Core: Change FCoreDelegates::OnObjectPropertyChanged to take the property change event as an additional parameter

diff --git a/Engine/Plugins/Editor/IntroTutorials/Source/IntroTutorials/Private/InteractiveTutorial.h b/Engine/Plugins/Editor/IntroTutorials/Source/IntroTutorials/Private/InteractiveTutorial.h
--- a/Engine/Plugins/Editor/IntroTutorials/Source/IntroTutorials/Private/InteractiveTutorial.h
+++ b/Engine/Plugins/Editor/IntroTutorials/Source/IntroTutorials/Private/InteractiveTutorial.h
@@ -6,7 +6,7 @@ DECLARE_DELEGATE( FOnBeginExcerptDelegate );
 DECLARE_DELEGATE_RetVal(bool, FShouldSkipExcerptDelegate);
 
 DECLARE_DELEGATE_RetVal_OneParam(bool, FActorSelectionChangedTrigger, const TArray<UObject*>&);
-DECLARE_DELEGATE_RetVal_OneParam(bool, FObjectPropertyChangedTrigger, UObject*);
+DECLARE_DELEGATE_RetVal_TwoParams(bool, FObjectPropertyChangedTrigger, UObject*, FPropertyChangedEvent&);
 DECLARE_DELEGATE_RetVal_OneParam(bool, FAssetEditorRequestedOpenTrigger, UObject*);
 DECLARE_DELEGATE_RetVal_TwoParams(bool, FEditorModeChangedTrigger, FEdMode*, bool);
 DECLARE_DELEGATE_RetVal_TwoParams(bool, FEditorModeChangedTrigger, FEdMode*, bool);
diff --git a/Engine/Plugins/Editor/IntroTutorials/Source/IntroTutorials/Private/InteractivityData.cpp b/Engine/Plugins/Editor/IntroTutorials/Source/IntroTutorials/Private/InteractivityData.cpp
--- a/Engine/Plugins/Editor/IntroTutorials/Source/IntroTutorials/Private/InteractivityData.cpp
+++ b/Engine/Plugins/Editor/IntroTutorials/Source/IntroTutorials/Private/InteractivityData.cpp
@@ -1685,7 +1685,7 @@ void PopulateInteractivityData(TSharedPtr<FInteractiveTutorials> InteractiveTuto
 			return bSphereIsSelected && SelectedObjectCount == 1;
 		}
 
-		static bool CompleteEditingDetails(UObject* Object)
+		static bool CompleteEditingDetails(UObject* Object, FPropertyChangedEvent& PropertyChangedEvent)
 		{
 			auto Actor = Cast<AActor>(Object);
 			if (Actor && Actor->GetFName() == FName("StaticMesh_Sphere"))
diff --git a/Engine/Plugins/Experimental/Paper2D/Source/Paper2DEditor/Private/Paper2DEditorModule.cpp b/Engine/Plugins/Experimental/Paper2D/Source/Paper2DEditor/Private/Paper2DEditorModule.cpp
--- a/Engine/Plugins/Experimental/Paper2D/Source/Paper2DEditor/Private/Paper2DEditorModule.cpp
+++ b/Engine/Plugins/Experimental/Paper2D/Source/Paper2DEditor/Private/Paper2DEditorModule.cpp
@@ -167,7 +167,7 @@ private:
 	}
 
 	// Called when a property on the specified object is modified
-	void OnPropertyChanged(UObject* ObjectBeingModified)
+	void OnPropertyChanged(UObject* ObjectBeingModified, FPropertyChangedEvent& PropertyChangedEvent)
 	{
 		if (UPaperSpriteAtlas* Atlas = Cast<UPaperSpriteAtlas>(ObjectBeingModified))
 		{
diff --git a/Engine/Source/Editor/ContentBrowser/Private/ContentBrowserSingleton.cpp b/Engine/Source/Editor/ContentBrowser/Private/ContentBrowserSingleton.cpp
--- a/Engine/Source/Editor/ContentBrowser/Private/ContentBrowserSingleton.cpp
+++ b/Engine/Source/Editor/ContentBrowser/Private/ContentBrowserSingleton.cpp
@@ -342,7 +342,7 @@ void FContentBrowserSingleton::SummonNewBrowser(bool bAllowLockedBrowsers)
 	}
 }
 
-void FContentBrowserSingleton::OnObjectPropertyChanged(UObject* ObjectBeingModified)
+void FContentBrowserSingleton::OnObjectPropertyChanged(UObject* ObjectBeingModified, FPropertyChangedEvent& PropertyChangedEvent)
 {
 	if ( !ensure(ObjectBeingModified) )
 	{
diff --git a/Engine/Source/Editor/ContentBrowser/Private/ContentBrowserSingleton.h b/Engine/Source/Editor/ContentBrowser/Private/ContentBrowserSingleton.h
--- a/Engine/Source/Editor/ContentBrowser/Private/ContentBrowserSingleton.h
+++ b/Engine/Source/Editor/ContentBrowser/Private/ContentBrowserSingleton.h
@@ -54,7 +54,7 @@ private:
 	void SummonNewBrowser(bool bAllowLockedBrowsers = false);
 
 	/** Handler for when a property changes on any object */
-	void OnObjectPropertyChanged(UObject* ObjectBeingModified);
+	void OnObjectPropertyChanged(UObject* ObjectBeingModified, FPropertyChangedEvent& PropertyChangedEvent);
 
 	/** Handler for a request to spawn a new content browser tab */
 	TSharedRef<SDockTab> SpawnContentBrowserTab( const FSpawnTabArgs& SpawnTabArgs, int32 BrowserIdx );
diff --git a/Engine/Source/Editor/ContentBrowser/Private/SAssetView.cpp b/Engine/Source/Editor/ContentBrowser/Private/SAssetView.cpp
--- a/Engine/Source/Editor/ContentBrowser/Private/SAssetView.cpp
+++ b/Engine/Source/Editor/ContentBrowser/Private/SAssetView.cpp
@@ -1919,7 +1919,7 @@ void SAssetView::OnAssetLoaded(UObject* Asset)
 	}
 }
 
-void SAssetView::OnObjectPropertyChanged(UObject* Asset)
+void SAssetView::OnObjectPropertyChanged(UObject* Asset, FPropertyChangedEvent& PropertyChangedEvent)
 {
 	if ( Asset != NULL )
 	{
diff --git a/Engine/Source/Editor/ContentBrowser/Private/SAssetView.h b/Engine/Source/Editor/ContentBrowser/Private/SAssetView.h
--- a/Engine/Source/Editor/ContentBrowser/Private/SAssetView.h
+++ b/Engine/Source/Editor/ContentBrowser/Private/SAssetView.h
@@ -315,7 +315,7 @@ private:
 	void OnAssetLoaded(UObject* Asset);
 
 	/** Handler for when an asset's property has changed */
-	void OnObjectPropertyChanged(UObject* Asset);
+	void OnObjectPropertyChanged(UObject* Asset, FPropertyChangedEvent& PropertyChangedEvent);
 
 	/** Handler for when any dynamic filters have been changed */
 	void OnDynamicFiltersChanged();
diff --git a/Engine/Source/Editor/EditorStyle/Private/SlateEditorStyle.cpp b/Engine/Source/Editor/EditorStyle/Private/SlateEditorStyle.cpp
--- a/Engine/Source/Editor/EditorStyle/Private/SlateEditorStyle.cpp
+++ b/Engine/Source/Editor/EditorStyle/Private/SlateEditorStyle.cpp
@@ -80,7 +80,7 @@ void SetColor( const TSharedRef< FLinearColor >& Source, const FLinearColor& Val
 	Source->A = Value.A;
 }
 
-void FSlateEditorStyle::FStyle::SettingsChanged( UObject* ChangedObject )
+void FSlateEditorStyle::FStyle::SettingsChanged( UObject* ChangedObject, FPropertyChangedEvent& PropertyChangedEvent )
 {
 	if ( ChangedObject == Settings.Get() )
 	{
diff --git a/Engine/Source/Editor/EditorStyle/Private/SlateEditorStyle.h b/Engine/Source/Editor/EditorStyle/Private/SlateEditorStyle.h
--- a/Engine/Source/Editor/EditorStyle/Private/SlateEditorStyle.h
+++ b/Engine/Source/Editor/EditorStyle/Private/SlateEditorStyle.h
@@ -80,7 +80,7 @@ public:
 		void SetupSourceControlStyles();
 		void SetupAutomationStyles();
 
-		void SettingsChanged( UObject* ChangedObject );
+		void SettingsChanged( UObject* ChangedObject, FPropertyChangedEvent& PropertyChangedEvent );
 		void SyncSettings();
 
 		const FVector2D Icon7x16;
diff --git a/Engine/Source/Editor/Persona/Private/Persona.cpp b/Engine/Source/Editor/Persona/Private/Persona.cpp
--- a/Engine/Source/Editor/Persona/Private/Persona.cpp
+++ b/Engine/Source/Editor/Persona/Private/Persona.cpp
@@ -1608,7 +1608,7 @@ void FPersona::SetPreviewMesh(USkeletalMesh* NewPreviewMesh)
 	}
 }
 
-void FPersona::OnPropertyChanged(UObject* ObjectBeingModified)
+void FPersona::OnPropertyChanged(UObject* ObjectBeingModified, FPropertyChangedEvent& PropertyChangedEvent)
 {
 	// Re-initialize the preview when a skeletal control is being edited
 	//@TODO: Should we still do this?
diff --git a/Engine/Source/Editor/Persona/Private/Persona.h b/Engine/Source/Editor/Persona/Private/Persona.h
--- a/Engine/Source/Editor/Persona/Private/Persona.h
+++ b/Engine/Source/Editor/Persona/Private/Persona.h
@@ -293,7 +293,7 @@ public:
 
 	// Property changed delegate
 	FCoreDelegates::FOnObjectPropertyChanged::FDelegate OnPropertyChangedHandle;
-	void OnPropertyChanged(UObject* ObjectBeingModified);
+	void OnPropertyChanged(UObject* ObjectBeingModified, FPropertyChangedEvent& PropertyChangedEvent);
 
 	/** Shared data between modes - for now only used for viewport **/
 	FPersonaModeSharedData ModeSharedData;
diff --git a/Engine/Source/Editor/Persona/Private/SAnimationBlendSpace1D.cpp b/Engine/Source/Editor/Persona/Private/SAnimationBlendSpace1D.cpp
--- a/Engine/Source/Editor/Persona/Private/SAnimationBlendSpace1D.cpp
+++ b/Engine/Source/Editor/Persona/Private/SAnimationBlendSpace1D.cpp
@@ -604,7 +604,7 @@ void SBlendSpaceEditor1D::OnBlendSpaceParamtersChanged()
 	UpdateBlendParameters();
 }
 
-void SBlendSpaceEditor1D::OnPropertyChanged(UObject* ObjectBeingModified)
+void SBlendSpaceEditor1D::OnPropertyChanged(UObject* ObjectBeingModified, FPropertyChangedEvent& PropertyChangedEvent)
 {
 	if(UBlendSpace1D* BlendSpace1D = Cast<UBlendSpace1D>(ObjectBeingModified))
 	{
diff --git a/Engine/Source/Editor/Persona/Private/SAnimationBlendSpace1D.h b/Engine/Source/Editor/Persona/Private/SAnimationBlendSpace1D.h
--- a/Engine/Source/Editor/Persona/Private/SAnimationBlendSpace1D.h
+++ b/Engine/Source/Editor/Persona/Private/SAnimationBlendSpace1D.h
@@ -187,7 +187,7 @@ private:
 
 	// Property changed delegate
 	FCoreDelegates::FOnObjectPropertyChanged::FDelegate OnPropertyChangedHandle;
-	void OnPropertyChanged(UObject* ObjectBeingModified);
+	void OnPropertyChanged(UObject* ObjectBeingModified, FPropertyChangedEvent& PropertyChangedEvent);
 
 	// Updates the UI to reflect the current blend space parameter values */
 	void UpdateBlendParameters();
diff --git a/Engine/Source/Editor/PropertyEditor/Private/UserInterface/PropertyTable/SPropertyTableCell.cpp b/Engine/Source/Editor/PropertyEditor/Private/UserInterface/PropertyTable/SPropertyTableCell.cpp
--- a/Engine/Source/Editor/PropertyEditor/Private/UserInterface/PropertyTable/SPropertyTableCell.cpp
+++ b/Engine/Source/Editor/PropertyEditor/Private/UserInterface/PropertyTable/SPropertyTableCell.cpp
@@ -39,7 +39,7 @@ void SPropertyTableCell::SetContent( const TSharedRef< SWidget >& NewContents )
 	];
 }
 
-void SPropertyTableCell::OnCellValueChanged( UObject* Object )
+void SPropertyTableCell::OnCellValueChanged( UObject* Object, FPropertyChangedEvent& PropertyChangedEvent )
 {
 	if ( Cell->GetObject().Get() == Object )
 	{
diff --git a/Engine/Source/Editor/PropertyEditor/Private/UserInterface/PropertyTable/SPropertyTableCell.h b/Engine/Source/Editor/PropertyEditor/Private/UserInterface/PropertyTable/SPropertyTableCell.h
--- a/Engine/Source/Editor/PropertyEditor/Private/UserInterface/PropertyTable/SPropertyTableCell.h
+++ b/Engine/Source/Editor/PropertyEditor/Private/UserInterface/PropertyTable/SPropertyTableCell.h
@@ -61,7 +61,7 @@ private:
 
 	void ExitedEditMode();
 
-	void OnCellValueChanged( UObject* Object );
+	void OnCellValueChanged( UObject* Object, FPropertyChangedEvent& PropertyChangedEvent );
 
 
 private:
diff --git a/Engine/Source/Editor/Sequencer/Private/SequencerObjectChangeListener.cpp b/Engine/Source/Editor/Sequencer/Private/SequencerObjectChangeListener.cpp
--- a/Engine/Source/Editor/Sequencer/Private/SequencerObjectChangeListener.cpp
+++ b/Engine/Source/Editor/Sequencer/Private/SequencerObjectChangeListener.cpp
@@ -76,7 +76,7 @@ void FSequencerObjectChangeListener::OnObjectPreEditChange( UObject* Object )
 	}
 }
 
-void FSequencerObjectChangeListener::OnObjectPostEditChange( UObject* Object )
+void FSequencerObjectChangeListener::OnObjectPostEditChange( UObject* Object, FPropertyChangedEvent& PropertyChangedEvent )
 {
 	if( Object )
 	{
@@ -121,7 +121,8 @@ void FSequencerObjectChangeListener::OnActorPostEditMove( AActor* Actor )
 {
 	// @todo sequencer actors: Currently this only fires on a "final" move.  For our purposes we probably
 	// want to get an update every single movement, even while dragging an object.
-	OnObjectPostEditChange( Actor );
+	FPropertyChangedEvent PropertyChangedEvent(nullptr);
+	OnObjectPostEditChange( Actor, PropertyChangedEvent );
 }
 
 void FSequencerObjectChangeListener::TriggerAllPropertiesChanged(UObject* Object)
diff --git a/Engine/Source/Editor/Sequencer/Private/SequencerObjectChangeListener.h b/Engine/Source/Editor/Sequencer/Private/SequencerObjectChangeListener.h
--- a/Engine/Source/Editor/Sequencer/Private/SequencerObjectChangeListener.h
+++ b/Engine/Source/Editor/Sequencer/Private/SequencerObjectChangeListener.h
@@ -31,7 +31,7 @@ private:
 	 *
 	 * @param Object	The object that PostEditChange was called on
 	 */
-	void OnObjectPostEditChange( UObject* Object );
+	void OnObjectPostEditChange( UObject* Object, FPropertyChangedEvent& PropertyChangedEvent );
 
 	/**
 	 * Called after an actor is moved (PostEditMove)
diff --git a/Engine/Source/Editor/StaticMeshEditor/Private/SStaticMeshEditorViewport.cpp b/Engine/Source/Editor/StaticMeshEditor/Private/SStaticMeshEditorViewport.cpp
--- a/Engine/Source/Editor/StaticMeshEditor/Private/SStaticMeshEditorViewport.cpp
+++ b/Engine/Source/Editor/StaticMeshEditor/Private/SStaticMeshEditorViewport.cpp
@@ -68,7 +68,7 @@ void SStaticMeshEditorViewport::RefreshViewport()
 	SceneViewport->Invalidate();
 }
 
-void SStaticMeshEditorViewport::OnObjectPropertyChanged(UObject* ObjectBeingModified)
+void SStaticMeshEditorViewport::OnObjectPropertyChanged(UObject* ObjectBeingModified, FPropertyChangedEvent& PropertyChangedEvent)
 {
 	if ( !ensure(ObjectBeingModified) )
 	{
diff --git a/Engine/Source/Editor/StaticMeshEditor/Private/SStaticMeshEditorViewport.h b/Engine/Source/Editor/StaticMeshEditor/Private/SStaticMeshEditorViewport.h
--- a/Engine/Source/Editor/StaticMeshEditor/Private/SStaticMeshEditorViewport.h
+++ b/Engine/Source/Editor/StaticMeshEditor/Private/SStaticMeshEditorViewport.h
@@ -94,7 +94,7 @@ private:
 	void SetRealtimePreview();
 
 	/** Callback for updating preview socket meshes if the static mesh or socket has been modified. */
-	void OnObjectPropertyChanged(UObject* ObjectBeingModified);
+	void OnObjectPropertyChanged(UObject* ObjectBeingModified, FPropertyChangedEvent& PropertyChangedEvent);
 private:
 	
 	/** The parent tab where this viewport resides */
diff --git a/Engine/Source/Editor/UMGEditor/Private/SUMGDesigner.cpp b/Engine/Source/Editor/UMGEditor/Private/SUMGDesigner.cpp
--- a/Engine/Source/Editor/UMGEditor/Private/SUMGDesigner.cpp
+++ b/Engine/Source/Editor/UMGEditor/Private/SUMGDesigner.cpp
@@ -167,7 +167,7 @@ void SUMGDesigner::OnBlueprintChanged(UBlueprint* InBlueprint)
 	}
 }
 
-void SUMGDesigner::OnObjectPropertyChanged(UObject* ObjectBeingModified)
+void SUMGDesigner::OnObjectPropertyChanged(UObject* ObjectBeingModified, FPropertyChangedEvent& PropertyChangedEvent)
 {
 	if ( !ensure(ObjectBeingModified) )
 	{
diff --git a/Engine/Source/Editor/UMGEditor/Private/SUMGDesigner.h b/Engine/Source/Editor/UMGEditor/Private/SUMGDesigner.h
--- a/Engine/Source/Editor/UMGEditor/Private/SUMGDesigner.h
+++ b/Engine/Source/Editor/UMGEditor/Private/SUMGDesigner.h
@@ -37,7 +37,7 @@ private:
 
 	UWidgetBlueprint* GetBlueprint() const;
 	void OnBlueprintChanged(UBlueprint* InBlueprint);
-	void OnObjectPropertyChanged(UObject* ObjectBeingModified);
+	void OnObjectPropertyChanged(UObject* ObjectBeingModified, FPropertyChangedEvent& PropertyChangedEvent);
 	void ShowDetailsForObjects(TArray<USlateWrapperComponent*> Widgets);
 
 	bool GetArrangedWidget(TSharedRef<SWidget> Widget, FArrangedWidget& ArrangedWidget) const;
diff --git a/Engine/Source/Editor/UMGEditor/Private/SUMGEditorTree.cpp b/Engine/Source/Editor/UMGEditor/Private/SUMGEditorTree.cpp
--- a/Engine/Source/Editor/UMGEditor/Private/SUMGEditorTree.cpp
+++ b/Engine/Source/Editor/UMGEditor/Private/SUMGEditorTree.cpp
@@ -94,7 +94,7 @@ void SUMGEditorTree::OnBlueprintChanged(UBlueprint* InBlueprint)
 	}
 }
 
-void SUMGEditorTree::OnObjectPropertyChanged(UObject* ObjectBeingModified)
+void SUMGEditorTree::OnObjectPropertyChanged(UObject* ObjectBeingModified, FPropertyChangedEvent& PropertyChangedEvent)
 {
 	if ( !ensure(ObjectBeingModified) )
 	{
diff --git a/Engine/Source/Editor/UMGEditor/Private/SUMGEditorTree.h b/Engine/Source/Editor/UMGEditor/Private/SUMGEditorTree.h
--- a/Engine/Source/Editor/UMGEditor/Private/SUMGEditorTree.h
+++ b/Engine/Source/Editor/UMGEditor/Private/SUMGEditorTree.h
@@ -19,7 +19,7 @@ private:
 	UWidgetBlueprint* GetBlueprint() const;
 
 	void OnBlueprintChanged(UBlueprint* InBlueprint);
-	void OnObjectPropertyChanged(UObject* ObjectBeingModified);
+	void OnObjectPropertyChanged(UObject* ObjectBeingModified, FPropertyChangedEvent& PropertyChangedEvent);
 
 	void ShowDetailsForObjects(TArray<USlateWrapperComponent*> Widgets);
 
diff --git a/Engine/Source/Editor/UnrealEd/Classes/CookOnTheSide/CookOnTheFlyServer.h b/Engine/Source/Editor/UnrealEd/Classes/CookOnTheSide/CookOnTheFlyServer.h
--- a/Engine/Source/Editor/UnrealEd/Classes/CookOnTheSide/CookOnTheFlyServer.h
+++ b/Engine/Source/Editor/UnrealEd/Classes/CookOnTheSide/CookOnTheFlyServer.h
@@ -309,7 +309,7 @@ public:
 	 */
 
 	void OnObjectModified( UObject *ObjectMoving );
-	void OnObjectPropertyChanged(UObject* ObjectBeingModified);
+	void OnObjectPropertyChanged(UObject* ObjectBeingModified, FPropertyChangedEvent& PropertyChangedEvent);
 	void OnObjectUpdated( UObject *Object );
 
 	/**
diff --git a/Engine/Source/Editor/UnrealEd/Private/AssetThumbnail.cpp b/Engine/Source/Editor/UnrealEd/Private/AssetThumbnail.cpp
--- a/Engine/Source/Editor/UnrealEd/Private/AssetThumbnail.cpp
+++ b/Engine/Source/Editor/UnrealEd/Private/AssetThumbnail.cpp
@@ -1208,7 +1208,7 @@ void FAssetThumbnailPool::OnAssetLoaded( UObject* Asset )
 	}
 }
 
-void FAssetThumbnailPool::OnObjectPropertyChanged( UObject* ObjectBeingModified )
+void FAssetThumbnailPool::OnObjectPropertyChanged( UObject* ObjectBeingModified, FPropertyChangedEvent& PropertyChangedEvent )
 {
 	if ( ObjectBeingModified->HasAnyFlags(RF_ClassDefaultObject) && ObjectBeingModified->GetClass()->ClassGeneratedBy != NULL )
 	{
diff --git a/Engine/Source/Editor/UnrealEd/Private/CookOnTheFlyServer.cpp b/Engine/Source/Editor/UnrealEd/Private/CookOnTheFlyServer.cpp
--- a/Engine/Source/Editor/UnrealEd/Private/CookOnTheFlyServer.cpp
+++ b/Engine/Source/Editor/UnrealEd/Private/CookOnTheFlyServer.cpp
@@ -345,7 +345,7 @@ void UCookOnTheFlyServer::OnObjectModified( UObject *ObjectMoving )
 	OnObjectUpdated( ObjectMoving );
 }
 
-void UCookOnTheFlyServer::OnObjectPropertyChanged(UObject* ObjectBeingModified)
+void UCookOnTheFlyServer::OnObjectPropertyChanged(UObject* ObjectBeingModified, FPropertyChangedEvent& PropertyChangedEvent)
 {
 	OnObjectUpdated( ObjectBeingModified );
 }
diff --git a/Engine/Source/Editor/UnrealEd/Public/AssetThumbnail.h b/Engine/Source/Editor/UnrealEd/Public/AssetThumbnail.h
--- a/Engine/Source/Editor/UnrealEd/Public/AssetThumbnail.h
+++ b/Engine/Source/Editor/UnrealEd/Public/AssetThumbnail.h
@@ -189,7 +189,7 @@ private:
 	void OnAssetLoaded( UObject* Asset );
 
 	/** Handler for when an asset is loaded */
-	void OnObjectPropertyChanged( UObject* Asset );
+	void OnObjectPropertyChanged( UObject* Asset, FPropertyChangedEvent& PropertyChangedEvent );
 
 private:
 	/** Information about a thumbnail */
diff --git a/Engine/Source/Runtime/Core/Public/Misc/CallbackDevice.h b/Engine/Source/Runtime/Core/Public/Misc/CallbackDevice.h
--- a/Engine/Source/Runtime/Core/Public/Misc/CallbackDevice.h
+++ b/Engine/Source/Runtime/Core/Public/Misc/CallbackDevice.h
@@ -45,7 +45,7 @@ public:
 	DECLARE_DELEGATE_TwoParams(FHotFixDelegate, void *, int32);
 
 	// Callback for object property modifications
-	DECLARE_MULTICAST_DELEGATE_OneParam(FOnObjectPropertyChanged, UObject*);
+	DECLARE_MULTICAST_DELEGATE_TwoParams(FOnObjectPropertyChanged, UObject*, struct FPropertyChangedEvent&);
 
 	// Callback for object property modifications
 	DECLARE_MULTICAST_DELEGATE_OneParam(FOnActorLabelChanged, AActor*);
diff --git a/Engine/Source/Runtime/CoreUObject/Private/UObject/Obj.cpp b/Engine/Source/Runtime/CoreUObject/Private/UObject/Obj.cpp
--- a/Engine/Source/Runtime/CoreUObject/Private/UObject/Obj.cpp
+++ b/Engine/Source/Runtime/CoreUObject/Private/UObject/Obj.cpp
@@ -207,7 +207,7 @@ void UObject::PostEditChange(void)
 
 void UObject::PostEditChangeProperty(FPropertyChangedEvent& PropertyChangedEvent)
 {
-	FCoreDelegates::OnObjectPropertyChanged.Broadcast(this);
+	FCoreDelegates::OnObjectPropertyChanged.Broadcast(this, PropertyChangedEvent);
 }
 
 
