Fixed sequencer not opening once it has been closed (world centric bug)
Removed ISequencerInternals interface as it was not necessary and confusing.

diff --git a/Engine/Source/Editor/LevelEditor/Private/LevelEditor.cpp b/Engine/Source/Editor/LevelEditor/Private/LevelEditor.cpp
--- a/Engine/Source/Editor/LevelEditor/Private/LevelEditor.cpp
+++ b/Engine/Source/Editor/LevelEditor/Private/LevelEditor.cpp
@@ -16,6 +16,8 @@
 #include "GlobalEditorCommonCommands.h"
 #include "IUserFeedbackModule.h"
 #include "SlateReflector.h"
+#include "SDockTab.h"
+#include "ToolkitManager.h"
 
 // @todo Editor: remove this circular dependency
 #include "Editor/MainFrame/Public/Interfaces/IMainFrameModule.h"
@@ -255,13 +257,32 @@ void FLevelEditorModule::SummonWorldBrowserComposition()
 }
 
 // @todo remove when world-centric mode is added
-void FLevelEditorModule::AttachSequencer(TSharedPtr<SWidget> Sequencer)
+void FLevelEditorModule::AttachSequencer(TSharedPtr<SWidget> SequencerWidget, TSharedPtr<IAssetEditorInstance> SequencerAssetEditor )
 {
 	if( FParse::Param( FCommandLine::Get(), TEXT( "Sequencer" ) ) )
 	{
+		struct Local
+		{
+			static void OnSequencerClosed( TSharedRef<SDockTab> DockTab, TWeakPtr<IAssetEditorInstance> SequencerAssetEditor )
+			{
+				SequencerAssetEditor.Pin()->CloseWindow();
+			}
+		};
+		
 		TSharedPtr<SLevelEditor> LevelEditorInstance = LevelEditorInstancePtr.Pin();
-		LevelEditorInstance->InvokeTab("Sequencer");
-		LevelEditorInstance->SequencerTab->SetContent(Sequencer.ToSharedRef());
+
+		if( SequencerWidget.IsValid() && SequencerAssetEditor.IsValid() )
+		{
+			LevelEditorInstance->InvokeTab("Sequencer");
+			LevelEditorInstance->SequencerTab->SetOnTabClosed( SDockTab::FOnTabClosedCallback::CreateStatic( &Local::OnSequencerClosed, TWeakPtr<IAssetEditorInstance>( SequencerAssetEditor ) ) );
+			LevelEditorInstance->SequencerTab->SetContent(SequencerWidget.ToSharedRef());
+		}
+		else
+		{
+			LevelEditorInstance->SequencerTab.Reset();
+		}
+
+	
 	}
 }
 
@@ -275,7 +296,7 @@ void FLevelEditorModule::FocusPIEViewport()
 {
 	TSharedPtr<SLevelEditor> LevelEditorInstance = LevelEditorInstancePtr.Pin();
 	if( LevelEditorInstance.IsValid() && LevelEditorTabManager.IsValid() && LevelEditorInstance->HasActivePlayInEditorViewport() )
-		{
+	{
 		FGlobalTabmanager::Get()->DrawAttentionToTabManager( LevelEditorTabManager.ToSharedRef() );
 	}
 }
diff --git a/Engine/Source/Editor/LevelEditor/Public/LevelEditor.h b/Engine/Source/Editor/LevelEditor/Public/LevelEditor.h
--- a/Engine/Source/Editor/LevelEditor/Public/LevelEditor.h
+++ b/Engine/Source/Editor/LevelEditor/Public/LevelEditor.h
@@ -79,7 +79,7 @@ public:
 	 * Spawns a new sequencer tab if one doesn't exist already
 	 * @todo This only works with the first level editor. Fix it.
 	 */
-	virtual void AttachSequencer(TSharedPtr<SWidget> Sequencer);
+	virtual void AttachSequencer(TSharedPtr<SWidget> SequencerWidget, TSharedPtr<class IAssetEditorInstance> SequencerAssetEditor );
 
 	/**
 	 * Starts a play in editor session using the active viewport
diff --git a/Engine/Source/Editor/Sequencer/Private/SAnimationOutlinerView.cpp b/Engine/Source/Editor/Sequencer/Private/SAnimationOutlinerView.cpp
--- a/Engine/Source/Editor/Sequencer/Private/SAnimationOutlinerView.cpp
+++ b/Engine/Source/Editor/Sequencer/Private/SAnimationOutlinerView.cpp
@@ -4,7 +4,7 @@
 
 #include "SAnimationOutlinerView.h"
 #include "ScopedTransaction.h"
-#include "ISequencerInternals.h"
+#include "Sequencer.h"
 #include "MovieScene.h"
 
 
@@ -154,7 +154,7 @@ EVisibility SAnimationOutlinerTreeNode::GetExpanderVisibility() const
 	return DisplayNode->GetNumChildren() > 0 ? EVisibility::Visible : EVisibility::Hidden;
 }
 
-void SAnimationOutlinerView::Construct( const FArguments& InArgs, TSharedRef<FSequencerDisplayNode> InRootNode, TSharedRef<ISequencerInternals> InSequencer )
+void SAnimationOutlinerView::Construct( const FArguments& InArgs, TSharedRef<FSequencerDisplayNode> InRootNode, TSharedRef<FSequencer> InSequencer )
 {
 	SAnimationOutlinerViewBase::Construct( SAnimationOutlinerViewBase::FArguments(), InRootNode );
 
diff --git a/Engine/Source/Editor/Sequencer/Private/SAnimationOutlinerView.h b/Engine/Source/Editor/Sequencer/Private/SAnimationOutlinerView.h
--- a/Engine/Source/Editor/Sequencer/Private/SAnimationOutlinerView.h
+++ b/Engine/Source/Editor/Sequencer/Private/SAnimationOutlinerView.h
@@ -80,7 +80,7 @@ public:
 	SLATE_END_ARGS()
 
 	/** Construct this widget.  Called by the SNew() Slate macro. */
-	void Construct( const FArguments& InArgs, TSharedRef<FSequencerDisplayNode> RootNode, TSharedRef<ISequencerInternals> InSequencer );
+	void Construct( const FArguments& InArgs, TSharedRef<FSequencerDisplayNode> RootNode, TSharedRef<FSequencer> InSequencer );
 
 	/** SAnimationOutlinerView destructor */
 	virtual ~SAnimationOutlinerView();
@@ -103,5 +103,5 @@ private:
 
 private:
 	/** Internal sequencer interface */
-	TWeakPtr<ISequencerInternals> Sequencer;
+	TWeakPtr<FSequencer> Sequencer;
 };
\ No newline at end of file
diff --git a/Engine/Source/Editor/Sequencer/Private/SSection.cpp b/Engine/Source/Editor/Sequencer/Private/SSection.cpp
--- a/Engine/Source/Editor/Sequencer/Private/SSection.cpp
+++ b/Engine/Source/Editor/Sequencer/Private/SSection.cpp
@@ -173,7 +173,7 @@ void SSection::CheckForEdgeInteraction( const FPointerEvent& MouseEvent, const F
 	}
 }
 
-ISequencerInternals& SSection::GetSequencer() const
+FSequencer& SSection::GetSequencer() const
 {
 	return ParentSectionArea->GetSequencer();
 }
@@ -304,7 +304,7 @@ void SSection::PaintKeys( const FGeometry& AllottedGeometry, const FSlateRect& M
 {
 	UMovieSceneSection& SectionObject = *SectionInterface->GetSectionObject();
 
-	const ISequencerInternals& Sequencer = ParentSectionArea->GetSequencer();
+	const FSequencer& Sequencer = ParentSectionArea->GetSequencer();
 
 	const FSlateBrush* BackgroundBrush = FEditorStyle::GetBrush("Sequencer.SectionArea.Background");
 
@@ -393,7 +393,7 @@ TSharedPtr<SWidget> SSection::OnSummonContextMenu( const FGeometry& MyGeometry,
 {
 	FSelectedKey Key = GetKeyUnderMouse( MouseEvent.GetScreenSpacePosition(), MyGeometry );
 
-	ISequencerInternals& Sequencer = GetSequencer();
+	FSequencer& Sequencer = GetSequencer();
 
 	UMovieSceneSection* SceneSection = SectionInterface->GetSectionObject();
 
@@ -407,7 +407,7 @@ TSharedPtr<SWidget> SSection::OnSummonContextMenu( const FGeometry& MyGeometry,
 			NSLOCTEXT("Sequencer", "DeleteKey", "Delete"),
 			NSLOCTEXT("Sequencer", "DeleteKeyToolTip", "Deletes the selected keys."),
 			FSlateIcon(),
-			FUIAction(FExecuteAction::CreateSP(&Sequencer, &ISequencerInternals::DeleteKeys, Sequencer.GetSelectedKeys().Array()))
+			FUIAction(FExecuteAction::CreateSP(&Sequencer, &FSequencer::DeleteSelectedKeys))
 			);
 	}
 	else
@@ -420,14 +420,14 @@ TSharedPtr<SWidget> SSection::OnSummonContextMenu( const FGeometry& MyGeometry,
 				NSLOCTEXT("Sequencer", "RenameShot", "Rename"),
 				NSLOCTEXT("Sequencer", "RenameShotToolTip", "Renames this shot."),
 				FSlateIcon(),
-				FUIAction(FExecuteAction::CreateSP(&Sequencer, &ISequencerInternals::RenameShot, SceneSection))
+				FUIAction(FExecuteAction::CreateSP(&Sequencer, &FSequencer::RenameShot, SceneSection))
 				);
 
 			MenuBuilder.AddMenuEntry(
 				NSLOCTEXT("Sequencer", "FilterToShots", "Filter To Shots"),
 				NSLOCTEXT("Sequencer", "FilterToShotsToolTip", "Filters to the selected shot sections"),
 				FSlateIcon(),
-				FUIAction(FExecuteAction::CreateSP(&Sequencer, &ISequencerInternals::FilterToSelectedShotSections, true))
+				FUIAction(FExecuteAction::CreateSP(&Sequencer, &FSequencer::FilterToSelectedShotSections, true))
 				);
 		}
 
@@ -437,7 +437,7 @@ TSharedPtr<SWidget> SSection::OnSummonContextMenu( const FGeometry& MyGeometry,
 			NSLOCTEXT("Sequencer", "DeleteSection", "Delete"),
 			NSLOCTEXT("Sequencer", "DeleteSectionToolTip", "Deletes this section."),
 			FSlateIcon(),
-			FUIAction(FExecuteAction::CreateSP(&Sequencer, &ISequencerInternals::DeleteSection, SceneSection))
+			FUIAction(FExecuteAction::CreateSP(&Sequencer, &FSequencer::DeleteSection, SceneSection))
 			);
 	}
 	
diff --git a/Engine/Source/Editor/Sequencer/Private/SSection.h b/Engine/Source/Editor/Sequencer/Private/SSection.h
--- a/Engine/Source/Editor/Sequencer/Private/SSection.h
+++ b/Engine/Source/Editor/Sequencer/Private/SSection.h
@@ -120,7 +120,7 @@ private:
 	void HandleSectionSelection( const FPointerEvent& MouseEvent );
 
 	/** @return the sequencer interface */
-	ISequencerInternals& GetSequencer() const;
+	FSequencer& GetSequencer() const;
 
 	/**
 	 * Creates a drag operation based on the current mouse event and what is selected
diff --git a/Engine/Source/Editor/Sequencer/Private/SSequencer.cpp b/Engine/Source/Editor/Sequencer/Private/SSequencer.cpp
--- a/Engine/Source/Editor/Sequencer/Private/SSequencer.cpp
+++ b/Engine/Source/Editor/Sequencer/Private/SSequencer.cpp
@@ -4,7 +4,7 @@
 
 #include "SSequencer.h"
 #include "Editor/SequencerWidgets/Public/SequencerWidgetsModule.h"
-#include "ISequencerInternals.h"
+#include "Sequencer.h"
 #include "MovieScene.h"
 #include "MovieSceneSection.h"
 #include "ISequencerSection.h"
@@ -38,7 +38,7 @@ public:
 		SLATE_ATTRIBUTE( TRange<float>, ViewRange )
 	SLATE_END_ARGS()
 
-	void Construct(const FArguments& InArgs, TWeakPtr<ISequencerInternals> InSequencer)
+	void Construct(const FArguments& InArgs, TWeakPtr<FSequencer> InSequencer)
 	{
 		ViewRange = InArgs._ViewRange;
 		Sequencer = InSequencer;
@@ -132,14 +132,14 @@ private:
 	/** The current minimum view range */
 	TAttribute< TRange<float> > ViewRange;
 	/** The main sequencer interface */
-	TWeakPtr<ISequencerInternals> Sequencer;
+	TWeakPtr<FSequencer> Sequencer;
 	/** Cached set of ranges that are being filtering currently */
 	TArray< TRange<float> > CachedFilteredRanges;
 };
 
 
 BEGIN_SLATE_FUNCTION_BUILD_OPTIMIZATION
-void SSequencer::Construct( const FArguments& InArgs, TSharedRef< class ISequencerInternals > InSequencer )
+void SSequencer::Construct( const FArguments& InArgs, TSharedRef< class FSequencer > InSequencer )
 {
 	Sequencer = InSequencer;
 
@@ -312,6 +312,10 @@ void SSequencer::Construct( const FArguments& InArgs, TSharedRef< class ISequenc
 }
 END_SLATE_FUNCTION_BUILD_OPTIMIZATION
 
+SSequencer::~SSequencer()
+{
+}
+
 void SSequencer::Tick( const FGeometry& AllottedGeometry, const double InCurrentTime, const float InDeltaTime )
 {
 	Sequencer.Pin()->Tick(InDeltaTime);
@@ -511,7 +515,7 @@ FReply SSequencer::OnKeyDown( const FGeometry& MyGeometry, const FKeyboardEvent&
 
 void SSequencer::OnAssetsDropped( const FAssetDragDropOp& DragDropOp )
 {
-	ISequencerInternals& SequencerRef = *Sequencer.Pin();
+	FSequencer& SequencerRef = *Sequencer.Pin();
 
 	bool bObjectAdded = false;
 	bool bSpawnableAdded = false;
@@ -577,7 +581,7 @@ void SSequencer::OnAssetsDropped( const FAssetDragDropOp& DragDropOp )
 
 void SSequencer::OnClassesDropped( const FClassDragDropOp& DragDropOp )
 {
-	ISequencerInternals& SequencerRef = *Sequencer.Pin();
+	FSequencer& SequencerRef = *Sequencer.Pin();
 
 	bool bSpawnableAdded = false;
 	for( auto ClassIter = DragDropOp.ClassesToDrop.CreateConstIterator(); ClassIter; ++ClassIter )
@@ -605,7 +609,7 @@ void SSequencer::OnClassesDropped( const FClassDragDropOp& DragDropOp )
 
 void SSequencer::OnUnloadedClassesDropped( const FUnloadedClassDragDropOp& DragDropOp )
 {
-	ISequencerInternals& SequencerRef = *Sequencer.Pin();
+	FSequencer& SequencerRef = *Sequencer.Pin();
 	bool bSpawnableAdded = false;
 	for( auto ClassDataIter = DragDropOp.AssetsToDrop->CreateConstIterator(); ClassDataIter; ++ClassDataIter )
 	{
@@ -660,7 +664,7 @@ void SSequencer::OnUnloadedClassesDropped( const FUnloadedClassDragDropOp& DragD
 
 void SSequencer::OnActorsDropped( FActorDragDropGraphEdOp& DragDropOp )
 {
-	ISequencerInternals& SequencerRef = *Sequencer.Pin();
+	FSequencer& SequencerRef = *Sequencer.Pin();
 	bool bPossessableAdded = false;
 	for( auto ActorIter = DragDropOp.Actors.CreateIterator(); ActorIter; ++ActorIter )
 	{
@@ -746,7 +750,7 @@ void SSequencer::DeleteSelectedNodes()
 	{
 		const FScopedTransaction Transaction( LOCTEXT("UndoDeletingObject", "Delete Object from MovieScene") );
 
-		ISequencerInternals& SequencerRef = *Sequencer.Pin();
+		FSequencer& SequencerRef = *Sequencer.Pin();
 
 		for( auto It = SelectedNodes.CreateConstIterator(); It; ++It )
 		{
diff --git a/Engine/Source/Editor/Sequencer/Private/SSequencer.h b/Engine/Source/Editor/Sequencer/Private/SSequencer.h
--- a/Engine/Source/Editor/Sequencer/Private/SSequencer.h
+++ b/Engine/Source/Editor/Sequencer/Private/SSequencer.h
@@ -72,7 +72,9 @@ public:
 	SLATE_END_ARGS()
 
 
-	void Construct( const FArguments& InArgs, TSharedRef< class ISequencerInternals > InSequencer );
+	void Construct( const FArguments& InArgs, TSharedRef< class FSequencer > InSequencer );
+
+	~SSequencer();
 
 	virtual void Tick( const FGeometry& AllottedGeometry, const double InCurrentTime, const float InDeltaTime ) OVERRIDE;
 	virtual bool SupportsKeyboardFocus() const OVERRIDE { return true; }
@@ -198,5 +200,5 @@ private:
 	/** Whether the clean sequencer view is enabled */
 	TAttribute<bool> CleanViewEnabled;
 	/** The main sequencer interface */
-	TWeakPtr<ISequencerInternals> Sequencer;
+	TWeakPtr<FSequencer> Sequencer;
 };
diff --git a/Engine/Source/Editor/Sequencer/Private/SSequencerSectionAreaView.cpp b/Engine/Source/Editor/Sequencer/Private/SSequencerSectionAreaView.cpp
--- a/Engine/Source/Editor/Sequencer/Private/SSequencerSectionAreaView.cpp
+++ b/Engine/Source/Editor/Sequencer/Private/SSequencerSectionAreaView.cpp
@@ -7,7 +7,7 @@
 #include "MovieSceneSection.h"
 #include "SequencerDisplayNode.h"
 #include "SSequencer.h"
-#include "ISequencerInternals.h"
+#include "Sequencer.h"
 #include "MovieSceneShotSection.h"
 #include "CommonMovieSceneTools.h"
 
@@ -193,7 +193,7 @@ FTimeToPixel SSequencerSectionAreaView::GetTimeToPixel( const FGeometry& Allotte
 }
 
 
-ISequencerInternals& SSequencerSectionAreaView::GetSequencer() const
+FSequencer& SSequencerSectionAreaView::GetSequencer() const
 {
 	return SectionAreaNode->GetSequencer();
 }
\ No newline at end of file
diff --git a/Engine/Source/Editor/Sequencer/Private/SSequencerSectionAreaView.h b/Engine/Source/Editor/Sequencer/Private/SSequencerSectionAreaView.h
--- a/Engine/Source/Editor/Sequencer/Private/SSequencerSectionAreaView.h
+++ b/Engine/Source/Editor/Sequencer/Private/SSequencerSectionAreaView.h
@@ -63,7 +63,7 @@ private:
 	void GenerateSectionWidgets();
 
 	/** @return The sequencer interface */
-	ISequencerInternals& GetSequencer() const;
+	FSequencer& GetSequencer() const;
 private:
 
 	/** List of all sequencer drag operations going on  in this area */
diff --git a/Engine/Source/Editor/Sequencer/Private/SSequencerTrackArea.h b/Engine/Source/Editor/Sequencer/Private/SSequencerTrackArea.h
--- a/Engine/Source/Editor/Sequencer/Private/SSequencerTrackArea.h
+++ b/Engine/Source/Editor/Sequencer/Private/SSequencerTrackArea.h
@@ -16,7 +16,7 @@ public:
 		SLATE_ATTRIBUTE( float, OutlinerFillPercent )
 	SLATE_END_ARGS()
 
-	void Construct( const FArguments& InArgs, TSharedRef<ISequencerInternals> InSequencer )
+	void Construct( const FArguments& InArgs, TSharedRef<FSequencer> InSequencer )
 	{
 		ViewRange = InArgs._ViewRange;
 		OutlinerFillPercent = InArgs._OutlinerFillPercent;
@@ -53,5 +53,5 @@ private:
 	/** The fill percentage of the animation outliner */
 	TAttribute<float> OutlinerFillPercent;
 	/** The main sequencer interface */
-	TWeakPtr<ISequencerInternals> Sequencer;
+	TWeakPtr<FSequencer> Sequencer;
 };
diff --git a/Engine/Source/Editor/Sequencer/Private/SelectedKey.h b/Engine/Source/Editor/Sequencer/Private/SelectedKey.h
new file mode 100644
--- /dev/null
+++ b/Engine/Source/Editor/Sequencer/Private/SelectedKey.h
@@ -0,0 +1,49 @@
+// Copyright 1998-2014 Epic Games, Inc. All Rights Reserved.
+
+#pragma once
+
+class UMovieSceneSection;
+class IKeyArea;
+
+/**
+ * Represents a selected key in the sequencer
+ */
+struct FSelectedKey
+{	
+	/** Section that the key belongs to */
+	UMovieSceneSection* Section;
+	/** Key area providing the key */
+	TSharedPtr<IKeyArea> KeyArea;
+	/** Index of the key in the key area */
+	TOptional<FKeyHandle> KeyHandle;
+
+	FSelectedKey( UMovieSceneSection& InSection, TSharedPtr<IKeyArea> InKeyArea, FKeyHandle InKeyHandle )
+		: Section( &InSection )
+		, KeyArea( InKeyArea )
+		, KeyHandle(InKeyHandle)
+	{}
+
+	FSelectedKey()
+		: Section( NULL )
+		, KeyArea( NULL )
+		, KeyHandle()
+	{}
+
+	/** @return Whether or not this is a valid selected key */
+	bool IsValid() const { return Section != NULL && KeyArea.IsValid() && KeyHandle.IsSet(); }
+
+	friend uint32 GetTypeHash( const FSelectedKey& SelectedKey )
+	{
+		return GetTypeHash( SelectedKey.Section ) ^ GetTypeHash( SelectedKey.KeyArea ) ^ 
+			(SelectedKey.KeyHandle.IsSet() ? GetTypeHash(SelectedKey.KeyHandle.GetValue()) : 0);
+	} 
+
+	bool operator==( const FSelectedKey& OtherKey ) const
+	{
+		return Section == OtherKey.Section && KeyArea == OtherKey.KeyArea &&
+			KeyHandle.IsSet() && OtherKey.KeyHandle.IsSet() &&
+			KeyHandle.GetValue() == OtherKey.KeyHandle.GetValue();
+	}
+
+};
+
diff --git a/Engine/Source/Editor/Sequencer/Private/Sequencer.cpp b/Engine/Source/Editor/Sequencer/Private/Sequencer.cpp
--- a/Engine/Source/Editor/Sequencer/Private/Sequencer.cpp
+++ b/Engine/Source/Editor/Sequencer/Private/Sequencer.cpp
@@ -52,29 +52,25 @@ void FSequencer::RegisterTabSpawners(const TSharedRef<class FTabManager>& TabMan
 {
 	const IWorkspaceMenuStructure& MenuStructure = WorkspaceMenu::GetMenuStructure();
 
-	if( FParse::Param( FCommandLine::Get(), TEXT( "Sequencer" ) ) )
+	if( FParse::Param( FCommandLine::Get(), TEXT( "Sequencer" ) ) && !IsWorldCentricAssetEditor() )
 	{
 		TabManager->RegisterTabSpawner( SequencerMainTabId, FOnSpawnTab::CreateSP(this, &FSequencer::SpawnTab_SequencerMain) )
 			.SetDisplayName( LOCTEXT("SequencerMainTab", "Sequencer") )
 			.SetGroup( MenuStructure.GetAssetEditorCategory() );
 	}
 
-	TabManager->RegisterTabSpawner( SequencerDetailsTabId, FOnSpawnTab::CreateSP(this, &FSequencer::SpawnTab_Details) )
-		.SetDisplayName( LOCTEXT("SequencerDetailsTab", "Details") )
-		.SetGroup( MenuStructure.GetAssetEditorCategory() );
 }
 
 void FSequencer::UnregisterTabSpawners(const TSharedRef<class FTabManager>& TabManager)
 {
-	if( FParse::Param( FCommandLine::Get(), TEXT( "Sequencer" ) ) )
+	if( FParse::Param( FCommandLine::Get(), TEXT( "Sequencer" ) ) && !IsWorldCentricAssetEditor() )
 	{
 		TabManager->UnregisterTabSpawner( SequencerMainTabId );
 	}
-	TabManager->UnregisterTabSpawner( SequencerDetailsTabId );
 	
 	// @todo remove when world-centric mode is added
 	FLevelEditorModule& LevelEditorModule = FModuleManager::LoadModuleChecked<FLevelEditorModule>("LevelEditor");
-	LevelEditorModule.AttachSequencer(SNullWidget::NullWidget);
+	LevelEditorModule.AttachSequencer( SNullWidget::NullWidget, nullptr );
 }
 
 void FSequencer::InitSequencer( const EToolkitMode::Type Mode, const TSharedPtr< class IToolkitHost >& InitToolkitHost, UObject* ObjectToEdit, const TArray<FOnCreateTrackEditor>& TrackEditorDelegates )
@@ -82,14 +78,15 @@ void FSequencer::InitSequencer( const EToolkitMode::Type Mode, const TSharedPtr<
 	if( FParse::Param( FCommandLine::Get(), TEXT( "Sequencer" ) ) )
 	{
 		const TSharedRef<FTabManager::FLayout> StandaloneDefaultLayout = FTabManager::NewLayout( "Standalone_Sequencer_Layout" )
-			->AddArea(
-			FTabManager::NewPrimaryArea()
-			->Split
+			->AddArea
 			(
-			FTabManager::NewStack()
-			->AddTab( SequencerMainTabId, ETabState::OpenedTab )
-			->AddTab( SequencerDetailsTabId, ETabState::OpenedTab )
-			)
+				FTabManager::NewPrimaryArea()
+				->Split
+				(
+					FTabManager::NewStack()
+					->AddTab(SequencerMainTabId, ETabState::OpenedTab)
+					->AddTab(SequencerDetailsTabId, ETabState::OpenedTab)
+				)
 			);
 
 		const bool bCreateDefaultStandaloneMenu = true;
@@ -123,7 +120,7 @@ void FSequencer::InitSequencer( const EToolkitMode::Type Mode, const TSharedPtr<
 
 		// @todo remove when world-centric mode is added
 		FLevelEditorModule& LevelEditorModule = FModuleManager::LoadModuleChecked<FLevelEditorModule>("LevelEditor");
-		LevelEditorModule.AttachSequencer(SequencerWidget);
+		LevelEditorModule.AttachSequencer(SequencerWidget, SharedThis( this ) );
 
 		// Hook into the editor's mechanism for checking whether we need live capture of PIE/SIE actor state
 		GEditor->GetActorRecordingState().AddSP( this, &FSequencer::GetActorRecordingState );
@@ -132,29 +129,10 @@ void FSequencer::InitSequencer( const EToolkitMode::Type Mode, const TSharedPtr<
 		GEditor->RegisterForUndo(this);
 
 
-		// Setup our tool's layout
-		// @todo re-enable once world centric works again
-		/*if( IsWorldCentricAssetEditor() )
-		{
-		if( !SequencerMainTab.IsValid() )
-		{
-		const FString TabInitializationPayload;
-		SequencerMainTab = SpawnToolkitTab( SequencerMainTabId, TabInitializationPayload, EToolkitTabSpot::BelowLevelEditor );
-		}
-
-		// @todo sequencer: Ideally we should be possessing the level editor's details view instead of spawning our own
-		if( !DetailsTab.IsValid() )
-		{
-		const FString TabInitializationPayload;		// NOTE: Payload not currently used for details
-		DetailsTab = SpawnToolkitTab( SequencerDetailsTabId, TabInitializationPayload, EToolkitTabSpot::Details );
-		}
-		}*/
-
-
 		// We need to find out when the user loads a new map, because we might need to re-create puppet actors
 		// when previewing a MovieScene
 		//auto& LevelEditorModule = FModuleManager::LoadModuleChecked< FLevelEditorModule >( TEXT( "LevelEditor" ) );
-		LevelEditorModule.OnMapChanged().AddRaw( this, &FSequencer::OnMapChanged );
+		LevelEditorModule.OnMapChanged().AddSP( this, &FSequencer::OnMapChanged );
 
 
 		// Start listening for important changes on this MovieScene
@@ -223,7 +201,7 @@ FSequencer::~FSequencer()
 	if( FModuleManager::Get().IsModuleLoaded( TEXT( "LevelEditor" ) ) )
 	{
 		auto& LevelEditorModule = FModuleManager::LoadModuleChecked< FLevelEditorModule >( TEXT( "LevelEditor" ) );
-		LevelEditorModule.OnMapChanged().RemoveRaw( this, &FSequencer::OnMapChanged );
+		LevelEditorModule.OnMapChanged().RemoveAll( this );
 	}
 
 	if( PlayMovieSceneNode.IsValid() )
@@ -483,15 +461,14 @@ void FSequencer::DeleteSection(class UMovieSceneSection* Section)
 	}
 }
 
-void FSequencer::DeleteKeys(TArray<FSelectedKey> KeysToDelete)
+void FSequencer::DeleteSelectedKeys()
 {
-	for (int32 i = 0; i < KeysToDelete.Num(); ++i)
+	TArray<FSelectedKey> SelectedKeysArray = SelectedKeys.Array();
+	for ( const FSelectedKey& Key : SelectedKeysArray )
 	{
-		const FSelectedKey& SelectedKey = KeysToDelete[i];
-
-		if (SelectedKey.IsValid())
+		if (Key.IsValid())
 		{
-			SelectedKey.KeyArea->DeleteKey(SelectedKey.KeyHandle.GetValue());
+			Key.KeyArea->DeleteKey(Key.KeyHandle.GetValue());
 		}
 	}
 }
@@ -584,7 +561,7 @@ UK2Node_PlayMovieScene* FSequencer::BindToPlayMovieSceneNode( const bool bCreate
 							* @param	Sequencer			The Sequencer we're using
 							* @param	PlayMovieSceneNode	The node that we need to jump to
 							*/
-						static void NavigateToPlayMovieSceneNode( ISequencerInternals* Sequencer, UK2Node_PlayMovieScene* PlayMovieSceneNode )
+						static void NavigateToPlayMovieSceneNode( FSequencer* Sequencer, UK2Node_PlayMovieScene* PlayMovieSceneNode )
 						{
 							check( Sequencer != NULL );
 							check( PlayMovieSceneNode != NULL );
@@ -612,7 +589,7 @@ UK2Node_PlayMovieScene* FSequencer::BindToPlayMovieSceneNode( const bool bCreate
 						}
 					};
 
-					ISequencerInternals* SequencerInternals = this;
+					FSequencer* SequencerInternals = this;
 					Info.Hyperlink = FSimpleDelegate::CreateStatic( &Local::NavigateToPlayMovieSceneNode, SequencerInternals, FoundPlayMovieSceneNode );
 					Info.HyperlinkText = LOCTEXT("AddedPlayMovieSceneEventToLevelScriptGraph_Hyperlink", "Show Graph");
 
@@ -1741,7 +1718,7 @@ bool FSequencer::IsSectionVisible(UMovieSceneSection* Section) const
 
 void FSequencer::DeleteSelectedItems()
 {
-	DeleteKeys(SelectedKeys.Array());
+	DeleteSelectedKeys();
 
 	for (int32 i = 0; i < SelectedSections.Num(); ++i)
 	{
diff --git a/Engine/Source/Editor/Sequencer/Private/Sequencer.h b/Engine/Source/Editor/Sequencer/Private/Sequencer.h
--- a/Engine/Source/Editor/Sequencer/Private/Sequencer.h
+++ b/Engine/Source/Editor/Sequencer/Private/Sequencer.h
@@ -2,19 +2,17 @@
 
 #pragma once
 
-#include "ISequencerInternals.h"
 #include "Toolkits/AssetEditorToolkit.h"
 #include "Editor/LevelEditor/Public/LevelEditor.h"	// For EMapChangeType::Type
 #include "Editor/EditorWidgets/Public/ITransportControl.h"
+#include "SelectedKey.h"
 
 class UMovieScene;
 
-
 /**
  * Sequencer is the editing tool for MovieScene assets
  */
-class FSequencer :
-	public ISequencerInternals, FGCObject, public FEditorUndoClient
+class FSequencer : public ISequencer, FGCObject, public FEditorUndoClient
 { 
 
 public:
@@ -38,66 +36,226 @@ public:
 	/** Destructor */
 	virtual ~FSequencer();
 
+	/** FGCObject interface */
+	virtual void AddReferencedObjects(FReferenceCollector& Collector) OVERRIDE;
+
 	/** IToolkit interface */
 	virtual FName GetToolkitFName() const OVERRIDE;
 	virtual FText GetBaseToolkitName() const OVERRIDE;
 	virtual FString GetWorldCentricTabPrefix() const OVERRIDE;
 	virtual FLinearColor GetWorldCentricTabColorScale() const OVERRIDE;
 
-	/** ISequencerInternals interface */
-	virtual TArray< UMovieScene* > GetMovieScenesBeingEdited() OVERRIDE;
-	virtual UMovieScene* GetRootMovieScene() const OVERRIDE;
+	/** ISequencer interface */
+	virtual TArray< UMovieScene* > GetMovieScenesBeingEdited();
+	virtual UMovieScene* GetRootMovieScene() const;
 	virtual UMovieScene* GetFocusedMovieScene() const OVERRIDE;
 	virtual TSharedRef<FMovieSceneInstance> GetRootMovieSceneInstance() const OVERRIDE;
 	virtual TSharedRef<FMovieSceneInstance> GetFocusedMovieSceneInstance() const OVERRIDE;
 	virtual void FocusSubMovieScene( TSharedRef<FMovieSceneInstance> SubMovieSceneInstance ) OVERRIDE;
 	TSharedRef<FMovieSceneInstance> GetInstanceForSubMovieSceneSection( UMovieSceneSection& SubMovieSceneSection ) const OVERRIDE;
-	virtual void PopToMovieScene( TSharedRef<FMovieSceneInstance> SubMovieSceneInstance ) OVERRIDE;
-	virtual void EditDetailsForObjects( TArray< UObject* > ObjectsToEdit ) OVERRIDE;
-	virtual void SpawnOrDestroyPuppetObjects( TSharedRef<FMovieSceneInstance> MovieSceneInstance ) OVERRIDE;
 	virtual void AddNewShot(FGuid CameraGuid) OVERRIDE;
-	virtual void AddAnimation( FGuid ObjectGuid, class UAnimSequence* AnimSequence ) OVERRIDE;
-	virtual void RenameShot(class UMovieSceneSection* ShotSection) OVERRIDE;
-	virtual void DeleteSection(class UMovieSceneSection* Section) OVERRIDE;
-	virtual void DeleteKeys(TArray<FSelectedKey> KeysToDelete) OVERRIDE;
-	virtual class UK2Node_PlayMovieScene* BindToPlayMovieSceneNode( const bool bCreateIfNotFound ) OVERRIDE;
+	virtual void AddAnimation(FGuid ObjectGuid, class UAnimSequence* AnimSequence) OVERRIDE;
 	virtual bool IsAutoKeyEnabled() const OVERRIDE;
 	virtual bool IsRecordingLive() const OVERRIDE;
-	virtual float GetCurrentLocalTime( UMovieScene& MovieScene ) OVERRIDE;
+	virtual float GetCurrentLocalTime(UMovieScene& MovieScene) OVERRIDE;
 	virtual float GetGlobalTime() OVERRIDE;
 	virtual void SetGlobalTime(float Time) OVERRIDE;
 	virtual void SetPerspectiveViewportPossessionEnabled(bool bEnabled) OVERRIDE;
-	virtual FGuid GetHandleToObject( UObject* Object ) OVERRIDE;
-	virtual const TArray< TSharedPtr<FMovieSceneTrackEditor> >& GetTrackEditors() const OVERRIDE { return TrackEditors; }
+	virtual FGuid GetHandleToObject(UObject* Object) OVERRIDE;
 	virtual ISequencerObjectChangeListener& GetObjectChangeListener() const OVERRIDE;
 	virtual void NotifyMovieSceneDataChanged() OVERRIDE;
-	virtual void Tick( const float InDeltaTime ) OVERRIDE;
-	virtual FGuid AddSpawnableForAssetOrClass( UObject* Object, UObject* CounterpartGamePreviewObject ) OVERRIDE;
-	virtual void AddSubMovieScene( UMovieScene* SubMovieScene ) OVERRIDE;
-	virtual bool OnHandleAssetDropped(UObject* DroppedAsset, const FGuid& TargetObjectGuid) OVERRIDE;
-	virtual void OnRequestNodeDeleted( TSharedRef<const FSequencerDisplayNode>& NodeToBeDeleted ) OVERRIDE;
-	virtual void CopyActorProperties( AActor* PuppetActor, AActor* TargetActor ) const OVERRIDE;
-	virtual TArray< TWeakObjectPtr<UMovieSceneSection> > GetSelectedSections() const OVERRIDE;
-	virtual void SelectSection(UMovieSceneSection* Section) OVERRIDE;
-	virtual bool IsSectionSelected( UMovieSceneSection* Section) const OVERRIDE;
-	virtual void ClearSectionSelection() OVERRIDE;
-	virtual void ZoomToSelectedSections() OVERRIDE;
-	virtual void SelectKey( const FSelectedKey& Key ) OVERRIDE;
-	virtual void ClearKeySelection() OVERRIDE;
-	virtual bool IsKeySelected( const FSelectedKey& Key ) const OVERRIDE;
-	virtual TSet<FSelectedKey>& GetSelectedKeys() OVERRIDE;
-	virtual void FilterToShotSections( const TArray< TWeakObjectPtr<class UMovieSceneSection> >& ShotSections, bool bZoomToShotBounds ) OVERRIDE;
+	virtual void AddSubMovieScene(UMovieScene* SubMovieScene) OVERRIDE;
+	virtual void FilterToShotSections(const TArray< TWeakObjectPtr<class UMovieSceneSection> >& ShotSections, bool bZoomToShotBounds = true) OVERRIDE;
 	virtual void FilterToSelectedShotSections(bool bZoomToShotBounds = true) OVERRIDE;
-	virtual TArray< TWeakObjectPtr<UMovieSceneSection> > GetFilteringShotSections() const OVERRIDE;
-	virtual bool IsObjectUnfilterable(const FGuid& ObjectGuid) const OVERRIDE;
-	virtual void AddUnfilterableObject(const FGuid& ObjectGuid) OVERRIDE;
-	virtual bool IsShotFilteringOn() const OVERRIDE;
-	virtual bool IsUsingCleanView() const OVERRIDE;
-	virtual float GetOverlayFadeCurve() const OVERRIDE;
-	virtual bool IsSectionVisible(UMovieSceneSection* Section) const OVERRIDE;
-	virtual FUICommandList& GetCommandBindings() OVERRIDE {return *SequencerCommandBindings;}
-	virtual void BuildObjectBindingContextMenu(class FMenuBuilder& MenuBuilder, const FGuid& ObjectBinding, const class UClass* ObjectClass) OVERRIDE;
 
+	/**
+	 * Pops the current focused movie scene from the stack.  The parent of this movie scene will be come the focused one
+	 */
+	void PopToMovieScene( TSharedRef<FMovieSceneInstance> SubMovieSceneInstance );
+
+	/**
+	 * Starts editing details for the specified list of UObjects.  This is used to edit CDO's for blueprints
+	 * that are stored with the MovieScene asset
+	 *
+	 * @param	ObjectsToEdit	The objects to edit in the details view
+	 */
+	void EditDetailsForObjects( TArray< UObject* > ObjectsToEdit );
+
+	/**
+	 * Spawn (or destroy) puppet objects as needed to match the spawnables array in the MovieScene we're editing
+	 *
+	 * @param MovieSceneInstance	The movie scene instance to spawn or destroy puppet objects for
+	 */
+	void SpawnOrDestroyPuppetObjects( TSharedRef<FMovieSceneInstance> MovieSceneInstance );
+
+	/** 
+	 * Opens a renaming dialog for the passed in shot section
+	 */
+	void RenameShot(class UMovieSceneSection* ShotSection);
+	 
+	/** 
+	 * Deletes the passed in section
+	 */
+	void DeleteSection(class UMovieSceneSection* Section);
+
+	/**
+	 * Deletes the currently selected in keys
+	 */
+	void DeleteSelectedKeys();
+
+	/**
+	 * Binds to a 'PlayMovieScene' node in level script that is associated with the MovieScene we're editing.  If
+	 * requested, can also create a new node.  Only valid to call when in world-centric mode.
+	 *
+	 * @param	bCreateIfNotFound	True to create a node if we can't find one
+	 *
+	 * @return	The LevelScript node we found or created, or NULL if nothing is still bound
+	 */
+	class UK2Node_PlayMovieScene* BindToPlayMovieSceneNode( const bool bCreateIfNotFound );
+
+	/**
+	 * @return Movie scene tools used by the sequencer
+	 */
+	const TArray< TSharedPtr<FMovieSceneTrackEditor> >& GetTrackEditors() const { return TrackEditors; }
+
+	/** Ticks the sequencer by InDeltaTime */
+	void Tick( const float InDeltaTime );
+
+	/**
+	 * Attempts to add a new spawnable to the MovieScene for the specified asset or class
+	 *
+	 * @param	Object	The asset or class object to add a spawnable for
+	 *
+	 * @return	The spawnable guid for the spawnable, or an invalid Guid if we were not able to create a spawnable
+	 */
+	virtual FGuid AddSpawnableForAssetOrClass( UObject* Object, UObject* CounterpartGamePreviewObject ) ;
+
+	/**
+	 * Call when an asset is dropped into the sequencer. Will proprogate this
+	 * to all tracks, and potentially consume this asset
+	 * so it won't be added as a spawnable
+	 *
+	 * @param DroppedAsset		The asset that is dropped in
+	 * @param TargetObjectGuid	Object to be targeted on dropping
+	 * @return					If true, this asset should be consumed
+	 */
+	virtual bool OnHandleAssetDropped( UObject* DroppedAsset, const FGuid& TargetObjectGuid );
+	
+	/**
+	 * Called to delete all moviescene data from a node
+	 *
+	 * @param NodeToBeDeleted	Node with data that should be deleted
+	 */
+	virtual void OnRequestNodeDeleted( TSharedRef<const FSequencerDisplayNode>& NodeToBeDeleted );
+
+	/**
+	 * Copies properties from one actor to another.  Only properties that are different are copied.  This is used
+	 * to propagate changed properties from a living actor in the scene to properties of a blueprint CDO actor
+	 *
+	 * @param	PuppetActor	The source actor (the puppet actor in the scene we're copying properties from)
+	 * @param	TargetActor	The destination actor (data-only blueprint actor CDO)
+	 */
+	virtual void CopyActorProperties( AActor* PuppetActor, AActor* TargetActor ) const;
+
+	/**
+	 * Gets all selected sections in the sequencer
+	 */
+	virtual TArray< TWeakObjectPtr<UMovieSceneSection> > GetSelectedSections() const;
+
+	/**
+	 * Selects a section in the sequencer
+	 */
+	void SelectSection(UMovieSceneSection* Section);
+
+	/**
+	 * Returns whether or not a section is selected
+	 *
+	 * @param Section The section to check
+	 * @return true if the section is selected
+	 */
+	bool IsSectionSelected( UMovieSceneSection* Section) const;
+
+	/**
+	 * Clears all selected sections
+	 */
+	void ClearSectionSelection();
+
+	/**
+	 * Zooms to the edges of all currently selected sections
+	 */
+	void ZoomToSelectedSections();
+
+	/**
+	 * Selects a key
+	 * 
+	 * @param Key	Representation of the key to select
+	 */
+	void SelectKey( const FSelectedKey& Key );
+
+	/**
+	 * Clears the entire set of selected keys
+	 */
+	void ClearKeySelection();
+
+	/**
+	 * Returns whether or not a key is selected
+	 * 
+	 * @param Key	Representation of the key to check for selection
+	 */
+	bool IsKeySelected( const FSelectedKey& Key ) const;
+
+	/**
+	 * @return The entire set of selected keys
+	 */
+	TSet<FSelectedKey>& GetSelectedKeys();
+
+	/**
+	 * Gets all shots that are filtering currently
+	 */
+	TArray< TWeakObjectPtr<UMovieSceneSection> > GetFilteringShotSections() const;
+
+	/**
+	 * Checks to see if an object is unfilterable
+	 */
+	bool IsObjectUnfilterable(const FGuid& ObjectGuid) const;
+
+	/**
+	 * Declares an object unfilterable
+	 */
+	void AddUnfilterableObject(const FGuid& ObjectGuid);
+
+	/**
+	 * Checks to see if shot filtering is on
+	 */
+	bool IsShotFilteringOn() const;
+
+	/**
+	 * Returns true if the sequencer is using the 'Clean View' mode
+	 * Clean View simply means no non-global tracks will appear if no shots are filtering
+	 */
+	bool IsUsingCleanView() const;
+
+	/**
+	 * Gets the overlay fading animation curve lerp
+	 */
+	float GetOverlayFadeCurve() const;
+
+	/**
+	 * Checks if a section is visible, given current shot filtering
+	 */
+	bool IsSectionVisible(UMovieSceneSection* Section) const;
+
+	/** Gets the command bindings for the sequencer */
+	FUICommandList& GetCommandBindings() { return *SequencerCommandBindings; }
+
+	/**
+	 * Builds up the context menu for object binding nodes in the outliner
+	 * 
+	 * @param MenuBuilder	The context menu builder to add things to
+	 * @param ObjectBinding	The object binding of the selected node
+	 * @param ObjectClass	The class of the selected object
+	 */
+	void BuildObjectBindingContextMenu(class FMenuBuilder& MenuBuilder, const FGuid& ObjectBinding, const class UClass* ObjectClass);
 
 	/** IMovieScenePlayer interface */
 	virtual void GetRuntimeObjects( TSharedRef<FMovieSceneInstance> MovieSceneInstance, const FGuid& ObjectHandle, TArray< UObject* >& OutObjects ) const;
@@ -108,8 +266,7 @@ public:
 
 	virtual void SpawnActorsForMovie( TSharedRef<FMovieSceneInstance> MovieSceneInstance );
 
-	/** FGCObject interface */
-	virtual void AddReferencedObjects( FReferenceCollector& Collector ) OVERRIDE;
+
 
 	// @todo remove when world-centric mode is added
 	TSharedRef<class SSequencer> GetMainSequencer();
diff --git a/Engine/Source/Editor/Sequencer/Private/SequencerDisplayNode.cpp b/Engine/Source/Editor/Sequencer/Private/SequencerDisplayNode.cpp
--- a/Engine/Source/Editor/Sequencer/Private/SequencerDisplayNode.cpp
+++ b/Engine/Source/Editor/Sequencer/Private/SequencerDisplayNode.cpp
@@ -3,7 +3,7 @@
 #include "SequencerPrivatePCH.h"
 #include "SectionLayoutBuilder.h"
 #include "ISequencerSection.h"
-#include "ISequencerInternals.h"
+#include "Sequencer.h"
 #include "MovieScene.h"
 #include "SSequencer.h"
 #include "SSequencerSectionAreaView.h"
@@ -213,7 +213,7 @@ void FSequencerDisplayNode::AddKeyAreaNode( FName KeyAreaName, const FString& Di
 	KeyAreaNode->AddKeyArea( KeyArea );
 }
 
-TSharedRef<SWidget> FSequencerDisplayNode::GenerateWidgetForOutliner( TSharedRef<class ISequencerInternals> Sequencer )
+TSharedRef<SWidget> FSequencerDisplayNode::GenerateWidgetForOutliner( TSharedRef<class FSequencer> Sequencer )
 {
 	return SNew( SAnimationOutlinerView, SharedThis( this ), Sequencer );
 }
@@ -459,7 +459,7 @@ bool FObjectBindingNode::GetShotFilteredVisibilityToCache() const
 
 TSharedPtr<SWidget> FObjectBindingNode::OnSummonContextMenu(const FGeometry& MyGeometry, const FPointerEvent& MouseEvent)
 {
-	ISequencerInternals& ParentSequencer = GetSequencer();
+	FSequencer& ParentSequencer = GetSequencer();
 
 	UMovieScene* MovieScene = GetSequencer().GetFocusedMovieScene();
 
diff --git a/Engine/Source/Editor/Sequencer/Private/SequencerDisplayNode.h b/Engine/Source/Editor/Sequencer/Private/SequencerDisplayNode.h
--- a/Engine/Source/Editor/Sequencer/Private/SequencerDisplayNode.h
+++ b/Engine/Source/Editor/Sequencer/Private/SequencerDisplayNode.h
@@ -85,7 +85,7 @@ public:
 	 * @param Sequencer	Sequencer interface to pass to outliner widgets
 	 * @return Generated outliner widget
 	 */
-	virtual TSharedRef<SWidget> GenerateWidgetForOutliner( TSharedRef<class ISequencerInternals> Sequencer );
+	virtual TSharedRef<SWidget> GenerateWidgetForOutliner( TSharedRef<class FSequencer> Sequencer );
 
 	/**
 	 * Generates a widget for display in the section area portion of the track area
@@ -129,7 +129,7 @@ public:
 	TSharedPtr<FSequencerDisplayNode> GetParent() const { return ParentNode.Pin(); }
 	
 	/** Gets the sequencer that owns this node */
-	ISequencerInternals& GetSequencer() const { return ParentTree.GetSequencer(); }
+	FSequencer& GetSequencer() const { return ParentTree.GetSequencer(); }
 	
 	/** Gets all the key area nodes recursively, including this node if applicable */
 	virtual void GetChildKeyAreaNodesRecursively(TArray< TSharedRef<class FSectionKeyAreaNode> >& OutNodes) const;
diff --git a/Engine/Source/Editor/Sequencer/Private/SequencerNodeTree.cpp b/Engine/Source/Editor/Sequencer/Private/SequencerNodeTree.cpp
--- a/Engine/Source/Editor/Sequencer/Private/SequencerNodeTree.cpp
+++ b/Engine/Source/Editor/Sequencer/Private/SequencerNodeTree.cpp
@@ -3,7 +3,7 @@
 #include "SequencerPrivatePCH.h"
 #include "MovieSceneSection.h"
 #include "SequencerNodeTree.h"
-#include "ISequencerInternals.h"
+#include "Sequencer.h"
 #include "ScopedTransaction.h"
 #include "MovieScene.h"
 #include "MovieSceneTrack.h"
diff --git a/Engine/Source/Editor/Sequencer/Private/SequencerNodeTree.h b/Engine/Source/Editor/Sequencer/Private/SequencerNodeTree.h
--- a/Engine/Source/Editor/Sequencer/Private/SequencerNodeTree.h
+++ b/Engine/Source/Editor/Sequencer/Private/SequencerNodeTree.h
@@ -13,7 +13,7 @@ class UMovieSceneTrack;
 class FSequencerNodeTree : public TSharedFromThis<FSequencerNodeTree>
 {
 public:
-	FSequencerNodeTree( class ISequencerInternals& InSequencer )
+	FSequencerNodeTree( class FSequencer& InSequencer )
 		: Sequencer( InSequencer )
 	{}
 	
@@ -72,7 +72,7 @@ public:
 	void FilterNodes( const FString& InFilter );
 
 	/** Gets the parent sequencer of this tree */
-	ISequencerInternals& GetSequencer() {return Sequencer;}
+	FSequencer& GetSequencer() {return Sequencer;}
 
 	/**
 	 * Saves the expansion state of a display node
@@ -129,5 +129,5 @@ private:
 	/** Active filter string if any */
 	FString FilterString;
 	/** Sequencer interface */
-	ISequencerInternals& Sequencer;
+	FSequencer& Sequencer;
 };
diff --git a/Engine/Source/Editor/Sequencer/Private/SequencerObjectSpawner.cpp b/Engine/Source/Editor/Sequencer/Private/SequencerObjectSpawner.cpp
--- a/Engine/Source/Editor/Sequencer/Private/SequencerObjectSpawner.cpp
+++ b/Engine/Source/Editor/Sequencer/Private/SequencerObjectSpawner.cpp
@@ -2,7 +2,7 @@
 
 #include "SequencerPrivatePCH.h"
 #include "SequencerObjectSpawner.h"
-#include "ISequencerInternals.h"
+#include "Sequencer.h"
 #include "MovieScene.h"
 #include "Toolkits/IToolkitHost.h"
 #include "Kismet2/BlueprintEditorUtils.h"
@@ -11,7 +11,7 @@
 
 #define LOCTEXT_NAMESPACE "Sequencer"
 
-FSequencerActorObjectSpawner::FSequencerActorObjectSpawner( ISequencerInternals& InitSequencer )
+FSequencerActorObjectSpawner::FSequencerActorObjectSpawner( FSequencer& InitSequencer )
 	: Sequencer( &InitSequencer )
 {
 	// Register to be notified when object changes should be propagated
diff --git a/Engine/Source/Editor/Sequencer/Private/SequencerObjectSpawner.h b/Engine/Source/Editor/Sequencer/Private/SequencerObjectSpawner.h
--- a/Engine/Source/Editor/Sequencer/Private/SequencerObjectSpawner.h
+++ b/Engine/Source/Editor/Sequencer/Private/SequencerObjectSpawner.h
@@ -109,7 +109,7 @@ public:
 	 *
 	 * @param	InitSequencer	Sequencer that owns this object
 	 */
-	FSequencerActorObjectSpawner( class ISequencerInternals& InitSequencer );
+	FSequencerActorObjectSpawner( class FSequencer& InitSequencer );
 
 
 	/** Destructor */
@@ -147,7 +147,7 @@ protected:
 private:
 
 	/** The Sequencer that owns this spawner  */
-	ISequencerInternals* Sequencer;
+	FSequencer* Sequencer;
 
 	/** 
 	 * List of puppet objects that we've spawned and are actively managing. 
diff --git a/Engine/Source/Editor/Sequencer/Private/SequencerPrivatePCH.h b/Engine/Source/Editor/Sequencer/Private/SequencerPrivatePCH.h
--- a/Engine/Source/Editor/Sequencer/Private/SequencerPrivatePCH.h
+++ b/Engine/Source/Editor/Sequencer/Private/SequencerPrivatePCH.h
@@ -13,8 +13,8 @@
 #include "SequencerNodeTree.h"
 #include "SequencerDisplayNode.h"
 #include "ISectionLayoutBuilder.h"
+#include "SelectedKey.h"
 #include "ITimeSlider.h"
-#include "ISequencerInternals.h"
 #include "SSequencer.h"
 #include "SSection.h"
 #include "SSequencerTrackViewPanel.h"
diff --git a/Engine/Source/Editor/Sequencer/Public/ISequencerInternals.h b/Engine/Source/Editor/Sequencer/Public/ISequencerInternals.h
deleted file mode 100644
