Basic hookup for Sequencer in UMG.  Not actually working yet.

- Broke out the asset editor portion of ISequencer into another class as it doesnt make sense to use an asset editor for an embedded sequencer
- Some minor cleanup to make sequencer work outside of the level editor (more work to do)

diff --git a/Engine/Source/Developer/MovieSceneTools/Private/AssetTypeActions_MovieScene.cpp b/Engine/Source/Developer/MovieSceneTools/Private/AssetTypeActions_MovieScene.cpp
--- a/Engine/Source/Developer/MovieSceneTools/Private/AssetTypeActions_MovieScene.cpp
+++ b/Engine/Source/Developer/MovieSceneTools/Private/AssetTypeActions_MovieScene.cpp
@@ -38,13 +38,13 @@ void FAssetTypeActions_MovieScene::OpenAssetEditor( const TArray<UObject*>& InOb
 
 	for (auto ObjIt = InObjects.CreateConstIterator(); ObjIt; ++ObjIt)
 	{
-		auto MovieScene = Cast<UMovieScene>(*ObjIt);
+		UMovieScene* MovieScene = Cast<UMovieScene>(*ObjIt);
 		if (MovieScene != NULL)
 		{
 			// @todo sequencer: Only allow users to create new MovieScenes if that feature is turned on globally.
 			if( FParse::Param( FCommandLine::Get(), TEXT( "Sequencer" ) ) )
 			{
-				auto NewSequencer = FModuleManager::LoadModuleChecked< ISequencerModule >( "Sequencer" ).CreateSequencer( Mode, EditWithinLevelEditor, MovieScene );
+				FModuleManager::LoadModuleChecked< ISequencerModule >( "Sequencer" ).CreateSequencerAssetEditor( Mode, EditWithinLevelEditor, MovieScene, true );
 			}
 		}
 	}
diff --git a/Engine/Source/Editor/Sequencer/Private/Sequencer.cpp b/Engine/Source/Editor/Sequencer/Private/Sequencer.cpp
--- a/Engine/Source/Editor/Sequencer/Private/Sequencer.cpp
+++ b/Engine/Source/Editor/Sequencer/Private/Sequencer.cpp
@@ -3,7 +3,6 @@
 #include "SequencerPrivatePCH.h"
 #include "Sequencer.h"
 #include "SequencerObjectSpawner.h"
-#include "Toolkits/IToolkitHost.h"
 #include "MovieScene.h"
 #include "Engine/LevelScriptBlueprint.h"
 #include "Editor/PropertyEditor/Public/PropertyEditorModule.h"
@@ -25,7 +24,6 @@
 #include "MovieSceneAudioTrack.h"
 #include "MovieSceneAnimationTrack.h"
 #include "MovieSceneTrackEditor.h"
-#include "Toolkits/IToolkitHost.h"
 #include "Kismet2/KismetEditorUtilities.h"
 #include "Kismet2/BlueprintEditorUtils.h"
 #include "K2Node_PlayMovieScene.h"
@@ -42,67 +40,27 @@
 
 DEFINE_LOG_CATEGORY_STATIC( LogSequencer, Log, All );
 
-const FName FSequencer::SequencerMainTabId( TEXT( "Sequencer_SequencerMain" ) );
-const FName FSequencer::SequencerDetailsTabId( TEXT( "Sequencer_Details" ) );
-
-namespace SequencerDefs
-{
-	static const FName SequencerAppIdentifier( TEXT( "SequencerApp" ) );
-}
-void FSequencer::RegisterTabSpawners(const TSharedRef<class FTabManager>& TabManager)
+static bool IsSequencerEnabled()
 {
-	const IWorkspaceMenuStructure& MenuStructure = WorkspaceMenu::GetMenuStructure();
-
-	if( FParse::Param( FCommandLine::Get(), TEXT( "Sequencer" ) ) && !IsWorldCentricAssetEditor() )
-	{
-		TabManager->RegisterTabSpawner( SequencerMainTabId, FOnSpawnTab::CreateSP(this, &FSequencer::SpawnTab_SequencerMain) )
-			.SetDisplayName( LOCTEXT("SequencerMainTab", "Sequencer") )
-			.SetGroup( MenuStructure.GetAssetEditorCategory() );
-	}
-
+	return FParse::Param( FCommandLine::Get(), TEXT( "Sequencer" ) ) || FParse::Param( FCommandLine::Get(), TEXT( "umg" ) );
 }
 
-void FSequencer::UnregisterTabSpawners(const TSharedRef<class FTabManager>& TabManager)
-{
-	if( FParse::Param( FCommandLine::Get(), TEXT( "Sequencer" ) ) && !IsWorldCentricAssetEditor() )
-	{
-		TabManager->UnregisterTabSpawner( SequencerMainTabId );
-	}
-	
-	// @todo remove when world-centric mode is added
-	FLevelEditorModule& LevelEditorModule = FModuleManager::LoadModuleChecked<FLevelEditorModule>("LevelEditor");
-	LevelEditorModule.AttachSequencer( SNullWidget::NullWidget, nullptr );
-}
 
-void FSequencer::InitSequencer( const EToolkitMode::Type Mode, const TSharedPtr< class IToolkitHost >& InitToolkitHost, UObject* ObjectToEdit, const TArray<FOnCreateTrackEditor>& TrackEditorDelegates )
+void FSequencer::InitSequencer( UMovieScene* InRootMovieScene, TSharedPtr<IToolkitHost> InToolKitHost, const TArray<FOnCreateTrackEditor>& TrackEditorDelegates, bool bEditWithinLevelEditor )
 {
-	if( FParse::Param( FCommandLine::Get(), TEXT( "Sequencer" ) ) )
+	if( IsSequencerEnabled() )
 	{
-		const TSharedRef<FTabManager::FLayout> StandaloneDefaultLayout = FTabManager::NewLayout( "Standalone_Sequencer_Layout" )
-			->AddArea
-			(
-				FTabManager::NewPrimaryArea()
-				->Split
-				(
-					FTabManager::NewStack()
-					->AddTab(SequencerMainTabId, ETabState::OpenedTab)
-					->AddTab(SequencerDetailsTabId, ETabState::OpenedTab)
-				)
-			);
-
-		const bool bCreateDefaultStandaloneMenu = true;
-		const bool bCreateDefaultToolbar = false;
-		FAssetEditorToolkit::InitAssetEditor( Mode, InitToolkitHost, SequencerDefs::SequencerAppIdentifier, StandaloneDefaultLayout, bCreateDefaultStandaloneMenu,bCreateDefaultToolbar, ObjectToEdit );
-
+		ToolkitHost = InToolKitHost;
 		// Create an object change listener for various systems that need to be notified when objects change
 		ObjectChangeListener = MakeShareable( new FSequencerObjectChangeListener( SharedThis( this ) ) );
 		// If this is a world-centric editor, then setup our puppet spawner
-		if( IsWorldCentricAssetEditor() )
+		if( bEditWithinLevelEditor )
 		{
 			ObjectSpawner = MakeShareable( new FSequencerActorObjectSpawner( *this ) );
 		}
 
-		UMovieScene& RootMovieScene = *CastChecked<UMovieScene>( GetEditingObject() );
+		UMovieScene& RootMovieScene = *InRootMovieScene;
+
 		// Focusing the initial movie scene needs to be done before the first time GetFocusedMovieSceneInstane or GetRootMovieSceneInstance is used
 		RootMovieSceneInstance = MakeShareable( new FMovieSceneInstance( RootMovieScene ) );
 		MovieSceneStack.Add( RootMovieSceneInstance.ToSharedRef() );
@@ -119,27 +77,20 @@ void FSequencer::InitSequencer( const EToolkitMode::Type Mode, const TSharedPtr<
 			.OnToggleCleanView( this, &FSequencer::OnToggleCleanView );
 
 
-		// @todo remove when world-centric mode is added
-		FLevelEditorModule& LevelEditorModule = FModuleManager::LoadModuleChecked<FLevelEditorModule>("LevelEditor");
-		LevelEditorModule.AttachSequencer(SequencerWidget, SharedThis( this ) );
-
-		// Hook into the editor's mechanism for checking whether we need live capture of PIE/SIE actor state
-		GEditor->GetActorRecordingState().AddSP( this, &FSequencer::GetActorRecordingState );
-
-		// When undo occurs, get a notification so we can make sure our view is up to date
-		GEditor->RegisterForUndo(this);
-
-
-		// We need to find out when the user loads a new map, because we might need to re-create puppet actors
-		// when previewing a MovieScene
-		//auto& LevelEditorModule = FModuleManager::LoadModuleChecked< FLevelEditorModule >( TEXT( "LevelEditor" ) );
-		LevelEditorModule.OnMapChanged().AddSP( this, &FSequencer::OnMapChanged );
+		if( bEditWithinLevelEditor )
+		{
+			// @todo remove when world-centric mode is added
+			// Hook into the editor's mechanism for checking whether we need live capture of PIE/SIE actor state
+			GEditor->GetActorRecordingState().AddSP(this, &FSequencer::GetActorRecordingState);
 
+			// When undo occurs, get a notification so we can make sure our view is up to date
+			GEditor->RegisterForUndo(this);
 
-		// Start listening for important changes on this MovieScene
-		// @todo sequencer: We need to correctly register for any new moviescenes that are dropped in!
-		const bool bCreateIfNotFound = false;
-		BindToPlayMovieSceneNode( bCreateIfNotFound );
+			// Start listening for important changes on this MovieScene
+			// @todo sequencer: We need to correctly register for any new moviescenes that are dropped in!
+			const bool bCreateIfNotFound = false;
+			BindToPlayMovieSceneNode(bCreateIfNotFound);
+		}
 
 		// Create tools and bind them to this sequencer
 		for( int32 DelegateIndex = 0; DelegateIndex < TrackEditorDelegates.Num(); ++DelegateIndex )
@@ -161,7 +112,7 @@ void FSequencer::InitSequencer( const EToolkitMode::Type Mode, const TSharedPtr<
 			TrackEditors.Add( TrackEditor );
 		}
 
-		AttachTransportControlsToViewports();
+
 
 		ZoomAnimation = FCurveSequence();
 		ZoomCurve = ZoomAnimation.AddCurve(0.f, 0.35f, ECurveEaseFunction::QuadIn);
@@ -195,16 +146,6 @@ FSequencer::FSequencer()
 
 FSequencer::~FSequencer()
 {
-	DetachTransportControlsFromViewports();
-
-
-	// Unregister delegates
-	if( FModuleManager::Get().IsModuleLoaded( TEXT( "LevelEditor" ) ) )
-	{
-		auto& LevelEditorModule = FModuleManager::LoadModuleChecked< FLevelEditorModule >( TEXT( "LevelEditor" ) );
-		LevelEditorModule.OnMapChanged().RemoveAll( this );
-	}
-
 	if( PlayMovieSceneNode.IsValid() )
 	{
 		PlayMovieSceneNode->OnBindingsChanged().RemoveAll( this );
@@ -225,91 +166,10 @@ FSequencer::~FSequencer()
 
 	TrackEditors.Empty();
 
-	DetailsView.Reset();
-
-	/*if( SequencerMainTab.IsValid() )
-	{
-		// Kill the tab!
-		// NOTE: It's possible that the user already closed the tab manually, but that's OK.
-		SequencerMainTab->RemoveTabFromParent();
-		SequencerMainTab.Reset();
-	}
-
-	if( DetailsTab.IsValid() )
-	{
-		// Kill the tab!
-		// NOTE: It's possible that the user already closed the tab manually, but that's OK.
-		DetailsTab->RemoveTabFromParent();
-		DetailsTab.Reset();
-	}*/
-
 	SequencerWidget.Reset();
 }
 
 
-TSharedRef<SDockTab> FSequencer::SpawnTab_SequencerMain(const FSpawnTabArgs& Args)
-{
-	check(Args.GetTabId() == SequencerMainTabId);
-
-	return SNew(SDockTab)
-		.Icon( FEditorStyle::GetBrush("Sequencer.Tabs.SequencerMain") )
-		.Label( LOCTEXT("SequencerMainTitle", "Sequencer") )
-		.TabColorScale( GetTabColorScale() )
-		[
-			SequencerWidget.ToSharedRef()
-		];
-}
-
-TSharedRef<SDockTab> FSequencer::SpawnTab_Details(const FSpawnTabArgs& Args)
-{
-	check(Args.GetTabId() == SequencerDetailsTabId);
-
-	const bool bIsUpdatable = false;
-	const bool bAllowFavorites = true;
-	const bool bIsLockable = false;
-
-	FPropertyEditorModule& PropertyEditorModule = FModuleManager::GetModuleChecked<FPropertyEditorModule>( "PropertyEditor" );
-	const FDetailsViewArgs DetailsViewArgs( bIsUpdatable, bIsLockable, true, false, false );
-	DetailsView = PropertyEditorModule.CreateDetailView( DetailsViewArgs );
-
-	return SNew(SDockTab)
-		.Icon( FEditorStyle::GetBrush("Sequencer.Tabs.Details") )	// @todo sequencer: Missing icon
-		.Label( LOCTEXT("SequencerDetailsTitle", "Details") )
-		.TabColorScale( GetTabColorScale() )
-		[
-			SNew(SBorder)
-			.Padding(4)
-			.BorderImage( FEditorStyle::GetBrush( "ToolPanel.GroupBorder" ) )
-			[
-				DetailsView.ToSharedRef()
-			]
-		]
-	;
-}
-
-FName FSequencer::GetToolkitFName() const
-{
-	return FName("Sequencer");
-}
-
-FText FSequencer::GetBaseToolkitName() const
-{
-	return LOCTEXT("AppLabel", "Sequencer");
-}
-
-
-FLinearColor FSequencer::GetWorldCentricTabColorScale() const
-{
-	return FLinearColor( 0.7, 0.0f, 0.0f, 0.5f );
-}
-
-
-FString FSequencer::GetWorldCentricTabPrefix() const
-{
-	return LOCTEXT("WorldCentricTabPrefix", "Sequencer ").ToString();
-}
-
-
 TArray< UMovieScene* > FSequencer::GetMovieScenesBeingEdited()
 {
 	TArray<UMovieScene* > OutMovieScenes;
@@ -388,15 +248,12 @@ void FSequencer::PopToMovieScene( TSharedRef<FMovieSceneInstance> SubMovieSceneI
 	}
 }
 
-void FSequencer::EditDetailsForObjects( TArray< UObject* > ObjectsToEdit )
-{
-	// @todo re-enable when world centric is back up
-	//DetailsView->SetObjects( ObjectsToEdit );
-}
-
 void FSequencer::SpawnOrDestroyPuppetObjects( TSharedRef<FMovieSceneInstance> MovieSceneInstance )
 {
-	ObjectSpawner->SpawnOrDestroyPuppetObjects( MovieSceneInstance );
+	if(ObjectSpawner.IsValid())
+	{
+		ObjectSpawner->SpawnOrDestroyPuppetObjects( MovieSceneInstance );
+	}
 }
 
 
@@ -784,6 +641,50 @@ void FSequencer::GetRuntimeObjects( TSharedRef<FMovieSceneInstance> MovieSceneIn
 	}
 }
 
+void FSequencer::UpdateViewports(AActor* ActorToViewThrough) const
+{
+	if(!IsPerspectiveViewportPosessionEnabled())
+	{
+		return;
+	}
+
+	for(FLevelEditorViewportClient* LevelVC : GEditor->LevelViewportClients)
+	{
+		if(LevelVC && LevelVC->IsPerspective() && LevelVC->AllowMatineePreview())
+		{
+			if(ActorToViewThrough)
+			{
+				LevelVC->SetViewLocation(ActorToViewThrough->GetActorLocation());
+				LevelVC->SetViewRotation(ActorToViewThrough->GetActorRotation());
+			}
+			else
+			{
+				LevelVC->ViewFOV = LevelVC->FOVAngle;
+			}
+
+			ACameraActor* Camera = Cast<ACameraActor>(ActorToViewThrough);
+
+			// If viewing through a camera - PP settings of camera.
+			LevelVC->SetPostprocessCameraActor(Camera);
+
+			if(Camera)
+			{
+				LevelVC->ViewFOV = LevelVC->FOVAngle = Camera->CameraComponent->FieldOfView;
+				// If the Camera's aspect ratio is zero, put a more reasonable default here - this at least stops it from crashing
+				// nb. the AspectRatio will be reported as a Map Check Warning
+				if(Camera->CameraComponent->AspectRatio == 0)
+				{
+					LevelVC->AspectRatio = 1.7f;
+				}
+				else
+				{
+					LevelVC->AspectRatio = Camera->CameraComponent->AspectRatio;
+				}
+			}
+		}
+	}
+}
+
 EMovieScenePlayerStatus::Type FSequencer::GetPlaybackStatus() const
 {
 	return PlaybackState;
@@ -850,49 +751,6 @@ void FSequencer::Tick(const float InDeltaTime)
 	}
 }
 
-
-void FSequencer::UpdateViewports(AActor* ActorToViewThrough) const
-{
-	if (!bPerspectiveViewportPossessionEnabled) {return;}
-
-	for (int32 i = 0; i < GEditor->LevelViewportClients.Num(); ++i)
-	{
-		FLevelEditorViewportClient* LevelVC = GEditor->LevelViewportClients[i];
-		if (LevelVC && LevelVC->IsPerspective() && LevelVC->AllowMatineePreview())
-		{
-			if (ActorToViewThrough)
-			{
-				LevelVC->SetViewLocation( ActorToViewThrough->GetActorLocation() );
-				LevelVC->SetViewRotation( ActorToViewThrough->GetActorRotation() );
-			}
-			else
-			{
-				LevelVC->ViewFOV = LevelVC->FOVAngle;
-			}
-
-			ACameraActor* Camera = Cast<ACameraActor>(ActorToViewThrough);
-			
-			// If viewing through a camera - PP settings of camera.
-			LevelVC->SetPostprocessCameraActor(Camera);
-
-			if (Camera)
-			{
-				LevelVC->ViewFOV = LevelVC->FOVAngle = Camera->CameraComponent->FieldOfView;
-				// If the Camera's aspect ratio is zero, put a more reasonable default here - this at least stops it from crashing
-				// nb. the AspectRatio will be reported as a Map Check Warning
-				if( Camera->CameraComponent->AspectRatio == 0 )
-				{
-					LevelVC->AspectRatio = 1.7f;
-				}
-				else
-				{
-					LevelVC->AspectRatio = Camera->CameraComponent->AspectRatio;
-				}
-			}
-		}
-	}
-}
-
 void FSequencer::AddReferencedObjects( FReferenceCollector& Collector )
 {
 	for( int32 MovieSceneIndex = 0; MovieSceneIndex < MovieSceneStack.Num(); ++MovieSceneIndex )
@@ -902,83 +760,6 @@ void FSequencer::AddReferencedObjects( FReferenceCollector& Collector )
 	}
 }
 
-void FSequencer::AttachTransportControlsToViewports()
-{
-	FLevelEditorModule* Module = FModuleManager::Get().LoadModulePtr<FLevelEditorModule>("LevelEditor");
-	if (Module)
-	{
-		TSharedPtr<ILevelEditor> LevelEditor = Module->GetFirstLevelEditor();
-		const TArray< TSharedPtr<ILevelViewport> >& LevelViewports = LevelEditor->GetViewports();
-		
-		FEditorWidgetsModule& EditorWidgetsModule = FModuleManager::Get().LoadModuleChecked<FEditorWidgetsModule>( "EditorWidgets" );
-
-		FTransportControlArgs TransportControlArgs;
-		TransportControlArgs.OnForwardPlay.BindSP(this, &FSequencer::OnPlay);
-		TransportControlArgs.OnRecord.BindSP(this, &FSequencer::OnRecord);
-		TransportControlArgs.OnForwardStep.BindSP(this, &FSequencer::OnStepForward);
-		TransportControlArgs.OnBackwardStep.BindSP(this, &FSequencer::OnStepBackward);
-		TransportControlArgs.OnForwardEnd.BindSP(this, &FSequencer::OnStepToEnd);
-		TransportControlArgs.OnBackwardEnd.BindSP(this, &FSequencer::OnStepToBeginning);
-		TransportControlArgs.OnToggleLooping.BindSP(this, &FSequencer::OnToggleLooping);
-		TransportControlArgs.OnGetLooping.BindSP(this, &FSequencer::IsLooping);
-		TransportControlArgs.OnGetPlaybackMode.BindSP(this, &FSequencer::GetPlaybackMode);
-
-		for (int32 i = 0; i < LevelViewports.Num(); ++i)
-		{
-			const TSharedPtr<ILevelViewport>& LevelViewport = LevelViewports[i];
-			
-			TSharedPtr<SWidget> TransportControl =
-				SNew(SHorizontalBox)
-				.Visibility(EVisibility::SelfHitTestInvisible)
-				+SHorizontalBox::Slot()
-				.FillWidth(1)
-				.HAlign(HAlign_Center)
-				.VAlign(VAlign_Bottom)
-				.Padding(4.f)
-				[
-					SNew(SBorder)
-					.Padding(4.f)
-					.Cursor( EMouseCursor::Default )
-					.BorderImage( FEditorStyle::GetBrush( "FilledBorder" ) )
-					.Visibility(this, &FSequencer::GetTransportControlVisibility, LevelViewport)
-					.Content()
-					[
-						EditorWidgetsModule.CreateTransportControl(TransportControlArgs)
-					]
-				];
-
-			LevelViewport->AddOverlayWidget(TransportControl.ToSharedRef());
-
-			TransportControls.Add(LevelViewport, TransportControl);
-		}
-	}
-}
-
-
-void FSequencer::DetachTransportControlsFromViewports()
-{
-	FLevelEditorModule* Module = FModuleManager::Get().LoadModulePtr<FLevelEditorModule>("LevelEditor");
-	if (Module)
-	{
-		TSharedPtr<ILevelEditor> LevelEditor = Module->GetFirstLevelEditor();
-		if (LevelEditor.IsValid())
-		{
-			const TArray< TSharedPtr<ILevelViewport> >& LevelViewports = LevelEditor->GetViewports();
-		
-			for (int32 i = 0; i < LevelViewports.Num(); ++i)
-			{
-				const TSharedPtr<ILevelViewport>& LevelViewport = LevelViewports[i];
-
-				TSharedPtr<SWidget>* TransportControl = TransportControls.Find(LevelViewport);
-				if (TransportControl && TransportControl->IsValid())
-				{
-					LevelViewport->RemoveOverlayWidget(TransportControl->ToSharedRef());
-				}
-			}
-		}
-	}
-}
-
 void FSequencer::ResetPerMovieSceneData()
 {
 	FilterToShotSections( TArray< TWeakObjectPtr<UMovieSceneSection> >(), false );
@@ -1005,11 +786,6 @@ void FSequencer::DestroySpawnablesForAllMovieScenes()
 	ObjectSpawner->DestroyAllPuppetObjects();
 }
 
-EVisibility FSequencer::GetTransportControlVisibility(TSharedPtr<ILevelViewport> LevelViewport) const
-{
-	FLevelEditorViewportClient& ViewportClient = LevelViewport->GetLevelViewportClient();
-	return (ViewportClient.IsPerspective() && ViewportClient.AllowMatineePreview()) ? EVisibility::Visible : EVisibility::Collapsed;
-}
 
 FReply FSequencer::OnPlay()
 {
@@ -1158,8 +934,7 @@ UK2Node_PlayMovieScene* FSequencer::FindPlayMovieSceneNodeInLevelScript( const U
 {
 	// Grab the world object for this editor
 	check( MovieScene != NULL );
-	check( IsWorldCentricAssetEditor() );
-	UWorld* World = GetToolkitHost()->GetWorld();
+	UWorld* World = ToolkitHost.Pin()->GetWorld();
 	check( World != NULL );
 
 	// Search all levels in the specified world
@@ -1196,8 +971,7 @@ UK2Node_PlayMovieScene* FSequencer::CreateNewPlayMovieSceneNode( UMovieScene* Mo
 {
 	// Grab the world object for this editor
 	check( MovieScene != NULL );
-	check( IsWorldCentricAssetEditor() );
-	UWorld* World = GetToolkitHost()->GetWorld();
+	UWorld* World = ToolkitHost.Pin()->GetWorld();
 	check( World != NULL );
 
 	ULevel* Level = World->GetCurrentLevel();
@@ -1277,11 +1051,6 @@ void FSequencer::OnToggleCleanView( bool bInCleanViewEnabled )
 	bCleanViewEnabled = bInCleanViewEnabled;
 }
 
-// @todo remove when world-centric mode is added
-TSharedRef<class SSequencer> FSequencer::GetMainSequencer()
-{
-	return SequencerWidget.ToSharedRef();
-}
 
 
 FGuid FSequencer::AddSpawnableForAssetOrClass( UObject* Object, UObject* CounterpartGamePreviewObject )
diff --git a/Engine/Source/Editor/Sequencer/Private/Sequencer.h b/Engine/Source/Editor/Sequencer/Private/Sequencer.h
--- a/Engine/Source/Editor/Sequencer/Private/Sequencer.h
+++ b/Engine/Source/Editor/Sequencer/Private/Sequencer.h
@@ -2,33 +2,28 @@
 
 #pragma once
 
-#include "Toolkits/AssetEditorToolkit.h"
-#include "Editor/LevelEditor/Public/LevelEditor.h"	// For EMapChangeType::Type
 #include "Editor/EditorWidgets/Public/ITransportControl.h"
 #include "SelectedKey.h"
 
 class UMovieScene;
+class IToolkitHost;
 
 /**
  * Sequencer is the editing tool for MovieScene assets
  */
-class FSequencer : public ISequencer, FGCObject, public FEditorUndoClient
+class FSequencer : public ISequencer, public FGCObject, public FEditorUndoClient
 { 
 
 public:
-
-	virtual void RegisterTabSpawners(const TSharedRef<class FTabManager>& TabManager) OVERRIDE;
-	virtual void UnregisterTabSpawners(const TSharedRef<class FTabManager>& TabManager) OVERRIDE;
-
 	/**
-	 * Edits the specified object
+	 * Initializes sequencer
 	 *
-	 * @param	Mode							Asset editing mode for this editor (standalone or world-centric)
-	 * @param	InitToolkitHost					When Mode is WorldCentric, this is the level editor instance to spawn this editor within
-	 * @param	ObjectToEdit					The object to edit
-	 * @param	TrackEditorDelegates		Delegates to call to create auto-key handlers for this sequencer
+	 * @param	RootMovieScene					The root movie scene being edited
+	 * @param	InAssetEditor					The asset editor created for this (if any) 
+	 * @param	TrackEditorDelegates			Delegates to call to create auto-key handlers for this sequencer
+	 * @param	bEditWithinLevelEditor			Whether or not sequencer should be edited within the level editor
 	 */
-	void InitSequencer( const EToolkitMode::Type Mode, const TSharedPtr< class IToolkitHost >& InitToolkitHost, UObject* ObjectToEdit, const TArray<FOnCreateTrackEditor>& TrackEditorDelegates );
+	void InitSequencer( UMovieScene* InRootMovieScene, TSharedPtr<IToolkitHost> InToolKitHost, const TArray<FOnCreateTrackEditor>& TrackEditorDelegates, bool bEditWithinLevelEditor );
 
 	/** Constructor */
 	FSequencer();
@@ -39,14 +34,9 @@ public:
 	/** FGCObject interface */
 	virtual void AddReferencedObjects(FReferenceCollector& Collector) OVERRIDE;
 
-	/** IToolkit interface */
-	virtual FName GetToolkitFName() const OVERRIDE;
-	virtual FText GetBaseToolkitName() const OVERRIDE;
-	virtual FString GetWorldCentricTabPrefix() const OVERRIDE;
-	virtual FLinearColor GetWorldCentricTabColorScale() const OVERRIDE;
 
 	/** ISequencer interface */
-	virtual TArray< UMovieScene* > GetMovieScenesBeingEdited();
+	virtual TSharedRef<SWidget> GetSequencerWidget() const OVERRIDE { return SequencerWidget.ToSharedRef(); }
 	virtual UMovieScene* GetRootMovieScene() const;
 	virtual UMovieScene* GetFocusedMovieScene() const OVERRIDE;
 	virtual TSharedRef<FMovieSceneInstance> GetRootMovieSceneInstance() const OVERRIDE;
@@ -68,19 +58,13 @@ public:
 	virtual void FilterToShotSections(const TArray< TWeakObjectPtr<class UMovieSceneSection> >& ShotSections, bool bZoomToShotBounds = true) OVERRIDE;
 	virtual void FilterToSelectedShotSections(bool bZoomToShotBounds = true) OVERRIDE;
 
+	bool IsPerspectiveViewportPosessionEnabled() const { return bPerspectiveViewportPossessionEnabled; }
+
 	/**
 	 * Pops the current focused movie scene from the stack.  The parent of this movie scene will be come the focused one
 	 */
 	void PopToMovieScene( TSharedRef<FMovieSceneInstance> SubMovieSceneInstance );
 
-	/**
-	 * Starts editing details for the specified list of UObjects.  This is used to edit CDO's for blueprints
-	 * that are stored with the MovieScene asset
-	 *
-	 * @param	ObjectsToEdit	The objects to edit in the details view
-	 */
-	void EditDetailsForObjects( TArray< UObject* > ObjectsToEdit );
-
 	/**
 	 * Spawn (or destroy) puppet objects as needed to match the spawnables array in the MovieScene we're editing
 	 *
@@ -103,16 +87,6 @@ public:
 	 */
 	void DeleteSelectedKeys();
 
-	/**
-	 * Binds to a 'PlayMovieScene' node in level script that is associated with the MovieScene we're editing.  If
-	 * requested, can also create a new node.  Only valid to call when in world-centric mode.
-	 *
-	 * @param	bCreateIfNotFound	True to create a node if we can't find one
-	 *
-	 * @return	The LevelScript node we found or created, or NULL if nothing is still bound
-	 */
-	class UK2Node_PlayMovieScene* BindToPlayMovieSceneNode( const bool bCreateIfNotFound );
-
 	/**
 	 * @return Movie scene tools used by the sequencer
 	 */
@@ -266,21 +240,36 @@ public:
 
 	virtual void SpawnActorsForMovie( TSharedRef<FMovieSceneInstance> MovieSceneInstance );
 
+	virtual TArray< UMovieScene* > GetMovieScenesBeingEdited();
 
+	/**
+	 * Binds to a 'PlayMovieScene' node in level script that is associated with the MovieScene we're editing.  If
+	 * requested, can also create a new node.  Only valid to call when in world-centric mode.
+	 *
+	 * @param	bCreateIfNotFound	True to create a node if we can't find one
+	 *
+	 * @return	The LevelScript node we found or created, or NULL if nothing is still bound
+	 */
+	class UK2Node_PlayMovieScene* BindToPlayMovieSceneNode( const bool bCreateIfNotFound );
 
-	// @todo remove when world-centric mode is added
-	TSharedRef<class SSequencer> GetMainSequencer();
+	/** Called by LevelEditor when the map changes */
+	void OnMapChanged(class UWorld* NewWorld, EMapChangeType::Type MapChangeType);
 
-private:
-	TSharedRef<SDockTab> SpawnTab_SequencerMain(const FSpawnTabArgs& Args);
-	TSharedRef<SDockTab> SpawnTab_Details(const FSpawnTabArgs& Args);	
-protected:
+	/** Functions to push on to the transport controls we use */
+	FReply OnPlay();
+	FReply OnRecord();
+	FReply OnStepForward();
+	FReply OnStepBackward();
+	FReply OnStepToEnd();
+	FReply OnStepToBeginning();
+	FReply OnToggleLooping();
+	bool IsLooping() const;
+	EPlaybackMode::Type GetPlaybackMode() const;
 
-	/** Generates and attaches transport control widgets to the main level editor viewports */
-	void AttachTransportControlsToViewports();
-	/** Purges all transport control widgets from the main level editor viewports */
-	void DetachTransportControlsFromViewports();
-	
+	/** @return The toolkit that this sequencer is hosted in (if any) */
+	TSharedPtr<IToolkitHost> GetToolkitHost() const { return ToolkitHost.Pin(); }
+
+protected:
 	/**
 	 * Reset data about a movie scene when pushing or popping a movie scene
 	 */
@@ -294,22 +283,9 @@ protected:
 	 */
 	void DestroySpawnablesForAllMovieScenes();
 
-	/** Gets the visibility of a transport control based on the viewport it's attached to */
-	EVisibility GetTransportControlVisibility(TSharedPtr<ILevelViewport> LevelViewport) const;
-
 	/** Sets the actor CDO such that it is placed in front of the active perspective viewport camera, if we have one */
 	static void PlaceActorInFrontOfCamera( AActor* ActorCDO );
 
-	/** Functions to push on to the transport controls we use */
-	FReply OnPlay();
-	FReply OnRecord();
-	FReply OnStepForward();
-	FReply OnStepBackward();
-	FReply OnStepToEnd();
-	FReply OnStepToBeginning();
-	FReply OnToggleLooping();
-	bool IsLooping() const;
-	EPlaybackMode::Type GetPlaybackMode() const;
 
 	/**
 	 * Gets the far time boundaries of the currently edited movie scene
@@ -346,8 +322,6 @@ protected:
 	 */
 	class UK2Node_PlayMovieScene* CreateNewPlayMovieSceneNode( UMovieScene* MovieScene );
 
-	/** Called by LevelEditor when the map changes */
-	void OnMapChanged( class UWorld* NewWorld, EMapChangeType::Type MapChangeType );
 
 	/** Called by UK2Node_PlayMovieScene when the bindings for that node are changed through Kismet */
 	void OnPlayMovieSceneBindingsChanged();
@@ -411,12 +385,11 @@ protected:
 	// End of FEditorUndoClient
 
 private:
+	TMap< TWeakObjectPtr<UMovieSceneSection>, TSharedRef<FMovieSceneInstance> > MovieSceneSectionToInstanceMap;
+
 	/** Command list for seququencer commands */
 	TSharedRef<FUICommandList> SequencerCommandBindings;
 
-	/** A map of all the transport controls to viewports that this sequencer has made */
-	TMap< TSharedPtr<class ILevelViewport>, TSharedPtr<class SWidget> > TransportControls;
-
 	/** List of tools we own */
 	TArray< TSharedPtr<FMovieSceneTrackEditor> > TrackEditors;
 
@@ -424,9 +397,6 @@ private:
 	TWeakPtr<FMovieSceneTrackEditor> DirectorTrackEditor;
 	TWeakPtr<FMovieSceneTrackEditor> AnimationTrackEditor;
 
-	/** Details view */
-	TSharedPtr< class IDetailsView > DetailsView;
-
 	/** Object spawner */
 	TSharedPtr< class ISequencerObjectSpawner > ObjectSpawner;
 
@@ -436,14 +406,15 @@ private:
 	/** The runtime instance for the root movie scene */
 	TSharedPtr< class FMovieSceneInstance > RootMovieSceneInstance;
 
-	TMap< TWeakObjectPtr<UMovieSceneSection>, TSharedRef<FMovieSceneInstance> > MovieSceneSectionToInstanceMap;
-
 	/** Main sequencer widget */
 	TSharedPtr< class SSequencer > SequencerWidget;
 	
 	/** Reference to owner of the current popup */
 	TWeakPtr<class SWindow> NameEntryPopupWindow;
 
+	/** The asset editor that created this Sequencer if any */
+	TWeakPtr<IToolkitHost> ToolkitHost;
+
 	/** 'PlayMovieScene' node in level script that we're currently associating with the active MovieScene.  This node contains the bindings
 	    information that we need to be able to preview possessed actors in the level */
 	TWeakObjectPtr< UK2Node_PlayMovieScene > PlayMovieSceneNode;
@@ -493,14 +464,8 @@ private:
 	/** Whether or not we are allowing autokey */
 	bool bAllowAutoKey;
 
-	/** Whether or not we allow perspective viewports to be hijacked by this sequencer */
 	bool bPerspectiveViewportPossessionEnabled;
 
-
-	/**	The tab ids for all the tabs used */
-	static const FName SequencerMainTabId;
-	static const FName SequencerDetailsTabId;
-
 	/** Stores a dirty bit for whether the sequencer tree (and other UI bits) may need to be refreshed.  We
 	    do this simply to avoid refreshing the UI more than once per frame. (e.g. during live recording where
 		the MovieScene data can change many times per frame.) */
diff --git a/Engine/Source/Editor/Sequencer/Private/SequencerAssetEditor.cpp b/Engine/Source/Editor/Sequencer/Private/SequencerAssetEditor.cpp
new file mode 100644
--- /dev/null
+++ b/Engine/Source/Editor/Sequencer/Private/SequencerAssetEditor.cpp
@@ -0,0 +1,227 @@
+// Copyright 1998-2014 Epic Games, Inc. All Rights Reserved.
+
+#include "SequencerPrivatePCH.h"
+#include "Toolkits/IToolkitHost.h"
+#include "SequencerAssetEditor.h"
+#include "Editor/LevelEditor/Public/LevelEditor.h"
+
+#define LOCTEXT_NAMESPACE "Sequencer"
+
+const FName FSequencerAssetEditor::SequencerMainTabId( TEXT( "Sequencer_SequencerMain" ) );
+
+namespace SequencerDefs
+{
+	static const FName SequencerAppIdentifier( TEXT( "SequencerApp" ) );
+}
+void FSequencerAssetEditor::RegisterTabSpawners(const TSharedRef<class FTabManager>& TabManager)
+{
+	const IWorkspaceMenuStructure& MenuStructure = WorkspaceMenu::GetMenuStructure();
+
+	if( IsSequencerEnabled() && !IsWorldCentricAssetEditor() )
+	{
+		TabManager->RegisterTabSpawner( SequencerMainTabId, FOnSpawnTab::CreateSP(this, &FSequencerAssetEditor::SpawnTab_SequencerMain) )
+			.SetDisplayName( LOCTEXT("SequencerMainTab", "Sequencer") )
+			.SetGroup( MenuStructure.GetAssetEditorCategory() );
+	}
+
+}
+
+void FSequencerAssetEditor::UnregisterTabSpawners(const TSharedRef<class FTabManager>& TabManager)
+{
+	if( IsSequencerEnabled() && !IsWorldCentricAssetEditor() )
+	{
+		TabManager->UnregisterTabSpawner( SequencerMainTabId );
+	}
+	
+	// @todo remove when world-centric mode is added
+	FLevelEditorModule& LevelEditorModule = FModuleManager::LoadModuleChecked<FLevelEditorModule>("LevelEditor");
+	LevelEditorModule.AttachSequencer( SNullWidget::NullWidget, nullptr );
+}
+
+void FSequencerAssetEditor::InitSequencerAssetEditor( const EToolkitMode::Type Mode, const TSharedPtr< class IToolkitHost >& InitToolkitHost, UMovieScene* InRootMovieScene, const TArray<FOnCreateTrackEditor>& TrackEditorDelegates, bool bEditWithinLevelEditor )
+{
+	if(InitToolkitHost.IsValid())
+	{
+		const TSharedRef<FTabManager::FLayout> StandaloneDefaultLayout = FTabManager::NewLayout("Standalone_Sequencer_Layout")
+		->AddArea
+		(
+			FTabManager::NewPrimaryArea()
+			->Split
+			(
+				FTabManager::NewStack()
+				->AddTab(SequencerMainTabId, ETabState::OpenedTab)
+			)
+		);
+
+		const bool bCreateDefaultStandaloneMenu = true;
+		const bool bCreateDefaultToolbar = false;
+		FAssetEditorToolkit::InitAssetEditor(Mode, InitToolkitHost, SequencerDefs::SequencerAppIdentifier, StandaloneDefaultLayout, bCreateDefaultStandaloneMenu, bCreateDefaultToolbar, InRootMovieScene);
+	}
+
+	// Init Sequencer 
+	Sequencer = MakeShareable(new FSequencer);
+	Sequencer->InitSequencer( InRootMovieScene, InitToolkitHost, TrackEditorDelegates, bEditWithinLevelEditor);
+
+	if(bEditWithinLevelEditor)
+	{
+		// @todo remove when world-centric mode is added
+		FLevelEditorModule& LevelEditorModule = FModuleManager::LoadModuleChecked<FLevelEditorModule>("LevelEditor");
+
+		LevelEditorModule.AttachSequencer( Sequencer->GetSequencerWidget(), SharedThis(this));
+			
+		// We need to find out when the user loads a new map, because we might need to re-create puppet actors
+		// when previewing a MovieScene
+		LevelEditorModule.OnMapChanged().AddSP(Sequencer.ToSharedRef(), &FSequencer::OnMapChanged);
+
+		AttachTransportControlsToViewports();
+	}
+}
+
+
+FSequencerAssetEditor::FSequencerAssetEditor()
+{
+
+}
+
+FSequencerAssetEditor::~FSequencerAssetEditor()
+{
+	DetachTransportControlsFromViewports();
+
+
+	// Unregister delegates
+	if( FModuleManager::Get().IsModuleLoaded( TEXT( "LevelEditor" ) ) )
+	{
+		auto& LevelEditorModule = FModuleManager::LoadModuleChecked< FLevelEditorModule >( TEXT( "LevelEditor" ) );
+		LevelEditorModule.OnMapChanged().RemoveAll( this );
+	}
+}
+
+
+TSharedRef<SDockTab> FSequencerAssetEditor::SpawnTab_SequencerMain(const FSpawnTabArgs& Args)
+{
+	check(Args.GetTabId() == SequencerMainTabId);
+
+	return SNew(SDockTab)
+		.Icon( FEditorStyle::GetBrush("Sequencer.Tabs.SequencerMain") )
+		.Label( LOCTEXT("SequencerMainTitle", "Sequencer") )
+		.TabColorScale( GetTabColorScale() )
+		[
+			Sequencer->GetSequencerWidget()
+		];
+}
+
+
+FName FSequencerAssetEditor::GetToolkitFName() const
+{
+	static FName SequencerName("Sequencer");
+	return SequencerName;
+}
+
+FText FSequencerAssetEditor::GetBaseToolkitName() const
+{
+	return LOCTEXT("AppLabel", "Sequencer");
+}
+
+
+FLinearColor FSequencerAssetEditor::GetWorldCentricTabColorScale() const
+{
+	return FLinearColor( 0.7, 0.0f, 0.0f, 0.5f );
+}
+
+
+FString FSequencerAssetEditor::GetWorldCentricTabPrefix() const
+{
+	return LOCTEXT("WorldCentricTabPrefix", "Sequencer ").ToString();
+}
+
+
+void FSequencerAssetEditor::AttachTransportControlsToViewports()
+{
+	FLevelEditorModule* Module = FModuleManager::Get().LoadModulePtr<FLevelEditorModule>("LevelEditor");
+	if (Module)
+	{
+		TSharedPtr<ILevelEditor> LevelEditor = Module->GetFirstLevelEditor();
+		const TArray< TSharedPtr<ILevelViewport> >& LevelViewports = LevelEditor->GetViewports();
+		
+		FEditorWidgetsModule& EditorWidgetsModule = FModuleManager::Get().LoadModuleChecked<FEditorWidgetsModule>( "EditorWidgets" );
+
+		TSharedRef<FSequencer> SequencerRef = Sequencer.ToSharedRef();
+
+		FTransportControlArgs TransportControlArgs;
+		TransportControlArgs.OnForwardPlay.BindSP(SequencerRef, &FSequencer::OnPlay);
+		TransportControlArgs.OnRecord.BindSP(SequencerRef, &FSequencer::OnRecord);
+		TransportControlArgs.OnForwardStep.BindSP(SequencerRef, &FSequencer::OnStepForward);
+		TransportControlArgs.OnBackwardStep.BindSP(SequencerRef, &FSequencer::OnStepBackward);
+		TransportControlArgs.OnForwardEnd.BindSP(SequencerRef, &FSequencer::OnStepToEnd);
+		TransportControlArgs.OnBackwardEnd.BindSP(SequencerRef, &FSequencer::OnStepToBeginning);
+		TransportControlArgs.OnToggleLooping.BindSP(SequencerRef, &FSequencer::OnToggleLooping);
+		TransportControlArgs.OnGetLooping.BindSP(SequencerRef, &FSequencer::IsLooping);
+		TransportControlArgs.OnGetPlaybackMode.BindSP(SequencerRef, &FSequencer::GetPlaybackMode);
+
+		for ( const TSharedPtr<ILevelViewport>& LevelViewport : LevelViewports )
+		{
+			TWeakPtr<ILevelViewport> LevelViewportWeakPtr = LevelViewport;
+			
+			TSharedPtr<SWidget> TransportControl =
+				SNew(SHorizontalBox)
+				.Visibility(EVisibility::SelfHitTestInvisible)
+				+SHorizontalBox::Slot()
+				.FillWidth(1)
+				.HAlign(HAlign_Center)
+				.VAlign(VAlign_Bottom)
+				.Padding(4.f)
+				[
+					SNew(SBorder)
+					.Padding(4.f)
+					.Cursor( EMouseCursor::Default )
+					.BorderImage( FEditorStyle::GetBrush( "FilledBorder" ) )
+					.Visibility(this, &FSequencerAssetEditor::GetTransportControlVisibility, LevelViewportWeakPtr)
+					.Content()
+					[
+						EditorWidgetsModule.CreateTransportControl(TransportControlArgs)
+					]
+				];
+
+			LevelViewport->AddOverlayWidget(TransportControl.ToSharedRef());
+
+			TransportControls.Add(LevelViewportWeakPtr, TransportControl);
+		}
+	}
+}
+
+
+void FSequencerAssetEditor::DetachTransportControlsFromViewports()
+{
+	FLevelEditorModule* Module = FModuleManager::Get().LoadModulePtr<FLevelEditorModule>("LevelEditor");
+	if (Module)
+	{
+		TSharedPtr<ILevelEditor> LevelEditor = Module->GetFirstLevelEditor();
+		if (LevelEditor.IsValid())
+		{
+			const TArray< TSharedPtr<ILevelViewport> >& LevelViewports = LevelEditor->GetViewports();
+		
+			for (int32 i = 0; i < LevelViewports.Num(); ++i)
+			{
+				const TSharedPtr<ILevelViewport>& LevelViewport = LevelViewports[i];
+
+				TSharedPtr<SWidget>* TransportControl = TransportControls.Find(LevelViewport);
+				if (TransportControl && TransportControl->IsValid())
+				{
+					LevelViewport->RemoveOverlayWidget(TransportControl->ToSharedRef());
+				}
+			}
+		}
+	}
+}
+
+EVisibility FSequencerAssetEditor::GetTransportControlVisibility(TWeakPtr<ILevelViewport> LevelViewport) const
+{
+	TSharedPtr<ILevelViewport> LevelViewportPin = LevelViewport.Pin();
+
+	FLevelEditorViewportClient& ViewportClient = LevelViewportPin->GetLevelViewportClient();
+	return (ViewportClient.IsPerspective() && ViewportClient.AllowMatineePreview()) ? EVisibility::Visible : EVisibility::Collapsed;
+}
+
+
+
+#undef LOCTEXT_NAMESPACE
diff --git a/Engine/Source/Editor/Sequencer/Private/SequencerAssetEditor.h b/Engine/Source/Editor/Sequencer/Private/SequencerAssetEditor.h
new file mode 100644
--- /dev/null
+++ b/Engine/Source/Editor/Sequencer/Private/SequencerAssetEditor.h
@@ -0,0 +1,58 @@
+// Copyright 1998-2014 Epic Games, Inc. All Rights Reserved.
+
+#pragma once
+
+#include "ISequencerAssetEditor.h"
+
+class FSequencerAssetEditor : public ISequencerAssetEditor
+{ 
+public:
+	/** Constructor */
+	FSequencerAssetEditor();
+
+	/** Destructor */
+	virtual ~FSequencerAssetEditor();
+
+	/**
+	 * Edits the specified object
+	 *
+	 * @param	Mode							Asset editing mode for this editor (standalone or world-centric)
+	 * @param	InitToolkitHost					When Mode is WorldCentric, this is the level editor instance to spawn this editor within
+	 * @param	ObjectToEdit					The object to edit
+	 * @param	TrackEditorDelegates			Delegates to call to create auto-key handlers for this sequencer
+	 * @param	bEditWithinLevelEditor			Whether or not sequencer should be edited within the level editor
+	 */
+	void InitSequencerAssetEditor( const EToolkitMode::Type Mode, const TSharedPtr< class IToolkitHost >& InitToolkitHost, UMovieScene* InRootMovieScene, const TArray<FOnCreateTrackEditor>& TrackEditorDelegates, bool bEditWithinLevelEditor );
+
+	TSharedRef<ISequencer> GetSequencerInterface() const OVERRIDE { return Sequencer.ToSharedRef(); }
+
+	/** IToolkit interface */
+	virtual FName GetToolkitFName() const OVERRIDE;
+	virtual FText GetBaseToolkitName() const OVERRIDE;
+	virtual FString GetWorldCentricTabPrefix() const OVERRIDE;
+	virtual FLinearColor GetWorldCentricTabColorScale() const OVERRIDE;
+
+	virtual void RegisterTabSpawners(const TSharedRef<class FTabManager>& TabManager) OVERRIDE;
+	virtual void UnregisterTabSpawners(const TSharedRef<class FTabManager>& TabManager) OVERRIDE;
+
+	void UpdateViewports(AActor* ActorToViewThrough) const;
+
+private:
+	EVisibility GetTransportControlVisibility(TWeakPtr<ILevelViewport> LevelViewport) const;
+	TSharedRef<SDockTab> SpawnTab_SequencerMain(const FSpawnTabArgs& Args);	
+protected:
+
+	/** Generates and attaches transport control widgets to the main level editor viewports */
+	void AttachTransportControlsToViewports();
+	/** Purges all transport control widgets from the main level editor viewports */
+	void DetachTransportControlsFromViewports();
+	
+private:
+	TSharedPtr<FSequencer> Sequencer;
+
+	/** A map of all the transport controls to viewports that this sequencer has made */
+	TMap< TWeakPtr<class ILevelViewport>, TSharedPtr<class SWidget> > TransportControls;
+	/**	The tab ids for all the tabs used */
+	static const FName SequencerMainTabId;
+
+};
diff --git a/Engine/Source/Editor/Sequencer/Private/SequencerModule.cpp b/Engine/Source/Editor/Sequencer/Private/SequencerModule.cpp
--- a/Engine/Source/Editor/Sequencer/Private/SequencerModule.cpp
+++ b/Engine/Source/Editor/Sequencer/Private/SequencerModule.cpp
@@ -13,13 +13,20 @@ class FSequencerModule : public ISequencerModule
 {
 
 	/** ISequencerModule interface */
-	virtual TSharedPtr<ISequencer> CreateSequencer( const EToolkitMode::Type Mode, const TSharedPtr< IToolkitHost >& InitToolkitHost, UObject* ObjectToEdit ) OVERRIDE
+	virtual TSharedPtr<ISequencer> CreateSequencer( UMovieScene* RootMovieScene ) OVERRIDE
 	{
-		TSharedRef< FSequencer > SequencerInUse = MakeShareable(new FSequencer);
-		SequencerInUse->InitSequencer( Mode, InitToolkitHost, ObjectToEdit, TrackEditorDelegates );
-		return SequencerInUse;
+		TSharedRef< FSequencer > Sequencer = MakeShareable(new FSequencer);
+		Sequencer->InitSequencer( RootMovieScene, nullptr, TrackEditorDelegates, false );
+		return Sequencer;
 	}
 	
+	virtual TSharedPtr<ISequencer> CreateSequencerAssetEditor( const EToolkitMode::Type Mode, const TSharedPtr< class IToolkitHost >& InitToolkitHost, UMovieScene* InRootMovieScene, bool bEditWithinLevelEditor ) OVERRIDE
+	{
+		TSharedRef< FSequencerAssetEditor > SequencerAssetEditor = MakeShareable(new FSequencerAssetEditor);
+		SequencerAssetEditor->InitSequencerAssetEditor( Mode, InitToolkitHost, InRootMovieScene, TrackEditorDelegates, bEditWithinLevelEditor );
+		return SequencerAssetEditor->GetSequencerInterface();
+	}
+
 	virtual void RegisterTrackEditor( FOnCreateTrackEditor InOnCreateTrackEditor ) OVERRIDE
 	{
 		TrackEditorDelegates.AddUnique( InOnCreateTrackEditor );
diff --git a/Engine/Source/Editor/Sequencer/Public/ISequencer.h b/Engine/Source/Editor/Sequencer/Public/ISequencer.h
--- a/Engine/Source/Editor/Sequencer/Public/ISequencer.h
+++ b/Engine/Source/Editor/Sequencer/Public/ISequencer.h
@@ -12,10 +12,13 @@ class UMovieScene;
 /**
  * Sequencer public interface
  */
-class ISequencer : public FAssetEditorToolkit, public IMovieScenePlayer
+class ISequencer : public IMovieScenePlayer, public TSharedFromThis<ISequencer>
 {
 
 public:
+	/** @return Widget used to display the sequencer */
+	virtual TSharedRef<SWidget> GetSequencerWidget() const = 0;
+
 	/** @return Returns the MovieScene that is currently focused for editing by the sequencer.  This can change at any time. */
 	virtual UMovieScene* GetFocusedMovieScene() const = 0;
 
diff --git a/Engine/Source/Editor/Sequencer/Public/ISequencerAssetEditor.h b/Engine/Source/Editor/Sequencer/Public/ISequencerAssetEditor.h
new file mode 100644
--- /dev/null
+++ b/Engine/Source/Editor/Sequencer/Public/ISequencerAssetEditor.h
@@ -0,0 +1,21 @@
+
+
+// Copyright 1998-2014 Epic Games, Inc. All Rights Reserved.
+
+#pragma once
+
+
+#include "Toolkits/AssetEditorToolkit.h"
+#include "Toolkits/IToolkit.h"
+
+class UMovieScene;
+
+/**
+ * Sequencer public interface
+ */
+class ISequencerAssetEditor : public FAssetEditorToolkit
+{
+
+public:
+	virtual TSharedRef<ISequencer> GetSequencerInterface() const = 0;
+};
diff --git a/Engine/Source/Editor/Sequencer/Public/ISequencerModule.h b/Engine/Source/Editor/Sequencer/Public/ISequencerModule.h
--- a/Engine/Source/Editor/Sequencer/Public/ISequencerModule.h
+++ b/Engine/Source/Editor/Sequencer/Public/ISequencerModule.h
@@ -20,7 +20,15 @@ class ISequencerModule : public IModuleInterface
 public:
 
 	/**
-	 * Creates a new instance of a Sequencer, the editor for MovieScene assets
+	 * Creates a new instance of a standalone sequencer that can be added to other UIs
+	 *
+	 * @param 	InRootMovieScene	The movie scene to edit
+	 * @return	Interface to the new editor
+	 */
+	virtual TSharedPtr<ISequencer> CreateSequencer( UMovieScene* InRootMovieScene ) = 0;
+
+	/**
+	 * Creates a new instance of a Sequencer, the editor for MovieScene assets in an asset editor
 	 *
 	 * @param	Mode					Mode that this editor should operate in
 	 * @param	InitToolkitHost			When Mode is WorldCentric, this is the level editor instance to spawn this editor within
@@ -28,7 +36,7 @@ public:
 	 *
 	 * @return	Interface to the new editor
 	 */
-	virtual TSharedPtr<ISequencer> CreateSequencer( const EToolkitMode::Type Mode, const TSharedPtr< IToolkitHost >& InitToolkitHost, UObject* ObjectToEdit ) = 0;
+	virtual TSharedPtr<ISequencer> CreateSequencerAssetEditor( const EToolkitMode::Type Mode, const TSharedPtr< class IToolkitHost >& InitToolkitHost, UMovieScene* InRootMovieScene, bool bEditWithinLevelEditor ) = 0;
 
 	/** 
 	 * Registers a delegate that will create an editor for a track in each sequencer 
diff --git a/Engine/Source/Editor/UMGEditor/Private/UMGBlueprintEditorExtensionHook.cpp b/Engine/Source/Editor/UMGEditor/Private/UMGBlueprintEditorExtensionHook.cpp
--- a/Engine/Source/Editor/UMGEditor/Private/UMGBlueprintEditorExtensionHook.cpp
+++ b/Engine/Source/Editor/UMGEditor/Private/UMGBlueprintEditorExtensionHook.cpp
@@ -8,8 +8,6 @@
 #include "BlueprintEditor.h"
 #include "BlueprintEditorModes.h"
 #include "BlueprintEditorTabs.h"
-#include "SEditorViewport.h"
-#include "SSCSEditorViewport.h"
 
 #include "SUMGEditorTree.h"
 #include "SUMGDesigner.h"
@@ -22,6 +20,9 @@
 
 #include "WidgetBlueprintEditor.h"
 
+#include "MovieScene.h"
+#include "ISequencerModule.h"
+
 #define LOCTEXT_NAMESPACE "UMG_EXTENSION"
 
 /////////////////////////////////////////////////////
@@ -32,9 +33,9 @@ static const FName SlatePreviewTabID(TEXT("SlatePreview"));
 struct FSlatePreviewSummoner : public FWorkflowTabFactory
 {
 protected:
-	TWeakPtr<class FBlueprintEditor> BlueprintEditor;
+	TWeakPtr<class FWidgetBlueprintEditor> BlueprintEditor;
 public:
-	FSlatePreviewSummoner(TSharedPtr<class FBlueprintEditor> InBlueprintEditor)
+	FSlatePreviewSummoner(TSharedPtr<class FWidgetBlueprintEditor> InBlueprintEditor)
 		: FWorkflowTabFactory(SlatePreviewTabID, InBlueprintEditor)
 		, BlueprintEditor(InBlueprintEditor)
 	{
@@ -49,8 +50,6 @@ public:
 
 	virtual TSharedRef<SWidget> CreateTabBody(const FWorkflowTabSpawnInfo& Info) const OVERRIDE
 	{
-		TSharedPtr<FWidgetBlueprintEditor> BlueprintEditorPtr = StaticCastSharedPtr<FWidgetBlueprintEditor>(BlueprintEditor.Pin());
-
 		//TSharedPtr<SWidget> Result;
 		//if ( BlueprintEditorPtr->CanAccessComponentsMode() )
 		//{
@@ -61,7 +60,7 @@ public:
 			+ SHorizontalBox::Slot()
 			.FillWidth(1)
 			[
-				SNew(SUMGDesigner, BlueprintEditorPtr)
+				SNew(SUMGDesigner, BlueprintEditor.Pin())
 			];
 	}
 };
@@ -74,9 +73,9 @@ static const FName SlateHierarchyTabID(TEXT("SlateHierarchy"));
 struct FSlateTreeSummoner : public FWorkflowTabFactory
 {
 protected:
-	TWeakPtr<class FBlueprintEditor> BlueprintEditor;
+	TWeakPtr<class FWidgetBlueprintEditor> BlueprintEditor;
 public:
-	FSlateTreeSummoner(TSharedPtr<class FBlueprintEditor> InBlueprintEditor)
+	FSlateTreeSummoner(TSharedPtr<class FWidgetBlueprintEditor> InBlueprintEditor)
 		: FWorkflowTabFactory(SlateHierarchyTabID, InBlueprintEditor)
 		, BlueprintEditor(InBlueprintEditor)
 	{
@@ -104,9 +103,9 @@ static const FName UMGWidgetTemplatesTabID(TEXT("WidgetTemplates"));
 struct FWidgetTemplatesSummoner : public FWorkflowTabFactory
 {
 protected:
-	TWeakPtr<class FBlueprintEditor> BlueprintEditor;
+	TWeakPtr<class FWidgetBlueprintEditor> BlueprintEditor;
 public:
-	FWidgetTemplatesSummoner(TSharedPtr<class FBlueprintEditor> InBlueprintEditor)
+	FWidgetTemplatesSummoner(TSharedPtr<class FWidgetBlueprintEditor> InBlueprintEditor)
 		: FWorkflowTabFactory(UMGWidgetTemplatesTabID, InBlueprintEditor)
 		, BlueprintEditor(InBlueprintEditor)
 	{
@@ -125,6 +124,36 @@ public:
 	}
 };
 
+/////////////////////////////////////////////////////
+// FSequencerSummoner
+
+static const FName FSequencerSummonerTabID(TEXT("Sequencer"));
+
+struct FSequencerSummoner : public FWorkflowTabFactory
+{
+protected:
+	TWeakPtr<class FWidgetBlueprintEditor> BlueprintEditor;
+public:
+	FSequencerSummoner(TSharedPtr<class FWidgetBlueprintEditor> InBlueprintEditor)
+		: FWorkflowTabFactory(FSequencerSummonerTabID, InBlueprintEditor)
+		, BlueprintEditor(InBlueprintEditor)
+	{
+		TabLabel = LOCTEXT("SequencerLabel", "Timeline");
+
+		bIsSingleton = true;
+
+		ViewMenuDescription = LOCTEXT("Sequencer_ViewMenu_Desc", "Timeline");
+		ViewMenuTooltip = LOCTEXT("Sequencer_ViewMenu_ToolTip", "Show the Animation editor");
+	}
+
+	virtual TSharedRef<SWidget> CreateTabBody(const FWorkflowTabSpawnInfo& Info) const OVERRIDE
+	{
+		TSharedPtr<FWidgetBlueprintEditor> BlueprintEditorPinned = BlueprintEditor.Pin();
+
+		return BlueprintEditorPinned->GetSequencer()->GetSequencerWidget();
+	}
+};
+
 
 //////////////////////////////////////////////////////////////////////////
 // FComponentsEditorModeOverride
@@ -143,10 +172,15 @@ public:
 
 		//BlueprintComponentsTabFactories.UnregisterFactory(FBlueprintEditorTabs::ConstructionScriptEditorID);
 		//BlueprintDefaultsTabFactories.UnregisterFactory(FBlueprintEditorTabs::SCSViewportID);
-		BlueprintDefaultsTabFactories.RegisterFactory(MakeShareable(new FSelectionDetailsSummoner(InBlueprintEditor)));
-		BlueprintDefaultsTabFactories.RegisterFactory(MakeShareable(new FSlatePreviewSummoner(InBlueprintEditor)));
-		BlueprintDefaultsTabFactories.RegisterFactory(MakeShareable(new FSlateTreeSummoner(InBlueprintEditor)));
-		BlueprintDefaultsTabFactories.RegisterFactory(MakeShareable(new FWidgetTemplatesSummoner(InBlueprintEditor)));
+
+		TSharedPtr<FWidgetBlueprintEditor> WidgetBlueprintEditor = StaticCastSharedPtr<FWidgetBlueprintEditor>(InBlueprintEditor);
+
+
+		BlueprintDefaultsTabFactories.RegisterFactory(MakeShareable(new FSelectionDetailsSummoner(WidgetBlueprintEditor)));
+		BlueprintDefaultsTabFactories.RegisterFactory(MakeShareable(new FSlatePreviewSummoner(WidgetBlueprintEditor)));
+		BlueprintDefaultsTabFactories.RegisterFactory(MakeShareable(new FSlateTreeSummoner(WidgetBlueprintEditor)));
+		BlueprintDefaultsTabFactories.RegisterFactory(MakeShareable(new FWidgetTemplatesSummoner(WidgetBlueprintEditor)));
+		BlueprintDefaultsTabFactories.RegisterFactory(MakeShareable(new FSequencerSummoner(WidgetBlueprintEditor)));
 
 		TabLayout = FTabManager::NewLayout( "Standalone_UMGEditor_Layout_v2" )
 		->AddArea
@@ -207,6 +241,13 @@ public:
 					)
 				)
 			)
+			->Split
+			(
+				FTabManager::NewStack()
+				->SetSizeCoefficient(0.2f)
+				->SetHideTabWell(true)
+				->AddTab( FSequencerSummonerTabID, ETabState::OpenedTab )
+			)
 		);
 	}
 
diff --git a/Engine/Source/Editor/UMGEditor/Private/UMGEditorModule.cpp b/Engine/Source/Editor/UMGEditor/Private/UMGEditorModule.cpp
--- a/Engine/Source/Editor/UMGEditor/Private/UMGEditorModule.cpp
+++ b/Engine/Source/Editor/UMGEditor/Private/UMGEditorModule.cpp
@@ -96,9 +96,6 @@ private:
 
 	/** All created asset type actions.  Cached here so that we can unregister it during shutdown. */
 	TArray< TSharedPtr<IAssetTypeActions> > CreatedAssetTypeActions;
-
-	/** List of open UMGEditor toolkits */
-	TArray<FUMGEditor*> UMGEditorToolkits;
 };
 
 IMPLEMENT_MODULE(FUMGEditorModule, UMGEditor);
diff --git a/Engine/Source/Editor/UMGEditor/Private/UMGEditorPrivatePCH.h b/Engine/Source/Editor/UMGEditor/Private/UMGEditorPrivatePCH.h
--- a/Engine/Source/Editor/UMGEditor/Private/UMGEditorPrivatePCH.h
+++ b/Engine/Source/Editor/UMGEditor/Private/UMGEditorPrivatePCH.h
@@ -14,7 +14,6 @@
 #include "UMGEditorClasses.h"
 
 #include "AssetTypeActions_Base.h"
-#include "AssetTypeActions_Blueprint.h"
 
 #include "DesignerExtension.h"
 #include "VerticalSlotExtension.h"
diff --git a/Engine/Source/Editor/UMGEditor/Private/WidgetBlueprintCompiler.cpp b/Engine/Source/Editor/UMGEditor/Private/WidgetBlueprintCompiler.cpp
--- a/Engine/Source/Editor/UMGEditor/Private/WidgetBlueprintCompiler.cpp
+++ b/Engine/Source/Editor/UMGEditor/Private/WidgetBlueprintCompiler.cpp
@@ -624,7 +624,11 @@ void FWidgetBlueprintCompiler::FinishCompilingClass(UClass* Class)
 
 	UWidgetBlueprintGeneratedClass* BPGClass = CastChecked<UWidgetBlueprintGeneratedClass>(Class);
 	BPGClass->WidgetTree = NULL;
+
+	// @todo UMG Possibly need duplication here 
 	BPGClass->WidgetTree = Blueprint->WidgetTree;
+	BPGClass->AnimationData = Blueprint->AnimationData;
+
 	BPGClass->Bindings.Reset();
 
 	for ( const FDelegateEditorBinding& EditorBinding : Blueprint->Bindings )
diff --git a/Engine/Source/Editor/UMGEditor/Private/WidgetBlueprintEditor.cpp b/Engine/Source/Editor/UMGEditor/Private/WidgetBlueprintEditor.cpp
--- a/Engine/Source/Editor/UMGEditor/Private/WidgetBlueprintEditor.cpp
+++ b/Engine/Source/Editor/UMGEditor/Private/WidgetBlueprintEditor.cpp
@@ -4,11 +4,14 @@
 
 #include "WidgetBlueprintEditor.h"
 
+#include "Editor/Sequencer/Public/ISequencerModule.h"
+
 #define LOCTEXT_NAMESPACE "UMG"
 
 FWidgetBlueprintEditor::FWidgetBlueprintEditor()
 	: PreviewScene(FPreviewScene::ConstructionValues().AllowAudioPlayback(true).ShouldSimulatePhysics(true))
 	, PreviewBlueprint(NULL)
+	, DefaultMovieScene(NULL)
 {
 }
 
@@ -66,6 +69,15 @@ void FWidgetBlueprintEditor::NotifyPostChange(const FPropertyChangedEvent& Prope
 	UpdatePreview(GetWidgetBlueprintObj(), true);
 }
 
+void FWidgetBlueprintEditor::AddReferencedObjects( FReferenceCollector& Collector )
+{
+	FBlueprintEditor::AddReferencedObjects( Collector );
+	if( DefaultMovieScene != nullptr )
+	{
+		Collector.AddReferencedObject( DefaultMovieScene );
+	}
+}
+
 UWidgetBlueprint* FWidgetBlueprintEditor::GetWidgetBlueprintObj() const
 {
 	return Cast<UWidgetBlueprint>(GetBlueprintObj());
@@ -92,6 +104,19 @@ UUserWidget* FWidgetBlueprintEditor::GetPreview() const
 	return PreviewWidgetActorPtr.Get();
 }
 
+TSharedPtr<ISequencer>& FWidgetBlueprintEditor::GetSequencer()
+{
+	if(!Sequencer.IsValid())
+	{
+		UWidgetBlueprint* Blueprint = GetWidgetBlueprintObj();
+
+		UMovieScene* MovieScene = Blueprint->AnimationData.Num() ? Blueprint->AnimationData[0] : GetDefaultMovieScene();
+		Sequencer = FModuleManager::LoadModuleChecked< ISequencerModule >("Sequencer").CreateSequencer(MovieScene);
+	}
+
+	return Sequencer;
+}
+
 void FWidgetBlueprintEditor::DestroyPreview()
 {
 	UUserWidget* PreviewActor = GetPreview();
@@ -145,4 +170,14 @@ void FWidgetBlueprintEditor::UpdatePreview(UBlueprint* InBlueprint, bool bInForc
 	}
 }
 
+UMovieScene* FWidgetBlueprintEditor::GetDefaultMovieScene()
+{
+	if( !DefaultMovieScene )
+	{
+		DefaultMovieScene = ConstructObject<UMovieScene>( UMovieScene::StaticClass(), GetTransientPackage(), "DefaultAnimationData" );
+	}
+
+	return DefaultMovieScene;
+}
+
 #undef LOCTEXT_NAMESPACE
diff --git a/Engine/Source/Editor/UMGEditor/Private/WidgetBlueprintEditor.h b/Engine/Source/Editor/UMGEditor/Private/WidgetBlueprintEditor.h
--- a/Engine/Source/Editor/UMGEditor/Private/WidgetBlueprintEditor.h
+++ b/Engine/Source/Editor/UMGEditor/Private/WidgetBlueprintEditor.h
@@ -6,6 +6,8 @@
 
 #include "PreviewScene.h"
 
+class ISequencer;
+
 /**
  * widget blueprint editor (extends Blueprint editor)
  */
@@ -21,21 +23,39 @@ public:
 	virtual void Tick(float DeltaTime) OVERRIDE;
 	virtual void NotifyPostChange(const FPropertyChangedEvent& PropertyChangedEvent, UProperty* PropertyThatChanged) OVERRIDE;
 
+	/** FGCObjectInterface */
+	virtual void AddReferencedObjects( FReferenceCollector& Collector ) OVERRIDE;
+
 	class UWidgetBlueprint* GetWidgetBlueprintObj() const;
 
 	UUserWidget* GetPreview() const;
 
+	/** @return The sequencer used to create widget animations */
+	TSharedPtr<ISequencer>& GetSequencer();
+
 private:
 	void OnBlueprintChanged(UBlueprint* InBlueprint);
 
 	void DestroyPreview();
 	void UpdatePreview(UBlueprint* InBlueprint, bool bInForceFullUpdate);
 
+	/**
+	 * Gets the default movie scene which is used when there is no 
+	 * animation data present on the widget blueprint
+	 */
+	UMovieScene* GetDefaultMovieScene();
+
 private:
 	FPreviewScene PreviewScene;
 
+	/** Sequencer for creating and previewing widget animations */
+	TSharedPtr<ISequencer> Sequencer;
+
 	/** The Blueprint associated with the current preview */
 	UBlueprint* PreviewBlueprint;
 
+	/** Default movie scene for new animations */
+	UMovieScene* DefaultMovieScene;
+
 	mutable TWeakObjectPtr<UUserWidget> PreviewWidgetActorPtr;
 };
diff --git a/Engine/Source/Editor/UMGEditor/UMGEditor.Build.cs b/Engine/Source/Editor/UMGEditor/UMGEditor.Build.cs
--- a/Engine/Source/Editor/UMGEditor/UMGEditor.Build.cs
+++ b/Engine/Source/Editor/UMGEditor/UMGEditor.Build.cs
@@ -8,9 +8,6 @@ public class UMGEditor : ModuleRules
 	{
 		PrivateIncludePaths.AddRange(
 			new string[] {
-				"Editor/GraphEditor/Private",
-				"Editor/Kismet/Private",
-				"Developer/AssetTools/Private/AssetTypeActions",
 				"Editor/UMGEditor/Private", // For PCH includes (because they don't work with relative paths, yet)
 				"Editor/UMGEditor/Private/Templates",
                 "Editor/UMGEditor/Private/Extensions",
@@ -23,20 +20,13 @@ public class UMGEditor : ModuleRules
 				"UMG",
 			});
 
-		PublicIncludePaths.AddRange(
-			new string[] {
-				"Editor/DistCurveEditor/Public",
-				"Editor/UnrealEd/Public",
-			}
-			);
-
 		PrivateDependencyModuleNames.AddRange(
 			new string[] {
 				"Core",
 				"CoreUObject",
+				"InputCore",
 				"Slate",
 				"Engine",
-				"InputCore",
 				"AssetTools",
 				"UnrealEd", // for FAssetEditorManager
 				"KismetWidgets",
@@ -45,24 +35,21 @@ public class UMGEditor : ModuleRules
 				"GraphEditor",
 				"Kismet",  // for FWorkflowCentricApplication
 				"PropertyEditor",
-				"RenderCore",
-				"LevelEditor", // for EdModes to get a toolkit host  //@TODO: PAPER: Should be a better way to do this (see the @todo in EdModeTileMap.cpp)
 				"UMG",
-				"DistCurveEditor",
-				"ContentBrowser",
-				"WorkspaceMenuStructure",
 				"EditorStyle",
 				"Slate",
 				"SlateCore",
+				"MovieSceneCore",
+				"Sequencer",
 			}
 			);
 
 		DynamicallyLoadedModuleNames.AddRange(
 			new string[] {
-				"MainFrame",
 				"PropertyEditor",
 				"AssetTools",
 				"UMG",
+				"Sequencer",
 			}
 			);
 	}
diff --git a/Engine/Source/Runtime/MovieSceneCore/Classes/MovieScene.h b/Engine/Source/Runtime/MovieSceneCore/Classes/MovieScene.h
--- a/Engine/Source/Runtime/MovieSceneCore/Classes/MovieScene.h
+++ b/Engine/Source/Runtime/MovieSceneCore/Classes/MovieScene.h
@@ -44,15 +44,6 @@ public:
 		return CounterpartGamePreviewObject;
 	}
 
-	/** ICPPStructOps implementation */
-	bool Serialize(FArchive& Ar);
-
-	friend FArchive& operator<<(FArchive& Ar,FMovieSceneSpawnable& P)
-	{
-		P.Serialize(Ar);
-		return Ar;
-	}
-
 private:
 
 	/** Guid */
@@ -77,16 +68,6 @@ private:
 	FWeakObjectPtr CounterpartGamePreviewObject;
 };
 
-template<>
-struct TStructOpsTypeTraits< FMovieSceneSpawnable > : public TStructOpsTypeTraitsBase
-{
-	enum 
-	{
-		WithSerializer = true,
-	};
-};
-
-
 /**
  * MovieScenePossessable is a "typed slot" used to allow the MovieScene to control an already-existing object
  */
diff --git a/Engine/Source/Runtime/MovieSceneCore/Private/MovieScene.cpp b/Engine/Source/Runtime/MovieSceneCore/Private/MovieScene.cpp
--- a/Engine/Source/Runtime/MovieSceneCore/Private/MovieScene.cpp
+++ b/Engine/Source/Runtime/MovieSceneCore/Private/MovieScene.cpp
@@ -4,29 +4,6 @@
 #include "MovieScene.h"
 
 
-bool FMovieSceneSpawnable::Serialize(FArchive& Ar)
-{
-	Ar << Guid;
-	Ar << Name;
-
-#if WITH_EDITOR
-	if (Ar.IsLoading() && Ar.UE4Ver() < VER_UE4_EDITORONLY_BLUEPRINTS)
-	{
-		UBlueprint* ObsoleteBlueprint = NULL;
-		Ar << ObsoleteBlueprint;
-		if(ObsoleteBlueprint)
-		{
-			GeneratedClass = ObsoleteBlueprint->GeneratedClass;
-		}
-	}
-	else
-#endif //WITH_EDITOR
-	{
-		Ar << GeneratedClass;
-	}
-	return true;
-}
-
 FMovieSceneSpawnable::FMovieSceneSpawnable( const FString& InitName, UClass* InitClass, UObject* InitCounterpartGamePreviewObject )
 {
 	Guid = FGuid::NewGuid();
diff --git a/Engine/Source/Runtime/UMG/Private/WidgetBlueprint.cpp b/Engine/Source/Runtime/UMG/Private/WidgetBlueprint.cpp
--- a/Engine/Source/Runtime/UMG/Private/WidgetBlueprint.cpp
+++ b/Engine/Source/Runtime/UMG/Private/WidgetBlueprint.cpp
@@ -1,6 +1,7 @@
 // Copyright 1998-2014 Epic Games, Inc. All Rights Reserved.
 
 #include "UMGPrivatePCH.h"
+#include "Runtime/MovieSceneCore/Classes/MovieScene.h"
 
 #define LOCTEXT_NAMESPACE "UMG"
 
diff --git a/Engine/Source/Runtime/UMG/Public/Blueprint/WidgetBlueprint.h b/Engine/Source/Runtime/UMG/Public/Blueprint/WidgetBlueprint.h
--- a/Engine/Source/Runtime/UMG/Public/Blueprint/WidgetBlueprint.h
+++ b/Engine/Source/Runtime/UMG/Public/Blueprint/WidgetBlueprint.h
@@ -33,6 +33,8 @@ struct UMG_API FDelegateEditorBinding
 	FDelegateRuntimeBinding ToRuntimeBinding(class UWidgetBlueprint* Blueprint) const;
 };
 
+class UMovieScene;
+
 /**
  * The widget blueprint enables extending UUserWidget the user extensible UWidget.
  */
@@ -49,6 +51,9 @@ public:
 	UPROPERTY()
 	TArray< FDelegateEditorBinding > Bindings;
 
+	UPROPERTY()
+	TArray< UMovieScene* > AnimationData;
+	 
 	virtual void PostLoad() OVERRIDE;
 	virtual void PostLoadSubobjects(FObjectInstancingGraph* OuterInstanceGraph) OVERRIDE;
 	
diff --git a/Engine/Source/Runtime/UMG/Public/Blueprint/WidgetBlueprintGeneratedClass.h b/Engine/Source/Runtime/UMG/Public/Blueprint/WidgetBlueprintGeneratedClass.h
--- a/Engine/Source/Runtime/UMG/Public/Blueprint/WidgetBlueprintGeneratedClass.h
+++ b/Engine/Source/Runtime/UMG/Public/Blueprint/WidgetBlueprintGeneratedClass.h
@@ -25,6 +25,7 @@ struct FDelegateRuntimeBinding
 	}*/
 };
 
+class UMovieScene;
 
 UCLASS()
 class UMG_API UWidgetBlueprintGeneratedClass : public UBlueprintGeneratedClass
@@ -38,6 +39,9 @@ class UMG_API UWidgetBlueprintGeneratedClass : public UBlueprintGeneratedClass
 	UPROPERTY()
 	TArray< FDelegateRuntimeBinding > Bindings;
 
+	UPROPERTY()
+	TArray< UMovieScene* > AnimationData;
+
 	/** This is transient data calculated at link time. */
 	TArray<UStructProperty*> WidgetNodeProperties;
 
diff --git a/Engine/Source/Runtime/UMG/UMG.Build.cs b/Engine/Source/Runtime/UMG/UMG.Build.cs
--- a/Engine/Source/Runtime/UMG/UMG.Build.cs
+++ b/Engine/Source/Runtime/UMG/UMG.Build.cs
@@ -12,19 +12,15 @@ public class UMG : ModuleRules
             })
 		;
 
-		PublicDependencyModuleNames.AddRange(
+		PrivateDependencyModuleNames.AddRange(
 			new string[] {
 				"Core",
 				"CoreUObject",
                 "Engine",
                 "InputCore",
 				"Slate",
-				"SlateCore",
-			}
-		);
-
-		PrivateDependencyModuleNames.AddRange(
-			new string[] {
+				"SlateCore",	
+				"MovieSceneCore",
 			}
 		);
 	}
