Fix Sequencer/UMG in non-unity

diff --git a/Engine/Source/Editor/Sequencer/Private/Sequencer.cpp b/Engine/Source/Editor/Sequencer/Private/Sequencer.cpp
--- a/Engine/Source/Editor/Sequencer/Private/Sequencer.cpp
+++ b/Engine/Source/Editor/Sequencer/Private/Sequencer.cpp
@@ -40,7 +40,7 @@
 
 DEFINE_LOG_CATEGORY_STATIC( LogSequencer, Log, All );
 
-static bool IsSequencerEnabled()
+bool FSequencer::IsSequencerEnabled()
 {
 	return FParse::Param( FCommandLine::Get(), TEXT( "Sequencer" ) ) || FParse::Param( FCommandLine::Get(), TEXT( "umg" ) );
 }
diff --git a/Engine/Source/Editor/Sequencer/Private/Sequencer.h b/Engine/Source/Editor/Sequencer/Private/Sequencer.h
--- a/Engine/Source/Editor/Sequencer/Private/Sequencer.h
+++ b/Engine/Source/Editor/Sequencer/Private/Sequencer.h
@@ -269,6 +269,8 @@ public:
 	/** @return The toolkit that this sequencer is hosted in (if any) */
 	TSharedPtr<IToolkitHost> GetToolkitHost() const { return ToolkitHost.Pin(); }
 
+	static bool IsSequencerEnabled();
+
 protected:
 	/**
 	 * Reset data about a movie scene when pushing or popping a movie scene
diff --git a/Engine/Source/Editor/Sequencer/Private/SequencerAssetEditor.cpp b/Engine/Source/Editor/Sequencer/Private/SequencerAssetEditor.cpp
--- a/Engine/Source/Editor/Sequencer/Private/SequencerAssetEditor.cpp
+++ b/Engine/Source/Editor/Sequencer/Private/SequencerAssetEditor.cpp
@@ -4,6 +4,7 @@
 #include "Toolkits/IToolkitHost.h"
 #include "SequencerAssetEditor.h"
 #include "Editor/LevelEditor/Public/LevelEditor.h"
+#include "Editor/WorkspaceMenuStructure/Public/WorkspaceMenuStructure.h"
 
 #define LOCTEXT_NAMESPACE "Sequencer"
 
@@ -17,7 +18,7 @@ void FSequencerAssetEditor::RegisterTabSpawners(const TSharedRef<class FTabManag
 {
 	const IWorkspaceMenuStructure& MenuStructure = WorkspaceMenu::GetMenuStructure();
 
-	if( IsSequencerEnabled() && !IsWorldCentricAssetEditor() )
+	if( FSequencer::IsSequencerEnabled() && !IsWorldCentricAssetEditor() )
 	{
 		TabManager->RegisterTabSpawner( SequencerMainTabId, FOnSpawnTab::CreateSP(this, &FSequencerAssetEditor::SpawnTab_SequencerMain) )
 			.SetDisplayName( LOCTEXT("SequencerMainTab", "Sequencer") )
@@ -28,7 +29,7 @@ void FSequencerAssetEditor::RegisterTabSpawners(const TSharedRef<class FTabManag
 
 void FSequencerAssetEditor::UnregisterTabSpawners(const TSharedRef<class FTabManager>& TabManager)
 {
-	if( IsSequencerEnabled() && !IsWorldCentricAssetEditor() )
+	if( FSequencer::IsSequencerEnabled() && !IsWorldCentricAssetEditor() )
 	{
 		TabManager->UnregisterTabSpawner( SequencerMainTabId );
 	}
diff --git a/Engine/Source/Editor/Sequencer/Private/SequencerModule.cpp b/Engine/Source/Editor/Sequencer/Private/SequencerModule.cpp
--- a/Engine/Source/Editor/Sequencer/Private/SequencerModule.cpp
+++ b/Engine/Source/Editor/Sequencer/Private/SequencerModule.cpp
@@ -5,6 +5,7 @@
 #include "Sequencer.h"
 #include "Toolkits/ToolkitManager.h"
 #include "SequencerCommands.h"
+#include "SequencerAssetEditor.h"
 
 /**
  * SequencerModule implementation (private)
diff --git a/Engine/Source/Editor/UMGEditor/Private/WidgetBlueprintEditor.cpp b/Engine/Source/Editor/UMGEditor/Private/WidgetBlueprintEditor.cpp
--- a/Engine/Source/Editor/UMGEditor/Private/WidgetBlueprintEditor.cpp
+++ b/Engine/Source/Editor/UMGEditor/Private/WidgetBlueprintEditor.cpp
@@ -3,7 +3,7 @@
 #include "UMGEditorPrivatePCH.h"
 
 #include "WidgetBlueprintEditor.h"
-
+#include "MovieScene.h"
 #include "Editor/Sequencer/Public/ISequencerModule.h"
 
 #define LOCTEXT_NAMESPACE "UMG"
