Fix scrollbar causing incorrect offset of timeline in Sequencer.  Added external scrollbar support to SScrollBox

diff --git a/Engine/Source/Editor/Sequencer/Private/SSequencer.cpp b/Engine/Source/Editor/Sequencer/Private/SSequencer.cpp
--- a/Engine/Source/Editor/Sequencer/Private/SSequencer.cpp
+++ b/Engine/Source/Editor/Sequencer/Private/SSequencer.cpp
@@ -138,7 +138,6 @@ private:
 };
 
 
-BEGIN_SLATE_FUNCTION_BUILD_OPTIMIZATION
 void SSequencer::Construct( const FArguments& InArgs, TSharedRef< class FSequencer > InSequencer )
 {
 	Sequencer = InSequencer;
@@ -180,7 +179,7 @@ void SSequencer::Construct( const FArguments& InArgs, TSharedRef< class FSequenc
 			[
 				SNew( SHorizontalBox )
 				+SHorizontalBox::Slot()
-				.Padding( FMargin( 0.0f, 2.0f, 0.0f, 0.0f ) )
+				.Padding( FMargin( 0.0f, 0.0f, 0.0f, 0.0f ) )
 				.FillWidth( TAttribute<float>( this, &SSequencer::GetAnimationOutlinerFillPercentage ) )
 				.VAlign(VAlign_Center)
 				[
@@ -205,6 +204,7 @@ void SSequencer::Construct( const FArguments& InArgs, TSharedRef< class FSequenc
 					[
 						// @todo Sequencer - Temp clean view button
 						SNew( SCheckBox )
+						.Visibility( this, &SSequencer::OnGetCleanViewVisibility )
 						.Style(FEditorStyle::Get(), "ToggleButtonCheckbox")
 						.IsChecked( this, &SSequencer::OnGetCleanViewCheckState )
 						.OnCheckStateChanged( this, &SSequencer::OnCleanViewChecked )
@@ -226,13 +226,18 @@ void SSequencer::Construct( const FArguments& InArgs, TSharedRef< class FSequenc
 			[
 				SNew( SHorizontalBox )
 				+ SHorizontalBox::Slot()
-				.Padding( FMargin( 0.0f, 2.0f, 0.0f, 0.0f ) )
+				// @todo Sequencer Do not change the paddings or the sliders scrub widgets wont line up
+				.Padding( FMargin(0) )
 				.FillWidth( TAttribute<float>( this, &SSequencer::GetAnimationOutlinerFillPercentage ) )
 				.VAlign(VAlign_Center)
 				[
-					// Search box for searching through the outliner
-					SNew( SSearchBox )
-					.OnTextChanged( this, &SSequencer::OnOutlinerSearchChanged )
+					SNew( SBox )
+					.Padding( FMargin( 0,0,4,0) )
+					[
+						// Search box for searching through the outliner
+						SNew( SSearchBox )
+						.OnTextChanged( this, &SSequencer::OnOutlinerSearchChanged )
+					]
 				]
 				+ SHorizontalBox::Slot()
 				.FillWidth(1.0f)
@@ -310,7 +315,6 @@ void SSequencer::Construct( const FArguments& InArgs, TSharedRef< class FSequenc
 
 	BreadcrumbTrail->PushCrumb(TAttribute<FText>::Create(TAttribute<FText>::FGetter::CreateSP(this, &SSequencer::GetRootMovieSceneName)), FSequencerBreadcrumb( Sequencer.Pin()->GetRootMovieSceneInstance() ));
 }
-END_SLATE_FUNCTION_BUILD_OPTIMIZATION
 
 SSequencer::~SSequencer()
 {
@@ -405,6 +409,11 @@ void SSequencer::OnAutoKeyChecked( ESlateCheckBoxState::Type InState )
 	OnToggleAutoKey.ExecuteIfBound( InState == ESlateCheckBoxState::Checked ? true : false );
 }
 
+EVisibility SSequencer::OnGetCleanViewVisibility() const 
+{
+	return Sequencer.Pin()->IsLevelEditorSequencer() ? EVisibility::Visible : EVisibility::Collapsed;
+}
+
 ESlateCheckBoxState::Type SSequencer::OnGetCleanViewCheckState() const
 {
 	return CleanViewEnabled.Get() ? ESlateCheckBoxState::Checked : ESlateCheckBoxState::Unchecked;
diff --git a/Engine/Source/Editor/Sequencer/Private/SSequencer.h b/Engine/Source/Editor/Sequencer/Private/SSequencer.h
--- a/Engine/Source/Editor/Sequencer/Private/SSequencer.h
+++ b/Engine/Source/Editor/Sequencer/Private/SSequencer.h
@@ -106,6 +106,9 @@ private:
 	 */
 	void OnAutoKeyChecked( ESlateCheckBoxState::Type InState );	
 	
+	/** @return The visibility of the clean view button */
+	EVisibility OnGetCleanViewVisibility() const;
+
 	/**
 	 * @return The state of the auto-key check box
 	 */
diff --git a/Engine/Source/Editor/Sequencer/Private/SSequencerTrackArea.cpp b/Engine/Source/Editor/Sequencer/Private/SSequencerTrackArea.cpp
--- a/Engine/Source/Editor/Sequencer/Private/SSequencerTrackArea.cpp
+++ b/Engine/Source/Editor/Sequencer/Private/SSequencerTrackArea.cpp
@@ -12,9 +12,24 @@
 void SSequencerTrackArea::Update( const FSequencerNodeTree& InSequencerNodeTree )
 {
 	// @todo Sequencer: Currently need to recreate the scroll box to remove all its children
+
+	TSharedRef<SScrollBar> ScrollBar =
+		SNew( SScrollBar )
+		.Thickness(FVector2D(5.0f, 5.0f));
+
 	ChildSlot
 	[
-		SAssignNew( ScrollBox, SScrollBox )
+		SNew( SOverlay )
+		+ SOverlay::Slot()
+		[
+			SAssignNew( ScrollBox, SScrollBox )
+			.ExternalScrollbar( ScrollBar )
+		]
+		+ SOverlay::Slot()
+		.HAlign(HAlign_Right)
+		[
+			ScrollBar
+		]
 	];
 
 	const TArray< TSharedRef<FSequencerDisplayNode> >& RootNodes = InSequencerNodeTree.GetRootNodes();
diff --git a/Engine/Source/Editor/Sequencer/Private/Sequencer.cpp b/Engine/Source/Editor/Sequencer/Private/Sequencer.cpp
--- a/Engine/Source/Editor/Sequencer/Private/Sequencer.cpp
+++ b/Engine/Source/Editor/Sequencer/Private/Sequencer.cpp
@@ -50,6 +50,8 @@ void FSequencer::InitSequencer( UMovieScene* InRootMovieScene, TSharedPtr<IToolk
 {
 	if( IsSequencerEnabled() )
 	{
+		bIsEditingWithinLevelEditor = bEditWithinLevelEditor;
+
 		ToolkitHost = InToolKitHost;
 		// Create an object change listener for various systems that need to be notified when objects change
 		ObjectChangeListener = MakeShareable( new FSequencerObjectChangeListener( SharedThis( this ) ) );
@@ -140,6 +142,7 @@ FSequencer::FSequencer()
 	, bAllowAutoKey( false )
 	, bPerspectiveViewportPossessionEnabled( true )
 	, bNeedTreeRefresh( false )
+	, bIsEditingWithinLevelEditor( false )
 {
 
 }
diff --git a/Engine/Source/Editor/Sequencer/Private/Sequencer.h b/Engine/Source/Editor/Sequencer/Private/Sequencer.h
--- a/Engine/Source/Editor/Sequencer/Private/Sequencer.h
+++ b/Engine/Source/Editor/Sequencer/Private/Sequencer.h
@@ -269,8 +269,10 @@ public:
 	/** @return The toolkit that this sequencer is hosted in (if any) */
 	TSharedPtr<IToolkitHost> GetToolkitHost() const { return ToolkitHost.Pin(); }
 
+	/** @return Whether or not this sequencer is used in the level editor */
+	bool IsLevelEditorSequencer() const { return bIsEditingWithinLevelEditor; }
+	
 	static bool IsSequencerEnabled();
-
 protected:
 	/**
 	 * Reset data about a movie scene when pushing or popping a movie scene
@@ -443,6 +445,7 @@ private:
 	TRange<float> TargetViewRange;
 	/** The last time range that was viewed */
 	TRange<float> LastViewRange;
+
 	/** Zoom smoothing curves */
 	FCurveSequence ZoomAnimation;
 	FCurveHandle ZoomCurve;
@@ -450,6 +453,9 @@ private:
 	FCurveSequence OverlayAnimation;
 	FCurveHandle OverlayCurve;
 
+	/** Whether we are playing, recording, etc. */
+	EMovieScenePlayerStatus::Type PlaybackState;
+
 	/** The current scrub position */
 	// @todo sequencer: Should use FTimespan or "double" for Time Cursor Position! (cascades)
 	float ScrubPosition;
@@ -457,8 +463,6 @@ private:
 	/** Whether the clean sequencer view is enabled */
 	bool bCleanViewEnabled;
 
-	/** Whether we are playing, recording, etc. */
-	EMovieScenePlayerStatus::Type PlaybackState;
 
 	/** Whether looping while playing is enabled for this sequencer */
 	bool bLoopingEnabled;
@@ -468,6 +472,9 @@ private:
 
 	bool bPerspectiveViewportPossessionEnabled;
 
+	/** True if this sequencer is being edited within the level editor */
+	bool bIsEditingWithinLevelEditor;
+
 	/** Stores a dirty bit for whether the sequencer tree (and other UI bits) may need to be refreshed.  We
 	    do this simply to avoid refreshing the UI more than once per frame. (e.g. during live recording where
 		the MovieScene data can change many times per frame.) */
diff --git a/Engine/Source/Runtime/Slate/Private/Widgets/Layout/SScrollBox.cpp b/Engine/Source/Runtime/Slate/Private/Widgets/Layout/SScrollBox.cpp
--- a/Engine/Source/Runtime/Slate/Private/Widgets/Layout/SScrollBox.cpp
+++ b/Engine/Source/Runtime/Slate/Private/Widgets/Layout/SScrollBox.cpp
@@ -106,9 +106,10 @@ void SScrollBox::Construct( const FArguments& InArgs )
 	bShowSoftwareCursor = false;
 	SoftwareCursorPosition = FVector2D::ZeroVector;
 
+	TSharedPtr<SHorizontalBox> PanelAndScrollbar;
 	this->ChildSlot
 	[
-		SNew(SHorizontalBox)
+		SAssignNew(PanelAndScrollbar, SHorizontalBox)
 		+SHorizontalBox::Slot()
 		.FillWidth(1)
 		[
@@ -140,15 +141,33 @@ void SScrollBox::Construct( const FArguments& InArgs )
 				.Image( &InArgs._Style->BottomShadowBrush )	
 			]
 		]
-		+SHorizontalBox::Slot()
-		.AutoWidth()
-		[
-			SAssignNew( ScrollBar, SScrollBar )
-			.Thickness( FVector2D(5.0f, 5.0f) )
-			.OnUserScrolled( this, &SScrollBox::ScrollBar_OnUserScrolled )
-		]
+	
 	];
 
+
+	if (InArgs._ExternalScrollbar.IsValid())
+	{
+		// An external scroll bar was sepecified by the user
+		ScrollBar = InArgs._ExternalScrollbar;
+		ScrollBar->SetOnUserScrolled(FOnUserScrolled::CreateSP(this, &SScrollBox::ScrollBar_OnUserScrolled));
+	}
+	else
+	{
+		// Make a scroll bar 
+		ScrollBar = 
+			SNew( SScrollBar)
+			.Thickness(FVector2D(5.0f, 5.0f))
+			.OnUserScrolled(this, &SScrollBox::ScrollBar_OnUserScrolled);
+
+		PanelAndScrollbar->AddSlot()
+			.AutoWidth()
+			[
+				ScrollBar.ToSharedRef()
+			];
+	}
+
+
+
 	ScrollBar->SetState( 0.0f, 1.0f );
 }
 
diff --git a/Engine/Source/Runtime/Slate/Public/Widgets/Layout/SScrollBox.h b/Engine/Source/Runtime/Slate/Public/Widgets/Layout/SScrollBox.h
--- a/Engine/Source/Runtime/Slate/Public/Widgets/Layout/SScrollBox.h
+++ b/Engine/Source/Runtime/Slate/Public/Widgets/Layout/SScrollBox.h
@@ -27,6 +27,7 @@ public:
 
 	SLATE_BEGIN_ARGS(SScrollBox)
 		: _Style( &FCoreStyle::Get().GetWidgetStyle<FScrollBoxStyle>("ScrollBox") )
+		, _ExternalScrollbar()
 		{}
 		
 		SLATE_SUPPORTS_SLOT( FSlot )
@@ -34,6 +35,8 @@ public:
 		/** Style used to draw this scrollbox */
 		SLATE_STYLE_ARGUMENT( FScrollBoxStyle, Style )
 
+		/** Custom scroll bar */
+		SLATE_ARGUMENT( TSharedPtr<SScrollBar>, ExternalScrollbar )
 	SLATE_END_ARGS()
 
 
