Non-unity build fixes for Sequencer, CollisionConversions etc.

diff --git a/Engine/Source/Editor/Sequencer/Private/SequencerAssetEditor.cpp b/Engine/Source/Editor/Sequencer/Private/SequencerAssetEditor.cpp
--- a/Engine/Source/Editor/Sequencer/Private/SequencerAssetEditor.cpp
+++ b/Engine/Source/Editor/Sequencer/Private/SequencerAssetEditor.cpp
@@ -4,7 +4,12 @@
 #include "Toolkits/IToolkitHost.h"
 #include "SequencerAssetEditor.h"
 #include "Editor/LevelEditor/Public/LevelEditor.h"
+#include "Editor/LevelEditor/Public/ILevelViewport.h"
 #include "Editor/WorkspaceMenuStructure/Public/WorkspaceMenuStructure.h"
+#include "Editor/WorkspaceMenuStructure/Public/WorkspaceMenuStructureModule.h"
+#include "Sequencer.h"
+#include "MovieScene.h"
+#include "EditorWidgetsModule.h"
 
 #define LOCTEXT_NAMESPACE "Sequencer"
 
diff --git a/Engine/Source/Editor/UMGEditor/Private/SUMGEditorTree.cpp b/Engine/Source/Editor/UMGEditor/Private/SUMGEditorTree.cpp
--- a/Engine/Source/Editor/UMGEditor/Private/SUMGEditorTree.cpp
+++ b/Engine/Source/Editor/UMGEditor/Private/SUMGEditorTree.cpp
@@ -13,6 +13,7 @@
 #include "SKismetInspector.h"
 #include "BlueprintEditorUtils.h"
 #include "WidgetTemplateClass.h"
+#include "WidgetBlueprintEditor.h"
 
 #define LOCTEXT_NAMESPACE "UMG"
 
diff --git a/Engine/Source/Runtime/Engine/Private/Collision/CollisionConversions.cpp b/Engine/Source/Runtime/Engine/Private/Collision/CollisionConversions.cpp
--- a/Engine/Source/Runtime/Engine/Private/Collision/CollisionConversions.cpp
+++ b/Engine/Source/Runtime/Engine/Private/Collision/CollisionConversions.cpp
@@ -4,7 +4,6 @@
 
 #if WITH_PHYSX
 #include "PhysicsPublic.h"
-#include "../PhysicsEngine/PhysXSupport.h"
 #include "PhysXCollision.h"
 #include "CollisionDebugDrawing.h"
 #include "CollisionConversions.h"
diff --git a/Engine/Source/Runtime/Engine/Private/Collision/CollisionConversions.h b/Engine/Source/Runtime/Engine/Private/Collision/CollisionConversions.h
--- a/Engine/Source/Runtime/Engine/Private/Collision/CollisionConversions.h
+++ b/Engine/Source/Runtime/Engine/Private/Collision/CollisionConversions.h
@@ -5,6 +5,8 @@
 
 #if WITH_PHYSX
 
+#include "../PhysicsEngine/PhysXSupport.h"
+
 /** 
  * Util to convert single physX hit (raycast or sweep) to our hit result
  *
diff --git a/Engine/Source/Runtime/Engine/Private/Collision/PhysXCollision.h b/Engine/Source/Runtime/Engine/Private/Collision/PhysXCollision.h
--- a/Engine/Source/Runtime/Engine/Private/Collision/PhysXCollision.h
+++ b/Engine/Source/Runtime/Engine/Private/Collision/PhysXCollision.h
@@ -8,6 +8,7 @@
 
 #if WITH_PHYSX
 
+#include "../PhysicsEngine/PhysXSupport.h"
 
 /** Temporary result buffer size */
 #define HIT_BUFFER_SIZE							512		// Hit buffer size for traces and sweeps. This is the total size allowed for sync + async tests.
