Basic UMG animation work
- Refactored sequencer to take in an object binding manager to manage the handles to objects that are being used at runtime.
- Wrote binding manager for UMG
- Added basic UFunction calling instead of setting properties directly

diff --git a/Engine/Source/Developer/MovieSceneTools/Private/AnimationTrackEditor.cpp b/Engine/Source/Developer/MovieSceneTools/Private/AnimationTrackEditor.cpp
--- a/Engine/Source/Developer/MovieSceneTools/Private/AnimationTrackEditor.cpp
+++ b/Engine/Source/Developer/MovieSceneTools/Private/AnimationTrackEditor.cpp
@@ -239,9 +239,10 @@ void FAnimationTrackEditor::AddKeyInternal( float KeyTime, const TArray<UObject*
 	{
 		UObject* Object = Objects[ObjectIndex];
 
-		if (Object)
+		FGuid ObjectHandle = FindOrCreateHandleToObject( Object );
+		if (ObjectHandle.IsValid())
 		{
-			UMovieSceneTrack* Track = GetTrackForObject( Object, UMovieSceneAnimationTrack::StaticClass(), FName("Animation"));
+			UMovieSceneTrack* Track = GetTrackForObject( ObjectHandle, UMovieSceneAnimationTrack::StaticClass(), FName("Animation"));
 
 			if (ensure(Track))
 			{
diff --git a/Engine/Source/Developer/MovieSceneTools/Private/AudioTrackEditor.cpp b/Engine/Source/Developer/MovieSceneTools/Private/AudioTrackEditor.cpp
--- a/Engine/Source/Developer/MovieSceneTools/Private/AudioTrackEditor.cpp
+++ b/Engine/Source/Developer/MovieSceneTools/Private/AudioTrackEditor.cpp
@@ -498,9 +498,11 @@ void FAudioTrackEditor::AddNewAttachedSound( float KeyTime, USoundBase* Sound, T
 	{
 		UObject* Object = ObjectsToAttachTo[ObjectIndex];
 
-		if (Object)
+		FGuid ObjectHandle = FindOrCreateHandleToObject( Object );
+		if (ObjectHandle.IsValid())
 		{
-			UMovieSceneTrack* Track = GetTrackForObject( Object, UMovieSceneAudioTrack::StaticClass(), AudioTrackConstants::UniqueTrackName );
+			
+			UMovieSceneTrack* Track = GetTrackForObject( ObjectHandle, UMovieSceneAudioTrack::StaticClass(), AudioTrackConstants::UniqueTrackName );
 
 			if (ensure(Track))
 			{
diff --git a/Engine/Source/Developer/MovieSceneTools/Private/ParticleTrackEditor.cpp b/Engine/Source/Developer/MovieSceneTools/Private/ParticleTrackEditor.cpp
--- a/Engine/Source/Developer/MovieSceneTools/Private/ParticleTrackEditor.cpp
+++ b/Engine/Source/Developer/MovieSceneTools/Private/ParticleTrackEditor.cpp
@@ -169,9 +169,10 @@ void FParticleTrackEditor::AddKeyInternal( float KeyTime, const TArray<UObject*>
 	{
 		UObject* Object = Objects[ObjectIndex];
 
-		if (Object)
+		FGuid ObjectHandle = FindOrCreateHandleToObject( Object );
+		if (ObjectHandle.IsValid())
 		{
-			UMovieSceneTrack* Track = GetTrackForObject( Object, UMovieSceneParticleTrack::StaticClass(), FName("ParticleSystem"));
+			UMovieSceneTrack* Track = GetTrackForObject( ObjectHandle, UMovieSceneParticleTrack::StaticClass(), FName("ParticleSystem"));
 
 			if (ensure(Track))
 			{
diff --git a/Engine/Source/Developer/MovieSceneTools/Private/PropertyTrackEditor.cpp b/Engine/Source/Developer/MovieSceneTools/Private/PropertyTrackEditor.cpp
--- a/Engine/Source/Developer/MovieSceneTools/Private/PropertyTrackEditor.cpp
+++ b/Engine/Source/Developer/MovieSceneTools/Private/PropertyTrackEditor.cpp
@@ -487,11 +487,10 @@ void FPropertyTrackEditor::OnKeyProperty( float KeyTime, FName PropertyName, con
 	{
 		UObject* Object = InObjectsThatChanged[ObjectIndex];
 
-		if( Object )
+		FGuid ObjectHandle = FindOrCreateHandleToObject( Object );
+		if (ObjectHandle.IsValid())
 		{
-			bool bTrackExists = TrackForObjectExists(Object, TrackType::StaticClass(), PropertyName);
-
-			UMovieSceneTrack* Track = GetTrackForObject( Object, TrackType::StaticClass(), PropertyName );
+			UMovieSceneTrack* Track = GetTrackForObject( ObjectHandle, TrackType::StaticClass(), PropertyName );
 			if( ensure( Track ) )
 			{
 				TrackType* TypedTrack = CastChecked<TrackType>(Track);
@@ -499,7 +498,7 @@ void FPropertyTrackEditor::OnKeyProperty( float KeyTime, FName PropertyName, con
 				// Find or add a new section at the auto-key time and changing the property same property
 				// AddKeyToSection is not actually a virtual, it's redefined in each class with a different type
 				bool bSuccessfulAdd = TypedTrack->AddKeyToSection( KeyTime, *Value );
-				if (bSuccessfulAdd && (bAutoKeying || bTrackExists))
+				if (bSuccessfulAdd)
 				{
 					TypedTrack->SetAsShowable();
 				}
@@ -507,3 +506,4 @@ void FPropertyTrackEditor::OnKeyProperty( float KeyTime, FName PropertyName, con
 		}
 	}
 }
+
diff --git a/Engine/Source/Developer/MovieSceneTools/Private/TransformTrackEditor.cpp b/Engine/Source/Developer/MovieSceneTools/Private/TransformTrackEditor.cpp
--- a/Engine/Source/Developer/MovieSceneTools/Private/TransformTrackEditor.cpp
+++ b/Engine/Source/Developer/MovieSceneTools/Private/TransformTrackEditor.cpp
@@ -212,10 +212,8 @@ void FTransformTrackEditor::AddKey(const FGuid& ObjectGuid, UObject* AdditionalA
 	TArray<UObject*> OutObjects;
 	GetSequencer()->GetRuntimeObjects( GetSequencer()->GetFocusedMovieSceneInstance(), ObjectGuid, OutObjects);
 
-	for (int32 i = 0; i < OutObjects.Num(); ++i)
+	for ( UObject* Object : OutObjects )
 	{
-		UObject* Object = OutObjects[i];
-		
 		FGuid ObjectHandle;
 		USceneComponent* SceneComponent = NULL;
 		AActor* Actor = Cast<AActor>( Object );
@@ -254,23 +252,24 @@ void FTransformTrackEditor::OnTransformChangedInternals(float KeyTime, UObject*
 	const bool bUnwindRotation = GetSequencer()->IsRecordingLive();
 
 	FName Transform("Transform");
-	bool bTrackExists = TrackForObjectExists(InObject, UMovieSceneTransformTrack::StaticClass(), Transform);
-
-	UMovieSceneTrack* Track = GetTrackForObject( InObject, UMovieSceneTransformTrack::StaticClass(), Transform );
-	UMovieSceneTransformTrack* TransformTrack = CastChecked<UMovieSceneTransformTrack>( Track );
-	TransformTrack->SetPropertyName( Transform );
-	
-	if (!TransformPair.LastTransformData.IsValid())
+	if (ObjectHandle.IsValid())
 	{
-		bool bHasTranslationKeys = false, bHasRotationKeys = false, bHasScaleKeys = false;
-		TransformPair.LastTransformData.bValid = TransformTrack->Eval(KeyTime, KeyTime, TransformPair.LastTransformData.Translation, TransformPair.LastTransformData.Rotation, TransformPair.LastTransformData.Scale, bHasTranslationKeys, bHasRotationKeys, bHasScaleKeys);
-	}
+		UMovieSceneTrack* Track = GetTrackForObject( ObjectHandle, UMovieSceneTransformTrack::StaticClass(), Transform );
+		UMovieSceneTransformTrack* TransformTrack = CastChecked<UMovieSceneTransformTrack>( Track );
+		TransformTrack->SetPropertyName( Transform );
+	
+		if (!TransformPair.LastTransformData.IsValid())
+		{
+			bool bHasTranslationKeys = false, bHasRotationKeys = false, bHasScaleKeys = false;
+			TransformPair.LastTransformData.bValid = TransformTrack->Eval(KeyTime, KeyTime, TransformPair.LastTransformData.Translation, TransformPair.LastTransformData.Rotation, TransformPair.LastTransformData.Scale, bHasTranslationKeys, bHasRotationKeys, bHasScaleKeys);
+		}
 
-	FTransformKey TransformKey = FTransformKey(KeyTime, TransformPair.TransformData, TransformPair.LastTransformData);
+		FTransformKey TransformKey = FTransformKey(KeyTime, TransformPair.TransformData, TransformPair.LastTransformData);
 
-	bool bSuccessfulAdd = TransformTrack->AddKeyToSection( ObjectHandle, TransformKey, bUnwindRotation );
-	if (bSuccessfulAdd && (bAutoKeying || bTrackExists))
-	{
-		TransformTrack->SetAsShowable();
+		bool bSuccessfulAdd = TransformTrack->AddKeyToSection( ObjectHandle, TransformKey, bUnwindRotation );
+		if (bSuccessfulAdd)
+		{
+			TransformTrack->SetAsShowable();
+		}
 	}
 }
diff --git a/Engine/Source/Editor/Sequencer/Private/SSequencer.cpp b/Engine/Source/Editor/Sequencer/Private/SSequencer.cpp
--- a/Engine/Source/Editor/Sequencer/Private/SSequencer.cpp
+++ b/Engine/Source/Editor/Sequencer/Private/SSequencer.cpp
@@ -568,8 +568,8 @@ void SSequencer::OnAssetsDropped( const FAssetDragDropOp& DragDropOp )
 
 		if (!SequencerRef.OnHandleAssetDropped(CurObject, TargetObjectGuid))
 		{
-			SequencerRef.AddSpawnableForAssetOrClass( CurObject, NULL );
-			bSpawnableAdded = true;
+			FGuid NewGuid = SequencerRef.AddSpawnableForAssetOrClass( CurObject, NULL );
+			bSpawnableAdded = NewGuid.IsValid();
 		}
 		bObjectAdded = true;
 	}
@@ -600,9 +600,9 @@ void SSequencer::OnClassesDropped( const FClassDragDropOp& DragDropOp )
 		{
 			UObject* Object = Class->GetDefaultObject();
 
-			SequencerRef.AddSpawnableForAssetOrClass( Object, NULL );
+			FGuid NewGuid = SequencerRef.AddSpawnableForAssetOrClass( Object, NULL );
 
-			bSpawnableAdded = true;
+			bSpawnableAdded = NewGuid.IsValid();
 
 		}
 	}
@@ -657,8 +657,8 @@ void SSequencer::OnUnloadedClassesDropped( const FUnloadedClassDragDropOp& DragD
 
 		if( Object != NULL )
 		{
-			SequencerRef.AddSpawnableForAssetOrClass( Object, NULL );
-			bSpawnableAdded = true;
+			FGuid NewGuid = SequencerRef.AddSpawnableForAssetOrClass( Object, NULL );
+			bSpawnableAdded = NewGuid.IsValid();
 		}
 	}
 
@@ -673,50 +673,7 @@ void SSequencer::OnUnloadedClassesDropped( const FUnloadedClassDragDropOp& DragD
 
 void SSequencer::OnActorsDropped( FActorDragDropGraphEdOp& DragDropOp )
 {
-	FSequencer& SequencerRef = *Sequencer.Pin();
-	bool bPossessableAdded = false;
-	for( auto ActorIter = DragDropOp.Actors.CreateIterator(); ActorIter; ++ActorIter )
-	{
-		AActor* Actor = ( *ActorIter ).Get();
-		if( Actor != NULL )
-		{
-			// Grab the MovieScene that is currently focused.  We'll add our Blueprint as an inner of the
-			// MovieScene asset.
-			UMovieScene* OwnerMovieScene = SequencerRef.GetFocusedMovieScene();
-
-			// @todo sequencer: Undo doesn't seem to be working at all
-			const FScopedTransaction Transaction( LOCTEXT("UndoPossessingObject", "Possess Object with MovieScene") );
-
-			// Possess the object!
-			{
-				// Update level script
-				const bool bCreateIfNotFound = true;
-				UK2Node_PlayMovieScene* PlayMovieSceneNode = SequencerRef.BindToPlayMovieSceneNode( bCreateIfNotFound );
-
-				// Create a new possessable
-				OwnerMovieScene->Modify();
-				const FString PossessableName = Actor->GetActorLabel();
-				const FGuid PossessableGuid = OwnerMovieScene->AddPossessable( PossessableName, Actor->GetClass() );
-
-				if (Sequencer.Pin()->IsShotFilteringOn()) {Sequencer.Pin()->AddUnfilterableObject(PossessableGuid);}
-
-				// Bind the actor to the new possessable
-				// @todo sequencer: Handle Undo properly here (new UObjects being created and modified)
-				TArray< UObject* > Actors;
-				Actors.Add( Actor );
-				PlayMovieSceneNode->BindPossessableToObjects( PossessableGuid, Actors );
-
-				bPossessableAdded = true;
-			}
-		}
-	}
-
-	if( bPossessableAdded )
-	{
-		SequencerRef.SpawnOrDestroyPuppetObjects( SequencerRef.GetFocusedMovieSceneInstance() );
-
-		SequencerRef.NotifyMovieSceneDataChanged();
-	}
+	Sequencer.Pin()->OnActorsDropped( DragDropOp.Actors );
 }
 
 void SSequencer::OnCrumbClicked(const FSequencerBreadcrumb& Item)
diff --git a/Engine/Source/Editor/Sequencer/Private/SSequencerTrackArea.cpp b/Engine/Source/Editor/Sequencer/Private/SSequencerTrackArea.cpp
--- a/Engine/Source/Editor/Sequencer/Private/SSequencerTrackArea.cpp
+++ b/Engine/Source/Editor/Sequencer/Private/SSequencerTrackArea.cpp
@@ -8,22 +8,23 @@
 #include "TimeSliderController.h"
 #include "SSequencerSectionAreaView.h"
 
-
-void SSequencerTrackArea::Update( const FSequencerNodeTree& InSequencerNodeTree )
+void SSequencerTrackArea::Construct( const FArguments& InArgs, TSharedRef<FSequencer> InSequencer )
 {
-	// @todo Sequencer: Currently need to recreate the scroll box to remove all its children
+	ViewRange = InArgs._ViewRange;
+	OutlinerFillPercent = InArgs._OutlinerFillPercent;
+	Sequencer = InSequencer;
 
 	TSharedRef<SScrollBar> ScrollBar =
-		SNew( SScrollBar )
+		SNew(SScrollBar)
 		.Thickness(FVector2D(5.0f, 5.0f));
 
 	ChildSlot
 	[
-		SNew( SOverlay )
+		SNew(SOverlay)
 		+ SOverlay::Slot()
 		[
-			SAssignNew( ScrollBox, SScrollBox )
-			.ExternalScrollbar( ScrollBar )
+			SAssignNew(ScrollBox, SScrollBox)
+			.ExternalScrollbar(ScrollBar)
 		]
 		+ SOverlay::Slot()
 		.HAlign(HAlign_Right)
@@ -31,6 +32,11 @@ void SSequencerTrackArea::Update( const FSequencerNodeTree& InSequencerNodeTree
 			ScrollBar
 		]
 	];
+}
+
+void SSequencerTrackArea::Update( const FSequencerNodeTree& InSequencerNodeTree )
+{
+	ScrollBox->ClearChildren();
 
 	const TArray< TSharedRef<FSequencerDisplayNode> >& RootNodes = InSequencerNodeTree.GetRootNodes();
 
@@ -44,6 +50,9 @@ void SSequencerTrackArea::Update( const FSequencerNodeTree& InSequencerNodeTree
 		// Generate a node for each children of the root node
 		GenerateLayoutNodeWidgetsRecursive( RootNode->GetChildNodes() );
 	}
+
+	// @todo Sequencer - Remove this expensive operation
+	ScrollBox->SlatePrepass();
 }
 
 void SSequencerTrackArea::GenerateLayoutNodeWidgetsRecursive( const TArray< TSharedRef<FSequencerDisplayNode> >& Nodes )
diff --git a/Engine/Source/Editor/Sequencer/Private/SSequencerTrackArea.h b/Engine/Source/Editor/Sequencer/Private/SSequencerTrackArea.h
--- a/Engine/Source/Editor/Sequencer/Private/SSequencerTrackArea.h
+++ b/Engine/Source/Editor/Sequencer/Private/SSequencerTrackArea.h
@@ -16,16 +16,7 @@ public:
 		SLATE_ATTRIBUTE( float, OutlinerFillPercent )
 	SLATE_END_ARGS()
 
-	void Construct( const FArguments& InArgs, TSharedRef<FSequencer> InSequencer )
-	{
-		ViewRange = InArgs._ViewRange;
-		OutlinerFillPercent = InArgs._OutlinerFillPercent;
-		Sequencer = InSequencer;
-		ChildSlot
-		[
-			SAssignNew( ScrollBox, SScrollBox )
-		];
-	}
+	void Construct( const FArguments& InArgs, TSharedRef<FSequencer> InSequencer );
 
 	/**
 	 * Updates the track area with new nodes
diff --git a/Engine/Source/Editor/Sequencer/Private/Sequencer.cpp b/Engine/Source/Editor/Sequencer/Private/Sequencer.cpp
--- a/Engine/Source/Editor/Sequencer/Private/Sequencer.cpp
+++ b/Engine/Source/Editor/Sequencer/Private/Sequencer.cpp
@@ -2,11 +2,10 @@
 
 #include "SequencerPrivatePCH.h"
 #include "Sequencer.h"
-#include "SequencerObjectSpawner.h"
+#include "ISequencerObjectBindingManager.h"
 #include "MovieScene.h"
 #include "Engine/LevelScriptBlueprint.h"
 #include "Editor/PropertyEditor/Public/PropertyEditorModule.h"
-#include "Editor/PropertyEditor/Public/IDetailsView.h"
 #include "Editor/SequencerWidgets/Public/ITimeSlider.h"
 #include "Editor/EditorWidgets/Public/ITransportControl.h"
 #include "Editor/EditorWidgets/Public/EditorWidgetsModule.h"
@@ -16,9 +15,7 @@
 #include "SSequencer.h"
 #include "ScopedTransaction.h"
 #include "Kismet2/BlueprintEditorUtils.h"
-#include "K2Node_PlayMovieScene.h"
 #include "BlueprintEditorModule.h"
-#include "SequencerObjectChangeListener.h"
 #include "MovieSceneTrack.h"
 #include "MovieSceneDirectorTrack.h"
 #include "MovieSceneAudioTrack.h"
@@ -46,23 +43,23 @@ bool FSequencer::IsSequencerEnabled()
 }
 
 
-void FSequencer::InitSequencer( UMovieScene* InRootMovieScene, TSharedPtr<IToolkitHost> InToolKitHost, const TArray<FOnCreateTrackEditor>& TrackEditorDelegates, bool bEditWithinLevelEditor )
+void FSequencer::InitSequencer( const FSequencerInitParams& InitParams, const TArray<FOnCreateTrackEditor>& TrackEditorDelegates )
 {
 	if( IsSequencerEnabled() )
 	{
-		bIsEditingWithinLevelEditor = bEditWithinLevelEditor;
+		bIsEditingWithinLevelEditor = InitParams.bEditWithinLevelEditor;
 
-		ToolkitHost = InToolKitHost;
-		// Create an object change listener for various systems that need to be notified when objects change
-		ObjectChangeListener = MakeShareable( new FSequencerObjectChangeListener( SharedThis( this ) ) );
-		// If this is a world-centric editor, then setup our puppet spawner
-		if( bEditWithinLevelEditor )
-		{
-			ObjectSpawner = MakeShareable( new FSequencerActorObjectSpawner( *this ) );
-		}
+		ToolkitHost = InitParams.ToolkitHost;
 
-		UMovieScene& RootMovieScene = *InRootMovieScene;
+		ObjectChangeListener = InitParams.ObjectChangeListener;
+		ObjectBindingManager = InitParams.ObjectBindingManager;
 
+		check( ObjectChangeListener.IsValid() );
+		check( ObjectBindingManager.IsValid() );
+		
+		UMovieScene& RootMovieScene = *InitParams.RootMovieScene;
+		RootMovieScene.SetFlags( RF_Transactional );
+		
 		// Focusing the initial movie scene needs to be done before the first time GetFocusedMovieSceneInstane or GetRootMovieSceneInstance is used
 		RootMovieSceneInstance = MakeShareable( new FMovieSceneInstance( RootMovieScene ) );
 		MovieSceneStack.Add( RootMovieSceneInstance.ToSharedRef() );
@@ -78,20 +75,15 @@ void FSequencer::InitSequencer( UMovieScene* InRootMovieScene, TSharedPtr<IToolk
 			.OnToggleAutoKey( this, &FSequencer::OnToggleAutoKey )
 			.OnToggleCleanView( this, &FSequencer::OnToggleCleanView );
 
+		
+		// When undo occurs, get a notification so we can make sure our view is up to date
+		GEditor->RegisterForUndo(this);
 
-		if( bEditWithinLevelEditor )
+		if( bIsEditingWithinLevelEditor )
 		{
 			// @todo remove when world-centric mode is added
 			// Hook into the editor's mechanism for checking whether we need live capture of PIE/SIE actor state
 			GEditor->GetActorRecordingState().AddSP(this, &FSequencer::GetActorRecordingState);
-
-			// When undo occurs, get a notification so we can make sure our view is up to date
-			GEditor->RegisterForUndo(this);
-
-			// Start listening for important changes on this MovieScene
-			// @todo sequencer: We need to correctly register for any new moviescenes that are dropped in!
-			const bool bCreateIfNotFound = false;
-			BindToPlayMovieSceneNode(bCreateIfNotFound);
 		}
 
 		// Create tools and bind them to this sequencer
@@ -149,23 +141,13 @@ FSequencer::FSequencer()
 
 FSequencer::~FSequencer()
 {
-	if( PlayMovieSceneNode.IsValid() )
-	{
-		PlayMovieSceneNode->OnBindingsChanged().RemoveAll( this );
-		PlayMovieSceneNode.Reset();
-	}
-
 	GEditor->GetActorRecordingState().RemoveAll( this );
 
 	GEditor->UnregisterForUndo( this );
 
-	// Clean up puppet objects
-	if( ObjectSpawner.IsValid() )
-	{
-		DestroySpawnablesForAllMovieScenes();
+	DestroySpawnablesForAllMovieScenes();
 	
-		ObjectSpawner.Reset();
-	}
+	ObjectBindingManager.Reset();
 
 	TrackEditors.Empty();
 
@@ -251,15 +233,6 @@ void FSequencer::PopToMovieScene( TSharedRef<FMovieSceneInstance> SubMovieSceneI
 	}
 }
 
-void FSequencer::SpawnOrDestroyPuppetObjects( TSharedRef<FMovieSceneInstance> MovieSceneInstance )
-{
-	if(ObjectSpawner.IsValid())
-	{
-		ObjectSpawner->SpawnOrDestroyPuppetObjects( MovieSceneInstance );
-	}
-}
-
-
 void FSequencer::AddNewShot(FGuid CameraGuid)
 {
 	if (DirectorTrackEditor.IsValid())
@@ -348,136 +321,84 @@ void FSequencer::RenameShotCommitted(const FText& RenameText, ETextCommit::Type
 	}
 }
 
-
+void FSequencer::SpawnOrDestroyPuppetObjects( TSharedRef<FMovieSceneInstance> MovieSceneInstance )
+{
+	if( ObjectBindingManager->AllowsSpawnableObjects() )
+	{
+		const bool bDestroyAll = false;
+		ObjectBindingManager->SpawnOrDestroyObjectsForInstance( MovieSceneInstance, bDestroyAll );
+	}
+}
 
 void FSequencer::OnMapChanged( class UWorld* NewWorld, EMapChangeType::Type MapChangeType )
 {
 	// @todo Sequencer Sub-MovieScenes Needs more investigation of what to spawn
 	// Destroy our puppets because the world is going away.  We probably don't have to do this (the actors will
 	// be naturally destroyed with the level, but we might as well.)
-	if( ObjectSpawner.IsValid() )
-	{
-		const bool bDestroyAll = true;
-		DestroySpawnablesForAllMovieScenes();
-	}
+	DestroySpawnablesForAllMovieScenes();
+
 
 	// @todo sequencer: We should only wipe/respawn puppets that are affected by the world that is being changed! (multi-UWorld support)
-	if( MapChangeType == EMapChangeType::LoadMap || MapChangeType == EMapChangeType::NewMap )
+	if( ( MapChangeType == EMapChangeType::LoadMap || MapChangeType == EMapChangeType::NewMap ) )
 	{
 		SpawnOrDestroyPuppetObjects( GetFocusedMovieSceneInstance() );
 	}
 }
 
-void FSequencer::NotifyMovieSceneDataChanged()
+void FSequencer::OnActorsDropped( const TArray<TWeakObjectPtr<AActor> >& Actors )
 {
-	PlaybackState = EMovieScenePlayerStatus::Stopped;
-	bNeedTreeRefresh = true;
-}
-
-void FSequencer::OnPlayMovieSceneBindingsChanged()
-{
-	// Something may have changed with the objects bound to the PlayMovieScene node, which may require us to change
-	// which actors are possessed, and thus create new puppets for different actors.
-	// @todo Sequencer Sub-MovieScenes Spawn other movie scene puppets?
-	SpawnOrDestroyPuppetObjects( GetFocusedMovieSceneInstance() );
-}
-
-TRange<float> FSequencer::OnGetViewRange() const
-{
-	return TRange<float>(FMath::Lerp(LastViewRange.GetLowerBoundValue(), TargetViewRange.GetLowerBoundValue(), ZoomCurve.GetLerp()),
-		FMath::Lerp(LastViewRange.GetUpperBoundValue(), TargetViewRange.GetUpperBoundValue(), ZoomCurve.GetLerp()));
-}
-
-UK2Node_PlayMovieScene* FSequencer::BindToPlayMovieSceneNode( const bool bCreateIfNotFound )
-{
-	auto* MovieScene = GetFocusedMovieScene();
-
-	// Update level script
-	UK2Node_PlayMovieScene* FoundPlayMovieSceneNode = FindPlayMovieSceneNodeInLevelScript( MovieScene );
-	if( FoundPlayMovieSceneNode == NULL )
+	bool bPossessableAdded = false;
+	for( TWeakObjectPtr<AActor> WeakActor : Actors )
 	{
-		if( bCreateIfNotFound )
+		AActor* Actor = WeakActor.Get();
+		if( Actor != NULL )
 		{
-			// Couldn't find an existing node that uses this MovieScene, so we'll create one now.
-			FoundPlayMovieSceneNode = CreateNewPlayMovieSceneNode( MovieScene );
-			if( FoundPlayMovieSceneNode != NULL )
+			// Grab the MovieScene that is currently focused.  We'll add our Blueprint as an inner of the
+			// MovieScene asset.
+			UMovieScene* OwnerMovieScene = GetFocusedMovieScene();
+			
+			// @todo sequencer: Undo doesn't seem to be working at all
+			const FScopedTransaction Transaction( LOCTEXT("UndoPossessingObject", "Possess Object with MovieScene") );
+			
+			// Possess the object!
 			{
-				// Let the user know that we plopped down a new PlayMovieScene event in their level script graph
+				// Create a new possessable
+				OwnerMovieScene->Modify();
+				
+				const FString PossessableName = Actor->GetActorLabel();
+				const FGuid PossessableGuid = OwnerMovieScene->AddPossessable( PossessableName, Actor->GetClass() );
+				 
+				if ( IsShotFilteringOn() )
 				{
-					FNotificationInfo Info( LOCTEXT("AddedPlayMovieSceneEventToLevelScriptGraph", "A new event for this MovieScene was added to this level's script") );
-					Info.bFireAndForget = true;
-					Info.bUseThrobber = false;
-					Info.bUseSuccessFailIcons = false;
-					Info.bUseLargeFont = false;
-					Info.ExpireDuration = 5.0f;		// Stay visible for awhile so the user has time to click "Show Graph" if they want to
-
-					// @todo sequencer: Needs better artwork (see DefaultStyle.cpp)
-					Info.Image = FEditorStyle::GetBrush( TEXT( "Sequencer.NotificationImage_AddedPlayMovieSceneEvent" ) );
-
-					struct Local
-					{
-						/**
-							* Called by our notification's hyperlink to open the Level Script editor and navigate to our newly-created PlayMovieScene node
-							*
-							* @param	Sequencer			The Sequencer we're using
-							* @param	PlayMovieSceneNode	The node that we need to jump to
-							*/
-						static void NavigateToPlayMovieSceneNode( FSequencer* Sequencer, UK2Node_PlayMovieScene* PlayMovieSceneNode )
-						{
-							check( Sequencer != NULL );
-							check( PlayMovieSceneNode != NULL );
-
-							IBlueprintEditor* BlueprintEditor = NULL;
-
-							auto* LSB = CastChecked< ULevelScriptBlueprint >( PlayMovieSceneNode->GetBlueprint() );
-
-							// @todo sequencer: Support using world-centric editing here?  (just need to set EToolkitMode::WorldCentric)
-							if( FAssetEditorManager::Get().OpenEditorForAsset( LSB, EToolkitMode::Standalone, Sequencer->GetToolkitHost() ) )
-							{
-								const bool bFocusIfOpen = true;
-								BlueprintEditor = static_cast< IBlueprintEditor* >( FAssetEditorManager::Get().FindEditorForAsset( LSB, bFocusIfOpen ) );
-							}
-
-							if( BlueprintEditor != NULL )
-							{
-								const bool bRequestRename = false;
-								BlueprintEditor->JumpToHyperlink( PlayMovieSceneNode, false );
-							}
-							else
-							{
-								UE_LOG( LogSequencer, Warning, TEXT( "Unable to open Blueprint Editor to edit newly-created PlayMovieScene event" ) );
-							}
-						}
-					};
-
-					FSequencer* SequencerInternals = this;
-					Info.Hyperlink = FSimpleDelegate::CreateStatic( &Local::NavigateToPlayMovieSceneNode, SequencerInternals, FoundPlayMovieSceneNode );
-					Info.HyperlinkText = LOCTEXT("AddedPlayMovieSceneEventToLevelScriptGraph_Hyperlink", "Show Graph");
-
-					TSharedPtr< SNotificationItem > NewNotification = FSlateNotificationManager::Get().AddNotification(Info);
+					AddUnfilterableObject(PossessableGuid);
 				}
+				
+				ObjectBindingManager->BindPossessableObject( PossessableGuid, *Actor );
+				
+				bPossessableAdded = true;
 			}
 		}
 	}
-
-	if( PlayMovieSceneNode.Get() != FoundPlayMovieSceneNode )
+	
+	if( bPossessableAdded )
 	{
-		if( PlayMovieSceneNode.IsValid() )
-		{
-			// Unhook from old node
-			PlayMovieSceneNode.Get()->OnBindingsChanged().RemoveAll( this );
-		}
+		SpawnOrDestroyPuppetObjects( GetFocusedMovieSceneInstance() );
 		
-		PlayMovieSceneNode = FoundPlayMovieSceneNode;
-
-		if( PlayMovieSceneNode.IsValid() )
-		{
-			// Bind to the new node
-			PlayMovieSceneNode.Get()->OnBindingsChanged().AddSP( this, &FSequencer::OnPlayMovieSceneBindingsChanged );
-		}
+		NotifyMovieSceneDataChanged();
 	}
+}
 
-	return PlayMovieSceneNode.Get();
+void FSequencer::NotifyMovieSceneDataChanged()
+{
+	PlaybackState = EMovieScenePlayerStatus::Stopped;
+	bNeedTreeRefresh = true;
+}
+
+
+TRange<float> FSequencer::OnGetViewRange() const
+{
+	return TRange<float>(FMath::Lerp(LastViewRange.GetLowerBoundValue(), TargetViewRange.GetLowerBoundValue(), ZoomCurve.GetLerp()),
+		FMath::Lerp(LastViewRange.GetUpperBoundValue(), TargetViewRange.GetUpperBoundValue(), ZoomCurve.GetLerp()));
 }
 
 bool FSequencer::IsAutoKeyEnabled() const 
@@ -523,90 +444,51 @@ FGuid FSequencer::GetHandleToObject( UObject* Object )
 	TSharedRef<FMovieSceneInstance> FocusedMovieSceneInstance = GetFocusedMovieSceneInstance();
 	UMovieScene* FocusedMovieScene = FocusedMovieSceneInstance->GetMovieScene();
 
-	// First check to see if this is a puppet object.  If so, we'll get a handle to the movie scene's spawnable
-	// entry for that puppet.
-	FGuid ObjectGuid = ObjectSpawner->FindSpawnableGuidForPuppetObject( Object );
-	if( ObjectGuid.IsValid() )
-	{
-		// Found a puppet's spawnable entry!		
-	}
-	else
+	FGuid ObjectGuid = ObjectBindingManager->FindGuidForObject( *FocusedMovieScene, *Object );
+	
+	bool bPossessableAdded = false;
+	
+	// If the object guid was not found attempt to add it
+	// Note: Only possessed actors can be added like this
+	if( !ObjectGuid.IsValid() && ObjectBindingManager->CanPossessObject( *Object ) )
 	{
-		// Is this a game preview object?
-		const bool bIsGamePreviewObject = !!( Object->GetOutermost()->PackageFlags & PKG_PlayInEditor );
-		if( bIsGamePreviewObject )
-		{
-			// OK, so someone is asking for a handle to an object from a game preview session, probably because
-			// they want to capture keys during live simulation.
-
-			// Object is in the editor world, so go ahead and possess it if we can
-
-			// Check to see if we already have a puppet that was generated from a recording of this game preview object
-			// @todo sequencer livecapture: We could support recalling counterpart by full name instead of weak pointer, to allow "overdubbing" of previously recorded actors, when the new actors in the current play session have the same path name
-			// @todo sequencer livecapture: Ideally we could capture from editor-world actors that are "puppeteered" as well (real time)
-			FMovieSceneSpawnable* FoundSpawnable = FocusedMovieScene->FindSpawnableForCounterpart( Object );
-			if( FoundSpawnable != NULL )
-			{
-				ObjectGuid = FoundSpawnable->GetGuid();
-			}
-			else
-			{
-				// No existing object was found, so we'll need to create a new spawnable (and puppet) for this
-				// game preview object right now.
-				UObject* ActorTemplate = Object;	// Pass in the actor itself as the template.  AddSpawnableForAssetOrClass() will use this actor's class!
-				UObject* CounterpartGamePreviewObject = Object;
-				ObjectGuid = AddSpawnableForAssetOrClass( Object, CounterpartGamePreviewObject );
-
-				// Update puppet actors
-				SpawnOrDestroyPuppetObjects( FocusedMovieSceneInstance );
-			}
-		}
-		else
+		// Grab the MovieScene that is currently focused.  We'll add our Blueprint as an inner of the
+		// MovieScene asset.
+		UMovieScene* OwnerMovieScene = GetFocusedMovieScene();
+		
+		// @todo sequencer: Undo doesn't seem to be working at all
+		const FScopedTransaction Transaction( LOCTEXT("UndoPossessingObject", "Possess Object with MovieScene") );
+		
+		// Possess the object!
 		{
-			// Object is in the editor world, so go ahead and possess it if we can
-			FocusedMovieScene = GetFocusedMovieScene();
-
-			// Make sure we're bound to the level script node which contains data about possessables.  But don't bother creating
-			// a node if we don't have one, at least not yet.
-			{
-				const bool bCreateIfNotFound = false;
-				BindToPlayMovieSceneNode( bCreateIfNotFound );
-			}
+			// Create a new possessable
+			OwnerMovieScene->Modify();
 			
-			if( PlayMovieSceneNode.IsValid() )
-			{
-				ObjectGuid = PlayMovieSceneNode->FindGuidForObject( Object );
-			}
-
-			// The object was not found, bind it now by adding a new possessable for it
-			if( !ObjectGuid.IsValid() )
+			const FString PossessableName = Object->GetName();
+			ObjectGuid = OwnerMovieScene->AddPossessable( PossessableName, Object->GetClass() );
+			
+			if ( IsShotFilteringOn() )
 			{
-				// Add a new possessable for this object.  Create a new PlayMovieScene node now, if we don't already have one.
-				const bool bCreateIfNotFound = true;
-				BindToPlayMovieSceneNode( bCreateIfNotFound );
-
-				const FString& PossessableName = Object->GetName();
-				const FGuid PossessableGuid = FocusedMovieScene->AddPossessable( PossessableName, Object->GetClass() );
-				
-				if (IsShotFilteringOn()) {AddUnfilterableObject(PossessableGuid);}
-
-				// Bind the object to the handle
-				TArray< UObject* > Objects;
-				Objects.Add( Object );
-				PlayMovieSceneNode->BindPossessableToObjects( PossessableGuid, Objects );
-
-				// A possessable was created so we need to respawn its puppet
-				SpawnOrDestroyPuppetObjects( FocusedMovieSceneInstance );
-
-				ObjectGuid = PossessableGuid;
+				AddUnfilterableObject(ObjectGuid);
 			}
+			
+			ObjectBindingManager->BindPossessableObject( ObjectGuid, *Object );
+			
+			bPossessableAdded = true;
 		}
 	}
-
+	
+	if( bPossessableAdded )
+	{
+		SpawnOrDestroyPuppetObjects( GetFocusedMovieSceneInstance() );
+			
+		NotifyMovieSceneDataChanged();
+	}
+	
 	return ObjectGuid;
 }
 
-ISequencerObjectChangeListener& FSequencer::GetObjectChangeListener() const 
+ISequencerObjectChangeListener& FSequencer::GetObjectChangeListener()
 { 
 	return *ObjectChangeListener;
 }
@@ -614,15 +496,14 @@ ISequencerObjectChangeListener& FSequencer::GetObjectChangeListener() const
 
 void FSequencer::SpawnActorsForMovie( TSharedRef<FMovieSceneInstance> MovieSceneInstance )
 {
-	if( ObjectSpawner.IsValid() )
-	{
-		ObjectSpawner->SpawnOrDestroyPuppetObjects( MovieSceneInstance );
-	}
+	SpawnOrDestroyPuppetObjects( MovieSceneInstance );
 }
 
 void FSequencer::GetRuntimeObjects( TSharedRef<FMovieSceneInstance> MovieSceneInstance, const FGuid& ObjectHandle, TArray< UObject*>& OutObjects ) const
 {
-	if( ObjectSpawner.IsValid() )
+	ObjectBindingManager->GetRuntimeObjects( MovieSceneInstance, ObjectHandle, OutObjects );
+	
+	/*if( ObjectSpawner.IsValid() )
 	{	
 		// First, try to find spawnable puppet objects for the specified Guid
 		UObject* FoundPuppetObject = ObjectSpawner->FindPuppetObjectForSpawnableGuid( MovieSceneInstance, ObjectHandle );
@@ -642,6 +523,14 @@ void FSequencer::GetRuntimeObjects( TSharedRef<FMovieSceneInstance> MovieSceneIn
 			}
 		}
 	}
+	else
+	{
+		FMovieScenePossessable* Possessable = MovieSceneInstance->GetMovieScene()->FindPossessable( ObjectHandle );
+		if( Possessable && Possessable->GetPossessableObject() )
+		{
+			OutObjects.Add( Possessable->GetPossessableObject() );
+		}
+	}*/
 }
 
 void FSequencer::UpdateViewports(AActor* ActorToViewThrough) const
@@ -702,8 +591,12 @@ void FSequencer::AddMovieSceneInstance( UMovieSceneSection& MovieSceneSection, T
 
 void FSequencer::RemoveMovieSceneInstance( UMovieSceneSection& MovieSceneSection, TSharedRef<FMovieSceneInstance> InstanceToRemove )
 {
-	const bool bDestroyAll = true;
-	ObjectSpawner->SpawnOrDestroyPuppetObjects( InstanceToRemove, bDestroyAll );
+	if( ObjectBindingManager->AllowsSpawnableObjects() )
+	{
+		const bool bDestroyAll = true;
+		ObjectBindingManager->SpawnOrDestroyObjectsForInstance( InstanceToRemove, bDestroyAll );
+	}
+	
 
 	MovieSceneSectionToInstanceMap.Remove( &MovieSceneSection );
 }
@@ -786,7 +679,10 @@ void FSequencer::UpdateRuntimeInstances()
 
 void FSequencer::DestroySpawnablesForAllMovieScenes()
 {
-	ObjectSpawner->DestroyAllPuppetObjects();
+	if( ObjectBindingManager->AllowsSpawnableObjects() )
+	{
+		ObjectBindingManager->DestroyAllSpawnedObjects();
+	}
 }
 
 
@@ -933,78 +829,6 @@ TRange<float> FSequencer::GetFilteringShotsTimeBounds() const
 }
 
 
-UK2Node_PlayMovieScene* FSequencer::FindPlayMovieSceneNodeInLevelScript( const UMovieScene* MovieScene )
-{
-	// Grab the world object for this editor
-	check( MovieScene != NULL );
-	UWorld* World = ToolkitHost.Pin()->GetWorld();
-	check( World != NULL );
-
-	// Search all levels in the specified world
-	for( TArray<ULevel*>::TConstIterator LevelIter( World->GetLevels().CreateConstIterator() ); LevelIter; ++LevelIter )
-	{
-		auto* Level = *LevelIter;
-		if( Level != NULL )
-		{
-			// We don't want to create a level script if one doesn't exist yet.  We just want to grab the one
-			// that we already have, if one exists.
-			const bool bDontCreate = true;
-			ULevelScriptBlueprint* LSB = Level->GetLevelScriptBlueprint( bDontCreate );
-			if( LSB != NULL )
-			{
-				TArray<UK2Node_PlayMovieScene*> EventNodes;
-				FBlueprintEditorUtils::GetAllNodesOfClass( LSB, EventNodes );
-				for( auto FoundNodeIter( EventNodes.CreateIterator() ); FoundNodeIter; ++FoundNodeIter )
-				{
-					UK2Node_PlayMovieScene* FoundNode = *FoundNodeIter;
-					if( FoundNode->GetMovieScene() == MovieScene )
-					{
-						return FoundNode;
-					}
-				}
-			}
-		}
-	}
-
-	return NULL;
-}
-
-
-UK2Node_PlayMovieScene* FSequencer::CreateNewPlayMovieSceneNode( UMovieScene* MovieScene )
-{
-	// Grab the world object for this editor
-	check( MovieScene != NULL );
-	UWorld* World = ToolkitHost.Pin()->GetWorld();
-	check( World != NULL );
-
-	ULevel* Level = World->GetCurrentLevel();
-	check( Level != NULL );
-
-	// Here, we'll create a level script if one does not yet exist.
-	const bool bDontCreate = false;
-	ULevelScriptBlueprint* LSB = Level->GetLevelScriptBlueprint( bDontCreate );
-	if( LSB != NULL )
-	{
-		UEdGraph* TargetGraph = NULL;
-		if( LSB->UbergraphPages.Num() > 0 )
-		{
-			TargetGraph = LSB->UbergraphPages[0]; // Just use the first graph
-		}
-
-		if( ensure( TargetGraph != NULL ) )
-		{
-			// Figure out a decent place to stick the node
-			const FVector2D NewNodePos = TargetGraph->GetGoodPlaceForNewNode();
-
-			// Create a new node
-			UK2Node_PlayMovieScene* TemplateNode = NewObject<UK2Node_PlayMovieScene>();
-			TemplateNode->SetMovieScene( MovieScene );
-			return FEdGraphSchemaAction_K2NewNode::SpawnNodeFromTemplate<UK2Node_PlayMovieScene>(TargetGraph, TemplateNode, NewNodePos);;
-		}
-	}
-
-	return NULL;
-}
 
 void FSequencer::OnViewRangeChanged( TRange<float> NewViewRange, bool bSmoothZoom )
 {
@@ -1058,82 +882,86 @@ void FSequencer::OnToggleCleanView( bool bInCleanViewEnabled )
 
 FGuid FSequencer::AddSpawnableForAssetOrClass( UObject* Object, UObject* CounterpartGamePreviewObject )
 {
-	// Grab the MovieScene that is currently focused.  We'll add our Blueprint as an inner of the
-	// MovieScene asset.
-	UMovieScene* OwnerMovieScene = GetFocusedMovieScene();
+	FGuid NewSpawnableGuid;
+	
+	if( ObjectBindingManager->AllowsSpawnableObjects() )
+	{
+		// Grab the MovieScene that is currently focused.  We'll add our Blueprint as an inner of the
+		// MovieScene asset.
+		UMovieScene* OwnerMovieScene = GetFocusedMovieScene();
 
-	// @todo sequencer: Undo doesn't seem to be working at all
-	const FScopedTransaction Transaction( LOCTEXT("UndoAddingObject", "Add Object to MovieScene") );
+		// @todo sequencer: Undo doesn't seem to be working at all
+		const FScopedTransaction Transaction( LOCTEXT("UndoAddingObject", "Add Object to MovieScene") );
 
-	// Use the class as the spawnable's name if this is an actor class, otherwise just use the object name (asset)
-	const bool bIsActorClass = Object->IsA( AActor::StaticClass() ) && !Object->HasAnyFlags( RF_ArchetypeObject );
-	const FName AssetName = bIsActorClass ? Object->GetClass()->GetFName() : Object->GetFName();
+		// Use the class as the spawnable's name if this is an actor class, otherwise just use the object name (asset)
+		const bool bIsActorClass = Object->IsA( AActor::StaticClass() ) && !Object->HasAnyFlags( RF_ArchetypeObject );
+		const FName AssetName = bIsActorClass ? Object->GetClass()->GetFName() : Object->GetFName();
 
-	// Inner objects don't need a name (it will be auto-generated by the UObject system), but we want one in this case
-	// because the class of any actors that are created from this Blueprint may end up being user-facing.
-	const FName BlueprintName = MakeUniqueObjectName( OwnerMovieScene, UBlueprint::StaticClass(), AssetName );
+		// Inner objects don't need a name (it will be auto-generated by the UObject system), but we want one in this case
+		// because the class of any actors that are created from this Blueprint may end up being user-facing.
+		const FName BlueprintName = MakeUniqueObjectName( OwnerMovieScene, UBlueprint::StaticClass(), AssetName );
 
-	// Use the asset name as the initial spawnable name
-	const FString& NewSpawnableName = AssetName.ToString();		// @todo sequencer: Need UI to allow user to rename these slots
+		// Use the asset name as the initial spawnable name
+		const FString& NewSpawnableName = AssetName.ToString();		// @todo sequencer: Need UI to allow user to rename these slots
 
-	// Create our new blueprint!
-	UBlueprint* NewBlueprint = NULL;
-	{
-		// @todo sequencer: Add support for forcing specific factories for an asset
-		UActorFactory* FactoryToUse = NULL;
-		if( bIsActorClass )
-		{
-			// Placing an actor class directly::
-			FactoryToUse = GEditor->FindActorFactoryForActorClass( Object->GetClass() );
-		}
-		else
+		// Create our new blueprint!
+		UBlueprint* NewBlueprint = NULL;
 		{
-			// Placing an asset
-			FactoryToUse = FActorFactoryAssetProxy::GetFactoryForAssetObject( Object );
-		}
+			// @todo sequencer: Add support for forcing specific factories for an asset
+			UActorFactory* FactoryToUse = NULL;
+			if( bIsActorClass )
+			{
+				// Placing an actor class directly::
+				FactoryToUse = GEditor->FindActorFactoryForActorClass( Object->GetClass() );
+			}
+			else
+			{
+				// Placing an asset
+				FactoryToUse = FActorFactoryAssetProxy::GetFactoryForAssetObject( Object );
+			}
 
-		if( FactoryToUse != NULL )
-		{
-			// Create the blueprint
-			NewBlueprint = FactoryToUse->CreateBlueprint( Object, OwnerMovieScene, BlueprintName );
-		}
-		else if( bIsActorClass )
-		{
-			// We don't have a factory, but we can still try to create a blueprint for this actor class
-			NewBlueprint = FKismetEditorUtilities::CreateBlueprint( Object->GetClass(), OwnerMovieScene, BlueprintName, EBlueprintType::BPTYPE_Normal, UBlueprint::StaticClass(), UBlueprintGeneratedClass::StaticClass() );
+			if( FactoryToUse != NULL )
+			{
+				// Create the blueprint
+				NewBlueprint = FactoryToUse->CreateBlueprint( Object, OwnerMovieScene, BlueprintName );
+			}
+			else if( bIsActorClass )
+			{
+				// We don't have a factory, but we can still try to create a blueprint for this actor class
+				NewBlueprint = FKismetEditorUtilities::CreateBlueprint( Object->GetClass(), OwnerMovieScene, BlueprintName, EBlueprintType::BPTYPE_Normal, UBlueprint::StaticClass(), UBlueprintGeneratedClass::StaticClass() );
+			}
 		}
-	}
 
-	FGuid NewSpawnableGuid;
-	if( ensure( NewBlueprint != NULL ) )
-	{
-		if( NewBlueprint->GeneratedClass != NULL && FBlueprintEditorUtils::IsActorBased( NewBlueprint ) )
+		if( ensure( NewBlueprint != NULL ) )
 		{
-			AActor* ActorCDO = CastChecked< AActor >( NewBlueprint->GeneratedClass->ClassDefaultObject );
-
-			// If we have a counterpart object, then copy the properties from that object back into our blueprint's CDO
-			// @todo sequencer livecapture: This isn't really good enough to handle complex actors.  The dynamically-spawned actor could have components that
-			// were created in its construction script or via straight-up C++ code.  Instead what we should probably do is duplicate the PIE actor and generate
-			// our CDO from that duplicate.  It could get pretty complex though.
-			if( CounterpartGamePreviewObject != NULL )
+			if( NewBlueprint->GeneratedClass != NULL && FBlueprintEditorUtils::IsActorBased( NewBlueprint ) )
 			{
-				AActor* CounterpartGamePreviewActor = CastChecked< AActor >( CounterpartGamePreviewObject );
-				CopyActorProperties( CounterpartGamePreviewActor, ActorCDO );
+				AActor* ActorCDO = CastChecked< AActor >( NewBlueprint->GeneratedClass->ClassDefaultObject );
+
+				// If we have a counterpart object, then copy the properties from that object back into our blueprint's CDO
+				// @todo sequencer livecapture: This isn't really good enough to handle complex actors.  The dynamically-spawned actor could have components that
+				// were created in its construction script or via straight-up C++ code.  Instead what we should probably do is duplicate the PIE actor and generate
+				// our CDO from that duplicate.  It could get pretty complex though.
+				if( CounterpartGamePreviewObject != NULL )
+				{
+					AActor* CounterpartGamePreviewActor = CastChecked< AActor >( CounterpartGamePreviewObject );
+					CopyActorProperties( CounterpartGamePreviewActor, ActorCDO );
+				}
+				else
+				{
+					// Place the new spawnable in front of the camera (unless we were automatically created from a PIE actor)
+					PlaceActorInFrontOfCamera( ActorCDO );
+				}
 			}
-			else
+
+			NewSpawnableGuid = OwnerMovieScene->AddSpawnable( NewSpawnableName, NewBlueprint, CounterpartGamePreviewObject );
+
+			if (IsShotFilteringOn())
 			{
-				// Place the new spawnable in front of the camera (unless we were automatically created from a PIE actor)
-				PlaceActorInFrontOfCamera( ActorCDO );
+				AddUnfilterableObject(NewSpawnableGuid);
 			}
 		}
-
-		NewSpawnableGuid = OwnerMovieScene->AddSpawnable( NewSpawnableName, NewBlueprint, CounterpartGamePreviewObject );
-
-		if (IsShotFilteringOn()) {AddUnfilterableObject(NewSpawnableGuid);}
-	}
-	else
-	{
-		// Factory could not produce a blueprint or some other error occurred.  Not expected to ever happen.
+	
 	}
 
 	return NewSpawnableGuid;
@@ -1142,7 +970,8 @@ FGuid FSequencer::AddSpawnableForAssetOrClass( UObject* Object, UObject* Counter
 void FSequencer::AddSubMovieScene( UMovieScene* SubMovieScene )
 {
 	// @todo Sequencer - sub-moviescenes This should be moved to the sub-moviescene editor
-
+	SubMovieScene->SetFlags( RF_Transactional );
+	
 	// Grab the MovieScene that is currently focused.  THis is the movie scene that will contain the sub-moviescene
 	UMovieScene* OwnerMovieScene = GetFocusedMovieScene();
 
@@ -1193,6 +1022,8 @@ void FSequencer::OnRequestNodeDeleted( TSharedRef<const FSequencerDisplayNode>&
 		OwnerMovieScene->SetFlags( RF_Transactional );
 		const FGuid& BindingToRemove = StaticCastSharedRef<const FObjectBindingNode>( NodeToBeDeleted )->GetObjectBinding();
 
+		//@todo Sequencer - add transaction
+		
 		// Try to remove as a spawnable first
 		bool bRemoved = OwnerMovieScene->RemoveSpawnable( BindingToRemove );
 		if( bRemoved )
@@ -1204,6 +1035,10 @@ void FSequencer::OnRequestNodeDeleted( TSharedRef<const FSequencerDisplayNode>&
 		{
 			// The guid should be associated with a possessable if it wasnt a spawnable
 			bRemoved = OwnerMovieScene->RemovePossessable( BindingToRemove );
+			
+			// @todo Sequencer - undo needs to work here
+			ObjectBindingManager->UnbindPossessableObjects( BindingToRemove );
+			
 			// If this check fails the guid was not associated with a spawnable or possessable so there was an invalid guid being stored on a node
 			check( bRemoved );
 		}
@@ -1511,12 +1346,12 @@ void FSequencer::SetKey()
 		// @todo Handle case of actors which aren't in sequencer yet
 
 		FGuid ObjectGuid = GetHandleToObject(*It);
-		for (int i = 0; i < TrackEditors.Num(); ++i)
+		for ( auto& TrackEditor : TrackEditors )
 		{
 			// @todo Handle this director track business better
-			if (TrackEditors[i] != DirectorTrackEditor.Pin())
+			if (TrackEditor != DirectorTrackEditor.Pin())
 			{
-				TrackEditors[i]->AddKey(ObjectGuid);
+				TrackEditor->AddKey(ObjectGuid);
 			}
 		}
 	}
diff --git a/Engine/Source/Editor/Sequencer/Private/Sequencer.h b/Engine/Source/Editor/Sequencer/Private/Sequencer.h
--- a/Engine/Source/Editor/Sequencer/Private/Sequencer.h
+++ b/Engine/Source/Editor/Sequencer/Private/Sequencer.h
@@ -7,6 +7,27 @@
 
 class UMovieScene;
 class IToolkitHost;
+class ISequencerObjectBindingManager;
+class ISequencerObjectChangeListener;
+
+/** Parameters for initializing a Sequencer */
+struct FSequencerInitParams
+{
+	/** The root movie scene being edited */
+	UMovieScene* RootMovieScene;
+	
+	/** The interface for managing runtime object bindings */
+	TSharedPtr<ISequencerObjectBindingManager> ObjectBindingManager;
+	
+	/** The Object change listener for various systems that need to be notified when objects change */
+	TSharedPtr<ISequencerObjectChangeListener> ObjectChangeListener;
+	
+	/** The asset editor created for this (if any) */
+	TSharedPtr<IToolkitHost> ToolkitHost;
+	
+	/** Whether or not sequencer should be edited within the level editor */
+	bool bEditWithinLevelEditor;
+};
 
 /**
  * Sequencer is the editing tool for MovieScene assets
@@ -18,12 +39,11 @@ public:
 	/**
 	 * Initializes sequencer
 	 *
-	 * @param	RootMovieScene					The root movie scene being edited
-	 * @param	InAssetEditor					The asset editor created for this (if any) 
-	 * @param	TrackEditorDelegates			Delegates to call to create auto-key handlers for this sequencer
-	 * @param	bEditWithinLevelEditor			Whether or not sequencer should be edited within the level editor
+	 * @param	InitParams		Initialization options
+	 * @param	TrackEditorDelegates	Delegates to call to create auto-key handlers for this sequencer
+
 	 */
-	void InitSequencer( UMovieScene* InRootMovieScene, TSharedPtr<IToolkitHost> InToolKitHost, const TArray<FOnCreateTrackEditor>& TrackEditorDelegates, bool bEditWithinLevelEditor );
+	void InitSequencer( const FSequencerInitParams& InitParams, const TArray<FOnCreateTrackEditor>& TrackEditorDelegates );
 
 	/** Constructor */
 	FSequencer();
@@ -52,7 +72,7 @@ public:
 	virtual void SetGlobalTime(float Time) override;
 	virtual void SetPerspectiveViewportPossessionEnabled(bool bEnabled) override;
 	virtual FGuid GetHandleToObject(UObject* Object) override;
-	virtual ISequencerObjectChangeListener& GetObjectChangeListener() const override;
+	virtual ISequencerObjectChangeListener& GetObjectChangeListener() override;
 	virtual void NotifyMovieSceneDataChanged() override;
 	virtual void AddSubMovieScene(UMovieScene* SubMovieScene) override;
 	virtual void FilterToShotSections(const TArray< TWeakObjectPtr<class UMovieSceneSection> >& ShotSections, bool bZoomToShotBounds = true) override;
@@ -242,19 +262,12 @@ public:
 
 	virtual TArray< UMovieScene* > GetMovieScenesBeingEdited();
 
-	/**
-	 * Binds to a 'PlayMovieScene' node in level script that is associated with the MovieScene we're editing.  If
-	 * requested, can also create a new node.  Only valid to call when in world-centric mode.
-	 *
-	 * @param	bCreateIfNotFound	True to create a node if we can't find one
-	 *
-	 * @return	The LevelScript node we found or created, or NULL if nothing is still bound
-	 */
-	class UK2Node_PlayMovieScene* BindToPlayMovieSceneNode( const bool bCreateIfNotFound );
-
 	/** Called by LevelEditor when the map changes */
 	void OnMapChanged(class UWorld* NewWorld, EMapChangeType::Type MapChangeType);
 
+	/** Called when an actor is dropped into Sequencer */
+	void OnActorsDropped( const TArray<TWeakObjectPtr<AActor> >& Actors );
+
 	/** Functions to push on to the transport controls we use */
 	FReply OnPlay();
 	FReply OnRecord();
@@ -304,32 +317,6 @@ protected:
 	 */
 	TRange<float> GetFilteringShotsTimeBounds() const;
 
-
-	/**
-	 * Given a MovieScene to search for, attempts to find a LevelScript node that is playing the
-	 * supplied MovieScene.  Only valid to call when in world-centric mode.
-	 *
-	 * @param	MovieScene	The MovieScene to search for
-	 *
-	 * @return	The LevelScript node we found, or NULL for no matches
-	 */
-	class UK2Node_PlayMovieScene* FindPlayMovieSceneNodeInLevelScript( const class UMovieScene* MovieScene );
-
-
-	/**
-	 * Creates a new "PlayMovieScene" node in the level script for the current level and associates it with the
-	 * supplied MovieScene.  Only valid to call when in world-centric mode.
-	 *
-	 * @param	MovieScene	The MovieScene to associate with the newly-created node
-	 *
-	 * @return	The LevelScript node we created, or NULL on failure (does not assert)
-	 */
-	class UK2Node_PlayMovieScene* CreateNewPlayMovieSceneNode( UMovieScene* MovieScene );
-
-
-	/** Called by UK2Node_PlayMovieScene when the bindings for that node are changed through Kismet */
-	void OnPlayMovieSceneBindingsChanged();
-
 	/** @return The current view range */
 	virtual TRange<float> OnGetViewRange() const;
 
@@ -401,11 +388,11 @@ private:
 	TWeakPtr<FMovieSceneTrackEditor> DirectorTrackEditor;
 	TWeakPtr<FMovieSceneTrackEditor> AnimationTrackEditor;
 
-	/** Object spawner */
-	TSharedPtr< class ISequencerObjectSpawner > ObjectSpawner;
+	/** Manager for handling runtime object bindings */
+	TSharedPtr< ISequencerObjectBindingManager > ObjectBindingManager;
 
 	/** Listener for object changes being made while this sequencer is open*/
-	TSharedPtr< class FSequencerObjectChangeListener > ObjectChangeListener; 
+	TSharedPtr< class ISequencerObjectChangeListener > ObjectChangeListener;
 
 	/** The runtime instance for the root movie scene */
 	TSharedPtr< class FMovieSceneInstance > RootMovieSceneInstance;
@@ -419,10 +406,6 @@ private:
 	/** The asset editor that created this Sequencer if any */
 	TWeakPtr<IToolkitHost> ToolkitHost;
 
-	/** 'PlayMovieScene' node in level script that we're currently associating with the active MovieScene.  This node contains the bindings
-	    information that we need to be able to preview possessed actors in the level */
-	TWeakObjectPtr< UK2Node_PlayMovieScene > PlayMovieSceneNode;
-
 	/** A list of all shots that are acting as filters */
 	TArray< TWeakObjectPtr<class UMovieSceneSection> > FilteringShots;
 
diff --git a/Engine/Source/Editor/Sequencer/Private/SequencerActorBindingManager.cpp b/Engine/Source/Editor/Sequencer/Private/SequencerActorBindingManager.cpp
new file mode 100644
--- /dev/null
+++ b/Engine/Source/Editor/Sequencer/Private/SequencerActorBindingManager.cpp
@@ -0,0 +1,623 @@
+// Copyright 1998-2014 Epic Games, Inc. All Rights Reserved.
+
+#include "SequencerPrivatePCH.h"
+#include "SequencerObjectSpawner.h"
+#include "Sequencer.h"
+#include "MovieScene.h"
+#include "Toolkits/IToolkitHost.h"
+#include "Kismet2/BlueprintEditorUtils.h"
+#include "K2Node_PlayMovieScene.h"
+#include "MovieSceneInstance.h"
+
+#define LOCTEXT_NAMESPACE "Sequencer"
+
+FSequencerActorBindingManager::FSequencerActorBindingManager( UWorld* InWorld, TSharedRef<ISequencerObjectChangeListener> InObjectChangeListener, TSharedRef<FSequencer> InSequencer )
+	: ActorWorld( InWorld )
+	, PlayMovieSceneNode( nullptr )
+	, Sequencer( InSequencer )
+	, ObjectChangeListener( InObjectChangeListener )
+{
+	// Register to be notified when object changes should be propagated
+	InObjectChangeListener->GetOnPropagateObjectChanges().AddRaw( this, &FSequencerActorBindingManager::OnPropagateObjectChanges );
+}
+
+
+FSequencerActorBindingManager::~FSequencerActorBindingManager()
+{
+	if( ObjectChangeListener.IsValid() )
+	{
+		ObjectChangeListener.Pin()->GetOnPropagateObjectChanges().RemoveAll( this );
+	}
+	
+	if( PlayMovieSceneNode.IsValid() )
+	{
+		PlayMovieSceneNode->OnBindingsChanged().RemoveAll( this );
+		PlayMovieSceneNode.Reset();
+	}
+}
+
+FGuid FSequencerActorBindingManager::FindGuidForObject( const UMovieScene& MovieScene, UObject& Object ) const
+{
+	FGuid ObjectGuid = FindSpawnableGuidForPuppetObject( &Object );
+	if( !ObjectGuid.IsValid() )
+	{
+		// Spawnable
+		// Is this a game preview object?
+		const bool bIsGamePreviewObject = !!( Object.GetOutermost()->PackageFlags & PKG_PlayInEditor );
+		if( bIsGamePreviewObject )
+		{
+			// OK, so someone is asking for a handle to an object from a game preview session, probably because
+			// they want to capture keys during live simulation.
+
+			// Check to see if we already have a puppet that was generated from a recording of this game preview object
+			// @todo sequencer livecapture: We could support recalling counterpart by full name instead of weak pointer, to allow "overdubbing" of previously recorded actors, when the new actors in the current play session have the same path name
+			// @todo sequencer livecapture: Ideally we could capture from editor-world actors that are "puppeteered" as well (real time)
+			const FMovieSceneSpawnable* FoundSpawnable = MovieScene.FindSpawnableForCounterpart( &Object );
+			if( FoundSpawnable != NULL )
+			{
+				ObjectGuid = FoundSpawnable->GetGuid();
+			}
+		}
+		else
+		{
+			BindToPlayMovieSceneNode( false );
+
+			// Possessable
+			// When editing within the level editor, make sure we're bound to the level script node which contains data about possessables.
+			if( PlayMovieSceneNode.IsValid() )
+			{
+				ObjectGuid = PlayMovieSceneNode->FindGuidForObject( &Object );
+			}
+		}
+
+	}
+	
+	return ObjectGuid;
+}
+
+void FSequencerActorBindingManager::SpawnOrDestroyObjectsForInstance( TSharedRef<FMovieSceneInstance> MovieSceneInstance, const bool bDestroyAll )
+{
+	bool bAnyLevelActorsChanged = false;
+
+	// Get the list of puppet objects for the movie scene
+	TArray< TSharedRef<FPuppetActorInfo> >& PuppetObjects = InstanceToPuppetObjectsMap.FindOrAdd( MovieSceneInstance );
+
+	UMovieScene* MovieScene = MovieSceneInstance->GetMovieScene();
+
+	// Remove any puppet objects that we no longer need
+	{
+		for( auto PuppetObjectIndex = 0; PuppetObjectIndex < PuppetObjects.Num(); ++PuppetObjectIndex )
+		{
+			if( PuppetObjects[ PuppetObjectIndex ]->GetType() == EPuppetObjectType::Actor )
+			{
+				TSharedRef< FPuppetActorInfo > PuppetActorInfo = StaticCastSharedRef< FPuppetActorInfo >( PuppetObjects[ PuppetObjectIndex ] );
+
+				// Figure out if we still need this puppet actor
+				bool bShouldDestroyActor = true;
+				if( !bDestroyAll )
+				{
+					for( auto SpawnableIndex = 0; SpawnableIndex < MovieScene->GetSpawnableCount(); ++SpawnableIndex )
+					{
+						auto& Spawnable = MovieScene->GetSpawnable( SpawnableIndex );
+						if( Spawnable.GetGuid() == PuppetActorInfo->SpawnableGuid )
+						{
+							bShouldDestroyActor = false;
+							break;
+						}
+					}
+				}
+
+				if( bShouldDestroyActor )
+				{
+					AActor* PuppetActor = PuppetActorInfo->PuppetActor.Get();
+					if( PuppetActor != NULL )
+					{
+						UWorld* PuppetWorld = PuppetActor->GetWorld();
+						if( ensure( PuppetWorld != NULL ) )
+						{
+							// Destroy this actor
+							{
+								// Ignored unless called while game is running
+								const bool bNetForce = false;
+
+								// We don't want to dirty the level for puppet actor changes
+								const bool bShouldModifyLevel = false;
+
+								if( !bAnyLevelActorsChanged )
+								{
+									DeselectAllPuppetObjects();
+								}
+
+								// Actor should never be selected in the editor at this point.  We took care of that up above.
+								ensure( !PuppetActor->IsSelected() );
+
+								const bool bWasDestroyed = PuppetWorld->DestroyActor( PuppetActor, bNetForce, bShouldModifyLevel );
+
+								if( bWasDestroyed )
+								{
+									bAnyLevelActorsChanged = true;
+									PuppetObjects.RemoveAt( PuppetObjectIndex-- );
+								}
+								else
+								{
+									// @todo sequencer: At least one puppet couldn't be cleaned up!
+								}
+							}
+						}
+					}
+					else
+					{
+						// Actor is no longer valid (probably the world was destroyed)
+					}
+				}
+			}
+			else
+			{
+				check(0);	// Unhandled type
+			}
+		}
+	}
+
+	if( !bDestroyAll )
+	{
+		for( auto SpawnableIndex = 0; SpawnableIndex < MovieScene->GetSpawnableCount(); ++SpawnableIndex )
+		{
+			auto& Spawnable = MovieScene->GetSpawnable( SpawnableIndex );
+
+			// Must have a valid world for us to be able to do this
+			if( ActorWorld != NULL )
+			{
+				// Do we already have a puppet for this spawnable?
+				bool bIsAlreadySpawned = false;
+				for( auto PuppetIndex = 0; PuppetIndex < PuppetObjects.Num(); ++PuppetIndex )
+				{
+					auto& PuppetObject = PuppetObjects[ PuppetIndex ];
+					if( PuppetObject->SpawnableGuid == Spawnable.GetGuid() )
+					{
+						bIsAlreadySpawned = true;
+						break;
+					}
+				}
+
+				if( !bIsAlreadySpawned )
+				{
+					UClass* GeneratedClass = Spawnable.GetClass();
+					if ( GeneratedClass != NULL && GeneratedClass->IsChildOf(AActor::StaticClass()))
+					{
+						AActor* ActorCDO = CastChecked< AActor >( GeneratedClass->ClassDefaultObject );
+
+						const FVector SpawnLocation = ActorCDO->GetRootComponent()->RelativeLocation;
+						const FRotator SpawnRotation = ActorCDO->GetRootComponent()->RelativeRotation;
+
+						// @todo sequencer: We should probably spawn these in a specific sub-level!
+						// World->CurrentLevel = ???;
+
+						const FName PuppetActorName = NAME_None;
+
+						// Override the object flags so that RF_Transactional is not set.  Puppet actors are never transactional
+						// @todo sequencer: These actors need to avoid any transaction history.  However, RF_Transactional can currently be set on objects on the fly!
+						const EObjectFlags ObjectFlags = RF_Transient;		// NOTE: We are omitting RF_Transactional intentionally
+
+						// @todo sequencer livecapture: Consider using SetPlayInEditorWorld() and RestoreEditorWorld() here instead
+						
+						// @todo sequencer actors: We need to make sure puppet objects aren't copied into PIE/SIE sessions!  They should be omitted from that duplication!
+
+						// Spawn the puppet actor
+						FActorSpawnParameters SpawnInfo;
+						SpawnInfo.Name = PuppetActorName;
+						SpawnInfo.ObjectFlags = ObjectFlags;
+						AActor* NewActor = ActorWorld->SpawnActor( GeneratedClass, &SpawnLocation, &SpawnRotation, SpawnInfo );
+				
+						if( NewActor )
+						{
+							// @todo sequencer: We're naming the actor based off of the spawnable's name.  Is that really what we want?
+							GEditor->SetActorLabelUnique( NewActor, Spawnable.GetName() );
+
+							// Actor was spawned OK!
+
+							// Keep track of this actor
+							TSharedRef< FPuppetActorInfo > NewPuppetInfo( new FPuppetActorInfo() );
+							NewPuppetInfo->SpawnableGuid = Spawnable.GetGuid();
+							NewPuppetInfo->PuppetActor = NewActor;
+							PuppetObjects.Add( NewPuppetInfo );
+						}
+						else
+						{
+							// Actor failed to spawn
+							// @todo sequencer: What should we do when this happens to one or more actors?
+						}
+					}
+				}
+			}
+		}
+	}
+}
+
+void FSequencerActorBindingManager::DestroyAllSpawnedObjects()
+{
+	const bool bDestroyAll = true;
+	for( auto MovieSceneIter = InstanceToPuppetObjectsMap.CreateConstIterator(); MovieSceneIter; ++MovieSceneIter )
+	{
+		SpawnOrDestroyObjectsForInstance( MovieSceneIter.Key().Pin().ToSharedRef(), bDestroyAll );
+	}
+
+	InstanceToPuppetObjectsMap.Empty();
+}
+
+bool FSequencerActorBindingManager::CanPossessObject( UObject& Object ) const
+{
+	return Object.IsA<AActor>();
+}
+
+void FSequencerActorBindingManager::BindPossessableObject( const FGuid& PossessableGuid, UObject& PossessedObject )
+{
+	 // When editing within the level editor, make sure we're bound to the level script node which contains data about possessables.
+	 const bool bCreateIfNotFound = true;
+	 BindToPlayMovieSceneNode( bCreateIfNotFound );
+	
+	 //const FString& PossessableName = Object->GetName();
+	 //const FGuid PossessableGuid = FocusedMovieScene->AddPossessable( PossessableName, Object->GetClass() );
+	 
+	 //if (IsShotFilteringOn())
+	 //{
+	 //	 AddUnfilterableObject(PossessableGuid);
+	 // }
+	 
+	 // Bind the object to the handle
+	 TArray< UObject* > Objects;
+	 Objects.Add( &PossessedObject );
+	 PlayMovieSceneNode->BindPossessableToObjects( PossessableGuid, Objects );
+	 
+	 // A possessable was created so we need to respawn its puppet
+	 // SpawnOrDestroyPuppetObjects( FocusedMovieSceneInstance );
+}
+
+void FSequencerActorBindingManager::UnbindPossessableObjects( const FGuid& PossessableGuid )
+{
+	// @todo Sequencer: Should we remove the pin and unlink?
+}
+
+void FSequencerActorBindingManager::GetRuntimeObjects( const TSharedRef<FMovieSceneInstance>& MovieSceneInstance, const FGuid& ObjectGuid, TArray<UObject*>& OutRuntimeObjects ) const
+{
+	UObject* FoundSpawnedObject = FindPuppetObjectForSpawnableGuid( MovieSceneInstance, ObjectGuid );
+	if( FoundSpawnedObject )
+	{
+		// Spawnable
+		OutRuntimeObjects.Reset();
+		OutRuntimeObjects.Add( FoundSpawnedObject );
+	}
+	else if( PlayMovieSceneNode.IsValid() )
+	{
+		// Possessable
+		OutRuntimeObjects = PlayMovieSceneNode->FindBoundObjects( ObjectGuid );
+	}
+}
+
+UK2Node_PlayMovieScene* FSequencerActorBindingManager::FindPlayMovieSceneNodeInLevelScript( const UMovieScene* MovieScene ) const
+{
+	// Grab the world object for this editor
+	check( MovieScene != NULL );
+	
+	// Search all levels in the specified world
+	for( TArray<ULevel*>::TConstIterator LevelIter( ActorWorld->GetLevels().CreateConstIterator() ); LevelIter; ++LevelIter )
+	{
+		auto* Level = *LevelIter;
+		if( Level != NULL )
+		{
+			// We don't want to create a level script if one doesn't exist yet.  We just want to grab the one
+			// that we already have, if one exists.
+			const bool bDontCreate = true;
+			ULevelScriptBlueprint* LSB = Level->GetLevelScriptBlueprint( bDontCreate );
+			if( LSB != NULL )
+			{
+				TArray<UK2Node_PlayMovieScene*> EventNodes;
+				FBlueprintEditorUtils::GetAllNodesOfClass( LSB, EventNodes );
+				for( auto FoundNodeIter( EventNodes.CreateIterator() ); FoundNodeIter; ++FoundNodeIter )
+				{
+					UK2Node_PlayMovieScene* FoundNode = *FoundNodeIter;
+					if( FoundNode->GetMovieScene() == MovieScene )
+					{
+						return FoundNode;
+					}
+				}
+			}
+		}
+	}
+	
+	return NULL;
+}
+
+
+UK2Node_PlayMovieScene* FSequencerActorBindingManager::CreateNewPlayMovieSceneNode( UMovieScene* MovieScene ) const
+{
+	// Grab the world object for this editor
+	check( MovieScene != NULL );
+	
+	ULevel* Level = ActorWorld->GetCurrentLevel();
+	check( Level != NULL );
+	
+	// Here, we'll create a level script if one does not yet exist.
+	const bool bDontCreate = false;
+	ULevelScriptBlueprint* LSB = Level->GetLevelScriptBlueprint( bDontCreate );
+	if( LSB != NULL )
+	{
+		UEdGraph* TargetGraph = NULL;
+		if( LSB->UbergraphPages.Num() > 0 )
+		{
+			TargetGraph = LSB->UbergraphPages[0]; // Just use the first graph
+		}
+		
+		if( ensure( TargetGraph != NULL ) )
+		{
+			// Figure out a decent place to stick the node
+			const FVector2D NewNodePos = TargetGraph->GetGoodPlaceForNewNode();
+			
+			// Create a new node
+			UK2Node_PlayMovieScene* TemplateNode = NewObject<UK2Node_PlayMovieScene>();
+			TemplateNode->SetMovieScene( MovieScene );
+			return FEdGraphSchemaAction_K2NewNode::SpawnNodeFromTemplate<UK2Node_PlayMovieScene>(TargetGraph, TemplateNode, NewNodePos);;
+		}
+	}
+	
+	return NULL;
+}
+
+UK2Node_PlayMovieScene* FSequencerActorBindingManager::BindToPlayMovieSceneNode( const bool bCreateIfNotFound ) const
+{
+	auto* MovieScene = Sequencer.Pin()->GetFocusedMovieScene();
+	
+	// Update level script
+	UK2Node_PlayMovieScene* FoundPlayMovieSceneNode = FindPlayMovieSceneNodeInLevelScript( MovieScene );
+	if( FoundPlayMovieSceneNode == NULL )
+	{
+		if( bCreateIfNotFound )
+		{
+			// Couldn't find an existing node that uses this MovieScene, so we'll create one now.
+			FoundPlayMovieSceneNode = CreateNewPlayMovieSceneNode( MovieScene );
+			if( FoundPlayMovieSceneNode != NULL )
+			{
+				// Let the user know that we plopped down a new PlayMovieScene event in their level script graph
+				{
+					FNotificationInfo Info( LOCTEXT("AddedPlayMovieSceneEventToLevelScriptGraph", "A new event for this MovieScene was added to this level's script") );
+					Info.bFireAndForget = true;
+					Info.bUseThrobber = false;
+					Info.bUseSuccessFailIcons = false;
+					Info.bUseLargeFont = false;
+					Info.ExpireDuration = 5.0f;		// Stay visible for awhile so the user has time to click "Show Graph" if they want to
+					
+					// @todo sequencer: Needs better artwork (see DefaultStyle.cpp)
+					Info.Image = FEditorStyle::GetBrush( TEXT( "Sequencer.NotificationImage_AddedPlayMovieSceneEvent" ) );
+					
+					struct Local
+					{
+						/**
+						 * Called by our notification's hyperlink to open the Level Script editor and navigate to our newly-created PlayMovieScene node
+						 *
+						 * @param	Sequencer			The Sequencer we're using
+						 * @param	PlayMovieSceneNode	The node that we need to jump to
+						 */
+						static void NavigateToPlayMovieSceneNode( UK2Node_PlayMovieScene* PlayMovieSceneNode )
+						{
+							check( PlayMovieSceneNode != NULL );
+							
+							IBlueprintEditor* BlueprintEditor = NULL;
+							
+							auto* LSB = CastChecked< ULevelScriptBlueprint >( PlayMovieSceneNode->GetBlueprint() );
+							
+							// @todo sequencer: Support using world-centric editing here?  (just need to set EToolkitMode::WorldCentric)
+							FAssetEditorManager::Get().OpenEditorForAsset( LSB );
+							
+							if( BlueprintEditor != NULL )
+							{
+								const bool bRequestRename = false;
+								BlueprintEditor->JumpToHyperlink( PlayMovieSceneNode, false );
+							}
+							else
+							{
+								UE_LOG( LogSequencer, Warning, TEXT( "Unable to open Blueprint Editor to edit newly-created PlayMovieScene event" ) );
+							}
+						}
+					};
+					
+					Info.Hyperlink = FSimpleDelegate::CreateStatic( &Local::NavigateToPlayMovieSceneNode, FoundPlayMovieSceneNode );
+					Info.HyperlinkText = LOCTEXT("AddedPlayMovieSceneEventToLevelScriptGraph_Hyperlink", "Show Graph");
+					
+					TSharedPtr< SNotificationItem > NewNotification = FSlateNotificationManager::Get().AddNotification(Info);
+				}
+			}
+		}
+	}
+	
+	if( PlayMovieSceneNode.Get() != FoundPlayMovieSceneNode )
+	{
+		if( PlayMovieSceneNode.IsValid() )
+		{
+			// Unhook from old node
+			PlayMovieSceneNode.Get()->OnBindingsChanged().RemoveAll( this );
+		}
+		
+		PlayMovieSceneNode = FoundPlayMovieSceneNode;
+		
+		if( PlayMovieSceneNode.IsValid() )
+		{
+			// Bind to the new node
+			PlayMovieSceneNode.Get()->OnBindingsChanged().AddSP( this, &FSequencerActorBindingManager::OnPlayMovieSceneBindingsChanged );
+		}
+	}
+	
+	return PlayMovieSceneNode.Get();
+}
+
+void FSequencerActorBindingManager::OnPlayMovieSceneBindingsChanged()
+{
+	const bool bDestroyAll = false;
+	SpawnOrDestroyObjectsForInstance( Sequencer.Pin()->GetFocusedMovieSceneInstance(), bDestroyAll );
+}
+
+void FSequencerActorBindingManager::DeselectAllPuppetObjects()
+{
+	TArray< AActor* > ActorsToDeselect;
+
+	for( auto MovieSceneIter( InstanceToPuppetObjectsMap.CreateConstIterator() ); MovieSceneIter; ++MovieSceneIter ) 
+	{
+		const TArray< TSharedRef<FPuppetActorInfo> >& PuppetObjects = MovieSceneIter.Value();
+
+		for( auto PuppetObjectIter( PuppetObjects.CreateConstIterator() ); PuppetObjectIter; ++PuppetObjectIter )
+		{
+			if( PuppetObjectIter->Get().GetType() == EPuppetObjectType::Actor )
+			{
+				TSharedRef< FPuppetActorInfo > PuppetActorInfo = StaticCastSharedRef< FPuppetActorInfo >( *PuppetObjectIter );
+				AActor* PuppetActor = PuppetActorInfo->PuppetActor.Get();
+				if( PuppetActor != NULL )
+				{
+					ActorsToDeselect.Add( PuppetActor );
+				}
+			}
+		}
+	}
+
+	if( ActorsToDeselect.Num() > 0 )
+	{
+		GEditor->GetSelectedActors()->BeginBatchSelectOperation();
+		GEditor->GetSelectedActors()->MarkBatchDirty();
+		for( auto ActorIt( ActorsToDeselect.CreateIterator() ); ActorIt; ++ActorIt )
+		{
+			GEditor->GetSelectedActors()->Deselect( *ActorIt );
+		}
+		GEditor->GetSelectedActors()->EndBatchSelectOperation();
+	}
+}
+
+
+void FSequencerActorBindingManager::OnPropagateObjectChanges( UObject* Object )
+{
+	// @todo sequencer major: Many other types of changes to the puppet actor may occur which should be propagated (or disallowed?):
+	//		- Deleting the puppet actor (Del key) should probably delete the reference actor and remove it from the MovieScene
+	//		- Attachment changes
+	//		- Replace/Convert operations?
+
+	// We only care about actor objects for this type of spawner.  Note that sometimes we can get component objects passed in, but we'll
+	// rely on the caller also calling OnObjectPropertyChange for the outer Actor if a component property changes.
+	if( Object->IsA<AActor>() )
+	{
+		for( auto MovieSceneIter( InstanceToPuppetObjectsMap.CreateConstIterator() ); MovieSceneIter; ++MovieSceneIter ) 
+		{
+			const TArray< TSharedRef<FPuppetActorInfo> >& PuppetObjects = MovieSceneIter.Value();
+
+			// Is this an object that we care about?
+			for( auto PuppetObjectIter( PuppetObjects.CreateConstIterator() ); PuppetObjectIter; ++PuppetObjectIter )
+			{
+				if( PuppetObjectIter->Get().GetType() == EPuppetObjectType::Actor )
+				{
+					TSharedRef< FPuppetActorInfo > PuppetActorInfo = StaticCastSharedRef< FPuppetActorInfo >( *PuppetObjectIter );
+					AActor* PuppetActor = PuppetActorInfo->PuppetActor.Get();
+					if( PuppetActor != NULL )
+					{
+						if( PuppetActor == Object )
+						{
+							// @todo sequencer propagation: Don't propagate changes that are being auto-keyed or key-adjusted!
+
+							// Our puppet actor was modified.
+							PropagatePuppetActorChanges( PuppetActorInfo );
+						}
+					}
+				}
+			}
+		}
+	}
+}
+
+
+
+void FSequencerActorBindingManager::PropagatePuppetActorChanges( const TSharedRef< FPuppetActorInfo > PuppetActorInfo )
+{
+	AActor* PuppetActor = PuppetActorInfo->PuppetActor.Get();
+	AActor* TargetActor = NULL;
+
+	{
+		// Find the spawnable for this puppet actor
+		FMovieSceneSpawnable* FoundSpawnable = NULL;
+		TArray< UMovieScene* > MovieScenesBeingEdited = Sequencer.Pin()->GetMovieScenesBeingEdited();
+		for( auto CurMovieSceneIt( MovieScenesBeingEdited.CreateIterator() ); CurMovieSceneIt; ++CurMovieSceneIt )
+		{
+			auto CurMovieScene = *CurMovieSceneIt;
+			FoundSpawnable = CurMovieScene->FindSpawnable( PuppetActorInfo->SpawnableGuid );
+			if( FoundSpawnable != NULL )
+			{
+				break;
+			}
+		}
+		
+		if( ensure( FoundSpawnable != NULL ) )
+		{
+			UClass* ActorClass = FoundSpawnable->GetClass();
+
+			// The puppet actor's class should always be the blueprint that it was spawned from!
+			UClass* SpawnedActorClass = PuppetActor->GetClass();
+			check( ActorClass == SpawnedActorClass );
+
+			// We'll be copying properties into the class default object of the Blueprint's generated class
+			TargetActor = CastChecked<AActor>( ActorClass->GetDefaultObject() );
+		}
+	}
+
+
+	if( PuppetActor != NULL && TargetActor != NULL )
+	{
+		Sequencer.Pin()->CopyActorProperties( PuppetActor, TargetActor );
+	}
+}
+
+
+FGuid FSequencerActorBindingManager::FindSpawnableGuidForPuppetObject( UObject* Object ) const
+{
+	for( auto MovieSceneIter( InstanceToPuppetObjectsMap.CreateConstIterator() ); MovieSceneIter; ++MovieSceneIter ) 
+	{
+		const TArray< TSharedRef<FPuppetActorInfo> >& PuppetObjects = MovieSceneIter.Value();
+
+		// Is this an object that we care about?
+		for( auto PuppetObjectIter( PuppetObjects.CreateConstIterator() ); PuppetObjectIter; ++PuppetObjectIter )
+		{
+			if( PuppetObjectIter->Get().GetType() == EPuppetObjectType::Actor )
+			{
+				TSharedRef< FPuppetActorInfo > PuppetActorInfo = StaticCastSharedRef< FPuppetActorInfo >( *PuppetObjectIter );
+				AActor* PuppetActor = PuppetActorInfo->PuppetActor.Get();
+				if( PuppetActor != NULL )
+				{
+					if( PuppetActor == Object )
+					{
+						// Found it!
+						return PuppetActorInfo->SpawnableGuid;
+					}
+				}
+			}
+		}
+	}
+
+	// Not found
+	return FGuid();
+}
+
+
+UObject* FSequencerActorBindingManager::FindPuppetObjectForSpawnableGuid( const TSharedRef<FMovieSceneInstance>& MovieSceneInstance, const FGuid& Guid ) const
+{
+	const TArray< TSharedRef<FPuppetActorInfo> >& PuppetObjects = InstanceToPuppetObjectsMap.FindChecked( MovieSceneInstance );
+
+	for( auto PuppetObjectIter( PuppetObjects.CreateConstIterator() ); PuppetObjectIter; ++PuppetObjectIter )
+	{
+		if( PuppetObjectIter->Get().GetType() == EPuppetObjectType::Actor )
+		{
+			TSharedRef< FPuppetActorInfo > PuppetActorInfo = StaticCastSharedRef< FPuppetActorInfo >( *PuppetObjectIter );
+			if( PuppetActorInfo->SpawnableGuid == Guid )
+			{
+				return PuppetActorInfo->PuppetActor.Get();
+			}
+		}
+	}
+
+	// Not found
+	return NULL;
+}
+
+
+#undef LOCTEXT_NAMESPACE
\ No newline at end of file
diff --git a/Engine/Source/Editor/Sequencer/Private/SequencerActorBindingManager.h b/Engine/Source/Editor/Sequencer/Private/SequencerActorBindingManager.h
new file mode 100644
--- /dev/null
+++ b/Engine/Source/Editor/Sequencer/Private/SequencerActorBindingManager.h
@@ -0,0 +1,173 @@
+// Copyright 1998-2014 Epic Games, Inc. All Rights Reserved.
+
+#pragma once
+
+class FSequencer;
+
+namespace EPuppetObjectType
+{
+	enum Type
+	{
+		/** Puppet actor */
+		Actor,
+	};
+}
+
+
+/**
+ * Stores the relationship between a spawned puppet object (e.g. actor) and it's data descriptor (e.g. blueprint)
+ */
+class FBasePuppetInfo
+{	
+
+public:
+
+	/** @return Gets the type of puppet */
+	virtual EPuppetObjectType::Type GetType() const = 0;
+};
+
+
+/**
+ * Puppet actor info
+ */
+struct FPuppetActorInfo : public FBasePuppetInfo
+{	
+
+public:
+
+	/** @return Gets the type of puppet */
+	virtual EPuppetObjectType::Type GetType() const override
+	{
+		return EPuppetObjectType::Actor;
+	}
+
+
+public:
+
+	/** Spawnable guid */
+	FGuid SpawnableGuid;
+
+	/** The actual puppet actor that we are actively managing.  These actors actually live in the editor level
+	    just like any other actor, but we'll directly manage their lifetimes while in the editor */
+	TWeakObjectPtr< class AActor > PuppetActor;	
+};
+
+
+/**
+ * Manages spawned puppet objects for Sequencer
+ */
+class FSequencerActorObjectSpawner : public ISequencerObjectBindingManager, public TSharedFromThis<FSequencerActorObjectSpawner>
+{
+
+public:
+
+	/**
+	 * Constructor
+	 *
+	 * @param	InWorld	The world to spawn actors in
+	 */
+	FSequencerActorObjectSpawner( UWorld* InWorld, TSharedRef<FSequencer> InSequencer );
+
+
+	/** Destructor */
+	virtual ~FSequencerActorObjectSpawner();
+
+
+	/** ISequencerObjectBindingManager interface */
+	virtual bool AllowsSpawnableObjects() const override { return true; }
+	virtual FGuid FindGuidForObject( const UMovieScene& MovieScene, UObject& Object ) const override;
+	virtual void SpawnOrDestroyObjectsForInstance( TSharedRef<FMovieSceneInstance> MovieSceneInstance, bool bDestroyAll ) override;
+	virtual void DestroyAllSpawnedObjects() override;
+	virtual bool CanPossessObject( UObject& Object ) const override;
+	virtual void BindPossessableObject( const FGuid& PossessableGuid, UObject& PossessedObject ) override;
+	virtual void UnbindPossessableObjects( const FGuid& PossessableGuid ) override;
+	virtual void GetRuntimeObjects( const TSharedRef<FMovieSceneInstance>& MovieSceneInstance, const FGuid& ObjectGuid, TArray<UObject*>& OutRuntimeObjects ) const override;
+
+protected:
+	/**
+	 * Finds the spawnable guid for the specified object, if one exists
+	 *
+	 * @param	Object	The object to find a spawnable for
+	 *
+	 * @return	The spawnable guid for the specified object, or an Invalid guid if none was found
+	 */
+	FGuid FindSpawnableGuidForPuppetObject( UObject* Object ) const;
+	
+	/**
+	 * Given a guid for a spawnable, attempts to find a spawned puppet object for that guid
+	 *
+	 * @param	MovieSceneInstance	The instance of a movie scene which the puppet object would be spawned for
+	 * @param	Guid	The guid of the spawnable to find a puppet object for
+	 * @return	The puppet object, or NULL if we don't recognize this guid
+	 */
+	UObject* FindPuppetObjectForSpawnableGuid( const TSharedRef<FMovieSceneInstance>& MovieSceneInstance, const FGuid& Guid ) const;
+	
+	/** Deselects all puppet objects in the editor */
+	void DeselectAllPuppetObjects();
+
+
+	/**
+	 * Called by FSequencerObjectChangeListener after a non-auto-keyed property has changed on any object.
+	 *
+	 * @param	Object	Object that was modified
+	 */
+	void OnPropagateObjectChanges( UObject* Object );
+
+
+	/**
+	 * Called to refresh either the reference actor or the data descriptor blueprint after its puppet proxy actor has been changed
+	 *
+	 * @param	PuppetActorInfo		Info about the puppet actor that needs to be refreshed
+	 */
+	void PropagatePuppetActorChanges( const TSharedRef< FPuppetActorInfo > PuppetActorInfo );
+	
+	
+	/**
+	 * Given a MovieScene to search for, attempts to find a LevelScript node that is playing the
+	 * supplied MovieScene.  Only valid to call when in world-centric mode.
+	 *
+	 * @param	MovieScene	The MovieScene to search for
+	 *
+	 * @return	The LevelScript node we found, or NULL for no matches
+	 */
+	class UK2Node_PlayMovieScene* FindPlayMovieSceneNodeInLevelScript( const class UMovieScene* MovieScene );
+	
+	/**
+	 * Creates a new "PlayMovieScene" node in the level script for the current level and associates it with the
+	 * supplied MovieScene.  Only valid to call when in world-centric mode.
+	 *
+	 * @param	MovieScene	The MovieScene to associate with the newly-created node
+	 *
+	 * @return	The LevelScript node we created, or NULL on failure (does not assert)
+	 */
+	class UK2Node_PlayMovieScene* CreateNewPlayMovieSceneNode( UMovieScene* MovieScene );
+	
+	
+	/** Called by UK2Node_PlayMovieScene when the bindings for that node are changed through Kismet */
+	void OnPlayMovieSceneBindingsChanged();
+	/**
+	 * Binds to a 'PlayMovieScene' node in level script that is associated with the MovieScene we're editing.  If
+	 * requested, can also create a new node.  Only valid to call when in world-centric mode.
+	 *
+	 * @param	bCreateIfNotFound	True to create a node if we can't find one
+	 *
+	 * @return	The LevelScript node we found or created, or NULL if nothing is still bound
+	 */
+	class UK2Node_PlayMovieScene* BindToPlayMovieSceneNode( const bool bCreateIfNotFound );
+	
+private:
+	TWeakObjectPtr< UWorld > ActorWorld;
+	
+	/** 'PlayMovieScene' node in level script that we're currently associating with the active MovieScene.  This node contains the bindings
+	 information that we need to be able to preview possessed actors in the level */
+	TWeakObjectPtr< UK2Node_PlayMovieScene > PlayMovieSceneNode;
+
+	/** The Sequencer that owns this spawner  */
+	TWeakPtr<FSequencer> Sequencer;
+
+	/**
+	 * List of puppet objects that we've spawned and are actively managing. 
+	 */
+	TMap< TWeakPtr<FMovieSceneInstance>, TArray< TSharedRef< FPuppetActorInfo > > > InstanceToPuppetObjectsMap;
+};
+
diff --git a/Engine/Source/Editor/Sequencer/Private/SequencerAssetEditor.cpp b/Engine/Source/Editor/Sequencer/Private/SequencerAssetEditor.cpp
--- a/Engine/Source/Editor/Sequencer/Private/SequencerAssetEditor.cpp
+++ b/Engine/Source/Editor/Sequencer/Private/SequencerAssetEditor.cpp
@@ -10,6 +10,8 @@
 #include "Sequencer.h"
 #include "MovieScene.h"
 #include "EditorWidgetsModule.h"
+#include "SequencerObjectSpawner.h"
+#include "SequencerObjectChangeListener.h"
 
 #define LOCTEXT_NAMESPACE "Sequencer"
 
@@ -46,7 +48,7 @@ void FSequencerAssetEditor::UnregisterTabSpawners(const TSharedRef<class FTabMan
 
 void FSequencerAssetEditor::InitSequencerAssetEditor( const EToolkitMode::Type Mode, const TSharedPtr< class IToolkitHost >& InitToolkitHost, UMovieScene* InRootMovieScene, const TArray<FOnCreateTrackEditor>& TrackEditorDelegates, bool bEditWithinLevelEditor )
 {
-	if(InitToolkitHost.IsValid())
+	
 	{
 		const TSharedRef<FTabManager::FLayout> StandaloneDefaultLayout = FTabManager::NewLayout("Standalone_Sequencer_Layout")
 		->AddArea
@@ -63,10 +65,22 @@ void FSequencerAssetEditor::InitSequencerAssetEditor( const EToolkitMode::Type M
 		const bool bCreateDefaultToolbar = false;
 		FAssetEditorToolkit::InitAssetEditor(Mode, InitToolkitHost, SequencerDefs::SequencerAppIdentifier, StandaloneDefaultLayout, bCreateDefaultStandaloneMenu, bCreateDefaultToolbar, InRootMovieScene);
 	}
-
-	// Init Sequencer 
+	
+	// Init Sequencer
 	Sequencer = MakeShareable(new FSequencer);
-	Sequencer->InitSequencer( InRootMovieScene, InitToolkitHost, TrackEditorDelegates, bEditWithinLevelEditor);
+	
+	FSequencerInitParams SequencerInitParams;
+	SequencerInitParams.ObjectChangeListener = MakeShareable( new FSequencerObjectChangeListener( Sequencer.ToSharedRef(), bEditWithinLevelEditor ) );
+	
+	SequencerInitParams.ObjectBindingManager = MakeShareable( new FSequencerActorBindingManager( InitToolkitHost->GetWorld(), SequencerInitParams.ObjectChangeListener.ToSharedRef(), Sequencer.ToSharedRef() ) );
+
+	SequencerInitParams.RootMovieScene = InRootMovieScene;
+	
+	SequencerInitParams.bEditWithinLevelEditor = bEditWithinLevelEditor;
+	
+	SequencerInitParams.ToolkitHost = InitToolkitHost;
+	
+	Sequencer->InitSequencer( SequencerInitParams, TrackEditorDelegates );
 
 	if(bEditWithinLevelEditor)
 	{
diff --git a/Engine/Source/Editor/Sequencer/Private/SequencerModule.cpp b/Engine/Source/Editor/Sequencer/Private/SequencerModule.cpp
--- a/Engine/Source/Editor/Sequencer/Private/SequencerModule.cpp
+++ b/Engine/Source/Editor/Sequencer/Private/SequencerModule.cpp
@@ -14,10 +14,23 @@ class FSequencerModule : public ISequencerModule
 {
 
 	/** ISequencerModule interface */
-	virtual TSharedPtr<ISequencer> CreateSequencer( UMovieScene* RootMovieScene ) override
+	virtual TSharedPtr<ISequencer> CreateSequencer( UMovieScene* RootMovieScene, TSharedRef<ISequencerObjectBindingManager> ObjectBindingManager ) override
 	{
 		TSharedRef< FSequencer > Sequencer = MakeShareable(new FSequencer);
-		Sequencer->InitSequencer( RootMovieScene, nullptr, TrackEditorDelegates, false );
+		
+		FSequencerInitParams SequencerInitParams;
+		SequencerInitParams.ObjectChangeListener = MakeShareable( new FSequencerObjectChangeListener( Sequencer, false ) );
+		
+		SequencerInitParams.ObjectBindingManager = ObjectBindingManager;
+		
+		SequencerInitParams.RootMovieScene = RootMovieScene;
+		
+		SequencerInitParams.bEditWithinLevelEditor = false;
+		
+		SequencerInitParams.ToolkitHost = nullptr;
+		
+		Sequencer->InitSequencer( SequencerInitParams, TrackEditorDelegates );
+		
 		return Sequencer;
 	}
 	
diff --git a/Engine/Source/Editor/Sequencer/Private/SequencerObjectChangeListener.cpp b/Engine/Source/Editor/Sequencer/Private/SequencerObjectChangeListener.cpp
--- a/Engine/Source/Editor/Sequencer/Private/SequencerObjectChangeListener.cpp
+++ b/Engine/Source/Editor/Sequencer/Private/SequencerObjectChangeListener.cpp
@@ -7,8 +7,9 @@
 
 DEFINE_LOG_CATEGORY_STATIC(LogSequencerTools, Log, All);
 
-FSequencerObjectChangeListener::FSequencerObjectChangeListener( TSharedRef<ISequencer> InSequencer )
+FSequencerObjectChangeListener::FSequencerObjectChangeListener( TSharedRef<ISequencer> InSequencer, bool bInListenForActorsOnly )
 	: Sequencer( InSequencer )
+	, bListenForActorsOnly( bInListenForActorsOnly )
 {
 	FCoreDelegates::OnPreObjectPropertyChanged.AddRaw(this, &FSequencerObjectChangeListener::OnObjectPreEditChange);
 	FCoreDelegates::OnObjectPropertyChanged.AddRaw(this, &FSequencerObjectChangeListener::OnObjectPostEditChange);
@@ -28,6 +29,13 @@ void FSequencerObjectChangeListener::OnPropertyChanged( const TArray<UObject*>&
 	ClassToPropertyChangedMap.FindRef( PropertyValue->GetProperty()->GetClass() ).Broadcast( ChangedObjects, *PropertyValue, bRequireAutoKey );
 }
 
+bool FSequencerObjectChangeListener::IsObjectValidForListening( UObject* Object ) const
+{
+	// @todo Sequencer - Pre/PostEditChange is sometimes called for inner objects of other objects (like actors with components)
+	// We only want the outer object so assume it's an actor for now
+	return bListenForActorsOnly ? Object->IsA<AActor>() : true;
+}
+
 FOnAnimatablePropertyChanged& FSequencerObjectChangeListener::GetOnAnimatablePropertyChanged( TSubclassOf<UProperty> PropertyClass )
 {
 	return ClassToPropertyChangedMap.FindOrAdd( PropertyClass );
@@ -47,9 +55,7 @@ void FSequencerObjectChangeListener::OnObjectPreEditChange( UObject* Object )
 		// Register with the property editor module that we'd like to know about animated float properties that change
 		FPropertyEditorModule& PropertyEditor = FModuleManager::Get().LoadModuleChecked<FPropertyEditorModule>( "PropertyEditor" );
 
-		// @todo Sequencer PreEditChange is sometimes called for inner objects of other objects (like actors with components)
-		// We only want the outer object so assume it's an actor for now
-		if( Object->IsA( AActor::StaticClass() ) )
+		if( IsObjectValidForListening(Object) )
 		{
 			TSharedPtr<IPropertyChangeListener> PropertyChangeListener = ActivePropertyChangeListeners.FindRef( Object );
 
@@ -66,7 +72,7 @@ void FSequencerObjectChangeListener::OnObjectPreEditChange( UObject* Object )
 				Settings.bIgnoreArrayProperties = true;
 				Settings.bIgnoreObjectProperties = true;
 				// Property flags which must be on the property
-				Settings.RequiredPropertyFlags = CPF_Interp;
+				Settings.RequiredPropertyFlags = 0;
 				// Property flags which cannot be on the property
 				Settings.DisallowedPropertyFlags = CPF_EditConst;
 
@@ -80,16 +86,12 @@ void FSequencerObjectChangeListener::OnObjectPostEditChange( UObject* Object, FP
 {
 	if( Object )
 	{
-		// @todo Sequencer - Pre/PostEditChange is sometimes called for inner objects of other objects (like actors with components)
-		// We only want the outer object so assume it's an actor for now
-		bool bObjectIsActor = Object->IsA( AActor::StaticClass() );
+		bool bIsObjectValid = IsObjectValidForListening( Object );
 
-		// Default to propagating changes to objects only if they are actors
-		// if this change is handled by auto-key we will not propagate changes
-		bool bShouldPropagateChanges = bObjectIsActor;
+		bool bShouldPropagateChanges = bIsObjectValid;
 
 		// We only care if we are in auto-key mode and not attempting to change properties of a CDO (which cannot be animated)
-		if( Sequencer.IsValid() && Sequencer.Pin()->IsAutoKeyEnabled() && bObjectIsActor && !Object->HasAnyFlags(RF_ClassDefaultObject) )
+		if( Sequencer.IsValid() && Sequencer.Pin()->IsAutoKeyEnabled() && bIsObjectValid && !Object->HasAnyFlags(RF_ClassDefaultObject) )
 		{
 			TSharedPtr< IPropertyChangeListener > Listener;
 			ActivePropertyChangeListeners.RemoveAndCopyValue( Object, Listener );
diff --git a/Engine/Source/Editor/Sequencer/Private/SequencerObjectChangeListener.h b/Engine/Source/Editor/Sequencer/Private/SequencerObjectChangeListener.h
--- a/Engine/Source/Editor/Sequencer/Private/SequencerObjectChangeListener.h
+++ b/Engine/Source/Editor/Sequencer/Private/SequencerObjectChangeListener.h
@@ -10,7 +10,7 @@ class IPropertyHandle;
 class FSequencerObjectChangeListener : public ISequencerObjectChangeListener
 {
 public:
-	FSequencerObjectChangeListener( TSharedRef<ISequencer> InSequencer );
+	FSequencerObjectChangeListener( TSharedRef<ISequencer> InSequencer, bool bInListenForActorsOnly );
 	~FSequencerObjectChangeListener();
 
 	/** ISequencerObjectChangeListener interface */
@@ -47,6 +47,11 @@ private:
 	 */
 	void OnPropertyChanged( const TArray<UObject*>& ChangedObjects, const TSharedRef< const IPropertyHandle> PropertyValue, bool bRequireAutoKey );
 
+	/**
+	 * @return True if an object is valid for listening to property changes 
+	 */
+	bool IsObjectValidForListening( UObject* Object ) const;
+
 private:
 	/** Mapping of object to a listener used to check for property changes */
 	TMap< TWeakObjectPtr<UObject>, TSharedPtr<class IPropertyChangeListener> > ActivePropertyChangeListeners;
@@ -59,4 +64,7 @@ private:
 
 	/** The owning sequencer */
 	TWeakPtr< ISequencer > Sequencer;
+
+	/** True to listen for actors only */
+	bool bListenForActorsOnly;
 };
\ No newline at end of file
diff --git a/Engine/Source/Editor/Sequencer/Private/SequencerObjectSpawner.cpp b/Engine/Source/Editor/Sequencer/Private/SequencerObjectSpawner.cpp
deleted file mode 100644
diff --git a/Engine/Source/Editor/Sequencer/Private/SequencerObjectSpawner.h b/Engine/Source/Editor/Sequencer/Private/SequencerObjectSpawner.h
deleted file mode 100644
diff --git a/Engine/Source/Editor/Sequencer/Public/ISequencer.h b/Engine/Source/Editor/Sequencer/Public/ISequencer.h
--- a/Engine/Source/Editor/Sequencer/Public/ISequencer.h
+++ b/Engine/Source/Editor/Sequencer/Public/ISequencer.h
@@ -111,7 +111,7 @@ public:
 	/**
 	 * @return Returns the object change listener for sequencer instance
 	 */
-	virtual class ISequencerObjectChangeListener& GetObjectChangeListener() const = 0;
+	virtual class ISequencerObjectChangeListener& GetObjectChangeListener() = 0;
 
 	virtual void NotifyMovieSceneDataChanged() = 0;
 };
diff --git a/Engine/Source/Editor/Sequencer/Public/ISequencerModule.h b/Engine/Source/Editor/Sequencer/Public/ISequencerModule.h
--- a/Engine/Source/Editor/Sequencer/Public/ISequencerModule.h
+++ b/Engine/Source/Editor/Sequencer/Public/ISequencerModule.h
@@ -5,6 +5,8 @@
 #include "ModuleInterface.h"
 #include "ISequencer.h"
 
+class ISequencerObjectBindingManager;
+
 /**
  * A delegate which will create an auto-key handler
  */
@@ -25,7 +27,7 @@ public:
 	 * @param 	InRootMovieScene	The movie scene to edit
 	 * @return	Interface to the new editor
 	 */
-	virtual TSharedPtr<ISequencer> CreateSequencer( UMovieScene* InRootMovieScene ) = 0;
+	virtual TSharedPtr<ISequencer> CreateSequencer( UMovieScene* InRootMovieScene, TSharedRef<ISequencerObjectBindingManager> ObjectBindingManager ) = 0;
 
 	/**
 	 * Creates a new instance of a Sequencer, the editor for MovieScene assets in an asset editor
diff --git a/Engine/Source/Editor/Sequencer/Public/ISequencerObjectBindingManager.h b/Engine/Source/Editor/Sequencer/Public/ISequencerObjectBindingManager.h
new file mode 100644
--- /dev/null
+++ b/Engine/Source/Editor/Sequencer/Public/ISequencerObjectBindingManager.h
@@ -0,0 +1,53 @@
+// Copyright 1998-2014 Epic Games, Inc. All Rights Reserved.
+
+#pragma once
+
+
+class ISequencerObjectBindingManager
+{
+public:
+	virtual ~ISequencerObjectBindingManager() {}
+	
+	/**
+	 * @return true to allow objects to be spawned by sequencer, false if only existing objects can be possessed
+	 */
+	virtual bool AllowsSpawnableObjects() const = 0;
+	
+	/**
+	 * Finds an existing guid for a bound object
+	 *
+	 * @param MovieScene The movie scene that may contain the binding
+	 * @param The object to get the guid for
+	 * @return The object guid or an invalid guid if not found
+	 */
+	virtual FGuid FindGuidForObject( const UMovieScene& MovieScene, UObject& Object ) const = 0;
+	
+	/**
+	 * Called to spawn or destroy objects for a movie instance
+	 *
+	 * @param MovieSceneInstance	The instance to spawn or destroy objects for
+	 * @param bDestroyAll			If true, destroy all spawned objects for the instance, if false only destroy unused objects
+	 */
+	virtual void SpawnOrDestroyObjectsForInstance( TSharedRef<FMovieSceneInstance> MovieSceneInstance, bool bDestroyAll ) = 0;
+	
+	virtual void DestroyAllSpawnedObjects() = 0;
+	
+	virtual bool CanPossessObject( UObject& Object ) const = 0;
+	
+	/**
+	 * Called when Sequencer has created an object binding for a possessable object
+	 * 
+	 * @param PossessableGuid	The guid used to map to the possessable object.  Note the guid can be bound to multiple objects at once
+	 * @param PossessedObject	The runtime object which was possessed.
+	 */
+	virtual void BindPossessableObject( const FGuid& PossessableGuid, UObject& PossessedObject ) = 0;
+	
+	/**
+	 * Unbinds all possessable objects from the provided guid
+	 *
+	 * @param PossessableGuid	The guid bound to possessable objects that should be removed
+	 */
+	virtual void UnbindPossessableObjects( const FGuid& PossessableGuid ) = 0;
+	
+	virtual void GetRuntimeObjects( const TSharedRef<FMovieSceneInstance>& MovieSceneInstance, const FGuid& ObjectGuid, TArray<UObject*>& OutRuntimeObjects ) const = 0;
+};
diff --git a/Engine/Source/Editor/Sequencer/Public/MovieSceneTrackEditor.h b/Engine/Source/Editor/Sequencer/Public/MovieSceneTrackEditor.h
--- a/Engine/Source/Editor/Sequencer/Public/MovieSceneTrackEditor.h
+++ b/Engine/Source/Editor/Sequencer/Public/MovieSceneTrackEditor.h
@@ -86,25 +86,21 @@ public:
 		}
 	}
 	
-	bool TrackForObjectExists( UObject* Object, TSubclassOf<class UMovieSceneTrack> TrackClass, FName UniqueTypeName )
+	/**
+	 * Finds or creates a binding to an object
+	 *
+	 * @param Object	The object to create a binding for
+	 * @return A handle to the binding or an invalid handle if the object could not be bound
+	 */
+	FGuid FindOrCreateHandleToObject( UObject* Object )
 	{
-		check( UniqueTypeName != NAME_None );
-
-		// Get a handle bound to the key/section we are adding so we know what objects to change during playback
-		FGuid ObjectHandle = GetSequencer()->GetHandleToObject( Object );
-		UMovieScene* MovieScene = GetSequencer()->GetFocusedMovieScene();
-
-		UMovieSceneTrack* Type = MovieScene->FindTrack( TrackClass, ObjectHandle, UniqueTypeName );
-
-		return Type != NULL;
+		return GetSequencer()->GetHandleToObject( Object );
 	}
 
-	UMovieSceneTrack* GetTrackForObject( UObject* Object, TSubclassOf<class UMovieSceneTrack> TrackClass, FName UniqueTypeName )
+	UMovieSceneTrack* GetTrackForObject( const FGuid& ObjectHandle, TSubclassOf<class UMovieSceneTrack> TrackClass, FName UniqueTypeName )
 	{
 		check( UniqueTypeName != NAME_None );
 
-		// Get a handle bound to the key/section we are adding so we know what objects to change during playback
-		FGuid ObjectHandle = GetSequencer()->GetHandleToObject( Object );
 		UMovieScene* MovieScene = GetSequencer()->GetFocusedMovieScene();
 
 		UMovieSceneTrack* Type = MovieScene->FindTrack( TrackClass, ObjectHandle, UniqueTypeName );
diff --git a/Engine/Source/Editor/UMGEditor/Private/UMGSequencerObjectBindingManager.cpp b/Engine/Source/Editor/UMGEditor/Private/UMGSequencerObjectBindingManager.cpp
new file mode 100644
--- /dev/null
+++ b/Engine/Source/Editor/UMGEditor/Private/UMGSequencerObjectBindingManager.cpp
@@ -0,0 +1,49 @@
+// Copyright 1998-2014 Epic Games, Inc. All Rights Reserved.
+
+#include "UMGEditorPrivatePCH.h"
+#include "UMGSequencerObjectBindingManager.h"
+
+FUMGSequencerObjectBindingManager::FUMGSequencerObjectBindingManager( UWidgetBlueprint& InWidgetBlueprint )
+	: WidgetBlueprint( &InWidgetBlueprint )
+{
+	
+}
+
+FGuid FUMGSequencerObjectBindingManager::FindGuidForObject( const UMovieScene& MovieScene, UObject& Object ) const
+{
+	return PreviewObjectToGuidMap.FindRef( &Object );
+}
+
+
+bool FUMGSequencerObjectBindingManager::CanPossessObject( UObject& Object ) const
+{
+	return !Object.IsIn( WidgetBlueprint.Get() );
+}
+
+void FUMGSequencerObjectBindingManager::BindPossessableObject( const FGuid& PossessableGuid, UObject& PossessedObject )
+{
+	PreviewObjectToGuidMap.Add( &PossessedObject, PossessableGuid );
+}
+
+void FUMGSequencerObjectBindingManager::UnbindPossessableObjects( const FGuid& PossessableGuid )
+{
+	for( auto It = PreviewObjectToGuidMap.CreateIterator(); It; ++It )
+	{
+		if( It.Value() == PossessableGuid )
+		{
+			It.RemoveCurrent();
+		}
+	}
+}
+
+void FUMGSequencerObjectBindingManager::GetRuntimeObjects( const TSharedRef<FMovieSceneInstance>& MovieSceneInstance, const FGuid& ObjectGuid, TArray<UObject*>& OutRuntimeObjects ) const
+{
+	for( auto It = PreviewObjectToGuidMap.CreateConstIterator(); It; ++It )
+	{
+		UObject* Object = It.Key().Get();
+		if( Object )
+		{
+			OutRuntimeObjects.Add( Object );
+		}
+	}
+}
diff --git a/Engine/Source/Editor/UMGEditor/Private/UMGSequencerObjectBindingManager.h b/Engine/Source/Editor/UMGEditor/Private/UMGSequencerObjectBindingManager.h
new file mode 100644
--- /dev/null
+++ b/Engine/Source/Editor/UMGEditor/Private/UMGSequencerObjectBindingManager.h
@@ -0,0 +1,27 @@
+// Copyright 1998-2014 Epic Games, Inc. All Rights Reserved.
+
+#pragma once
+
+#include "ISequencerObjectBindingManager.h"
+
+class FUMGSequencerObjectBindingManager : public ISequencerObjectBindingManager
+{
+public:
+	FUMGSequencerObjectBindingManager( UWidgetBlueprint& InWidgetBlueprint );
+	
+	/** ISequencerObjectBindingManager interface */
+	virtual bool AllowsSpawnableObjects() const override { return false; }
+	virtual FGuid FindGuidForObject( const UMovieScene& MovieScene, UObject& Object ) const override;
+	virtual void SpawnOrDestroyObjectsForInstance( TSharedRef<FMovieSceneInstance> MovieSceneInstance, bool bDestroyAll ) {}
+	virtual void DestroyAllSpawnedObjects()  {}
+	virtual bool CanPossessObject( UObject& Object ) const override;
+	virtual void BindPossessableObject( const FGuid& PossessableGuid, UObject& PossessedObject ) override;
+	virtual void UnbindPossessableObjects( const FGuid& PossessableGuid ) override;
+	virtual void GetRuntimeObjects( const TSharedRef<FMovieSceneInstance>& MovieSceneInstance, const FGuid& ObjectGuid, TArray<UObject*>& OutRuntimeObjects ) const override;
+	
+private:
+	/** Mapping of preview objects to sequencer guids */
+	TMap< TWeakObjectPtr<UObject>, FGuid> PreviewObjectToGuidMap;
+	
+	TWeakObjectPtr<UWidgetBlueprint> WidgetBlueprint;
+};
diff --git a/Engine/Source/Editor/UMGEditor/Private/WidgetBlueprintEditor.cpp b/Engine/Source/Editor/UMGEditor/Private/WidgetBlueprintEditor.cpp
--- a/Engine/Source/Editor/UMGEditor/Private/WidgetBlueprintEditor.cpp
+++ b/Engine/Source/Editor/UMGEditor/Private/WidgetBlueprintEditor.cpp
@@ -232,7 +232,10 @@ TSharedPtr<ISequencer>& FWidgetBlueprintEditor::GetSequencer()
 		UWidgetBlueprint* Blueprint = GetWidgetBlueprintObj();
 
 		UMovieScene* MovieScene = Blueprint->AnimationData.Num() ? Blueprint->AnimationData[0] : GetDefaultMovieScene();
-		Sequencer = FModuleManager::LoadModuleChecked< ISequencerModule >("Sequencer").CreateSequencer(MovieScene);
+		
+		TSharedRef<FUMGSequencerObjectBindingManager> ObjectBindingManager = MakeShareable( new FUMGSequencerObjectBindingManager( *Blueprint ) );
+		
+		Sequencer = FModuleManager::LoadModuleChecked< ISequencerModule >("Sequencer").CreateSequencer( MovieScene, ObjectBindingManager );
 	}
 
 	return Sequencer;
diff --git a/Engine/Source/Runtime/MovieSceneCore/Classes/MovieScene.h b/Engine/Source/Runtime/MovieSceneCore/Classes/MovieScene.h
--- a/Engine/Source/Runtime/MovieSceneCore/Classes/MovieScene.h
+++ b/Engine/Source/Runtime/MovieSceneCore/Classes/MovieScene.h
@@ -267,7 +267,7 @@ public:
 	 *
 	 * @return	Spawnable object that was found (or NULL if not found)
 	 */
-	virtual struct FMovieSceneSpawnable* FindSpawnableForCounterpart( UObject* GamePreviewObject );
+	virtual const struct FMovieSceneSpawnable* FindSpawnableForCounterpart( UObject* GamePreviewObject ) const;
 
 	/**
 	 * @return Returns the number of spawnables in this MovieScene
diff --git a/Engine/Source/Runtime/MovieSceneCore/Private/MovieScene.cpp b/Engine/Source/Runtime/MovieSceneCore/Private/MovieScene.cpp
--- a/Engine/Source/Runtime/MovieSceneCore/Private/MovieScene.cpp
+++ b/Engine/Source/Runtime/MovieSceneCore/Private/MovieScene.cpp
@@ -129,7 +129,7 @@ FMovieSceneSpawnable* UMovieScene::FindSpawnable( const FGuid& Guid )
 }
 
 
-FMovieSceneSpawnable* UMovieScene::FindSpawnableForCounterpart( UObject* GamePreviewObject )
+const FMovieSceneSpawnable* UMovieScene::FindSpawnableForCounterpart( UObject* GamePreviewObject ) const
 {
 	check( GamePreviewObject != NULL );
 
@@ -137,7 +137,7 @@ FMovieSceneSpawnable* UMovieScene::FindSpawnableForCounterpart( UObject* GamePre
 	const bool bIsGamePreviewObject = !!( GamePreviewObject->GetOutermost()->PackageFlags & PKG_PlayInEditor );
 	check( bIsGamePreviewObject );
 
-	for( auto CurSpawnableIt( Spawnables.CreateIterator() ); CurSpawnableIt; ++CurSpawnableIt )
+	for( auto CurSpawnableIt( Spawnables.CreateConstIterator() ); CurSpawnableIt; ++CurSpawnableIt )
 	{
 		auto& CurSpawnable = *CurSpawnableIt;
 		if( CurSpawnable.GetCounterpartGamePreviewObject() == GamePreviewObject )
@@ -272,8 +272,10 @@ void UMovieScene::RemoveObjectBinding( const FGuid& Guid )
 UMovieSceneTrack* UMovieScene::FindTrack( TSubclassOf<UMovieSceneTrack> TrackClass, const FGuid& ObjectGuid, FName UniqueTrackName ) const
 {
 	UMovieSceneTrack* FoundTrack = NULL;
-	check( UniqueTrackName != NAME_None && ObjectGuid.IsValid() )
-
+	
+	check( UniqueTrackName != NAME_None );
+	check( ObjectGuid.IsValid() );
+	
 	for( int32 ObjectBindingIndex = 0; ObjectBindingIndex < ObjectBindings.Num(); ++ObjectBindingIndex )
 	{
 		const FMovieSceneObjectBinding& ObjectBinding = ObjectBindings[ObjectBindingIndex];
diff --git a/Engine/Source/Runtime/MovieSceneCore/Private/MovieSceneInstance.cpp b/Engine/Source/Runtime/MovieSceneCore/Private/MovieSceneInstance.cpp
--- a/Engine/Source/Runtime/MovieSceneCore/Private/MovieSceneInstance.cpp
+++ b/Engine/Source/Runtime/MovieSceneCore/Private/MovieSceneInstance.cpp
@@ -32,8 +32,9 @@ void FMovieSceneInstance::Update( float Position, float LastPosition, class IMov
 void FMovieSceneInstance::RefreshInstance( IMovieScenePlayer& Player )
 {
 	// Get all the master tracks and create instances for them if needed
-	const TArray<UMovieSceneTrack*>& MasterTracks = MovieScene->GetMasterTracks();	
-	RefreshInstanceMap( MasterTracks, MasterTrackInstances, Player );
+	const TArray<UMovieSceneTrack*>& MasterTracks = MovieScene->GetMasterTracks();
+	TArray<UObject*> Objects;
+	RefreshInstanceMap( MasterTracks, Objects, MasterTrackInstances, Player );
 
 	TSet< FGuid > FoundObjectBindings;
 	// Get all tracks for each object binding and create instances for them if needed
@@ -48,14 +49,14 @@ void FMovieSceneInstance::RefreshInstance( IMovieScenePlayer& Player )
 
 		FoundObjectBindings.Add( ObjectBinding.GetObjectGuid() );
 
-		// Refresh the instance's tracks
-		const TArray<UMovieSceneTrack*>& Tracks = ObjectBinding.GetTracks();
-		RefreshInstanceMap( Tracks, BindingInstance.TrackInstances, Player );
-
 		// Spawn the runtime objects
 		// @todo Sequencer SubMovieScenes: We need to know which actors were removed and which actors were added so we know which saved actor state to restore/create
 		BindingInstance.RuntimeObjects.Empty();
 		Player.GetRuntimeObjects( SharedThis( this ), BindingInstance.ObjectGuid, BindingInstance.RuntimeObjects );
+		
+		// Refresh the instance's tracks
+		const TArray<UMovieSceneTrack*>& Tracks = ObjectBinding.GetTracks();
+		RefreshInstanceMap( Tracks, BindingInstance.RuntimeObjects, BindingInstance.TrackInstances, Player );
 	}
 
 
@@ -72,7 +73,7 @@ void FMovieSceneInstance::RefreshInstance( IMovieScenePlayer& Player )
 }
 
 
-void FMovieSceneInstance::RefreshInstanceMap( const TArray<UMovieSceneTrack*>& Tracks, FMovieSceneInstanceMap& InstanceMap, class IMovieScenePlayer& Player )
+void FMovieSceneInstance::RefreshInstanceMap( const TArray<UMovieSceneTrack*>& Tracks, const TArray<UObject*>& RuntimeObjects, FMovieSceneInstanceMap& TrackInstances, IMovieScenePlayer& Player  )
 {
 	// All the tracks we found during this pass
 	TSet< UMovieSceneTrack* > FoundTracks;
@@ -86,26 +87,26 @@ void FMovieSceneInstance::RefreshInstanceMap( const TArray<UMovieSceneTrack*>& T
 		FoundTracks.Add( Track );
 
 		// See if the track has an instance
-		TSharedPtr<IMovieSceneTrackInstance> Instance = InstanceMap.FindRef( Track );
+		TSharedPtr<IMovieSceneTrackInstance> Instance = TrackInstances.FindRef( Track );
 		if ( !Instance.IsValid() )
 		{
 			// The track does not have an instance, create one
 			Instance = Track->CreateInstance();
 
-			Instance->RefreshInstance( Player );
+			Instance->RefreshInstance( RuntimeObjects, Player );
 
-			InstanceMap.Add( Track, Instance );
+			TrackInstances.Add( Track, Instance );
 		}
 		else
 		{
 			// The track has an instance, refresh it
-			Instance->RefreshInstance( Player );
+			Instance->RefreshInstance( RuntimeObjects, Player );
 		}
 
 	}
 
 	// Remove entries which no longer have a track associated with them
-	FMovieSceneInstanceMap::TIterator It = InstanceMap.CreateIterator();
+	FMovieSceneInstanceMap::TIterator It = TrackInstances.CreateIterator();
 	for( ; It; ++It )
 	{
 		if( !FoundTracks.Contains( It.Key().Get() ) )
diff --git a/Engine/Source/Runtime/MovieSceneCore/Public/IMovieSceneTrackInstance.h b/Engine/Source/Runtime/MovieSceneCore/Public/IMovieSceneTrackInstance.h
--- a/Engine/Source/Runtime/MovieSceneCore/Public/IMovieSceneTrackInstance.h
+++ b/Engine/Source/Runtime/MovieSceneCore/Public/IMovieSceneTrackInstance.h
@@ -24,5 +24,5 @@ public:
 	/**
 	 * Refreshes the current instance
 	 */
-	virtual void RefreshInstance( class IMovieScenePlayer& Player ) = 0;
+	virtual void RefreshInstance( const TArray<UObject*>& RuntimeObjects, class IMovieScenePlayer& Player ) = 0;
 };
diff --git a/Engine/Source/Runtime/MovieSceneCore/Public/MovieSceneInstance.h b/Engine/Source/Runtime/MovieSceneCore/Public/MovieSceneInstance.h
--- a/Engine/Source/Runtime/MovieSceneCore/Public/MovieSceneInstance.h
+++ b/Engine/Source/Runtime/MovieSceneCore/Public/MovieSceneInstance.h
@@ -25,7 +25,7 @@ public:
 	/** @return The movie scene being instanced */
 	UMovieScene* GetMovieScene() { return MovieScene.Get(); }
 private:
-	void RefreshInstanceMap( const TArray<UMovieSceneTrack*>& Tracks, FMovieSceneInstanceMap& InstanceMap, class IMovieScenePlayer& Player  );
+	void RefreshInstanceMap( const TArray<UMovieSceneTrack*>& Tracks, const TArray<UObject*>& RuntimeObjects, FMovieSceneInstanceMap& TrackInstances, class IMovieScenePlayer& Player );
 private:
 	/** A paring of a runtime object guid to a runtime data and track instances animating the runtime objects */
 	struct FMovieSceneObjectBindingInstance
diff --git a/Engine/Source/Runtime/MovieSceneCoreTypes/Private/Instances/MovieSceneAnimationTrackInstance.h b/Engine/Source/Runtime/MovieSceneCoreTypes/Private/Instances/MovieSceneAnimationTrackInstance.h
--- a/Engine/Source/Runtime/MovieSceneCoreTypes/Private/Instances/MovieSceneAnimationTrackInstance.h
+++ b/Engine/Source/Runtime/MovieSceneCoreTypes/Private/Instances/MovieSceneAnimationTrackInstance.h
@@ -13,7 +13,7 @@ public:
 
 	/** IMovieSceneTrackInstance interface */
 	virtual void Update( float Position, float LastPosition, const TArray<UObject*>& RuntimeObjects, class IMovieScenePlayer& Player ) override;
-	virtual void RefreshInstance( IMovieScenePlayer& Player ) override {}
+	virtual void RefreshInstance( const TArray<UObject*>& RuntimeObjects, class IMovieScenePlayer& Player ) override {}
 
 private:
 	/** Track that is being instanced */
diff --git a/Engine/Source/Runtime/MovieSceneCoreTypes/Private/Instances/MovieSceneAudioTrackInstance.h b/Engine/Source/Runtime/MovieSceneCoreTypes/Private/Instances/MovieSceneAudioTrackInstance.h
--- a/Engine/Source/Runtime/MovieSceneCoreTypes/Private/Instances/MovieSceneAudioTrackInstance.h
+++ b/Engine/Source/Runtime/MovieSceneCoreTypes/Private/Instances/MovieSceneAudioTrackInstance.h
@@ -12,7 +12,7 @@ public:
 
 	/** IMovieSceneTrackInstance interface */
 	virtual void Update( float Position, float LastPosition, const TArray<UObject*>& RuntimeObjects, class IMovieScenePlayer& Player ) override;
-	virtual void RefreshInstance( class IMovieScenePlayer& Player ) override {}
+	virtual void RefreshInstance( const TArray<UObject*>& RuntimeObjects, class IMovieScenePlayer& Player ) override {}
 private:
 	/** Plays the sound of the given section at the given time */
 	void PlaySound(class UMovieSceneAudioSection* AudioSection, class UAudioComponent* Component, float Time);
diff --git a/Engine/Source/Runtime/MovieSceneCoreTypes/Private/Instances/MovieSceneBoolTrackInstance.h b/Engine/Source/Runtime/MovieSceneCoreTypes/Private/Instances/MovieSceneBoolTrackInstance.h
--- a/Engine/Source/Runtime/MovieSceneCoreTypes/Private/Instances/MovieSceneBoolTrackInstance.h
+++ b/Engine/Source/Runtime/MovieSceneCoreTypes/Private/Instances/MovieSceneBoolTrackInstance.h
@@ -12,7 +12,7 @@ public:
 
 	/** IMovieSceneTrackInstance interface */
 	virtual void Update( float Position, float LastPosition, const TArray<UObject*>& RuntimeObjects, class IMovieScenePlayer& Player ) override;
-	virtual void RefreshInstance( class IMovieScenePlayer& Player ) override {}
+	virtual void RefreshInstance( const TArray<UObject*>& RuntimeObjects, class IMovieScenePlayer& Player ) override {}
 private:
 	/** Bool track that is being instanced */
 	UMovieSceneBoolTrack* BoolTrack;
diff --git a/Engine/Source/Runtime/MovieSceneCoreTypes/Private/Instances/MovieSceneColorTrackInstance.cpp b/Engine/Source/Runtime/MovieSceneCoreTypes/Private/Instances/MovieSceneColorTrackInstance.cpp
--- a/Engine/Source/Runtime/MovieSceneCoreTypes/Private/Instances/MovieSceneColorTrackInstance.cpp
+++ b/Engine/Source/Runtime/MovieSceneCoreTypes/Private/Instances/MovieSceneColorTrackInstance.cpp
@@ -19,35 +19,68 @@ void FMovieSceneColorTrackInstance::Update( float Position, float LastPosition,
 		{
 			UObject* Object = RuntimeObjects[ObjIndex];
 
-			UObject* PropertyOwner = NULL;
-			UProperty* Property = NULL;
-			//@todo Sequencer - Major performance problems here.  This needs to be initialized and stored (not serialized) somewhere
-			void* Address = FMatineeUtils::GetPropertyAddress<void>( Object, ColorTrack->GetPropertyName(), Property, PropertyOwner );
-
-			if (Address)
+			UFunction* Function = RuntimeObjectToFunctionMap.FindRef( Object );
+			if( Function )
 			{
-				bool bIsFColor = false;
-				const UStructProperty* StructProp = Cast<const UStructProperty>(Property);
-				check(StructProp && StructProp->Struct);
-				FName StructName = StructProp->Struct->GetFName();
-
-				bIsFColor = StructName == NAME_Color;
-				check(bIsFColor || StructName == NAME_LinearColor);
+				Object->ProcessEvent( Function, &ColorValue );
+			}
+			else
+			{
+				// @todo Sequencer - remove this legacy method of setting the property directly
+				UObject* PropertyOwner = NULL;
+				UProperty* Property = NULL;
+				//@todo Sequencer - Major performance problems here.  This needs to be initialized and stored (not serialized) somewhere
+				void* Address = FMatineeUtils::GetPropertyAddress<void>( Object, ColorTrack->GetPropertyName(), Property, PropertyOwner );
 
-				if( bIsFColor )
+				if (Address)
 				{
-					*(FColor*)Address = ColorValue.ToFColor(false);
-					// Let the property owner know that we changed one of its properties
-					PropertyOwner->PostInterpChange( Property );
-				}
-				else
-				{
-					*(FLinearColor*)Address = ColorValue;
-					// Let the property owner know that we changed one of its properties
-					PropertyOwner->PostInterpChange( Property );
+					bool bIsFColor = false;
+					const UStructProperty* StructProp = Cast<const UStructProperty>(Property);
+					check(StructProp && StructProp->Struct);
+					FName StructName = StructProp->Struct->GetFName();
+
+					bIsFColor = StructName == NAME_Color;
+					check(bIsFColor || StructName == NAME_LinearColor);
+
+					if( bIsFColor )
+					{
+						*(FColor*)Address = ColorValue.ToFColor(false);
+						// Let the property owner know that we changed one of its properties
+						PropertyOwner->PostInterpChange( Property );
+					}
+					else
+					{
+						*(FLinearColor*)Address = ColorValue;
+						// Let the property owner know that we changed one of its properties
+						PropertyOwner->PostInterpChange( Property );
+					}
 				}
 			}
 		}
 	}
 }
 
+void FMovieSceneColorTrackInstance::RefreshInstance( const TArray<UObject*>& RuntimeObjects, class IMovieScenePlayer& Player )
+{
+	static const FString Set(TEXT("Set") );
+	
+	FString FunctionString = Set + ColorTrack->GetPropertyName().ToString();
+	
+	FName FunctionName = FName( *FunctionString );
+	
+	//Player.GetRuntimeObjects( )
+	for( UObject* Object : RuntimeObjects )
+	{
+		// If the mapping for the object already exists we dont need to search for it again
+		UFunction* Function = RuntimeObjectToFunctionMap.FindRef( Object );
+		
+		if( !Function )
+		{
+			Function = Object->FindFunction( FunctionName );
+			if( Function )
+			{
+				RuntimeObjectToFunctionMap.Add( Object, Function );
+			}
+		}
+	}
+}
diff --git a/Engine/Source/Runtime/MovieSceneCoreTypes/Private/Instances/MovieSceneColorTrackInstance.h b/Engine/Source/Runtime/MovieSceneCoreTypes/Private/Instances/MovieSceneColorTrackInstance.h
--- a/Engine/Source/Runtime/MovieSceneCoreTypes/Private/Instances/MovieSceneColorTrackInstance.h
+++ b/Engine/Source/Runtime/MovieSceneCoreTypes/Private/Instances/MovieSceneColorTrackInstance.h
@@ -12,8 +12,10 @@ public:
 
 	/** IMovieSceneTrackInstance interface */
 	virtual void Update( float Position, float LastPosition, const TArray<UObject*>& RuntimeObjects, class IMovieScenePlayer& Player ) override;
-	virtual void RefreshInstance( class IMovieScenePlayer& Player ) override {}
+	virtual void RefreshInstance( const TArray<UObject*>& RuntimeObjects, class IMovieScenePlayer& Player ) override;
 private:
 	/** The track being instanced */
 	UMovieSceneColorTrack* ColorTrack;
+	/** Mapping of objects to bound functions that will be called to update data on the track */
+	TMap< TWeakObjectPtr<UObject>, UFunction* > RuntimeObjectToFunctionMap;
 };
diff --git a/Engine/Source/Runtime/MovieSceneCoreTypes/Private/Instances/MovieSceneDirectorTrackInstance.h b/Engine/Source/Runtime/MovieSceneCoreTypes/Private/Instances/MovieSceneDirectorTrackInstance.h
--- a/Engine/Source/Runtime/MovieSceneCoreTypes/Private/Instances/MovieSceneDirectorTrackInstance.h
+++ b/Engine/Source/Runtime/MovieSceneCoreTypes/Private/Instances/MovieSceneDirectorTrackInstance.h
@@ -11,7 +11,7 @@ public:
 	FMovieSceneDirectorTrackInstance( UMovieSceneDirectorTrack& InDirectorTrack );
 	/** IMovieSceneTrackInstance interface */
 	virtual void Update( float Position, float LastPosition, const TArray<UObject*>& RuntimeObjects, class IMovieScenePlayer& Player ) override;
-	virtual void RefreshInstance( class IMovieScenePlayer& Player ) override {}
+	virtual void RefreshInstance( const TArray<UObject*>& RuntimeObjects, class IMovieScenePlayer& Player  ) override {}
 private:
 	/** Track that is being instanced */
 	UMovieSceneDirectorTrack* DirectorTrack;
diff --git a/Engine/Source/Runtime/MovieSceneCoreTypes/Private/Instances/MovieSceneFloatTrackInstance.h b/Engine/Source/Runtime/MovieSceneCoreTypes/Private/Instances/MovieSceneFloatTrackInstance.h
--- a/Engine/Source/Runtime/MovieSceneCoreTypes/Private/Instances/MovieSceneFloatTrackInstance.h
+++ b/Engine/Source/Runtime/MovieSceneCoreTypes/Private/Instances/MovieSceneFloatTrackInstance.h
@@ -12,7 +12,7 @@ public:
 
 	/** IMovieSceneTrackInstance interface */
 	virtual void Update( float Position, float LastPosition, const TArray<UObject*>& RuntimeObjects, class IMovieScenePlayer& Player ) override;
-	virtual void RefreshInstance( IMovieScenePlayer& Player ) override {}
+	virtual void RefreshInstance( const TArray<UObject*>& RuntimeObjects, class IMovieScenePlayer& Player ) override {}
 private:
 	/** Track that is being instanced */
 	UMovieSceneFloatTrack* FloatTrack;
diff --git a/Engine/Source/Runtime/MovieSceneCoreTypes/Private/Instances/MovieSceneParticleTrackInstance.h b/Engine/Source/Runtime/MovieSceneCoreTypes/Private/Instances/MovieSceneParticleTrackInstance.h
--- a/Engine/Source/Runtime/MovieSceneCoreTypes/Private/Instances/MovieSceneParticleTrackInstance.h
+++ b/Engine/Source/Runtime/MovieSceneCoreTypes/Private/Instances/MovieSceneParticleTrackInstance.h
@@ -16,7 +16,7 @@ public:
 
 	/** IMovieSceneTrackInstance interface */
 	virtual void Update( float Position, float LastPosition, const TArray<UObject*>& RuntimeObjects, class IMovieScenePlayer& Player ) override;
-	virtual void RefreshInstance( class IMovieScenePlayer& Player ) override {}
+	virtual void RefreshInstance( const TArray<UObject*>& RuntimeObjects, class IMovieScenePlayer& Player ) override {}
 
 private:
 	/** Track that is being instanced */
diff --git a/Engine/Source/Runtime/MovieSceneCoreTypes/Private/Instances/MovieSceneTransformTrackInstance.h b/Engine/Source/Runtime/MovieSceneCoreTypes/Private/Instances/MovieSceneTransformTrackInstance.h
--- a/Engine/Source/Runtime/MovieSceneCoreTypes/Private/Instances/MovieSceneTransformTrackInstance.h
+++ b/Engine/Source/Runtime/MovieSceneCoreTypes/Private/Instances/MovieSceneTransformTrackInstance.h
@@ -12,7 +12,7 @@ public:
 
 	/** IMovieSceneTrackInstance interface */
 	virtual void Update( float Position, float LastPosition, const TArray<UObject*>& RuntimeObjects, class IMovieScenePlayer& Player ) override;
-	virtual void RefreshInstance( IMovieScenePlayer& Player ) override {}
+	virtual void RefreshInstance( const TArray<UObject*>& RuntimeObjects, class IMovieScenePlayer& Player ) override {}
 private:
 	/** Track that is being instanced */
 	UMovieSceneTransformTrack* TransformTrack;
diff --git a/Engine/Source/Runtime/MovieSceneCoreTypes/Private/Instances/MovieSceneVectorTrackInstance.h b/Engine/Source/Runtime/MovieSceneCoreTypes/Private/Instances/MovieSceneVectorTrackInstance.h
--- a/Engine/Source/Runtime/MovieSceneCoreTypes/Private/Instances/MovieSceneVectorTrackInstance.h
+++ b/Engine/Source/Runtime/MovieSceneCoreTypes/Private/Instances/MovieSceneVectorTrackInstance.h
@@ -12,7 +12,7 @@ public:
 
 	/** IMovieSceneTrackInstance interface */
 	virtual void Update( float Position, float LastPosition, const TArray<UObject*>& RuntimeObjects, class IMovieScenePlayer& Player ) override;
-	virtual void RefreshInstance( IMovieScenePlayer& Player ) override {}
+	virtual void RefreshInstance( const TArray<UObject*>& RuntimeObjects, class IMovieScenePlayer& Player ) override {}
 private:
 	/** Track that is being instanced */
 	UMovieSceneVectorTrack* VectorTrack;
diff --git a/Engine/Source/Runtime/MovieSceneCoreTypes/Private/Instances/SubMovieSceneTrackInstance.cpp b/Engine/Source/Runtime/MovieSceneCoreTypes/Private/Instances/SubMovieSceneTrackInstance.cpp
--- a/Engine/Source/Runtime/MovieSceneCoreTypes/Private/Instances/SubMovieSceneTrackInstance.cpp
+++ b/Engine/Source/Runtime/MovieSceneCoreTypes/Private/Instances/SubMovieSceneTrackInstance.cpp
@@ -12,7 +12,7 @@ FSubMovieSceneTrackInstance::FSubMovieSceneTrackInstance( USubMovieSceneTrack& I
 	SubMovieSceneTrack = &InTrack;
 }
 
-void FSubMovieSceneTrackInstance::RefreshInstance( class IMovieScenePlayer& Player )
+void FSubMovieSceneTrackInstance::RefreshInstance( const TArray<UObject*>& RuntimeObjects, class IMovieScenePlayer& Player )
 {
 	TArray<UMovieSceneSection*> AllSections = SubMovieSceneTrack->GetAllSections();
 
diff --git a/Engine/Source/Runtime/MovieSceneCoreTypes/Private/Instances/SubMovieSceneTrackInstance.h b/Engine/Source/Runtime/MovieSceneCoreTypes/Private/Instances/SubMovieSceneTrackInstance.h
--- a/Engine/Source/Runtime/MovieSceneCoreTypes/Private/Instances/SubMovieSceneTrackInstance.h
+++ b/Engine/Source/Runtime/MovieSceneCoreTypes/Private/Instances/SubMovieSceneTrackInstance.h
@@ -12,7 +12,7 @@ public:
 
 	/** IMovieSceneTrackInstance interface */
 	virtual void Update( float Position, float LastPosition, const TArray<UObject*>& RuntimeObjects, class IMovieScenePlayer& Player ) override;
-	virtual void RefreshInstance( class IMovieScenePlayer& Player ) override;
+	virtual void RefreshInstance( const TArray<UObject*>& RuntimeObjects, class IMovieScenePlayer& Player ) override;
 private:
 	/** Track that is being instanced */
 	TWeakObjectPtr<class USubMovieSceneTrack> SubMovieSceneTrack;
