Added missing files

diff --git a/Engine/Source/Editor/Sequencer/Private/SequencerActorBindingManager.cpp b/Engine/Source/Editor/Sequencer/Private/SequencerActorBindingManager.cpp
--- a/Engine/Source/Editor/Sequencer/Private/SequencerActorBindingManager.cpp
+++ b/Engine/Source/Editor/Sequencer/Private/SequencerActorBindingManager.cpp
@@ -1,7 +1,7 @@
 // Copyright 1998-2014 Epic Games, Inc. All Rights Reserved.
 
 #include "SequencerPrivatePCH.h"
-#include "SequencerObjectSpawner.h"
+#include "SequencerActorBindingManager.h"
 #include "Sequencer.h"
 #include "MovieScene.h"
 #include "Toolkits/IToolkitHost.h"
diff --git a/Engine/Source/Editor/Sequencer/Private/SequencerActorBindingManager.h b/Engine/Source/Editor/Sequencer/Private/SequencerActorBindingManager.h
--- a/Engine/Source/Editor/Sequencer/Private/SequencerActorBindingManager.h
+++ b/Engine/Source/Editor/Sequencer/Private/SequencerActorBindingManager.h
@@ -56,7 +56,7 @@ public:
 /**
  * Manages spawned puppet objects for Sequencer
  */
-class FSequencerActorObjectSpawner : public ISequencerObjectBindingManager, public TSharedFromThis<FSequencerActorObjectSpawner>
+class FSequencerActorBindingManager : public ISequencerObjectBindingManager, public TSharedFromThis<FSequencerActorBindingManager>
 {
 
 public:
@@ -66,11 +66,11 @@ public:
 	 *
 	 * @param	InWorld	The world to spawn actors in
 	 */
-	FSequencerActorObjectSpawner( UWorld* InWorld, TSharedRef<FSequencer> InSequencer );
+	FSequencerActorBindingManager( UWorld* InWorld, TSharedRef<ISequencerObjectChangeListener> InObjectChangeListener, TSharedRef<FSequencer> InSequencer );
 
 
 	/** Destructor */
-	virtual ~FSequencerActorObjectSpawner();
+	virtual ~FSequencerActorBindingManager();
 
 
 	/** ISequencerObjectBindingManager interface */
@@ -84,22 +84,9 @@ public:
 	virtual void GetRuntimeObjects( const TSharedRef<FMovieSceneInstance>& MovieSceneInstance, const FGuid& ObjectGuid, TArray<UObject*>& OutRuntimeObjects ) const override;
 
 protected:
-	/**
-	 * Finds the spawnable guid for the specified object, if one exists
-	 *
-	 * @param	Object	The object to find a spawnable for
-	 *
-	 * @return	The spawnable guid for the specified object, or an Invalid guid if none was found
-	 */
+	
 	FGuid FindSpawnableGuidForPuppetObject( UObject* Object ) const;
 	
-	/**
-	 * Given a guid for a spawnable, attempts to find a spawned puppet object for that guid
-	 *
-	 * @param	MovieSceneInstance	The instance of a movie scene which the puppet object would be spawned for
-	 * @param	Guid	The guid of the spawnable to find a puppet object for
-	 * @return	The puppet object, or NULL if we don't recognize this guid
-	 */
 	UObject* FindPuppetObjectForSpawnableGuid( const TSharedRef<FMovieSceneInstance>& MovieSceneInstance, const FGuid& Guid ) const;
 	
 	/** Deselects all puppet objects in the editor */
@@ -130,7 +117,7 @@ protected:
 	 *
 	 * @return	The LevelScript node we found, or NULL for no matches
 	 */
-	class UK2Node_PlayMovieScene* FindPlayMovieSceneNodeInLevelScript( const class UMovieScene* MovieScene );
+	class UK2Node_PlayMovieScene* FindPlayMovieSceneNodeInLevelScript( const class UMovieScene* MovieScene ) const;
 	
 	/**
 	 * Creates a new "PlayMovieScene" node in the level script for the current level and associates it with the
@@ -140,7 +127,7 @@ protected:
 	 *
 	 * @return	The LevelScript node we created, or NULL on failure (does not assert)
 	 */
-	class UK2Node_PlayMovieScene* CreateNewPlayMovieSceneNode( UMovieScene* MovieScene );
+	class UK2Node_PlayMovieScene* CreateNewPlayMovieSceneNode( UMovieScene* MovieScene ) const;
 	
 	
 	/** Called by UK2Node_PlayMovieScene when the bindings for that node are changed through Kismet */
@@ -153,18 +140,18 @@ protected:
 	 *
 	 * @return	The LevelScript node we found or created, or NULL if nothing is still bound
 	 */
-	class UK2Node_PlayMovieScene* BindToPlayMovieSceneNode( const bool bCreateIfNotFound );
+	class UK2Node_PlayMovieScene* BindToPlayMovieSceneNode( const bool bCreateIfNotFound ) const;
 	
 private:
 	TWeakObjectPtr< UWorld > ActorWorld;
 	
 	/** 'PlayMovieScene' node in level script that we're currently associating with the active MovieScene.  This node contains the bindings
 	 information that we need to be able to preview possessed actors in the level */
-	TWeakObjectPtr< UK2Node_PlayMovieScene > PlayMovieSceneNode;
+	mutable TWeakObjectPtr< UK2Node_PlayMovieScene > PlayMovieSceneNode;
 
-	/** The Sequencer that owns this spawner  */
 	TWeakPtr<FSequencer> Sequencer;
 
+	TWeakPtr<ISequencerObjectChangeListener> ObjectChangeListener;
 	/**
 	 * List of puppet objects that we've spawned and are actively managing. 
 	 */
diff --git a/Engine/Source/Editor/Sequencer/Private/SequencerAssetEditor.cpp b/Engine/Source/Editor/Sequencer/Private/SequencerAssetEditor.cpp
--- a/Engine/Source/Editor/Sequencer/Private/SequencerAssetEditor.cpp
+++ b/Engine/Source/Editor/Sequencer/Private/SequencerAssetEditor.cpp
@@ -10,7 +10,7 @@
 #include "Sequencer.h"
 #include "MovieScene.h"
 #include "EditorWidgetsModule.h"
-#include "SequencerObjectSpawner.h"
+#include "SequencerActorBindingManager.h"
 #include "SequencerObjectChangeListener.h"
 
 #define LOCTEXT_NAMESPACE "Sequencer"
