UE4: Non-unity fixes

diff --git a/Engine/Source/Editor/Sequencer/Private/Sequencer.cpp b/Engine/Source/Editor/Sequencer/Private/Sequencer.cpp
--- a/Engine/Source/Editor/Sequencer/Private/Sequencer.cpp
+++ b/Engine/Source/Editor/Sequencer/Private/Sequencer.cpp
@@ -35,7 +35,7 @@
 
 #define LOCTEXT_NAMESPACE "Sequencer"
 
-DEFINE_LOG_CATEGORY_STATIC( LogSequencer, Log, All );
+DEFINE_LOG_CATEGORY(LogSequencer);
 
 bool FSequencer::IsSequencerEnabled()
 {
diff --git a/Engine/Source/Editor/Sequencer/Private/SequencerActorBindingManager.cpp b/Engine/Source/Editor/Sequencer/Private/SequencerActorBindingManager.cpp
--- a/Engine/Source/Editor/Sequencer/Private/SequencerActorBindingManager.cpp
+++ b/Engine/Source/Editor/Sequencer/Private/SequencerActorBindingManager.cpp
@@ -8,6 +8,8 @@
 #include "Kismet2/BlueprintEditorUtils.h"
 #include "K2Node_PlayMovieScene.h"
 #include "MovieSceneInstance.h"
+#include "Engine/LevelScriptBlueprint.h"
+#include "Editor/Kismet/Public/BlueprintEditorModule.h"
 
 #define LOCTEXT_NAMESPACE "Sequencer"
 
diff --git a/Engine/Source/Editor/Sequencer/Private/SequencerActorBindingManager.h b/Engine/Source/Editor/Sequencer/Private/SequencerActorBindingManager.h
--- a/Engine/Source/Editor/Sequencer/Private/SequencerActorBindingManager.h
+++ b/Engine/Source/Editor/Sequencer/Private/SequencerActorBindingManager.h
@@ -2,6 +2,8 @@
 
 #pragma once
 
+#include "Editor/Sequencer/Public/ISequencerObjectBindingManager.h"
+
 class FSequencer;
 
 namespace EPuppetObjectType
diff --git a/Engine/Source/Editor/Sequencer/Private/SequencerModule.cpp b/Engine/Source/Editor/Sequencer/Private/SequencerModule.cpp
--- a/Engine/Source/Editor/Sequencer/Private/SequencerModule.cpp
+++ b/Engine/Source/Editor/Sequencer/Private/SequencerModule.cpp
@@ -6,6 +6,7 @@
 #include "Toolkits/ToolkitManager.h"
 #include "SequencerCommands.h"
 #include "SequencerAssetEditor.h"
+#include "SequencerObjectChangeListener.h"
 
 /**
  * SequencerModule implementation (private)
diff --git a/Engine/Source/Editor/Sequencer/Private/SequencerPrivatePCH.h b/Engine/Source/Editor/Sequencer/Private/SequencerPrivatePCH.h
--- a/Engine/Source/Editor/Sequencer/Private/SequencerPrivatePCH.h
+++ b/Engine/Source/Editor/Sequencer/Private/SequencerPrivatePCH.h
@@ -21,4 +21,6 @@
 #include "SAnimationOutlinerView.h"
 #include "MovieSceneInstance.h"
 
+DECLARE_LOG_CATEGORY_EXTERN(LogSequencer, Log, All);
+
 #endif // __SequencerPrivatePCH_h__
diff --git a/Engine/Source/Editor/Sequencer/Public/ISequencerObjectBindingManager.h b/Engine/Source/Editor/Sequencer/Public/ISequencerObjectBindingManager.h
--- a/Engine/Source/Editor/Sequencer/Public/ISequencerObjectBindingManager.h
+++ b/Engine/Source/Editor/Sequencer/Public/ISequencerObjectBindingManager.h
@@ -2,6 +2,7 @@
 
 #pragma once
 
+class FMovieSceneInstance;
 
 class ISequencerObjectBindingManager
 {
diff --git a/Engine/Source/Editor/UMGEditor/Private/WidgetBlueprintEditor.cpp b/Engine/Source/Editor/UMGEditor/Private/WidgetBlueprintEditor.cpp
--- a/Engine/Source/Editor/UMGEditor/Private/WidgetBlueprintEditor.cpp
+++ b/Engine/Source/Editor/UMGEditor/Private/WidgetBlueprintEditor.cpp
@@ -6,6 +6,7 @@
 #include "WidgetBlueprintEditor.h"
 #include "MovieScene.h"
 #include "Editor/Sequencer/Public/ISequencerModule.h"
+#include "UMGSequencerObjectBindingManager.h"
 
 #define LOCTEXT_NAMESPACE "UMG"
 
