Moved a few Sequencer files around and refactored a bit to properly support new track types added from other modules

diff --git a/Engine/Source/Developer/MovieSceneTools/Private/AnimationTrackEditor.cpp b/Engine/Source/Developer/MovieSceneTools/Private/AnimationTrackEditor.cpp
--- a/Engine/Source/Developer/MovieSceneTools/Private/AnimationTrackEditor.cpp
+++ b/Engine/Source/Developer/MovieSceneTools/Private/AnimationTrackEditor.cpp
@@ -44,14 +44,14 @@ UMovieSceneSection* FAnimationSection::GetSectionObject()
 	return &Section;
 }
 
-FString FAnimationSection::GetDisplayName() const
+FText FAnimationSection::GetDisplayName() const
 {
-	return TEXT("Animation");
+	return NSLOCTEXT("FAnimationSection", "AnimationSection", "Animation");
 }
 
-FString FAnimationSection::GetSectionTitle() const
+FText FAnimationSection::GetSectionTitle() const
 {
-	return Cast<UMovieSceneAnimationSection>(&Section)->GetAnimSequence()->GetName();
+	return FText::FromString( Cast<UMovieSceneAnimationSection>(&Section)->GetAnimSequence()->GetName() );
 }
 
 float FAnimationSection::GetSectionHeight() const
diff --git a/Engine/Source/Developer/MovieSceneTools/Private/AnimationTrackEditor.h b/Engine/Source/Developer/MovieSceneTools/Private/AnimationTrackEditor.h
--- a/Engine/Source/Developer/MovieSceneTools/Private/AnimationTrackEditor.h
+++ b/Engine/Source/Developer/MovieSceneTools/Private/AnimationTrackEditor.h
@@ -51,8 +51,8 @@ public:
 
 	/** ISequencerSection interface */
 	virtual UMovieSceneSection* GetSectionObject() override;
-	virtual FString GetDisplayName() const override;
-	virtual FString GetSectionTitle() const override;
+	virtual FText GetDisplayName() const override;
+	virtual FText GetSectionTitle() const override;
 	virtual float GetSectionHeight() const override;
 	virtual void GenerateSectionLayout( class ISectionLayoutBuilder& LayoutBuilder ) const override {}
 	virtual int32 OnPaintSection( const FGeometry& AllottedGeometry, const FSlateRect& SectionClippingRect, FSlateWindowElementList& OutDrawElements, int32 LayerId, bool bParentEnabled ) const override;
diff --git a/Engine/Source/Developer/MovieSceneTools/Private/AudioTrackEditor.cpp b/Engine/Source/Developer/MovieSceneTools/Private/AudioTrackEditor.cpp
--- a/Engine/Source/Developer/MovieSceneTools/Private/AudioTrackEditor.cpp
+++ b/Engine/Source/Developer/MovieSceneTools/Private/AudioTrackEditor.cpp
@@ -329,14 +329,14 @@ UMovieSceneSection* FAudioSection::GetSectionObject()
 	return &Section;
 }
 
-FString FAudioSection::GetDisplayName() const
+FText FAudioSection::GetDisplayName() const
 {
-	return bIsOnAMasterTrack ? TEXT("Master Audio") : TEXT("Audio");
+	return bIsOnAMasterTrack ? NSLOCTEXT("FAudioSection", "MasterAudioDisplayName", "Master Audio") :  NSLOCTEXT("FAudioSection", "AudioDisplayName", "Audio");
 }
 
-FString FAudioSection::GetSectionTitle() const
+FText FAudioSection::GetSectionTitle() const
 {
-	return Cast<UMovieSceneAudioSection>(&Section)->GetSound()->GetName();
+	return FText::FromString( Cast<UMovieSceneAudioSection>(&Section)->GetSound()->GetName() );
 }
 
 float FAudioSection::GetSectionHeight() const
diff --git a/Engine/Source/Developer/MovieSceneTools/Private/AudioTrackEditor.h b/Engine/Source/Developer/MovieSceneTools/Private/AudioTrackEditor.h
--- a/Engine/Source/Developer/MovieSceneTools/Private/AudioTrackEditor.h
+++ b/Engine/Source/Developer/MovieSceneTools/Private/AudioTrackEditor.h
@@ -48,8 +48,8 @@ public:
 
 	/** ISequencerSection interface */
 	virtual UMovieSceneSection* GetSectionObject() override;
-	virtual FString GetDisplayName() const override;
-	virtual FString GetSectionTitle() const override;
+	virtual FText GetDisplayName() const override;
+	virtual FText GetSectionTitle() const override;
 	virtual float GetSectionHeight() const override;
 	virtual void GenerateSectionLayout( class ISectionLayoutBuilder& LayoutBuilder ) const override {}
 	virtual int32 OnPaintSection( const FGeometry& AllottedGeometry, const FSlateRect& SectionClippingRect, FSlateWindowElementList& OutDrawElements, int32 LayerId, bool bParentEnabled ) const override;
diff --git a/Engine/Source/Developer/MovieSceneTools/Private/DirectorTrackEditor.cpp b/Engine/Source/Developer/MovieSceneTools/Private/DirectorTrackEditor.cpp
--- a/Engine/Source/Developer/MovieSceneTools/Private/DirectorTrackEditor.cpp
+++ b/Engine/Source/Developer/MovieSceneTools/Private/DirectorTrackEditor.cpp
@@ -198,9 +198,9 @@ UMovieSceneSection* FShotSection::GetSectionObject()
 	return Section;
 }
 
-FString FShotSection::GetSectionTitle() const
+FText FShotSection::GetSectionTitle() const
 {
-	return Cast<UMovieSceneShotSection>(Section)->GetTitle().ToString();
+	return Cast<UMovieSceneShotSection>(Section)->GetTitle();
 }
 
 float FShotSection::GetSectionHeight() const
diff --git a/Engine/Source/Developer/MovieSceneTools/Private/DirectorTrackEditor.h b/Engine/Source/Developer/MovieSceneTools/Private/DirectorTrackEditor.h
--- a/Engine/Source/Developer/MovieSceneTools/Private/DirectorTrackEditor.h
+++ b/Engine/Source/Developer/MovieSceneTools/Private/DirectorTrackEditor.h
@@ -132,8 +132,8 @@ public:
 	virtual UMovieSceneSection* GetSectionObject() override;
 	virtual int32 OnPaintSection( const FGeometry& AllottedGeometry, const FSlateRect& SectionClippingRect, FSlateWindowElementList& OutDrawElements, int32 LayerId, bool bParentEnabled ) const override;
 	virtual void Tick( const FGeometry& AllottedGeometry, const FGeometry& ParentGeometry, const double InCurrentTime, const float InDeltaTime ) override;
-	virtual FString GetDisplayName() const override { return TEXT("Director Track"); }
-	virtual FString GetSectionTitle() const override;
+	virtual FText GetDisplayName() const override { return NSLOCTEXT("FShotSection", "DirectorTrack", "Director Track"); }
+	virtual FText GetSectionTitle() const override;
 	virtual float GetSectionHeight() const override;
 	virtual void GenerateSectionLayout( class ISectionLayoutBuilder& LayoutBuilder ) const override {}
 	virtual FReply OnSectionDoubleClicked( const FGeometry& SectionGeometry, const FPointerEvent& MouseEvent ) override;
diff --git a/Engine/Source/Developer/MovieSceneTools/Private/K2Node_PlayMovieScene.cpp b/Engine/Source/Developer/MovieSceneTools/Private/K2Node_PlayMovieScene.cpp
--- a/Engine/Source/Developer/MovieSceneTools/Private/K2Node_PlayMovieScene.cpp
+++ b/Engine/Source/Developer/MovieSceneTools/Private/K2Node_PlayMovieScene.cpp
@@ -284,7 +284,7 @@ void UK2Node_PlayMovieScene::CreatePinForBoundObject( FMovieSceneBoundObject& Bo
 				if( Possessable.GetGuid() == BoundObject.GetPossessableGuid() )
 				{
 					// Found a name for this possessable
-					PinFriendlyName = FText::FromString(Possessable.GetName());
+					PinFriendlyName = Possessable.GetDisplayName();
 					break;
 				}
 			}
diff --git a/Engine/Source/Developer/MovieSceneTools/Private/MovieSceneToolHelpers.h b/Engine/Source/Developer/MovieSceneTools/Private/MovieSceneToolHelpers.h
deleted file mode 100644
diff --git a/Engine/Source/Developer/MovieSceneTools/Private/ParticleTrackEditor.cpp b/Engine/Source/Developer/MovieSceneTools/Private/ParticleTrackEditor.cpp
--- a/Engine/Source/Developer/MovieSceneTools/Private/ParticleTrackEditor.cpp
+++ b/Engine/Source/Developer/MovieSceneTools/Private/ParticleTrackEditor.cpp
@@ -43,16 +43,16 @@ UMovieSceneSection* FParticleSection::GetSectionObject()
 	return &Section;
 }
 
-FString FParticleSection::GetDisplayName() const
+FText FParticleSection::GetDisplayName() const
 {
-	return TEXT("Emitter");
+	return NSLOCTEXT("FParticleSection", "Emitter", "Emitter");
 }
 
-FString FParticleSection::GetSectionTitle() const
+FText FParticleSection::GetSectionTitle() const
 {
 	EParticleKey::Type KeyType = Cast<UMovieSceneParticleSection>(&Section)->GetKeyType();
-	return KeyType == EParticleKey::Toggle ? TEXT("Toggle") :
-		KeyType == EParticleKey::Trigger ? TEXT("Trigger") : TEXT("None");
+	return KeyType == EParticleKey::Toggle ? NSLOCTEXT("FParticleSection", "Toggle", "Toggle") :
+		KeyType == EParticleKey::Trigger ? NSLOCTEXT("FParticleSection", "Trigger", "Trigger") : NSLOCTEXT("FParticleSection", "None", "None");
 }
 
 float FParticleSection::GetSectionHeight() const
diff --git a/Engine/Source/Developer/MovieSceneTools/Private/ParticleTrackEditor.h b/Engine/Source/Developer/MovieSceneTools/Private/ParticleTrackEditor.h
--- a/Engine/Source/Developer/MovieSceneTools/Private/ParticleTrackEditor.h
+++ b/Engine/Source/Developer/MovieSceneTools/Private/ParticleTrackEditor.h
@@ -48,8 +48,8 @@ public:
 
 	/** ISequencerSection interface */
 	virtual UMovieSceneSection* GetSectionObject() override;
-	virtual FString GetDisplayName() const override;
-	virtual FString GetSectionTitle() const override;
+	virtual FText GetDisplayName() const override;
+	virtual FText GetSectionTitle() const override;
 	virtual float GetSectionHeight() const override;
 	virtual void GenerateSectionLayout( class ISectionLayoutBuilder& LayoutBuilder ) const override {}
 	virtual int32 OnPaintSection( const FGeometry& AllottedGeometry, const FSlateRect& SectionClippingRect, FSlateWindowElementList& OutDrawElements, int32 LayerId, bool bParentEnabled ) const override;
diff --git a/Engine/Source/Developer/MovieSceneTools/Private/PropertyTrackEditor.cpp b/Engine/Source/Developer/MovieSceneTools/Private/PropertyTrackEditor.cpp
--- a/Engine/Source/Developer/MovieSceneTools/Private/PropertyTrackEditor.cpp
+++ b/Engine/Source/Developer/MovieSceneTools/Private/PropertyTrackEditor.cpp
@@ -21,51 +21,7 @@
 #include "MovieSceneToolHelpers.h"
 #include "PropertyTrackEditor.h"
 #include "CommonMovieSceneTools.h"
-
-/**
- * A generic implementation for displaying simple property sections
- */
-class FPropertySection : public ISequencerSection
-{
-public:
-	FPropertySection( UMovieSceneSection& InSectionObject, FName SectionName )
-		: SectionObject( InSectionObject )
-	{
-		DisplayName = FName::NameToDisplayString( SectionName.ToString(), false );
-	}
-
-	/** ISequencerSection interface */
-	virtual UMovieSceneSection* GetSectionObject() override { return &SectionObject; }
-
-	virtual FString GetDisplayName() const override
-	{
-		return DisplayName;
-	}
-	
-	virtual FString GetSectionTitle() const override { return FString(); }
-
-	virtual void GenerateSectionLayout( class ISectionLayoutBuilder& LayoutBuilder ) const override {}
-
-	virtual int32 OnPaintSection( const FGeometry& AllottedGeometry, const FSlateRect& SectionClippingRect, FSlateWindowElementList& OutDrawElements, int32 LayerId, bool bParentEnabled ) const override 
-	{
-		// Add a box for the section
-		FSlateDrawElement::MakeBox( 
-			OutDrawElements,
-			LayerId,
-			AllottedGeometry.ToPaintGeometry(),
-			FEditorStyle::GetBrush("Sequencer.GenericSection.Background"),
-			SectionClippingRect
-		);
-
-		return LayerId;
-	}
-	
-protected:
-	/** Display name of the section */
-	FString DisplayName;
-	/** The section we are visualizing */
-	UMovieSceneSection& SectionObject;
-};
+#include "PropertySection.h"
 
 /**
  * A color section implementation
@@ -80,10 +36,10 @@ public:
 	{
 		UMovieSceneColorSection* ColorSection = Cast<UMovieSceneColorSection>(&SectionObject);
 
-		LayoutBuilder.AddKeyArea("R", TEXT("Red"), MakeShareable(new FFloatCurveKeyArea(ColorSection->GetRedCurve())));
-		LayoutBuilder.AddKeyArea("G", TEXT("Green"), MakeShareable(new FFloatCurveKeyArea(ColorSection->GetGreenCurve())));
-		LayoutBuilder.AddKeyArea("B", TEXT("Blue"), MakeShareable(new FFloatCurveKeyArea(ColorSection->GetBlueCurve())));
-		LayoutBuilder.AddKeyArea("A", TEXT("Alpha"), MakeShareable(new FFloatCurveKeyArea(ColorSection->GetAlphaCurve())));
+		LayoutBuilder.AddKeyArea("R", NSLOCTEXT("FColorPropertySection", "RedArea", "Red"), MakeShareable(new FFloatCurveKeyArea(ColorSection->GetRedCurve())));
+		LayoutBuilder.AddKeyArea("G", NSLOCTEXT("FColorPropertySection", "GreenArea", "Green"), MakeShareable(new FFloatCurveKeyArea(ColorSection->GetGreenCurve())));
+		LayoutBuilder.AddKeyArea("B", NSLOCTEXT("FColorPropertySection", "BlueArea", "Blue"), MakeShareable(new FFloatCurveKeyArea(ColorSection->GetBlueCurve())));
+		LayoutBuilder.AddKeyArea("A", NSLOCTEXT("FColorPropertySection", "OpacityArea", "Opacity"), MakeShareable(new FFloatCurveKeyArea(ColorSection->GetAlphaCurve())));
 	}
 
 	virtual int32 OnPaintSection( const FGeometry& AllottedGeometry, const FSlateRect& SectionClippingRect, FSlateWindowElementList& OutDrawElements, int32 LayerId, bool bParentEnabled ) const override 
@@ -220,12 +176,16 @@ public:
 		int32 ChannelsUsed = VectorSection->GetChannelsUsed();
 		check(ChannelsUsed >= 2 && ChannelsUsed <= 4);
 
-		LayoutBuilder.AddKeyArea("Vector.X", TEXT("X"), MakeShareable(new FFloatCurveKeyArea(VectorSection->GetCurve(0))));
-		LayoutBuilder.AddKeyArea("Vector.Y", TEXT("Y"), MakeShareable(new FFloatCurveKeyArea(VectorSection->GetCurve(1))));
+		LayoutBuilder.AddKeyArea("Vector.X", NSLOCTEXT("FVectorPropertySection", "XArea", "X"), MakeShareable(new FFloatCurveKeyArea(VectorSection->GetCurve(0))));
+		LayoutBuilder.AddKeyArea("Vector.Y", NSLOCTEXT("FVectorPropertySection", "YArea", "Y"), MakeShareable(new FFloatCurveKeyArea(VectorSection->GetCurve(1))));
 		if (ChannelsUsed >= 3)
-		{LayoutBuilder.AddKeyArea("Vector.Z", TEXT("Z"), MakeShareable(new FFloatCurveKeyArea(VectorSection->GetCurve(2))));}
+		{
+			LayoutBuilder.AddKeyArea("Vector.Z", NSLOCTEXT("FVectorPropertySection", "ZArea", "Z"), MakeShareable(new FFloatCurveKeyArea(VectorSection->GetCurve(2))));
+		}
 		if (ChannelsUsed >= 4)
-		{LayoutBuilder.AddKeyArea("Vector.W", TEXT("W"), MakeShareable(new FFloatCurveKeyArea(VectorSection->GetCurve(3))));}
+		{
+			LayoutBuilder.AddKeyArea("Vector.W", NSLOCTEXT("FVectorPropertySection", "WArea", "W"), MakeShareable(new FFloatCurveKeyArea(VectorSection->GetCurve(3))));
+		}
 	}
 };
 
diff --git a/Engine/Source/Developer/MovieSceneTools/Private/SubMovieSceneTrackEditor.cpp b/Engine/Source/Developer/MovieSceneTools/Private/SubMovieSceneTrackEditor.cpp
--- a/Engine/Source/Developer/MovieSceneTools/Private/SubMovieSceneTrackEditor.cpp
+++ b/Engine/Source/Developer/MovieSceneTools/Private/SubMovieSceneTrackEditor.cpp
@@ -14,7 +14,7 @@ public:
 	FSubMovieSceneSection( TSharedPtr<ISequencer> InSequencer, UMovieSceneSection& InSectionObject, FName SectionName )
 		: Sequencer( InSequencer )
 		, SectionObject( *CastChecked<USubMovieSceneSection>( &InSectionObject ) )
-		, DisplayName( NSLOCTEXT("SubMovieSceneSection", "DisplayName", "Scenes").ToString() )
+		, DisplayName( NSLOCTEXT("SubMovieSceneSection", "DisplayName", "Scenes") )
 	{
 		MovieSceneInstance = InSequencer->GetInstanceForSubMovieSceneSection( InSectionObject );
 	}
@@ -22,14 +22,14 @@ public:
 	/** ISequencerSection interface */
 	virtual UMovieSceneSection* GetSectionObject() override { return &SectionObject; }
 
-	virtual FString GetDisplayName() const override
+	virtual FText GetDisplayName() const override
 	{
 		return DisplayName;
 	}
 
-	virtual FString GetSectionTitle() const override
+	virtual FText GetSectionTitle() const override
 	{
-		return SectionObject.GetMovieScene()->GetName();
+		return FText::FromString( SectionObject.GetMovieScene()->GetName() );
 	}
 
 	virtual void GenerateSectionLayout( class ISectionLayoutBuilder& LayoutBuilder ) const override
@@ -65,7 +65,7 @@ public:
 	
 private:
 	/** Display name of the section */
-	FString DisplayName;
+	FText DisplayName;
 	/** The section we are visualizing */
 	USubMovieSceneSection& SectionObject;
 	/** The instance that this section is part of */
diff --git a/Engine/Source/Developer/MovieSceneTools/Private/TransformTrackEditor.cpp b/Engine/Source/Developer/MovieSceneTools/Private/TransformTrackEditor.cpp
--- a/Engine/Source/Developer/MovieSceneTools/Private/TransformTrackEditor.cpp
+++ b/Engine/Source/Developer/MovieSceneTools/Private/TransformTrackEditor.cpp
@@ -31,34 +31,34 @@ public:
 		return &Section;
 	}
 
-	virtual FString GetDisplayName() const override
+	virtual FText GetDisplayName() const override
 	{ 
-		return TEXT("Transform");
+		return NSLOCTEXT("FTransformSection", "DisplayName", "Transform");
 	}
 	
-	virtual FString GetSectionTitle() const override { return FString(); }
+	virtual FText GetSectionTitle() const override { return FText::GetEmpty(); }
 
 	virtual void GenerateSectionLayout( class ISectionLayoutBuilder& LayoutBuilder ) const override
 	{
 		UMovieSceneTransformSection* TransformSection = Cast<UMovieSceneTransformSection>( &Section );
 
 		// This generates the tree structure for the transform section
-		LayoutBuilder.PushCategory( "Location", TEXT("Location") );
-			LayoutBuilder.AddKeyArea("Location.X", TEXT("X"), MakeShareable( new FFloatCurveKeyArea( TransformSection->GetTranslationCurve( EAxis::X ) ) ) );
-			LayoutBuilder.AddKeyArea("Location.Y", TEXT("Y"), MakeShareable( new FFloatCurveKeyArea( TransformSection->GetTranslationCurve( EAxis::Y ) ) ) );
-			LayoutBuilder.AddKeyArea("Location.Z", TEXT("Z"), MakeShareable( new FFloatCurveKeyArea( TransformSection->GetTranslationCurve( EAxis::Z ) ) ) );
+		LayoutBuilder.PushCategory( "Location", NSLOCTEXT("FTransformSection", "LocationArea", "Location") );
+			LayoutBuilder.AddKeyArea("Location.X", NSLOCTEXT("FTransformSection", "LocXArea", "X"), MakeShareable( new FFloatCurveKeyArea( TransformSection->GetTranslationCurve( EAxis::X ) ) ) );
+			LayoutBuilder.AddKeyArea("Location.Y", NSLOCTEXT("FTransformSection", "LocYArea", "Y"), MakeShareable( new FFloatCurveKeyArea( TransformSection->GetTranslationCurve( EAxis::Y ) ) ) );
+			LayoutBuilder.AddKeyArea("Location.Z", NSLOCTEXT("FTransformSection", "LocZArea", "Z"), MakeShareable( new FFloatCurveKeyArea( TransformSection->GetTranslationCurve( EAxis::Z ) ) ) );
 		LayoutBuilder.PopCategory();
 
-		LayoutBuilder.PushCategory( "Rotation", TEXT("Rotation") );
-			LayoutBuilder.AddKeyArea("Rotation.X", TEXT("X"), MakeShareable( new FFloatCurveKeyArea( TransformSection->GetRotationCurve( EAxis::X ) ) ) );
-			LayoutBuilder.AddKeyArea("Rotation.Y", TEXT("Y"), MakeShareable( new FFloatCurveKeyArea( TransformSection->GetRotationCurve( EAxis::Y ) ) ) );
-			LayoutBuilder.AddKeyArea("Rotation.Z", TEXT("Z"), MakeShareable( new FFloatCurveKeyArea( TransformSection->GetRotationCurve( EAxis::Z ) ) ) );
+		LayoutBuilder.PushCategory( "Rotation", NSLOCTEXT("FTransformSection", "RotationArea", "Rotation") );
+			LayoutBuilder.AddKeyArea("Rotation.X", NSLOCTEXT("FTransformSection", "RotXArea", "X"), MakeShareable( new FFloatCurveKeyArea( TransformSection->GetRotationCurve( EAxis::X ) ) ) );
+			LayoutBuilder.AddKeyArea("Rotation.Y", NSLOCTEXT("FTransformSection", "RotYArea", "Y"), MakeShareable( new FFloatCurveKeyArea( TransformSection->GetRotationCurve( EAxis::Y ) ) ) );
+			LayoutBuilder.AddKeyArea("Rotation.Z", NSLOCTEXT("FTransformSection", "RotZArea", "Z"), MakeShareable( new FFloatCurveKeyArea( TransformSection->GetRotationCurve( EAxis::Z ) ) ) );
 		LayoutBuilder.PopCategory();
 
-		LayoutBuilder.PushCategory( "Scale", TEXT("Scale") );
-			LayoutBuilder.AddKeyArea("Scale.X", TEXT("X"), MakeShareable( new FFloatCurveKeyArea( TransformSection->GetScaleCurve( EAxis::X ) ) ) );
-			LayoutBuilder.AddKeyArea("Scale.Y", TEXT("Y"), MakeShareable( new FFloatCurveKeyArea( TransformSection->GetScaleCurve( EAxis::Y ) ) ) );
-			LayoutBuilder.AddKeyArea("Scale.Z", TEXT("Z"), MakeShareable( new FFloatCurveKeyArea( TransformSection->GetScaleCurve( EAxis::Z ) ) ) );
+		LayoutBuilder.PushCategory( "Scale", NSLOCTEXT("FTransformSection", "ScaleArea", "Scale") );
+			LayoutBuilder.AddKeyArea("Scale.X", NSLOCTEXT("FTransformSection", "ScaleXArea", "X"), MakeShareable( new FFloatCurveKeyArea( TransformSection->GetScaleCurve( EAxis::X ) ) ) );
+			LayoutBuilder.AddKeyArea("Scale.Y", NSLOCTEXT("FTransformSection", "ScaleYArea", "Y"), MakeShareable( new FFloatCurveKeyArea( TransformSection->GetScaleCurve( EAxis::Y ) ) ) );
+			LayoutBuilder.AddKeyArea("Scale.Z", NSLOCTEXT("FTransformSection", "ScaleZArea", "Z"), MakeShareable( new FFloatCurveKeyArea( TransformSection->GetScaleCurve( EAxis::Z ) ) ) );
 		LayoutBuilder.PopCategory();
 	}
 
diff --git a/Engine/Source/Developer/MovieSceneTools/Public/MovieSceneToolHelpers.h b/Engine/Source/Developer/MovieSceneTools/Public/MovieSceneToolHelpers.h
new file mode 100644
--- /dev/null
+++ b/Engine/Source/Developer/MovieSceneTools/Public/MovieSceneToolHelpers.h
@@ -0,0 +1,88 @@
+// Copyright 1998-2014 Epic Games, Inc. All Rights Reserved.
+
+#pragma once
+
+#include "ISequencerSection.h"
+#include "IKeyArea.h"
+
+/**
+ * A key area for float keys
+ */
+class FFloatCurveKeyArea : public IKeyArea
+{
+public:
+	FFloatCurveKeyArea( FRichCurve& InCurve )
+		: Curve( InCurve )
+	{}
+
+	/** IKeyArea interface */
+	virtual TArray<FKeyHandle> GetUnsortedKeyHandles() const override
+	{
+		TArray<FKeyHandle> OutKeyHandles;
+		for (auto It(Curve.GetKeyHandleIterator()); It; ++It)
+		{
+			OutKeyHandles.Add(It.Key());
+		}
+		return OutKeyHandles;
+	}
+
+	virtual float GetKeyTime( FKeyHandle KeyHandle ) const override
+	{
+		return Curve.GetKeyTime( KeyHandle );
+	}
+
+	virtual FKeyHandle MoveKey( FKeyHandle KeyHandle, float DeltaPosition ) override
+	{
+		return Curve.SetKeyTime( KeyHandle, Curve.GetKeyTime( KeyHandle ) + DeltaPosition );
+	}
+
+	virtual void DeleteKey(FKeyHandle KeyHandle) override
+	{
+		Curve.DeleteKey(KeyHandle);
+	}
+private:
+	/** Curve with keys in this area */
+	FRichCurve& Curve;
+};
+
+
+
+/** A key area for integral keys */
+class FIntegralKeyArea : public IKeyArea
+{
+public:
+	FIntegralKeyArea( FIntegralCurve& InCurve )
+		: Curve(InCurve)
+	{}
+
+	/** IKeyArea interface */
+	virtual TArray<FKeyHandle> GetUnsortedKeyHandles() const override
+	{
+		TArray<FKeyHandle> OutKeyHandles;
+		for (auto It(Curve.GetKeyHandleIterator()); It; ++It)
+		{
+			OutKeyHandles.Add(It.Key());
+		}
+		return OutKeyHandles;
+	}
+
+	virtual float GetKeyTime( FKeyHandle KeyHandle ) const override
+	{
+		return Curve.GetKeyTime( KeyHandle );
+	}
+
+	virtual FKeyHandle MoveKey( FKeyHandle KeyHandle, float DeltaPosition ) override
+	{
+		return Curve.SetKeyTime( KeyHandle, Curve.GetKeyTime( KeyHandle ) + DeltaPosition );
+	}
+
+	virtual void DeleteKey(FKeyHandle KeyHandle) override
+	{
+		Curve.DeleteKey(KeyHandle);
+	}
+private:
+	/** Curve with keys in this area */
+	FIntegralCurve& Curve;
+};
+
+
diff --git a/Engine/Source/Developer/MovieSceneTools/Public/PropertySection.h b/Engine/Source/Developer/MovieSceneTools/Public/PropertySection.h
new file mode 100644
--- /dev/null
+++ b/Engine/Source/Developer/MovieSceneTools/Public/PropertySection.h
@@ -0,0 +1,51 @@
+// Copyright 1998-2014 Epic Games, Inc. All Rights Reserved.
+
+#pragma once
+
+#include "ISequencerSection.h"
+
+/**
+ * A generic implementation for displaying simple property sections
+ */
+class FPropertySection : public ISequencerSection
+{
+public:
+	FPropertySection( UMovieSceneSection& InSectionObject, FName SectionName )
+		: SectionObject( InSectionObject )
+	{
+		DisplayName = FText::FromString( FName::NameToDisplayString( SectionName.ToString(), false ) );
+	}
+
+	/** ISequencerSection interface */
+	virtual UMovieSceneSection* GetSectionObject() override { return &SectionObject; }
+
+	virtual FText GetDisplayName() const override
+	{
+		return DisplayName;
+	}
+	
+	virtual FText GetSectionTitle() const override { return FText::GetEmpty(); }
+
+	virtual void GenerateSectionLayout( class ISectionLayoutBuilder& LayoutBuilder ) const override {}
+
+	virtual int32 OnPaintSection( const FGeometry& AllottedGeometry, const FSlateRect& SectionClippingRect, FSlateWindowElementList& OutDrawElements, int32 LayerId, bool bParentEnabled ) const override 
+	{
+		// Add a box for the section
+		FSlateDrawElement::MakeBox( 
+			OutDrawElements,
+			LayerId,
+			AllottedGeometry.ToPaintGeometry(),
+			FEditorStyle::GetBrush("Sequencer.GenericSection.Background"),
+			SectionClippingRect
+		);
+
+		return LayerId;
+	}
+	
+protected:
+	/** Display name of the section */
+	FText DisplayName;
+	/** The section we are visualizing */
+	UMovieSceneSection& SectionObject;
+};
+
diff --git a/Engine/Source/Editor/PropertyEditor/Private/PropertyHandleImpl.cpp b/Engine/Source/Editor/PropertyEditor/Private/PropertyHandleImpl.cpp
--- a/Engine/Source/Editor/PropertyEditor/Private/PropertyHandleImpl.cpp
+++ b/Engine/Source/Editor/PropertyEditor/Private/PropertyHandleImpl.cpp
@@ -473,15 +473,15 @@ FPropertyAccess::Result FPropertyValueImpl::ImportText( const TArray<FObjectBase
 void FPropertyValueImpl::AccessRawData( TArray<void*>& RawData )
 {
 	TArray<FObjectBaseAddress> ObjectAddresses;
-	GetObjectsToModify( ObjectAddresses, PropertyNode.Pin().Get() );
+	GetObjectsToModify(ObjectAddresses, PropertyNode.Pin().Get());
 
-	if( ObjectAddresses.Num() > 0 )
+	if (ObjectAddresses.Num() > 0)
 	{
 		RawData.Empty();
-		RawData.AddZeroed( ObjectAddresses.Num() );
-		for ( int32 ObjectIndex = 0 ; ObjectIndex < ObjectAddresses.Num() ; ++ObjectIndex )
-		{	
-			const FObjectBaseAddress& Cur = ObjectAddresses[ ObjectIndex ];
+		RawData.AddZeroed(ObjectAddresses.Num());
+		for (int32 ObjectIndex = 0; ObjectIndex < ObjectAddresses.Num(); ++ObjectIndex)
+		{
+			const FObjectBaseAddress& Cur = ObjectAddresses[ObjectIndex];
 			if (Cur.BaseAddress != NULL)
 			{
 				RawData[ObjectIndex] = Cur.BaseAddress;
@@ -490,6 +490,25 @@ void FPropertyValueImpl::AccessRawData( TArray<void*>& RawData )
 	}
 }
 
+void FPropertyValueImpl::AccessRawData( TArray<const void*>& RawData ) const
+{
+	TArray<FObjectBaseAddress> ObjectAddresses;
+	GetObjectsToModify(ObjectAddresses, PropertyNode.Pin().Get());
+
+	if (ObjectAddresses.Num() > 0)
+	{
+		RawData.Empty();
+		RawData.AddZeroed(ObjectAddresses.Num());
+		for (int32 ObjectIndex = 0; ObjectIndex < ObjectAddresses.Num(); ++ObjectIndex)
+		{
+			const FObjectBaseAddress& Cur = ObjectAddresses[ObjectIndex];
+			if (Cur.BaseAddress != NULL)
+			{
+				RawData[ObjectIndex] = Cur.BaseAddress;
+			}
+		}
+	}
+}
 
 void FPropertyValueImpl::SetOnPropertyValueChanged( const FSimpleDelegate& InOnPropertyValueChanged )
 {
@@ -1634,6 +1653,12 @@ void FPropertyHandleBase::AccessRawData( TArray<void*>& RawData )
 	Implementation->AccessRawData( RawData );
 }
 
+void FPropertyHandleBase::AccessRawData( TArray<const void*>& RawData ) const
+{
+	Implementation->AccessRawData(RawData);
+}
+
+
 void FPropertyHandleBase::SetOnPropertyValueChanged( const FSimpleDelegate& InOnPropertyValueChanged )
 {
 	return Implementation->SetOnPropertyValueChanged(InOnPropertyValueChanged);
diff --git a/Engine/Source/Editor/PropertyEditor/Private/PropertyHandleImpl.h b/Engine/Source/Editor/PropertyEditor/Private/PropertyHandleImpl.h
--- a/Engine/Source/Editor/PropertyEditor/Private/PropertyHandleImpl.h
+++ b/Engine/Source/Editor/PropertyEditor/Private/PropertyHandleImpl.h
@@ -123,6 +123,7 @@ public:
 	 * @param RawData	An array of raw data. 
 	 */ 
 	void AccessRawData( TArray<void*>& RawData );
+	void AccessRawData( TArray<const void*>& RawData ) const;
 
 	/**
 	 * Sets a delegate to call when the property value changes
@@ -385,6 +386,7 @@ public:
 	virtual TSharedPtr<IPropertyHandle> GetChildHandle( FName ChildName ) const override;
 	virtual TSharedPtr<IPropertyHandle> GetParentHandle() const override;
 	virtual void AccessRawData( TArray<void*>& RawData ) override;
+	virtual void AccessRawData( TArray<const void*>& RawData ) const override;
 	virtual uint32 GetNumOuterObjects() const override;
 	virtual void GetOuterObjects( TArray<UObject*>& OuterObjects ) override;
 	virtual FPropertyAccess::Result GetNumChildren( uint32& OutNumChildren ) const override;
diff --git a/Engine/Source/Editor/PropertyEditor/Public/PropertyHandle.h b/Engine/Source/Editor/PropertyEditor/Public/PropertyHandle.h
--- a/Engine/Source/Editor/PropertyEditor/Public/PropertyHandle.h
+++ b/Engine/Source/Editor/PropertyEditor/Public/PropertyHandle.h
@@ -278,6 +278,7 @@ public:
 	 * @param RawData	An array of raw data.  The elements in this array are the raw data for this property on each of the objects in the  property editor
 	 */ 
 	virtual void AccessRawData( TArray<void*>& RawData ) = 0;
+	virtual void AccessRawData( TArray<const void*>& RawData ) const = 0;
 
 	/**
 	 * Returns this handle as an array if possible
diff --git a/Engine/Source/Editor/Sequencer/Private/SSection.cpp b/Engine/Source/Editor/Sequencer/Private/SSection.cpp
--- a/Engine/Source/Editor/Sequencer/Private/SSection.cpp
+++ b/Engine/Source/Editor/Sequencer/Private/SSection.cpp
@@ -230,7 +230,7 @@ int32 SSection::OnPaint( const FGeometry& AllottedGeometry, const FSlateRect& My
 	TransparentSelectionColor.A = .5f;
 		
 	// Section name with drop shadow
-	FString SectionTitle = SectionInterface->GetSectionTitle();
+	FText SectionTitle = SectionInterface->GetSectionTitle();
 	if (!SectionTitle.IsEmpty())
 	{
 		FSlateDrawElement::MakeText(
diff --git a/Engine/Source/Editor/Sequencer/Private/SectionLayoutBuilder.cpp b/Engine/Source/Editor/Sequencer/Private/SectionLayoutBuilder.cpp
--- a/Engine/Source/Editor/Sequencer/Private/SectionLayoutBuilder.cpp
+++ b/Engine/Source/Editor/Sequencer/Private/SectionLayoutBuilder.cpp
@@ -9,7 +9,7 @@ FSectionLayoutBuilder::FSectionLayoutBuilder( TSharedRef<FTrackNode> InRootNode
 	, CurrentNode( InRootNode )
 {}
 
-void FSectionLayoutBuilder::PushCategory( FName CategoryName, const FString& DisplayLabel )
+void FSectionLayoutBuilder::PushCategory( FName CategoryName, const FText& DisplayLabel )
 {
 	CurrentNode = CurrentNode->AddCategoryNode( CategoryName, DisplayLabel );
 }
@@ -28,7 +28,7 @@ void FSectionLayoutBuilder::SetSectionAsKeyArea( TSharedRef<IKeyArea> KeyArea )
 	RootNode->SetSectionAsKeyArea( KeyArea ); 
 }
 
-void FSectionLayoutBuilder::AddKeyArea( FName KeyAreaName, const FString& DisplayName, TSharedRef<IKeyArea> KeyArea )
+void FSectionLayoutBuilder::AddKeyArea( FName KeyAreaName, const FText& DisplayName, TSharedRef<IKeyArea> KeyArea )
 {
 	CurrentNode->AddKeyAreaNode( KeyAreaName, DisplayName, KeyArea );
 }
diff --git a/Engine/Source/Editor/Sequencer/Private/SectionLayoutBuilder.h b/Engine/Source/Editor/Sequencer/Private/SectionLayoutBuilder.h
--- a/Engine/Source/Editor/Sequencer/Private/SectionLayoutBuilder.h
+++ b/Engine/Source/Editor/Sequencer/Private/SectionLayoutBuilder.h
@@ -11,9 +11,9 @@ public:
 	FSectionLayoutBuilder( TSharedRef<FTrackNode> InRootNode );
 
 	/** ISectionLayoutBuilder Interface */
-	virtual void PushCategory( FName CategoryName, const FString& DisplayLabel ) override;
+	virtual void PushCategory( FName CategoryName, const FText& DisplayLabel ) override;
 	virtual void SetSectionAsKeyArea( TSharedRef<IKeyArea> KeyArea ) override;
-	virtual void AddKeyArea( FName KeyAreaName, const FString& DisplayName, TSharedRef<IKeyArea> KeyArea ) override;
+	virtual void AddKeyArea( FName KeyAreaName, const FText& DisplayName, TSharedRef<IKeyArea> KeyArea ) override;
 	virtual void PopCategory() override;
 private:
 	/** Root node of the tree */
diff --git a/Engine/Source/Editor/Sequencer/Private/Sequencer.cpp b/Engine/Source/Editor/Sequencer/Private/Sequencer.cpp
--- a/Engine/Source/Editor/Sequencer/Private/Sequencer.cpp
+++ b/Engine/Source/Editor/Sequencer/Private/Sequencer.cpp
@@ -365,7 +365,7 @@ void FSequencer::OnActorsDropped( const TArray<TWeakObjectPtr<AActor> >& Actors
 				// Create a new possessable
 				OwnerMovieScene->Modify();
 				
-				const FString PossessableName = Actor->GetActorLabel();
+				const FText PossessableName = FText::FromString( Actor->GetActorLabel() );
 				const FGuid PossessableGuid = OwnerMovieScene->AddPossessable( PossessableName, Actor->GetClass() );
 				 
 				if ( IsShotFilteringOn() )
@@ -464,7 +464,7 @@ FGuid FSequencer::GetHandleToObject( UObject* Object )
 			// Create a new possessable
 			OwnerMovieScene->Modify();
 			
-			const FString PossessableName = Object->GetName();
+			const FText PossessableName = FText::FromString( Object->GetName() );
 			ObjectGuid = OwnerMovieScene->AddPossessable( PossessableName, Object->GetClass() );
 			
 			if ( IsShotFilteringOn() )
@@ -902,7 +902,7 @@ FGuid FSequencer::AddSpawnableForAssetOrClass( UObject* Object, UObject* Counter
 		const FName BlueprintName = MakeUniqueObjectName( OwnerMovieScene, UBlueprint::StaticClass(), AssetName );
 
 		// Use the asset name as the initial spawnable name
-		const FString& NewSpawnableName = AssetName.ToString();		// @todo sequencer: Need UI to allow user to rename these slots
+		const FText NewSpawnableName = FText::FromName( AssetName );		// @todo sequencer: Need UI to allow user to rename these slots
 
 		// Create our new blueprint!
 		UBlueprint* NewBlueprint = NULL;
diff --git a/Engine/Source/Editor/Sequencer/Private/SequencerActorBindingManager.cpp b/Engine/Source/Editor/Sequencer/Private/SequencerActorBindingManager.cpp
--- a/Engine/Source/Editor/Sequencer/Private/SequencerActorBindingManager.cpp
+++ b/Engine/Source/Editor/Sequencer/Private/SequencerActorBindingManager.cpp
@@ -164,7 +164,7 @@ void FSequencerActorBindingManager::SpawnOrDestroyObjectsForInstance( TSharedRef
 	{
 		for( auto SpawnableIndex = 0; SpawnableIndex < MovieScene->GetSpawnableCount(); ++SpawnableIndex )
 		{
-			auto& Spawnable = MovieScene->GetSpawnable( SpawnableIndex );
+			FMovieSceneSpawnable& Spawnable = MovieScene->GetSpawnable( SpawnableIndex );
 
 			// Must have a valid world for us to be able to do this
 			if( ActorWorld != NULL )
@@ -213,7 +213,7 @@ void FSequencerActorBindingManager::SpawnOrDestroyObjectsForInstance( TSharedRef
 						if( NewActor )
 						{
 							// @todo sequencer: We're naming the actor based off of the spawnable's name.  Is that really what we want?
-							GEditor->SetActorLabelUnique( NewActor, Spawnable.GetName() );
+							GEditor->SetActorLabelUnique( NewActor, Spawnable.GetDisplayName().ToString() );
 
 							// Actor was spawned OK!
 
diff --git a/Engine/Source/Editor/Sequencer/Private/SequencerDisplayNode.cpp b/Engine/Source/Editor/Sequencer/Private/SequencerDisplayNode.cpp
--- a/Engine/Source/Editor/Sequencer/Private/SequencerDisplayNode.cpp
+++ b/Engine/Source/Editor/Sequencer/Private/SequencerDisplayNode.cpp
@@ -138,7 +138,7 @@ FSequencerDisplayNode::FSequencerDisplayNode( FName InNodeName, TSharedPtr<FSequ
 	bExpanded = ParentTree.GetSavedExpansionState( *this );
 }
 
-TSharedRef<FSectionCategoryNode> FSequencerDisplayNode::AddCategoryNode( FName CategoryName, const FString& DisplayLabel )
+TSharedRef<FSectionCategoryNode> FSequencerDisplayNode::AddCategoryNode( FName CategoryName, const FText& DisplayLabel )
 {
 	TSharedPtr<FSectionCategoryNode> CategoryNode;
 
@@ -189,7 +189,7 @@ TSharedRef<FTrackNode> FSequencerDisplayNode::AddSectionAreaNode( FName SectionN
 	return SectionNode.ToSharedRef();
 }
 
-void FSequencerDisplayNode::AddKeyAreaNode( FName KeyAreaName, const FString& DisplayName, TSharedRef<IKeyArea> KeyArea )
+void FSequencerDisplayNode::AddKeyAreaNode( FName KeyAreaName, const FText& DisplayName, TSharedRef<IKeyArea> KeyArea )
 {
 	TSharedPtr<FSectionKeyAreaNode> KeyAreaNode;
 
@@ -362,10 +362,10 @@ float FTrackNode::GetNodeHeight() const
 	return Sections.Num() > 0 ? Sections[0]->GetSectionHeight() * (GetMaxRowIndex()+1) : SequencerLayoutConstants::SectionAreaDefaultHeight;
 }
 
-FString FTrackNode::GetDisplayName() const
+FText FTrackNode::GetDisplayName() const
 {
 	// @todo Sequencer - IS there a better way to get the section interface name for the animation outliner?
-	return Sections.Num() > 0 ? Sections[0]->GetDisplayName() : NodeName.ToString();
+	return Sections.Num() > 0 ? Sections[0]->GetDisplayName() : FText::FromName(NodeName);
 }
 
 void FTrackNode::SetSectionAsKeyArea( TSharedRef<IKeyArea>& KeyArea )
@@ -373,7 +373,7 @@ void FTrackNode::SetSectionAsKeyArea( TSharedRef<IKeyArea>& KeyArea )
 	if( !TopLevelKeyNode.IsValid() )
 	{
 		bool bTopLevel = true;
-		TopLevelKeyNode = MakeShareable( new FSectionKeyAreaNode( GetNodeName(), TEXT(""), NULL, ParentTree, bTopLevel ) );
+		TopLevelKeyNode = MakeShareable( new FSectionKeyAreaNode( GetNodeName(), FText::GetEmpty(), NULL, ParentTree, bTopLevel ) );
 	}
 
 	TopLevelKeyNode->AddKeyArea( KeyArea );
diff --git a/Engine/Source/Editor/Sequencer/Private/SequencerDisplayNode.h b/Engine/Source/Editor/Sequencer/Private/SequencerDisplayNode.h
--- a/Engine/Source/Editor/Sequencer/Private/SequencerDisplayNode.h
+++ b/Engine/Source/Editor/Sequencer/Private/SequencerDisplayNode.h
@@ -42,7 +42,7 @@ public:
 	 * @param CategoryName	Name of the category
 	 * @param DisplayName	Localized label to display in the animation outliner
 	 */
-	TSharedRef<class FSectionCategoryNode> AddCategoryNode( FName CategoryName, const FString& DisplayLabel );
+	TSharedRef<class FSectionCategoryNode> AddCategoryNode( FName CategoryName, const FText& DisplayLabel );
 
 	/**
 	 * Adds a new section area for this node.
@@ -59,7 +59,7 @@ public:
 	 * @param DisplayLabel	Localized label to display in the animation outliner
 	 * @param KeyArea		Key area interface for drawing and interaction with keys
 	 */
-	void AddKeyAreaNode( FName KeyAreaName, const FString& DisplayLabel, TSharedRef<IKeyArea> KeyArea );
+	void AddKeyAreaNode( FName KeyAreaName, const FText& DisplayLabel, TSharedRef<IKeyArea> KeyArea );
 
 	/**
 	 * @return The type of node this is
@@ -77,7 +77,7 @@ public:
 	/**
 	 * @return The localized display name of this node
 	 */
-	virtual FString GetDisplayName() const = 0;
+	virtual FText GetDisplayName() const = 0;
 
 	/**
 	 * Generates a widget for display in the animation outliner portion of the track area
@@ -240,7 +240,7 @@ public:
 	 * @param InParentTree	The tree this node is in
 	 * @param bInTopLevel	If true the node is part of the section itself instead of taking up extra height in the section
 	 */
-	FSectionKeyAreaNode( FName NodeName, const FString& InDisplayName, TSharedPtr<FSequencerDisplayNode> InParentNode, FSequencerNodeTree& InParentTree, bool bInTopLevel = false )
+	FSectionKeyAreaNode( FName NodeName, const FText& InDisplayName, TSharedPtr<FSequencerDisplayNode> InParentNode, FSequencerNodeTree& InParentTree, bool bInTopLevel = false )
 		: FSequencerDisplayNode( NodeName, InParentNode, InParentTree )
 		, DisplayName( InDisplayName )
 		, bTopLevel( bInTopLevel )
@@ -249,7 +249,7 @@ public:
 	/** FSequencerDisplayNodeInterface */
 	virtual ESequencerNode::Type GetType() const override { return ESequencerNode::KeyArea; }
 	virtual float GetNodeHeight() const override;
-	virtual FString GetDisplayName() const override { return DisplayName; }
+	virtual FText GetDisplayName() const override { return DisplayName; }
 	virtual bool GetShotFilteredVisibilityToCache() const override;
 
 	/**
@@ -278,7 +278,7 @@ public:
 	bool IsTopLevel() const { return bTopLevel; }
 private:
 	/** The display name of the key area */
-	FString DisplayName;
+	FText DisplayName;
 	/** All key areas on this node (one per section) */
 	TArray< TSharedRef<IKeyArea> > KeyAreas;
 	/** If true the node is part of the section itself instead of taking up extra height in the section */
@@ -305,7 +305,7 @@ public:
 	/** FSequencerDisplayNodeInterface */
 	virtual ESequencerNode::Type GetType() const override { return ESequencerNode::Track; }
 	virtual float GetNodeHeight() const override;
-	virtual FString GetDisplayName() const override;
+	virtual FText GetDisplayName() const override;
 	virtual bool IsSelectable() const override { return true; }
 	virtual bool GetShotFilteredVisibilityToCache() const override;
 	virtual void GetChildKeyAreaNodesRecursively(TArray< TSharedRef<class FSectionKeyAreaNode> >& OutNodes) const override;
@@ -369,7 +369,7 @@ public:
 	 * @param InParentNode		The parent of this node or NULL if this is a root node
 	 * @param InParentTree		The tree this node is in
 	 */
-	FObjectBindingNode( FName NodeName, const FString& InObjectName, const FGuid& InObjectBinding, TSharedPtr<FSequencerDisplayNode> InParentNode, FSequencerNodeTree& InParentTree )
+	FObjectBindingNode( FName NodeName, const FText& InObjectName, const FGuid& InObjectBinding, TSharedPtr<FSequencerDisplayNode> InParentNode, FSequencerNodeTree& InParentTree )
 		: FSequencerDisplayNode( NodeName, InParentNode, InParentTree )
 		, ObjectBinding( InObjectBinding )
 		, DisplayName( InObjectName )
@@ -377,7 +377,7 @@ public:
 
 	/** FSequencerDisplayNodeInterface */
 	virtual ESequencerNode::Type GetType() const override { return ESequencerNode::Object; }
-	virtual FString GetDisplayName() const override { return DisplayName; }
+	virtual FText GetDisplayName() const override { return DisplayName; }
 	virtual float GetNodeHeight() const override;
 	virtual bool IsSelectable() const override { return true; }
 	virtual bool GetShotFilteredVisibilityToCache() const override;
@@ -392,7 +392,7 @@ private:
 	/** The binding to live objects */
 	FGuid ObjectBinding;
 	/** Display name of the object */
-	FString DisplayName;
+	FText DisplayName;
 };
 
 /**
@@ -409,7 +409,7 @@ public:
 	 * @param InParentNode	The parent of this node or NULL if this is a root node
 	 * @param InParentTree	The tree this node is in
 	 */
-	FSectionCategoryNode( FName NodeName, const FString& InDisplayName, TSharedPtr<FSequencerDisplayNode> InParentNode, FSequencerNodeTree& InParentTree )
+	FSectionCategoryNode( FName NodeName, const FText& InDisplayName, TSharedPtr<FSequencerDisplayNode> InParentNode, FSequencerNodeTree& InParentTree )
 		: FSequencerDisplayNode( NodeName, InParentNode, InParentTree )
 		, DisplayName( InDisplayName )
 	{}
@@ -417,9 +417,9 @@ public:
 	/** FSequencerDisplayNodeInterface */
 	virtual ESequencerNode::Type GetType() const override { return ESequencerNode::Category; }
 	virtual float GetNodeHeight() const override;
-	virtual FString GetDisplayName() const override { return DisplayName; }
+	virtual FText GetDisplayName() const override { return DisplayName; }
 	virtual bool GetShotFilteredVisibilityToCache() const override;
 private:
 	/** The display name of the category */
-	FString DisplayName;
+	FText DisplayName;
 };
diff --git a/Engine/Source/Editor/Sequencer/Private/SequencerNodeTree.cpp b/Engine/Source/Editor/Sequencer/Private/SequencerNodeTree.cpp
--- a/Engine/Source/Editor/Sequencer/Private/SequencerNodeTree.cpp
+++ b/Engine/Source/Editor/Sequencer/Private/SequencerNodeTree.cpp
@@ -48,7 +48,7 @@ void FSequencerNodeTree::Update()
 	// Make nodes for all object bindings
 	for( int32 BindingIndex = 0; BindingIndex < ObjectBindings.Num(); ++BindingIndex )
 	{
-		TSharedRef<FObjectBindingNode> ObjectBindingNode = AddObjectBinding( ObjectBindings[BindingIndex].GetBindingName(), ObjectBindings[BindingIndex].GetObjectGuid() );
+		TSharedRef<FObjectBindingNode> ObjectBindingNode = AddObjectBinding( ObjectBindings[BindingIndex].GetDisplayName(), ObjectBindings[BindingIndex].GetObjectGuid() );
 
 		const TArray<UMovieSceneTrack*>& Tracks = ObjectBindings[BindingIndex].GetTracks();
 
@@ -133,7 +133,7 @@ const TArray< TSharedRef<FSequencerDisplayNode> >& FSequencerNodeTree::GetRootNo
 }
 
 
-TSharedRef<FObjectBindingNode> FSequencerNodeTree::AddObjectBinding( const FString& ObjectName, const FGuid& ObjectBinding )
+TSharedRef<FObjectBindingNode> FSequencerNodeTree::AddObjectBinding( const FText& ObjectName, const FGuid& ObjectBinding )
 {
 	// The node name is the object guid
 	FName ObjectNodeName = *ObjectBinding.ToString();
@@ -226,7 +226,7 @@ static bool FilterNodesRecursive( const TSharedRef<FSequencerDisplayNode>& Start
 	{
 		const FString& TestName = FilterStrings[TestNameIndex];
 
-		if ( !StartNode->GetDisplayName().Contains( TestName ) ) 
+		if ( !StartNode->GetDisplayName().ToString().Contains( TestName ) ) 
 		{
 			bFilterAcceptable = false;
 			break;
diff --git a/Engine/Source/Editor/Sequencer/Private/SequencerNodeTree.h b/Engine/Source/Editor/Sequencer/Private/SequencerNodeTree.h
--- a/Engine/Source/Editor/Sequencer/Private/SequencerNodeTree.h
+++ b/Engine/Source/Editor/Sequencer/Private/SequencerNodeTree.h
@@ -114,7 +114,7 @@ private:
 	 * 
 	 * @param ObjectBinding	The object binding guid
 	 */
-	TSharedRef<FObjectBindingNode> AddObjectBinding( const FString& ObjectName, const FGuid& ObjectBinding );
+	TSharedRef<FObjectBindingNode> AddObjectBinding( const FText& ObjectName, const FGuid& ObjectBinding );
 private:
 	/** Tools for building movie scene section layouts.  One tool for each track */
 	TMap< class UMovieSceneTrack*, TSharedPtr<FMovieSceneTrackEditor> > EditorMap;
diff --git a/Engine/Source/Editor/Sequencer/Public/ISectionLayoutBuilder.h b/Engine/Source/Editor/Sequencer/Public/ISectionLayoutBuilder.h
--- a/Engine/Source/Editor/Sequencer/Public/ISectionLayoutBuilder.h
+++ b/Engine/Source/Editor/Sequencer/Public/ISectionLayoutBuilder.h
@@ -14,7 +14,7 @@ public:
 	 * @param CategoryName	The name of the category
 	 * @param DisplayLabel	The localized display label for the category
 	 */
-	virtual void PushCategory( FName CategoryName, const FString& DisplayLabel ) = 0;
+	virtual void PushCategory( FName CategoryName, const FText& DisplayLabel ) = 0;
 	
 	/**
 	 * Sets the section as a key area itself
@@ -29,7 +29,7 @@ public:
 	 * @param DisplayLabel	The localized display label for the key area
 	 * @param KeyArea		Interface for accessing and drawing keys
 	 */
-	virtual void AddKeyArea( FName KeyAreaName, const FString& DisplayLabel, TSharedRef<class IKeyArea> KeyArea ) = 0;
+	virtual void AddKeyArea( FName KeyAreaName, const FText& DisplayLabel, TSharedRef<class IKeyArea> KeyArea ) = 0;
 
 	/**
 	 * Pops a category off the stack
diff --git a/Engine/Source/Editor/Sequencer/Public/ISequencerSection.h b/Engine/Source/Editor/Sequencer/Public/ISequencerSection.h
--- a/Engine/Source/Editor/Sequencer/Public/ISequencerSection.h
+++ b/Engine/Source/Editor/Sequencer/Public/ISequencerSection.h
@@ -40,12 +40,12 @@ public:
 	/**
 	 * @return The display name of the section for the animation outliner
 	 */
-	virtual FString GetDisplayName() const = 0;
+	virtual FText GetDisplayName() const = 0;
 	
 	/**
 	 * @return The display name of the section in the section view
 	 */
-	virtual FString GetSectionTitle() const = 0;
+	virtual FText GetSectionTitle() const = 0;
 
 	/**
 	 * Generates the inner layout for this section
diff --git a/Engine/Source/Runtime/MovieSceneCore/Classes/MovieScene.h b/Engine/Source/Runtime/MovieSceneCore/Classes/MovieScene.h
--- a/Engine/Source/Runtime/MovieSceneCore/Classes/MovieScene.h
+++ b/Engine/Source/Runtime/MovieSceneCore/Classes/MovieScene.h
@@ -21,7 +21,7 @@ public:
 	}
 
 	/** FMovieSceneSpawnable initialization constructor */
-	FMovieSceneSpawnable( const FString& InitName, UClass* InitClass, UObject* InitCounterpartGamePreviewObject );
+	FMovieSceneSpawnable( const FText& InitName, UClass* InitClass, UObject* InitCounterpartGamePreviewObject );
 
 	/** @return Returns the guid for this possessable */
 	const FGuid& GetGuid() const
@@ -30,7 +30,7 @@ public:
 	}
 
 	/** @return Returns the name of this spawnable */
-	const FString& GetName() const
+	const FText& GetDisplayName() const
 	{
 		return Name;
 	}
@@ -55,7 +55,7 @@ private:
 	/** Name label */
 	// @todo sequencer: Should be editor-only probably
 	UPROPERTY()
-	FString Name;
+	FText Name;
 
 	/** Data-only blueprint-generated-class for this object */
 	// @todo sequencer: Could be weak object ptr here, IF blueprints that are inners are housekept properly without references
@@ -84,7 +84,7 @@ public:
 	}
 
 	/** FMovieScenePossessable initialization constructor */
-	FMovieScenePossessable( const FString& InitName, UClass* InitPossessedObjectClass );
+	FMovieScenePossessable( const FText& InitName, UClass* InitPossessedObjectClass );
 
 	/** @return Returns the guid for this possessable */
 	const FGuid& GetGuid() const
@@ -93,7 +93,7 @@ public:
 	}
 
 	/** @return Returns the name of this spawnable */
-	const FString& GetName() const
+	const FText& GetDisplayName() const
 	{
 		return Name;
 	}
@@ -115,7 +115,7 @@ private:
 	/** Name label for this slot */
 	// @todo sequencer: Should be editor-only probably
 	UPROPERTY()
-	FString Name;
+	FText Name;
 
 	/** Type of the object we'll be possessing */
 	// @todo sequencer: Might be able to be editor-only.  We'll see.
@@ -151,13 +151,13 @@ struct FMovieSceneObjectBinding
 	FMovieSceneObjectBinding()
 	{}
 
-	FMovieSceneObjectBinding( const FGuid& InObjectGuid, const FString& InBindingName, const TArray<class UMovieSceneTrack*>& InTracks )
+	FMovieSceneObjectBinding( const FGuid& InObjectGuid, const FText& InBindingName, const TArray<class UMovieSceneTrack*>& InTracks )
 		: ObjectGuid( InObjectGuid )
 		, BindingName( InBindingName )
 		, Tracks( InTracks )
 	{}
 
-	FMovieSceneObjectBinding( const FGuid& InObjectGuid, const FString& InBindingName )
+	FMovieSceneObjectBinding( const FGuid& InObjectGuid, const FText& InBindingName )
 		: ObjectGuid( InObjectGuid )
 		, BindingName( InBindingName )
 	{}
@@ -175,7 +175,7 @@ struct FMovieSceneObjectBinding
 	/**
 	 * @return The display name of the binding
 	 */
-	const FString& GetBindingName() const { return BindingName; } 
+	const FText& GetDisplayName() const { return BindingName; } 
 
 	/**
 	 * Adds a new track to this binding
@@ -203,7 +203,7 @@ private:
 	
 	/** Display name */
 	UPROPERTY()
-	FString BindingName;
+	FText BindingName;
 
 	/** All tracks in this binding */
 	UPROPERTY()
@@ -231,7 +231,7 @@ public:
 	 *
 	 * @return	Guid of the newly-added spawnable
 	 */
-	virtual FGuid AddSpawnable( const FString& Name, class UBlueprint* Blueprint, UObject* CounterpartGamePreviewObject );
+	virtual FGuid AddSpawnable( const FText& Name, class UBlueprint* Blueprint, UObject* CounterpartGamePreviewObject );
 
 	/**
 	 * Removes a spawnable from this movie scene.
@@ -282,7 +282,7 @@ public:
 	 *
 	 * @return	Guid of the newly-added possessable
 	 */
-	virtual FGuid AddPossessable( const FString& Name, UClass* Class );
+	virtual FGuid AddPossessable( const FText& Name, UClass* Class );
 
 	/**
 	 * Removes a possessable from this movie scene.
diff --git a/Engine/Source/Runtime/MovieSceneCore/Classes/MovieScenePropertyTrack.h b/Engine/Source/Runtime/MovieSceneCore/Classes/MovieScenePropertyTrack.h
--- a/Engine/Source/Runtime/MovieSceneCore/Classes/MovieScenePropertyTrack.h
+++ b/Engine/Source/Runtime/MovieSceneCore/Classes/MovieScenePropertyTrack.h
@@ -2,6 +2,7 @@
 
 #pragma once
 
+#include "MovieSceneTrack.h"
 #include "MovieScenePropertyTrack.generated.h"
 
 
diff --git a/Engine/Source/Runtime/MovieSceneCore/Classes/MovieSceneTrack.h b/Engine/Source/Runtime/MovieSceneCore/Classes/MovieSceneTrack.h
--- a/Engine/Source/Runtime/MovieSceneCore/Classes/MovieSceneTrack.h
+++ b/Engine/Source/Runtime/MovieSceneCore/Classes/MovieSceneTrack.h
@@ -4,6 +4,7 @@
 
 #include "MovieSceneTrack.generated.h"
 
+class UMovieSceneSection;
 
 /*
  * Base class for a track in a Movie Scene
diff --git a/Engine/Source/Runtime/MovieSceneCore/Private/MovieScene.cpp b/Engine/Source/Runtime/MovieSceneCore/Private/MovieScene.cpp
--- a/Engine/Source/Runtime/MovieSceneCore/Private/MovieScene.cpp
+++ b/Engine/Source/Runtime/MovieSceneCore/Private/MovieScene.cpp
@@ -4,7 +4,7 @@
 #include "MovieScene.h"
 
 
-FMovieSceneSpawnable::FMovieSceneSpawnable( const FString& InitName, UClass* InitClass, UObject* InitCounterpartGamePreviewObject )
+FMovieSceneSpawnable::FMovieSceneSpawnable( const FText& InitName, UClass* InitClass, UObject* InitCounterpartGamePreviewObject )
 {
 	Guid = FGuid::NewGuid();
 	Name = InitName;
@@ -14,7 +14,7 @@ FMovieSceneSpawnable::FMovieSceneSpawnable( const FString& InitName, UClass* Ini
 
 
 
-FMovieScenePossessable::FMovieScenePossessable( const FString& InitName, UClass* InitPossessedObjectClass )
+FMovieScenePossessable::FMovieScenePossessable( const FText& InitName, UClass* InitPossessedObjectClass )
 {
 	Guid = FGuid::NewGuid();
 	Name = InitName;
@@ -49,7 +49,7 @@ UMovieScene::UMovieScene( const FPostConstructInitializeProperties& PCIP )
 
 #if WITH_EDITOR
 // @todo sequencer: Some of these methods should only be used by tools, and should probably move out of MovieSceneCore!
-FGuid UMovieScene::AddSpawnable( const FString& Name, UBlueprint* Blueprint, UObject* CounterpartGamePreviewObject )
+FGuid UMovieScene::AddSpawnable( const FText& Name, UBlueprint* Blueprint, UObject* CounterpartGamePreviewObject )
 {
 	check( (Blueprint != NULL) && (Blueprint->GeneratedClass) );
 
@@ -59,7 +59,7 @@ FGuid UMovieScene::AddSpawnable( const FString& Name, UBlueprint* Blueprint, UOb
 	Spawnables.Add( NewSpawnable );
 
 	// Add a new binding so that tracks can be added to it
-	new (ObjectBindings) FMovieSceneObjectBinding( NewSpawnable.GetGuid(), NewSpawnable.GetName() );
+	new (ObjectBindings) FMovieSceneObjectBinding( NewSpawnable.GetGuid(), NewSpawnable.GetDisplayName() );
 
 	return NewSpawnable.GetGuid();
 }
@@ -150,7 +150,7 @@ const FMovieSceneSpawnable* UMovieScene::FindSpawnableForCounterpart( UObject* G
 }
 
 
-FGuid UMovieScene::AddPossessable( const FString& Name, UClass* Class )
+FGuid UMovieScene::AddPossessable( const FText& Name, UClass* Class )
 {
 	Modify();
 
@@ -158,7 +158,7 @@ FGuid UMovieScene::AddPossessable( const FString& Name, UClass* Class )
 	Possessables.Add( NewPossessable );
 
 	// Add a new binding so that tracks can be added to it
-	new (ObjectBindings) FMovieSceneObjectBinding( NewPossessable.GetGuid(), NewPossessable.GetName() );
+	new (ObjectBindings) FMovieSceneObjectBinding( NewPossessable.GetGuid(), NewPossessable.GetDisplayName() );
 
 	return NewPossessable.GetGuid();
 }
diff --git a/Engine/Source/Runtime/MovieSceneCore/Private/MovieSceneCommonHelpers.cpp b/Engine/Source/Runtime/MovieSceneCore/Private/MovieSceneCommonHelpers.cpp
new file mode 100644
--- /dev/null
+++ b/Engine/Source/Runtime/MovieSceneCore/Private/MovieSceneCommonHelpers.cpp
@@ -0,0 +1,74 @@
+// Copyright 1998-2014 Epic Games, Inc. All Rights Reserved.
+
+#include "MovieSceneCorePrivatePCH.h"
+#include "IMovieScenePlayer.h"
+#include "MovieSceneCommonHelpers.h"
+
+
+TArray<UMovieSceneSection*> MovieSceneHelpers::GetTraversedSections( const TArray<UMovieSceneSection*>& Sections, float CurrentTime, float PreviousTime )
+{
+	TArray<UMovieSceneSection*> TraversedSections;
+
+	bool bPlayingBackwards = CurrentTime - PreviousTime < 0.0f;
+	
+	float MaxTime = bPlayingBackwards ? PreviousTime : CurrentTime;
+	float MinTime = bPlayingBackwards ? CurrentTime : PreviousTime;
+
+	TRange<float> SliderRange(MinTime, MaxTime);
+
+	for( int32 SectionIndex = 0; SectionIndex < Sections.Num(); ++SectionIndex )
+	{
+		UMovieSceneSection* Section = Sections[SectionIndex];
+		if( SliderRange.Overlaps( TRange<float>( Section->GetRange() ) ) )
+		{
+			TraversedSections.Add( Section );
+		}
+	}
+
+	return TraversedSections;
+}
+
+UMovieSceneSection* MovieSceneHelpers::FindSectionAtTime( const TArray<UMovieSceneSection*>& Sections, float Time )
+{
+	// Only update the section if the position is within the time span of the section
+	// Or if there are no sections at the time, the left closest section to the time
+	// Or in the case that Time is before all sections, take the section with the earliest start time
+	UMovieSceneSection* ClosestSection = NULL;
+	float ClosestSectionTime = 0.f;
+	UMovieSceneSection* EarliestSection = NULL;
+	float EarliestSectionTime = 0.f;
+	for( int32 SectionIndex = 0; SectionIndex < Sections.Num(); ++SectionIndex )
+	{
+		UMovieSceneSection* Section = Sections[SectionIndex];
+
+		//@todo Sequencer - There can be multiple sections overlapping in time. Returning instantly does not account for that.
+		if( Section->IsTimeWithinSection( Time ) )
+		{
+			return Section;
+		}
+
+		float EndTime = Section->GetEndTime();
+		if (EndTime < Time)
+		{
+			float ClosestTime = Time - EndTime;
+			if (!ClosestSection || ClosestTime < ClosestSectionTime)
+			{
+				ClosestSection = Section;
+				ClosestSectionTime = ClosestTime;
+			}
+		}
+
+		float StartTime = Section->GetStartTime();
+		if (!EarliestSection || StartTime < EarliestSectionTime)
+		{
+			EarliestSection = Section;
+			EarliestSectionTime = StartTime;
+		}
+	}
+
+	// if we get here, we are off of any section
+	// if ClosestSection, then we take the closest to left of this time
+	// else, we take the EarliestSection
+	// if that's NULL, then there are no sections
+	return ClosestSection ? ClosestSection : EarliestSection;
+}
diff --git a/Engine/Source/Runtime/MovieSceneCore/Public/MovieSceneCommonHelpers.h b/Engine/Source/Runtime/MovieSceneCore/Public/MovieSceneCommonHelpers.h
new file mode 100644
--- /dev/null
+++ b/Engine/Source/Runtime/MovieSceneCore/Public/MovieSceneCommonHelpers.h
@@ -0,0 +1,21 @@
+// Copyright 1998-2014 Epic Games, Inc. All Rights Reserved.
+
+#pragma once
+
+
+class MOVIESCENECORE_API MovieSceneHelpers
+{
+public:
+	/**
+	 * Gets the sections that were traversed over between the current time and the previous time
+	 */
+	static TArray<UMovieSceneSection*> GetTraversedSections( const TArray<UMovieSceneSection*>& Sections, float CurrentTime, float PreviousTime );
+
+	/**
+	* Finds a section that exists at a given time
+	*
+	* @param Time	The time to find a section at
+	* @return The found section or null
+	*/
+	static UMovieSceneSection* FindSectionAtTime( const TArray<UMovieSceneSection*>& Sections, float Time );
+};
\ No newline at end of file
diff --git a/Engine/Source/Runtime/MovieSceneCoreTypes/Private/MovieSceneCommonHelpers.cpp b/Engine/Source/Runtime/MovieSceneCoreTypes/Private/MovieSceneCommonHelpers.cpp
deleted file mode 100644
diff --git a/Engine/Source/Runtime/MovieSceneCoreTypes/Private/MovieSceneCommonHelpers.h b/Engine/Source/Runtime/MovieSceneCoreTypes/Private/MovieSceneCommonHelpers.h
deleted file mode 100644
