Fix Sequencer auto-key not working on anything but structs

diff --git a/Engine/Source/Editor/Sequencer/Private/SequencerObjectChangeListener.cpp b/Engine/Source/Editor/Sequencer/Private/SequencerObjectChangeListener.cpp
--- a/Engine/Source/Editor/Sequencer/Private/SequencerObjectChangeListener.cpp
+++ b/Engine/Source/Editor/Sequencer/Private/SequencerObjectChangeListener.cpp
@@ -40,7 +40,7 @@ void FSequencerObjectChangeListener::OnPropertyChanged( const TArray<UObject*>&
 	UProperty* Property = HandleToUse->GetProperty();
 	UStructProperty* StructProperty = Cast<UStructProperty>( Property );
 
-	FName PropertyName = StructProperty ? StructProperty->Struct->GetFName() : Property->GetFName();
+	FName PropertyName = StructProperty ? StructProperty->Struct->GetFName() : Property->GetClass()->GetFName();
 
 	ClassToPropertyChangedMap.FindRef( PropertyName ).Broadcast( ChangedObjects, *HandleToUse, bRequireAutoKey );
 }
