UMG - Made slot animation work

diff --git a/Engine/Source/Developer/MovieSceneTools/Private/AssetTypeActions_MovieScene.cpp b/Engine/Source/Developer/MovieSceneTools/Private/AssetTypeActions_MovieScene.cpp
--- a/Engine/Source/Developer/MovieSceneTools/Private/AssetTypeActions_MovieScene.cpp
+++ b/Engine/Source/Developer/MovieSceneTools/Private/AssetTypeActions_MovieScene.cpp
@@ -44,7 +44,8 @@ void FAssetTypeActions_MovieScene::OpenAssetEditor( const TArray<UObject*>& InOb
 			// @todo sequencer: Only allow users to create new MovieScenes if that feature is turned on globally.
 			if( FParse::Param( FCommandLine::Get(), TEXT( "Sequencer" ) ) )
 			{
-				FModuleManager::LoadModuleChecked< ISequencerModule >( "Sequencer" ).CreateSequencerAssetEditor( Mode, EditWithinLevelEditor, MovieScene, true );
+				FSequencerViewParams ViewParams;
+				FModuleManager::LoadModuleChecked< ISequencerModule >( "Sequencer" ).CreateSequencerAssetEditor( Mode, ViewParams, EditWithinLevelEditor, MovieScene, true );
 			}
 		}
 	}
diff --git a/Engine/Source/Developer/MovieSceneTools/Public/MovieSceneToolHelpers.h b/Engine/Source/Developer/MovieSceneTools/Public/MovieSceneToolHelpers.h
--- a/Engine/Source/Developer/MovieSceneTools/Public/MovieSceneToolHelpers.h
+++ b/Engine/Source/Developer/MovieSceneTools/Public/MovieSceneToolHelpers.h
@@ -38,7 +38,10 @@ public:
 
 	virtual void DeleteKey(FKeyHandle KeyHandle) override
 	{
-		Curve.DeleteKey(KeyHandle);
+		if( Curve.IsKeyHandleValid(KeyHandle) )
+		{
+			Curve.DeleteKey(KeyHandle);
+		}
 	}
 private:
 	/** Curve with keys in this area */
diff --git a/Engine/Source/Editor/DetailCustomizations/Private/MarginCustomization.h b/Engine/Source/Editor/DetailCustomizations/Private/MarginCustomization.h
--- a/Engine/Source/Editor/DetailCustomizations/Private/MarginCustomization.h
+++ b/Engine/Source/Editor/DetailCustomizations/Private/MarginCustomization.h
@@ -7,6 +7,10 @@ class FMarginStructCustomization : public IPropertyTypeCustomization
 public:
 	static TSharedRef<IPropertyTypeCustomization> MakeInstance();
 
+	FMarginStructCustomization()
+		: bIsUsingSlider( false )
+	{}
+
 	/** IPropertyTypeCustomization interface */
 	virtual void CustomizeHeader( TSharedRef<class IPropertyHandle> StructPropertyHandle, class FDetailWidgetRow& HeaderRow, IPropertyTypeCustomizationUtils& StructCustomizationUtils ) override;
 
diff --git a/Engine/Source/Editor/Sequencer/Private/Sequencer.cpp b/Engine/Source/Editor/Sequencer/Private/Sequencer.cpp
--- a/Engine/Source/Editor/Sequencer/Private/Sequencer.cpp
+++ b/Engine/Source/Editor/Sequencer/Private/Sequencer.cpp
@@ -56,6 +56,9 @@ void FSequencer::InitSequencer( const FSequencerInitParams& InitParams, const TA
 
 		ToolkitHost = InitParams.ToolkitHost;
 
+		LastViewRange = InitParams.ViewParams.InitalViewRange;
+		ScrubPosition = InitParams.ViewParams.InitialScrubPosition;
+
 		ObjectChangeListener = InitParams.ObjectChangeListener;
 		ObjectBindingManager = InitParams.ObjectBindingManager;
 
diff --git a/Engine/Source/Editor/Sequencer/Private/Sequencer.h b/Engine/Source/Editor/Sequencer/Private/Sequencer.h
--- a/Engine/Source/Editor/Sequencer/Private/Sequencer.h
+++ b/Engine/Source/Editor/Sequencer/Private/Sequencer.h
@@ -24,7 +24,10 @@ struct FSequencerInitParams
 	
 	/** The asset editor created for this (if any) */
 	TSharedPtr<IToolkitHost> ToolkitHost;
-	
+
+	/** View parameters */
+	FSequencerViewParams ViewParams;
+
 	/** Whether or not sequencer should be edited within the level editor */
 	bool bEditWithinLevelEditor;
 };
diff --git a/Engine/Source/Editor/Sequencer/Private/SequencerAssetEditor.cpp b/Engine/Source/Editor/Sequencer/Private/SequencerAssetEditor.cpp
--- a/Engine/Source/Editor/Sequencer/Private/SequencerAssetEditor.cpp
+++ b/Engine/Source/Editor/Sequencer/Private/SequencerAssetEditor.cpp
@@ -46,7 +46,7 @@ void FSequencerAssetEditor::UnregisterTabSpawners(const TSharedRef<class FTabMan
 	LevelEditorModule.AttachSequencer( SNullWidget::NullWidget, nullptr );
 }
 
-void FSequencerAssetEditor::InitSequencerAssetEditor( const EToolkitMode::Type Mode, const TSharedPtr< class IToolkitHost >& InitToolkitHost, UMovieScene* InRootMovieScene, const TArray<FOnCreateTrackEditor>& TrackEditorDelegates, bool bEditWithinLevelEditor )
+void FSequencerAssetEditor::InitSequencerAssetEditor( const EToolkitMode::Type Mode, const FSequencerViewParams& InViewParams, const TSharedPtr< class IToolkitHost >& InitToolkitHost, UMovieScene* InRootMovieScene, const TArray<FOnCreateTrackEditor>& TrackEditorDelegates, bool bEditWithinLevelEditor )
 {
 	
 	{
@@ -70,6 +70,7 @@ void FSequencerAssetEditor::InitSequencerAssetEditor( const EToolkitMode::Type M
 	Sequencer = MakeShareable(new FSequencer);
 	
 	FSequencerInitParams SequencerInitParams;
+	SequencerInitParams.ViewParams = InViewParams;
 	SequencerInitParams.ObjectChangeListener = MakeShareable( new FSequencerObjectChangeListener( Sequencer.ToSharedRef(), bEditWithinLevelEditor ) );
 	
 	SequencerInitParams.ObjectBindingManager = MakeShareable( new FSequencerActorBindingManager( InitToolkitHost->GetWorld(), SequencerInitParams.ObjectChangeListener.ToSharedRef(), Sequencer.ToSharedRef() ) );
diff --git a/Engine/Source/Editor/Sequencer/Private/SequencerAssetEditor.h b/Engine/Source/Editor/Sequencer/Private/SequencerAssetEditor.h
--- a/Engine/Source/Editor/Sequencer/Private/SequencerAssetEditor.h
+++ b/Engine/Source/Editor/Sequencer/Private/SequencerAssetEditor.h
@@ -17,12 +17,13 @@ public:
 	 * Edits the specified object
 	 *
 	 * @param	Mode							Asset editing mode for this editor (standalone or world-centric)
+	 * @param	InViewParams					Parameters for how to view sequencer UI
 	 * @param	InitToolkitHost					When Mode is WorldCentric, this is the level editor instance to spawn this editor within
 	 * @param	ObjectToEdit					The object to edit
 	 * @param	TrackEditorDelegates			Delegates to call to create auto-key handlers for this sequencer
 	 * @param	bEditWithinLevelEditor			Whether or not sequencer should be edited within the level editor
 	 */
-	void InitSequencerAssetEditor( const EToolkitMode::Type Mode, const TSharedPtr< class IToolkitHost >& InitToolkitHost, UMovieScene* InRootMovieScene, const TArray<FOnCreateTrackEditor>& TrackEditorDelegates, bool bEditWithinLevelEditor );
+	void InitSequencerAssetEditor( const EToolkitMode::Type Mode, const FSequencerViewParams& InViewParams, const TSharedPtr< class IToolkitHost >& InitToolkitHost, UMovieScene* InRootMovieScene, const TArray<FOnCreateTrackEditor>& TrackEditorDelegates, bool bEditWithinLevelEditor );
 
 	TSharedRef<ISequencer> GetSequencerInterface() const override { return Sequencer.ToSharedRef(); }
 
diff --git a/Engine/Source/Editor/Sequencer/Private/SequencerModule.cpp b/Engine/Source/Editor/Sequencer/Private/SequencerModule.cpp
--- a/Engine/Source/Editor/Sequencer/Private/SequencerModule.cpp
+++ b/Engine/Source/Editor/Sequencer/Private/SequencerModule.cpp
@@ -15,11 +15,12 @@ class FSequencerModule : public ISequencerModule
 {
 
 	/** ISequencerModule interface */
-	virtual TSharedPtr<ISequencer> CreateSequencer( UMovieScene* RootMovieScene, TSharedRef<ISequencerObjectBindingManager> ObjectBindingManager ) override
+	virtual TSharedPtr<ISequencer> CreateSequencer( UMovieScene* RootMovieScene, const FSequencerViewParams& InViewParams, TSharedRef<ISequencerObjectBindingManager> ObjectBindingManager ) override
 	{
 		TSharedRef< FSequencer > Sequencer = MakeShareable(new FSequencer);
 		
 		FSequencerInitParams SequencerInitParams;
+		SequencerInitParams.ViewParams = InViewParams;
 		SequencerInitParams.ObjectChangeListener = MakeShareable( new FSequencerObjectChangeListener( Sequencer, false ) );
 		
 		SequencerInitParams.ObjectBindingManager = ObjectBindingManager;
@@ -35,10 +36,11 @@ class FSequencerModule : public ISequencerModule
 		return Sequencer;
 	}
 	
-	virtual TSharedPtr<ISequencer> CreateSequencerAssetEditor( const EToolkitMode::Type Mode, const TSharedPtr< class IToolkitHost >& InitToolkitHost, UMovieScene* InRootMovieScene, bool bEditWithinLevelEditor ) override
+	virtual TSharedPtr<ISequencer> CreateSequencerAssetEditor( const EToolkitMode::Type Mode, const FSequencerViewParams& InViewParams, const TSharedPtr< class IToolkitHost >& InitToolkitHost, UMovieScene* InRootMovieScene, bool bEditWithinLevelEditor ) override
 	{
 		TSharedRef< FSequencerAssetEditor > SequencerAssetEditor = MakeShareable(new FSequencerAssetEditor);
-		SequencerAssetEditor->InitSequencerAssetEditor( Mode, InitToolkitHost, InRootMovieScene, TrackEditorDelegates, bEditWithinLevelEditor );
+
+		SequencerAssetEditor->InitSequencerAssetEditor( Mode, InViewParams, InitToolkitHost, InRootMovieScene, TrackEditorDelegates, bEditWithinLevelEditor );
 		return SequencerAssetEditor->GetSequencerInterface();
 	}
 
diff --git a/Engine/Source/Editor/Sequencer/Private/SequencerObjectChangeListener.cpp b/Engine/Source/Editor/Sequencer/Private/SequencerObjectChangeListener.cpp
--- a/Engine/Source/Editor/Sequencer/Private/SequencerObjectChangeListener.cpp
+++ b/Engine/Source/Editor/Sequencer/Private/SequencerObjectChangeListener.cpp
@@ -63,15 +63,16 @@ FOnPropagateObjectChanges& FSequencerObjectChangeListener::GetOnPropagateObjectC
 }
 
 
-void FSequencerObjectChangeListener::OnObjectPreEditChange( UObject* Object )
+void FSequencerObjectChangeListener::OnObjectPreEditChange( UObject* Object, const FEditPropertyChain& PropertyChain )
 {
 	// We only care if we are in auto-key mode and not attempting to change properties of a CDO (which cannot be animated)
-	if( Sequencer.IsValid() && Sequencer.Pin()->IsAutoKeyEnabled() && !Object->HasAnyFlags(RF_ClassDefaultObject) )
+	if( Sequencer.IsValid() && Sequencer.Pin()->IsAutoKeyEnabled() && !Object->HasAnyFlags(RF_ClassDefaultObject) && PropertyChain.GetActiveMemberNode() )
 	{
 		// Register with the property editor module that we'd like to know about animated float properties that change
 		FPropertyEditorModule& PropertyEditor = FModuleManager::Get().LoadModuleChecked<FPropertyEditorModule>( "PropertyEditor" );
 
-		if( IsObjectValidForListening(Object) )
+		// Sometimes due to edit inline new the object is not actually the object that contains the property
+		if( IsObjectValidForListening(Object) && PropertyChain.GetActiveMemberNode()->GetValue()->IsIn( Object->GetClass() ) )
 		{
 			TSharedPtr<IPropertyChangeListener> PropertyChangeListener = ActivePropertyChangeListeners.FindRef( Object );
 
diff --git a/Engine/Source/Editor/Sequencer/Private/SequencerObjectChangeListener.h b/Engine/Source/Editor/Sequencer/Private/SequencerObjectChangeListener.h
--- a/Engine/Source/Editor/Sequencer/Private/SequencerObjectChangeListener.h
+++ b/Engine/Source/Editor/Sequencer/Private/SequencerObjectChangeListener.h
@@ -23,8 +23,9 @@ private:
 	 * Called when PreEditChange is called on an object
 	 *
 	 * @param Object	The object that PreEditChange was called on
+	 * @param Property	The property that is about to change
 	 */
-	void OnObjectPreEditChange( UObject* Object );
+	void OnObjectPreEditChange( UObject* Object, const FEditPropertyChain& PropertyChain );
 
 	/**
 	 * Called when PostEditChange is called on an object
diff --git a/Engine/Source/Editor/Sequencer/Public/ISequencerModule.h b/Engine/Source/Editor/Sequencer/Public/ISequencerModule.h
--- a/Engine/Source/Editor/Sequencer/Public/ISequencerModule.h
+++ b/Engine/Source/Editor/Sequencer/Public/ISequencerModule.h
@@ -15,6 +15,20 @@ class ISequencerObjectBindingManager;
  */
 DECLARE_DELEGATE_RetVal_OneParam( TSharedRef<class FMovieSceneTrackEditor>, FOnCreateTrackEditor, TSharedRef<ISequencer> );
 
+/** Parameters for initializing a Sequencer */
+struct FSequencerViewParams
+{
+	/** The initial time range of the view */
+	TRange<float> InitalViewRange;
+
+	/** Initial Scrub Position */
+	float InitialScrubPosition;
+
+	FSequencerViewParams()
+		: InitalViewRange( 0.0f, 5.0f )
+		, InitialScrubPosition( 0.0f )
+	{}
+};
 
 /**
  * The public interface of SequencerModule
@@ -28,20 +42,22 @@ public:
 	 * Creates a new instance of a standalone sequencer that can be added to other UIs
 	 *
 	 * @param 	InRootMovieScene	The movie scene to edit
+	 * @param	ViewParams			Parameters for how to view sequencer UI
 	 * @return	Interface to the new editor
 	 */
-	virtual TSharedPtr<ISequencer> CreateSequencer( UMovieScene* InRootMovieScene, TSharedRef<ISequencerObjectBindingManager> ObjectBindingManager ) = 0;
+	virtual TSharedPtr<ISequencer> CreateSequencer( UMovieScene* InRootMovieScene, const FSequencerViewParams& InViewParams, TSharedRef<ISequencerObjectBindingManager> ObjectBindingManager ) = 0;
 
 	/**
 	 * Creates a new instance of a Sequencer, the editor for MovieScene assets in an asset editor
 	 *
 	 * @param	Mode					Mode that this editor should operate in
+	 * @param	ViewParams				Parameters for how to view sequencer UI
 	 * @param	InitToolkitHost			When Mode is WorldCentric, this is the level editor instance to spawn this editor within
 	 * @param	ObjectToEdit			The object to start editing
 	 *
 	 * @return	Interface to the new editor
 	 */
-	virtual TSharedPtr<ISequencer> CreateSequencerAssetEditor( const EToolkitMode::Type Mode, const TSharedPtr< class IToolkitHost >& InitToolkitHost, UMovieScene* InRootMovieScene, bool bEditWithinLevelEditor ) = 0;
+	virtual TSharedPtr<ISequencer> CreateSequencerAssetEditor( const EToolkitMode::Type Mode, const FSequencerViewParams& InViewParams, const TSharedPtr< class IToolkitHost >& InitToolkitHost, UMovieScene* InRootMovieScene, bool bEditWithinLevelEditor ) = 0;
 
 	/** 
 	 * Registers a delegate that will create an editor for a track in each sequencer 
diff --git a/Engine/Source/Editor/UMGEditor/Private/Extensions/CanvasSlotExtension.cpp b/Engine/Source/Editor/UMGEditor/Private/Extensions/CanvasSlotExtension.cpp
--- a/Engine/Source/Editor/UMGEditor/Private/Extensions/CanvasSlotExtension.cpp
+++ b/Engine/Source/Editor/UMGEditor/Private/Extensions/CanvasSlotExtension.cpp
@@ -261,13 +261,13 @@ void FCanvasSlotExtension::MoveByAmount(FWidgetReference& WidgetRef, FVector2D D
 		Offsets.Bottom -= Delta.Y;
 	}
 
-	CanvasSlot->SetOffset(Offsets);
+	CanvasSlot->SetOffsets(Offsets);
 
 	// Update the Template widget to match
 	UWidget* TemplateWidget = WidgetRef.GetTemplate();
 	UCanvasPanelSlot* TemplateSlot = CastChecked<UCanvasPanelSlot>(TemplateWidget->Slot);
 
-	TemplateSlot->SetOffset(Offsets);
+	TemplateSlot->SetOffsets(Offsets);
 }
 
 #undef LOCTEXT_NAMESPACE
diff --git a/Engine/Source/Editor/UMGEditor/Private/UMGSequencerObjectBindingManager.cpp b/Engine/Source/Editor/UMGEditor/Private/UMGSequencerObjectBindingManager.cpp
--- a/Engine/Source/Editor/UMGEditor/Private/UMGSequencerObjectBindingManager.cpp
+++ b/Engine/Source/Editor/UMGEditor/Private/UMGSequencerObjectBindingManager.cpp
@@ -42,6 +42,12 @@ void FUMGSequencerObjectBindingManager::BindPossessableObject( const FGuid& Poss
 
 	FWidgetAnimationBinding NewBinding;
 	NewBinding.WidgetName = PossessedObject.GetFName();
+	if( PossessedObject.IsA<UPanelSlot>() )
+	{
+		// Save the name of the widget containing the slots.  This is the object to look up that contaisn the slot itself(the thing we are animating)
+		NewBinding.SlotWidgetName = PossessedObject.GetOuter()->GetFName();
+	}
+
 	NewBinding.AnimationGuid = PossessableGuid;
 
 	WidgetAnimation->AnimationBindings.Add( NewBinding );
@@ -69,7 +75,7 @@ void FUMGSequencerObjectBindingManager::GetRuntimeObjects( const TSharedRef<FMov
 	for( auto It = PreviewObjectToGuidMap.CreateConstIterator(); It; ++It )
 	{
 		UObject* Object = It.Key().Get();
-		if( Object )
+		if( Object && It.Value() == ObjectGuid )
 		{
 			OutRuntimeObjects.Add( Object );
 		}
@@ -91,10 +97,26 @@ void FUMGSequencerObjectBindingManager::InitPreviewObjects()
 	for( const FWidgetAnimationBinding& Binding : WidgetAnimation->AnimationBindings )
 	{
 		FName ObjectName = Binding.WidgetName;
-		UObject* FoundObject = FindObject<UObject>(WidgetTree, *ObjectName.ToString());
+		FName SlotWidgetName = Binding.SlotWidgetName;
+
+		FName NameToSearchFor = SlotWidgetName != NAME_None ? SlotWidgetName : ObjectName;
+
+		UObject* FoundObject = FindObject<UObject>(WidgetTree, *NameToSearchFor.ToString());
 		if(FoundObject)
 		{
-			PreviewObjectToGuidMap.Add(FoundObject, Binding.AnimationGuid);
+			if( SlotWidgetName == NAME_None )
+			{
+				PreviewObjectToGuidMap.Add(FoundObject, Binding.AnimationGuid);
+			}
+			else
+			{
+				FoundObject = FindObject<UObject>( FoundObject, *ObjectName.ToString() );
+				if( FoundObject )
+				{
+					PreviewObjectToGuidMap.Add(FoundObject, Binding.AnimationGuid);
+				}
+			}
+		
 		}
 	}
 }
diff --git a/Engine/Source/Editor/UMGEditor/Private/WidgetBlueprintEditor.cpp b/Engine/Source/Editor/UMGEditor/Private/WidgetBlueprintEditor.cpp
--- a/Engine/Source/Editor/UMGEditor/Private/WidgetBlueprintEditor.cpp
+++ b/Engine/Source/Editor/UMGEditor/Private/WidgetBlueprintEditor.cpp
@@ -402,7 +402,11 @@ TSharedPtr<ISequencer>& FWidgetBlueprintEditor::GetSequencer()
 
 		SequencerObjectBindingManager = ObjectBindingManager;
 
-		Sequencer = FModuleManager::LoadModuleChecked< ISequencerModule >("Sequencer").CreateSequencer( WidgetAnimation->MovieScene, ObjectBindingManager );
+		FSequencerViewParams ViewParams;
+		ViewParams.InitalViewRange = TRange<float>(-0.02f, 3.2f);
+		ViewParams.InitialScrubPosition = 0;
+
+		Sequencer = FModuleManager::LoadModuleChecked< ISequencerModule >("Sequencer").CreateSequencer( WidgetAnimation->MovieScene, ViewParams, ObjectBindingManager );
 	}
 
 	return Sequencer;
diff --git a/Engine/Source/Runtime/Core/Public/Misc/CallbackDevice.h b/Engine/Source/Runtime/Core/Public/Misc/CallbackDevice.h
--- a/Engine/Source/Runtime/Core/Public/Misc/CallbackDevice.h
+++ b/Engine/Source/Runtime/Core/Public/Misc/CallbackDevice.h
@@ -51,7 +51,7 @@ public:
 	DECLARE_MULTICAST_DELEGATE_OneParam(FOnActorLabelChanged, AActor*);
 
 	// Callback for PreEditChange
-	DECLARE_MULTICAST_DELEGATE_OneParam(FOnPreObjectPropertyChanged, UObject*);
+	DECLARE_MULTICAST_DELEGATE_TwoParams(FOnPreObjectPropertyChanged, UObject*, const class FEditPropertyChain& );
 
 	#if WITH_EDITOR
 	// Callback for all object modifications
diff --git a/Engine/Source/Runtime/CoreUObject/Private/UObject/Obj.cpp b/Engine/Source/Runtime/CoreUObject/Private/UObject/Obj.cpp
--- a/Engine/Source/Runtime/CoreUObject/Private/UObject/Obj.cpp
+++ b/Engine/Source/Runtime/CoreUObject/Private/UObject/Obj.cpp
@@ -194,7 +194,6 @@ void UObject::PostLoad()
 void UObject::PreEditChange(UProperty* PropertyAboutToChange)
 {
 	Modify();
-	FCoreDelegates::OnPreObjectPropertyChanged.Broadcast(this);
 }
 
 
@@ -216,6 +215,8 @@ void UObject::PreEditChange( FEditPropertyChain& PropertyAboutToChange )
 	// forward the notification to the UProperty* version of PreEditChange
 	PreEditChange(PropertyAboutToChange.GetActiveNode()->GetValue());
 
+	FCoreDelegates::OnPreObjectPropertyChanged.Broadcast(this, PropertyAboutToChange);
+
 	if ( HasAnyFlags(RF_ClassDefaultObject|RF_ArchetypeObject) && PropertyAboutToChange.GetActiveMemberNode() == PropertyAboutToChange.GetHead() && !FApp::IsGame())
 	{
 		// this object must now be included in the undo/redo buffer
diff --git a/Engine/Source/Runtime/UMG/Private/Animation/MovieSceneMarginSection.cpp b/Engine/Source/Runtime/UMG/Private/Animation/MovieSceneMarginSection.cpp
--- a/Engine/Source/Runtime/UMG/Private/Animation/MovieSceneMarginSection.cpp
+++ b/Engine/Source/Runtime/UMG/Private/Animation/MovieSceneMarginSection.cpp
@@ -14,16 +14,16 @@ void UMovieSceneMarginSection::MoveSection( float DeltaTime )
 	Super::MoveSection( DeltaTime );
 
 	// Move all the curves in this section
-	TopCurve.ShiftCurve(DeltaTime);
 	LeftCurve.ShiftCurve(DeltaTime);
+	TopCurve.ShiftCurve(DeltaTime);
 	RightCurve.ShiftCurve(DeltaTime);
 	BottomCurve.ShiftCurve(DeltaTime);
 }
 
 void UMovieSceneMarginSection::DilateSection( float DilationFactor, float Origin )
 {
-	TopCurve.ScaleCurve(Origin, DilationFactor);
 	LeftCurve.ScaleCurve(Origin, DilationFactor);
+	TopCurve.ScaleCurve(Origin, DilationFactor);
 	RightCurve.ScaleCurve(Origin, DilationFactor);
 	BottomCurve.ScaleCurve(Origin, DilationFactor);
 
@@ -32,16 +32,16 @@ void UMovieSceneMarginSection::DilateSection( float DilationFactor, float Origin
 
 FMargin UMovieSceneMarginSection::Eval( float Position ) const
 {
-	return FMargin(	TopCurve.Eval(Position),
-					LeftCurve.Eval(Position),
+	return FMargin(	LeftCurve.Eval(Position),
+					TopCurve.Eval(Position),
 					RightCurve.Eval(Position),
 					BottomCurve.Eval(Position));
 }
 
 void UMovieSceneMarginSection::AddKey( float Time, const FMargin& Value )
 {
-	AddKeyToCurve(TopCurve, Time, Value.Top);
 	AddKeyToCurve(LeftCurve, Time, Value.Left);
+	AddKeyToCurve(TopCurve, Time, Value.Top);
 	AddKeyToCurve(RightCurve, Time, Value.Right);
 	AddKeyToCurve(BottomCurve, Time, Value.Bottom);
 }
diff --git a/Engine/Source/Runtime/UMG/Private/Animation/UMGSequencePlayer.cpp b/Engine/Source/Runtime/UMG/Private/Animation/UMGSequencePlayer.cpp
--- a/Engine/Source/Runtime/UMG/Private/Animation/UMGSequencePlayer.cpp
+++ b/Engine/Source/Runtime/UMG/Private/Animation/UMGSequencePlayer.cpp
@@ -32,12 +32,27 @@ void UUMGSequencePlayer::InitSequencePlayer( const FWidgetAnimation& Animation,
 	for (const FWidgetAnimationBinding& Binding : Animation.AnimationBindings)
 	{
 		FName ObjectName = Binding.WidgetName;
-		UObject* FoundObject = FindObject<UObject>(WidgetTree, *ObjectName.ToString());
-		if (FoundObject)
+		FName SlotWidgetName = Binding.SlotWidgetName;
+
+		FName NameToSearchFor = SlotWidgetName != NAME_None ? SlotWidgetName : ObjectName;
+
+		UObject* FoundObject = FindObject<UObject>(WidgetTree, *NameToSearchFor.ToString());
+		if(FoundObject)
 		{
-			TArray<UObject*>& Objects = GuidToRuntimeObjectMap.FindOrAdd( Binding.AnimationGuid );
-			Objects.Add( FoundObject );
-		
+			if(SlotWidgetName == NAME_None)
+			{
+				TArray<UObject*>& Objects = GuidToRuntimeObjectMap.FindOrAdd(Binding.AnimationGuid);
+				Objects.Add(FoundObject);
+			}
+			else
+			{
+				FoundObject = FindObject<UObject>(FoundObject, *ObjectName.ToString());
+				if(FoundObject)
+				{
+					TArray<UObject*>& Objects = GuidToRuntimeObjectMap.FindOrAdd(Binding.AnimationGuid);
+					Objects.Add(FoundObject);
+				}
+			}
 		}
 	}
 
@@ -63,7 +78,7 @@ void UUMGSequencePlayer::Tick(float DeltaTime)
 
 		bool bIsFinished = false;
 
-		if (!TimeRange.Contains(TimeCursorPosition))
+		if ( TimeCursorPosition >= TimeRange.GetUpperBoundValue() )
 		{
 			bIsFinished = true;
 		}
diff --git a/Engine/Source/Runtime/UMG/Private/Components/CanvasPanelSlot.cpp b/Engine/Source/Runtime/UMG/Private/Components/CanvasPanelSlot.cpp
--- a/Engine/Source/Runtime/UMG/Private/Components/CanvasPanelSlot.cpp
+++ b/Engine/Source/Runtime/UMG/Private/Components/CanvasPanelSlot.cpp
@@ -34,12 +34,12 @@ void UCanvasPanelSlot::BuildSlot(TSharedRef<SConstraintCanvas> Canvas)
 
 void UCanvasPanelSlot::SetDesiredPosition(FVector2D InPosition)
 {
-	SetOffset(FMargin(InPosition.X, InPosition.Y, LayoutData.Offsets.Right, LayoutData.Offsets.Bottom));
+	SetOffsets(FMargin(InPosition.X, InPosition.Y, LayoutData.Offsets.Right, LayoutData.Offsets.Bottom));
 }
 
 void UCanvasPanelSlot::SetDesiredSize(FVector2D InSize)
 {
-	SetOffset(FMargin(LayoutData.Offsets.Left, LayoutData.Offsets.Top, InSize.X, InSize.Y));
+	SetOffsets(FMargin(LayoutData.Offsets.Left, LayoutData.Offsets.Top, InSize.X, InSize.Y));
 }
 
 void UCanvasPanelSlot::Resize(const FVector2D& Direction, const FVector2D& Amount)
@@ -103,7 +103,7 @@ bool UCanvasPanelSlot::CanResize(const FVector2D& Direction) const
 	return true;
 }
 
-void UCanvasPanelSlot::SetOffset(FMargin InOffset)
+void UCanvasPanelSlot::SetOffsets(FMargin InOffset)
 {
 	LayoutData.Offsets = InOffset;
 	if ( Slot )
@@ -141,7 +141,7 @@ void UCanvasPanelSlot::SetZOrder(int32 InZOrder)
 
 void UCanvasPanelSlot::SyncronizeProperties()
 {
-	SetOffset(LayoutData.Offsets);
+	SetOffsets(LayoutData.Offsets);
 	SetAnchors(LayoutData.Anchors);
 	SetAlignment(LayoutData.Alignment);
 	SetZOrder(ZOrder);
diff --git a/Engine/Source/Runtime/UMG/Public/Blueprint/WidgetBlueprintGeneratedClass.h b/Engine/Source/Runtime/UMG/Public/Blueprint/WidgetBlueprintGeneratedClass.h
--- a/Engine/Source/Runtime/UMG/Public/Blueprint/WidgetBlueprintGeneratedClass.h
+++ b/Engine/Source/Runtime/UMG/Public/Blueprint/WidgetBlueprintGeneratedClass.h
@@ -49,12 +49,15 @@ struct FWidgetAnimationBinding
 	UPROPERTY()
 	FName WidgetName;
 
+	UPROPERTY()
+	FName SlotWidgetName;
+
 	UPROPERTY()
 	FGuid AnimationGuid;
 
 	bool operator==( const FWidgetAnimationBinding& Other ) const
 	{
-		return WidgetName == Other.WidgetName && AnimationGuid == Other.AnimationGuid;
+		return WidgetName == Other.WidgetName && SlotWidgetName == Other.SlotWidgetName && AnimationGuid == Other.AnimationGuid;
 	}
 };
 
diff --git a/Engine/Source/Runtime/UMG/Public/Components/CanvasPanelSlot.h b/Engine/Source/Runtime/UMG/Public/Components/CanvasPanelSlot.h
--- a/Engine/Source/Runtime/UMG/Public/Components/CanvasPanelSlot.h
+++ b/Engine/Source/Runtime/UMG/Public/Components/CanvasPanelSlot.h
@@ -50,7 +50,7 @@ class UMG_API UCanvasPanelSlot : public UPanelSlot
 
 	/** Sets the position of the slot */
 	UFUNCTION(BlueprintCallable, Category="Appearance")
-	void SetOffset(FMargin InOffset);
+	void SetOffsets(FMargin InOffset);
 	
 	/** Sets the anchors on the slot */
 	UFUNCTION(BlueprintCallable, Category="Appearance")
