Fixed squashed expander arrows in sequencer

diff --git a/Engine/Source/Editor/Sequencer/Private/SAnimationOutlinerView.cpp b/Engine/Source/Editor/Sequencer/Private/SAnimationOutlinerView.cpp
--- a/Engine/Source/Editor/Sequencer/Private/SAnimationOutlinerView.cpp
+++ b/Engine/Source/Editor/Sequencer/Private/SAnimationOutlinerView.cpp
@@ -38,6 +38,7 @@ void SAnimationOutlinerTreeNode::Construct( const FArguments& InArgs, TSharedRef
 		.VAlign( VAlign_Center )
 		[
 			SNew(SButton)
+			.ContentPadding(0)
 			.Visibility( this, &SAnimationOutlinerTreeNode::GetExpanderVisibility )
 			.OnClicked( this, &SAnimationOutlinerTreeNode::OnExpanderClicked )
 			.ClickMethod( EButtonClickMethod::MouseDown )
