Make sequencer use a checkbox for auto-key until proper artwork is done

diff --git a/Engine/Source/Editor/Sequencer/Private/SSequencer.cpp b/Engine/Source/Editor/Sequencer/Private/SSequencer.cpp
--- a/Engine/Source/Editor/Sequencer/Private/SSequencer.cpp
+++ b/Engine/Source/Editor/Sequencer/Private/SSequencer.cpp
@@ -190,7 +190,6 @@ void SSequencer::Construct( const FArguments& InArgs, TSharedRef< class FSequenc
 					[
 						// @todo Sequencer - Temp auto-key button
 						SNew( SCheckBox )
-						.Style(FEditorStyle::Get(), "ToggleButtonCheckbox")
 						.IsChecked( this, &SSequencer::OnGetAutoKeyCheckState )
 						.OnCheckStateChanged( this, &SSequencer::OnAutoKeyChecked )
 						[
