Sequencer - Improved the drawing of the section area borders

diff --git a/Engine/Source/Editor/EditorStyle/Private/SlateEditorStyle.cpp b/Engine/Source/Editor/EditorStyle/Private/SlateEditorStyle.cpp
--- a/Engine/Source/Editor/EditorStyle/Private/SlateEditorStyle.cpp
+++ b/Engine/Source/Editor/EditorStyle/Private/SlateEditorStyle.cpp
@@ -1400,6 +1400,8 @@ void FSlateEditorStyle::FStyle::SetupGeneralStyles()
 		Set( "Sequencer.NotificationImage_AddedPlayMovieSceneEvent", new IMAGE_BRUSH( "Old/Checkbox_checked", Icon16x16 ) );
 		Set( "Sequencer.Key", new IMAGE_BRUSH( "Sequencer/Key", FVector2D(7.0f,12.0f) ) );
 		Set( "Sequencer.GenericSection.Background", new BOX_BRUSH( "Sequencer/GenericSectionBackground", FMargin(5.f/16.f) ) );
+		Set( "Sequencer.SectionGripLeft", new BOX_BRUSH( "Sequencer/SectionGripLeft", FMargin(5.f/16.f) ) );
+		Set( "Sequencer.SectionGripRight", new BOX_BRUSH( "Sequencer/SectionGripRight", FMargin(5.f/16.f) ) );
 		Set( "Sequencer.SectionArea.Background", new FSlateColorBrush( FColor::White ) );
 		Set( "Sequencer.AnimationOutliner.SelectionBorder", new BOX_BRUSH( "Common/RoundedSelection_16x", 4.0f/16.0f, SelectionColor_Pressed ) );
 		Set( "Sequencer.AnimationOutliner.BoldFont", TTF_CORE_FONT( "Fonts/Roboto-Bold", 11 ) );
diff --git a/Engine/Source/Editor/Sequencer/Private/SSection.cpp b/Engine/Source/Editor/Sequencer/Private/SSection.cpp
--- a/Engine/Source/Editor/Sequencer/Private/SSection.cpp
+++ b/Engine/Source/Editor/Sequencer/Private/SSection.cpp
@@ -10,16 +10,7 @@
 #include "MovieSceneShotSection.h"
 #include "CommonMovieSceneTools.h"
 
-namespace SequencerSectionConstants
-{
-	/** How far the user has to drag the mouse before we consider the action dragging rather than a click */
-	const float SectionDragStartDistance = 5.0f;
-
-	/** The size of each key */
-	const FVector2D KeySize( 11.0f, 11.0f );
 
-	const float SectionEdgeSize = 15.0f;
-}
 
 void SSection::Construct( const FArguments& InArgs, TSharedRef<FTrackNode> SectionNode, int32 InSectionIndex )
 {
@@ -142,11 +133,11 @@ void SSection::CheckForEdgeInteraction( const FPointerEvent& MouseEvent, const F
 	// Make areas to the left and right of the geometry.  We will use these areas to determine if someone dragged the left or right edge of a section
 	FGeometry SectionRectLeft = SectionGeometry.MakeChild(
 		FVector2D::ZeroVector,
-		FVector2D( SequencerSectionConstants::SectionEdgeSize, SectionGeometry.Size.Y )
+		FVector2D( SequencerSectionConstants::SectionGripSize, SectionGeometry.Size.Y )
 		);
 
 	FGeometry SectionRectRight = SectionGeometry.MakeChild(
-		FVector2D( SectionGeometry.Size.X - SequencerSectionConstants::SectionEdgeSize, 0 ), 
+		FVector2D( SectionGeometry.Size.X - SequencerSectionConstants::SectionGripSize, 0 ), 
 		SectionGeometry.Size 
 		);
 
@@ -214,21 +205,19 @@ int32 SSection::OnPaint( const FPaintArgs& Args, const FGeometry& AllottedGeomet
 {
 	int32 StartLayer = SCompoundWidget::OnPaint( Args, AllottedGeometry, MyClippingRect, OutDrawElements, LayerId, InWidgetStyle, bParentEnabled );
 
-	UMovieSceneSection* SectionObject = SectionInterface->GetSectionObject();
+	DrawSectionBorders( AllottedGeometry, MyClippingRect, OutDrawElements, StartLayer );
+
+	FGeometry SectionGeometry = AllottedGeometry.MakeChild( FVector2D( SequencerSectionConstants::SectionGripSize, 0 ), AllottedGeometry.GetDrawSize() - FVector2D( SequencerSectionConstants::SectionGripSize*2, 0.0f ) );
 
-	const bool bSelected = ParentSectionArea->GetSequencer().IsSectionSelected( SectionObject );
 
 	// Ask the interface to draw the section
-	int32 PostSectionLayer = SectionInterface->OnPaintSection( AllottedGeometry, MyClippingRect, OutDrawElements, LayerId, bParentEnabled );
+	int32 PostSectionLayer = SectionInterface->OnPaintSection( SectionGeometry, MyClippingRect, OutDrawElements, LayerId, bParentEnabled );
 
 	// @todo Sequencer - Temp indicators of the area of a section which can be dragged.
 
-	PaintKeys( AllottedGeometry, MyClippingRect, OutDrawElements, PostSectionLayer, InWidgetStyle );
+	PaintKeys( SectionGeometry, MyClippingRect, OutDrawElements, PostSectionLayer, InWidgetStyle );
 
-	FLinearColor SelectionColor = FEditorStyle::GetSlateColor("SelectionColor").GetColor( InWidgetStyle );
-	FLinearColor TransparentSelectionColor = SelectionColor;
-	TransparentSelectionColor.A = .5f;
-		
+	
 	// Section name with drop shadow
 	FText SectionTitle = SectionInterface->GetSectionTitle();
 	if (!SectionTitle.IsEmpty())
@@ -236,17 +225,18 @@ int32 SSection::OnPaint( const FPaintArgs& Args, const FGeometry& AllottedGeomet
 		FSlateDrawElement::MakeText(
 			OutDrawElements,
 			PostSectionLayer+1,
-			AllottedGeometry.ToOffsetPaintGeometry(FVector2D(6, 6)),
+			SectionGeometry.ToOffsetPaintGeometry(FVector2D(6, 6)),
 			SectionTitle,
 			FEditorStyle::GetFontStyle("NormalFont"),
 			MyClippingRect,
 			ESlateDrawEffect::None,
 			FLinearColor::Black
 			);
+
 		FSlateDrawElement::MakeText(
 			OutDrawElements,
 			PostSectionLayer+2,
-			AllottedGeometry.ToOffsetPaintGeometry(FVector2D(5, 5)),
+			SectionGeometry.ToOffsetPaintGeometry(FVector2D(5, 5)),
 			SectionTitle,
 			FEditorStyle::GetFontStyle("NormalFont"),
 			MyClippingRect,
@@ -255,48 +245,8 @@ int32 SSection::OnPaint( const FPaintArgs& Args, const FGeometry& AllottedGeomet
 			);
 	}
 
-	// draw selection box
-	if (bSelected)
-	{
-		FSlateDrawElement::MakeBox( 
-			OutDrawElements,
-			PostSectionLayer+3,
-			AllottedGeometry.ToPaintGeometry(),
-			FEditorStyle::GetBrush("PlainBorder"),
-			MyClippingRect,
-			ESlateDrawEffect::None,
-			SelectionColor
-			);
-	}
 	
 
-	if( bLeftEdgePressed || bLeftEdgeHovered )
-	{
-		FSlateDrawElement::MakeBox( 
-			OutDrawElements, 
-			PostSectionLayer+3,
-			AllottedGeometry.ToPaintGeometry( FVector2D::ZeroVector, FVector2D( 5.0f, AllottedGeometry.Size.Y ) ),
-			FEditorStyle::GetBrush( "WhiteBrush" ),
-			MyClippingRect,
-			ESlateDrawEffect::None,
-			TransparentSelectionColor
-			);
-
-	}
-
-	if( bRightEdgePressed || bRightEdgeHovered )
-	{
-		FSlateDrawElement::MakeBox( 
-			OutDrawElements, 
-			PostSectionLayer+3,
-			AllottedGeometry.ToPaintGeometry( FVector2D( AllottedGeometry.Size.X - 5.0f, 0 ), AllottedGeometry.Size ),
-			FEditorStyle::GetBrush( "WhiteBrush" ),
-			MyClippingRect,
-			ESlateDrawEffect::None,
-			TransparentSelectionColor
-			);
-
-	}
 
 	return LayerId;
 }
@@ -390,6 +340,56 @@ void SSection::PaintKeys( const FGeometry& AllottedGeometry, const FSlateRect& M
 	}
 }
 
+void SSection::DrawSectionBorders( const FGeometry& AllottedGeometry, const FSlateRect& MyClippingRect, FSlateWindowElementList& OutDrawElements, int32 LayerId ) const
+{
+	UMovieSceneSection* SectionObject = SectionInterface->GetSectionObject();
+
+	const bool bSelected = ParentSectionArea->GetSequencer().IsSectionSelected(SectionObject);
+
+	FLinearColor SelectionColor = FEditorStyle::GetSlateColor("SelectionColor").GetColor(FWidgetStyle());
+	FLinearColor TransparentSelectionColor = SelectionColor;
+
+	// Left Grip
+	FSlateDrawElement::MakeBox(
+		OutDrawElements,
+		LayerId,
+		// Center the key along Y.  Ensure the middle of the key is at the actual key time
+		AllottedGeometry.ToPaintGeometry( FVector2D( 0.0f, 0.0f ), FVector2D( SequencerSectionConstants::SectionGripSize, AllottedGeometry.GetDrawSize().Y) ) ,
+		FEditorStyle::GetBrush("Sequencer.SectionGripLeft"),
+		MyClippingRect,
+		ESlateDrawEffect::None,
+		(bLeftEdgePressed || bLeftEdgeHovered) ? TransparentSelectionColor : FLinearColor::White
+	);
+
+	// Right Grip
+	FSlateDrawElement::MakeBox(
+		OutDrawElements,
+		LayerId,
+		// Center the key along Y.  Ensure the middle of the key is at the actual key time
+		AllottedGeometry.ToPaintGeometry( FVector2D( AllottedGeometry.Size.X-SequencerSectionConstants::SectionGripSize, 0.0f), FVector2D(SequencerSectionConstants::SectionGripSize, AllottedGeometry.GetDrawSize().Y)),
+		FEditorStyle::GetBrush("Sequencer.SectionGripRight"),
+		MyClippingRect,
+		ESlateDrawEffect::None,
+		(bRightEdgePressed || bRightEdgeHovered) ? TransparentSelectionColor : FLinearColor::White
+		);
+
+
+	// draw selection box
+	if(bSelected)
+	{
+		FSlateDrawElement::MakeBox(
+			OutDrawElements,
+			LayerId,
+			AllottedGeometry.ToPaintGeometry(),
+			FEditorStyle::GetBrush("PlainBorder"),
+			MyClippingRect,
+			ESlateDrawEffect::None,
+			SelectionColor
+			);
+	}
+
+}
+
 TSharedPtr<SWidget> SSection::OnSummonContextMenu( const FGeometry& MyGeometry, const FPointerEvent& MouseEvent )
 {
 	FSelectedKey Key = GetKeyUnderMouse( MouseEvent.GetScreenSpacePosition(), MyGeometry );
diff --git a/Engine/Source/Editor/Sequencer/Private/SSection.h b/Engine/Source/Editor/Sequencer/Private/SSection.h
--- a/Engine/Source/Editor/Sequencer/Private/SSection.h
+++ b/Engine/Source/Editor/Sequencer/Private/SSection.h
@@ -2,6 +2,17 @@
 
 #pragma once
 
+namespace SequencerSectionConstants
+{
+	/** How far the user has to drag the mouse before we consider the action dragging rather than a click */
+	const float SectionDragStartDistance = 5.0f;
+
+	/** The size of each key */
+	const FVector2D KeySize(11.0f, 11.0f);
+
+	const float SectionGripSize = 7.0f;
+}
+
 class SSection : public SCompoundWidget
 {
 public:
@@ -92,6 +103,11 @@ private:
 	 */
 	void PaintKeys( const FGeometry& AllottedGeometry, const FSlateRect& MyClippingRect, FSlateWindowElementList& OutDrawElements, int32 LayerId, const FWidgetStyle& InWidgetStyle ) const;
 
+	/**
+	 * Draw the borders around each section.  The borders are the areas which can be used to expand or collapse the size of the section.
+	 */
+	void DrawSectionBorders( const FGeometry& AllottedGeometry, const FSlateRect& MyClippingRect, FSlateWindowElementList& OutDrawElements, int32 LayerId ) const;
+
 	/** Summons a context menu over the associated section */
 	TSharedPtr<SWidget> OnSummonContextMenu( const FGeometry& MyGeometry, const FPointerEvent& MouseEvent );
 
diff --git a/Engine/Source/Editor/Sequencer/Private/SSequencerSectionAreaView.cpp b/Engine/Source/Editor/Sequencer/Private/SSequencerSectionAreaView.cpp
--- a/Engine/Source/Editor/Sequencer/Private/SSequencerSectionAreaView.cpp
+++ b/Engine/Source/Editor/Sequencer/Private/SSequencerSectionAreaView.cpp
@@ -35,19 +35,16 @@ namespace SequencerSectionUtils
 		float StartX =  TimeToPixelConverter.TimeToPixel( Section->GetStartTime() );
 		// Where to stop drawing the section
 		float EndX = TimeToPixelConverter.TimeToPixel( Section->GetEndTime() );
+		
+		// Actual section length without grips.
+		float SectionLengthActual = EndX-StartX;
 
-		if( StartX == EndX )
-		{
-			// Section is a single point in time
-			// extend it out a bit so it is visible
-			StartX-=1.0f;
-			EndX+=1.0f;
-		}
+		float SectionLengthWithGrips = SectionLengthActual+SequencerSectionConstants::SectionGripSize*2;
 
 		float ActualHeight = NodeHeight / MaxTracks;
 
 		// Compute allotted geometry area that can be used to draw the section
-		return AllottedGeometry.MakeChild( FVector2D( StartX, ActualHeight * RowIndex ), FVector2D( EndX - StartX, ActualHeight ) );
+		return AllottedGeometry.MakeChild( FVector2D( StartX-SequencerSectionConstants::SectionGripSize, ActualHeight * RowIndex ), FVector2D( SectionLengthWithGrips, ActualHeight ) );
 	}
 
 }
