Sequencer - A  couple of selection bugs and undo fixes

diff --git a/Engine/Source/Editor/EditorStyle/Private/SlateEditorStyle.cpp b/Engine/Source/Editor/EditorStyle/Private/SlateEditorStyle.cpp
--- a/Engine/Source/Editor/EditorStyle/Private/SlateEditorStyle.cpp
+++ b/Engine/Source/Editor/EditorStyle/Private/SlateEditorStyle.cpp
@@ -1416,7 +1416,7 @@ void FSlateEditorStyle::FStyle::SetupGeneralStyles()
 		Set( "Sequencer.SectionGripRight", new BOX_BRUSH( "Sequencer/SectionGripRight", FMargin(5.f/16.f) ) );
 		Set( "Sequencer.SectionArea.Background", new FSlateColorBrush( FColor::White ) );
 		Set( "Sequencer.AnimationOutliner.SelectionBorder", new BOX_BRUSH( "Common/RoundedSelection_16x", 4.0f/16.0f, SelectionColor_Pressed ) );
-		Set( "Sequencer.AnimationOutliner.BoldFont", TTF_CORE_FONT( "Fonts/Roboto-Bold", 11 ) );
+		Set( "Sequencer.AnimationOutliner.BoldFont", TTF_CORE_FONT( "Fonts/Roboto-Bold", 10 ) );
 		Set( "Sequencer.AnimationOutliner.RegularFont", TTF_CORE_FONT( "Fonts/Roboto-Regular", 8 ) );
 		Set( "Sequencer.ShotFilter", new IMAGE_BRUSH( "Sequencer/FilteredArea", FVector2D(74,74), FLinearColor::White, ESlateBrushTileType::Both ) );
 		Set( "Sequencer.KeyMark", new IMAGE_BRUSH("Sequencer/KeyMark", FVector2D(3,21), FLinearColor::White, ESlateBrushTileType::NoTile ) );
diff --git a/Engine/Source/Editor/Sequencer/Private/SAnimationOutlinerView.cpp b/Engine/Source/Editor/Sequencer/Private/SAnimationOutlinerView.cpp
--- a/Engine/Source/Editor/Sequencer/Private/SAnimationOutlinerView.cpp
+++ b/Engine/Source/Editor/Sequencer/Private/SAnimationOutlinerView.cpp
@@ -23,7 +23,7 @@ void SAnimationOutlinerTreeNode::Construct( const FArguments& InArgs, TSharedRef
 
 	// Choose the font.  If the node is a root node, we show a larger font for it.
 	FSlateFontInfo NodeFont = Node->GetParent().IsValid() ? 
-		FEditorStyle::GetFontStyle("Sequencer.AnimationOutliner.NormalFont") :
+		FEditorStyle::GetFontStyle("Sequencer.AnimationOutliner.RegularFont") :
 		FEditorStyle::GetFontStyle("Sequencer.AnimationOutliner.BoldFont");
 
 	TSharedRef<SWidget> TextWidget = 
diff --git a/Engine/Source/Editor/Sequencer/Private/SSection.cpp b/Engine/Source/Editor/Sequencer/Private/SSection.cpp
--- a/Engine/Source/Editor/Sequencer/Private/SSection.cpp
+++ b/Engine/Source/Editor/Sequencer/Private/SSection.cpp
@@ -70,7 +70,7 @@ FGeometry SSection::GetKeyAreaGeometry( const struct FKeyAreaElement& KeyArea, c
 }
 
 
-FSelectedKey SSection::GetKeyUnderMouse( const FVector2D& MousePosition, const FGeometry& SectionGeometry ) const
+FSelectedKey SSection::GetKeyUnderMouse( const FVector2D& MousePosition, const FGeometry& AllottedGeometry ) const
 {
 	UMovieSceneSection& Section = *SectionInterface->GetSectionObject();
 
@@ -81,11 +81,14 @@ FSelectedKey SSection::GetKeyUnderMouse( const FVector2D& MousePosition, const F
 		TSharedRef<IKeyArea> KeyArea = Element.KeyAreaNode.GetKeyArea( SectionIndex ); 
 
 		// Compute the current key area geometry
-		FGeometry KeyAreaGeometry = GetKeyAreaGeometry( Element, SectionGeometry );
+		FGeometry KeyAreaGeometryPadded = GetKeyAreaGeometry( Element, AllottedGeometry );
 
 		// Is the key area under the mouse
-		if( KeyAreaGeometry.IsUnderLocation( MousePosition ) )
+		if( KeyAreaGeometryPadded.IsUnderLocation( MousePosition ) )
 		{
+			FGeometry SectionGeometry = AllottedGeometry.MakeChild(FVector2D(SequencerSectionConstants::SectionGripSize, 0), AllottedGeometry.GetDrawSize() - FVector2D(SequencerSectionConstants::SectionGripSize*2, 0.0f));
+			FGeometry KeyAreaGeometry = GetKeyAreaGeometry( Element, SectionGeometry );
+
 			FVector2D LocalSpaceMousePosition = KeyAreaGeometry.AbsoluteToLocal( MousePosition );
 
 			FTimeToPixel TimeToPixelConverter( KeyAreaGeometry, TRange<float>( Section.GetStartTime(), Section.GetEndTime() ) );
@@ -466,12 +469,12 @@ FReply SSection::OnMouseButtonDown( const FGeometry& MyGeometry, const FPointerE
 
 	bDragging = false;
 
-	if( MouseEvent.GetEffectingButton() == EKeys::LeftMouseButton )
+	if( MouseEvent.GetEffectingButton() == EKeys::LeftMouseButton || MouseEvent.GetEffectingButton() == EKeys::RightMouseButton )
 	{
 		// Check for clicking on a key and mark it as the pressed key for drag detection (if necessary) later
 		PressedKey = GetKeyUnderMouse( MouseEvent.GetScreenSpacePosition(), MyGeometry );
 
-		if( !PressedKey.IsValid() )
+		if( !PressedKey.IsValid() && MouseEvent.GetEffectingButton() == EKeys::LeftMouseButton )
 		{
 			CheckForEdgeInteraction( MouseEvent, MyGeometry );
 		}
@@ -513,7 +516,7 @@ FReply SSection::OnMouseButtonUp( const FGeometry& MyGeometry, const FPointerEve
 	}
 	else
 	{
-		if( MouseEvent.GetEffectingButton() == EKeys::LeftMouseButton && HasMouseCapture() && MyGeometry.IsUnderLocation( MouseEvent.GetScreenSpacePosition() ) )
+		if( ( MouseEvent.GetEffectingButton() == EKeys::LeftMouseButton || MouseEvent.GetEffectingButton() == EKeys::RightMouseButton ) && HasMouseCapture() && MyGeometry.IsUnderLocation( MouseEvent.GetScreenSpacePosition() ) )
 		{
 			HandleSelection( MyGeometry, MouseEvent );
 		}
@@ -676,7 +679,7 @@ void SSection::HandleSelection( const FGeometry& MyGeometry, const FPointerEvent
 		bool bSelectDueToDrag = false;
 		HandleKeySelection( Key, MouseEvent, bSelectDueToDrag );
 	}
-	else
+	else if( MouseEvent.GetEffectingButton() == EKeys::LeftMouseButton ) // Only select anything but keys on left mouse button
 	{
 		HandleSectionSelection( MouseEvent );
 	}
diff --git a/Engine/Source/Editor/Sequencer/Private/SSection.h b/Engine/Source/Editor/Sequencer/Private/SSection.h
--- a/Engine/Source/Editor/Sequencer/Private/SSection.h
+++ b/Engine/Source/Editor/Sequencer/Private/SSection.h
@@ -70,10 +70,10 @@ private:
 	 * Determines the key that is under the mouse
 	 *
 	 * @param MousePosition		The current screen space position of the mouse
-	 * @param SectionGeometry	The geometry of the section
+	 * @param AllottedGeometry	The geometry of the mouse event
 	 * @return The key that is under the mouse.  Invalid if there is no key under the mouse
 	 */
-	FSelectedKey GetKeyUnderMouse( const FVector2D& MousePosition, const FGeometry& SectionGeometry ) const;
+	FSelectedKey GetKeyUnderMouse( const FVector2D& MousePosition, const FGeometry& AllottedGeometry ) const;
 
 	/**
 	 * Checks for user interaction (via the mouse) with the left and right edge of a section
diff --git a/Engine/Source/Editor/Sequencer/Private/Sequencer.cpp b/Engine/Source/Editor/Sequencer/Private/Sequencer.cpp
--- a/Engine/Source/Editor/Sequencer/Private/Sequencer.cpp
+++ b/Engine/Source/Editor/Sequencer/Private/Sequencer.cpp
@@ -309,6 +309,9 @@ void FSequencer::DeleteSection(class UMovieSceneSection* Section)
 	//checkSlow( Track->HasSection(Section) );
 	
 	Track->SetFlags( RF_Transactional );
+
+	FScopedTransaction DeleteSectionTransaction( NSLOCTEXT("Sequencer", "DeleteSection_Transaction", "Delete Section") );
+	
 	Track->Modify();
 
 	Track->RemoveSection(Section);
@@ -324,11 +327,15 @@ void FSequencer::DeleteSection(class UMovieSceneSection* Section)
 
 void FSequencer::DeleteSelectedKeys()
 {
+	FScopedTransaction DeleteKeysTransaction( NSLOCTEXT("Sequencer", "DeleteSelectedKeys_Transaction", "Delete Selected Keys") );
+
 	TArray<FSelectedKey> SelectedKeysArray = SelectedKeys.Array();
 	for ( const FSelectedKey& Key : SelectedKeysArray )
 	{
 		if (Key.IsValid())
 		{
+			Key.Section->Modify();
+
 			Key.KeyArea->DeleteKey(Key.KeyHandle.GetValue());
 		}
 	}
