Fix a crash clicking on a breadcrumb trail in Sequencer after creating a new animation in UMG

diff --git a/Engine/Source/Editor/Sequencer/Private/SSequencer.cpp b/Engine/Source/Editor/Sequencer/Private/SSequencer.cpp
--- a/Engine/Source/Editor/Sequencer/Private/SSequencer.cpp
+++ b/Engine/Source/Editor/Sequencer/Private/SSequencer.cpp
@@ -312,7 +312,7 @@ void SSequencer::Construct( const FArguments& InArgs, TSharedRef< class FSequenc
 		]
 	];
 
-	BreadcrumbTrail->PushCrumb(TAttribute<FText>::Create(TAttribute<FText>::FGetter::CreateSP(this, &SSequencer::GetRootMovieSceneName)), FSequencerBreadcrumb( Sequencer.Pin()->GetRootMovieSceneInstance() ));
+	ResetBreadcrumbs();
 }
 
 SSequencer::~SSequencer()
@@ -373,6 +373,12 @@ void SSequencer::UpdateBreadcrumbs(const TArray< TWeakObjectPtr<class UMovieScen
 	SequencerNodeTree->UpdateCachedVisibilityBasedOnShotFiltersChanged();
 }
 
+void SSequencer::ResetBreadcrumbs()
+{
+	BreadcrumbTrail->ClearCrumbs();
+	BreadcrumbTrail->PushCrumb(TAttribute<FText>::Create(TAttribute<FText>::FGetter::CreateSP(this, &SSequencer::GetRootMovieSceneName)), FSequencerBreadcrumb(Sequencer.Pin()->GetRootMovieSceneInstance()));
+}
+
 void SSequencer::OnOutlinerSearchChanged( const FText& Filter )
 {
 	SequencerNodeTree->FilterNodes( Filter.ToString() );
diff --git a/Engine/Source/Editor/Sequencer/Private/SSequencer.h b/Engine/Source/Editor/Sequencer/Private/SSequencer.h
--- a/Engine/Source/Editor/Sequencer/Private/SSequencer.h
+++ b/Engine/Source/Editor/Sequencer/Private/SSequencer.h
@@ -89,6 +89,7 @@ public:
 	 * @param FilteringShots A list of shots that are now filtering, or none if filtering is off
 	 */
 	void UpdateBreadcrumbs(const TArray< TWeakObjectPtr<class UMovieSceneSection> >& FilteringShots);
+	void ResetBreadcrumbs();
 
 	/** Deletes selected nodes out of the sequencer node tree */
 	void DeleteSelectedNodes();
diff --git a/Engine/Source/Editor/Sequencer/Private/Sequencer.cpp b/Engine/Source/Editor/Sequencer/Private/Sequencer.cpp
--- a/Engine/Source/Editor/Sequencer/Private/Sequencer.cpp
+++ b/Engine/Source/Editor/Sequencer/Private/Sequencer.cpp
@@ -206,6 +206,8 @@ void FSequencer::ResetToNewRootMovieScene( UMovieScene& NewRoot, TSharedRef<ISeq
 	RootMovieSceneInstance = MakeShareable(new FMovieSceneInstance(NewRoot));
 	MovieSceneStack.Add(RootMovieSceneInstance.ToSharedRef());
 
+	SequencerWidget->ResetBreadcrumbs();
+
 	NotifyMovieSceneDataChanged();
 }
 
