Sequencer - Fix some issues keying properties inside of structs

diff --git a/Engine/Source/Editor/Sequencer/Private/SequencerObjectChangeListener.cpp b/Engine/Source/Editor/Sequencer/Private/SequencerObjectChangeListener.cpp
--- a/Engine/Source/Editor/Sequencer/Private/SequencerObjectChangeListener.cpp
+++ b/Engine/Source/Editor/Sequencer/Private/SequencerObjectChangeListener.cpp
@@ -44,29 +44,41 @@ void FSequencerObjectChangeListener::OnPropertyChanged( const TArray<UObject*>&
 	FName InnerStructPropName = !bRequireAutoKey && ParentStructProperty ? Property->GetFName() : NAME_None;
 
 	FKeyPropertyParams Params;
-	Params.InnerStructPropertyName = InnerStructPropName;
+
 	Params.bRequireAutoKey = bRequireAutoKey;
 	Params.ObjectsThatChanged = ChangedObjects;
 
+	bool bFoundAndBroadcastedDelegate = false;
 	if(ParentStructProperty)
 	{
+
 		Params.PropertyHandle = ParentHandle.Get();
 		Params.PropertyPath = ParentHandle->GeneratePathToProperty();
 		// If the property parent is a struct, see if this property parent can be keyed. (e.g R,G,B,A for a color)
-		ClassToPropertyChangedMap.FindRef(ParentStructProperty->Struct->GetFName()).Broadcast( Params );
-	}
-	else if(StructProperty)
-	{
-		Params.PropertyHandle = &PropertyHandle;
-		Params.PropertyPath = PropertyHandle.GeneratePathToProperty();
-		ClassToPropertyChangedMap.FindRef(StructProperty->Struct->GetFName()).Broadcast(Params);
+		FOnAnimatablePropertyChanged Delegate = ClassToPropertyChangedMap.FindRef(ParentStructProperty->Struct->GetFName());
+		if(Delegate.IsBound())
+		{
+			Params.InnerStructPropertyName = InnerStructPropName;
+			bFoundAndBroadcastedDelegate = true;
+			Delegate.Broadcast(Params);
+		}
 	}
-	else
+
+	if( !bFoundAndBroadcastedDelegate )
 	{
-		Params.PropertyHandle = &PropertyHandle;
-		Params.PropertyPath = PropertyHandle.GeneratePathToProperty();
-		// the property in question is not a struct or an inner of the struct. See if it is directly keyable
-		ClassToPropertyChangedMap.FindRef(Property->GetClass()->GetFName()).Broadcast( Params );
+		if(StructProperty)
+		{
+			Params.PropertyHandle = &PropertyHandle;
+			Params.PropertyPath = PropertyHandle.GeneratePathToProperty();
+			ClassToPropertyChangedMap.FindRef(StructProperty->Struct->GetFName()).Broadcast(Params);
+		}
+		else
+		{
+			Params.PropertyHandle = &PropertyHandle;
+			Params.PropertyPath = PropertyHandle.GeneratePathToProperty();
+			// the property in question is not a struct or an inner of the struct. See if it is directly keyable
+			ClassToPropertyChangedMap.FindRef(Property->GetClass()->GetFName()).Broadcast(Params);
+		}
 	}
 }
 
