Fix auto-key working for properties that are not keyable

diff --git a/Engine/Source/Editor/Sequencer/Private/SequencerObjectChangeListener.cpp b/Engine/Source/Editor/Sequencer/Private/SequencerObjectChangeListener.cpp
--- a/Engine/Source/Editor/Sequencer/Private/SequencerObjectChangeListener.cpp
+++ b/Engine/Source/Editor/Sequencer/Private/SequencerObjectChangeListener.cpp
@@ -26,58 +26,70 @@ FSequencerObjectChangeListener::~FSequencerObjectChangeListener()
 
 void FSequencerObjectChangeListener::OnPropertyChanged( const TArray<UObject*>& ChangedObjects, const IPropertyHandle& PropertyHandle, bool bRequireAutoKey ) const
 {
-	const UProperty* Property = PropertyHandle.GetProperty();
-	const UStructProperty* StructProperty = Cast<const UStructProperty>(Property);
-	const UStructProperty* ParentStructProperty = nullptr;
-	TSharedPtr<IPropertyHandle> ParentHandle;
+	bool bIsKeyable = false;
 	
-	ParentHandle = PropertyHandle.GetParentHandle();
-	if(ParentHandle->IsValidHandle())
+	for( UObject* Object : ChangedObjects )
 	{
-		ParentStructProperty = Cast<const UStructProperty>(ParentHandle->GetProperty());
+		bIsKeyable |= IsTypeKeyable( *Object->GetClass(), PropertyHandle );
 	}
-	
 
-	FString PropertyVarName;
+	if( bIsKeyable )
+	{
+		const UProperty* Property = PropertyHandle.GetProperty();
+		const UStructProperty* StructProperty = Cast<const UStructProperty>(Property);
+		const UStructProperty* ParentStructProperty = nullptr;
+		TSharedPtr<IPropertyHandle> ParentHandle;
 
-	// If not in auto-key allow partial keying of specific inner struct property values;
-	FName InnerStructPropName = !bRequireAutoKey && ParentStructProperty ? Property->GetFName() : NAME_None;
+		ParentHandle = PropertyHandle.GetParentHandle();
+		if (ParentHandle->IsValidHandle())
+		{
+			ParentStructProperty = Cast<const UStructProperty>(ParentHandle->GetProperty());
+		}
 
-	FKeyPropertyParams Params;
 
-	Params.bRequireAutoKey = bRequireAutoKey;
-	Params.ObjectsThatChanged = ChangedObjects;
+		FString PropertyVarName;
 
-	bool bFoundAndBroadcastedDelegate = false;
-	if(ParentStructProperty)
-	{
+		// If not in auto-key allow partial keying of specific inner struct property values;
+		FName InnerStructPropName = !bRequireAutoKey && ParentStructProperty ? Property->GetFName() : NAME_None;
 
-		Params.PropertyHandle = ParentHandle.Get();
-		Params.PropertyPath = ParentHandle->GeneratePathToProperty();
-		// If the property parent is a struct, see if this property parent can be keyed. (e.g R,G,B,A for a color)
-		FOnAnimatablePropertyChanged Delegate = ClassToPropertyChangedMap.FindRef(ParentStructProperty->Struct->GetFName());
-		if(Delegate.IsBound())
-		{
-			Params.InnerStructPropertyName = InnerStructPropName;
-			bFoundAndBroadcastedDelegate = true;
-			Delegate.Broadcast(Params);
-		}
-	}
+		FKeyPropertyParams Params;
 
-	if( !bFoundAndBroadcastedDelegate )
-	{
-		if(StructProperty)
+		Params.bRequireAutoKey = bRequireAutoKey;
+		Params.ObjectsThatChanged = ChangedObjects;
+
+		bool bFoundAndBroadcastedDelegate = false;
+		if (ParentStructProperty)
 		{
-			Params.PropertyHandle = &PropertyHandle;
-			Params.PropertyPath = PropertyHandle.GeneratePathToProperty();
-			ClassToPropertyChangedMap.FindRef(StructProperty->Struct->GetFName()).Broadcast(Params);
+
+			Params.PropertyHandle = ParentHandle.Get();
+			Params.PropertyPath = ParentHandle->GeneratePathToProperty();
+
+
+			// If the property parent is a struct, see if this property parent can be keyed. (e.g R,G,B,A for a color)
+			FOnAnimatablePropertyChanged Delegate = ClassToPropertyChangedMap.FindRef(ParentStructProperty->Struct->GetFName());
+			if (Delegate.IsBound())
+			{
+				Params.InnerStructPropertyName = InnerStructPropName;
+				bFoundAndBroadcastedDelegate = true;
+				Delegate.Broadcast(Params);
+			}
 		}
-		else
+
+		if (!bFoundAndBroadcastedDelegate)
 		{
-			Params.PropertyHandle = &PropertyHandle;
-			Params.PropertyPath = PropertyHandle.GeneratePathToProperty();
-			// the property in question is not a struct or an inner of the struct. See if it is directly keyable
-			ClassToPropertyChangedMap.FindRef(Property->GetClass()->GetFName()).Broadcast(Params);
+			if (StructProperty)
+			{
+				Params.PropertyHandle = &PropertyHandle;
+				Params.PropertyPath = PropertyHandle.GeneratePathToProperty();
+				ClassToPropertyChangedMap.FindRef(StructProperty->Struct->GetFName()).Broadcast(Params);
+			}
+			else
+			{
+				Params.PropertyHandle = &PropertyHandle;
+				Params.PropertyPath = PropertyHandle.GeneratePathToProperty();
+				// the property in question is not a struct or an inner of the struct. See if it is directly keyable
+				ClassToPropertyChangedMap.FindRef(Property->GetClass()->GetFName()).Broadcast(Params);
+			}
 		}
 	}
 }
