Fix cases of constant FName recreation in the editor which is affecting performance

diff --git a/Engine/Plugins/2D/Paper2D/Source/Paper2DEditor/Private/TileMapEditing/STileLayerItem.cpp b/Engine/Plugins/2D/Paper2D/Source/Paper2DEditor/Private/TileMapEditing/STileLayerItem.cpp
--- a/Engine/Plugins/2D/Paper2D/Source/Paper2DEditor/Private/TileMapEditing/STileLayerItem.cpp
+++ b/Engine/Plugins/2D/Paper2D/Source/Paper2DEditor/Private/TileMapEditing/STileLayerItem.cpp
@@ -243,7 +243,8 @@ const FSlateBrush* STileLayerItem::GetVisibilityBrushForLayer() const
 
 FSlateColor STileLayerItem::GetForegroundColorForVisibilityButton() const
 {
-	return (VisibilityButton.IsValid() && (VisibilityButton->IsHovered() || VisibilityButton->IsPressed())) ? FEditorStyle::GetSlateColor("InvertedForeground") : FSlateColor::UseForeground();
+	static const FName InvertedForeground("InvertedForeground");
+	return (VisibilityButton.IsValid() && (VisibilityButton->IsHovered() || VisibilityButton->IsPressed())) ? FEditorStyle::GetSlateColor(InvertedForeground) : FSlateColor::UseForeground();
 }
 
 //////////////////////////////////////////////////////////////////////////
diff --git a/Engine/Plugins/Editor/EpicSurvey/Source/EpicSurvey/Private/SSurveyNotification.h b/Engine/Plugins/Editor/EpicSurvey/Source/EpicSurvey/Private/SSurveyNotification.h
--- a/Engine/Plugins/Editor/EpicSurvey/Source/EpicSurvey/Private/SSurveyNotification.h
+++ b/Engine/Plugins/Editor/EpicSurvey/Source/EpicSurvey/Private/SSurveyNotification.h
@@ -88,7 +88,8 @@ private:
 
 	FSlateColor GetInvertedForegroundIfHovered() const
 	{
-		return this->IsHovered() ? FEditorStyle::GetSlateColor( "InvertedForeground" ) : FSlateColor::UseForeground();
+		static const FName InvertedForeground("InvertedForeground");
+		return this->IsHovered() ? FEditorStyle::GetSlateColor( InvertedForeground ) : FSlateColor::UseForeground();
 	}
 
 	FSlateColor GetPinColorAndOpacity() const
diff --git a/Engine/Source/Developer/SlateReflector/Private/Widgets/SWidgetReflector.h b/Engine/Source/Developer/SlateReflector/Private/Widgets/SWidgetReflector.h
--- a/Engine/Source/Developer/SlateReflector/Private/Widgets/SWidgetReflector.h
+++ b/Engine/Source/Developer/SlateReflector/Private/Widgets/SWidgetReflector.h
@@ -164,9 +164,12 @@ private:
 	/** Callback for getting the color of the pick button text. */
 	FSlateColor HandlePickButtonColorAndOpacity() const
 	{
+		static const FName SelectionColor("SelectionColor");
+		static const FName DefaultForeground("DefaultForeground");
+
 		return bIsPicking
-			? FCoreStyle::Get().GetSlateColor("SelectionColor")
-			: FCoreStyle::Get().GetSlateColor("DefaultForeground");
+			? FCoreStyle::Get().GetSlateColor(SelectionColor)
+			: FCoreStyle::Get().GetSlateColor(DefaultForeground);
 	}
 
 	/** Callback for getting the text of the pick button. */
diff --git a/Engine/Source/Editor/BehaviorTreeEditor/Private/SGraphNode_BehaviorTree.cpp b/Engine/Source/Editor/BehaviorTreeEditor/Private/SGraphNode_BehaviorTree.cpp
--- a/Engine/Source/Editor/BehaviorTreeEditor/Private/SGraphNode_BehaviorTree.cpp
+++ b/Engine/Source/Editor/BehaviorTreeEditor/Private/SGraphNode_BehaviorTree.cpp
@@ -1069,7 +1069,11 @@ FSlateColor SGraphNode_BehaviorTree::GetIndexColor(bool bHovered) const
 {
 	UBehaviorTreeGraphNode* ParentNode = GetParentNode(GraphNode);
 	const bool bHighlightHover = bHovered || (ParentNode && ParentNode->bHighlightChildNodeIndices);
-	return bHighlightHover ? FEditorStyle::Get().GetSlateColor("BTEditor.Graph.BTNode.Index.HoveredColor") : FEditorStyle::Get().GetSlateColor("BTEditor.Graph.BTNode.Index.Color");
+
+	static const FName HoveredColor("BTEditor.Graph.BTNode.Index.HoveredColor");
+	static const FName DefaultColor("BTEditor.Graph.BTNode.Index.Color");
+
+	return bHighlightHover ? FEditorStyle::Get().GetSlateColor(HoveredColor) : FEditorStyle::Get().GetSlateColor(DefaultColor);
 }
 
 EVisibility SGraphNode_BehaviorTree::GetIndexVisibility() const
diff --git a/Engine/Source/Editor/ContentBrowser/Private/SAssetView.cpp b/Engine/Source/Editor/ContentBrowser/Private/SAssetView.cpp
--- a/Engine/Source/Editor/ContentBrowser/Private/SAssetView.cpp
+++ b/Engine/Source/Editor/ContentBrowser/Private/SAssetView.cpp
@@ -2133,7 +2133,10 @@ FLinearColor SAssetView::GetThumbnailHintColorAndOpacity() const
 
 FSlateColor SAssetView::GetViewButtonForegroundColor() const
 {
-	return ViewOptionsComboButton->IsHovered() ? FEditorStyle::GetSlateColor("InvertedForeground") : FEditorStyle::GetSlateColor("DefaultForeground");
+	static const FName InvertedForegroundName("InvertedForeground");
+	static const FName DefaultForegroundName("DefaultForeground");
+
+	return ViewOptionsComboButton->IsHovered() ? FEditorStyle::GetSlateColor(InvertedForegroundName) : FEditorStyle::GetSlateColor(DefaultForegroundName);
 }
 
 TSharedRef<SWidget> SAssetView::GetViewButtonContent()
diff --git a/Engine/Source/Editor/ContentBrowser/Private/SContentBrowser.cpp b/Engine/Source/Editor/ContentBrowser/Private/SContentBrowser.cpp
--- a/Engine/Source/Editor/ContentBrowser/Private/SContentBrowser.cpp
+++ b/Engine/Source/Editor/ContentBrowser/Private/SContentBrowser.cpp
@@ -63,6 +63,8 @@ void SContentBrowser::Construct( const FArguments& InArgs, const FName& InInstan
 	FrontendFilters = MakeShareable(new AssetFilterCollectionType());
 	TextFilter = MakeShareable( new FFrontendFilter_Text() );
 
+	static const FName DefaultForegroundName("DefaultForeground");
+
 	FContentBrowserCommands::Register();
 
 	BindCommands();
@@ -273,7 +275,7 @@ void SContentBrowser::Construct( const FArguments& InArgs, const FName& InInstan
 							SNew(SButton)
 							.VAlign(EVerticalAlignment::VAlign_Center)
 							.ButtonStyle( FEditorStyle::Get(), "ToggleButton" )
-							.ForegroundColor( FEditorStyle::GetSlateColor("DefaultForeground") )
+							.ForegroundColor( FEditorStyle::GetSlateColor(DefaultForegroundName) )
 							.ToolTipText( this, &SContentBrowser::GetHistoryBackTooltip )
 							.ContentPadding( FMargin(1, 0) )
 							.OnClicked(this, &SContentBrowser::BackClicked)
@@ -299,7 +301,7 @@ void SContentBrowser::Construct( const FArguments& InArgs, const FName& InInstan
 							SNew(SButton)
 							.VAlign(EVerticalAlignment::VAlign_Center)
 							.ButtonStyle( FEditorStyle::Get(), "ToggleButton" )
-							.ForegroundColor( FEditorStyle::GetSlateColor("DefaultForeground") )
+							.ForegroundColor( FEditorStyle::GetSlateColor(DefaultForegroundName) )
 							.ToolTipText( this, &SContentBrowser::GetHistoryForwardTooltip )
 							.ContentPadding( FMargin(1, 0) )
 							.OnClicked(this, &SContentBrowser::ForwardClicked)
@@ -435,7 +437,7 @@ void SContentBrowser::Construct( const FArguments& InArgs, const FName& InInstan
 									.ButtonStyle( FEditorStyle::Get(), "ToggleButton" )
 									.ToolTipText( LOCTEXT( "SourcesTreeToggleTooltip", "Show or hide the sources panel" ) )
 									.ContentPadding( FMargin( 1, 0 ) )
-									.ForegroundColor( FEditorStyle::GetSlateColor( "DefaultForeground" ) )
+									.ForegroundColor( FEditorStyle::GetSlateColor(DefaultForegroundName) )
 									.OnClicked( this, &SContentBrowser::SourcesViewExpandClicked )
 									[
 										SNew( SImage )
@@ -493,7 +495,7 @@ void SContentBrowser::Construct( const FArguments& InArgs, const FName& InInstan
 								.ButtonStyle( FEditorStyle::Get(), "ToggleButton" )
 								.ToolTipText( LOCTEXT( "SourcesTreeToggleTooltip", "Show or hide the sources panel" ) )
 								.ContentPadding( FMargin( 1, 0 ) )
-								.ForegroundColor( FEditorStyle::GetSlateColor( "DefaultForeground" ) )
+								.ForegroundColor( FEditorStyle::GetSlateColor(DefaultForegroundName) )
 								.OnClicked( this, &SContentBrowser::SourcesViewExpandClicked )
 								.Visibility( this, &SContentBrowser::GetPathExpanderVisibility )
 								[
diff --git a/Engine/Source/Editor/DetailCustomizations/Private/DateTimeStructCustomization.cpp b/Engine/Source/Editor/DetailCustomizations/Private/DateTimeStructCustomization.cpp
--- a/Engine/Source/Editor/DetailCustomizations/Private/DateTimeStructCustomization.cpp
+++ b/Engine/Source/Editor/DetailCustomizations/Private/DateTimeStructCustomization.cpp
@@ -48,7 +48,8 @@ FSlateColor FDateTimeStructCustomization::HandleTextBoxForegroundColor( ) const
 {
 	if (InputValid)
 	{
-		return FEditorStyle::GetSlateColor("InvertedForeground");
+		static const FName InvertedForegroundName("InvertedForeground");
+		return FEditorStyle::GetSlateColor(InvertedForegroundName);
 	}
 
 	return FLinearColor::Red;
diff --git a/Engine/Source/Editor/DetailCustomizations/Private/GuidStructCustomization.cpp b/Engine/Source/Editor/DetailCustomizations/Private/GuidStructCustomization.cpp
--- a/Engine/Source/Editor/DetailCustomizations/Private/GuidStructCustomization.cpp
+++ b/Engine/Source/Editor/DetailCustomizations/Private/GuidStructCustomization.cpp
@@ -116,7 +116,8 @@ FSlateColor FGuidStructCustomization::HandleTextBoxForegroundColor( ) const
 {
 	if (InputValid)
 	{
-		return FEditorStyle::GetSlateColor("InvertedForeground");
+		static const FName InvertedForegroundName("InvertedForeground");
+		return FEditorStyle::GetSlateColor(InvertedForegroundName);
 	}
 
 	return FLinearColor::Red;
diff --git a/Engine/Source/Editor/DetailCustomizations/Private/TimespanStructCustomization.cpp b/Engine/Source/Editor/DetailCustomizations/Private/TimespanStructCustomization.cpp
--- a/Engine/Source/Editor/DetailCustomizations/Private/TimespanStructCustomization.cpp
+++ b/Engine/Source/Editor/DetailCustomizations/Private/TimespanStructCustomization.cpp
@@ -48,7 +48,8 @@ FSlateColor FTimespanStructCustomization::HandleTextBoxForegroundColor( ) const
 {
 	if (InputValid)
 	{
-		return FEditorStyle::GetSlateColor("InvertedForeground");
+		static const FName InvertedForegroundName("InvertedForeground");
+		return FEditorStyle::GetSlateColor(InvertedForegroundName);
 	}
 
 	return FLinearColor::Red;
diff --git a/Engine/Source/Editor/InputBindingEditor/Private/Widgets/SGestureEditBox.cpp b/Engine/Source/Editor/InputBindingEditor/Private/Widgets/SGestureEditBox.cpp
--- a/Engine/Source/Editor/InputBindingEditor/Private/Widgets/SGestureEditBox.cpp
+++ b/Engine/Source/Editor/InputBindingEditor/Private/Widgets/SGestureEditBox.cpp
@@ -16,6 +16,8 @@ void SGestureEditBox::Construct( const FArguments& InArgs, TSharedPtr<FGestureTr
 	BorderImageHovered = FEditorStyle::GetBrush( "EditableTextBox.Background.Hovered" );
 	BorderImageFocused = FEditorStyle::GetBrush( "EditableTextBox.Background.Focused" );
 
+	static const FName InvertedForegroundName("InvertedForeground");
+
 	ChildSlot
 	[
 		SAssignNew( ConflictPopup, SMenuAnchor )
@@ -29,7 +31,7 @@ void SGestureEditBox::Construct( const FArguments& InArgs, TSharedPtr<FGestureTr
 				.VAlign(VAlign_Center)
 				.Padding( FMargin( 4.0f, 2.0f ) )
 				.BorderImage( this, &SGestureEditBox::GetBorderImage )
-				.ForegroundColor( FEditorStyle::GetSlateColor("InvertedForeground") )
+				.ForegroundColor( FEditorStyle::GetSlateColor(InvertedForegroundName) )
 				[
 					SNew( SHorizontalBox )
 					+ SHorizontalBox::Slot()
diff --git a/Engine/Source/Editor/InputBindingEditor/Private/Widgets/SInputBindingEditorPanel.cpp b/Engine/Source/Editor/InputBindingEditor/Private/Widgets/SInputBindingEditorPanel.cpp
--- a/Engine/Source/Editor/InputBindingEditor/Private/Widgets/SInputBindingEditorPanel.cpp
+++ b/Engine/Source/Editor/InputBindingEditor/Private/Widgets/SInputBindingEditorPanel.cpp
@@ -248,6 +248,7 @@ TSharedRef< ITableRow > SInputBindingEditorPanel::OnGenerateWidgetForTreeItem( T
 {
 	if( InTreeItem->IsContext() )
 	{
+		static const FName InvertedForegroundName("InvertedForeground");
 		// contexts do not need columns
 		return
 			SNew( STableRow< TSharedPtr<FGestureTreeItem> >, OwnerTable )
@@ -272,7 +273,7 @@ TSharedRef< ITableRow > SInputBindingEditorPanel::OnGenerateWidgetForTreeItem( T
 					[
 						SNew( STextBlock )
 						.Font( FEditorStyle::GetFontStyle( "InputBindingEditor.ContextFont" ) )
-						.ColorAndOpacity( FEditorStyle::GetSlateColor( "InvertedForeground" ) )
+						.ColorAndOpacity(FEditorStyle::GetSlateColor(InvertedForegroundName) )
 						.Text( InTreeItem->GetBindingContext()->GetContextDesc() )
 					]
 				]
diff --git a/Engine/Source/Editor/Kismet/Private/SSCSEditorViewport.cpp b/Engine/Source/Editor/Kismet/Private/SSCSEditorViewport.cpp
--- a/Engine/Source/Editor/Kismet/Private/SSCSEditorViewport.cpp
+++ b/Engine/Source/Editor/Kismet/Private/SSCSEditorViewport.cpp
@@ -28,12 +28,14 @@ public:
 	{
 		EditorViewport = InArgs._EditorViewport;
 
+		static const FName DefaultForegroundName("DefaultForeground");
+
 		this->ChildSlot
 		[
 			SNew(SBorder)
 			.BorderImage(FEditorStyle::GetBrush("NoBorder"))
 			.ColorAndOpacity(this, &SViewportToolBar::OnGetColorAndOpacity)
-			.ForegroundColor(FEditorStyle::GetSlateColor("DefaultForeground"))
+			.ForegroundColor(FEditorStyle::GetSlateColor(DefaultForegroundName))
 			[
 				SNew(SHorizontalBox)
 				+SHorizontalBox::Slot()
diff --git a/Engine/Source/Editor/Kismet/Private/WorkflowOrientedApp/SContentReference.cpp b/Engine/Source/Editor/Kismet/Private/WorkflowOrientedApp/SContentReference.cpp
--- a/Engine/Source/Editor/Kismet/Private/WorkflowOrientedApp/SContentReference.cpp
+++ b/Engine/Source/Editor/Kismet/Private/WorkflowOrientedApp/SContentReference.cpp
@@ -30,6 +30,8 @@ void SContentReference::Construct(const FArguments& InArgs)
 	BorderImageNormal = FEditorStyle::GetBrush(InArgs._Style, ".Background.Normal");
 	BorderImageHovered = FEditorStyle::GetBrush(InArgs._Style, ".Background.Hovered");
 
+	static const FName InvertedForegroundName("InvertedForeground");
+
 	// Create the widgets
 	ChildSlot
 	[
@@ -44,7 +46,7 @@ void SContentReference::Construct(const FArguments& InArgs)
 			.BorderImage(this, &SContentReference::GetBorderImage)
 			.Padding(FEditorStyle::GetMargin(InArgs._Style, ".BorderPadding"))
 			.BorderBackgroundColor( FLinearColor::White )
-			.ForegroundColor(FEditorStyle::GetSlateColor("InvertedForeground"))
+			.ForegroundColor(FEditorStyle::GetSlateColor(InvertedForegroundName))
 			.ToolTipText(this, &SContentReference::GetAssetFullName)
 			[
 				SNew(SBox)
diff --git a/Engine/Source/Editor/Layers/Private/SLayerBrowser.h b/Engine/Source/Editor/Layers/Private/SLayerBrowser.h
--- a/Engine/Source/Editor/Layers/Private/SLayerBrowser.h
+++ b/Engine/Source/Editor/Layers/Private/SLayerBrowser.h
@@ -183,12 +183,15 @@ private:
 
 	FSlateColor GetInvertedForegroundIfHovered() const
 	{
-		return ( ToggleModeButton.IsValid() && ( ToggleModeButton->IsHovered() || ToggleModeButton->IsPressed() ) ) ? FEditorStyle::GetSlateColor( "InvertedForeground" ) : FSlateColor::UseForeground();
+		static const FName InvertedForegroundName("InvertedForeground");
+		return ( ToggleModeButton.IsValid() && ( ToggleModeButton->IsHovered() || ToggleModeButton->IsPressed() ) ) ? FEditorStyle::GetSlateColor(InvertedForegroundName): FSlateColor::UseForeground();
 	}
 
 	const FSlateBrush* GetToggleModeButtonImageBrush() const
 	{
-		return ( Mode == ELayerBrowserMode::Layers ) ? FEditorStyle::GetBrush( "LayerBrowser.ExploreLayerContents" ) : FEditorStyle::GetBrush( "LayerBrowser.ReturnToLayersList" );
+		static const FName ExploreLayerContents("LayerBrowser.ExploreLayerContents");
+		static const FName ReturnToLayersList("LayerBrowser.ReturnToLayersList");
+		return ( Mode == ELayerBrowserMode::Layers ) ? FEditorStyle::GetBrush( ExploreLayerContents ) : FEditorStyle::GetBrush( ReturnToLayersList );
 	}
 
 	FText GetLayerContentsHeaderText() const;
diff --git a/Engine/Source/Editor/Layers/Private/SLayerStats.h b/Engine/Source/Editor/Layers/Private/SLayerStats.h
--- a/Engine/Source/Editor/Layers/Private/SLayerStats.h
+++ b/Engine/Source/Editor/Layers/Private/SLayerStats.h
@@ -124,7 +124,9 @@ private:
 
 		const auto Button = StatButtonWidgets[ StatsIndex ];
 
-		return ( Button.IsValid() && ( Button->IsHovered() || Button->IsPressed() ) ) ? FEditorStyle::GetSlateColor( "InvertedForeground" ) : FSlateColor::UseForeground();
+		static const FName InvertedForegroundName("InvertedForeground");
+
+		return ( Button.IsValid() && ( Button->IsHovered() || Button->IsPressed() ) ) ? FEditorStyle::GetSlateColor(InvertedForegroundName): FSlateColor::UseForeground();
 	}
 
 	/**
diff --git a/Engine/Source/Editor/Layers/Private/SLayersViewRow.cpp b/Engine/Source/Editor/Layers/Private/SLayersViewRow.cpp
--- a/Engine/Source/Editor/Layers/Private/SLayersViewRow.cpp
+++ b/Engine/Source/Editor/Layers/Private/SLayersViewRow.cpp
@@ -156,7 +156,8 @@ FReply SLayersViewRow::OnDrop(const FGeometry& MyGeometry, const FDragDropEvent&
 
 FSlateColor SLayersViewRow::GetForegroundColorForButton() const
 {
-	return (VisibilityButton.IsValid() && (VisibilityButton->IsHovered() || VisibilityButton->IsPressed())) ? FEditorStyle::GetSlateColor("InvertedForeground") : FSlateColor::UseForeground();
+	static const FName InvertedForegroundName("InvertedForeground");
+	return (VisibilityButton.IsValid() && (VisibilityButton->IsHovered() || VisibilityButton->IsPressed())) ? FEditorStyle::GetSlateColor(InvertedForegroundName) : FSlateColor::UseForeground();
 }
 
 FSlateColor SLayersViewRow::GetColorAndOpacity() const
diff --git a/Engine/Source/Editor/LevelEditor/Private/LevelEditorContextMenu.cpp b/Engine/Source/Editor/LevelEditor/Private/LevelEditorContextMenu.cpp
--- a/Engine/Source/Editor/LevelEditor/Private/LevelEditorContextMenu.cpp
+++ b/Engine/Source/Editor/LevelEditor/Private/LevelEditorContextMenu.cpp
@@ -420,7 +420,8 @@ FSlateColor InvertOnHover( const TWeakPtr< SWidget > WidgetPtr )
 	TSharedPtr< SWidget > Widget = WidgetPtr.Pin();
 	if ( Widget.IsValid() && Widget->IsHovered() )
 	{
-		return FEditorStyle::GetSlateColor( "InvertedForeground" );
+		static const FName InvertedForegroundName("InvertedForeground");
+		return FEditorStyle::GetSlateColor(InvertedForegroundName);
 	}
 
 	return FSlateColor::UseForeground();
diff --git a/Engine/Source/Editor/LevelEditor/Private/LevelEditorToolBar.cpp b/Engine/Source/Editor/LevelEditor/Private/LevelEditorToolBar.cpp
--- a/Engine/Source/Editor/LevelEditor/Private/LevelEditorToolBar.cpp
+++ b/Engine/Source/Editor/LevelEditor/Private/LevelEditorToolBar.cpp
@@ -1914,13 +1914,15 @@ TSharedRef< SWidget > FLevelEditorToolBar::GenerateMatineeMenuContent( TSharedRe
 				FOnActorPicked::CreateStatic( &FLevelEditorToolBar::OnMatineeActorPicked ) )
 		];
 
+	static const FName DefaultForegroundName("DefaultForeground");
+
 	// Give the scene outliner a border and background
 	const FSlateBrush* BackgroundBrush = FEditorStyle::GetBrush( "Menu.Background" );
 	TSharedRef< SBorder > RootBorder =
 		SNew( SBorder )
 		.Padding(3)
 		.BorderImage( BackgroundBrush )
-		.ForegroundColor( FEditorStyle::GetSlateColor("DefaultForeground") )
+		.ForegroundColor( FEditorStyle::GetSlateColor(DefaultForegroundName) )
 
 		// Assign the box panel as the child
 		[
diff --git a/Engine/Source/Editor/LevelEditor/Private/SLevelViewport.cpp b/Engine/Source/Editor/LevelEditor/Private/SLevelViewport.cpp
--- a/Engine/Source/Editor/LevelEditor/Private/SLevelViewport.cpp
+++ b/Engine/Source/Editor/LevelEditor/Private/SLevelViewport.cpp
@@ -2772,7 +2772,8 @@ FSlateColor SActorPreview::GetBorderColorAndOpacity() const
 
 	if (HighlightSequence.IsPlaying())
 	{
-		const FLinearColor SelectionColor = FEditorStyle::Get().GetSlateColor("SelectionColor").GetSpecifiedColor().CopyWithNewOpacity(0.5f);
+		static const FName SelectionColorName("SelectionColor");
+		const FLinearColor SelectionColor = FEditorStyle::Get().GetSlateColor(SelectionColorName).GetSpecifiedColor().CopyWithNewOpacity(0.5f);
 		
 		const float Interp = FMath::Sin(HighlightSequence.GetLerp()*6*PI) / 2 + 1;
 		Color = FMath::Lerp(SelectionColor, Color, Interp);
@@ -3295,7 +3296,9 @@ EVisibility SLevelViewport::GetMouseCaptureLabelVisibility() const
 
 FLinearColor SLevelViewport::GetMouseCaptureLabelColorAndOpacity() const
 {
-	FSlateColor SlateColor = FEditorStyle::GetSlateColor("DefaultForeground");
+	static const FName DefaultForegroundName("DefaultForeground");
+
+	FSlateColor SlateColor = FEditorStyle::GetSlateColor(DefaultForegroundName);
 	FLinearColor Col = SlateColor.IsColorSpecified() ? SlateColor.GetSpecifiedColor() : FLinearColor::White; 
 
 	float Alpha = 0.0f;
diff --git a/Engine/Source/Editor/LevelEditor/Private/SLevelViewportToolBar.cpp b/Engine/Source/Editor/LevelEditor/Private/SLevelViewportToolBar.cpp
--- a/Engine/Source/Editor/LevelEditor/Private/SLevelViewportToolBar.cpp
+++ b/Engine/Source/Editor/LevelEditor/Private/SLevelViewportToolBar.cpp
@@ -112,13 +112,15 @@ void SLevelViewportToolBar::Construct( const FArguments& InArgs )
 	const FMargin ToolbarSlotPadding( 2.0f, 2.0f );
 	const FMargin ToolbarButtonPadding( 2.0f, 0.0f );
 
+	static const FName DefaultForegroundName("DefaultForeground");
+
 	ChildSlot
 	[
 		SNew( SBorder )
 		.BorderImage( FEditorStyle::GetBrush("NoBorder") )
 		// Color and opacity is changed based on whether or not the mouse cursor is hovering over the toolbar area
 		.ColorAndOpacity( this, &SViewportToolBar::OnGetColorAndOpacity )
-		.ForegroundColor( FEditorStyle::GetSlateColor("DefaultForeground") )
+		.ForegroundColor( FEditorStyle::GetSlateColor(DefaultForegroundName) )
 		[
 			SNew( SVerticalBox )
 			+ SVerticalBox::Slot()
diff --git a/Engine/Source/Editor/Levels/Private/SLevelsViewRow.h b/Engine/Source/Editor/Levels/Private/SLevelsViewRow.h
--- a/Engine/Source/Editor/Levels/Private/SLevelsViewRow.h
+++ b/Engine/Source/Editor/Levels/Private/SLevelsViewRow.h
@@ -67,7 +67,8 @@ public:
 
 	FSlateColor GetForegroundBasedOnSelection() const
 	{
-		return FEditorStyle::GetSlateColor("DefaultForeground");
+		static const FName DefaultForegroundName("DefaultForeground");
+		return FEditorStyle::GetSlateColor(DefaultForegroundName);
 	}
 
 protected:
diff --git a/Engine/Source/Editor/MaterialEditor/Private/SMaterialEditorViewportToolBar.cpp b/Engine/Source/Editor/MaterialEditor/Private/SMaterialEditorViewportToolBar.cpp
--- a/Engine/Source/Editor/MaterialEditor/Private/SMaterialEditorViewportToolBar.cpp
+++ b/Engine/Source/Editor/MaterialEditor/Private/SMaterialEditorViewportToolBar.cpp
@@ -31,13 +31,15 @@ void SMaterialEditorViewportToolBar::Construct(const FArguments& InArgs)
 	}
 	ToolbarBuilder.EndSection();
 
+	static const FName DefaultForegroundName("DefaultForeground");
+
 	ChildSlot
 	[
 		SNew(SBorder)
 		.BorderImage(FEditorStyle::GetBrush("NoBorder"))
 		// Color and opacity is changed based on whether or not the mouse cursor is hovering over the toolbar area
 		.ColorAndOpacity(this, &SViewportToolBar::OnGetColorAndOpacity)
-		.ForegroundColor(FEditorStyle::GetSlateColor("DefaultForeground"))
+		.ForegroundColor(FEditorStyle::GetSlateColor(DefaultForegroundName))
 		.HAlign(HAlign_Right)
 		[
 			ToolbarBuilder.MakeWidget()
diff --git a/Engine/Source/Editor/PListEditor/Private/PListNodeArray.cpp b/Engine/Source/Editor/PListEditor/Private/PListNodeArray.cpp
--- a/Engine/Source/Editor/PListEditor/Private/PListNodeArray.cpp
+++ b/Engine/Source/Editor/PListEditor/Private/PListNodeArray.cpp
@@ -406,17 +406,21 @@ FSlateColor FPListNodeArray::GetKeyBackgroundColor() const
 /** Delegate: Changes the color of the key string text box based on validity */
 FSlateColor FPListNodeArray::GetKeyForegroundColor() const
 {
+	static const FName InvertedForegroundName("InvertedForeground");
+
 	if(bArrayMember)
 	{
-		return FEditorStyle::GetSlateColor("InvertedForeground");
+		
+		return FEditorStyle::GetSlateColor(InvertedForegroundName);
 	}
 
 	if(!bKeyValid)
 	{
-		return FEditorStyle::GetColor("ErrorReporting.ForegroundColor");
+		static const FName ForegroundColor("ErrorReporting.ForegroundColor");
+		return FEditorStyle::GetColor(ForegroundColor);
 	}
 	else
 	{
-		return FEditorStyle::GetSlateColor("InvertedForeground");
+		return FEditorStyle::GetSlateColor(InvertedForegroundName);
 	}
 }
diff --git a/Engine/Source/Editor/PListEditor/Private/PListNodeBoolean.cpp b/Engine/Source/Editor/PListEditor/Private/PListNodeBoolean.cpp
--- a/Engine/Source/Editor/PListEditor/Private/PListNodeBoolean.cpp
+++ b/Engine/Source/Editor/PListEditor/Private/PListNodeBoolean.cpp
@@ -349,17 +349,20 @@ FSlateColor FPListNodeBoolean::GetKeyBackgroundColor() const
 /** Delegate: Changes the color of the key string text box based on validity */
 FSlateColor FPListNodeBoolean::GetKeyForegroundColor() const
 {
+	static const FName InvertedForegroundName("InvertedForeground");
+
 	if(bArrayMember)
 	{
-		return FEditorStyle::GetSlateColor("InvertedForeground");
+		return FEditorStyle::GetSlateColor(InvertedForegroundName);
 	}
 
 	if(!bKeyValid)
 	{
-		return FEditorStyle::GetColor("ErrorReporting.ForegroundColor");
+		static const FName ForegroundColor("ErrorReporting.ForegroundColor");
+		return FEditorStyle::GetColor(ForegroundColor);
 	}
 	else
 	{
-		return FEditorStyle::GetSlateColor("InvertedForeground");
+		return FEditorStyle::GetSlateColor(InvertedForegroundName);
 	}
 }
diff --git a/Engine/Source/Editor/PListEditor/Private/PListNodeString.cpp b/Engine/Source/Editor/PListEditor/Private/PListNodeString.cpp
--- a/Engine/Source/Editor/PListEditor/Private/PListNodeString.cpp
+++ b/Engine/Source/Editor/PListEditor/Private/PListNodeString.cpp
@@ -328,18 +328,21 @@ FSlateColor FPListNodeString::GetKeyBackgroundColor() const
 /** Delegate: Changes the color of the key string text box based on validity */
 FSlateColor FPListNodeString::GetKeyForegroundColor() const
 {
+	static const FName InvertedForegroundName("InvertedForeground");
+
 	if(bArrayMember)
 	{
-		return FEditorStyle::GetSlateColor("InvertedForeground");
+		return FEditorStyle::GetSlateColor(InvertedForegroundName);
 	}
 
 	if(!bKeyValid)
 	{
-		return FEditorStyle::GetColor("ErrorReporting.ForegroundColor");
+		static const FName ForegroundColor("ErrorReporting.ForegroundColor");
+		return FEditorStyle::GetColor(ForegroundColor);
 	}
 	else
 	{
-		return FEditorStyle::GetSlateColor("InvertedForeground");
+		return FEditorStyle::GetSlateColor(InvertedForegroundName);
 	}
 }
 
@@ -349,7 +352,8 @@ FSlateColor FPListNodeString::GetValueBackgroundColor() const
 {
 	if(!bValueValid)
 	{
-		return FEditorStyle::GetColor("ErrorReporting.BackgroundColor");
+		static const FName BackgroundColor("ErrorReporting.BackgroundColor");
+		return FEditorStyle::GetColor(BackgroundColor);
 	}
 	else
 	{
@@ -362,11 +366,13 @@ FSlateColor FPListNodeString::GetValueForegroundColor() const
 {
 	if(!bValueValid)
 	{
-		return FEditorStyle::GetSlateColor("ErrorReporting.ForegroundColor");
+		static const FName ForegroundColor("ErrorReporting.ForegroundColor");
+		return FEditorStyle::GetSlateColor(ForegroundColor);
 	}
 	else
 	{
-		return FEditorStyle::GetSlateColor("InvertedForeground");
+		static const FName InvertedForegroundName("InvertedForeground");
+		return FEditorStyle::GetSlateColor(InvertedForegroundName);
 	}
 }
 
diff --git a/Engine/Source/Editor/Persona/Private/AnimationEditorViewportClient.cpp b/Engine/Source/Editor/Persona/Private/AnimationEditorViewportClient.cpp
--- a/Engine/Source/Editor/Persona/Private/AnimationEditorViewportClient.cpp
+++ b/Engine/Source/Editor/Persona/Private/AnimationEditorViewportClient.cpp
@@ -1911,7 +1911,10 @@ void FAnimationViewportClient::DrawMeshSubsetBones(const USkeletalMeshComponent*
 		TArray<FBoneIndexType> RequiredBones;
 
 		const FReferenceSkeleton& RefSkeleton = MeshComponent->SkeletalMesh->RefSkeleton;
-		const FSlateColor SelectionColor = FEditorStyle::GetSlateColor("SelectionColor");
+
+		static const FName SelectionColorName("SelectionColor");
+
+		const FSlateColor SelectionColor = FEditorStyle::GetSlateColor(SelectionColorName);
 		const FLinearColor LinearSelectionColor( SelectionColor.IsColorSpecified() ? SelectionColor.GetSpecifiedColor() : FLinearColor::White );
 
 		// we could cache parent bones as we calculate, but right now I'm not worried about perf issue of this
diff --git a/Engine/Source/Editor/Persona/Private/SAnimViewportToolBar.cpp b/Engine/Source/Editor/Persona/Private/SAnimViewportToolBar.cpp
--- a/Engine/Source/Editor/Persona/Private/SAnimViewportToolBar.cpp
+++ b/Engine/Source/Editor/Persona/Private/SAnimViewportToolBar.cpp
@@ -337,6 +337,7 @@ void SAnimViewportToolBar::Construct(const FArguments& InArgs, TSharedPtr<class
 			];
 	//@TODO: Need clipping horizontal box: LeftToolbar->AddWrapButton();
 
+	static const FName DefaultForegroundName("DefaultForeground");
 
 	ChildSlot
 	[
@@ -344,7 +345,7 @@ void SAnimViewportToolBar::Construct(const FArguments& InArgs, TSharedPtr<class
 		.BorderImage( FEditorStyle::GetBrush("NoBorder") )
 		// Color and opacity is changed based on whether or not the mouse cursor is hovering over the toolbar area
 		.ColorAndOpacity( this, &SViewportToolBar::OnGetColorAndOpacity )
-		.ForegroundColor( FEditorStyle::GetSlateColor("DefaultForeground") )
+		.ForegroundColor( FEditorStyle::GetSlateColor(DefaultForegroundName) )
 		[
 			SNew( SVerticalBox )
 			+ SVerticalBox::Slot()
diff --git a/Engine/Source/Editor/PhAT/Private/SPhATPreviewToolbar.cpp b/Engine/Source/Editor/PhAT/Private/SPhATPreviewToolbar.cpp
--- a/Engine/Source/Editor/PhAT/Private/SPhATPreviewToolbar.cpp
+++ b/Engine/Source/Editor/PhAT/Private/SPhATPreviewToolbar.cpp
@@ -12,13 +12,15 @@ void SPhATPreviewViewportToolBar::Construct(const FArguments& InArgs)
 {
 	PhATPtr = InArgs._PhATPtr;
 	
+	static const FName DefaultForegroundName("DefaultForeground");
+
 	this->ChildSlot
 	[
 		SNew(SBorder)
 		.BorderImage(FEditorStyle::GetBrush("NoBorder"))
 		// Color and opacity is changed based on whether or not the mouse cursor is hovering over the toolbar area
 		.ColorAndOpacity(this, &SViewportToolBar::OnGetColorAndOpacity)
-		.ForegroundColor(FEditorStyle::GetSlateColor("DefaultForeground"))
+		.ForegroundColor( FEditorStyle::GetSlateColor(DefaultForegroundName) )
 		[
 			SNew(SHorizontalBox)
 			+ SHorizontalBox::Slot()
diff --git a/Engine/Source/Editor/PlacementMode/Private/SPlacementModeTools.cpp b/Engine/Source/Editor/PlacementMode/Private/SPlacementModeTools.cpp
--- a/Engine/Source/Editor/PlacementMode/Private/SPlacementModeTools.cpp
+++ b/Engine/Source/Editor/PlacementMode/Private/SPlacementModeTools.cpp
@@ -568,7 +568,8 @@ const FSlateBrush* SPlacementModeTools::PlacementGroupBorderImage( int32 Placeme
 {
 	if ( WidgetSwitcher->GetActiveWidgetIndex() == PlacementGroupIndex )
 	{
-		return FEditorStyle::GetBrush( "PlacementBrowser.ActiveTabBar" );
+		static FName PlacementBrowserActiveTabBarBrush( "PlacementBrowser.ActiveTabBar" );
+		return FEditorStyle::GetBrush( PlacementBrowserActiveTabBarBrush );
 	}
 	else
 	{
diff --git a/Engine/Source/Editor/PropertyEditor/Private/UserInterface/PropertyEditor/SPropertyEditorBool.cpp b/Engine/Source/Editor/PropertyEditor/Private/UserInterface/PropertyEditor/SPropertyEditorBool.cpp
--- a/Engine/Source/Editor/PropertyEditor/Private/UserInterface/PropertyEditor/SPropertyEditorBool.cpp
+++ b/Engine/Source/Editor/PropertyEditor/Private/UserInterface/PropertyEditor/SPropertyEditorBool.cpp
@@ -9,10 +9,12 @@ void SPropertyEditorBool::Construct( const FArguments& InArgs, const TSharedRef<
 {
 	PropertyEditor = InPropertyEditor;
 
+	static const FName DefaultForegroundName("DefaultForeground");
+
 	CheckBox = SNew( SCheckBox )
 		.OnCheckStateChanged( this, &SPropertyEditorBool::OnCheckStateChanged )
 		.IsChecked( this, &SPropertyEditorBool::OnGetCheckState )
-		.ForegroundColor( FEditorStyle::GetSlateColor("DefaultForeground") )
+		.ForegroundColor( FEditorStyle::GetSlateColor(DefaultForegroundName) )
 		.Padding(0.0f);
 
 	ChildSlot
diff --git a/Engine/Source/Editor/PropertyEditor/Private/UserInterface/PropertyTable/SPropertyTableCell.cpp b/Engine/Source/Editor/PropertyEditor/Private/UserInterface/PropertyTable/SPropertyTableCell.cpp
--- a/Engine/Source/Editor/PropertyEditor/Private/UserInterface/PropertyTable/SPropertyTableCell.cpp
+++ b/Engine/Source/Editor/PropertyEditor/Private/UserInterface/PropertyTable/SPropertyTableCell.cpp
@@ -22,7 +22,9 @@ void SPropertyTableCell::Construct( const FArguments& InArgs, const TSharedRef<
 
 	FCoreUObjectDelegates::OnObjectPropertyChanged.AddSP(this, &SPropertyTableCell::OnCellValueChanged);
 
-	SetForegroundColor( FEditorStyle::GetSlateColor( "InvertedForeground" ) );
+	static const FName InvertedForegroundName("InvertedForeground");
+
+	SetForegroundColor( FEditorStyle::GetSlateColor(InvertedForegroundName) );
 }
 
 void SPropertyTableCell::SetContent( const TSharedRef< SWidget >& NewContents )
diff --git a/Engine/Source/Editor/ReferenceViewer/Private/SReferenceViewer.cpp b/Engine/Source/Editor/ReferenceViewer/Private/SReferenceViewer.cpp
--- a/Engine/Source/Editor/ReferenceViewer/Private/SReferenceViewer.cpp
+++ b/Engine/Source/Editor/ReferenceViewer/Private/SReferenceViewer.cpp
@@ -52,6 +52,8 @@ void SReferenceViewer::Construct( const FArguments& InArgs )
 	FEditorWidgetsModule& EditorWidgetsModule = FModuleManager::LoadModuleChecked<FEditorWidgetsModule>("EditorWidgets");
 	TSharedRef<SWidget> AssetDiscoveryIndicator = EditorWidgetsModule.CreateAssetDiscoveryIndicator(EAssetDiscoveryIndicatorScaleMode::Scale_None, FMargin(16, 8), false);
 
+	static const FName DefaultForegroundName("DefaultForeground");
+
 	ChildSlot
 	[
 		SNew(SVerticalBox)
@@ -71,7 +73,7 @@ void SReferenceViewer::Construct( const FArguments& InArgs )
 			[
 				SNew(SButton)
 				.ButtonStyle( FEditorStyle::Get(), "HoverHintOnly" )
-				.ForegroundColor( FEditorStyle::GetSlateColor("DefaultForeground") )
+				.ForegroundColor( FEditorStyle::GetSlateColor(DefaultForegroundName) )
 				.ToolTipText( this, &SReferenceViewer::GetHistoryBackTooltip )
 				.ContentPadding( 0 )
 				.OnClicked(this, &SReferenceViewer::BackClicked)
@@ -89,7 +91,7 @@ void SReferenceViewer::Construct( const FArguments& InArgs )
 			[
 				SNew(SButton)
 				.ButtonStyle( FEditorStyle::Get(), "HoverHintOnly" )
-				.ForegroundColor( FEditorStyle::GetSlateColor("DefaultForeground") )
+				.ForegroundColor( FEditorStyle::GetSlateColor(DefaultForegroundName) )
 				.ToolTipText( this, &SReferenceViewer::GetHistoryForwardTooltip )
 				.ContentPadding( 0 )
 				.OnClicked(this, &SReferenceViewer::ForwardClicked)
diff --git a/Engine/Source/Editor/SceneOutliner/Private/SSceneOutliner.cpp b/Engine/Source/Editor/SceneOutliner/Private/SSceneOutliner.cpp
--- a/Engine/Source/Editor/SceneOutliner/Private/SSceneOutliner.cpp
+++ b/Engine/Source/Editor/SceneOutliner/Private/SSceneOutliner.cpp
@@ -498,7 +498,10 @@ namespace SceneOutliner
 
 	FSlateColor SSceneOutliner::GetViewButtonForegroundColor() const
 	{
-		return ViewOptionsComboButton->IsHovered() ? FEditorStyle::GetSlateColor("InvertedForeground") : FEditorStyle::GetSlateColor("DefaultForeground");
+		static const FName InvertedForegroundName("InvertedForeground");
+		static const FName DefaultForegroundName("DefaultForeground");
+
+		return ViewOptionsComboButton->IsHovered() ? FEditorStyle::GetSlateColor(InvertedForegroundName) : FEditorStyle::GetSlateColor(DefaultForegroundName);
 	}
 
 	TSharedRef<SWidget> SSceneOutliner::GetViewButtonContent()
diff --git a/Engine/Source/Editor/Sequencer/Private/SSection.cpp b/Engine/Source/Editor/Sequencer/Private/SSection.cpp
--- a/Engine/Source/Editor/Sequencer/Private/SSection.cpp
+++ b/Engine/Source/Editor/Sequencer/Private/SSection.cpp
@@ -260,12 +260,19 @@ void SSection::PaintKeys( const FGeometry& AllottedGeometry, const FSlateRect& M
 
 	const FSequencer& Sequencer = ParentSectionArea->GetSequencer();
 
-	const FSlateBrush* BackgroundBrush = FEditorStyle::GetBrush("Sequencer.SectionArea.Background");
+	static const FName BackgroundBrushName("Sequencer.SectionArea.Background");
+	static const FName KeyBrushName("Sequencer.Key");
 
-	const FSlateBrush* KeyBrush = FEditorStyle::GetBrush("Sequencer.Key");
+	const FSlateBrush* BackgroundBrush = FEditorStyle::GetBrush(BackgroundBrushName);
+
+	const FSlateBrush* KeyBrush = FEditorStyle::GetBrush(KeyBrushName);
+
+	static const FName SelectionColorName("SelectionColor");
+	static const FName SelectionColorPressedName("SelectionColor_Pressed");
+
+	const FLinearColor PressedKeyColor = FEditorStyle::GetSlateColor(SelectionColorPressedName).GetColor( InWidgetStyle );
+	const FLinearColor SelectedKeyColor = FEditorStyle::GetSlateColor(SelectionColorName).GetColor( InWidgetStyle );
 
-	const FLinearColor PressedKeyColor = FEditorStyle::GetSlateColor("SelectionColor_Pressed").GetColor( InWidgetStyle );
-	const FLinearColor SelectedKeyColor = FEditorStyle::GetSlateColor("SelectionColor").GetColor( InWidgetStyle );
 	// @todo Sequencer temp color, make hovered brighter than selected.
 	FLinearColor HoveredKeyColor = SelectedKeyColor * FLinearColor(1.5,1.5,1.5,1.0f);
 
@@ -349,16 +356,21 @@ void SSection::DrawSectionBorders( const FGeometry& AllottedGeometry, const FSla
 
 	const bool bSelected = ParentSectionArea->GetSequencer().IsSectionSelected(SectionObject);
 
-	FLinearColor SelectionColor = FEditorStyle::GetSlateColor("SelectionColor").GetColor(FWidgetStyle());
+	static const FName SelectionColorName("SelectionColor");
+
+	FLinearColor SelectionColor = FEditorStyle::GetSlateColor(SelectionColorName).GetColor(FWidgetStyle());
 	FLinearColor TransparentSelectionColor = SelectionColor;
 
+	static const FName SectionGripLeftName("Sequencer.SectionGripLeft");
+	static const FName SectionGripRightName("Sequencer.SectionGripRight");
+
 	// Left Grip
 	FSlateDrawElement::MakeBox(
 		OutDrawElements,
 		LayerId,
 		// Center the key along Y.  Ensure the middle of the key is at the actual key time
 		AllottedGeometry.ToPaintGeometry( FVector2D( 0.0f, 0.0f ), FVector2D( SequencerSectionConstants::SectionGripSize, AllottedGeometry.GetDrawSize().Y) ) ,
-		FEditorStyle::GetBrush("Sequencer.SectionGripLeft"),
+		FEditorStyle::GetBrush(SectionGripLeftName),
 		MyClippingRect,
 		ESlateDrawEffect::None,
 		(bLeftEdgePressed || bLeftEdgeHovered) ? TransparentSelectionColor : FLinearColor::White
@@ -370,7 +382,7 @@ void SSection::DrawSectionBorders( const FGeometry& AllottedGeometry, const FSla
 		LayerId,
 		// Center the key along Y.  Ensure the middle of the key is at the actual key time
 		AllottedGeometry.ToPaintGeometry( FVector2D( AllottedGeometry.Size.X-SequencerSectionConstants::SectionGripSize, 0.0f), FVector2D(SequencerSectionConstants::SectionGripSize, AllottedGeometry.GetDrawSize().Y)),
-		FEditorStyle::GetBrush("Sequencer.SectionGripRight"),
+		FEditorStyle::GetBrush(SectionGripRightName),
 		MyClippingRect,
 		ESlateDrawEffect::None,
 		(bRightEdgePressed || bRightEdgeHovered) ? TransparentSelectionColor : FLinearColor::White
@@ -380,11 +392,13 @@ void SSection::DrawSectionBorders( const FGeometry& AllottedGeometry, const FSla
 	// draw selection box
 	if(bSelected)
 	{
+		static const FName PlainBorder("PlainBorder");
+
 		FSlateDrawElement::MakeBox(
 			OutDrawElements,
 			LayerId,
 			AllottedGeometry.ToPaintGeometry(),
-			FEditorStyle::GetBrush("PlainBorder"),
+			FEditorStyle::GetBrush(PlainBorder),
 			MyClippingRect,
 			ESlateDrawEffect::None,
 			SelectionColor
diff --git a/Engine/Source/Editor/UnrealEd/Private/AssetThumbnail.cpp b/Engine/Source/Editor/UnrealEd/Private/AssetThumbnail.cpp
--- a/Engine/Source/Editor/UnrealEd/Private/AssetThumbnail.cpp
+++ b/Engine/Source/Editor/UnrealEd/Private/AssetThumbnail.cpp
@@ -308,17 +308,17 @@ private:
 	const FSlateBrush* GetBackgroundBrush() const
 	{
 		const FAssetData& AssetData = AssetThumbnail->GetAssetData();
-		FName Substyle;
+		FString Substyle;
 		if ( AssetData.AssetClass == UClass::StaticClass()->GetFName() )
 		{
-			Substyle = FName(".ClassBackground");
+			Substyle = TEXT(".ClassBackground");
 		}
 		else
 		{
-			Substyle = FName(".AssetBackground");
+			Substyle = TEXT(".AssetBackground");
 		}
 
-		const FName BackgroundBrushName( *(Style.ToString() + Substyle.ToString()) );
+		const FName BackgroundBrushName( *(Style.ToString() + Substyle) );
 		return FEditorStyle::GetBrush(BackgroundBrushName);
 	}
 
diff --git a/Engine/Source/Editor/UnrealEd/Private/SCommonEditorViewportToolbarBase.cpp b/Engine/Source/Editor/UnrealEd/Private/SCommonEditorViewportToolbarBase.cpp
--- a/Engine/Source/Editor/UnrealEd/Private/SCommonEditorViewportToolbarBase.cpp
+++ b/Engine/Source/Editor/UnrealEd/Private/SCommonEditorViewportToolbarBase.cpp
@@ -25,13 +25,15 @@ void SCommonEditorViewportToolbarBase::Construct(const FArguments& InArgs, TShar
 	const FMargin ToolbarSlotPadding( 2.0f, 2.0f );
 	const FMargin ToolbarButtonPadding( 2.0f, 0.0f );
 
+	static const FName DefaultForegroundName("DefaultForeground");
+
 	ChildSlot
 	[
 		SNew( SBorder )
 		.BorderImage( FEditorStyle::GetBrush("NoBorder") )
 		// Color and opacity is changed based on whether or not the mouse cursor is hovering over the toolbar area
 		.ColorAndOpacity( this, &SViewportToolBar::OnGetColorAndOpacity )
-		.ForegroundColor( FEditorStyle::GetSlateColor("DefaultForeground") )
+		.ForegroundColor( FEditorStyle::GetSlateColor(DefaultForegroundName) )
 		[
 			SNew( SVerticalBox )
 			+ SVerticalBox::Slot()
diff --git a/Engine/Source/Editor/UnrealEd/Private/Settings/GameMapsSettingsCustomization.h b/Engine/Source/Editor/UnrealEd/Private/Settings/GameMapsSettingsCustomization.h
--- a/Engine/Source/Editor/UnrealEd/Private/Settings/GameMapsSettingsCustomization.h
+++ b/Engine/Source/Editor/UnrealEd/Private/Settings/GameMapsSettingsCustomization.h
@@ -210,7 +210,9 @@ private:
 		{
 			if (Value.IsEmpty() || IsValidMapName(Value))
 			{
-				return FEditorStyle::GetSlateColor("InvertedForeground");
+				static const FName InvertedForegroundName("InvertedForeground");
+
+				return FEditorStyle::GetSlateColor(InvertedForegroundName);
 			}
 		}
 
diff --git a/Engine/Source/Editor/WorldBrowser/Private/SWorldHierarchy.cpp b/Engine/Source/Editor/WorldBrowser/Private/SWorldHierarchy.cpp
--- a/Engine/Source/Editor/WorldBrowser/Private/SWorldHierarchy.cpp
+++ b/Engine/Source/Editor/WorldBrowser/Private/SWorldHierarchy.cpp
@@ -533,7 +533,10 @@ private:
 	/** @return the foreground color for the view button */
 	FSlateColor GetViewButtonForegroundColor() const
 	{
-		return ViewOptionsComboButton->IsHovered() ? FEditorStyle::GetSlateColor("InvertedForeground") : FEditorStyle::GetSlateColor("DefaultForeground");
+		static const FName InvertedForegroundName("InvertedForeground");
+		static const FName DefaultForegroundName("DefaultForeground");
+
+		return ViewOptionsComboButton->IsHovered() ? FEditorStyle::GetSlateColor(InvertedForegroundName) : FEditorStyle::GetSlateColor(DefaultForegroundName);
 	}
 
 	/** Toggles state of 'display path' */
diff --git a/Engine/Source/Runtime/Slate/Private/Framework/MultiBox/SMenuEntryBlock.cpp b/Engine/Source/Runtime/Slate/Private/Framework/MultiBox/SMenuEntryBlock.cpp
--- a/Engine/Source/Runtime/Slate/Private/Framework/MultiBox/SMenuEntryBlock.cpp
+++ b/Engine/Source/Runtime/Slate/Private/Framework/MultiBox/SMenuEntryBlock.cpp
@@ -1184,7 +1184,8 @@ FSlateColor SMenuEntryBlock::TintOnHover() const
 		TSharedPtr<SMultiBoxWidget> MultiBoxWidget = OwnerMultiBoxWidget.Pin();
 		const ISlateStyle* const StyleSet = MultiBoxWidget->GetStyleSet();
 
-		return StyleSet->GetSlateColor("SelectionColor");
+		static const FName SelectionColorName("SelectionColor");
+		return StyleSet->GetSlateColor(SelectionColorName);
 	}
 	else
 	{
diff --git a/Engine/Source/Runtime/Slate/Public/Widgets/Navigation/SBreadcrumbTrail.h b/Engine/Source/Runtime/Slate/Public/Widgets/Navigation/SBreadcrumbTrail.h
--- a/Engine/Source/Runtime/Slate/Public/Widgets/Navigation/SBreadcrumbTrail.h
+++ b/Engine/Source/Runtime/Slate/Public/Widgets/Navigation/SBreadcrumbTrail.h
@@ -340,11 +340,13 @@ private:
 
 			if ( CrumbButton.IsValid() && CrumbButton->IsHovered() )
 			{
-				return FCoreStyle::Get().GetSlateColor("InvertedForeground");
+				static const FName InvertedForegroundName("InvertedForeground");
+				return FCoreStyle::Get().GetSlateColor(InvertedForegroundName);
 			}
 		}
 
-		return FCoreStyle::Get().GetSlateColor("DefaultForeground");
+		static const FName DefaultForegroundName;
+		return FCoreStyle::Get().GetSlateColor(DefaultForegroundName);
 	}
 
 	/** Handler for when a crumb is clicked. Will pop crumbs down to the selected one. */
