Fixing more doxygen warnings

diff --git a/Engine/Source/Editor/Sequencer/Private/TimeSliderController.cpp b/Engine/Source/Editor/Sequencer/Private/TimeSliderController.cpp
--- a/Engine/Source/Editor/Sequencer/Private/TimeSliderController.cpp
+++ b/Engine/Source/Editor/Sequencer/Private/TimeSliderController.cpp
@@ -126,7 +126,7 @@ struct FDrawTickArgs
 	
 };
 
-void FSequencerTimeSliderController::DrawTicks( FSlateWindowElementList& OutDrawElements, const FScrubRangeToScreen& RangeToScreen, FDrawTickArgs& InArgs ) const
+void FSequencerTimeSliderController::DrawTicks( FSlateWindowElementList& OutDrawElements, const struct FScrubRangeToScreen& RangeToScreen, FDrawTickArgs& InArgs ) const
 {
 	const float Spacing = DetermineOptimalSpacing( RangeToScreen.PixelsPerInput, ScrubConstants::MinPixelsPerDisplayTick, ScrubConstants::MinDisplayTickSpacing );
 
diff --git a/Engine/Source/Runtime/GameplayAbilities/Private/GameplayCueManager.cpp b/Engine/Source/Runtime/GameplayAbilities/Private/GameplayCueManager.cpp
--- a/Engine/Source/Runtime/GameplayAbilities/Private/GameplayCueManager.cpp
+++ b/Engine/Source/Runtime/GameplayAbilities/Private/GameplayCueManager.cpp
@@ -12,7 +12,7 @@
 #endif
 
 
-UGameplayCueManager::UGameplayCueManager(const class FObjectInitializer& PCIP)
+UGameplayCueManager::UGameplayCueManager(const FObjectInitializer& PCIP)
 : Super(PCIP)
 {
 #if WITH_EDITOR
diff --git a/Engine/Source/Runtime/GameplayAbilities/Private/GameplayCueNotify_HitImpact.cpp b/Engine/Source/Runtime/GameplayAbilities/Private/GameplayCueNotify_HitImpact.cpp
--- a/Engine/Source/Runtime/GameplayAbilities/Private/GameplayCueNotify_HitImpact.cpp
+++ b/Engine/Source/Runtime/GameplayAbilities/Private/GameplayCueNotify_HitImpact.cpp
@@ -3,7 +3,7 @@
 #include "AbilitySystemPrivatePCH.h"
 #include "GameplayCueNotify_HitImpact.h"
 
-UGameplayCueNotify_HitImpact::UGameplayCueNotify_HitImpact(const class FObjectInitializer& PCIP)
+UGameplayCueNotify_HitImpact::UGameplayCueNotify_HitImpact(const FObjectInitializer& PCIP)
 : Super(PCIP)
 {
 
diff --git a/Engine/Source/Runtime/GameplayAbilities/Private/GameplayCueNotify_Static.cpp b/Engine/Source/Runtime/GameplayAbilities/Private/GameplayCueNotify_Static.cpp
--- a/Engine/Source/Runtime/GameplayAbilities/Private/GameplayCueNotify_Static.cpp
+++ b/Engine/Source/Runtime/GameplayAbilities/Private/GameplayCueNotify_Static.cpp
@@ -6,7 +6,7 @@
 #include "GameplayCueManager.h"
 #include "AbilitySystemComponent.h"
 
-UGameplayCueNotify_Static::UGameplayCueNotify_Static(const class FObjectInitializer& PCIP)
+UGameplayCueNotify_Static::UGameplayCueNotify_Static(const FObjectInitializer& PCIP)
 : Super(PCIP)
 {
 	IsOverride = true;
diff --git a/Engine/Source/Runtime/GameplayAbilities/Private/GameplayEffectTypes.cpp b/Engine/Source/Runtime/GameplayAbilities/Private/GameplayEffectTypes.cpp
--- a/Engine/Source/Runtime/GameplayAbilities/Private/GameplayEffectTypes.cpp
+++ b/Engine/Source/Runtime/GameplayAbilities/Private/GameplayEffectTypes.cpp
@@ -297,7 +297,7 @@ void FGameplayTagCountContainer::UpdateTagMap(const struct FGameplayTag& Tag, in
 	}
 }
 
-void FGameplayTagCountContainer::UpdateTagMap(const FGameplayTagContainer& Container, int32 CountDelta)
+void FGameplayTagCountContainer::UpdateTagMap(const struct FGameplayTagContainer& Container, int32 CountDelta)
 {
 	for (auto TagIt = Container.CreateConstIterator(); TagIt; ++TagIt)
 	{
