Sequencer - Add support for snapping, and different snapping options.

diff --git a/Engine/Source/Editor/EditorStyle/Private/SlateEditorStyle.cpp b/Engine/Source/Editor/EditorStyle/Private/SlateEditorStyle.cpp
--- a/Engine/Source/Editor/EditorStyle/Private/SlateEditorStyle.cpp
+++ b/Engine/Source/Editor/EditorStyle/Private/SlateEditorStyle.cpp
@@ -1481,6 +1481,8 @@ void FSlateEditorStyle::FStyle::SetupGeneralStyles()
 		Set( "Sequencer.AnimationOutliner.RegularFont", TTF_CORE_FONT( "Fonts/Roboto-Regular", 8 ) );
 		Set( "Sequencer.ShotFilter", new IMAGE_BRUSH( "Sequencer/FilteredArea", FVector2D(74,74), FLinearColor::White, ESlateBrushTileType::Both ) );
 		Set( "Sequencer.KeyMark", new IMAGE_BRUSH("Sequencer/KeyMark", FVector2D(3,21), FLinearColor::White, ESlateBrushTileType::NoTile ) );
+		Set( "Sequencer.ToggleIsSnapEnabled", new IMAGE_BRUSH( "Icons/icon_Sequencer_ToggleSnap_40x", Icon40x40 ) );
+		Set( "Sequencer.ToggleIsSnapEnabled.Small", new IMAGE_BRUSH( "Icons/icon_Sequencer_ToggleSnap_40x", Icon16x16 ) );
 
 		const FButtonStyle DetailsKeyButton = FButtonStyle(NoBorder)
 			.SetNormal( IMAGE_BRUSH("Sequencer/AddKey_Details", FVector2D(11,11) )  )
diff --git a/Engine/Source/Editor/Sequencer/Private/SSequencer.cpp b/Engine/Source/Editor/Sequencer/Private/SSequencer.cpp
--- a/Engine/Source/Editor/Sequencer/Private/SSequencer.cpp
+++ b/Engine/Source/Editor/Sequencer/Private/SSequencer.cpp
@@ -23,6 +23,7 @@
 #include "MovieSceneShotSection.h"
 #include "CommonMovieSceneTools.h"
 #include "SSearchBox.h"
+#include "SNumericEntryBox.h"
 
 
 #define LOCTEXT_NAMESPACE "Sequencer"
@@ -146,8 +147,6 @@ void SSequencer::Construct( const FArguments& InArgs, TSharedRef< class FSequenc
 	// Create a node tree which contains a tree of movie scene data to display in the sequence
 	SequencerNodeTree = MakeShareable( new FSequencerNodeTree( InSequencer.Get() ) );
 
-	AllowAutoKey = InArgs._AllowAutoKey;
-	OnToggleAutoKey = InArgs._OnToggleAutoKey;
 	CleanViewEnabled = InArgs._CleanViewEnabled;
 	OnToggleCleanView = InArgs._OnToggleCleanView;
 
@@ -185,18 +184,10 @@ void SSequencer::Construct( const FArguments& InArgs, TSharedRef< class FSequenc
 				.VAlign(VAlign_Center)
 				[
 					SNew(SHorizontalBox)
-					+SHorizontalBox::Slot()
+					+ SHorizontalBox::Slot()
 					.AutoWidth()
-					.HAlign(HAlign_Left)
 					[
-						// @todo Sequencer - Temp auto-key button
-						SNew( SCheckBox )
-						.IsChecked( this, &SSequencer::OnGetAutoKeyCheckState )
-						.OnCheckStateChanged( this, &SSequencer::OnAutoKeyChecked )
-						[
-							SNew( STextBlock )
-							.Text( LOCTEXT("ToggleAutoKey", "Auto Key") )
-						]
+						MakeToolBar()
 					]
 					+SHorizontalBox::Slot()
 					.HAlign(HAlign_Right)
@@ -316,6 +307,68 @@ void SSequencer::Construct( const FArguments& InArgs, TSharedRef< class FSequenc
 	ResetBreadcrumbs();
 }
 
+TSharedRef<SWidget> SSequencer::MakeToolBar()
+{
+	FToolBarBuilder ToolBarBuilder( Sequencer.Pin()->GetCommandBindings(), FMultiBoxCustomization::None, TSharedPtr<FExtender>(), Orient_Horizontal, true);
+	
+	ToolBarBuilder.SetLabelVisibility(EVisibility::Visible);
+	ToolBarBuilder.AddToolBarButton( FSequencerCommands::Get().ToggleAutoKeyEnabled );
+
+	ToolBarBuilder.SetLabelVisibility(EVisibility::Collapsed);
+	ToolBarBuilder.AddToolBarButton( FSequencerCommands::Get().ToggleIsSnapEnabled, NAME_None, TAttribute<FText>( FText::GetEmpty() ) );
+	ToolBarBuilder.AddComboButton(
+		FUIAction(),
+		FOnGetContent::CreateSP( this, &SSequencer::MakeSnapMenu ),
+		LOCTEXT( "SnapOptions", "Options" ),
+		LOCTEXT( "SnapOptionsToolTip", "Snapping Options" ),
+		TAttribute<FSlateIcon>(),
+		true );
+	ToolBarBuilder.AddWidget(
+		SNew( SBox )
+		.VAlign( VAlign_Center )
+		[
+			SNew( SNumericEntryBox<float> )
+			.AllowSpin( true )
+			.MinValue( 0 )
+			.MinSliderValue( 0 )
+			.MaxSliderValue( 1 )
+			.MinDesiredValueWidth( 50 )
+			.Delta( 0.05f )
+			.ToolTipText( LOCTEXT( "SnappingIntervalToolTip", "Snapping interval" ) )
+			.Value( this, &SSequencer::OnGetSnapInterval )
+			.OnValueChanged( this, &SSequencer::OnSnapIntervalChanged ) 
+		] );
+
+	return ToolBarBuilder.MakeWidget();
+}
+
+TSharedRef<SWidget> SSequencer::MakeSnapMenu()
+{
+	FMenuBuilder MenuBuilder( false, Sequencer.Pin()->GetCommandBindings() );
+
+	MenuBuilder.BeginSection( "KeySnapping", LOCTEXT( "SnappingMenuKeyHeader", "Key Snapping" ) );
+	{
+		MenuBuilder.AddMenuEntry( FSequencerCommands::Get().ToggleSnapKeysToInterval );
+		MenuBuilder.AddMenuEntry( FSequencerCommands::Get().ToggleSnapKeysToKeys );
+	}
+	MenuBuilder.EndSection();
+
+	MenuBuilder.BeginSection( "SectionSnapping", LOCTEXT( "SnappingMenuSectionHeader", "Section Snapping" ) );
+	{
+		MenuBuilder.AddMenuEntry( FSequencerCommands::Get().ToggleSnapSectionsToInterval );
+		MenuBuilder.AddMenuEntry( FSequencerCommands::Get().ToggleSnapSectionsToSections );
+	}
+	MenuBuilder.EndSection();
+
+	MenuBuilder.BeginSection( "PlayTimeSnapping", LOCTEXT( "SnappingMenuPlayTimeHeader", "Play Time Snapping" ) );
+	{
+		MenuBuilder.AddMenuEntry( FSequencerCommands::Get().ToggleSnapPlayTimeToInterval );
+	}
+	MenuBuilder.EndSection();
+
+	return MenuBuilder.MakeWidget();
+}
+
 SSequencer::~SSequencer()
 {
 }
@@ -405,16 +458,6 @@ TSharedRef<SWidget> SSequencer::MakeSectionOverlay( TSharedRef<FSequencerTimeSli
 		];
 }
 
-ESlateCheckBoxState::Type SSequencer::OnGetAutoKeyCheckState() const
-{
-	return AllowAutoKey.Get() ? ESlateCheckBoxState::Checked : ESlateCheckBoxState::Unchecked;
-}
-
-void SSequencer::OnAutoKeyChecked( ESlateCheckBoxState::Type InState ) 
-{
-	OnToggleAutoKey.ExecuteIfBound( InState == ESlateCheckBoxState::Checked ? true : false );
-}
-
 EVisibility SSequencer::OnGetCleanViewVisibility() const 
 {
 	return Sequencer.Pin()->IsLevelEditorSequencer() ? EVisibility::Visible : EVisibility::Collapsed;
@@ -430,6 +473,16 @@ void SSequencer::OnCleanViewChecked( ESlateCheckBoxState::Type InState )
 	OnToggleCleanView.ExecuteIfBound( InState == ESlateCheckBoxState::Checked ? true : false );
 }
 
+TOptional<float> SSequencer::OnGetSnapInterval() const
+{
+	return GetDefault<USequencerSnapSettings>()->GetSnapInterval();
+}
+
+void SSequencer::OnSnapIntervalChanged( float InInterval )
+{
+	GetMutableDefault<USequencerSnapSettings>()->SetSnapInterval(InInterval);
+}
+
 
 void SSequencer::OnDragEnter( const FGeometry& MyGeometry, const FDragDropEvent& DragDropEvent )
 {
@@ -520,7 +573,7 @@ FReply SSequencer::OnDrop( const FGeometry& MyGeometry, const FDragDropEvent& Dr
 FReply SSequencer::OnKeyDown( const FGeometry& MyGeometry, const FKeyEvent& InKeyEvent ) 
 {
 	// A toolkit tab is active, so direct all command processing to it
-	if( Sequencer.Pin()->GetCommandBindings().ProcessCommandBindings( InKeyEvent ) )
+	if( Sequencer.Pin()->GetCommandBindings()->ProcessCommandBindings( InKeyEvent ) )
 	{
 		return FReply::Handled();
 	}
diff --git a/Engine/Source/Editor/Sequencer/Private/SSequencer.h b/Engine/Source/Editor/Sequencer/Private/SSequencer.h
--- a/Engine/Source/Editor/Sequencer/Private/SSequencer.h
+++ b/Engine/Source/Editor/Sequencer/Private/SSequencer.h
@@ -57,8 +57,6 @@ public:
 		SLATE_ATTRIBUTE( TRange<float>, ViewRange )
 		/** The current scrub position in (seconds) */
 		SLATE_ATTRIBUTE( float, ScrubPosition )
-		/** Whether or not to allow auto-key */
-		SLATE_ATTRIBUTE( bool, AllowAutoKey )
 		/** Whether the clean sequencer view is enabled */
 		SLATE_ATTRIBUTE( bool, CleanViewEnabled )
 		/** Called when the user changes the view range */
@@ -66,8 +64,6 @@ public:
 		/** Called when the user changes the scrub position */
 		SLATE_EVENT( FOnScrubPositionChanged, OnScrubPositionChanged )
 		/** Called when the user toggles auto-key */
-		SLATE_EVENT( FOnToggleBoolOption, OnToggleAutoKey )
-		/** Called when the user toggles auto-key */
 		SLATE_EVENT( FOnToggleBoolOption, OnToggleCleanView )
 	SLATE_END_ARGS()
 
@@ -95,17 +91,11 @@ public:
 	void DeleteSelectedNodes();
 
 private:
-	/**
-	 * @return The state of the auto-key check box
-	 */
-	ESlateCheckBoxState::Type OnGetAutoKeyCheckState() const;
-	
-	/**
-	 * Called when the auto-key checkbox state changes
-	 * 
-	 * @param InState	The new state
-	 */
-	void OnAutoKeyChecked( ESlateCheckBoxState::Type InState );	
+	/** Makes the toolbar for the outline section. */
+	TSharedRef<SWidget> MakeToolBar();
+
+	/** Makes the snapping menu for the toolbar. */
+	TSharedRef<SWidget> MakeSnapMenu();
 	
 	/** @return The visibility of the clean view button */
 	EVisibility OnGetCleanViewVisibility() const;
@@ -122,6 +112,18 @@ private:
 	 */
 	void OnCleanViewChecked( ESlateCheckBoxState::Type InState );	
 
+	/**
+	* @return The value of the current snap interval.
+	*/
+	TOptional<float> OnGetSnapInterval() const;
+
+	/**
+	* Called when the snap interval slider changes.
+	*
+	* @param InInterval	The new snap interval
+	*/
+	void OnSnapIntervalChanged(float InInterval);
+
 	/**
 	 * Called when the outliner search terms change                                                              
 	 */
@@ -131,7 +133,7 @@ private:
 	 * @return The fill percentage of the animation outliner
 	 * @todo Sequencer Make this user adjustable
 	 */
-	float GetAnimationOutlinerFillPercentage() const { return .2f; }
+	float GetAnimationOutlinerFillPercentage() const { return .25f; }
 
 	/**
 	 * Creates an overlay for various components of the time slider that should be displayed on top or bottom of the section area
@@ -195,10 +197,6 @@ private:
 	TSharedPtr<class FSequencerNodeTree> SequencerNodeTree;
 	/** The breadcrumb trail widget for this sequencer */
 	TSharedPtr< class SBreadcrumbTrail<FSequencerBreadcrumb> > BreadcrumbTrail;
-	/** Delegate to call when auto-key is toggled */
-	FOnToggleBoolOption OnToggleAutoKey;
-	/** Whether or not to allow auto-key */
-	TAttribute<bool> AllowAutoKey;
 	/** Delegate to call when clean view is toggled */
 	FOnToggleBoolOption OnToggleCleanView;
 	/** Whether the clean sequencer view is enabled */
diff --git a/Engine/Source/Editor/Sequencer/Private/SectionDragOperations.cpp b/Engine/Source/Editor/Sequencer/Private/SectionDragOperations.cpp
--- a/Engine/Source/Editor/Sequencer/Private/SectionDragOperations.cpp
+++ b/Engine/Source/Editor/Sequencer/Private/SectionDragOperations.cpp
@@ -165,6 +165,10 @@ void FResizeSection::OnDrag( const FPointerEvent& MouseEvent, const FVector2D& L
 		TRange<float> SectionBoundaries = GetSectionBoundaries(Section.Get(), SequencerNode);
 		
 		// Snapping
+		if ( SnapSettings->GetIsSnapEnabled() )
+		{
+			bool bSnappedToSection = false;
+			if ( SnapSettings->GetSnapSectionsToSections() )
 		{
 			TArray<float> TimesToSnapTo;
 			GetSectionSnapTimes(TimesToSnapTo, Section.Get(), SequencerNode, bIsDilating);
@@ -174,6 +178,13 @@ void FResizeSection::OnDrag( const FPointerEvent& MouseEvent, const FVector2D& L
 			if (NewSnappedTime.IsSet())
 			{
 				NewTime = NewSnappedTime.GetValue();
+					bSnappedToSection = true;
+				}
+			}
+
+			if ( bSnappedToSection == false && SnapSettings->GetSnapSectionsToInterval() )
+			{
+				NewTime = SnapSettings->SnapToInterval(NewTime);
 			}
 		}
 
@@ -259,7 +270,11 @@ void FMoveSection::OnDrag( const FPointerEvent& MouseEvent, const FVector2D& Loc
 		float DistanceMoved = TotalDelta.X / TimeToPixelConverter.GetPixelsPerInput();
 		
 		float DeltaTime = DistanceMoved;
-		// Snapping
+		
+		if ( SnapSettings->GetIsSnapEnabled() )
+		{
+			bool bSnappedToSection = false;
+			if ( SnapSettings->GetSnapSectionsToSections() )
 		{
 			TArray<float> TimesToSnapTo;
 			GetSectionSnapTimes(TimesToSnapTo, Section.Get(), SequencerNode, true);
@@ -270,10 +285,17 @@ void FMoveSection::OnDrag( const FPointerEvent& MouseEvent, const FVector2D& Loc
 
 			float OutSnappedTime = 0.f;
 			float OutNewTime = 0.f;
-			bool bSuccess = SnapToTimes(TimesToSnap, TimesToSnapTo, TimeToPixelConverter, OutSnappedTime, OutNewTime);
-			if (bSuccess)
+				if ( SnapToTimes( TimesToSnap, TimesToSnapTo, TimeToPixelConverter, OutSnappedTime, OutNewTime ) )
 			{
 				DeltaTime = OutNewTime - (OutSnappedTime - DistanceMoved);
+					bSnappedToSection = true;
+			}
+		}
+
+			if ( bSnappedToSection == false && SnapSettings->GetSnapSectionsToInterval() )
+			{
+				float NewStartTime = DistanceMoved + Section->GetStartTime();
+				DeltaTime = SnapSettings->SnapToInterval( NewStartTime ) - Section->GetStartTime();
 			}
 		}
 
@@ -407,7 +429,11 @@ void FMoveKeys::OnDrag( const FPointerEvent& MouseEvent, const FVector2D& LocalM
 	{
 		float TimeDelta = DistanceMoved;
 		// Snapping
+		if ( SnapSettings->GetIsSnapEnabled() )
 		{
+			bool bSnappedToKeyTime = false;
+			if ( SnapSettings->GetSnapKeysToKeys() )
+			{
 			TArray<float> OutSnapTimes;
 			GetKeySnapTimes(OutSnapTimes, SequencerNode);
 
@@ -419,11 +445,16 @@ void FMoveKeys::OnDrag( const FPointerEvent& MouseEvent, const FVector2D& LocalM
 			}
 			float OutInitialTime = 0.f;
 			float OutSnapTime = 0.f;
-			bool bSuccess = SnapToTimes(InitialTimes, OutSnapTimes, TimeToPixelConverter, OutInitialTime, OutSnapTime);
+				if ( SnapToTimes( InitialTimes, OutSnapTimes, TimeToPixelConverter, OutInitialTime, OutSnapTime ) )
+				{
+					bSnappedToKeyTime = true;
+					TimeDelta = OutSnapTime - (OutInitialTime - DistanceMoved);
+				}
+			}
 
-			if (bSuccess)
+			if ( bSnappedToKeyTime == false && SnapSettings->GetSnapKeysToInterval() )
 			{
-				TimeDelta = OutSnapTime - (OutInitialTime - DistanceMoved);
+				TimeDelta = SnapSettings->SnapToInterval( MouseTime ) - SelectedKeyTime;
 			}
 		}
 
diff --git a/Engine/Source/Editor/Sequencer/Private/SectionDragOperations.h b/Engine/Source/Editor/Sequencer/Private/SectionDragOperations.h
--- a/Engine/Source/Editor/Sequencer/Private/SectionDragOperations.h
+++ b/Engine/Source/Editor/Sequencer/Private/SectionDragOperations.h
@@ -8,6 +8,11 @@
 class FSequencerDragOperation
 {
 public:
+	FSequencerDragOperation()
+	{
+		SnapSettings = GetDefault<USequencerSnapSettings>();
+	}
+
 	virtual ~FSequencerDragOperation(){}
 	void BeginTransaction( UMovieSceneSection& Section, const FText& TransactionDesc );
 	void EndTransaction();
@@ -69,6 +74,9 @@ protected:
 	 * @return						The snap time that was snapped to, or potentially nothing if nothing snapped
 	 */
 	TOptional<float> SnapToTimes(float InitialTimes, const TArray<float>& SnapTimes, const FTimeToPixel& TimeToPixelConverter);
+
+	/** The current snap settings */
+	const USequencerSnapSettings* SnapSettings;
 };
 
 
diff --git a/Engine/Source/Editor/Sequencer/Private/Sequencer.cpp b/Engine/Source/Editor/Sequencer/Private/Sequencer.cpp
--- a/Engine/Source/Editor/Sequencer/Private/Sequencer.cpp
+++ b/Engine/Source/Editor/Sequencer/Private/Sequencer.cpp
@@ -75,14 +75,11 @@ void FSequencer::InitSequencer( const FSequencerInitParams& InitParams, const TA
 		SequencerWidget = SNew( SSequencer, SharedThis( this ) )
 			.ViewRange( this, &FSequencer::OnGetViewRange )
 			.ScrubPosition( this, &FSequencer::OnGetScrubPosition )
-			.AllowAutoKey( this, &FSequencer::OnGetAllowAutoKey )
 			.CleanViewEnabled( this, &FSequencer::IsUsingCleanView )
 			.OnScrubPositionChanged( this, &FSequencer::OnScrubPositionChanged )
 			.OnViewRangeChanged( this, &FSequencer::OnViewRangeChanged, false )
-			.OnToggleAutoKey( this, &FSequencer::OnToggleAutoKey )
 			.OnToggleCleanView( this, &FSequencer::OnToggleCleanView );
 
-		
 		// When undo occurs, get a notification so we can make sure our view is up to date
 		GEditor->RegisterForUndo(this);
 
@@ -880,9 +877,9 @@ void FSequencer::OnScrubPositionChanged( float NewScrubPosition, bool bScrubbing
 	SetGlobalTime( NewScrubPosition );
 }
 
-void FSequencer::OnToggleAutoKey( bool bInAllowAutoKey )
+void FSequencer::OnToggleAutoKey()
 {
-	bAllowAutoKey = bInAllowAutoKey;
+	bAllowAutoKey = !bAllowAutoKey;
 }
 
 void FSequencer::OnToggleCleanView( bool bInCleanViewEnabled )
@@ -1391,6 +1388,50 @@ void FSequencer::BindSequencerCommands()
 	SequencerCommandBindings->MapAction(
 		Commands.SetKey,
 		FExecuteAction::CreateSP( this, &FSequencer::SetKey ) );
+
+	USequencerSnapSettings* SnapSettings = GetMutableDefault<USequencerSnapSettings>();
+
+	SequencerCommandBindings->MapAction(
+		Commands.ToggleAutoKeyEnabled,
+		FExecuteAction::CreateSP( this, &FSequencer::OnToggleAutoKey ),
+		FCanExecuteAction::CreateLambda( []{ return true; } ),
+		FIsActionChecked::CreateSP( this, &FSequencer::OnGetAllowAutoKey ) );
+
+	SequencerCommandBindings->MapAction(
+		Commands.ToggleIsSnapEnabled,
+		FExecuteAction::CreateLambda( [SnapSettings]{ SnapSettings->SetIsSnapEnabled( !SnapSettings->GetIsSnapEnabled() ); } ),
+		FCanExecuteAction::CreateLambda( []{ return true; } ),
+		FIsActionChecked::CreateLambda( [SnapSettings]{ return SnapSettings->GetIsSnapEnabled(); } ) );
+
+	SequencerCommandBindings->MapAction(
+		Commands.ToggleSnapKeysToInterval,
+		FExecuteAction::CreateLambda( [SnapSettings]{ SnapSettings->SetSnapKeysToInterval( !SnapSettings->GetSnapKeysToInterval() ); } ),
+		FCanExecuteAction::CreateLambda( []{ return true; } ),
+		FIsActionChecked::CreateLambda( [SnapSettings]{ return SnapSettings->GetSnapKeysToInterval(); } ) );
+
+	SequencerCommandBindings->MapAction(
+		Commands.ToggleSnapKeysToKeys,
+		FExecuteAction::CreateLambda( [SnapSettings]{ SnapSettings->SetSnapKeysToKeys( !SnapSettings->GetSnapKeysToKeys() ); } ),
+		FCanExecuteAction::CreateLambda( []{ return true; } ),
+		FIsActionChecked::CreateLambda( [SnapSettings]{ return SnapSettings->GetSnapKeysToKeys(); } ) );
+
+	SequencerCommandBindings->MapAction(
+		Commands.ToggleSnapSectionsToInterval,
+		FExecuteAction::CreateLambda( [SnapSettings]{ SnapSettings->SetSnapSectionsToInterval( !SnapSettings->GetSnapSectionsToInterval() ); } ),
+		FCanExecuteAction::CreateLambda( []{ return true; } ),
+		FIsActionChecked::CreateLambda( [SnapSettings]{ return SnapSettings->GetSnapSectionsToInterval(); } ) );
+
+	SequencerCommandBindings->MapAction(
+		Commands.ToggleSnapSectionsToSections,
+		FExecuteAction::CreateLambda( [SnapSettings]{ SnapSettings->SetSnapSectionsToSections( !SnapSettings->GetSnapSectionsToSections() ); } ),
+		FCanExecuteAction::CreateLambda( []{ return true; } ),
+		FIsActionChecked::CreateLambda( [SnapSettings]{ return SnapSettings->GetSnapSectionsToSections(); } ) );
+
+	SequencerCommandBindings->MapAction(
+		Commands.ToggleSnapPlayTimeToInterval,
+		FExecuteAction::CreateLambda( [SnapSettings]{ SnapSettings->SetSnapPlayTimeToInterval( !SnapSettings->GetSnapPlayTimeToInterval() ); } ),
+		FCanExecuteAction::CreateLambda( []{ return true; } ),
+		FIsActionChecked::CreateLambda( [SnapSettings]{ return SnapSettings->GetSnapPlayTimeToInterval(); } ) );
 }
 
 void FSequencer::BuildObjectBindingContextMenu(FMenuBuilder& MenuBuilder, const FGuid& ObjectBinding, const UClass* ObjectClass)
diff --git a/Engine/Source/Editor/Sequencer/Private/Sequencer.h b/Engine/Source/Editor/Sequencer/Private/Sequencer.h
--- a/Engine/Source/Editor/Sequencer/Private/Sequencer.h
+++ b/Engine/Source/Editor/Sequencer/Private/Sequencer.h
@@ -248,7 +248,7 @@ public:
 	bool IsSectionVisible(UMovieSceneSection* Section) const;
 
 	/** Gets the command bindings for the sequencer */
-	FUICommandList& GetCommandBindings() { return *SequencerCommandBindings; }
+	TSharedPtr<FUICommandList> GetCommandBindings() { return SequencerCommandBindings; }
 
 	/**
 	 * Builds up the context menu for object binding nodes in the outliner
@@ -352,7 +352,7 @@ protected:
 	 *
 	 * @param bInAllowAutoKey	The new auto key state
 	 */
-	void OnToggleAutoKey( bool bInAllowAutoKey );
+	void OnToggleAutoKey();
 	
 	/**
 	 * Called when auto-key is toggled by a user
diff --git a/Engine/Source/Editor/Sequencer/Private/SequencerCommands.cpp b/Engine/Source/Editor/Sequencer/Private/SequencerCommands.cpp
--- a/Engine/Source/Editor/Sequencer/Private/SequencerCommands.cpp
+++ b/Engine/Source/Editor/Sequencer/Private/SequencerCommands.cpp
@@ -6,4 +6,17 @@
 void FSequencerCommands::RegisterCommands()
 {
 	UI_COMMAND( SetKey, "Set Key", "Sets a key at the current time for the selected actor.", EUserInterfaceActionType::Button, FInputGesture(EKeys::K) );
+
+	UI_COMMAND( ToggleAutoKeyEnabled, "Auto Key", "Enables and disables auto keying.", EUserInterfaceActionType::ToggleButton, FInputGesture() );
+
+	UI_COMMAND( ToggleIsSnapEnabled, "Enable Snapping", "Enables and disables snapping.", EUserInterfaceActionType::ToggleButton, FInputGesture() );
+
+	UI_COMMAND( ToggleSnapKeysToInterval, "Snap to the interval", "Snap keys to the interval.", EUserInterfaceActionType::ToggleButton, FInputGesture() );
+	UI_COMMAND( ToggleSnapKeysToKeys, "Snap to other keys", "Snap keys to other keys in the section.", EUserInterfaceActionType::ToggleButton, FInputGesture() );
+
+	UI_COMMAND( ToggleSnapSectionsToInterval, "Snap to the interval", "Snap sections to the interval.", EUserInterfaceActionType::ToggleButton, FInputGesture() );
+	UI_COMMAND( ToggleSnapSectionsToSections, "Snap to other sections", "Snap sections to other sections.", EUserInterfaceActionType::ToggleButton, FInputGesture() );
+
+	UI_COMMAND( ToggleSnapPlayTimeToInterval, "Snap to the interval while scrubbing", "Snap the play time to the interval while scrubbing.", EUserInterfaceActionType::ToggleButton, FInputGesture() );
+	
 }
diff --git a/Engine/Source/Editor/Sequencer/Private/SequencerCommands.h b/Engine/Source/Editor/Sequencer/Private/SequencerCommands.h
--- a/Engine/Source/Editor/Sequencer/Private/SequencerCommands.h
+++ b/Engine/Source/Editor/Sequencer/Private/SequencerCommands.h
@@ -18,6 +18,27 @@ public:
 	/** Sets a key at the current time for the selected actor */
 	TSharedPtr< FUICommandInfo > SetKey;
 
+	/** Turns auto keying on and off. */
+	TSharedPtr< FUICommandInfo > ToggleAutoKeyEnabled;
+
+	/** Turns snapping on and off. */
+	TSharedPtr< FUICommandInfo > ToggleIsSnapEnabled;
+
+	/** Toggles whether or not keys should snap to the selected interval. */
+	TSharedPtr< FUICommandInfo > ToggleSnapKeysToInterval;
+
+	/** Toggles whether or not keys should snap to other keys in the section. */
+	TSharedPtr< FUICommandInfo > ToggleSnapKeysToKeys;
+
+	/** Toggles whether or not sections should snap to the selected interval. */
+	TSharedPtr< FUICommandInfo > ToggleSnapSectionsToInterval;
+
+	/** Toggles whether or not sections should snap to other sections. */
+	TSharedPtr< FUICommandInfo > ToggleSnapSectionsToSections;
+
+	/** Toggles whether or not the play time should snap to the selected interval. */
+	TSharedPtr< FUICommandInfo > ToggleSnapPlayTimeToInterval;
+
 	/**
 	 * Initialize commands
 	 */
diff --git a/Engine/Source/Editor/Sequencer/Private/SequencerPrivatePCH.h b/Engine/Source/Editor/Sequencer/Private/SequencerPrivatePCH.h
--- a/Engine/Source/Editor/Sequencer/Private/SequencerPrivatePCH.h
+++ b/Engine/Source/Editor/Sequencer/Private/SequencerPrivatePCH.h
@@ -20,6 +20,7 @@
 #include "SSequencerTrackViewPanel.h"
 #include "SAnimationOutlinerView.h"
 #include "MovieSceneInstance.h"
+#include "SequencerSnapSettings.h"
 
 DECLARE_LOG_CATEGORY_EXTERN(LogSequencer, Log, All);
 
diff --git a/Engine/Source/Editor/Sequencer/Private/SequencerSnapSettings.cpp b/Engine/Source/Editor/Sequencer/Private/SequencerSnapSettings.cpp
new file mode 100644
--- /dev/null
+++ b/Engine/Source/Editor/Sequencer/Private/SequencerSnapSettings.cpp
@@ -0,0 +1,118 @@
+// Copyright 1998-2014 Epic Games, Inc. All Rights Reserved.
+
+#include "SequencerPrivatePCH.h"
+
+USequencerSnapSettings::USequencerSnapSettings( const FObjectInitializer& ObjectInitializer )
+	: Super( ObjectInitializer )
+{
+	bIsEnabled = true;
+	SnapInterval = .05f;
+	bSnapKeysToInterval = true;
+	bSnapKeysToKeys = true;
+	bSnapSectionsToInterval = true;
+	bSnapSectionsToSections = true;
+	bSnapPlayTimeToInterval = true;
+}
+
+bool USequencerSnapSettings::GetIsSnapEnabled() const
+{
+	return bIsEnabled;
+}
+
+void USequencerSnapSettings::SetIsSnapEnabled(bool InbIsEnabled)
+{
+	if ( bIsEnabled != InbIsEnabled )
+	{
+		bIsEnabled = InbIsEnabled;
+		SaveConfig();
+	}
+}
+
+float USequencerSnapSettings::GetSnapInterval() const
+{
+	return SnapInterval;
+}
+
+void USequencerSnapSettings::SetSnapInterval(float InSnapInterval)
+{
+	if ( SnapInterval != InSnapInterval )
+	{
+		SnapInterval = InSnapInterval;
+		SaveConfig();
+	}
+}
+
+bool USequencerSnapSettings::GetSnapKeysToInterval() const
+{
+	return bSnapKeysToInterval;
+}
+
+void USequencerSnapSettings::SetSnapKeysToInterval(bool InbSnapKeysToInterval)
+{
+	if ( bSnapKeysToInterval != InbSnapKeysToInterval )
+	{
+		bSnapKeysToInterval = InbSnapKeysToInterval;
+		SaveConfig();
+	}
+}
+
+bool USequencerSnapSettings::GetSnapKeysToKeys() const
+{
+	return bSnapKeysToKeys;
+}
+
+void USequencerSnapSettings::SetSnapKeysToKeys(bool InbSnapKeysToKeys)
+{
+	if ( bSnapKeysToKeys != InbSnapKeysToKeys )
+	{
+		bSnapKeysToKeys = InbSnapKeysToKeys;
+		SaveConfig();
+	}
+}
+
+bool USequencerSnapSettings::GetSnapSectionsToInterval() const
+{
+	return bSnapSectionsToInterval;
+}
+
+void USequencerSnapSettings::SetSnapSectionsToInterval(bool InbSnapSectionsToInterval)
+{
+	if ( bSnapSectionsToInterval != InbSnapSectionsToInterval )
+	{
+		bSnapSectionsToInterval = InbSnapSectionsToInterval;
+		SaveConfig();
+	}
+}
+
+bool USequencerSnapSettings::GetSnapSectionsToSections() const
+{
+	return bSnapSectionsToSections;
+}
+
+void USequencerSnapSettings::SetSnapSectionsToSections( bool InbSnapSectionsToSections )
+{
+	if ( bSnapSectionsToSections != InbSnapSectionsToSections )
+	{
+		bSnapSectionsToSections = InbSnapSectionsToSections;
+		SaveConfig();
+	}
+}
+
+bool USequencerSnapSettings::GetSnapPlayTimeToInterval() const
+{
+	return bSnapPlayTimeToInterval;
+}
+
+void USequencerSnapSettings::SetSnapPlayTimeToInterval(bool InbSnapPlayTimeToInterval)
+{
+	if ( bSnapPlayTimeToInterval != InbSnapPlayTimeToInterval )
+	{
+		bSnapPlayTimeToInterval = InbSnapPlayTimeToInterval;
+		SaveConfig();
+	}
+}
+
+float USequencerSnapSettings::SnapToInterval(float InValue) const
+{
+	return FMath::RoundToInt(InValue / SnapInterval) * SnapInterval;
+}
diff --git a/Engine/Source/Editor/Sequencer/Private/SequencerSnapSettings.h b/Engine/Source/Editor/Sequencer/Private/SequencerSnapSettings.h
new file mode 100644
--- /dev/null
+++ b/Engine/Source/Editor/Sequencer/Private/SequencerSnapSettings.h
@@ -0,0 +1,73 @@
+// Copyright 1998-2014 Epic Games, Inc. All Rights Reserved.
+
+#pragma once
+
+#include "SequencerSnapSettings.generated.h"
+
+/** Serializable snapping options for sequencer. */
+UCLASS(config=EditorUserSettings)
+class USequencerSnapSettings : public UObject
+{
+public:
+	GENERATED_UCLASS_BODY()
+
+	/** Gets whether or not snapping is enabled. */
+	bool GetIsSnapEnabled() const;
+	/** Sets whether or not snapping is enabled. */
+	void SetIsSnapEnabled(bool InbIsEnabled);
+
+	/** Gets the time in seconds used for interval snapping. */
+	float GetSnapInterval() const;
+	/** Sets the time in seconds used for interval snapping. */
+	void SetSnapInterval(float InSnapInterval);
+
+	/** Gets whether or not to snap keys to the interval. */
+	bool GetSnapKeysToInterval() const;
+	/** Sets whether or not to snap keys to the interval. */
+	void SetSnapKeysToInterval(bool InbSnapKeysToInterval);
+
+	/** Gets whether or not to snap keys to other keys. */
+	bool GetSnapKeysToKeys() const;
+	/** Sets whether or not to snap keys to other keys. */
+	void SetSnapKeysToKeys(bool InbSnapKeysToKeys);
+
+	/** Gets whether or not to snap sections to the interval. */
+	bool GetSnapSectionsToInterval() const;
+	/** Sets whether or not to snap sections to the interval. */
+	void SetSnapSectionsToInterval(bool InbSnapSectionsToInterval);
+
+	/** Gets whether or not to snap sections to other sections. */
+	bool GetSnapSectionsToSections() const;
+	/** sets whether or not to snap sections to other sections. */
+	void SetSnapSectionsToSections( bool InbSnapSectionsToSections );
+
+	/** Gets whether or not to snap the play time to the interval while scrubbing. */
+	bool GetSnapPlayTimeToInterval() const;
+	/** Sets whether or not to snap the play time to the interval while scrubbing. */
+	void SetSnapPlayTimeToInterval(bool InbSnapPlayTimeToInterval);
+
+	/** Snaps a time value in seconds to the currently selected interval. */
+	float SnapToInterval( float InValue ) const;
+
+protected:
+	UPROPERTY( config )
+	bool bIsEnabled;
+
+	UPROPERTY( config )
+	float SnapInterval;
+
+	UPROPERTY( config )
+	bool bSnapKeysToInterval;
+
+	UPROPERTY( config )
+	bool bSnapKeysToKeys;
+
+	UPROPERTY( config )
+	bool bSnapSectionsToInterval;
+
+	UPROPERTY( config )
+	bool bSnapSectionsToSections;
+
+	UPROPERTY( config )
+	bool bSnapPlayTimeToInterval;
+};
\ No newline at end of file
diff --git a/Engine/Source/Editor/Sequencer/Private/TimeSliderController.cpp b/Engine/Source/Editor/Sequencer/Private/TimeSliderController.cpp
--- a/Engine/Source/Editor/Sequencer/Private/TimeSliderController.cpp
+++ b/Engine/Source/Editor/Sequencer/Private/TimeSliderController.cpp
@@ -334,6 +334,12 @@ FReply FSequencerTimeSliderController::OnMouseButtonUp( TSharedRef<SWidget> Widg
 			FVector2D CursorPos = MyGeometry.AbsoluteToLocal(MouseEvent.GetLastScreenSpacePosition());
 			float NewValue = RangeToScreen.LocalXToInput(CursorPos.X);
 
+			const USequencerSnapSettings* SnapSettings = GetDefault<USequencerSnapSettings>();
+			if ( SnapSettings->GetIsSnapEnabled() && SnapSettings->GetSnapPlayTimeToInterval() )
+			{
+				NewValue = SnapSettings->SnapToInterval( NewValue );
+			}
+
 			CommitScrubPosition( NewValue, /*bIsScrubbing=*/false );
 		}
 
@@ -413,6 +419,12 @@ FReply FSequencerTimeSliderController::OnMouseMove( TSharedRef<SWidget> WidgetOw
 				FVector2D CursorPos = MyGeometry.AbsoluteToLocal( MouseEvent.GetLastScreenSpacePosition() );
 				float NewValue = RangeToScreen.LocalXToInput( CursorPos.X );
 
+				const USequencerSnapSettings* SnapSettings = GetDefault<USequencerSnapSettings>();
+				if ( SnapSettings->GetIsSnapEnabled() && SnapSettings->GetSnapPlayTimeToInterval() )
+				{
+					NewValue = SnapSettings->SnapToInterval(NewValue);
+				}
+
 				CommitScrubPosition( NewValue, /*bIsScrubbing=*/true );
 			}
 		}
