Sequencer - Snap fixes.
+ Tweak the MaxPixelsPerInput value to work better when snapping is turned on.
+ When the snap interval is set to 0, don't try to snap.

diff --git a/Engine/Source/Developer/MovieSceneTools/Public/CommonMovieSceneTools.h b/Engine/Source/Developer/MovieSceneTools/Public/CommonMovieSceneTools.h
--- a/Engine/Source/Developer/MovieSceneTools/Public/CommonMovieSceneTools.h
+++ b/Engine/Source/Developer/MovieSceneTools/Public/CommonMovieSceneTools.h
@@ -9,7 +9,7 @@ struct FTimeToPixel
 {
 public:
 	FTimeToPixel( const FGeometry& AllottedGeometry, TRange<float> InLocalViewRange )
-		: MaxPixelsPerInput(10000.0f)
+		: MaxPixelsPerInput(1000.0f)
 		, LocalViewRange( InLocalViewRange )
 	{
 		const float ViewRange = LocalViewRange.Size<float>();
diff --git a/Engine/Source/Editor/Sequencer/Private/SequencerSnapSettings.cpp b/Engine/Source/Editor/Sequencer/Private/SequencerSnapSettings.cpp
--- a/Engine/Source/Editor/Sequencer/Private/SequencerSnapSettings.cpp
+++ b/Engine/Source/Editor/Sequencer/Private/SequencerSnapSettings.cpp
@@ -114,5 +114,7 @@ void USequencerSnapSettings::SetSnapPlayTimeToInterval(bool InbSnapPlayTimeToInt
 
 float USequencerSnapSettings::SnapToInterval(float InValue) const
 {
-	return FMath::RoundToInt(InValue / SnapInterval) * SnapInterval;
+	return SnapInterval > 0 
+		? FMath::RoundToInt(InValue / SnapInterval) * SnapInterval
+		: InValue;
 }
