Sequencer - Fix the build from curve editor change.

diff --git a/Engine/Source/Editor/Sequencer/Private/SSequencerCurveEditor.cpp b/Engine/Source/Editor/Sequencer/Private/SSequencerCurveEditor.cpp
--- a/Engine/Source/Editor/Sequencer/Private/SSequencerCurveEditor.cpp
+++ b/Engine/Source/Editor/Sequencer/Private/SSequencerCurveEditor.cpp
@@ -3,7 +3,7 @@
 #include "SequencerPrivatePCH.h"
 #include "SCurveEditor.h"
 
-#define LOCTEXT_NAMESPACE "SequencerTrackArea"
+#define LOCTEXT_NAMESPACE "SequencerCurveEditor"
 
 void SSequencerCurveEditor::Construct( const FArguments& InArgs, TSharedRef<FSequencer> InSequencer )
 {
@@ -112,4 +112,6 @@ SSequencerCurveEditor::~SSequencerCurveEditor()
 	{
 		SequencerNodeTree->GetOnSelectionChanged()->Remove( NodeTreeSelectionChangedHandle );
 	}
-}
\ No newline at end of file
+}
+
+#undef LOCTEXT_NAMESPACE
\ No newline at end of file
diff --git a/Engine/Source/Editor/Sequencer/Private/SSequencerCurveEditor.h b/Engine/Source/Editor/Sequencer/Private/SSequencerCurveEditor.h
--- a/Engine/Source/Editor/Sequencer/Private/SSequencerCurveEditor.h
+++ b/Engine/Source/Editor/Sequencer/Private/SSequencerCurveEditor.h
@@ -16,7 +16,7 @@ public:
 
 	SLATE_END_ARGS()
 
-	void SSequencerCurveEditor::Construct( const FArguments& InArgs, TSharedRef<FSequencer> InSequencer );
+	void Construct( const FArguments& InArgs, TSharedRef<FSequencer> InSequencer );
 
 	/** Sets the sequencer node tree which supplies the curves. */
 	void SetSequencerNodeTree( TSharedPtr<FSequencerNodeTree> InSequencerNodeTree );
