Fix non-unity compile

diff --git a/Engine/Source/Editor/Sequencer/Private/SequencerCurveOwner.cpp b/Engine/Source/Editor/Sequencer/Private/SequencerCurveOwner.cpp
--- a/Engine/Source/Editor/Sequencer/Private/SequencerCurveOwner.cpp
+++ b/Engine/Source/Editor/Sequencer/Private/SequencerCurveOwner.cpp
@@ -2,6 +2,8 @@
 
 #include "SequencerPrivatePCH.h"
 #include "SequencerCurveOwner.h"
+#include "IKeyArea.h"
+#include "MovieSceneSection.h"
 
 void GetAllKeyAreaNodes( TSharedPtr<FSequencerNodeTree> InSequencerNodeTree, TArray<TSharedRef<FSectionKeyAreaNode>>& KeyAreaNodes )
 {
