Sequencer - Add transport controls.

diff --git a/Engine/Source/Editor/Sequencer/Private/SSequencer.cpp b/Engine/Source/Editor/Sequencer/Private/SSequencer.cpp
--- a/Engine/Source/Editor/Sequencer/Private/SSequencer.cpp
+++ b/Engine/Source/Editor/Sequencer/Private/SSequencer.cpp
@@ -24,6 +24,7 @@
 #include "CommonMovieSceneTools.h"
 #include "SSearchBox.h"
 #include "SNumericDropDown.h"
+#include "EditorWidgetsModule.h"
 
 
 #define LOCTEXT_NAMESPACE "Sequencer"
@@ -263,8 +264,9 @@ void SSequencer::Construct( const FArguments& InArgs, TSharedRef< class FSequenc
 				SNew( SHorizontalBox )
 				+ SHorizontalBox::Slot()
 				.FillWidth( TAttribute<float>( this, &SSequencer::GetAnimationOutlinerFillPercentage ) )
+				.HAlign(HAlign_Center)
 				[
-					SNew( SSpacer )
+					MakeTransportControls()
 				]
 				+ SHorizontalBox::Slot()
 				.FillWidth(1.0f)
@@ -433,6 +435,24 @@ TSharedRef<SWidget> SSequencer::MakeCurveEditorMenu()
 	return MenuBuilder.MakeWidget();
 }
 
+TSharedRef<SWidget> SSequencer::MakeTransportControls()
+{
+	FEditorWidgetsModule& EditorWidgetsModule = FModuleManager::Get().LoadModuleChecked<FEditorWidgetsModule>( "EditorWidgets" );
+	TSharedRef<FSequencer> SequencerPinned = Sequencer.Pin().ToSharedRef();
+
+	FTransportControlArgs TransportControlArgs;
+	TransportControlArgs.OnBackwardEnd.BindSP( SequencerPinned, &FSequencer::OnStepToBeginning );
+	TransportControlArgs.OnBackwardStep.BindSP( SequencerPinned, &FSequencer::OnStepBackward );
+	TransportControlArgs.OnForwardPlay.BindSP( SequencerPinned, &FSequencer::OnPlay );
+	TransportControlArgs.OnForwardStep.BindSP( SequencerPinned, &FSequencer::OnStepForward );
+	TransportControlArgs.OnForwardEnd.BindSP( SequencerPinned, &FSequencer::OnStepToEnd );
+	TransportControlArgs.OnToggleLooping.BindSP( SequencerPinned, &FSequencer::OnToggleLooping );
+	TransportControlArgs.OnGetLooping.BindSP( SequencerPinned, &FSequencer::IsLooping );
+	TransportControlArgs.OnGetPlaybackMode.BindSP( SequencerPinned, &FSequencer::GetPlaybackMode );
+
+	return EditorWidgetsModule.CreateTransportControl( TransportControlArgs );
+}
+
 SSequencer::~SSequencer()
 {
 }
diff --git a/Engine/Source/Editor/Sequencer/Private/SSequencer.h b/Engine/Source/Editor/Sequencer/Private/SSequencer.h
--- a/Engine/Source/Editor/Sequencer/Private/SSequencer.h
+++ b/Engine/Source/Editor/Sequencer/Private/SSequencer.h
@@ -100,6 +100,9 @@ private:
 	/** Makes the curve editor menu for the toolbar. */
 	TSharedRef<SWidget> MakeCurveEditorMenu();
 
+	/** Makes and configures a set of the standard UE transport controls. */
+	TSharedRef<SWidget> MakeTransportControls();
+
 	/**
 	* @return The value of the current time snap interval.
 	*/
diff --git a/Engine/Source/Runtime/MovieSceneCore/Classes/MovieSceneSection.h b/Engine/Source/Runtime/MovieSceneCore/Classes/MovieSceneSection.h
--- a/Engine/Source/Runtime/MovieSceneCore/Classes/MovieSceneSection.h
+++ b/Engine/Source/Runtime/MovieSceneCore/Classes/MovieSceneSection.h
@@ -45,7 +45,11 @@ public:
 	/**
 	 * @return The range of times of the section
 	 */
-	TRange<float> GetRange() const {return TRange<float>(StartTime, EndTime);}
+	TRange<float> GetRange() const 
+	{
+		// Uses an inclusive range so that sections that start and end on the same value aren't considered empty ranges.
+		return TRange<float>( TRange<float>::BoundsType::Inclusive( StartTime ), TRange<float>::BoundsType::Inclusive( EndTime ) );
+	}
 	
 	/**
 	 * Sets a new range of times for this section
