Sequencer - Fix an issue where an animation would become unusable if you undid the first keyframe addition.

Jira: UE-6365

diff --git a/Engine/Source/Editor/Sequencer/Private/Sequencer.cpp b/Engine/Source/Editor/Sequencer/Private/Sequencer.cpp
--- a/Engine/Source/Editor/Sequencer/Private/Sequencer.cpp
+++ b/Engine/Source/Editor/Sequencer/Private/Sequencer.cpp
@@ -529,6 +529,15 @@ FGuid FSequencer::GetHandleToObject( UObject* Object )
 	UMovieScene* FocusedMovieScene = FocusedMovieSceneInstance->GetMovieScene();
 
 	FGuid ObjectGuid = ObjectBindingManager->FindGuidForObject( *FocusedMovieScene, *Object );
+
+	// Make sure that the possessable is still valid, if it's not remove the binding so new one 
+	// can be created.  This can happen due to undo.
+	FMovieScenePossessable* Possessable = FocusedMovieScene->FindPossessable(ObjectGuid);
+	if ( Possessable == nullptr )
+	{
+		ObjectBindingManager->UnbindPossessableObjects(ObjectGuid);
+		ObjectGuid.Invalidate();
+	}
 	
 	bool bPossessableAdded = false;
 	
@@ -536,19 +545,15 @@ FGuid FSequencer::GetHandleToObject( UObject* Object )
 	// Note: Only possessed actors can be added like this
 	if( !ObjectGuid.IsValid() && ObjectBindingManager->CanPossessObject( *Object ) )
 	{
-		// Grab the MovieScene that is currently focused.  We'll add our Blueprint as an inner of the
-		// MovieScene asset.
-		UMovieScene* OwnerMovieScene = GetFocusedMovieScene();
-		
 		// @todo sequencer: Undo doesn't seem to be working at all
 		const FScopedTransaction Transaction( LOCTEXT("UndoPossessingObject", "Possess Object with MovieScene") );
 		
 		// Possess the object!
 		{
 			// Create a new possessable
-			OwnerMovieScene->Modify();
-			
-			ObjectGuid = OwnerMovieScene->AddPossessable( Object->GetName(), Object->GetClass() );
+			FocusedMovieScene->Modify();
+
+			ObjectGuid = FocusedMovieScene->AddPossessable( Object->GetName(), Object->GetClass() );
 			
 			if ( IsShotFilteringOn() )
 			{
