Sequencer - Fix curve name generation so that it's not backwards.

diff --git a/Engine/Source/Editor/Sequencer/Private/SequencerCurveOwner.cpp b/Engine/Source/Editor/Sequencer/Private/SequencerCurveOwner.cpp
--- a/Engine/Source/Editor/Sequencer/Private/SequencerCurveOwner.cpp
+++ b/Engine/Source/Editor/Sequencer/Private/SequencerCurveOwner.cpp
@@ -35,7 +35,7 @@ FName BuildCurveName( TSharedPtr<FSectionKeyAreaNode> KeyAreaNode )
 	TArray<FString> NameParts;
 	while ( CurrentNameNode.IsValid() )
 	{
-		NameParts.Add( CurrentNameNode->GetDisplayName().ToString() );
+		NameParts.Insert( CurrentNameNode->GetDisplayName().ToString(), 0);
 		CurrentNameNode = CurrentNameNode->GetParent();
 	}
 	return FName(*FString::Join(NameParts, TEXT(" - ")));
