Sequencer - Fix a hang when zooming to a zero sized section.
+ Guard agains an infinate loop in the time slider controller.
+ Don't zoom to zero sized sections in Sequencer.

Fixes UE-11677

diff --git a/Engine/Source/Editor/Sequencer/Private/Sequencer.cpp b/Engine/Source/Editor/Sequencer/Private/Sequencer.cpp
--- a/Engine/Source/Editor/Sequencer/Private/Sequencer.cpp
+++ b/Engine/Source/Editor/Sequencer/Private/Sequencer.cpp
@@ -1186,7 +1186,7 @@ void FSequencer::ZoomToSelectedSections()
 		BoundsHull = GetTimeBounds();
 	}
 
-	if (!BoundsHull.IsEmpty())
+	if (!BoundsHull.IsEmpty() && !BoundsHull.IsDegenerate())
 	{
 		OnViewRangeChanged(BoundsHull, true);
 	}
diff --git a/Engine/Source/Editor/Sequencer/Private/TimeSliderController.cpp b/Engine/Source/Editor/Sequencer/Private/TimeSliderController.cpp
--- a/Engine/Source/Editor/Sequencer/Private/TimeSliderController.cpp
+++ b/Engine/Source/Editor/Sequencer/Private/TimeSliderController.cpp
@@ -83,10 +83,13 @@ float DetermineOptimalSpacing( float InPixelsPerInput, uint32 MinTick, float Min
 	// Start with the smallest spacing
 	float Spacing = MinTickSpacing;
 
-	while( Spacing * InPixelsPerInput < MinTick )
+	if (InPixelsPerInput > 0)
 	{
-		Spacing = MinTickSpacing * GetNextSpacing( CurStep );
-		CurStep++;
+		while( Spacing * InPixelsPerInput < MinTick )
+		{
+			Spacing = MinTickSpacing * GetNextSpacing( CurStep );
+			CurStep++;
+		}
 	}
 
 	return Spacing;
