Sequencer - Add an object to mange various selections in sequencer.
+ Add the FSequencerSelection class and move all selection functionality there.
+ Add the concept of an active selection which is used to determine what to delete when the user hits the delete key.

diff --git a/Engine/Source/Editor/EditorStyle/Private/SlateEditorStyle.cpp b/Engine/Source/Editor/EditorStyle/Private/SlateEditorStyle.cpp
--- a/Engine/Source/Editor/EditorStyle/Private/SlateEditorStyle.cpp
+++ b/Engine/Source/Editor/EditorStyle/Private/SlateEditorStyle.cpp
@@ -1569,6 +1569,7 @@ void FSlateEditorStyle::FStyle::SetupGeneralStyles()
 		Set( "Sequencer.SectionGripRight", new BOX_BRUSH( "Sequencer/SectionGripRight", FMargin(5.f/16.f) ) );
 		Set( "Sequencer.SectionArea.Background", new FSlateColorBrush( FColor::White ) );
 		Set( "Sequencer.AnimationOutliner.SelectionBorder", new BOX_BRUSH( "Common/RoundedSelection_16x", 4.0f/16.0f, SelectionColor_Pressed ) );
+		Set( "Sequencer.AnimationOutliner.SelectionBorderInactive", new BOX_BRUSH("Common/RoundedSelection_16x", 4.0f / 16.0f, SelectionColor_Inactive));
 		Set( "Sequencer.AnimationOutliner.BoldFont", TTF_CORE_FONT( "Fonts/Roboto-Bold", 10 ) );
 		Set( "Sequencer.AnimationOutliner.RegularFont", TTF_CORE_FONT( "Fonts/Roboto-Regular", 8 ) );
 		Set( "Sequencer.ShotFilter", new IMAGE_BRUSH( "Sequencer/FilteredArea", FVector2D(74,74), FLinearColor::White, ESlateBrushTileType::Both ) );
diff --git a/Engine/Source/Editor/Sequencer/Private/SAnimationOutlinerView.cpp b/Engine/Source/Editor/Sequencer/Private/SAnimationOutlinerView.cpp
--- a/Engine/Source/Editor/Sequencer/Private/SAnimationOutlinerView.cpp
+++ b/Engine/Source/Editor/Sequencer/Private/SAnimationOutlinerView.cpp
@@ -17,6 +17,7 @@ void SAnimationOutlinerTreeNode::Construct( const FArguments& InArgs, TSharedRef
 	OnSelectionChanged = InArgs._OnSelectionChanged;
 
 	SelectedBrush = FEditorStyle::GetBrush( "Sequencer.AnimationOutliner.SelectionBorder" );
+	SelectedBrushInactive = FEditorStyle::GetBrush("Sequencer.AnimationOutliner.SelectionBorderInactive");
 	NotSelectedBrush = FEditorStyle::GetBrush( "NoBorder" );
 	ExpandedBrush = FEditorStyle::GetBrush( "TreeArrow_Expanded" );
 	CollapsedBrush = FEditorStyle::GetBrush( "TreeArrow_Collapsed" );
@@ -77,19 +78,19 @@ FReply SAnimationOutlinerTreeNode::OnMouseButtonDown( const FGeometry& MyGeometr
 {
 	if( MouseEvent.GetEffectingButton() == EKeys::LeftMouseButton && DisplayNode->IsSelectable() )
 	{
-		bool bSelected = DisplayNode->IsSelected();
+		FSequencer& Sequencer = DisplayNode->GetSequencer();
+		bool bSelected = Sequencer.GetSelection()->IsSelected(DisplayNode.ToSharedRef());
 
 		if( MouseEvent.IsControlDown() )
 		{
-			const bool bDeselectOtherNodes = false;
 			// Select the node if we were clicked on
-			DisplayNode->SetSelectionState( !bSelected, bDeselectOtherNodes );
+			Sequencer.GetSelection()->AddToSelection(DisplayNode.ToSharedRef());
 		}
 		else
 		{
-			const bool bDeselectOtherNodes = true;
-			// Select the node if we were clicked on
-			DisplayNode->SetSelectionState( true, bDeselectOtherNodes );
+			// Deselect the other nodes and select this node.
+			Sequencer.GetSelection()->EmptySelectedOutlinerNodes();
+			Sequencer.GetSelection()->AddToSelection(DisplayNode.ToSharedRef());
 		}
 
 		OnSelectionChanged.ExecuteIfBound( DisplayNode );
@@ -133,10 +134,25 @@ FReply SAnimationOutlinerTreeNode::OnExpanderClicked()
 }
 
 const FSlateBrush* SAnimationOutlinerTreeNode::GetNodeBorderImage() const
-{		
+{
 	// Display a highlight when the node is selected
-	const bool bIsSelected = DisplayNode->IsSelected();
-	return bIsSelected ? SelectedBrush : NotSelectedBrush;
+	FSequencer& Sequencer = DisplayNode->GetSequencer();
+	const bool bIsSelected = Sequencer.GetSelection()->IsSelected(DisplayNode.ToSharedRef());
+	if (bIsSelected)
+	{
+		if (Sequencer.GetSelection()->GetActiveSelection() == FSequencerSelection::EActiveSelection::OutlinerNode)
+		{
+			return SelectedBrush;
+		}
+		else
+		{
+			return SelectedBrushInactive;
+		}
+	}
+	else
+	{
+		return  NotSelectedBrush;
+	}
 }
 
 const FSlateBrush* SAnimationOutlinerTreeNode::OnGetExpanderImage() const
diff --git a/Engine/Source/Editor/Sequencer/Private/SAnimationOutlinerView.h b/Engine/Source/Editor/Sequencer/Private/SAnimationOutlinerView.h
--- a/Engine/Source/Editor/Sequencer/Private/SAnimationOutlinerView.h
+++ b/Engine/Source/Editor/Sequencer/Private/SAnimationOutlinerView.h
@@ -65,6 +65,8 @@ private:
 	FOnNodeSelectionChanged OnSelectionChanged;
 	/** Brush to display a border around the widget when it is selected */
 	const FSlateBrush* SelectedBrush;
+	/** Brush to display a border around the widget when it is selected but inactive */
+	const FSlateBrush* SelectedBrushInactive;
 	/** Brush to use if the node is not selected */
 	const FSlateBrush* NotSelectedBrush;
 	/** Brush to use for the expander image when this node is expanded */
diff --git a/Engine/Source/Editor/Sequencer/Private/SSection.cpp b/Engine/Source/Editor/Sequencer/Private/SSection.cpp
--- a/Engine/Source/Editor/Sequencer/Private/SSection.cpp
+++ b/Engine/Source/Editor/Sequencer/Private/SSection.cpp
@@ -179,7 +179,7 @@ void SSection::CreateDragOperation( const FGeometry& MyGeometry, const FPointerE
 
 	if( bKeysUnderMouse )
 	{
-		DragOperation = MakeShareable( new FMoveKeys( GetSequencer().GetSelectedKeys(), PressedKey ) );
+		DragOperation = MakeShareable( new FMoveKeys( GetSequencer().GetSelection()->GetSelectedKeys(), PressedKey ) );
 	}
 	else
 	{
@@ -258,7 +258,7 @@ void SSection::PaintKeys( const FGeometry& AllottedGeometry, const FSlateRect& M
 {
 	UMovieSceneSection& SectionObject = *SectionInterface->GetSectionObject();
 
-	const FSequencer& Sequencer = ParentSectionArea->GetSequencer();
+	FSequencer& Sequencer = ParentSectionArea->GetSequencer();
 
 	static const FName BackgroundBrushName("Sequencer.SectionArea.Background");
 	static const FName KeyBrushName("Sequencer.Key");
@@ -268,10 +268,13 @@ void SSection::PaintKeys( const FGeometry& AllottedGeometry, const FSlateRect& M
 	const FSlateBrush* KeyBrush = FEditorStyle::GetBrush(KeyBrushName);
 
 	static const FName SelectionColorName("SelectionColor");
+	static const FName SelectionInactiveColorName("SelectionColorInactive");
 	static const FName SelectionColorPressedName("SelectionColor_Pressed");
 
 	const FLinearColor PressedKeyColor = FEditorStyle::GetSlateColor(SelectionColorPressedName).GetColor( InWidgetStyle );
 	const FLinearColor SelectedKeyColor = FEditorStyle::GetSlateColor(SelectionColorName).GetColor( InWidgetStyle );
+	const FLinearColor SelectedInactiveColor = FEditorStyle::GetSlateColor(SelectionInactiveColorName).GetColor( InWidgetStyle )
+		* FLinearColor(.25, .25, .25, 1);  // Make the color a little darker since it's not very visible next to white keyframes.
 
 	// @todo Sequencer temp color, make hovered brighter than selected.
 	FLinearColor HoveredKeyColor = SelectedKeyColor * FLinearColor(1.5,1.5,1.5,1.0f);
@@ -318,7 +321,8 @@ void SSection::PaintKeys( const FGeometry& AllottedGeometry, const FSlateRect& M
 
 				FSelectedKey TestKey( SectionObject, KeyArea, KeyHandle );
 
-				bool bSelected = Sequencer.IsKeySelected( TestKey );
+				bool bSelected = Sequencer.GetSelection()->IsSelected( TestKey );
+				bool bActive = Sequencer.GetSelection()->GetActiveSelection() == FSequencerSelection::EActiveSelection::KeyAndSection;
 
 				if( TestKey == PressedKey )
 				{
@@ -330,7 +334,14 @@ void SSection::PaintKeys( const FGeometry& AllottedGeometry, const FSlateRect& M
 				}
 				else if( bSelected )
 				{
-					KeyColor = SelectedKeyColor;
+					if (bActive)
+					{
+						KeyColor = SelectedKeyColor;
+					}
+					else
+					{
+						KeyColor = SelectedInactiveColor;
+					}
 				}
 
 				// Draw the key
@@ -354,11 +365,15 @@ void SSection::DrawSectionBorders( const FGeometry& AllottedGeometry, const FSla
 {
 	UMovieSceneSection* SectionObject = SectionInterface->GetSectionObject();
 
-	const bool bSelected = ParentSectionArea->GetSequencer().IsSectionSelected(SectionObject);
+	FSequencerSelection* Selection = ParentSectionArea->GetSequencer().GetSelection();
+	const bool bSelected = Selection->IsSelected(SectionObject);
+	const bool bActive = Selection->GetActiveSelection() == FSequencerSelection::EActiveSelection::KeyAndSection;
 
 	static const FName SelectionColorName("SelectionColor");
+	static const FName SelectionInactiveColorName("SelectionColorInactive");
 
 	FLinearColor SelectionColor = FEditorStyle::GetSlateColor(SelectionColorName).GetColor(FWidgetStyle());
+	FLinearColor SelectionInactiveColor = FEditorStyle::GetSlateColor(SelectionInactiveColorName).GetColor(FWidgetStyle());
 	FLinearColor TransparentSelectionColor = SelectionColor;
 
 	static const FName SectionGripLeftName("Sequencer.SectionGripLeft");
@@ -401,7 +416,7 @@ void SSection::DrawSectionBorders( const FGeometry& AllottedGeometry, const FSla
 			FEditorStyle::GetBrush(PlainBorder),
 			MyClippingRect,
 			ESlateDrawEffect::None,
-			SelectionColor
+			bActive ? SelectionColor : SelectionInactiveColor
 			);
 	}
 
@@ -601,7 +616,7 @@ FReply SSection::OnMouseMove( const FGeometry& MyGeometry, const FPointerEvent&
 					if( PressedKey.IsValid() )
 					{
 						// Clear selected sections when beginning to drag keys
-						GetSequencer().ClearSectionSelection();
+						GetSequencer().GetSelection()->EmptySelectedSections();
 
 						bool bSelectDueToDrag = true;
 						HandleKeySelection( PressedKey, MouseEvent, bSelectDueToDrag );
@@ -612,7 +627,7 @@ FReply SSection::OnMouseMove( const FGeometry& MyGeometry, const FPointerEvent&
 					else
 					{
 						// Clear selected keys when beginning to drag a section
-						GetSequencer().ClearKeySelection();
+						GetSequencer().GetSelection()->EmptySelectedKeys();
 
 						HandleSectionSelection( MouseEvent );
 
@@ -706,13 +721,13 @@ void SSection::HandleKeySelection( const FSelectedKey& Key, const FPointerEvent&
 	{
 		// Clear previous key selection if:
 		// we are selecting due to drag and the key being dragged is not selected or control is not down
-		bool bShouldClearSelectionDueToDrag =  bSelectDueToDrag ? !GetSequencer().IsKeySelected( Key ) : true;
+		bool bShouldClearSelectionDueToDrag =  bSelectDueToDrag ? !GetSequencer().GetSelection()->IsSelected( Key ) : true;
 		
 		if( !MouseEvent.IsControlDown() && bShouldClearSelectionDueToDrag )
 		{
-			GetSequencer().ClearKeySelection();
+			GetSequencer().GetSelection()->EmptySelectedKeys();
 		}
-		GetSequencer().SelectKey( Key );
+		GetSequencer().GetSelection()->AddToSelection( Key );
 	}
 }
 
@@ -720,11 +735,11 @@ void SSection::HandleSectionSelection( const FPointerEvent& MouseEvent )
 {
 	if( !MouseEvent.IsControlDown() )
 	{
-		GetSequencer().ClearSectionSelection();
+		GetSequencer().GetSelection()->EmptySelectedSections();
 	}
 
 	// handle selecting sections 
 	UMovieSceneSection* Section = SectionInterface->GetSectionObject();
-	GetSequencer().SelectSection(Section);
+	GetSequencer().GetSelection()->AddToSelection(Section);
 }
 
diff --git a/Engine/Source/Editor/Sequencer/Private/SSequencer.cpp b/Engine/Source/Editor/Sequencer/Private/SSequencer.cpp
--- a/Engine/Source/Editor/Sequencer/Private/SSequencer.cpp
+++ b/Engine/Source/Editor/Sequencer/Private/SSequencer.cpp
@@ -696,15 +696,15 @@ void SSequencer::OnAssetsDropped( const FAssetDragDropOp& DragDropOp )
 		}
 	}
 
-	const TSet< TSharedRef<const FSequencerDisplayNode> >& SelectedNodes = SequencerNodeTree->GetSelectedNodes();
+	const TSet< TSharedRef<const FSequencerDisplayNode> >* SelectedNodes = Sequencer.Pin()->GetSelection()->GetSelectedOutlinerNodes();
 	FGuid TargetObjectGuid;
 	// if exactly one object node is selected, we have a target object guid
 	TSharedPtr<const FSequencerDisplayNode> DisplayNode;
-	if (SelectedNodes.Num() == 1)
+	if (SelectedNodes->Num() == 1)
 	{
-		for (TSet< TSharedRef<const FSequencerDisplayNode> >::TConstIterator It(SelectedNodes); It; ++It)
+		for (TSharedRef<const FSequencerDisplayNode> SelectedNode : *SelectedNodes )
 		{
-			DisplayNode = *It;
+			DisplayNode = SelectedNode;
 		}
 		if (DisplayNode.IsValid() && DisplayNode->GetType() == ESequencerNode::Object)
 		{
@@ -861,24 +861,21 @@ FText SSequencer::GetShotSectionTitle(UMovieSceneSection* ShotSection) const
 
 void SSequencer::DeleteSelectedNodes()
 {
-	const TSet< TSharedRef<const FSequencerDisplayNode> >& SelectedNodes = SequencerNodeTree->GetSelectedNodes();
+	const TSet< TSharedRef<const FSequencerDisplayNode> >* SelectedNodes = Sequencer.Pin()->GetSelection()->GetSelectedOutlinerNodes();
 
-	if( SelectedNodes.Num() > 0 )
+	if( SelectedNodes->Num() > 0 )
 	{
 		const FScopedTransaction Transaction( LOCTEXT("UndoDeletingObject", "Delete Object from MovieScene") );
 
 		FSequencer& SequencerRef = *Sequencer.Pin();
 
-		for( auto It = SelectedNodes.CreateConstIterator(); It; ++It )
+		for( auto SelectedNode : *SelectedNodes )
 		{
-			TSharedRef<const FSequencerDisplayNode> Node = *It;
-
-			if( Node->IsVisible() )
+			if( SelectedNode->IsVisible() )
 			{
 				// Delete everything in the entire node
-				SequencerRef.OnRequestNodeDeleted( Node );
+				SequencerRef.OnRequestNodeDeleted( SelectedNode );
 			}
-
 		}
 	}
 }
diff --git a/Engine/Source/Editor/Sequencer/Private/SSequencerCurveEditor.cpp b/Engine/Source/Editor/Sequencer/Private/SSequencerCurveEditor.cpp
--- a/Engine/Source/Editor/Sequencer/Private/SSequencerCurveEditor.cpp
+++ b/Engine/Source/Editor/Sequencer/Private/SSequencerCurveEditor.cpp
@@ -9,6 +9,7 @@ void SSequencerCurveEditor::Construct( const FArguments& InArgs, TSharedRef<FSeq
 {
 	ViewRange = InArgs._ViewRange;
 	Sequencer = InSequencer;
+	NodeTreeSelectionChangedHandle = Sequencer.Pin()->GetSelection()->GetOnOutlinerNodeSelectionChanged()->AddSP(this, &SSequencerCurveEditor::NodeTreeSelectionChanged);
 	GetMutableDefault<USequencerSettings>()->GetOnCurveVisibilityChanged()->AddSP( this, &SSequencerCurveEditor::SequencerCurveVisibilityChanged );
 
 	ChildSlot
@@ -46,14 +47,7 @@ TSharedPtr<ParentNodeType> GetParentOfType( TSharedRef<FSequencerDisplayNode> In
 
 void SSequencerCurveEditor::SetSequencerNodeTree( TSharedPtr<FSequencerNodeTree> InSequencerNodeTree )
 {
-	if ( SequencerNodeTree.IsValid() )
-	{
-		SequencerNodeTree->GetOnSelectionChanged()->Remove( NodeTreeSelectionChangedHandle );
-	}
-
 	SequencerNodeTree = InSequencerNodeTree;
-	NodeTreeSelectionChangedHandle = SequencerNodeTree->GetOnSelectionChanged()->AddSP( this, &SSequencerCurveEditor::NodeTreeSelectionChanged );
-
 	UpdateCurveOwner();
 }
 
@@ -90,9 +84,12 @@ bool SSequencerCurveEditor::GetShowCurveEditorCurveToolTips() const
 
 void SSequencerCurveEditor::NodeTreeSelectionChanged()
 {
-	if ( GetDefault<USequencerSettings>()->GetCurveVisibility() == ESequencerCurveVisibility::SelectedCurves )
+	if (SequencerNodeTree.IsValid())
 	{
-		UpdateCurveOwner();
+		if (GetDefault<USequencerSettings>()->GetCurveVisibility() == ESequencerCurveVisibility::SelectedCurves)
+		{
+			UpdateCurveOwner();
+		}
 	}
 }
 
@@ -108,9 +105,9 @@ TSharedPtr<FUICommandList> SSequencerCurveEditor::GetCommands()
 
 SSequencerCurveEditor::~SSequencerCurveEditor()
 {
-	if ( SequencerNodeTree.IsValid() )
+	if ( Sequencer.IsValid() )
 	{
-		SequencerNodeTree->GetOnSelectionChanged()->Remove( NodeTreeSelectionChangedHandle );
+		Sequencer.Pin()->GetSelection()->GetOnOutlinerNodeSelectionChanged()->Remove( NodeTreeSelectionChangedHandle );
 	}
 }
 
diff --git a/Engine/Source/Editor/Sequencer/Private/SSequencerSectionAreaView.cpp b/Engine/Source/Editor/Sequencer/Private/SSequencerSectionAreaView.cpp
--- a/Engine/Source/Editor/Sequencer/Private/SSequencerSectionAreaView.cpp
+++ b/Engine/Source/Editor/Sequencer/Private/SSequencerSectionAreaView.cpp
@@ -137,7 +137,7 @@ FReply SSequencerSectionAreaView::OnMouseButtonDown( const FGeometry& MyGeometry
 	// Clear selected sections
 	if( !MouseEvent.IsControlDown() )
 	{
-		GetSequencer().ClearSectionSelection();
+		GetSequencer().GetSelection()->EmptySelectedSections();
 	}
 
 	return FReply::Handled();
diff --git a/Engine/Source/Editor/Sequencer/Private/SectionDragOperations.cpp b/Engine/Source/Editor/Sequencer/Private/SectionDragOperations.cpp
--- a/Engine/Source/Editor/Sequencer/Private/SectionDragOperations.cpp
+++ b/Engine/Source/Editor/Sequencer/Private/SectionDragOperations.cpp
@@ -385,15 +385,14 @@ bool FMoveSection::IsSectionAtNewLocationValid(UMovieSceneSection* InSection, in
 
 void FMoveKeys::OnBeginDrag(const FVector2D& LocalMousePos, TSharedPtr<FTrackNode> SequencerNode)
 {
-	check( SelectedKeys.Num() > 0 )
+	check( SelectedKeys->Num() > 0 )
 
 	// Begin an editor transaction and mark the section as transactional so it's state will be saved
 	GEditor->BeginTransaction( NSLOCTEXT("Sequencer", "MoveKeysTransaction", "Move Keys") );
 
 	TSet<UMovieSceneSection*> ModifiedSections;
-	for( auto It = SelectedKeys.CreateIterator(); It; ++It )
+	for( FSelectedKey SelectedKey : *SelectedKeys )
 	{
-		FSelectedKey& SelectedKey = *It;
 		UMovieSceneSection* OwningSection = SelectedKey.Section;
 
 		// Only modify sections once
@@ -438,9 +437,8 @@ void FMoveKeys::OnDrag( const FPointerEvent& MouseEvent, const FVector2D& LocalM
 				GetKeySnapTimes(OutSnapTimes, SequencerNode);
 
 				TArray<float> InitialTimes;
-				for (TSet<FSelectedKey>::TConstIterator It(SelectedKeys); It; ++It)
+				for ( FSelectedKey SelectedKey : *SelectedKeys )
 				{
-					const FSelectedKey& SelectedKey = *It;
 					InitialTimes.Add(SelectedKey.KeyArea->GetKeyTime(SelectedKey.KeyHandle.GetValue()) + DistanceMoved);
 				}
 				float OutInitialTime = 0.f;
@@ -458,10 +456,8 @@ void FMoveKeys::OnDrag( const FPointerEvent& MouseEvent, const FVector2D& LocalM
 			}
 		}
 
-		for( auto It = SelectedKeys.CreateIterator(); It; ++It )
+		for( FSelectedKey SelectedKey : *SelectedKeys )
 		{
-			FSelectedKey& SelectedKey = *It;
-
 			UMovieSceneSection* Section = SelectedKey.Section;
 
 			TSharedPtr<IKeyArea>& KeyArea = SelectedKey.KeyArea;
@@ -505,9 +501,8 @@ void FMoveKeys::GetKeySnapTimes(TArray<float>& OutSnapTimes, TSharedPtr<FTrackNo
 			{
 				FKeyHandle KeyHandle = KeyHandles[KeyIndex];
 				bool bKeyIsSnappable = true;
-				for (TSet<FSelectedKey>::TConstIterator It(SelectedKeys); It; ++It)
+				for ( FSelectedKey SelectedKey : *SelectedKeys )
 				{
-					const FSelectedKey& SelectedKey = *It;
 					if (SelectedKey.KeyArea == KeyArea && SelectedKey.KeyHandle.GetValue() == KeyHandle)
 					{
 						bKeyIsSnappable = false;
diff --git a/Engine/Source/Editor/Sequencer/Private/SectionDragOperations.h b/Engine/Source/Editor/Sequencer/Private/SectionDragOperations.h
--- a/Engine/Source/Editor/Sequencer/Private/SectionDragOperations.h
+++ b/Engine/Source/Editor/Sequencer/Private/SectionDragOperations.h
@@ -132,7 +132,7 @@ private:
 class FMoveKeys : public FSequencerDragOperation
 {
 public:
-	FMoveKeys( TSet<FSelectedKey>& InSelectedKeys, FSelectedKey& PressedKey )
+	FMoveKeys( const TSet<FSelectedKey>* InSelectedKeys, FSelectedKey& PressedKey )
 		: SelectedKeys( InSelectedKeys )
 		, DraggedKey( PressedKey )
 	{}
@@ -148,7 +148,7 @@ protected:
 
 private:
 	/** The selected keys being moved. */
-	TSet<FSelectedKey>& SelectedKeys;
+	const TSet<FSelectedKey>* SelectedKeys;
 	/** The exact key that we're dragging */
 	FSelectedKey DraggedKey;
 };
diff --git a/Engine/Source/Editor/Sequencer/Private/Sequencer.cpp b/Engine/Source/Editor/Sequencer/Private/Sequencer.cpp
--- a/Engine/Source/Editor/Sequencer/Private/Sequencer.cpp
+++ b/Engine/Source/Editor/Sequencer/Private/Sequencer.cpp
@@ -240,8 +240,7 @@ void FSequencer::ResetToNewRootMovieScene( UMovieScene& NewRoot, TSharedRef<ISeq
 
 	//@todo Sequencer - Encapsulate this better
 	MovieSceneStack.Empty();
-	SelectedSections.Empty();
-	SelectedKeys.Empty();
+	Selection.Empty();
 	FilteringShots.Empty();
 	UnfilterableSections.Empty();
 	UnfilterableObjects.Empty();
@@ -378,7 +377,7 @@ void FSequencer::DeleteSelectedKeys()
 {
 	FScopedTransaction DeleteKeysTransaction( NSLOCTEXT("Sequencer", "DeleteSelectedKeys_Transaction", "Delete Selected Keys") );
 
-	TArray<FSelectedKey> SelectedKeysArray = SelectedKeys.Array();
+	TArray<FSelectedKey> SelectedKeysArray = Selection.GetSelectedKeys()->Array();
 	for ( const FSelectedKey& Key : SelectedKeysArray )
 	{
 		if (Key.IsValid())
@@ -676,8 +675,8 @@ void FSequencer::ResetPerMovieSceneData()
 	FilterToShotSections( TArray< TWeakObjectPtr<UMovieSceneSection> >(), false );
 
 	//@todo Sequencer - We may want to preserve selections when moving between movie scenes
-	SelectedSections.Empty();
-	SelectedKeys.Empty();
+	Selection.EmptySelectedKeys();
+	Selection.EmptySelectedSections();
 
 	// @todo run through all tracks for new movie scene changes
 	//  needed for audio track decompression
@@ -1148,38 +1147,22 @@ void FSequencer::PostUndo(bool bSuccess)
 	NotifyMovieSceneDataChanged();
 }
 
-TArray< TWeakObjectPtr<UMovieSceneSection> > FSequencer::GetSelectedSections() const
-{
-	return SelectedSections;
-}
 
-void FSequencer::SelectSection(UMovieSceneSection* Section)
+void FSequencer::OnSectionSelectionChanged()
 {
-	if (!Section->IsA<UMovieSceneShotSection>())
+	for (TWeakObjectPtr<UMovieSceneSection> SelectedSection : *Selection.GetSelectedSections())
 	{
 		// if we select something, consider it unfilterable until we change shot filters
-		UnfilterableSections.AddUnique(TWeakObjectPtr<UMovieSceneSection>(Section));
+		UnfilterableSections.AddUnique(TWeakObjectPtr<UMovieSceneSection>(SelectedSection));
 	}
-
-	SelectedSections.Add(TWeakObjectPtr<UMovieSceneSection>(Section));
-}
-
-bool FSequencer::IsSectionSelected(UMovieSceneSection* Section) const
-{
-	return SelectedSections.Contains(TWeakObjectPtr<UMovieSceneSection>(Section));
-}
-
-void FSequencer::ClearSectionSelection()
-{
-	SelectedSections.Empty();
 }
 
 void FSequencer::ZoomToSelectedSections()
 {
 	TArray< TRange<float> > Bounds;
-	for (int32 i = 0; i < SelectedSections.Num(); ++i)
+	for (TWeakObjectPtr<UMovieSceneSection> SelectedSection : *Selection.GetSelectedSections())
 	{
-		Bounds.Add(SelectedSections[i]->GetRange());
+		Bounds.Add(SelectedSection->GetRange());
 	}
 	TRange<float> BoundsHull = TRange<float>::Hull(Bounds);
 
@@ -1194,27 +1177,6 @@ void FSequencer::ZoomToSelectedSections()
 	}
 }
 
-
-void FSequencer::SelectKey( const FSelectedKey& Key )
-{
-	SelectedKeys.Add( Key );
-}
-
-void FSequencer::ClearKeySelection()
-{
-	SelectedKeys.Empty();
-}
-
-bool FSequencer::IsKeySelected( const FSelectedKey& Key ) const
-{
-	return SelectedKeys.Contains( Key );
-}
-
-TSet<FSelectedKey>& FSequencer::GetSelectedKeys()
-{
-	return SelectedKeys; 
-}
-
 void FSequencer::FilterToShotSections(const TArray< TWeakObjectPtr<class UMovieSceneSection> >& ShotSections, bool bZoomToShotBounds )
 {
 	TArray< TWeakObjectPtr<UMovieSceneSection> > ActualShotSections;
@@ -1277,11 +1239,11 @@ void FSequencer::FilterToShotSections(const TArray< TWeakObjectPtr<class UMovieS
 void FSequencer::FilterToSelectedShotSections(bool bZoomToShotBounds)
 {
 	TArray< TWeakObjectPtr<UMovieSceneSection> > SelectedShotSections;
-	for (int32 i = 0; i < SelectedSections.Num(); ++i)
+	for (TWeakObjectPtr<UMovieSceneSection> SelectedSection : *Selection.GetSelectedSections())
 	{
-		if (SelectedSections[i]->IsA<UMovieSceneShotSection>())
+		if (SelectedSection->IsA<UMovieSceneShotSection>())
 		{
-			SelectedShotSections.Add(SelectedSections[i]);
+			SelectedShotSections.Add(SelectedSection);
 		}
 	}
 	FilterToShotSections(SelectedShotSections, bZoomToShotBounds);
@@ -1302,6 +1264,11 @@ TSharedRef<ISequencerObjectBindingManager> FSequencer::GetObjectBindingManager()
 	return ObjectBindingManager.ToSharedRef();
 }
 
+FSequencerSelection* FSequencer::GetSelection()
+{
+	return &Selection;
+}
+
 TArray< TWeakObjectPtr<UMovieSceneSection> > FSequencer::GetFilteringShotSections() const
 {
 	return FilteringShots;
@@ -1350,14 +1317,18 @@ bool FSequencer::IsSectionVisible(UMovieSceneSection* Section) const
 
 void FSequencer::DeleteSelectedItems()
 {
-	DeleteSelectedKeys();
-
-	for (int32 i = 0; i < SelectedSections.Num(); ++i)
+	if (Selection.GetActiveSelection() == FSequencerSelection::EActiveSelection::KeyAndSection)
 	{
-		DeleteSection(SelectedSections[i].Get());
+		DeleteSelectedKeys();
+		for (TWeakObjectPtr<UMovieSceneSection> SelectedSection : *Selection.GetSelectedSections())
+		{
+			DeleteSection(SelectedSection.Get());
+		}
+	}
+	else if (Selection.GetActiveSelection() == FSequencerSelection::EActiveSelection::OutlinerNode)
+	{
+		SequencerWidget->DeleteSelectedNodes();
 	}
-
-	SequencerWidget->DeleteSelectedNodes();
 }
 
 void FSequencer::SetKey()
diff --git a/Engine/Source/Editor/Sequencer/Private/Sequencer.h b/Engine/Source/Editor/Sequencer/Private/Sequencer.h
--- a/Engine/Source/Editor/Sequencer/Private/Sequencer.h
+++ b/Engine/Source/Editor/Sequencer/Private/Sequencer.h
@@ -91,6 +91,7 @@ public:
 	virtual bool CanKeyProperty(const UClass& ObjectClass, const class IPropertyHandle& PropertyHandle) const override;
 	virtual void KeyProperty(const TArray<UObject*>& ObjectsToKey, const class IPropertyHandle& PropertyHandle) override;
 	virtual TSharedRef<ISequencerObjectBindingManager> GetObjectBindingManager() const override;
+	virtual FSequencerSelection* GetSelection() override;
 
 	bool IsPerspectiveViewportPosessionEnabled() const { return bPerspectiveViewportPossessionEnabled; }
 
@@ -162,58 +163,11 @@ public:
 	 */
 	virtual void CopyActorProperties( AActor* PuppetActor, AActor* TargetActor ) const;
 
-	/**
-	 * Gets all selected sections in the sequencer
-	 */
-	virtual TArray< TWeakObjectPtr<UMovieSceneSection> > GetSelectedSections() const;
-
-	/**
-	 * Selects a section in the sequencer
-	 */
-	void SelectSection(UMovieSceneSection* Section);
-
-	/**
-	 * Returns whether or not a section is selected
-	 *
-	 * @param Section The section to check
-	 * @return true if the section is selected
-	 */
-	bool IsSectionSelected( UMovieSceneSection* Section) const;
-
-	/**
-	 * Clears all selected sections
-	 */
-	void ClearSectionSelection();
-
 	/**
 	 * Zooms to the edges of all currently selected sections
 	 */
 	void ZoomToSelectedSections();
 
-	/**
-	 * Selects a key
-	 * 
-	 * @param Key	Representation of the key to select
-	 */
-	void SelectKey( const FSelectedKey& Key );
-
-	/**
-	 * Clears the entire set of selected keys
-	 */
-	void ClearKeySelection();
-
-	/**
-	 * Returns whether or not a key is selected
-	 * 
-	 * @param Key	Representation of the key to check for selection
-	 */
-	bool IsKeySelected( const FSelectedKey& Key ) const;
-
-	/**
-	 * @return The entire set of selected keys
-	 */
-	TSet<FSelectedKey>& GetSelectedKeys();
-
 	/**
 	 * Gets all shots that are filtering currently
 	 */
@@ -371,6 +325,8 @@ protected:
 	virtual void PostRedo(bool bSuccess) override { PostUndo(bSuccess); }
 	// End of FEditorUndoClient
 
+	void OnSectionSelectionChanged();
+
 private:
 	TMap< TWeakObjectPtr<UMovieSceneSection>, TSharedRef<FMovieSceneInstance> > MovieSceneSectionToInstanceMap;
 
@@ -411,12 +367,6 @@ private:
 	/** A list of object guids that will be visible, regardless of shot filters */
 	TArray<FGuid> UnfilterableObjects;
 
-	/** Selected non-shot sections */
-	TArray< TWeakObjectPtr<class UMovieSceneSection> > SelectedSections;
-
-	/** Set of selected keys */
-	TSet< FSelectedKey > SelectedKeys;
-
 	/** Stack of movie scenes.  The first element is always the root movie scene.  The last element is the focused movie scene */
 	TArray< TSharedRef<FMovieSceneInstance> > MovieSceneStack;
 
@@ -454,4 +404,6 @@ private:
 	    do this simply to avoid refreshing the UI more than once per frame. (e.g. during live recording where
 		the MovieScene data can change many times per frame.) */
 	bool bNeedTreeRefresh;
+
+	FSequencerSelection Selection;
 };
diff --git a/Engine/Source/Editor/Sequencer/Private/SequencerCurveOwner.cpp b/Engine/Source/Editor/Sequencer/Private/SequencerCurveOwner.cpp
--- a/Engine/Source/Editor/Sequencer/Private/SequencerCurveOwner.cpp
+++ b/Engine/Source/Editor/Sequencer/Private/SequencerCurveOwner.cpp
@@ -61,7 +61,7 @@ FSequencerCurveOwner::FSequencerCurveOwner( TSharedPtr<FSequencerNodeTree> InSeq
 					bAddCurve = true;
 					break;
 				case ESequencerCurveVisibility::SelectedCurves:
-					bAddCurve = KeyAreaNode->IsSelected();
+					bAddCurve = KeyAreaNode->GetSequencer().GetSelection()->IsSelected(KeyAreaNode);
 					break;
 				case ESequencerCurveVisibility::AnimatedCurves:
 					bAddCurve = CurveInfo->Curve->GetNumKeys() > 0;
diff --git a/Engine/Source/Editor/Sequencer/Private/SequencerDisplayNode.cpp b/Engine/Source/Editor/Sequencer/Private/SequencerDisplayNode.cpp
--- a/Engine/Source/Editor/Sequencer/Private/SequencerDisplayNode.cpp
+++ b/Engine/Source/Editor/Sequencer/Private/SequencerDisplayNode.cpp
@@ -266,11 +266,6 @@ void FSequencerDisplayNode::GetChildKeyAreaNodesRecursively(TArray< TSharedRef<F
 	}
 }
 
-void FSequencerDisplayNode::SetSelectionState( bool bSelect, bool bDeselectOtherNodes )
-{
-	ParentTree.SetSelectionState( AsShared(), bSelect, bDeselectOtherNodes );
-}
-
 void FSequencerDisplayNode::ToggleExpansion()
 {
 	bExpanded = !bExpanded;
@@ -281,12 +276,6 @@ void FSequencerDisplayNode::ToggleExpansion()
 	UpdateCachedShotFilteredVisibility();
 }
 
-bool FSequencerDisplayNode::IsSelected() const
-{
-	// Ask the tree if we are selected
-	return ParentTree.IsNodeSelected( SharedThis(this) );
-}
-
 bool FSequencerDisplayNode::IsExpanded() const
 {
 	return ParentTree.HasActiveFilter() ? ParentTree.IsNodeFiltered( AsShared() ) : bExpanded;
diff --git a/Engine/Source/Editor/Sequencer/Private/SequencerDisplayNode.h b/Engine/Source/Editor/Sequencer/Private/SequencerDisplayNode.h
--- a/Engine/Source/Editor/Sequencer/Private/SequencerDisplayNode.h
+++ b/Engine/Source/Editor/Sequencer/Private/SequencerDisplayNode.h
@@ -134,24 +134,11 @@ public:
 	/** Gets all the key area nodes recursively, including this node if applicable */
 	virtual void GetChildKeyAreaNodesRecursively(TArray< TSharedRef<class FSectionKeyAreaNode> >& OutNodes) const;
 
-	/**
-	 * Selects or deselects this node
-	 *
-	 * @param bSelect				Whether or not this node should be selected
-	 * @param bDeselectOtherNodes	Whether or not to deselect other nodes
-	 */
-	void SetSelectionState( bool bSelect, bool bDeselectOtherNodes );
-
 	/**
 	 * Toggles the expansion state of this node
 	 */
 	void ToggleExpansion();
 
-	/**
-	 * @return Whether or not this node is selected
-	 */
-	bool IsSelected() const;
-
 	/**
 	 * @return Whether or not this node is expanded                                                              
 	 */
diff --git a/Engine/Source/Editor/Sequencer/Private/SequencerNodeTree.cpp b/Engine/Source/Editor/Sequencer/Private/SequencerNodeTree.cpp
--- a/Engine/Source/Editor/Sequencer/Private/SequencerNodeTree.cpp
+++ b/Engine/Source/Editor/Sequencer/Private/SequencerNodeTree.cpp
@@ -15,7 +15,7 @@ void FSequencerNodeTree::Empty()
 {
 	RootNodes.Empty();
 	ObjectBindingMap.Empty();
-	SelectedNodes.Empty();
+	Sequencer.GetSelection()->EmptySelectedOutlinerNodes();
 	EditorMap.Empty();
 	FilteredNodes.Empty();
 }
@@ -180,31 +180,6 @@ bool FSequencerNodeTree::GetSavedExpansionState( const FSequencerDisplayNode& No
 	return !bCollapsed;
 }
 
-void FSequencerNodeTree::SetSelectionState( TSharedRef<const FSequencerDisplayNode> AffectedNode, bool bSelect, bool bDeselectOtherNodes )
-{
-	if( bSelect && AffectedNode->IsSelectable() )
-	{
-		if( bDeselectOtherNodes )
-		{
-			// empty current selection set unless multiple selecting
-			SelectedNodes.Empty();
-		}
-
-		SelectedNodes.Add( AffectedNode );
-	}
-	else
-	{
-		// Not selecting so remove the node from the selection set
-		SelectedNodes.Remove( AffectedNode );
-	}
-	OnSelectionChanged.Broadcast();
-}
-
-bool FSequencerNodeTree::IsNodeSelected( TSharedRef<const FSequencerDisplayNode> Node ) const
-{
-	return SelectedNodes.Contains( Node );
-}
-
 bool FSequencerNodeTree::IsNodeFiltered( const TSharedRef<const FSequencerDisplayNode> Node ) const
 {
 	return FilteredNodes.Contains( Node );
@@ -296,8 +271,3 @@ void FSequencerNodeTree::UpdateCachedVisibilityBasedOnShotFiltersChanged()
 		RootNodes[i]->UpdateCachedShotFilteredVisibility();
 	}
 }
-
-FSequencerNodeTree::FOnSelectionChanged* FSequencerNodeTree::GetOnSelectionChanged()
-{
-	return &OnSelectionChanged;
-}
diff --git a/Engine/Source/Editor/Sequencer/Private/SequencerNodeTree.h b/Engine/Source/Editor/Sequencer/Private/SequencerNodeTree.h
--- a/Engine/Source/Editor/Sequencer/Private/SequencerNodeTree.h
+++ b/Engine/Source/Editor/Sequencer/Private/SequencerNodeTree.h
@@ -13,8 +13,6 @@ class UMovieSceneTrack;
 class FSequencerNodeTree : public TSharedFromThis<FSequencerNodeTree>
 {
 public:
-	DECLARE_MULTICAST_DELEGATE(FOnSelectionChanged)
-
 	FSequencerNodeTree( class FSequencer& InSequencer )
 		: Sequencer( InSequencer )
 	{}
@@ -34,28 +32,6 @@ public:
 	 */
 	const TArray< TSharedRef<FSequencerDisplayNode> >& GetRootNodes() const;
 
-	/**
-	 * Selects a node in the tree
-	 *
-	 * @param AffectedNode			The node to select
-	 * @param bSelect				Whether or not to select
-	 * @param bDeselectOtherNodes	Whether or not to deselect all other nodes
-	 */
-	void SetSelectionState( TSharedRef<const FSequencerDisplayNode> AffectedNode, bool bSelect, bool bDeselectOtherNodes = true );
-
-	/**
-	 * @return All currently selected nodes
-	 */
-	const TSet< TSharedRef<const FSequencerDisplayNode> >& GetSelectedNodes() const { return SelectedNodes; }
-
-	/**
-	 * Returns whether or not a node is selected
-	 *
-	 * @param Node	The node to check for selection
-	 * @return true if the node is selected
-	 */
-	bool IsNodeSelected( const TSharedRef<const FSequencerDisplayNode> Node ) const;
-
 	/** @return Whether or not there is an active filter */
 	bool HasActiveFilter() const { return !FilterString.IsEmpty(); }
 
@@ -95,9 +71,6 @@ public:
 	/** Updates nodes visibility for use when shot filters change at all */
 	void UpdateCachedVisibilityBasedOnShotFiltersChanged();
 
-	/** Gets the multicast delegate which is run when the node selection changes. */
-	FOnSelectionChanged* GetOnSelectionChanged();
-
 private:
 	/**
 	 * Finds or adds a type editor for the track
@@ -111,7 +84,7 @@ private:
 	 * Makes section interfaces for all sections in a track
 	 *
 	 * @param Track	The type to get sections from
-	 * @param SectionAreaNode	The section area whech section interfaces belong to
+	 * @param SectionAreaNode	The section area which section interfaces belong to
 	 */
 	void MakeSectionInterfaces( UMovieSceneTrack& Track, TSharedRef<class FTrackNode>& SectionAreaNode );
 
@@ -128,14 +101,10 @@ private:
 	TArray< TSharedRef<FSequencerDisplayNode> > RootNodes;
 	/** Mapping of object binding guids to their node (for fast lookup) */
 	TMap< FGuid, TSharedPtr<FObjectBindingNode> > ObjectBindingMap;
-	/** Set of all selected nodes */
-	TSet< TSharedRef<const FSequencerDisplayNode> > SelectedNodes;
 	/** Set of all filtered nodes */
 	TSet< TSharedRef<const FSequencerDisplayNode> > FilteredNodes;
 	/** Active filter string if any */
 	FString FilterString;
 	/** Sequencer interface */
 	FSequencer& Sequencer;
-	/** The multicast delegate which is run when the node selection changes */
-	FOnSelectionChanged OnSelectionChanged;
 };
diff --git a/Engine/Source/Editor/Sequencer/Private/SequencerPrivatePCH.h b/Engine/Source/Editor/Sequencer/Private/SequencerPrivatePCH.h
--- a/Engine/Source/Editor/Sequencer/Private/SequencerPrivatePCH.h
+++ b/Engine/Source/Editor/Sequencer/Private/SequencerPrivatePCH.h
@@ -9,6 +9,7 @@
 
 #include "ISequencerModule.h"
 #include "ISequencerObjectChangeListener.h"
+#include "SequencerSelection.h"
 #include "SequencerCommands.h"
 #include "SequencerNodeTree.h"
 #include "SequencerDisplayNode.h"
diff --git a/Engine/Source/Editor/Sequencer/Private/SequencerSelection.cpp b/Engine/Source/Editor/Sequencer/Private/SequencerSelection.cpp
new file mode 100644
--- /dev/null
+++ b/Engine/Source/Editor/Sequencer/Private/SequencerSelection.cpp
@@ -0,0 +1,105 @@
+// Copyright 1998-2015 Epic Games, Inc. All Rights Reserved.
+
+#include "SequencerPrivatePCH.h"
+#include "SequencerSelection.h"
+
+FSequencerSelection::FSequencerSelection()
+{
+	ActiveSelection = EActiveSelection::None;
+}
+
+const TSet<FSelectedKey>* FSequencerSelection::GetSelectedKeys() const
+{
+	return &SelectedKeys;
+}
+
+const TSet<TWeakObjectPtr<UMovieSceneSection>>* FSequencerSelection::GetSelectedSections() const
+{
+	return &SelectedSections;
+}
+
+const TSet<TSharedRef<const FSequencerDisplayNode>>* FSequencerSelection::GetSelectedOutlinerNodes() const
+{
+	return &SelectedOutlinerNodes;
+}
+
+FSequencerSelection::EActiveSelection FSequencerSelection::GetActiveSelection() const
+{
+	return ActiveSelection;
+}
+
+FSequencerSelection::FOnSelectionChanged* FSequencerSelection::GetOnKeySelectionChanged()
+{
+	return &OnKeySelectionChanged;
+}
+
+FSequencerSelection::FOnSelectionChanged* FSequencerSelection::GetOnSectionSelectionChanged()
+{
+	return &OnSectionSelectionChanged;
+}
+
+FSequencerSelection::FOnSelectionChanged* FSequencerSelection::GetOnOutlinerNodeSelectionChanged()
+{
+	return &OnOutlinerNodeSelectionChanged;
+}
+
+void FSequencerSelection::AddToSelection(FSelectedKey Key)
+{
+	SelectedKeys.Add(Key);
+	ActiveSelection = EActiveSelection::KeyAndSection;
+	OnKeySelectionChanged.Broadcast();
+}
+
+void FSequencerSelection::AddToSelection(UMovieSceneSection* Section)
+{
+	SelectedSections.Add(Section);
+	ActiveSelection = EActiveSelection::KeyAndSection;
+	OnSectionSelectionChanged.Broadcast();
+}
+
+void FSequencerSelection::AddToSelection(TSharedRef<const FSequencerDisplayNode> OutlinerNode)
+{
+	SelectedOutlinerNodes.Add(OutlinerNode);
+	ActiveSelection = EActiveSelection::OutlinerNode;
+	OnOutlinerNodeSelectionChanged.Broadcast();
+}
+
+bool FSequencerSelection::IsSelected(FSelectedKey Key) const
+{
+	return SelectedKeys.Contains(Key);
+}
+
+bool FSequencerSelection::IsSelected(UMovieSceneSection* Section) const
+{
+	return SelectedSections.Contains(Section);
+}
+
+bool FSequencerSelection::IsSelected(TSharedRef<FSequencerDisplayNode> OutlinerNode) const
+{
+	return SelectedOutlinerNodes.Contains(OutlinerNode);
+}
+
+void FSequencerSelection::Empty()
+{
+	EmptySelectedKeys();
+	EmptySelectedSections();
+	EmptySelectedOutlinerNodes();
+}
+
+void FSequencerSelection::EmptySelectedKeys()
+{
+	SelectedKeys.Empty();
+	OnKeySelectionChanged.Broadcast();
+}
+
+void FSequencerSelection::EmptySelectedSections()
+{
+	SelectedSections.Empty();
+	OnSectionSelectionChanged.Broadcast();
+}
+
+void FSequencerSelection::EmptySelectedOutlinerNodes()
+{
+	SelectedOutlinerNodes.Empty();
+	OnOutlinerNodeSelectionChanged.Broadcast();
+}
\ No newline at end of file
diff --git a/Engine/Source/Editor/Sequencer/Private/SequencerSelection.h b/Engine/Source/Editor/Sequencer/Private/SequencerSelection.h
new file mode 100644
--- /dev/null
+++ b/Engine/Source/Editor/Sequencer/Private/SequencerSelection.h
@@ -0,0 +1,93 @@
+// Copyright 1998-2015 Epic Games, Inc. All Rights Reserved.
+
+#pragma once
+
+#include "SelectedKey.h"
+
+class UMovieSceneSection;
+class FSequencerDisplayNode;
+
+/**
+ * Manages the selection of keys, sections, and outliner nodes for the sequencer.
+ */
+class FSequencerSelection
+{
+public:
+	DECLARE_MULTICAST_DELEGATE(FOnSelectionChanged)
+
+	/** Options for the current active selection. */
+	enum class EActiveSelection
+	{
+		KeyAndSection,
+		OutlinerNode,
+		None
+	};
+
+public:
+	FSequencerSelection();
+
+	/** Gets a set of the selected keys. */
+	const TSet<FSelectedKey>* GetSelectedKeys() const;
+
+	/** Gets a set of the selected sections */
+	const TSet<TWeakObjectPtr<UMovieSceneSection>>* GetSelectedSections() const;
+
+	/** Gets a set of the selected outliner nodes. */
+	const TSet<TSharedRef<const FSequencerDisplayNode>>* GetSelectedOutlinerNodes() const;
+
+	/** 
+	 * Gets the currently active selection.  This is used to determine which selection actions
+	 * like delete should act on.
+	 */
+	EActiveSelection GetActiveSelection() const;
+
+	/** Adds a key to the selection */
+	void AddToSelection(FSelectedKey Key);
+
+	/** Adds a section to the selection */
+	void AddToSelection(UMovieSceneSection* Section);
+
+	/** Adds an outliner node to the selection */
+	void AddToSelection(TSharedRef<const FSequencerDisplayNode> OutlinerNode);
+
+	/** Returns whether or not the key is selected. */
+	bool IsSelected(FSelectedKey Key) const;
+
+	/** Returns whether or not the section is selected. */
+	bool IsSelected(UMovieSceneSection* Section) const;
+
+	/** Returns whether or not the outliner node is selected. */
+	bool IsSelected(TSharedRef<FSequencerDisplayNode> OutlinerNode) const;
+
+	/** Empties all selections. */
+	void Empty();
+
+	/** Empties the key selection. */
+	void EmptySelectedKeys();
+
+	/** Empties the section selection. */
+	void EmptySelectedSections();
+
+	/** Empties the outliner node selection. */
+	void EmptySelectedOutlinerNodes();
+
+	/** Gets a multicast delegate which is called when the key selection changes. */
+	FOnSelectionChanged* GetOnKeySelectionChanged();
+
+	/** Gets a multicast delegate which is called when the section selection changes. */
+	FOnSelectionChanged* GetOnSectionSelectionChanged();
+
+	/** Gets a multicast delegate which is called when the outliner node selection changes. */
+	FOnSelectionChanged* GetOnOutlinerNodeSelectionChanged();
+	
+private:
+	TSet<FSelectedKey> SelectedKeys;
+	TSet<TWeakObjectPtr<UMovieSceneSection>> SelectedSections;
+	TSet<TSharedRef<const FSequencerDisplayNode>> SelectedOutlinerNodes;
+
+	FOnSelectionChanged OnKeySelectionChanged;
+	FOnSelectionChanged OnSectionSelectionChanged;
+	FOnSelectionChanged OnOutlinerNodeSelectionChanged;
+
+	EActiveSelection ActiveSelection;
+};
\ No newline at end of file
diff --git a/Engine/Source/Editor/Sequencer/Public/ISequencer.h b/Engine/Source/Editor/Sequencer/Public/ISequencer.h
--- a/Engine/Source/Editor/Sequencer/Public/ISequencer.h
+++ b/Engine/Source/Editor/Sequencer/Public/ISequencer.h
@@ -6,6 +6,7 @@
 #include "Runtime/MovieSceneCore/Public/IMovieScenePlayer.h"
 
 class UMovieScene;
+class FSequencerSelection;
 
 /**
  * Sequencer public interface
@@ -126,4 +127,6 @@ public:
 	virtual void UpdateRuntimeInstances() = 0;
 
 	virtual TSharedRef<ISequencerObjectBindingManager> GetObjectBindingManager() const = 0;
+
+	virtual FSequencerSelection* GetSelection() = 0;
 };
