Fix CIS (shadow variables)

diff --git a/Engine/Source/Editor/Sequencer/Private/Sequencer.cpp b/Engine/Source/Editor/Sequencer/Private/Sequencer.cpp
--- a/Engine/Source/Editor/Sequencer/Private/Sequencer.cpp
+++ b/Engine/Source/Editor/Sequencer/Private/Sequencer.cpp
@@ -1333,9 +1333,9 @@ void FSequencer::DeleteSelectedItems()
 
 void FSequencer::SetKey()
 {
-	USelection* Selection = GEditor->GetSelectedActors();
+	USelection* CurrentSelection = GEditor->GetSelectedActors();
 	TArray<UObject*> SelectedActors;
-	Selection->GetSelectedObjects( AActor::StaticClass(), SelectedActors );
+	CurrentSelection->GetSelectedObjects( AActor::StaticClass(), SelectedActors );
 	for (TArray<UObject*>::TIterator It(SelectedActors); It; ++It)
 	{
 		// @todo Handle case of actors which aren't in sequencer yet
