UMG - Sequencer - Fix an issue where modifying the active animation wouldn't change the preview widgets.

diff --git a/Engine/Source/Editor/Sequencer/Private/SSection.cpp b/Engine/Source/Editor/Sequencer/Private/SSection.cpp
--- a/Engine/Source/Editor/Sequencer/Private/SSection.cpp
+++ b/Engine/Source/Editor/Sequencer/Private/SSection.cpp
@@ -179,7 +179,7 @@ void SSection::CreateDragOperation( const FGeometry& MyGeometry, const FPointerE
 
 	if( bKeysUnderMouse )
 	{
-		DragOperation = MakeShareable( new FMoveKeys( GetSequencer().GetSelection()->GetSelectedKeys(), PressedKey ) );
+		DragOperation = MakeShareable( new FMoveKeys( GetSequencer(), GetSequencer().GetSelection()->GetSelectedKeys(), PressedKey ) );
 	}
 	else
 	{
@@ -188,17 +188,17 @@ void SSection::CreateDragOperation( const FGeometry& MyGeometry, const FPointerE
 		if( bLeftEdgePressed || bLeftEdgeHovered )
 		{
 			// Selected the start of a section
-			DragOperation = MakeShareable( new FResizeSection( *SectionObject, false ) );
+			DragOperation = MakeShareable( new FResizeSection( GetSequencer(), *SectionObject, false ) );
 		}
 		else if( bRightEdgePressed || bRightEdgeHovered )
 		{
 			// Selected the end of a section
-			DragOperation = MakeShareable( new FResizeSection( *SectionObject, true ) );
+			DragOperation = MakeShareable( new FResizeSection( GetSequencer(), *SectionObject, true ) );
 		}
 		else
 		{
 			// Entire selection moved
-			DragOperation = MakeShareable( new FMoveSection( *SectionObject ) );
+			DragOperation = MakeShareable( new FMoveSection( GetSequencer(), *SectionObject ) );
 		}
 	}
 	
diff --git a/Engine/Source/Editor/Sequencer/Private/SectionDragOperations.cpp b/Engine/Source/Editor/Sequencer/Private/SectionDragOperations.cpp
--- a/Engine/Source/Editor/Sequencer/Private/SectionDragOperations.cpp
+++ b/Engine/Source/Editor/Sequencer/Private/SectionDragOperations.cpp
@@ -29,6 +29,7 @@ void FSequencerDragOperation::EndTransaction()
 	{
 		GEditor->EndTransaction();
 	}
+	Sequencer.UpdateRuntimeInstances();
 }
 
 TRange<float> FSequencerDragOperation::GetSectionBoundaries(UMovieSceneSection* Section, TSharedPtr<FTrackNode> SequencerNode) const
@@ -133,8 +134,9 @@ TOptional<float> FSequencerDragOperation::SnapToTimes(float InitialTime, const T
 }
 
 
-FResizeSection::FResizeSection( UMovieSceneSection& InSection, bool bInDraggingByEnd )
-	: Section( &InSection )
+FResizeSection::FResizeSection( FSequencer& Sequencer, UMovieSceneSection& InSection, bool bInDraggingByEnd )
+	: FSequencerDragOperation( Sequencer )
+	, Section( &InSection )
 	, bDraggingByEnd( bInDraggingByEnd )
 {
 }
@@ -225,8 +227,9 @@ void FResizeSection::OnDrag( const FPointerEvent& MouseEvent, const FVector2D& L
 	}
 }
 
-FMoveSection::FMoveSection( UMovieSceneSection& InSection )
-	: Section( &InSection )
+FMoveSection::FMoveSection( FSequencer& Sequencer, UMovieSceneSection& InSection )
+	: FSequencerDragOperation( Sequencer )
+	, Section( &InSection )
 	, DragOffset(ForceInit)
 {
 }
@@ -243,10 +246,7 @@ void FMoveSection::OnBeginDrag(const FVector2D& LocalMousePos, TSharedPtr<FTrack
 
 void FMoveSection::OnEndDrag(TSharedPtr<FTrackNode> SequencerNode)
 {
-	if( GEditor->IsTransactionActive() )
-	{
-		GEditor->EndTransaction();
-	}
+	EndTransaction();
 
 	if (Section.IsValid())
 	{
@@ -411,10 +411,7 @@ void FMoveKeys::OnBeginDrag(const FVector2D& LocalMousePos, TSharedPtr<FTrackNod
 
 void FMoveKeys::OnEndDrag(TSharedPtr<FTrackNode> SequencerNode)
 {
-	if( GEditor->IsTransactionActive() )
-	{
-		GEditor->EndTransaction();
-	}
+	EndTransaction();
 }
 
 void FMoveKeys::OnDrag( const FPointerEvent& MouseEvent, const FVector2D& LocalMousePos, const FTimeToPixel& TimeToPixelConverter, TSharedPtr<FTrackNode> SequencerNode )
diff --git a/Engine/Source/Editor/Sequencer/Private/SectionDragOperations.h b/Engine/Source/Editor/Sequencer/Private/SectionDragOperations.h
--- a/Engine/Source/Editor/Sequencer/Private/SectionDragOperations.h
+++ b/Engine/Source/Editor/Sequencer/Private/SectionDragOperations.h
@@ -8,7 +8,8 @@
 class FSequencerDragOperation
 {
 public:
-	FSequencerDragOperation()
+	FSequencerDragOperation( FSequencer& InSequencer )
+		: Sequencer(InSequencer)
 	{
 		Settings = GetDefault<USequencerSettings>();
 	}
@@ -77,6 +78,8 @@ protected:
 
 	/** The current sequencer settings */
 	const USequencerSettings* Settings;
+
+	FSequencer& Sequencer;
 };
 
 
@@ -87,7 +90,7 @@ protected:
 class FResizeSection : public FSequencerDragOperation
 {
 public:
-	FResizeSection( UMovieSceneSection& InSection, bool bInDraggingByEnd );
+	FResizeSection( FSequencer& Sequencer, UMovieSceneSection& InSection, bool bInDraggingByEnd );
 
 	/** FSequencerDragOperation interface */
 	virtual void OnBeginDrag(const FVector2D& LocalMousePos, TSharedPtr<FTrackNode> SequencerNode) override;
@@ -107,7 +110,7 @@ private:
 class FMoveSection : public FSequencerDragOperation
 {
 public:
-	FMoveSection( UMovieSceneSection& InSection );
+	FMoveSection( FSequencer& Sequencer, UMovieSceneSection& InSection );
 
 	/** FSequencerDragOperation interface */
 	virtual void OnBeginDrag(const FVector2D& LocalMousePos, TSharedPtr<FTrackNode> SequencerNode) override;
@@ -132,8 +135,9 @@ private:
 class FMoveKeys : public FSequencerDragOperation
 {
 public:
-	FMoveKeys( const TSet<FSelectedKey>* InSelectedKeys, FSelectedKey& PressedKey )
-		: SelectedKeys( InSelectedKeys )
+	FMoveKeys( FSequencer& Sequencer,  const TSet<FSelectedKey>* InSelectedKeys, FSelectedKey& PressedKey )
+		: FSequencerDragOperation(Sequencer)
+		, SelectedKeys( InSelectedKeys )
 		, DraggedKey( PressedKey )
 	{}
 
diff --git a/Engine/Source/Editor/Sequencer/Private/Sequencer.cpp b/Engine/Source/Editor/Sequencer/Private/Sequencer.cpp
--- a/Engine/Source/Editor/Sequencer/Private/Sequencer.cpp
+++ b/Engine/Source/Editor/Sequencer/Private/Sequencer.cpp
@@ -368,25 +368,30 @@ void FSequencer::DeleteSection(class UMovieSceneSection* Section)
 	
 	if( bAnythingRemoved )
 	{
-		// @todo Sequencer - Remove this
-		NotifyMovieSceneDataChanged();
+		UpdateRuntimeInstances();
 	}
 }
 
 void FSequencer::DeleteSelectedKeys()
 {
 	FScopedTransaction DeleteKeysTransaction( NSLOCTEXT("Sequencer", "DeleteSelectedKeys_Transaction", "Delete Selected Keys") );
-
+	bool bAnythingRemoved = false;
 	TArray<FSelectedKey> SelectedKeysArray = Selection.GetSelectedKeys()->Array();
+
 	for ( const FSelectedKey& Key : SelectedKeysArray )
 	{
 		if (Key.IsValid())
 		{
 			Key.Section->Modify();
-
 			Key.KeyArea->DeleteKey(Key.KeyHandle.GetValue());
+			bAnythingRemoved = true;
 		}
 	}
+
+	if (bAnythingRemoved)
+	{
+		UpdateRuntimeInstances();
+	}
 }
 
 void FSequencer::RenameShotCommitted(const FText& RenameText, ETextCommit::Type CommitInfo, UMovieSceneSection* Section)
@@ -675,8 +680,7 @@ void FSequencer::ResetPerMovieSceneData()
 	FilterToShotSections( TArray< TWeakObjectPtr<UMovieSceneSection> >(), false );
 
 	//@todo Sequencer - We may want to preserve selections when moving between movie scenes
-	Selection.EmptySelectedKeys();
-	Selection.EmptySelectedSections();
+	Selection.Empty();
 
 	// @todo run through all tracks for new movie scene changes
 	//  needed for audio track decompression
