Non-unity build fix

diff --git a/Engine/Source/Editor/Sequencer/Private/SectionDragOperations.cpp b/Engine/Source/Editor/Sequencer/Private/SectionDragOperations.cpp
--- a/Engine/Source/Editor/Sequencer/Private/SectionDragOperations.cpp
+++ b/Engine/Source/Editor/Sequencer/Private/SectionDragOperations.cpp
@@ -2,6 +2,7 @@
 
 #include "SequencerPrivatePCH.h"
 #include "SectionDragOperations.h"
+#include "Sequencer.h"
 #include "MovieSceneSection.h"
 #include "ISequencerSection.h"
 #include "IKeyArea.h"
