Fix sequencer auto-key crashing on components  who are no longer valid because their owners have had their construction scripts re-run

diff --git a/Engine/Source/Editor/Sequencer/Private/SequencerObjectChangeListener.cpp b/Engine/Source/Editor/Sequencer/Private/SequencerObjectChangeListener.cpp
--- a/Engine/Source/Editor/Sequencer/Private/SequencerObjectChangeListener.cpp
+++ b/Engine/Source/Editor/Sequencer/Private/SequencerObjectChangeListener.cpp
@@ -30,7 +30,10 @@ void FSequencerObjectChangeListener::OnPropertyChanged( const TArray<UObject*>&
 	
 	for( UObject* Object : ChangedObjects )
 	{
-		bIsKeyable |= IsTypeKeyable( *Object->GetClass(), PropertyHandle );
+		if( Object )
+		{
+			bIsKeyable |= IsTypeKeyable( *Object->GetClass(), PropertyHandle );
+		}
 	}
 
 	if( bIsKeyable )
