Sequencer improvements:
 - Reworked UI hierarchy to better afford input on its constituent panels
 - Added scroll-zoom/pan support to the track area and curve editor (using CTRL and SHIFT mouse wheel modifiers)
 - View range no longer zooms when changing the size of the sequencer window

diff --git a/Engine/Source/Developer/LogVisualizer/Private/SVisualLoggerView.cpp b/Engine/Source/Developer/LogVisualizer/Private/SVisualLoggerView.cpp
--- a/Engine/Source/Developer/LogVisualizer/Private/SVisualLoggerView.cpp
+++ b/Engine/Source/Developer/LogVisualizer/Private/SVisualLoggerView.cpp
@@ -75,7 +75,7 @@ void SVisualLoggerView::Construct(const FArguments& InArgs, const TSharedRef<FUI
 	VisualLoggerEvents = FLogVisualizer::Get().GetVisualLoggerEvents();
 
 	FVisualLoggerTimeSliderArgs TimeSliderArgs;
-	TimeSliderArgs.ViewRange = InArgs._ViewRange;
+	TimeSliderArgs.ViewRange = FAnimatedRange::WrapAttribute(InArgs._ViewRange);
 	TimeSliderArgs.ClampMin = InArgs._ViewRange.Get().GetLowerBoundValue();
 	TimeSliderArgs.ClampMax = InArgs._ViewRange.Get().GetUpperBoundValue();
 	TimeSliderArgs.ScrubPosition = InArgs._ScrubPosition;
diff --git a/Engine/Source/Developer/LogVisualizer/Private/TimeSliderController.cpp b/Engine/Source/Developer/LogVisualizer/Private/TimeSliderController.cpp
--- a/Engine/Source/Developer/LogVisualizer/Private/TimeSliderController.cpp
+++ b/Engine/Source/Developer/LogVisualizer/Private/TimeSliderController.cpp
@@ -422,7 +422,7 @@ FReply FVisualLoggerTimeSliderController::OnMouseMove( TSharedRef<SWidget> Widge
 					NewViewOutputMax = LocalClampMax.GetValue();
 				}
 
-				TimeSliderArgs.OnViewRangeChanged.ExecuteIfBound(TRange<float>(NewViewOutputMin, NewViewOutputMax));
+				TimeSliderArgs.OnViewRangeChanged.ExecuteIfBound(TRange<float>(NewViewOutputMin, NewViewOutputMax), EViewRangeInterpolation::Immediate);
 				if (Scrollbar.IsValid())
 				{
 					float InOffsetFraction = (NewViewOutputMin - LocalClampMin.GetValue()) / (LocalClampMax.GetValue() - LocalClampMin.GetValue());
@@ -597,7 +597,7 @@ FReply FVisualLoggerTimeSliderController::OnMouseWheel( TSharedRef<SWidget> Widg
 					NewViewOutputMax = LocalClampMax.GetValue();
 				}
 
-				TimeSliderArgs.OnViewRangeChanged.ExecuteIfBound(TRange<float>(NewViewOutputMin, NewViewOutputMax));
+				TimeSliderArgs.OnViewRangeChanged.ExecuteIfBound(TRange<float>(NewViewOutputMin, NewViewOutputMax), EViewRangeInterpolation::Immediate);
 				if (Scrollbar.IsValid())
 				{
 					float InOffsetFraction = (NewViewOutputMin - LocalClampMin.GetValue()) / (LocalClampMax.GetValue() - LocalClampMin.GetValue());
diff --git a/Engine/Source/Editor/Sequencer/Private/SAnimationOutlinerView.cpp b/Engine/Source/Editor/Sequencer/Private/SAnimationOutlinerView.cpp
--- a/Engine/Source/Editor/Sequencer/Private/SAnimationOutlinerView.cpp
+++ b/Engine/Source/Editor/Sequencer/Private/SAnimationOutlinerView.cpp
@@ -286,17 +286,13 @@ FReply SAnimationOutlinerTreeNode::OnMouseButtonDown( const FGeometry& MyGeometr
 		OnSelectionChanged.ExecuteIfBound( DisplayNode );
 		return FReply::Handled();
 	}
-	else if( MouseEvent.GetEffectingButton() == EKeys::RightMouseButton )
-	{
-		return FReply::Handled().CaptureMouse(SharedThis(this));
-	}
 
 	return FReply::Unhandled();
 }
 
 FReply SAnimationOutlinerTreeNode::OnMouseButtonUp( const FGeometry& MyGeometry, const FPointerEvent& MouseEvent )
 {
-	if( MouseEvent.GetEffectingButton() == EKeys::RightMouseButton && HasMouseCapture() )
+	if( MouseEvent.GetEffectingButton() == EKeys::RightMouseButton )
 	{
 		TSharedPtr<SWidget> MenuContent = DisplayNode->OnSummonContextMenu(MyGeometry, MouseEvent);
 		if (MenuContent.IsValid())
@@ -311,10 +307,10 @@ FReply SAnimationOutlinerTreeNode::OnMouseButtonUp( const FGeometry& MyGeometry,
 				FPopupTransitionEffect( FPopupTransitionEffect::ContextMenu )
 				);
 			
-			return FReply::Handled().ReleaseMouseCapture().SetUserFocus(MenuContent.ToSharedRef(), EFocusCause::SetDirectly);
+			return FReply::Handled().SetUserFocus(MenuContent.ToSharedRef(), EFocusCause::SetDirectly);
 		}
 
-		return FReply::Handled().ReleaseMouseCapture();
+		return FReply::Handled();
 	}
 
 	return FReply::Unhandled();
diff --git a/Engine/Source/Editor/Sequencer/Private/SSequencer.cpp b/Engine/Source/Editor/Sequencer/Private/SSequencer.cpp
--- a/Engine/Source/Editor/Sequencer/Private/SSequencer.cpp
+++ b/Engine/Source/Editor/Sequencer/Private/SSequencer.cpp
@@ -13,6 +13,7 @@
 #include "MovieSceneTrackEditor.h"
 #include "SSequencerSectionOverlay.h"
 #include "SSequencerTrackArea.h"
+#include "SSequencerTrackOutliner.h"
 #include "SequencerNodeTree.h"
 #include "TimeSliderController.h"
 #include "DragAndDrop/AssetDragDropOp.h"
@@ -25,7 +26,9 @@
 #include "SSearchBox.h"
 #include "SNumericDropDown.h"
 #include "EditorWidgetsModule.h"
+#include "SequencerTrackLaneFactory.h"
 
+#include "SSequencerSplitterOverlay.h"
 
 #define LOCTEXT_NAMESPACE "Sequencer"
 
@@ -167,17 +170,32 @@ void SSequencer::Construct( const FArguments& InArgs, TSharedRef< class FSequenc
 
 	OnGetAddMenuContent = InArgs._OnGetAddMenuContent;
 
-	OutlinerFillPercentage = .25f;
-	SectionFillPercentage = .75f;
-	OutlinerFillPercentageAttribute.BindRaw(this, &SSequencer::GetOutlinerFillPercentage);
-	SectionFillPercentageAttribute.BindRaw(this, &SSequencer::GetSectionFillPercentage);
+	ColumnFillCoefficients[0] = .25f;
+	ColumnFillCoefficients[1] = .75f;
 
-	SAssignNew( TrackArea, SSequencerTrackArea, Sequencer.Pin().ToSharedRef() )
-		.ViewRange( InArgs._ViewRange )
-		.OutlinerFillPercentage( OutlinerFillPercentageAttribute )
-		.SectionFillPercentage( SectionFillPercentageAttribute )
-		.OnOutlinerFillPercentageChanged( this, &SSequencer::OnOutlinerFillPercentageChanged )
-		.OnSectionFillPercentageChanged( this, &SSequencer::OnSectionFillPercentageChanged );
+	TAttribute<float> FillCoefficient_0, FillCoefficient_1;
+	FillCoefficient_0.Bind(TAttribute<float>::FGetter::CreateSP(this, &SSequencer::GetColumnFillCoefficient, 0));
+	FillCoefficient_1.Bind(TAttribute<float>::FGetter::CreateSP(this, &SSequencer::GetColumnFillCoefficient, 1));
+
+	TSharedRef<SScrollBar> ScrollBar =
+		SNew(SScrollBar)
+		.Thickness(FVector2D(5.0f, 5.0f));
+
+	auto PinnedSequencer = Sequencer.Pin().ToSharedRef();
+
+	SAssignNew( TrackOutliner, SSequencerTrackOutliner );
+
+	SAssignNew( TrackArea, SSequencerTrackArea, TimeSliderController )
+		.Visibility( this, &SSequencer::GetTrackAreaVisibility );
+
+	SAssignNew( CurveEditor, SSequencerCurveEditor, PinnedSequencer, TimeSliderController )
+		.Visibility( this, &SSequencer::GetCurveEditorVisibility )
+		.ViewRange( FAnimatedRange::WrapAttribute(InArgs._ViewRange) );
+
+	const int32 				Column0	= 0,	Column1	= 1;
+	const int32 Row0	= 0,
+				Row1	= 1,
+				Row2	= 2;
 
 	ChildSlot
 	[
@@ -198,127 +216,167 @@ void SSequencer::Construct( const FArguments& InArgs, TSharedRef< class FSequenc
 				+SHorizontalBox::Slot()
 				.AutoWidth()
 				[
-					SNew( SSequencerCurveEditorToolBar, TrackArea->GetCurveEditor()->GetCommands() )
+					SNew( SSequencerCurveEditorToolBar, CurveEditor->GetCommands() )
 					.Visibility( this, &SSequencer::GetCurveEditorToolBarVisibility )
 				]
 			]
+
 			+ SVerticalBox::Slot()
-			.AutoHeight()
 			[
-				SNew( SSplitter )
-				.Style( FEditorStyle::Get(), "Sequencer.AnimationOutliner.Splitter" )
-				+ SSplitter::Slot()
-				.Value(OutlinerFillPercentageAttribute)
-				.OnSlotResized(SSplitter::FOnSlotResized::CreateSP(this, &SSequencer::OnOutlinerFillPercentageChanged))
+				SNew(SOverlay)
+
+				+ SOverlay::Slot()
 				[
-					SNew( SBox )
-					.Padding( FMargin( 0,0,4,0) )
+					SNew( SGridPanel )
+					.FillRow( 1, 1.f )
+					.FillColumn( 0, FillCoefficient_0 )
+					.FillColumn( 1, FillCoefficient_1 )
+
+					+ SGridPanel::Slot( Column0, Row0 )
+					.VAlign( VAlign_Center )
 					[
 						// Search box for searching through the outliner
 						SNew( SSearchBox )
 						.OnTextChanged( this, &SSequencer::OnOutlinerSearchChanged )
 					]
-				]
-				+ SSplitter::Slot()
-				.Value(SectionFillPercentageAttribute)
-				.OnSlotResized(SSplitter::FOnSlotResized::CreateSP(this, &SSequencer::OnSectionFillPercentageChanged))
-				[
-					SNew( SBorder )
-					// @todo Sequencer Do not change the paddings or the sliders scrub widgets wont line up
-					.Padding( FMargin( 0.0f, 2.0f, 0.0f, 0.0f ) )
-					.BorderImage( FEditorStyle::GetBrush("ToolPanel.GroupBorder") )
-					.BorderBackgroundColor( FLinearColor(.50f, .50f, .50f, 1.0f ) )
+
+					+ SGridPanel::Slot( Column0, Row1 )
+					.ColumnSpan(2)
 					[
-						TopTimeSlider
+						SNew(SHorizontalBox)
+
+						+ SHorizontalBox::Slot()
+						[
+							SNew( SOverlay )
+
+							+ SOverlay::Slot()
+							[
+								SNew( SScrollBox )
+								.ExternalScrollbar(ScrollBar)
+
+								+ SScrollBox::Slot()
+								[
+									SNew( SHorizontalBox )
+
+									+ SHorizontalBox::Slot()
+									.FillWidth( FillCoefficient_0 )
+									// Padding to allow space for the scroll bar
+									.Padding(FMargin(0,0,10.f,0))
+									[
+										TrackOutliner.ToSharedRef()
+									]
+
+									+ SHorizontalBox::Slot()
+									.FillWidth( FillCoefficient_1 )
+									[
+										TrackArea.ToSharedRef()
+									]
+								]
+							]
+
+							+ SOverlay::Slot()
+							.HAlign( HAlign_Right )
+							[
+								ScrollBar
+							]
+						]
+
+						+ SHorizontalBox::Slot()
+						.FillWidth( TAttribute<float>( this, &SSequencer::GetOutlinerSpacerFill ) )
+						[
+							SNew(SSpacer)
+						]
 					]
-				]
-			]
-			+ SVerticalBox::Slot()
-			.FillHeight(1.0f)
-			[
-				SNew( SOverlay )
-				+ SOverlay::Slot()
-				[
-					MakeSectionOverlay( TimeSliderController, InArgs._ViewRange, InArgs._ScrubPosition, false )
-				]
-				+ SOverlay::Slot()
-				[
-					TrackArea.ToSharedRef()
-				]
-				+ SOverlay::Slot()
-				[
-					SNew( SSplitter )
-					.Style(FEditorStyle::Get(), "Sequencer.AnimationOutliner.Splitter")
-					.Visibility( EVisibility::HitTestInvisible )
-					+ SSplitter::Slot()
-					.Value(OutlinerFillPercentageAttribute)
-					.OnSlotResized(SSplitter::FOnSlotResized::CreateSP(this, &SSequencer::OnOutlinerFillPercentageChanged))
+
+					+ SGridPanel::Slot( Column0, Row2 )
+					.HAlign( HAlign_Center )
 					[
-						// Take up space but display nothing. This is required so that all areas dependent on time align correctly
-						SNullWidget::NullWidget
+						MakeTransportControls()
 					]
-					+ SSplitter::Slot()
-					.Value(SectionFillPercentageAttribute)
-					.OnSlotResized(SSplitter::FOnSlotResized::CreateSP(this, &SSequencer::OnSectionFillPercentageChanged))
+
+					// Second column
+					+ SGridPanel::Slot( Column1, Row0 )
 					[
-						SNew(SShotFilterOverlay, Sequencer)
-						.ViewRange( InArgs._ViewRange )
+						SNew( SBorder )
+						.BorderImage( FEditorStyle::GetBrush("ToolPanel.GroupBorder") )
+						.BorderBackgroundColor( FLinearColor(.50f, .50f, .50f, 1.0f ) )
+						.Padding(0)
+						[
+							TopTimeSlider
+						]
+					]
+
+					// Overlay that draws the tick lines
+					+ SGridPanel::Slot( Column1, Row1, SGridPanel::Layer(0) )
+					[
+						SNew( SSequencerSectionOverlay, TimeSliderController )
+						.Visibility( EVisibility::HitTestInvisible )
+						.DisplayScrubPosition( false )
+						.DisplayTickLines( true )
+					]
+
+					// Curve editor
+					+ SGridPanel::Slot( Column1, Row1 )
+					[
+						CurveEditor.ToSharedRef()
+					]
+
+					// Overlay that draws the scrub position
+					+ SGridPanel::Slot( Column1, Row1, SGridPanel::Layer(30) )
+					[
+						SNew( SSequencerSectionOverlay, TimeSliderController )
+						.Visibility( EVisibility::HitTestInvisible )
+						.DisplayScrubPosition( true )
+						.DisplayTickLines( false )
+					]
+
+					+ SGridPanel::Slot( Column1, Row2 )
+					[
+						SNew( SBorder )
+						.BorderImage( FEditorStyle::GetBrush("ToolPanel.GroupBorder") )
+						.BorderBackgroundColor( FLinearColor(.50f, .50f, .50f, 1.0f ) )
+						.Padding(0)
+						[
+							SNew(SVerticalBox)
+
+							+ SVerticalBox::Slot()
+							.AutoHeight()
+							[
+								BottomTimeSlider
+							]
+
+							+ SVerticalBox::Slot()
+							.AutoHeight()
+							[
+								SAssignNew( BreadcrumbTrail, SBreadcrumbTrail<FSequencerBreadcrumb> )
+								.Visibility( this, &SSequencer::GetBreadcrumbTrailVisibility )
+								.OnCrumbClicked( this, &SSequencer::OnCrumbClicked )
+							]
+						]
 					]
 				]
+
 				+ SOverlay::Slot()
 				[
-					MakeSectionOverlay( TimeSliderController, InArgs._ViewRange, InArgs._ScrubPosition, true )
-				]
-			]
-			+ SVerticalBox::Slot()
-			.AutoHeight()
-			[
-				SNew( SSplitter )
-				.Style(FEditorStyle::Get(), "Sequencer.AnimationOutliner.Splitter")
-				+ SSplitter::Slot()
-				.Value(OutlinerFillPercentageAttribute)
-				.OnSlotResized(SSplitter::FOnSlotResized::CreateSP(this, &SSequencer::OnOutlinerFillPercentageChanged))
-				[
-					SNew(SBox)
-					.HAlign(HAlign_Center)
+					SNew( SSequencerSplitterOverlay )
+					.Style(FEditorStyle::Get(), "Sequencer.AnimationOutliner.Splitter")
+					.Visibility(EVisibility::SelfHitTestInvisible)
+
+					+ SSplitter::Slot()
+					.Value( FillCoefficient_0 )
+					.OnSlotResized( SSplitter::FOnSlotResized::CreateSP(this, &SSequencer::OnColumnFillCoefficientChanged, 0) )
 					[
-						MakeTransportControls()
+						SNew(SSpacer)
 					]
-				]
-				+ SSplitter::Slot()
-				.Value(SectionFillPercentageAttribute)
-				.OnSlotResized(SSplitter::FOnSlotResized::CreateSP(this, &SSequencer::OnSectionFillPercentageChanged))
-				[
-					SNew( SBorder )
-					// @todo Sequencer Do not change the paddings or the sliders scrub widgets wont line up
-					.Padding( FMargin( 0.0f, 0.0f, 0.0f, 2.0f ) )
-					.BorderImage( FEditorStyle::GetBrush("ToolPanel.GroupBorder") )
-					.BorderBackgroundColor( FLinearColor(.50f, .50f, .50f, 1.0f ) )
+
+					+ SSplitter::Slot()
+					.Value( FillCoefficient_1 )
+					.OnSlotResized( SSplitter::FOnSlotResized::CreateSP(this, &SSequencer::OnColumnFillCoefficientChanged, 1) )
 					[
-						BottomTimeSlider
+						SNew(SSpacer)
 					]
 				]
 			]
-			+ SVerticalBox::Slot()
-			.AutoHeight()
-			[
-				SNew( SSplitter )
-				.Style(FEditorStyle::Get(), "Sequencer.AnimationOutliner.Splitter")
-				+ SSplitter::Slot()
-				.Value(OutlinerFillPercentageAttribute)
-				.OnSlotResized(SSplitter::FOnSlotResized::CreateSP(this, &SSequencer::OnOutlinerFillPercentageChanged))
-				[
-					SNew( SSpacer )
-				]
-				+ SSplitter::Slot()
-				.Value(SectionFillPercentageAttribute)
-				.OnSlotResized(SSplitter::FOnSlotResized::CreateSP(this, &SSequencer::OnSectionFillPercentageChanged))
-				[
-					SAssignNew( BreadcrumbTrail, SBreadcrumbTrail<FSequencerBreadcrumb> )
-					.Visibility(this, &SSequencer::GetBreadcrumbTrailVisibility)
-					.OnCrumbClicked( this, &SSequencer::OnCrumbClicked )
-				]
-			]
 		]
 	];
 
@@ -556,8 +614,11 @@ void SSequencer::UpdateLayoutTree()
 {	
 	// Update the node tree
 	SequencerNodeTree->Update();
-	
-	TrackArea->Update( SequencerNodeTree );
+
+	FSequencerTrackLaneFactory Factory(TrackOutliner.ToSharedRef(), TrackArea.ToSharedRef(), Sequencer.Pin().ToSharedRef() );
+	Factory.Repopulate( *SequencerNodeTree );
+
+	CurveEditor->SetSequencerNodeTree(SequencerNodeTree);
 
 	SequencerNodeTree->UpdateCachedVisibilityBasedOnShotFiltersChanged();
 }
@@ -606,29 +667,6 @@ void SSequencer::OnOutlinerSearchChanged( const FText& Filter )
 	SequencerNodeTree->FilterNodes( Filter.ToString() );
 }
 
-TSharedRef<SWidget> SSequencer::MakeSectionOverlay( TSharedRef<FSequencerTimeSliderController> TimeSliderController, const TAttribute< TRange<float> >& ViewRange, const TAttribute<float>& ScrubPosition, bool bTopOverlay )
-{
-	return
-		SNew( SSplitter )
-		.Style(FEditorStyle::Get(), "Sequencer.AnimationOutliner.Splitter")
-		.Visibility( EVisibility::HitTestInvisible )
-		+ SSplitter::Slot()
-		.Value(OutlinerFillPercentageAttribute)
-		.OnSlotResized(SSplitter::FOnSlotResized::CreateSP(this, &SSequencer::OnOutlinerFillPercentageChanged))
-		[
-			// Take up space but display nothing. This is required so that all areas dependent on time align correctly
-			SNullWidget::NullWidget
-		]
-		+ SSplitter::Slot()
-		.Value(SectionFillPercentageAttribute)
-		.OnSlotResized(SSplitter::FOnSlotResized::CreateSP(this, &SSequencer::OnSectionFillPercentageChanged))
-		[
-			SNew( SSequencerSectionOverlay, TimeSliderController )
-			.DisplayScrubPosition( bTopOverlay )
-			.DisplayTickLines( !bTopOverlay )
-		];
-}
-
 float SSequencer::OnGetTimeSnapInterval() const
 {
 	return GetDefault<USequencerSettings>()->GetTimeSnapInterval();
@@ -964,14 +1002,25 @@ EVisibility SSequencer::GetCurveEditorToolBarVisibility() const
 	return GetDefault<USequencerSettings>()->GetShowCurveEditor() ? EVisibility::Visible : EVisibility::Collapsed;
 }
 
-void SSequencer::OnOutlinerFillPercentageChanged(float InOutlinerFillPercentage)
+float SSequencer::GetOutlinerSpacerFill() const
+{
+	const float Column1Coeff = GetColumnFillCoefficient(1);
+	return GetDefault<USequencerSettings>()->GetShowCurveEditor() ? Column1Coeff / (1 - Column1Coeff) : 0.f;
+}
+
+void SSequencer::OnColumnFillCoefficientChanged(float FillCoefficient, int32 ColumnIndex)
 {
-	OutlinerFillPercentage = InOutlinerFillPercentage;
+	ColumnFillCoefficients[ColumnIndex] = FillCoefficient;
 }
 
-void SSequencer::OnSectionFillPercentageChanged(float InSectionFillPercentage)
+EVisibility SSequencer::GetTrackAreaVisibility() const
 {
-	SectionFillPercentage = InSectionFillPercentage;
+	return GetDefault<USequencerSettings>()->GetShowCurveEditor() ? EVisibility::Collapsed : EVisibility::Visible;
+}
+
+EVisibility SSequencer::GetCurveEditorVisibility() const
+{
+	return GetDefault<USequencerSettings>()->GetShowCurveEditor() ? EVisibility::Visible : EVisibility::Collapsed;
 }
 
 #undef LOCTEXT_NAMESPACE
diff --git a/Engine/Source/Editor/Sequencer/Private/SSequencer.h b/Engine/Source/Editor/Sequencer/Private/SSequencer.h
--- a/Engine/Source/Editor/Sequencer/Private/SSequencer.h
+++ b/Engine/Source/Editor/Sequencer/Private/SSequencer.h
@@ -50,11 +50,10 @@ class SSequencer : public SCompoundWidget
 public:
 	DECLARE_DELEGATE_OneParam( FOnToggleBoolOption, bool )
 	SLATE_BEGIN_ARGS( SSequencer )
-		: _ViewRange( TRange<float>( 0.0f, 5.0f ) )
-		, _ScrubPosition( 1.0f )
+		: _ScrubPosition( 1.0f )
 	{}
 		/** The current view range (seconds) */
-		SLATE_ATTRIBUTE( TRange<float>, ViewRange )
+		SLATE_ATTRIBUTE( FAnimatedRange, ViewRange )
 		/** The current scrub position in (seconds) */
 		SLATE_ATTRIBUTE( float, ScrubPosition )
 		/** Called when the user changes the view range */
@@ -140,22 +139,16 @@ private:
 	/**
 	 * @return The fill percentage of the animation outliner
 	 */
-	float GetOutlinerFillPercentage() const { return OutlinerFillPercentage; }
+	float GetColumnFillCoefficient(int32 ColumnIndex) const { return ColumnFillCoefficients[ColumnIndex]; }
 
-	/**
-	 * @return The fill percentage of the section area.
-	 */
-	float GetSectionFillPercentage() const { return SectionFillPercentage; }
+	/** Get the amount of space that the outliner spacer should fill */
+	float GetOutlinerSpacerFill() const;
 
-	/**
-	 * Creates an overlay for various components of the time slider that should be displayed on top or bottom of the section area
-	 * 
-	 * @param TimeSliderController	Time controller to route input to for adjustment of the time view
-	 * @param ViewRange		The current visible time range
-	 * @param ScrubPosition		The current scrub position
-	 * @param bTopOverlay		Whether or not the overlay is displayed on top or bottom of the section view
-	 */
-	TSharedRef<SWidget> MakeSectionOverlay( TSharedRef<class FSequencerTimeSliderController> TimeSliderController, const TAttribute< TRange<float> >& ViewRange, const TAttribute<float>& ScrubPosition, bool bTopOverlay );
+	/** Get the visibility of the curve area */
+	EVisibility GetCurveEditorVisibility() const;
+
+	/** Get the visibility of the track area */
+	EVisibility GetTrackAreaVisibility() const;
 
 	/** SWidget interface */
 	/** @todo Sequencer Basic drag and drop support.  Doesn't belong here most likely */
@@ -208,33 +201,25 @@ private:
 	/** Gets whether or not the curve editor toolbar should be visibe. */
 	EVisibility GetCurveEditorToolBarVisibility() const;
 
-	/** Called when the outliner fill percentage is changed by a splitter slot. */
-	void OnOutlinerFillPercentageChanged(float OutlinerFillPercentage);
-
-	/** Called when the section fill percentage is changed by a splitter slot. */
-	void OnSectionFillPercentageChanged(float SectionFillPercentage);
+	/** Called when a column fill percentage is changed by a splitter slot. */
+	void OnColumnFillCoefficientChanged(float FillCoefficient, int32 ColumnIndex);
 
 private:
+
 	/** Section area widget */
 	TSharedPtr<class SSequencerTrackArea> TrackArea;
+	/** Outliner widget */
+	TSharedPtr<class SSequencerTrackOutliner> TrackOutliner;
+	/** The curve editor. */
+	TSharedPtr<class SSequencerCurveEditor> CurveEditor;
 	/** Sequencer node tree for movie scene data */
 	TSharedPtr<class FSequencerNodeTree> SequencerNodeTree;
 	/** The breadcrumb trail widget for this sequencer */
 	TSharedPtr< class SBreadcrumbTrail<FSequencerBreadcrumb> > BreadcrumbTrail;
-	/** Whether or not the breadcrumb trail should be shown. */
-	TAttribute<bool> bShowBreadcrumbTrail;
 	/** The main sequencer interface */
 	TWeakPtr<FSequencer> Sequencer;
-	
-	/** The fill percentage of the outliner area. */
-	float OutlinerFillPercentage;
-	/** The fill percentage of the section area. */
-	float SectionFillPercentage;
-	/** An attribute for binding the fill percentage of the outliner area to a splitter slot. */
-	TAttribute<float>OutlinerFillPercentageAttribute;
-	/** An attribute for binding the fill percentage of the section area to a splitter slot. */
-	TAttribute<float>SectionFillPercentageAttribute;
-
+	/** The fill coefficients of each column in the grid. */
+	float ColumnFillCoefficients[2];
 	/** Whether the active timer is currently registered */
 	bool bIsActiveTimerRegistered;
 
diff --git a/Engine/Source/Editor/Sequencer/Private/SSequencerCurveEditor.cpp b/Engine/Source/Editor/Sequencer/Private/SSequencerCurveEditor.cpp
--- a/Engine/Source/Editor/Sequencer/Private/SSequencerCurveEditor.cpp
+++ b/Engine/Source/Editor/Sequencer/Private/SSequencerCurveEditor.cpp
@@ -3,19 +3,20 @@
 #include "SequencerPrivatePCH.h"
 #include "SCurveEditor.h"
 #include "Sequencer.h"
+#include "TimeSliderController.h"
 
 #define LOCTEXT_NAMESPACE "SequencerCurveEditor"
 
-void SSequencerCurveEditor::Construct( const FArguments& InArgs, TSharedRef<FSequencer> InSequencer )
+void SSequencerCurveEditor::Construct( const FArguments& InArgs, TSharedRef<FSequencer> InSequencer, TSharedRef<ITimeSliderController> InTimeSliderController )
 {
 	ViewRange = InArgs._ViewRange;
 	Sequencer = InSequencer;
+	TimeSliderController = InTimeSliderController;
 	NodeTreeSelectionChangedHandle = Sequencer.Pin()->GetSelection()->GetOnOutlinerNodeSelectionChanged()->AddSP(this, &SSequencerCurveEditor::NodeTreeSelectionChanged);
 	GetMutableDefault<USequencerSettings>()->GetOnCurveVisibilityChanged()->AddSP( this, &SSequencerCurveEditor::SequencerCurveVisibilityChanged );
 
-	ChildSlot
-	[
-		SAssignNew( CurveEditor, SCurveEditor )
+	SCurveEditor::Construct(
+		SCurveEditor::FArguments()
 		.ViewMinInput_Lambda( [this]{ return ViewRange.Get().GetLowerBoundValue(); } )
 		.ViewMaxInput_Lambda( [this]{ return ViewRange.Get().GetUpperBoundValue(); } )
 		.HideUI( false )
@@ -28,7 +29,55 @@ void SSequencerCurveEditor::Construct( const FArguments& InArgs, TSharedRef<FSeq
 		.TimelineLength( 0.0f )
 		.ShowCurveToolTips( this, &SSequencerCurveEditor::GetShowCurveEditorCurveToolTips )
 		.GridColor( FLinearColor( 0.3f, 0.3f, 0.3f, 0.3f ) )
-	];
+	);
+}
+
+FReply SSequencerCurveEditor::OnMouseButtonDown( const FGeometry& MyGeometry, const FPointerEvent& MouseEvent )
+{
+	// Curve Editor takes precedence
+	FReply Reply = SCurveEditor::OnMouseButtonDown( MyGeometry, MouseEvent );
+	if ( Reply.IsEventHandled() )
+	{
+		return Reply;
+	}
+
+	return TimeSliderController->OnMouseButtonDown( AsShared(), MyGeometry, MouseEvent );
+}
+
+
+FReply SSequencerCurveEditor::OnMouseButtonUp( const FGeometry& MyGeometry, const FPointerEvent& MouseEvent )
+{
+	// Curve Editor takes precedence
+	FReply Reply = SCurveEditor::OnMouseButtonUp( MyGeometry, MouseEvent );
+	if ( Reply.IsEventHandled() )
+	{
+		return Reply;
+	}
+
+	return TimeSliderController->OnMouseButtonUp( AsShared(), MyGeometry, MouseEvent );
+}
+
+
+FReply SSequencerCurveEditor::OnMouseMove( const FGeometry& MyGeometry, const FPointerEvent& MouseEvent )
+{
+	FReply Reply = SCurveEditor::OnMouseMove( MyGeometry, MouseEvent );
+	if ( Reply.IsEventHandled() )
+	{
+		return Reply;
+	}
+
+	return TimeSliderController->OnMouseMove( AsShared(), MyGeometry, MouseEvent );
+}
+
+FReply SSequencerCurveEditor::OnMouseWheel( const FGeometry& MyGeometry, const FPointerEvent& MouseEvent )
+{
+	FReply Reply = TimeSliderController->OnMouseWheel( AsShared(), MyGeometry, MouseEvent );
+	if ( Reply.IsEventHandled() )
+	{
+		return Reply;
+	}
+
+	return SCurveEditor::OnMouseWheel( MyGeometry, MouseEvent );
 }
 
 template<class ParentNodeType>
@@ -55,7 +104,7 @@ void SSequencerCurveEditor::SetSequencerNodeTree( TSharedPtr<FSequencerNodeTree>
 void SSequencerCurveEditor::UpdateCurveOwner()
 {
 	CurveOwner = MakeShareable( new FSequencerCurveOwner( SequencerNodeTree ) );
-	CurveEditor->SetCurveOwner( CurveOwner.Get() );
+	SetCurveOwner( CurveOwner.Get() );
 }
 
 bool SSequencerCurveEditor::GetCurveSnapEnabled() const
@@ -99,11 +148,6 @@ void SSequencerCurveEditor::SequencerCurveVisibilityChanged()
 	UpdateCurveOwner();
 }
 
-TSharedPtr<FUICommandList> SSequencerCurveEditor::GetCommands()
-{
-	return CurveEditor->GetCommands();
-}
-
 SSequencerCurveEditor::~SSequencerCurveEditor()
 {
 	if ( Sequencer.IsValid() )
diff --git a/Engine/Source/Editor/Sequencer/Private/SSequencerCurveEditor.h b/Engine/Source/Editor/Sequencer/Private/SSequencerCurveEditor.h
--- a/Engine/Source/Editor/Sequencer/Private/SSequencerCurveEditor.h
+++ b/Engine/Source/Editor/Sequencer/Private/SSequencerCurveEditor.h
@@ -2,10 +2,10 @@
 
 #pragma once
 
-class SCurveEditor;
+#include "SCurveEditor.h"
 
 /** A widget for displaying and managing an SCurveEditor in sequencer. */
-class SSequencerCurveEditor : public SCompoundWidget
+class SSequencerCurveEditor : public SCurveEditor
 {
 public:
 	SLATE_BEGIN_ARGS( SSequencerCurveEditor )
@@ -16,16 +16,20 @@ public:
 
 	SLATE_END_ARGS()
 
-	void Construct( const FArguments& InArgs, TSharedRef<FSequencer> InSequencer );
+	void Construct( const FArguments& InArgs, TSharedRef<FSequencer> InSequencer, TSharedRef<class ITimeSliderController> InTimeSliderController );
 
 	/** Sets the sequencer node tree which supplies the curves. */
 	void SetSequencerNodeTree( TSharedPtr<FSequencerNodeTree> InSequencerNodeTree );
-	
-	/** Gets the commands for building a toolbar for the curve editor. */
-	TSharedPtr<FUICommandList> GetCommands();
 
 	~SSequencerCurveEditor();
 
+protected:
+
+	virtual FReply OnMouseButtonDown( const FGeometry& MyGeometry, const FPointerEvent& MouseEvent ) override;
+	virtual FReply OnMouseButtonUp( const FGeometry& MyGeometry, const FPointerEvent& MouseEvent ) override;
+	virtual FReply OnMouseMove( const FGeometry& MyGeometry, const FPointerEvent& MouseEvent ) override;
+	virtual FReply OnMouseWheel( const FGeometry& MyGeometry, const FPointerEvent& MouseEvent ) override;
+
 private:
 	/** Builds and assigns a new curve owner to the FCurveEditor. */
 	void UpdateCurveOwner();
@@ -52,8 +56,8 @@ private:
 
 	/** The sequencer which owns this widget. */
 	TWeakPtr<FSequencer> Sequencer;
-	/** The curve editor widget contained by this widget. */
-	TSharedPtr<SCurveEditor> CurveEditor;
+	/** The class responsible for time sliding on the curve editor */
+	TSharedPtr<class ITimeSliderController> TimeSliderController;
 	/** The visible time range displayed by the curve editor. */
 	TAttribute<TRange<float>> ViewRange;
 	/** The sequencer node tree which contains the key area nodes which supply the curves to edit. */
diff --git a/Engine/Source/Editor/Sequencer/Private/SSequencerSectionOverlay.cpp b/Engine/Source/Editor/Sequencer/Private/SSequencerSectionOverlay.cpp
--- a/Engine/Source/Editor/Sequencer/Private/SSequencerSectionOverlay.cpp
+++ b/Engine/Source/Editor/Sequencer/Private/SSequencerSectionOverlay.cpp
@@ -9,24 +9,4 @@ int32 SSequencerSectionOverlay::OnPaint( const FPaintArgs& Args, const FGeometry
 	TimeSliderController->OnPaintSectionView( AllottedGeometry, MyClippingRect, OutDrawElements, LayerId, ShouldBeEnabled( bParentEnabled ), bDisplayTickLines.Get(), bDisplayScrubPosition.Get() );
 
 	return SCompoundWidget::OnPaint( Args, AllottedGeometry, MyClippingRect, OutDrawElements, LayerId, InWidgetStyle, bParentEnabled  );
-}
-
-FReply SSequencerSectionOverlay::OnMouseButtonDown( const FGeometry& MyGeometry, const FPointerEvent& MouseEvent )
-{
-	return TimeSliderController->OnMouseButtonDown( SharedThis(this), MyGeometry, MouseEvent );
-}
-
-FReply SSequencerSectionOverlay::OnMouseButtonUp( const FGeometry& MyGeometry, const FPointerEvent& MouseEvent )
-{
-	return TimeSliderController->OnMouseButtonUp( SharedThis(this),  MyGeometry, MouseEvent );
-}
-
-FReply SSequencerSectionOverlay::OnMouseMove( const FGeometry& MyGeometry, const FPointerEvent& MouseEvent )
-{
-	return TimeSliderController->OnMouseMove( SharedThis(this), MyGeometry, MouseEvent );
-}
-
-FReply SSequencerSectionOverlay::OnMouseWheel( const FGeometry& MyGeometry, const FPointerEvent& MouseEvent )
-{
-	return TimeSliderController->OnMouseWheel( SharedThis(this), MyGeometry, MouseEvent );
-}
+}
\ No newline at end of file
diff --git a/Engine/Source/Editor/Sequencer/Private/SSequencerSectionOverlay.h b/Engine/Source/Editor/Sequencer/Private/SSequencerSectionOverlay.h
--- a/Engine/Source/Editor/Sequencer/Private/SSequencerSectionOverlay.h
+++ b/Engine/Source/Editor/Sequencer/Private/SSequencerSectionOverlay.h
@@ -27,10 +27,6 @@ public:
 private:
 	/** SWidget Interface */
 	virtual int32 OnPaint( const FPaintArgs& Args, const FGeometry& AllottedGeometry, const FSlateRect& MyClippingRect, FSlateWindowElementList& OutDrawElements, int32 LayerId, const FWidgetStyle& InWidgetStyle, bool bParentEnabled ) const override;
-	virtual FReply OnMouseButtonDown( const FGeometry& MyGeometry, const FPointerEvent& MouseEvent ) override;
-	virtual FReply OnMouseButtonUp( const FGeometry& MyGeometry, const FPointerEvent& MouseEvent ) override;
-	virtual FReply OnMouseMove( const FGeometry& MyGeometry, const FPointerEvent& MouseEvent ) override;
-	virtual FReply OnMouseWheel( const FGeometry& MyGeometry, const FPointerEvent& MouseEvent ) override;
 
 private:
 	/** Controller for manipulating time */
diff --git a/Engine/Source/Editor/Sequencer/Private/SSequencerSplitterOverlay.h b/Engine/Source/Editor/Sequencer/Private/SSequencerSplitterOverlay.h
new file mode 100644
--- /dev/null
+++ b/Engine/Source/Editor/Sequencer/Private/SSequencerSplitterOverlay.h
@@ -0,0 +1,123 @@
+// Copyright 1998-2015 Epic Games, Inc. All Rights Reserved.
+
+#pragma once
+
+#include "SSplitter.h"
+
+/** Splitter used on the sequencer as an overlay. Input is disabled on all areas except the draggable positions */
+class SSequencerSplitterOverlay : public SOverlay
+{
+public:
+	typedef SSplitter::FArguments FArguments;
+
+	void Construct( const FArguments& InArgs )
+	{
+		SetVisibility(EVisibility::SelfHitTestInvisible);
+
+		Splitter = SNew(SSplitter) = InArgs;
+		Splitter->SetVisibility(EVisibility::HitTestInvisible);
+		AddSlot()
+		[
+			Splitter.ToSharedRef()
+		];
+
+		for (int32 Index = 0; Index < Splitter->GetChildren()->Num() - 1; ++Index)
+		{
+			AddSlot()
+			.Padding(TAttribute<FMargin>::Create(TAttribute<FMargin>::FGetter::CreateSP(this, &SSequencerSplitterOverlay::GetSplitterHandlePadding, Index)))
+			[
+				SNew(SBox)
+				.Visibility(EVisibility::Visible)
+			];
+		}
+	}
+
+	virtual void OnArrangeChildren( const FGeometry& AllottedGeometry, FArrangedChildren& ArrangedChildren ) const override
+	{
+		FArrangedChildren SplitterChildren(ArrangedChildren.GetFilter());
+		Splitter->ArrangeChildren(AllottedGeometry, SplitterChildren);
+
+		SlotPadding.Reset();
+
+		for (int32 Index = 0; Index < SplitterChildren.Num() - 1; ++Index)
+		{
+			const auto& ThisGeometry = SplitterChildren[Index].Geometry;
+			const auto& NextGeometry = SplitterChildren[Index + 1].Geometry;
+
+			if (Splitter->GetOrientation() == EOrientation::Orient_Horizontal)
+			{
+				SlotPadding.Add(FMargin(
+					ThisGeometry.Position.X + ThisGeometry.GetLocalSize().X,
+					0,
+					AllottedGeometry.Size.X - NextGeometry.Position.X,
+					0)
+				);
+			}
+			else
+			{
+				SlotPadding.Add(FMargin(
+					0,
+					ThisGeometry.Position.Y + ThisGeometry.GetLocalSize().Y,
+					0,
+					AllottedGeometry.Size.Y - NextGeometry.Position.Y)
+				);
+			}
+		}
+
+		SOverlay::OnArrangeChildren(AllottedGeometry, ArrangedChildren);
+	}
+
+	FMargin GetSplitterHandlePadding(int32 Index) const
+	{
+		if (SlotPadding.IsValidIndex(Index))
+		{
+			return SlotPadding[Index];
+		}
+
+		return 0.f;
+	}
+	
+	virtual FCursorReply OnCursorQuery( const FGeometry& MyGeometry, const FPointerEvent& CursorEvent ) const override
+	{
+		return Splitter->OnCursorQuery(MyGeometry, CursorEvent);
+	}
+
+	virtual FReply OnMouseButtonDown( const FGeometry& MyGeometry, const FPointerEvent& MouseEvent ) override
+	{
+		FReply Reply = Splitter->OnMouseButtonDown(MyGeometry, MouseEvent);
+		if (Reply.GetMouseCaptor().IsValid())
+		{
+			// Set us to be the mouse captor so we can forward events properly
+			Reply.CaptureMouse( SharedThis(this) );
+			SetVisibility(EVisibility::Visible);
+		}
+		return Reply;
+	}
+	
+	virtual void OnMouseCaptureLost() override
+	{
+		SetVisibility(EVisibility::SelfHitTestInvisible);
+		SOverlay::OnMouseCaptureLost();
+	}
+
+	virtual FReply OnMouseButtonUp( const FGeometry& MyGeometry, const FPointerEvent& MouseEvent ) override
+	{
+		FReply Reply = Splitter->OnMouseButtonUp(MyGeometry, MouseEvent);
+		if (Reply.ShouldReleaseMouse())
+		{
+			SetVisibility(EVisibility::SelfHitTestInvisible);
+		}
+		return Reply;
+	}
+	virtual FReply OnMouseMove( const FGeometry& MyGeometry, const FPointerEvent& MouseEvent ) override
+	{
+		return Splitter->OnMouseMove(MyGeometry, MouseEvent);
+	}
+	virtual void OnMouseLeave( const FPointerEvent& MouseEvent ) override
+	{
+		return Splitter->OnMouseLeave(MouseEvent);
+	}
+
+	TSharedPtr<SSplitter> Splitter;
+	mutable TArray<FMargin> SlotPadding;
+};
\ No newline at end of file
diff --git a/Engine/Source/Editor/Sequencer/Private/SSequencerTrackArea.cpp b/Engine/Source/Editor/Sequencer/Private/SSequencerTrackArea.cpp
--- a/Engine/Source/Editor/Sequencer/Private/SSequencerTrackArea.cpp
+++ b/Engine/Source/Editor/Sequencer/Private/SSequencerTrackArea.cpp
@@ -2,180 +2,50 @@
 
 #include "SequencerPrivatePCH.h"
 #include "SSequencerTrackArea.h"
-#include "ISequencerSection.h"
-#include "MovieSceneSection.h"
-#include "SSequencer.h"
 #include "TimeSliderController.h"
-#include "SSequencerSectionAreaView.h"
 
-void SSequencerTrackArea::Construct( const FArguments& InArgs, TSharedRef<FSequencer> InSequencer )
+void SSequencerTrackArea::Construct( const FArguments& InArgs, TSharedRef<FSequencerTimeSliderController> InTimeSliderController )
 {
-	ViewRange = InArgs._ViewRange;
-	OutlinerFillPercentage = InArgs._OutlinerFillPercentage;
-	SectionFillPercentage = InArgs._SectionFillPercentage;
-	OnOutlinerFillPercentageChanged = InArgs._OnOutlinerFillPercentageChanged;
-	OnSectionFillPercentageChanged = InArgs._OnSectionFillPercentageChanged;
-	Sequencer = InSequencer;
-
-	ScrollbarFillPercentage.BindRaw(this, &SSequencerTrackArea::GetScrollBarFillPercentage);
-	ScrollbarSpacerFillPercentage.BindRaw(this, &SSequencerTrackArea::GetScrollBarSpacerFillPercentage);
-
-	TSharedRef<SScrollBar> ScrollBar =
-		SNew(SScrollBar)
-		.Thickness(FVector2D(5.0f, 5.0f));
-
-	ChildSlot
-	[
-		SNew(SOverlay)
-		+ SOverlay::Slot()
-		[
-			SAssignNew(ScrollBox, SScrollBox)
-			.ExternalScrollbar(ScrollBar)
-		]
-		+ SOverlay::Slot()
-		[
-			SNew(SSplitter)
-			.Style(FEditorStyle::Get(), "Sequencer.AnimationOutliner.Splitter")
-			.Visibility(EVisibility::SelfHitTestInvisible)
-			+ SSplitter::Slot()
-			.Value(OutlinerFillPercentage)
-			.OnSlotResized(SSplitter::FOnSlotResized::CreateSP(this, &SSequencerTrackArea::OutlinerFillPercentageChanged))
-			[
-				SNew(SSpacer)
-			]
-			+ SSplitter::Slot()
-			.Value(SectionFillPercentage)
-			.OnSlotResized(SSplitter::FOnSlotResized::CreateSP(this, &SSequencerTrackArea::SectionFillPercentageChanged))
-			[
-				SAssignNew(CurveEditor, SSequencerCurveEditor, InSequencer)
-				.Visibility( this, &SSequencerTrackArea::GetCurveEditorVisibility )
-				.ViewRange(ViewRange)
-			]
-		]
-		+ SOverlay::Slot()
-		[
-			SNew(SSplitter)
-			.Style(FEditorStyle::Get(), "Sequencer.AnimationOutliner.Splitter")
-			.Visibility(EVisibility::SelfHitTestInvisible)
-			+ SSplitter::Slot()
-			.Value(ScrollbarFillPercentage)
-			[
-				SNew(SBox)
-				.HAlign(HAlign_Right)
-				[
-					ScrollBar
-				]
-			]
-			+ SSplitter::Slot()
-			.Value(ScrollbarSpacerFillPercentage)
-			[
-				SNew( SSpacer )
-			]
-		]
-	];
+	TimeSliderController = InTimeSliderController;
+	SVerticalBox::Construct( InArgs );
 }
 
-void SSequencerTrackArea::Update( TSharedPtr<FSequencerNodeTree> InSequencerNodeTree )
+FReply SSequencerTrackArea::OnMouseButtonDown( const FGeometry& MyGeometry, const FPointerEvent& MouseEvent )
 {
-	ScrollBox->ClearChildren();
-
-	const TArray< TSharedRef<FSequencerDisplayNode> >& RootNodes = InSequencerNodeTree->GetRootNodes();
-
-	for( int32 RootNodeIndex = 0; RootNodeIndex < RootNodes.Num(); ++RootNodeIndex )
-	{
-		TSharedRef<FSequencerDisplayNode> RootNode = RootNodes[RootNodeIndex];
-
-		// Generate a widget for each root node
-		GenerateWidgetForNode( RootNode );
-
-		// Generate a node for each children of the root node
-		GenerateLayoutNodeWidgetsRecursive( RootNode->GetChildNodes() );
-	}
-
-	// @todo Sequencer - Remove this expensive operation
-	ScrollBox->SlatePrepass();
-
-	CurveEditor->SetSequencerNodeTree(InSequencerNodeTree);
+	return TimeSliderController->OnMouseButtonDown( SharedThis(this), MyGeometry, MouseEvent );
 }
 
-void SSequencerTrackArea::GenerateLayoutNodeWidgetsRecursive( const TArray< TSharedRef<FSequencerDisplayNode> >& Nodes )
+FReply SSequencerTrackArea::OnMouseButtonUp( const FGeometry& MyGeometry, const FPointerEvent& MouseEvent )
 {
-	for( int32 NodeIndex = 0; NodeIndex < Nodes.Num(); ++NodeIndex )
-	{
-		TSharedRef<FSequencerDisplayNode> Node = Nodes[NodeIndex];
-
-		// The section area encompasses multiple logical child nodes so we do not generate children that arent section areas (they would not be top level nodes in that case)
-		if( Node->GetType() == ESequencerNode::Track )
-		{
-			GenerateWidgetForNode( Node );	
-
-			GenerateLayoutNodeWidgetsRecursive( Node->GetChildNodes() );
-		}
-	}
+	return TimeSliderController->OnMouseButtonUp( SharedThis(this),  MyGeometry, MouseEvent );
 }
 
-void SSequencerTrackArea::GenerateWidgetForNode( TSharedRef<FSequencerDisplayNode>& Node )
+FReply SSequencerTrackArea::OnMouseMove( const FGeometry& MyGeometry, const FPointerEvent& MouseEvent )
 {
-	ScrollBox->AddSlot()
-	[
-		SNew( SSplitter )
-		.Style(FEditorStyle::Get(), "Sequencer.AnimationOutliner.Splitter")
-		+ SSplitter::Slot()
-		.Value(OutlinerFillPercentage)
-		.OnSlotResized(SSplitter::FOnSlotResized::CreateSP(this, &SSequencerTrackArea::OutlinerFillPercentageChanged))
-		[
-			SNew(SBox)
-			.Padding(this, &SSequencerTrackArea::GetOutlinerPadding)
-			[
-				// Generate a widget for the animation outliner portion of the node
-				Node->GenerateContainerWidgetForOutliner()
-			]
-		]
-		+ SSplitter::Slot()
-		.Value(SectionFillPercentage)
-		.OnSlotResized(SSplitter::FOnSlotResized::CreateSP(this, &SSequencerTrackArea::SectionFillPercentageChanged))
-		[
-			SNew(SBox)
-			.Visibility(this, &SSequencerTrackArea::GetSectionControlVisibility)
-			[
-				// Generate a widget for the section area portion of the node
-				Node->GenerateWidgetForSectionArea( ViewRange )
-			]
-		]
-	];
+	return TimeSliderController->OnMouseMove( SharedThis(this), MyGeometry, MouseEvent );
 }
 
-EVisibility SSequencerTrackArea::GetCurveEditorVisibility() const
+FReply SSequencerTrackArea::OnMouseWheel( const FGeometry& MyGeometry, const FPointerEvent& MouseEvent )
 {
-	return GetDefault<USequencerSettings>()->GetShowCurveEditor() ? EVisibility::Visible : EVisibility::Hidden;
+	return TimeSliderController->OnMouseWheel( SharedThis(this), MyGeometry, MouseEvent );
 }
 
-EVisibility SSequencerTrackArea::GetSectionControlVisibility() const
+void SSequencerTrackArea::Tick( const FGeometry& AllottedGeometry, const double InCurrentTime, const float InDeltaTime )
 {
-	return GetDefault<USequencerSettings>()->GetShowCurveEditor() ? EVisibility::Hidden : EVisibility::Visible;
-}
+	FVector2D Size = AllottedGeometry.GetLocalSize();
 
-float SSequencerTrackArea::GetScrollBarFillPercentage() const
-{
-	return GetDefault<USequencerSettings>()->GetShowCurveEditor() ? OutlinerFillPercentage.Get() : 1.0f;
-}
-
-float SSequencerTrackArea::GetScrollBarSpacerFillPercentage() const
-{
-	return GetDefault<USequencerSettings>()->GetShowCurveEditor() ? SectionFillPercentage.Get() : 0.0f;
-}
-
-FMargin SSequencerTrackArea::GetOutlinerPadding() const
-{
-	return GetDefault<USequencerSettings>()->GetShowCurveEditor() ? FMargin(0, 0, 10, SequencerLayoutConstants::NodePadding) : FMargin(0, 0, 0, SequencerLayoutConstants::NodePadding);
-}
-
-void SSequencerTrackArea::OutlinerFillPercentageChanged(float Value)
-{
-	OnOutlinerFillPercentageChanged.ExecuteIfBound(Value);
-}
+	if (SizeLastFrame.IsSet() && Size.X != SizeLastFrame->X)
+	{
+		// Zoom by the difference in horizontal size
+		const float Difference = Size.X - SizeLastFrame->X;
+		TRange<float> OldRange = TimeSliderController->GetViewRange().GetAnimationTarget();
+
+		TimeSliderController->SetViewRange(
+			OldRange.GetLowerBoundValue(),
+			OldRange.GetUpperBoundValue() + (Difference * OldRange.Size<float>() / SizeLastFrame->X),
+			EViewRangeInterpolation::Immediate
+			);
+	}
 
-void SSequencerTrackArea::SectionFillPercentageChanged(float Value)
-{
-	OnSectionFillPercentageChanged.ExecuteIfBound(Value);
-}
+	SizeLastFrame = Size;
+}
\ No newline at end of file
diff --git a/Engine/Source/Editor/Sequencer/Private/SSequencerTrackArea.h b/Engine/Source/Editor/Sequencer/Private/SSequencerTrackArea.h
--- a/Engine/Source/Editor/Sequencer/Private/SSequencerTrackArea.h
+++ b/Engine/Source/Editor/Sequencer/Private/SSequencerTrackArea.h
@@ -2,90 +2,32 @@
 
 #pragma once
 
+class SSequencerOutlinerArea;
+class FSequencerTimeSliderController;
+
 /**
  * The area where tracks( rows of sections ) are displayed
  */
-class SSequencerTrackArea : public SCompoundWidget
+class SSequencerTrackArea : public SVerticalBox
 {
 public:
 	DECLARE_DELEGATE_OneParam(FOnValueChanged, float)
 
-	SLATE_BEGIN_ARGS( SSequencerTrackArea )
-	{}
-		/** The range of time being viewed */
-		SLATE_ATTRIBUTE( TRange<float>, ViewRange )
-		/** Percentage of horizontal space that the animation outliner takes up */
-		SLATE_ATTRIBUTE( float, OutlinerFillPercentage )
-		/** Percentage of horizontal space that the section area takes up. */
-		SLATE_ATTRIBUTE( float, SectionFillPercentage )
-		/** Called whenever the outliner fill percentage is changed by a splitter. */
-		SLATE_EVENT( FOnValueChanged, OnOutlinerFillPercentageChanged )
-		/** Called whenever the section fill percentage is changed by a splitter. */
-		SLATE_EVENT( FOnValueChanged, OnSectionFillPercentageChanged )
-	SLATE_END_ARGS()
-
-	void Construct( const FArguments& InArgs, TSharedRef<FSequencer> InSequencer );
-
-	/**
-	 * Updates the track area with new nodes
-	 */
-	virtual void Update( TSharedPtr<FSequencerNodeTree> InSequencerNodeTree );
-
-	TSharedPtr<SSequencerCurveEditor> GetCurveEditor() { return CurveEditor; }
 
-private:
-	/**
-	 * Generates a widget for the specified node and puts it into the scroll box
-	 *
-	 * @param Node	The node to generate a widget for
-	 */
-	void GenerateWidgetForNode( TSharedRef<FSequencerDisplayNode>& Node );
+	void Construct( const FArguments& InArgs, TSharedRef<FSequencerTimeSliderController> InTimeSliderController );
 	
-	/**
-	 * Recursively generates widgets for each layout node
-	 *
-	 * @param Nodes to generate widgets for.  Will recurse into each node's children
-	 */
-	void GenerateLayoutNodeWidgetsRecursive( const TArray< TSharedRef<FSequencerDisplayNode> >& Nodes );
-
-	/** Gets the visibility of the curve editor based on the the current settings. */
-	EVisibility GetCurveEditorVisibility() const;
-	/** Gets the visibility of the default section controls based on the current settings. */
-	EVisibility GetSectionControlVisibility() const;
+	virtual void Tick( const FGeometry& AllottedGeometry, const double InCurrentTime, const float InDeltaTime ) override;
 
-	/** Gets the fill value for the slot which contains the scrollbar.  This allows the scroll bar to move depending
-	  * whether or not the curve editor is visible. */
-	float GetScrollBarFillPercentage() const;
-	/** Gets the fill value for the spacer slot which is to the right of the scrollbar.  This allows the scroll bar 
-	  * to move depending whether or not the curve editor is visible. */
-	float GetScrollBarSpacerFillPercentage() const;
+	virtual FReply OnMouseButtonDown( const FGeometry& MyGeometry, const FPointerEvent& MouseEvent ) override;
+	virtual FReply OnMouseButtonUp( const FGeometry& MyGeometry, const FPointerEvent& MouseEvent ) override;
+	virtual FReply OnMouseMove( const FGeometry& MyGeometry, const FPointerEvent& MouseEvent ) override;
+	virtual FReply OnMouseWheel( const FGeometry& MyGeometry, const FPointerEvent& MouseEvent ) override;
 
-	FMargin GetOutlinerPadding() const;
+private:
 
-	/** Called whenever the outliner fill percentage is changed by a splitter. */
-	void OutlinerFillPercentageChanged(float Value);
-	/** Called whenever the section fill percentage is changed by a splitter. */
-	void SectionFillPercentageChanged(float Value);
+	/** Time slider controller for controlling zoom/pan etc */
+	TSharedPtr<class FSequencerTimeSliderController> TimeSliderController;
 
-private:
-	/** Scrollable area to display widgets */
-	TSharedPtr<SScrollBox> ScrollBox;
-	/** The current view range */
-	TAttribute< TRange<float> > ViewRange;
-	/** The fill percentage of the animation outliner */
-	TAttribute<float> OutlinerFillPercentage;
-	/** The fill percentage of the section area */
-	TAttribute<float> SectionFillPercentage;
-	/** Delegate to be called whenever the outliner fill percentage is changed */
-	FOnValueChanged OnOutlinerFillPercentageChanged;
-	/** Delegate to be called whenever the section fill percentage is changed */
-	FOnValueChanged OnSectionFillPercentageChanged;
-	/** Attribute representing the scrollbar fill percentage */
-	TAttribute<float> ScrollbarFillPercentage;
-	/** Attribute representing the spacer fill percentage which positions the scrollbar correctly. */
-	TAttribute<float> ScrollbarSpacerFillPercentage;
-	/** The main sequencer interface */
-	TWeakPtr<FSequencer> Sequencer;
-	/** The curve editor. */
-	TSharedPtr<SSequencerCurveEditor> CurveEditor;
+	/** Keep a hold of the size of the area so we can maintain zoom levels */
+	TOptional<FVector2D> SizeLastFrame;
 };
diff --git a/Engine/Source/Editor/Sequencer/Private/SSequencerTrackOutliner.cpp b/Engine/Source/Editor/Sequencer/Private/SSequencerTrackOutliner.cpp
new file mode 100644
--- /dev/null
+++ b/Engine/Source/Editor/Sequencer/Private/SSequencerTrackOutliner.cpp
@@ -0,0 +1,9 @@
+// Copyright 1998-2015 Epic Games, Inc. All Rights Reserved.
+
+#include "SequencerPrivatePCH.h"
+#include "SSequencerTrackOutliner.h"
+
+void SSequencerTrackOutliner::Construct( const FArguments& InArgs )
+{
+	SVerticalBox::Construct( InArgs );
+}
\ No newline at end of file
diff --git a/Engine/Source/Editor/Sequencer/Private/SSequencerTrackOutliner.h b/Engine/Source/Editor/Sequencer/Private/SSequencerTrackOutliner.h
new file mode 100644
--- /dev/null
+++ b/Engine/Source/Editor/Sequencer/Private/SSequencerTrackOutliner.h
@@ -0,0 +1,13 @@
+// Copyright 1998-2015 Epic Games, Inc. All Rights Reserved.
+
+#pragma once
+
+/**
+ * The area where outliner nodes for each track is stored
+ */
+class SSequencerTrackOutliner : public SVerticalBox
+{
+public:
+
+	void Construct( const FArguments& InArgs );
+};
diff --git a/Engine/Source/Editor/Sequencer/Private/Sequencer.cpp b/Engine/Source/Editor/Sequencer/Private/Sequencer.cpp
--- a/Engine/Source/Editor/Sequencer/Private/Sequencer.cpp
+++ b/Engine/Source/Editor/Sequencer/Private/Sequencer.cpp
@@ -56,7 +56,7 @@ void FSequencer::InitSequencer( const FSequencerInitParams& InitParams, const TA
 
 		ToolkitHost = InitParams.ToolkitHost;
 
-		LastViewRange = InitParams.ViewParams.InitalViewRange;
+		LastViewRange = TargetViewRange = InitParams.ViewParams.InitalViewRange;
 		ScrubPosition = InitParams.ViewParams.InitialScrubPosition;
 
 		ObjectChangeListener = InitParams.ObjectChangeListener;
@@ -74,10 +74,10 @@ void FSequencer::InitSequencer( const FSequencerInitParams& InitParams, const TA
 
 		// Make internal widgets
 		SequencerWidget = SNew( SSequencer, SharedThis( this ) )
-			.ViewRange( this, &FSequencer::OnGetViewRange )
+			.ViewRange( this, &FSequencer::GetViewRange )
 			.ScrubPosition( this, &FSequencer::OnGetScrubPosition )
 			.OnScrubPositionChanged( this, &FSequencer::OnScrubPositionChanged )
-			.OnViewRangeChanged( this, &FSequencer::OnViewRangeChanged, false )
+			.OnViewRangeChanged( this, &FSequencer::OnViewRangeChanged )
 			.OnGetAddMenuContent(InitParams.ViewParams.OnGetAddMenuContent);
 
 		// When undo occurs, get a notification so we can make sure our view is up to date
@@ -113,9 +113,9 @@ void FSequencer::InitSequencer( const FSequencerInitParams& InitParams, const TA
 
 
 		ZoomAnimation = FCurveSequence();
-		ZoomCurve = ZoomAnimation.AddCurve(0.f, 0.35f, ECurveEaseFunction::QuadIn);
+		ZoomCurve = ZoomAnimation.AddCurve(0.f, 0.2f, ECurveEaseFunction::QuadIn);
 		OverlayAnimation = FCurveSequence();
-		OverlayCurve = OverlayAnimation.AddCurve(0.f, 0.35f, ECurveEaseFunction::QuadIn);
+		OverlayCurve = OverlayAnimation.AddCurve(0.f, 0.2f, ECurveEaseFunction::QuadIn);
 
 		// Update initial movie scene data
 		NotifyMovieSceneDataChanged();
@@ -485,10 +485,17 @@ void FSequencer::NotifyMovieSceneDataChanged()
 }
 
 
-TRange<float> FSequencer::OnGetViewRange() const
+FAnimatedRange FSequencer::GetViewRange() const
 {
-	return TRange<float>(FMath::Lerp(LastViewRange.GetLowerBoundValue(), TargetViewRange.GetLowerBoundValue(), ZoomCurve.GetLerp()),
+	FAnimatedRange AnimatedRange(FMath::Lerp(LastViewRange.GetLowerBoundValue(), TargetViewRange.GetLowerBoundValue(), ZoomCurve.GetLerp()),
 		FMath::Lerp(LastViewRange.GetUpperBoundValue(), TargetViewRange.GetUpperBoundValue(), ZoomCurve.GetLerp()));
+
+	if (ZoomAnimation.IsPlaying())
+	{
+		AnimatedRange.AnimationTarget = TargetViewRange;
+	}
+
+	return AnimatedRange;
 }
 
 bool FSequencer::IsAutoKeyEnabled() const 
@@ -856,18 +863,23 @@ TRange<float> FSequencer::GetFilteringShotsTimeBounds() const
 	return TRange<float>::Empty();
 }
 
-
-
-void FSequencer::OnViewRangeChanged( TRange<float> NewViewRange, bool bSmoothZoom )
+void FSequencer::OnViewRangeChanged( TRange<float> NewViewRange, EViewRangeInterpolation Interpolation )
 {
+	const float AnimationLengthSeconds = Interpolation == EViewRangeInterpolation::Immediate ? 0.f : 0.1f;
+
 	if (!NewViewRange.IsEmpty())
 	{
-		if (bSmoothZoom)
+		if (AnimationLengthSeconds != 0.f)
 		{
+			if (ZoomAnimation.GetCurve(0).DurationSeconds != AnimationLengthSeconds)
+			{
+				ZoomAnimation = FCurveSequence();
+				ZoomCurve = ZoomAnimation.AddCurve(0.f, AnimationLengthSeconds, ECurveEaseFunction::QuadIn);
+			}
+
 			if (!ZoomAnimation.IsPlaying())
 			{
 				LastViewRange = TargetViewRange;
-
 				ZoomAnimation.Play( SequencerWidget.ToSharedRef() );
 			}
 			TargetViewRange = NewViewRange;
@@ -1181,7 +1193,7 @@ void FSequencer::ZoomToSelectedSections()
 
 	if (!BoundsHull.IsEmpty() && !BoundsHull.IsDegenerate())
 	{
-		OnViewRangeChanged(BoundsHull, true);
+		OnViewRangeChanged(BoundsHull);
 	}
 }
 
@@ -1238,7 +1250,7 @@ void FSequencer::FilterToShotSections(const TArray< TWeakObjectPtr<class UMovieS
 	if( bZoomToShotBounds )
 	{
 		// zoom in
-		OnViewRangeChanged(GetTimeBounds(), true);
+		OnViewRangeChanged(GetTimeBounds());
 	}
 
 	SequencerWidget->UpdateBreadcrumbs(ActualShotSections);
diff --git a/Engine/Source/Editor/Sequencer/Private/Sequencer.h b/Engine/Source/Editor/Sequencer/Private/Sequencer.h
--- a/Engine/Source/Editor/Sequencer/Private/Sequencer.h
+++ b/Engine/Source/Editor/Sequencer/Private/Sequencer.h
@@ -93,6 +93,9 @@ public:
 	virtual TSharedRef<ISequencerObjectBindingManager> GetObjectBindingManager() const override;
 	virtual FSequencerSelection* GetSelection() override;
 
+	/** @return The current view range */
+	FAnimatedRange GetViewRange() const;
+
 	bool IsPerspectiveViewportPosessionEnabled() const { return bPerspectiveViewportPossessionEnabled; }
 
 	/**
@@ -274,9 +277,6 @@ protected:
 	 */
 	TRange<float> GetFilteringShotsTimeBounds() const;
 
-	/** @return The current view range */
-	virtual TRange<float> OnGetViewRange() const;
-
 	/** @return The current scrub position */
 	virtual float OnGetScrubPosition() const { return ScrubPosition; }
 
@@ -288,7 +288,7 @@ protected:
 	 *
 	 * @param	NewViewRange The new view range
 	 */
-	void OnViewRangeChanged( TRange<float> NewViewRange, bool bSmoothZoom );
+	void OnViewRangeChanged( TRange<float> NewViewRange, EViewRangeInterpolation Interpolation = EViewRangeInterpolation::Animated );
 
 	/**
 	 * Called when the scrub position is changed by the user
diff --git a/Engine/Source/Editor/Sequencer/Private/SequencerTrackLaneFactory.cpp b/Engine/Source/Editor/Sequencer/Private/SequencerTrackLaneFactory.cpp
new file mode 100644
--- /dev/null
+++ b/Engine/Source/Editor/Sequencer/Private/SequencerTrackLaneFactory.cpp
@@ -0,0 +1,142 @@
+// Copyright 1998-2015 Epic Games, Inc. All Rights Reserved.
+
+#include "SequencerPrivatePCH.h"
+
+#include "SequencerTrackLaneFactory.h"
+#include "SSequencerTrackArea.h"
+#include "SSequencerTrackOutliner.h"
+#include "SequencerSettings.h"
+
+struct FWidgetSizeLinker
+{
+	FVector2D GetMinSize() const
+	{
+		TOptional<FVector2D> MinSize;
+
+		for (auto Widget : Widgets)
+		{
+			auto Pinned = Widget.Pin();
+			if (Pinned.IsValid())
+			{
+				FVector2D DesiredSize = Pinned->GetDesiredSize();
+				if (MinSize.IsSet())
+				{
+					MinSize->X = FMath::Min(MinSize->X, DesiredSize.X);
+				}
+				else
+				{
+					MinSize = DesiredSize;
+				}
+			}
+		}
+
+		return MinSize.IsSet() ? MinSize.GetValue() : FVector2D::ZeroVector;
+	}
+
+	FVector2D GetMaxSize() const
+	{
+		TOptional<FVector2D> MaxSize;
+
+		for (auto Widget : Widgets)
+		{
+			auto Pinned = Widget.Pin();
+			if (Pinned.IsValid())
+			{
+				FVector2D DesiredSize = Pinned->GetDesiredSize();
+				if (MaxSize.IsSet())
+				{
+					MaxSize->X = FMath::Max(MaxSize->X, DesiredSize.X);
+				}
+				else
+				{
+					MaxSize = DesiredSize;
+				}
+			}
+		}
+
+		return MaxSize.IsSet() ? MaxSize.GetValue() : FVector2D::ZeroVector;
+	}
+
+	TArray<TWeakPtr<SWidget>> Widgets;
+};
+
+FSequencerTrackLaneFactory::FSequencerTrackLaneFactory(
+	TSharedRef<SSequencerTrackOutliner>	InTrackOutliner,
+	TSharedRef<SSequencerTrackArea>		InTrackArea,
+	TSharedRef<FSequencer>				InSequencer
+	)
+	: TrackOutliner( MoveTemp( InTrackOutliner ) )
+	, TrackArea( MoveTemp( InTrackArea ) )
+	, Sequencer( MoveTemp( InSequencer ) )
+{
+
+}
+
+void FSequencerTrackLaneFactory::Repopulate(const FSequencerNodeTree& NodeTree)
+{
+	TrackArea->ClearChildren();
+	TrackOutliner->ClearChildren();
+
+	for ( TSharedRef<FSequencerDisplayNode> RootNode : NodeTree.GetRootNodes() )
+	{
+		GenerateWidgetForNode( *RootNode );
+		GenerateLayoutNodeWidgetsRecursive( RootNode->GetChildNodes() );
+	}
+}
+
+void FSequencerTrackLaneFactory::GenerateWidgetForNode( FSequencerDisplayNode& Node)
+{
+	TAttribute<FAnimatedRange> ViewRange;
+	ViewRange.Bind(TAttribute<FAnimatedRange>::FGetter::CreateSP(Sequencer, &FSequencer::GetViewRange));
+
+	TSharedRef<SWidget> OutlinerWidget = Node.GenerateContainerWidgetForOutliner();
+	TSharedRef<SWidget> TrackWidget = Node.GenerateWidgetForSectionArea( FAnimatedRange::WrapAttribute(ViewRange) );
+
+	// Link the height of the widets using a size linker
+	FWidgetSizeLinker Link;
+	Link.Widgets.Add(OutlinerWidget);
+	Link.Widgets.Add(TrackWidget);
+
+	auto GetMaxSize = [=]() -> float {
+		return Link.GetMaxSize().Y;
+	};
+	auto IsVisible = []() -> EVisibility{
+		return GetDefault<USequencerSettings>()->GetShowCurveEditor() ? EVisibility::Hidden : EVisibility::Visible;
+	};
+
+	TrackOutliner->AddSlot()
+	.AutoHeight()
+	[
+		SNew(SBox)
+		.HeightOverride_Lambda(GetMaxSize)
+		[
+			OutlinerWidget
+		]
+	];
+
+	TrackArea->AddSlot()
+	.AutoHeight()
+	[
+		SNew(SBox)
+		.HeightOverride_Lambda(GetMaxSize)
+		.Visibility_Lambda(IsVisible)
+		[
+			TrackWidget
+		]
+	];
+}
+
+void FSequencerTrackLaneFactory::GenerateLayoutNodeWidgetsRecursive( const TArray< TSharedRef<FSequencerDisplayNode> >& Nodes )
+{
+	for (auto& SharedNode : Nodes)
+	{
+		FSequencerDisplayNode& Node = *SharedNode;
+
+		// The section area encompasses multiple logical child nodes so we do not generate children that arent section areas (they would not be top level nodes in that case)
+		if( Node.GetType() == ESequencerNode::Track )
+		{
+			GenerateWidgetForNode( Node );
+			GenerateLayoutNodeWidgetsRecursive( Node.GetChildNodes() );
+		}
+	}
+}
\ No newline at end of file
diff --git a/Engine/Source/Editor/Sequencer/Private/SequencerTrackLaneFactory.h b/Engine/Source/Editor/Sequencer/Private/SequencerTrackLaneFactory.h
new file mode 100644
--- /dev/null
+++ b/Engine/Source/Editor/Sequencer/Private/SequencerTrackLaneFactory.h
@@ -0,0 +1,47 @@
+// Copyright 1998-2015 Epic Games, Inc. All Rights Reserved.
+
+#pragma once
+
+class FSequencer;
+class SSequencerTrackOutliner;
+class SSequencerTrackArea;
+
+/** Class responsible for creating widgets for the sequencer track outliner and track area */
+class FSequencerTrackLaneFactory
+{
+public:
+
+	/** Construct this facory */
+	FSequencerTrackLaneFactory(
+		TSharedRef<SSequencerTrackOutliner>	InTrackOutliner,
+		TSharedRef<SSequencerTrackArea>		InTrackArea,
+		TSharedRef<FSequencer>				InSequencer
+		);
+
+	/** Repopulate the supplied sequencer widgets with the specified node tree */
+	void Repopulate(const FSequencerNodeTree& NodeTree);
+
+private:
+
+	/**
+	 * Generates a widget for the specified node and puts it into the scroll box
+	 *
+	 * @param Node	The node to generate a widget for
+	 */
+	void GenerateWidgetForNode( FSequencerDisplayNode& Node);
+	
+	/**
+	 * Recursively generates widgets for each layout node
+	 *
+	 * @param Nodes to generate widgets for.  Will recurse into each node's children
+	 */
+	void GenerateLayoutNodeWidgetsRecursive( const TArray< TSharedRef<FSequencerDisplayNode> >& Nodes );
+
+private:
+	/** Pointer to the sequencer track outliner */
+	TSharedRef<SSequencerTrackOutliner>	TrackOutliner;
+	/** Pointer to the sequencer track area */
+	TSharedRef<SSequencerTrackArea>		TrackArea;
+	/** Pointer to the sequencer itself */
+	TSharedRef<FSequencer>				Sequencer;
+};
\ No newline at end of file
diff --git a/Engine/Source/Editor/Sequencer/Private/TimeSliderController.cpp b/Engine/Source/Editor/Sequencer/Private/TimeSliderController.cpp
--- a/Engine/Source/Editor/Sequencer/Private/TimeSliderController.cpp
+++ b/Engine/Source/Editor/Sequencer/Private/TimeSliderController.cpp
@@ -12,6 +12,9 @@ namespace ScrubConstants
 
 	/**The smallest number of units between between major tick marks */
 	const float MinDisplayTickSpacing = 0.001f;
+
+	/**The fraction of the current view range to scroll per unit delta  */
+	const float ScrollPanFraction = 0.1f;
 }
 
 
@@ -382,27 +385,7 @@ FReply FSequencerTimeSliderController::OnMouseMove( TSharedRef<SWidget> WidgetOw
 				float NewViewOutputMin = LocalViewRangeMin - InputDelta.X;
 				float NewViewOutputMax = LocalViewRangeMax - InputDelta.X;
 
-				TOptional<float> LocalClampMin = TimeSliderArgs.ClampMin.Get();
-				TOptional<float> LocalClampMax = TimeSliderArgs.ClampMax.Get();
-
-				// Clamp the range if clamp values are set
-				if ( LocalClampMin.IsSet() && NewViewOutputMin < LocalClampMin.GetValue() )
-				{
-					NewViewOutputMin = LocalClampMin.GetValue();
-				}
-			
-				if ( LocalClampMax.IsSet() && NewViewOutputMax > LocalClampMax.GetValue() )
-				{
-					NewViewOutputMax = LocalClampMax.GetValue();
-				}
-
-				TimeSliderArgs.OnViewRangeChanged.ExecuteIfBound(TRange<float>(NewViewOutputMin, NewViewOutputMax));
-
-				if( !TimeSliderArgs.ViewRange.IsBound() )
-				{	
-					// The  output is not bound to a delegate so we'll manage the value ourselves
-					TimeSliderArgs.ViewRange.Set( TRange<float>( NewViewOutputMin, NewViewOutputMax ) );
-				}
+				SetViewRange(NewViewOutputMin, NewViewOutputMax, EViewRangeInterpolation::Immediate);
 			}
 		}
 		else if (MouseEvent.IsMouseButtonDown( EKeys::LeftMouseButton ))
@@ -450,49 +433,23 @@ void FSequencerTimeSliderController::CommitScrubPosition( float NewValue, bool b
 
 FReply FSequencerTimeSliderController::OnMouseWheel( TSharedRef<SWidget> WidgetOwner, const FGeometry& MyGeometry, const FPointerEvent& MouseEvent )
 {
-	if ( TimeSliderArgs.AllowZoom )
-	{
-		const float ZoomDelta = -0.1f * MouseEvent.GetWheelDelta();
+	TOptional<TRange<float>> NewTargetRange;
 
+	float MouseFractionX = MyGeometry.AbsoluteToLocal(MouseEvent.GetScreenSpacePosition()).X / MyGeometry.GetLocalSize().X;
+	if ( TimeSliderArgs.AllowZoom && (MouseEvent.IsLeftControlDown() || MouseEvent.IsRightControlDown()) )
+	{
+		const float ZoomDelta = -0.2f * MouseEvent.GetWheelDelta();
+		if (ZoomByDelta(ZoomDelta, MouseFractionX))
 		{
-			TRange<float> LocalViewRange = TimeSliderArgs.ViewRange.Get();
-			float LocalViewRangeMax = LocalViewRange.GetUpperBoundValue();
-			float LocalViewRangeMin = LocalViewRange.GetLowerBoundValue();
-			const float OutputViewSize = LocalViewRangeMax - LocalViewRangeMin;
-			const float OutputChange = OutputViewSize * ZoomDelta;
-
-			float NewViewOutputMin = LocalViewRangeMin - (OutputChange * 0.5f);
-			float NewViewOutputMax = LocalViewRangeMax + (OutputChange * 0.5f);
-
-			if( FMath::Abs( OutputChange ) > 0.01f && NewViewOutputMin < NewViewOutputMax )
-			{
-				TOptional<float> LocalClampMin = TimeSliderArgs.ClampMin.Get();
-				TOptional<float> LocalClampMax = TimeSliderArgs.ClampMax.Get();
-
-				// Clamp the range if clamp values are set
-				if ( LocalClampMin.IsSet() && NewViewOutputMin < LocalClampMin.GetValue() )
-				{
-					NewViewOutputMin = LocalClampMin.GetValue();
-				}
-				
-				if ( LocalClampMax.IsSet() && NewViewOutputMax > LocalClampMax.GetValue() )
-				{
-					NewViewOutputMax = LocalClampMax.GetValue();
-				}
-
-				TimeSliderArgs.OnViewRangeChanged.ExecuteIfBound(TRange<float>(NewViewOutputMin, NewViewOutputMax));
-
-				if( !TimeSliderArgs.ViewRange.IsBound() )
-				{	
-					// The  output is not bound to a delegate so we'll manage the value ourselves
-					TimeSliderArgs.ViewRange.Set( TRange<float>( NewViewOutputMin, NewViewOutputMax ) );
-				}
-			}
+			return FReply::Handled();
 		}
-
+	}
+	else if (MouseEvent.IsLeftShiftDown() || MouseEvent.IsRightShiftDown())
+	{
+		PanByDelta(-MouseEvent.GetWheelDelta());
 		return FReply::Handled();
 	}
-
+	
 	return FReply::Unhandled();
 }
 
@@ -551,4 +508,65 @@ int32 FSequencerTimeSliderController::OnPaintSectionView( const FGeometry& Allot
 	return LayerId;
 }
 
+void FSequencerTimeSliderController::SetViewRange( float NewRangeMin, float NewRangeMax, EViewRangeInterpolation Interpolation )
+{
+	TOptional<float> LocalClampMin = TimeSliderArgs.ClampMin.Get();
+	TOptional<float> LocalClampMax = TimeSliderArgs.ClampMax.Get();
+
+	// Clamp the range if clamp values are set
+	if ( LocalClampMin.IsSet() && NewRangeMin < LocalClampMin.GetValue() )
+	{
+		NewRangeMin = LocalClampMin.GetValue();
+	}
+	
+	if ( LocalClampMax.IsSet() && NewRangeMax > LocalClampMax.GetValue() )
+	{
+		NewRangeMax = LocalClampMax.GetValue();
+	}
+
+	const TRange<float> NewRange(NewRangeMin, NewRangeMax);
+
+	TimeSliderArgs.OnViewRangeChanged.ExecuteIfBound( NewRange, Interpolation );
+
+	if( !TimeSliderArgs.ViewRange.IsBound() )
+	{	
+		// The  output is not bound to a delegate so we'll manage the value ourselves (no animation)
+		TimeSliderArgs.ViewRange.Set( NewRange );
+	}
+}
+
+bool FSequencerTimeSliderController::ZoomByDelta( float InDelta, float MousePositionFraction )
+{
+	TRange<float> LocalViewRange = TimeSliderArgs.ViewRange.Get().GetAnimationTarget();
+	float LocalViewRangeMax = LocalViewRange.GetUpperBoundValue();
+	float LocalViewRangeMin = LocalViewRange.GetLowerBoundValue();
+	const float OutputViewSize = LocalViewRangeMax - LocalViewRangeMin;
+	const float OutputChange = OutputViewSize * InDelta;
+
+	float NewViewOutputMin = LocalViewRangeMin - (OutputChange * MousePositionFraction);
+	float NewViewOutputMax = LocalViewRangeMax + (OutputChange * (1.f - MousePositionFraction));
+
+	if( FMath::Abs( OutputChange ) > 0.01f && NewViewOutputMin < NewViewOutputMax )
+	{
+		SetViewRange(NewViewOutputMin, NewViewOutputMax, EViewRangeInterpolation::Animated);
+		return true;
+	}
+
+	return false;
+}
+
+void FSequencerTimeSliderController::PanByDelta( float InDelta )
+{
+	TRange<float> LocalViewRange = TimeSliderArgs.ViewRange.Get().GetAnimationTarget();
+
+	float CurrentMin = LocalViewRange.GetLowerBoundValue();
+	float CurrentMax = LocalViewRange.GetUpperBoundValue();
+
+	// Adjust the delta to be a percentage of the current range
+	InDelta *= ScrubConstants::ScrollPanFraction * (CurrentMax - CurrentMin);
+
+	SetViewRange(CurrentMin + InDelta, CurrentMax + InDelta, EViewRangeInterpolation::Animated);
+}
+
+
 #undef LOCTEXT_NAMESPACE
diff --git a/Engine/Source/Editor/Sequencer/Private/TimeSliderController.h b/Engine/Source/Editor/Sequencer/Private/TimeSliderController.h
--- a/Engine/Source/Editor/Sequencer/Private/TimeSliderController.h
+++ b/Engine/Source/Editor/Sequencer/Private/TimeSliderController.h
@@ -19,11 +19,40 @@ public:
 	virtual FReply OnMouseButtonUp( TSharedRef<SWidget> WidgetOwner, const FGeometry& MyGeometry, const FPointerEvent& MouseEvent ) override;
 	virtual FReply OnMouseMove( TSharedRef<SWidget> WidgetOwner, const FGeometry& MyGeometry, const FPointerEvent& MouseEvent ) override;
 	virtual FReply OnMouseWheel( TSharedRef<SWidget> WidgetOwner, const FGeometry& MyGeometry, const FPointerEvent& MouseEvent ) override;
+	/** End ITimeSliderController Interface */
+
+	/** Get the current view range for this controller */
+	FAnimatedRange GetViewRange() const { return TimeSliderArgs.ViewRange.Get(); }
+
+	/**
+	 * Set a new range based on a min, max and an interpolation mode
+	 * 
+	 * @param NewRangeMin		The new lower bound of the range
+	 * @param NewRangeMax		The new upper bound of the range
+	 * @param Interpolation		How to set the new range (either immediately, or animated)
+	 */
+	void SetViewRange( float NewRangeMin, float NewRangeMax, EViewRangeInterpolation Interpolation );
+
+	/**
+	 * Zoom the range by a given delta.
+	 * 
+	 * @param InDelta		The total amount to zoom by (+ve = zoom out, -ve = zoom in)
+	 * @param ZoomBias		Bias to apply to lower/upper extents of the range. (0 = lower, 0.5 = equal, 1 = upper)
+	 */
+	bool ZoomByDelta( float InDelta, float ZoomBias = 0.5f );
+
+	/**
+	 * Pan the range by a given delta
+	 * 
+	 * @param InDelta		The total amount to pan by (+ve = pan forwards in time, -ve = pan backwards in time)
+	 */
+	void PanByDelta( float InDelta );
 
 	/**
 	 * Draws major tick lines in the section view                                                              
 	 */
 	int32 OnPaintSectionView( const FGeometry& AllottedGeometry, const FSlateRect& MyClippingRect, FSlateWindowElementList& OutDrawElements, int32 LayerId, bool bEnabled, bool bDisplayTickLines, bool bDisplayScrubPosition ) const;
+
 private:
 	/**
 	 * Call this method when the user's interaction has changed the scrub position
@@ -41,6 +70,7 @@ private:
 	 * @param InArgs			Parameters for drawing the tick lines
 	 */
 	void DrawTicks( FSlateWindowElementList& OutDrawElements, const struct FScrubRangeToScreen& RangeToScreen, struct FDrawTickArgs& InArgs ) const;
+
 private:
 	FTimeSliderArgs TimeSliderArgs;
 	/** Brush for drawing an upwards facing scrub handle */
diff --git a/Engine/Source/Editor/SequencerWidgets/Public/ITimeSlider.h b/Engine/Source/Editor/SequencerWidgets/Public/ITimeSlider.h
--- a/Engine/Source/Editor/SequencerWidgets/Public/ITimeSlider.h
+++ b/Engine/Source/Editor/SequencerWidgets/Public/ITimeSlider.h
@@ -2,14 +2,57 @@
 
 #pragma once
 
+/** Enum specifying how to interpolate to a new view range */
+enum class EViewRangeInterpolation
+{
+	/** Use an externally defined animated interpolation */
+	Animated,
+	/** Set the view range immediately */
+	Immediate,
+};
+
 DECLARE_DELEGATE_TwoParams( FOnScrubPositionChanged, float, bool )
-DECLARE_DELEGATE_OneParam( FOnViewRangeChanged, TRange<float> )
+DECLARE_DELEGATE_TwoParams( FOnViewRangeChanged, TRange<float>, EViewRangeInterpolation )
+
+/** Structure used to wrap up a range, and an optional animation target */
+struct FAnimatedRange : public TRange<float>
+{
+	/** Default Construction */
+	FAnimatedRange() : TRange() {}
+	/** Construction from a lower and upper bound */
+	FAnimatedRange( float LowerBound, float UpperBound ) : TRange( LowerBound, UpperBound ) {}
+	/** Copy-construction from simple range */
+	FAnimatedRange( const TRange<float>& InRange ) : TRange(InRange) {}
+
+	/** Helper function to wrap an attribute to an animated range with a non-animated one */
+	static TAttribute<TRange<float>> WrapAttribute( const TAttribute<FAnimatedRange>& InAttribute )
+	{
+		typedef TAttribute<TRange<float>> Attr;
+		return Attr::Create(Attr::FGetter::CreateLambda([=](){ return InAttribute.Get(); }));
+	}
+
+	/** Helper function to wrap an attribute to a non-animated range with an animated one */
+	static TAttribute<FAnimatedRange> WrapAttribute( const TAttribute<TRange<float>>& InAttribute )
+	{
+		typedef TAttribute<FAnimatedRange> Attr;
+		return Attr::Create(Attr::FGetter::CreateLambda([=](){ return InAttribute.Get(); }));
+	}
+
+	/** Get the current animation target, or the whole view range when not animating */
+	const TRange<float>& GetAnimationTarget() const
+	{
+		return AnimationTarget.IsSet() ? AnimationTarget.GetValue() : *this;
+	}
+	
+	/** The animation target, if animating */
+	TOptional<TRange<float>> AnimationTarget;
+};
 
 struct FTimeSliderArgs
 {
 	FTimeSliderArgs()
 		: ScrubPosition(0)
-		, ViewRange( TRange<float>(0.0f, 5.0f) )
+		, ViewRange( FAnimatedRange(0.0f, 5.0f) )
 		, ClampMin(-70000.0f)
 		, ClampMax(70000.0f)
 		, OnScrubPositionChanged()
@@ -22,7 +65,7 @@ struct FTimeSliderArgs
 	/** The scrub position */
 	TAttribute<float> ScrubPosition;
 	/** View time range */
-	TAttribute< TRange<float> > ViewRange;
+	TAttribute< FAnimatedRange > ViewRange;
 	/** Optional min output to clamp to */
 	TAttribute<TOptional<float> > ClampMin;
 	/** Optional max output to clamp to */
diff --git a/Engine/Source/Editor/UnrealEd/Public/SCurveEditor.h b/Engine/Source/Editor/UnrealEd/Public/SCurveEditor.h
--- a/Engine/Source/Editor/UnrealEd/Public/SCurveEditor.h
+++ b/Engine/Source/Editor/UnrealEd/Public/SCurveEditor.h
@@ -386,6 +386,7 @@ private:
 	/** Paints the marquee for selection */
 	void PaintMarquee(const FGeometry& AllottedGeometry, const FSlateRect& MyClippingRect, FSlateWindowElementList& OutDrawElements, int32 LayerId) const;
 
+protected:
 	// SWidget interface
 	UNREALED_API virtual FReply OnMouseButtonDown( const FGeometry& MyGeometry, const FPointerEvent& MouseEvent ) override;
 	UNREALED_API virtual FReply OnMouseButtonUp( const FGeometry& MyGeometry, const FPointerEvent& MouseEvent ) override;
@@ -394,6 +395,7 @@ private:
 	UNREALED_API virtual FReply OnKeyDown( const FGeometry& MyGeometry, const FKeyEvent& InKeyEvent ) override;
 	UNREALED_API virtual void   OnMouseCaptureLost() override;
 
+private:
 	/** Attempts to start a drag operation when the mouse moves. */
 	void TryStartDrag(const FGeometry& MyGeometry, const FPointerEvent& MouseEvent);
 
