Fixed padding offset in sequencer outliner area

diff --git a/Engine/Source/Editor/Sequencer/Private/SSequencer.cpp b/Engine/Source/Editor/Sequencer/Private/SSequencer.cpp
--- a/Engine/Source/Editor/Sequencer/Private/SSequencer.cpp
+++ b/Engine/Source/Editor/Sequencer/Private/SSequencer.cpp
@@ -260,10 +260,13 @@ void SSequencer::Construct( const FArguments& InArgs, TSharedRef< class FSequenc
 
 									+ SHorizontalBox::Slot()
 									.FillWidth( FillCoefficient_0 )
-									// Padding to allow space for the scroll bar
-									.Padding(FMargin(0,0,10.f,0))
 									[
-										TrackOutliner.ToSharedRef()
+										SNew(SBox)
+										// Padding to allow space for the scroll bar
+										.Padding(FMargin(0,0,10.f,0))
+										[
+											TrackOutliner.ToSharedRef()
+										]
 									]
 
 									+ SHorizontalBox::Slot()
