Sequencer: Fix for spawnables getting multiple possessable tracks.

diff --git a/Engine/Source/Editor/Sequencer/Private/Sequencer.cpp b/Engine/Source/Editor/Sequencer/Private/Sequencer.cpp
--- a/Engine/Source/Editor/Sequencer/Private/Sequencer.cpp
+++ b/Engine/Source/Editor/Sequencer/Private/Sequencer.cpp
@@ -543,7 +543,8 @@ FGuid FSequencer::GetHandleToObject( UObject* Object )
 
 	FGuid ObjectGuid = ObjectBindingManager->FindGuidForObject( *FocusedMovieScene, *Object );
 
-	if (ObjectGuid.IsValid())
+	// Check here for spawnable otherwise spawnables get recreated as possessables, which doesn't make sense
+	if (ObjectGuid.IsValid() && !FocusedMovieScene->FindSpawnable(ObjectGuid))
 	{
 		// Make sure that the possessable is still valid, if it's not remove the binding so new one 
 		// can be created.  This can happen due to undo.
