Sequencer: Hotkeys for transport control:

Spacebar - Toggle play
Down - Play
Up - Reset to beginning of the timeline
Left - Step to previous frame
Right - Step to next frame

These are pretty standard hotkeys, let me know if there's any contention for these.

diff --git a/Engine/Source/Editor/Sequencer/Private/SSequencer.cpp b/Engine/Source/Editor/Sequencer/Private/SSequencer.cpp
--- a/Engine/Source/Editor/Sequencer/Private/SSequencer.cpp
+++ b/Engine/Source/Editor/Sequencer/Private/SSequencer.cpp
@@ -575,7 +575,7 @@ TSharedRef<SWidget> SSequencer::MakeTransportControls()
 	FTransportControlArgs TransportControlArgs;
 	TransportControlArgs.OnBackwardEnd.BindSP( SequencerPinned, &FSequencer::OnStepToBeginning );
 	TransportControlArgs.OnBackwardStep.BindSP( SequencerPinned, &FSequencer::OnStepBackward );
-	TransportControlArgs.OnForwardPlay.BindSP( SequencerPinned, &FSequencer::OnPlay );
+	TransportControlArgs.OnForwardPlay.BindSP( SequencerPinned, &FSequencer::OnPlay, true );
 	TransportControlArgs.OnForwardStep.BindSP( SequencerPinned, &FSequencer::OnStepForward );
 	TransportControlArgs.OnForwardEnd.BindSP( SequencerPinned, &FSequencer::OnStepToEnd );
 	TransportControlArgs.OnToggleLooping.BindSP( SequencerPinned, &FSequencer::OnToggleLooping );
diff --git a/Engine/Source/Editor/Sequencer/Private/Sequencer.cpp b/Engine/Source/Editor/Sequencer/Private/Sequencer.cpp
--- a/Engine/Source/Editor/Sequencer/Private/Sequencer.cpp
+++ b/Engine/Source/Editor/Sequencer/Private/Sequencer.cpp
@@ -716,10 +716,10 @@ void FSequencer::DestroySpawnablesForAllMovieScenes()
 	}
 }
 
-FReply FSequencer::OnPlay()
+FReply FSequencer::OnPlay(bool bTogglePlay)
 {
-	if( PlaybackState == EMovieScenePlayerStatus::Playing ||
-		PlaybackState == EMovieScenePlayerStatus::Recording )
+	if( (PlaybackState == EMovieScenePlayerStatus::Playing ||
+		 PlaybackState == EMovieScenePlayerStatus::Recording) && bTogglePlay )
 	{
 		PlaybackState = EMovieScenePlayerStatus::Stopped;
 		// Update on stop (cleans up things like sounds that are playing)
@@ -1352,6 +1352,31 @@ void FSequencer::DeleteSelectedItems()
 	}
 }
 
+void FSequencer::TogglePlay()
+{
+	OnPlay();
+}
+
+void FSequencer::PlayForward()
+{
+	OnPlay(false);
+}
+
+void FSequencer::Rewind()
+{
+	OnStepToBeginning();
+}
+
+void FSequencer::StepForward()
+{
+	OnStepForward();
+}
+
+void FSequencer::StepBackward()
+{
+	OnStepBackward();
+}
+
 void FSequencer::SetKey()
 {
 	USelection* CurrentSelection = GEditor->GetSelectedActors();
@@ -1373,7 +1398,6 @@ void FSequencer::SetKey()
 	}
 }
 
-
 void FSequencer::BindSequencerCommands()
 {
 	const FSequencerCommands& Commands = FSequencerCommands::Get();
@@ -1382,6 +1406,26 @@ void FSequencer::BindSequencerCommands()
 		FGenericCommands::Get().Delete,
 		FExecuteAction::CreateSP( this, &FSequencer::DeleteSelectedItems ) );
 
+	SequencerCommandBindings->MapAction(
+		Commands.TogglePlay,
+		FExecuteAction::CreateSP( this, &FSequencer::TogglePlay ) );
+
+	SequencerCommandBindings->MapAction(
+		Commands.PlayForward,
+		FExecuteAction::CreateSP( this, &FSequencer::PlayForward ) );
+
+	SequencerCommandBindings->MapAction(
+		Commands.Rewind,
+		FExecuteAction::CreateSP( this, &FSequencer::Rewind ) );
+
+	SequencerCommandBindings->MapAction(
+		Commands.StepForward,
+		FExecuteAction::CreateSP( this, &FSequencer::StepForward ) );
+
+	SequencerCommandBindings->MapAction(
+		Commands.StepBackward,
+		FExecuteAction::CreateSP( this, &FSequencer::StepBackward ) );
+
 	SequencerCommandBindings->MapAction(
 		Commands.SetKey,
 		FExecuteAction::CreateSP( this, &FSequencer::SetKey ) );
diff --git a/Engine/Source/Editor/Sequencer/Private/Sequencer.h b/Engine/Source/Editor/Sequencer/Private/Sequencer.h
--- a/Engine/Source/Editor/Sequencer/Private/Sequencer.h
+++ b/Engine/Source/Editor/Sequencer/Private/Sequencer.h
@@ -231,7 +231,7 @@ public:
 	void OnActorsDropped( const TArray<TWeakObjectPtr<AActor> >& Actors );
 
 	/** Functions to push on to the transport controls we use */
-	FReply OnPlay();
+	FReply OnPlay(bool bTogglePlay=true);
 	FReply OnRecord();
 	FReply OnStepForward();
 	FReply OnStepBackward();
@@ -314,6 +314,13 @@ protected:
 	/** Called when a user executes the delete command to delete sections or keys */
 	void DeleteSelectedItems();
 	
+	/** Transport controls */
+	void TogglePlay();
+	void PlayForward();
+	void Rewind();
+	void StepForward();
+	void StepBackward();
+
 	/** Manually sets a key for the selected objects at the current time */
 	void SetKey();
 
diff --git a/Engine/Source/Editor/Sequencer/Private/SequencerAssetEditor.cpp b/Engine/Source/Editor/Sequencer/Private/SequencerAssetEditor.cpp
--- a/Engine/Source/Editor/Sequencer/Private/SequencerAssetEditor.cpp
+++ b/Engine/Source/Editor/Sequencer/Private/SequencerAssetEditor.cpp
@@ -173,7 +173,7 @@ void FSequencerAssetEditor::AttachTransportControlsToViewports()
 		TSharedRef<FSequencer> SequencerRef = Sequencer.ToSharedRef();
 
 		FTransportControlArgs TransportControlArgs;
-		TransportControlArgs.OnForwardPlay.BindSP(SequencerRef, &FSequencer::OnPlay);
+		TransportControlArgs.OnForwardPlay.BindSP(SequencerRef, &FSequencer::OnPlay, true);
 		TransportControlArgs.OnRecord.BindSP(SequencerRef, &FSequencer::OnRecord);
 		TransportControlArgs.OnForwardStep.BindSP(SequencerRef, &FSequencer::OnStepForward);
 		TransportControlArgs.OnBackwardStep.BindSP(SequencerRef, &FSequencer::OnStepBackward);
diff --git a/Engine/Source/Editor/Sequencer/Private/SequencerCommands.cpp b/Engine/Source/Editor/Sequencer/Private/SequencerCommands.cpp
--- a/Engine/Source/Editor/Sequencer/Private/SequencerCommands.cpp
+++ b/Engine/Source/Editor/Sequencer/Private/SequencerCommands.cpp
@@ -7,6 +7,12 @@
 
 void FSequencerCommands::RegisterCommands()
 {
+	UI_COMMAND( TogglePlay, "Toggle Play", "Toggle the timeline playing.", EUserInterfaceActionType::Button, FInputChord(EKeys::SpaceBar) );
+	UI_COMMAND( PlayForward, "Play Forward", "Play the timeline forward.", EUserInterfaceActionType::Button, FInputChord(EKeys::Down) );
+	UI_COMMAND( Rewind, "Rewind", "Rewind the timeline.", EUserInterfaceActionType::Button, FInputChord(EKeys::Up) );
+	UI_COMMAND( StepForward, "Step Forward", "Step the timeline forward.", EUserInterfaceActionType::Button, FInputChord(EKeys::Right) );
+	UI_COMMAND( StepBackward, "Step Backward", "Step the timeline backward", EUserInterfaceActionType::Button, FInputChord(EKeys::Left) );
+
 	UI_COMMAND( SetKey, "Set Key", "Sets a key at the current time for the selected actor.", EUserInterfaceActionType::Button, FInputChord(EKeys::K) );
 
 	UI_COMMAND( ToggleAutoKeyEnabled, "Auto Key", "Enables and disables auto keying.", EUserInterfaceActionType::ToggleButton, FInputChord() );
diff --git a/Engine/Source/Editor/Sequencer/Private/SequencerCommands.h b/Engine/Source/Editor/Sequencer/Private/SequencerCommands.h
--- a/Engine/Source/Editor/Sequencer/Private/SequencerCommands.h
+++ b/Engine/Source/Editor/Sequencer/Private/SequencerCommands.h
@@ -15,6 +15,21 @@ public:
 	)
 	{}
 	
+	/** Toggle play */
+	TSharedPtr< FUICommandInfo > TogglePlay;
+
+	/** Play forward */
+	TSharedPtr< FUICommandInfo > PlayForward;
+
+	/** Rewind */
+	TSharedPtr< FUICommandInfo > Rewind;
+
+	/** Step forward */
+	TSharedPtr< FUICommandInfo > StepForward;
+
+	/** Step backward */
+	TSharedPtr< FUICommandInfo > StepBackward;
+
 	/** Sets a key at the current time for the selected actor */
 	TSharedPtr< FUICommandInfo > SetKey;
 
