Sequencer: Restrict curve edits to section boundaries.

Move and dilation functions now take key handles to operate on.

diff --git a/Engine/Source/Editor/Sequencer/Private/SectionDragOperations.cpp b/Engine/Source/Editor/Sequencer/Private/SectionDragOperations.cpp
--- a/Engine/Source/Editor/Sequencer/Private/SectionDragOperations.cpp
+++ b/Engine/Source/Editor/Sequencer/Private/SectionDragOperations.cpp
@@ -147,12 +147,18 @@ void FResizeSection::OnBeginDrag(const FVector2D& LocalMousePos, TSharedPtr<FTra
 	if( Section.IsValid() )
 	{
 		BeginTransaction( *Section, NSLOCTEXT("Sequencer", "DragSectionEdgeTransaction", "Resize section") );
+
+		DraggedKeyHandles.Empty();
+
+		Section->GetKeyHandles(DraggedKeyHandles);
 	}
 }
 
 void FResizeSection::OnEndDrag(TSharedPtr<FTrackNode> SequencerNode)
 {
 	EndTransaction();
+
+	DraggedKeyHandles.Empty();
 }
 
 void FResizeSection::OnDrag( const FPointerEvent& MouseEvent, const FVector2D& LocalMousePos, const FTimeToPixel& TimeToPixelConverter, TSharedPtr<FTrackNode> SequencerNode )
@@ -201,7 +207,7 @@ void FResizeSection::OnDrag( const FPointerEvent& MouseEvent, const FVector2D& L
 			{
 				float NewSize = NewTime - Section->GetStartTime();
 				float DilationFactor = NewSize / Section->GetTimeSize();
-				Section->DilateSection(DilationFactor, Section->GetStartTime());
+				Section->DilateSection(DilationFactor, Section->GetStartTime(), DraggedKeyHandles);
 			}
 			else
 			{
@@ -218,7 +224,7 @@ void FResizeSection::OnDrag( const FPointerEvent& MouseEvent, const FVector2D& L
 			{
 				float NewSize = Section->GetEndTime() - NewTime;
 				float DilationFactor = NewSize / Section->GetTimeSize();
-				Section->DilateSection(DilationFactor, Section->GetEndTime());
+				Section->DilateSection(DilationFactor, Section->GetEndTime(), DraggedKeyHandles);
 			}
 			else
 			{
@@ -242,6 +248,10 @@ void FMoveSection::OnBeginDrag(const FVector2D& LocalMousePos, TSharedPtr<FTrack
 		BeginTransaction( *Section, NSLOCTEXT("Sequencer", "MoveSectionTransaction", "Move Section") );
 		
 		DragOffset = LocalMousePos;
+
+		DraggedKeyHandles.Empty();
+
+		Section->GetKeyHandles(DraggedKeyHandles);
 	}
 }
 
@@ -249,6 +259,8 @@ void FMoveSection::OnEndDrag(TSharedPtr<FTrackNode> SequencerNode)
 {
 	EndTransaction();
 
+	DraggedKeyHandles.Empty();
+
 	if (Section.IsValid())
 	{
 		SequencerNode->FixRowIndices();
@@ -329,7 +341,7 @@ void FMoveSection::OnDrag( const FPointerEvent& MouseEvent, const FVector2D& Loc
 		if (bDeltaX && bDeltaY &&
 			!Section->OverlapsWithSections(MovieSceneSections, TargetRowIndex - Section->GetRowIndex(), DeltaTime))
 		{
-			Section->MoveSection(DeltaTime);
+			Section->MoveSection(DeltaTime, DraggedKeyHandles);
 			Section->SetRowIndex(TargetRowIndex);
 		}
 		else
@@ -344,7 +356,7 @@ void FMoveSection::OnDrag( const FPointerEvent& MouseEvent, const FVector2D& Loc
 			{
 				if (!Section->OverlapsWithSections(MovieSceneSections, 0, DeltaTime))
 				{
-					Section->MoveSection(DeltaTime);
+					Section->MoveSection(DeltaTime, DraggedKeyHandles);
 				}
 				else
 				{
@@ -355,7 +367,7 @@ void FMoveSection::OnDrag( const FPointerEvent& MouseEvent, const FVector2D& Loc
 					float RightMovementMaximum = SectionBoundaries.GetUpperBoundValue() - Section->GetEndTime();
 
 					// Tell the section to move itself and any data it has
-					Section->MoveSection( FMath::Clamp(DeltaTime, LeftMovementMaximum, RightMovementMaximum) );
+					Section->MoveSection( FMath::Clamp(DeltaTime, LeftMovementMaximum, RightMovementMaximum), DraggedKeyHandles );
 				}
 			}
 		}
diff --git a/Engine/Source/Editor/Sequencer/Private/SectionDragOperations.h b/Engine/Source/Editor/Sequencer/Private/SectionDragOperations.h
--- a/Engine/Source/Editor/Sequencer/Private/SectionDragOperations.h
+++ b/Engine/Source/Editor/Sequencer/Private/SectionDragOperations.h
@@ -102,6 +102,8 @@ private:
 	TWeakObjectPtr<UMovieSceneSection> Section;
 	/** true if dragging  the end of the section, false if dragging the start */
 	bool bDraggingByEnd;
+	/** The exact key handles that we're dragging */
+	TSet<FKeyHandle> DraggedKeyHandles;
 };
 
 /**
@@ -127,6 +129,8 @@ private:
 	TWeakObjectPtr<UMovieSceneSection> Section;
 	/** Local mouse position when dragging the section */
 	FVector2D DragOffset;
+	/** The exact key handles that we're dragging */
+	TSet<FKeyHandle> DraggedKeyHandles;
 };
 
 /**
diff --git a/Engine/Source/Runtime/Engine/Classes/Curves/CurveBase.h b/Engine/Source/Runtime/Engine/Classes/Curves/CurveBase.h
--- a/Engine/Source/Runtime/Engine/Classes/Curves/CurveBase.h
+++ b/Engine/Source/Runtime/Engine/Classes/Curves/CurveBase.h
@@ -326,9 +326,11 @@ public:
 
 	/** Shifts all keys forwards or backwards in time by an even amount, preserving order */
 	void ShiftCurve(float DeltaTime);
+	void ShiftCurve(float DeltaTime, TSet<FKeyHandle>& KeyHandles);
 	
 	/** Scales all keys about an origin, preserving order */
 	void ScaleCurve(float ScaleOrigin, float ScaleFactor);
+	void ScaleCurve(float ScaleOrigin, float ScaleFactor, TSet<FKeyHandle>& KeyHandles);
 
 	/** Set the interp mode of the specified key */
 	void SetKeyInterpMode(FKeyHandle KeyHandle, ERichCurveInterpMode NewInterpMode);
@@ -584,9 +586,11 @@ public:
 	
 	/** Shifts all keys forwards or backwards in time by an even amount, preserving order */
 	void ShiftCurve(float DeltaTime);
+	void ShiftCurve(float DeltaTime, TSet<FKeyHandle>& KeyHandles);
 	
 	/** Scales all keys about an origin, preserving order */
 	void ScaleCurve(float ScaleOrigin, float ScaleFactor);
+	void ScaleCurve(float ScaleOrigin, float ScaleFactor, TSet<FKeyHandle>& KeyHandles);
 
 	/** Functions for getting keys based on handles */
 	FIntegralKey& GetKey(FKeyHandle KeyHandle);
diff --git a/Engine/Source/Runtime/Engine/Private/CurveBase.cpp b/Engine/Source/Runtime/Engine/Private/CurveBase.cpp
--- a/Engine/Source/Runtime/Engine/Private/CurveBase.cpp
+++ b/Engine/Source/Runtime/Engine/Private/CurveBase.cpp
@@ -460,19 +460,37 @@ float FRichCurve::GetKeyValue(FKeyHandle KeyHandle) const
 
 void FRichCurve::ShiftCurve(float DeltaTime)
 {
-	// Indices will not change, so we are ok to do this without regenerating the index handle map
-	for( int32 KeyIndex = 0; KeyIndex < Keys.Num(); ++KeyIndex )
+	TSet<FKeyHandle> KeyHandles;
+	ShiftCurve(DeltaTime, KeyHandles);
+}
+
+void FRichCurve::ShiftCurve(float DeltaTime, TSet<FKeyHandle>& KeyHandles)
+{
+	for (auto It = KeyHandlesToIndices.CreateIterator(); It; ++It)
 	{
-		Keys[KeyIndex].Time += DeltaTime;
+		const FKeyHandle& KeyHandle = It.Key();
+		if (KeyHandles.Num() != 0 && KeyHandles.Contains(KeyHandle))
+		{
+			SetKeyTime(KeyHandle, GetKeyTime(KeyHandle)+DeltaTime);
+		}
 	}
 }
 
 void FRichCurve::ScaleCurve(float ScaleOrigin, float ScaleFactor)
 {
-	// Indices will not change, so we are ok to do this without regenerating the index handle map
-	for( int32 KeyIndex = 0; KeyIndex < Keys.Num(); ++KeyIndex )
+	TSet<FKeyHandle> KeyHandles;
+	ScaleCurve(ScaleOrigin, ScaleFactor, KeyHandles);
+}
+
+void FRichCurve::ScaleCurve(float ScaleOrigin, float ScaleFactor, TSet<FKeyHandle>& KeyHandles)
+{
+	for (auto It = KeyHandlesToIndices.CreateIterator(); It; ++It)
 	{
-		Keys[KeyIndex].Time = (Keys[KeyIndex].Time - ScaleOrigin) * ScaleFactor + ScaleOrigin;
+		const FKeyHandle& KeyHandle = It.Key();
+		if (KeyHandles.Num() != 0 && KeyHandles.Contains(KeyHandle))
+		{
+			SetKeyTime(KeyHandle, (GetKeyTime(KeyHandle) - ScaleOrigin) * ScaleFactor + ScaleOrigin);
+		}
 	}
 }
 
@@ -1104,19 +1122,37 @@ float FIntegralCurve::GetKeyTime(FKeyHandle KeyHandle) const
 
 void FIntegralCurve::ShiftCurve(float DeltaTime)
 {
-	// Indices will not change, so we are ok to do this without regenerating the index handle map
-	for( int32 KeyIndex = 0; KeyIndex < Keys.Num(); ++KeyIndex )
+	TSet<FKeyHandle> KeyHandles;
+	ShiftCurve(DeltaTime, KeyHandles);
+}
+
+void FIntegralCurve::ShiftCurve(float DeltaTime, TSet<FKeyHandle>& KeyHandles)
+{
+	for (auto It = KeyHandlesToIndices.CreateIterator(); It; ++It)
 	{
-		Keys[KeyIndex].Time += DeltaTime;
+		const FKeyHandle& KeyHandle = It.Key();
+		if (KeyHandles.Num() != 0 && KeyHandles.Contains(KeyHandle))
+		{
+			SetKeyTime(KeyHandle, GetKeyTime(KeyHandle) + DeltaTime);
+		}
 	}
 }
 
 void FIntegralCurve::ScaleCurve(float ScaleOrigin, float ScaleFactor)
 {
-	// Indices will not change, so we are ok to do this without regenerating the index handle map
-	for( int32 KeyIndex = 0; KeyIndex < Keys.Num(); ++KeyIndex )
+	TSet<FKeyHandle> KeyHandles;
+	ScaleCurve(ScaleOrigin, ScaleFactor, KeyHandles);
+}
+
+void FIntegralCurve::ScaleCurve(float ScaleOrigin, float ScaleFactor, TSet<FKeyHandle>& KeyHandles)
+{
+	for (auto It = KeyHandlesToIndices.CreateIterator(); It; ++It)
 	{
-		Keys[KeyIndex].Time = (Keys[KeyIndex].Time - ScaleOrigin) * ScaleFactor + ScaleOrigin;
+		const FKeyHandle& KeyHandle = It.Key();
+		if (KeyHandles.Num() != 0 && KeyHandles.Contains(KeyHandle))
+		{
+			SetKeyTime(KeyHandle, (GetKeyTime(KeyHandle) - ScaleOrigin) * ScaleFactor + ScaleOrigin);
+		}
 	}
 }
 
diff --git a/Engine/Source/Runtime/MovieSceneCore/Classes/MovieSceneSection.h b/Engine/Source/Runtime/MovieSceneCore/Classes/MovieSceneSection.h
--- a/Engine/Source/Runtime/MovieSceneCore/Classes/MovieSceneSection.h
+++ b/Engine/Source/Runtime/MovieSceneCore/Classes/MovieSceneSection.h
@@ -81,8 +81,9 @@ public:
 	 * Moves the section by a specific amount of time
 	 *
 	 * @param DeltaTime	The distance in time to move the curve
+	 * @param KeyHandles The key handles to operate on
 	 */
-	virtual void MoveSection( float DeltaTime )
+	virtual void MoveSection( float DeltaTime, TSet<FKeyHandle>& KeyHandles )
 	{
 		Modify();
 
@@ -93,10 +94,11 @@ public:
 	/**
 	 * Dilates the section by a specific factor
 	 *
-	 * @param DilationFactor The multipler which scales this section
+	 * @param DilationFactor The multiplier which scales this section
 	 * @param bFromStart Whether to dilate from the beginning or end (whichever stays put)
+	 * @param KeyHandles The key handles to operate on
 	 */
-	virtual void DilateSection( float DilationFactor, float Origin )
+	virtual void DilateSection( float DilationFactor, float Origin, TSet<FKeyHandle>& KeyHandles )
 	{
 		Modify();
 
@@ -104,6 +106,13 @@ public:
 		EndTime = (EndTime - Origin) * DilationFactor + Origin;
 	}
 	
+	/**
+	 * Get the key handles for the keys on the curves within this section
+	 *
+	 * @param KeyHandles The key handles of the keys on the curves within this section
+	 */
+	virtual void GetKeyHandles(TSet<FKeyHandle>& KeyHandles) const {};
+
 	/**
 	 * Gets all snap times for this section
 	 *
diff --git a/Engine/Source/Runtime/MovieSceneCore/Private/MovieSceneSection.cpp b/Engine/Source/Runtime/MovieSceneCore/Private/MovieSceneSection.cpp
--- a/Engine/Source/Runtime/MovieSceneCore/Private/MovieSceneSection.cpp
+++ b/Engine/Source/Runtime/MovieSceneCore/Private/MovieSceneSection.cpp
@@ -53,7 +53,8 @@ void UMovieSceneSection::InitialPlacement(const TArray<UMovieSceneSection*>& Sec
 
 			if (OverlappedSection == NULL) {break;}
 
-			MoveSection(OverlappedSection->GetEndTime() - StartTime);
+			TSet<FKeyHandle> KeyHandles;
+			MoveSection(OverlappedSection->GetEndTime() - StartTime, KeyHandles);
 		}
 	}
 }
diff --git a/Engine/Source/Runtime/MovieSceneCoreTypes/Classes/MovieScene3DTransformSection.h b/Engine/Source/Runtime/MovieSceneCoreTypes/Classes/MovieScene3DTransformSection.h
--- a/Engine/Source/Runtime/MovieSceneCoreTypes/Classes/MovieScene3DTransformSection.h
+++ b/Engine/Source/Runtime/MovieSceneCoreTypes/Classes/MovieScene3DTransformSection.h
@@ -16,8 +16,9 @@ class UMovieScene3DTransformSection : public UMovieSceneSection
 public:
 
 	/** MovieSceneSection interface */
-	virtual void MoveSection( float DeltaPosition ) override;
-	virtual void DilateSection( float DilationFactor, float Origin ) override;
+	virtual void MoveSection( float DeltaPosition, TSet<FKeyHandle>& KeyHandles ) override;
+	virtual void DilateSection( float DilationFactor, float Origin, TSet<FKeyHandle>& KeyHandles ) override;
+	virtual void GetKeyHandles(TSet<FKeyHandle>& KeyHandles) const override;
 
 	/**
 	 * Evaluates the translation component of the transform
diff --git a/Engine/Source/Runtime/MovieSceneCoreTypes/Classes/MovieSceneAnimationSection.h b/Engine/Source/Runtime/MovieSceneCoreTypes/Classes/MovieSceneAnimationSection.h
--- a/Engine/Source/Runtime/MovieSceneCoreTypes/Classes/MovieSceneAnimationSection.h
+++ b/Engine/Source/Runtime/MovieSceneCoreTypes/Classes/MovieSceneAnimationSection.h
@@ -34,8 +34,9 @@ public:
 	float GetAnimationDilationFactor() const {return AnimationDilationFactor;}
 
 	/** MovieSceneSection interface */
-	virtual void MoveSection( float DeltaPosition ) override;
-	virtual void DilateSection( float DilationFactor, float Origin ) override;
+	virtual void MoveSection( float DeltaPosition, TSet<FKeyHandle>& KeyHandles ) override;
+	virtual void DilateSection( float DilationFactor, float Origin, TSet<FKeyHandle>& KeyHandles  ) override;
+	virtual void GetKeyHandles(TSet<FKeyHandle>& KeyHandles) const override;
 	virtual void GetSnapTimes(TArray<float>& OutSnapTimes, bool bGetSectionBorders) const override;
 
 private:
diff --git a/Engine/Source/Runtime/MovieSceneCoreTypes/Classes/MovieSceneAudioSection.h b/Engine/Source/Runtime/MovieSceneCoreTypes/Classes/MovieSceneAudioSection.h
--- a/Engine/Source/Runtime/MovieSceneCoreTypes/Classes/MovieSceneAudioSection.h
+++ b/Engine/Source/Runtime/MovieSceneCoreTypes/Classes/MovieSceneAudioSection.h
@@ -53,8 +53,9 @@ public:
 	}
 
 	/** MovieSceneSection interface */
-	virtual void MoveSection( float DeltaPosition ) override;
-	virtual void DilateSection( float DilationFactor, float Origin ) override;
+	virtual void MoveSection( float DeltaPosition, TSet<FKeyHandle>& KeyHandles ) override;
+	virtual void DilateSection( float DilationFactor, float Origin, TSet<FKeyHandle>& KeyHandles ) override;
+	virtual void GetKeyHandles(TSet<FKeyHandle>& KeyHandles) const override;
 	virtual void GetSnapTimes(TArray<float>& OutSnapTimes, bool bGetSectionBorders) const override;
 
 private:
diff --git a/Engine/Source/Runtime/MovieSceneCoreTypes/Classes/MovieSceneBoolSection.h b/Engine/Source/Runtime/MovieSceneCoreTypes/Classes/MovieSceneBoolSection.h
--- a/Engine/Source/Runtime/MovieSceneCoreTypes/Classes/MovieSceneBoolSection.h
+++ b/Engine/Source/Runtime/MovieSceneCoreTypes/Classes/MovieSceneBoolSection.h
@@ -39,8 +39,9 @@ public:
 	/**
 	 * UMovieSceneSection interface 
 	 */
-	virtual void MoveSection( float DeltaPosition ) override;
-	virtual void DilateSection( float DilationFactor, float Origin ) override;
+	virtual void MoveSection(float DeltaPosition, TSet<FKeyHandle>& KeyHandles) override;
+	virtual void DilateSection(float DilationFactor, float Origin, TSet<FKeyHandle>& KeyHandles) override;
+	virtual void GetKeyHandles(TSet<FKeyHandle>& KeyHandles) const override;
 
 	/** Gets all the keys of this boolean section */
 	FIntegralCurve& GetCurve() { return BoolCurve; }
diff --git a/Engine/Source/Runtime/MovieSceneCoreTypes/Classes/MovieSceneByteSection.h b/Engine/Source/Runtime/MovieSceneCoreTypes/Classes/MovieSceneByteSection.h
--- a/Engine/Source/Runtime/MovieSceneCoreTypes/Classes/MovieSceneByteSection.h
+++ b/Engine/Source/Runtime/MovieSceneCoreTypes/Classes/MovieSceneByteSection.h
@@ -39,8 +39,9 @@ public:
 	/**
 	 * UMovieSceneSection interface 
 	 */
-	virtual void MoveSection( float DeltaPosition ) override;
-	virtual void DilateSection( float DilationFactor, float Origin ) override;
+	virtual void MoveSection(float DeltaPosition, TSet<FKeyHandle>& KeyHandles) override;
+	virtual void DilateSection(float DilationFactor, float Origin, TSet<FKeyHandle>& KeyHandles) override;
+	virtual void GetKeyHandles(TSet<FKeyHandle>& KeyHandles) const override;
 
 	/** Gets all the keys of this byte section */
 	FIntegralCurve& GetCurve() { return ByteCurve; }
diff --git a/Engine/Source/Runtime/MovieSceneCoreTypes/Classes/MovieSceneColorSection.h b/Engine/Source/Runtime/MovieSceneCoreTypes/Classes/MovieSceneColorSection.h
--- a/Engine/Source/Runtime/MovieSceneCoreTypes/Classes/MovieSceneColorSection.h
+++ b/Engine/Source/Runtime/MovieSceneCoreTypes/Classes/MovieSceneColorSection.h
@@ -15,8 +15,9 @@ class UMovieSceneColorSection : public UMovieSceneSection
 	GENERATED_UCLASS_BODY()
 public:
 	/** MovieSceneSection interface */
-	virtual void MoveSection( float DeltaPosition ) override;
-	virtual void DilateSection( float DilationFactor, float Origin ) override;
+	virtual void MoveSection(float DeltaPosition, TSet<FKeyHandle>& KeyHandles) override;
+	virtual void DilateSection(float DilationFactor, float Origin, TSet<FKeyHandle>& KeyHandles) override;
+	virtual void GetKeyHandles(TSet<FKeyHandle>& KeyHandles) const override;
 
 	/**
 	 * Updates this section
diff --git a/Engine/Source/Runtime/MovieSceneCoreTypes/Classes/MovieSceneFloatSection.h b/Engine/Source/Runtime/MovieSceneCoreTypes/Classes/MovieSceneFloatSection.h
--- a/Engine/Source/Runtime/MovieSceneCoreTypes/Classes/MovieSceneFloatSection.h
+++ b/Engine/Source/Runtime/MovieSceneCoreTypes/Classes/MovieSceneFloatSection.h
@@ -39,8 +39,9 @@ public:
 	/**
 	 * UMovieSceneSection interface 
 	 */
-	virtual void MoveSection( float DeltaPosition ) override;
-	virtual void DilateSection( float DilationFactor, float Origin ) override;
+	virtual void MoveSection(float DeltaPosition, TSet<FKeyHandle>& KeyHandles) override;
+	virtual void DilateSection(float DilationFactor, float Origin, TSet<FKeyHandle>& KeyHandles) override;
+	virtual void GetKeyHandles(TSet<FKeyHandle>& KeyHandles) const override;
 
 	/**
 	 * @return The float curve on this section
diff --git a/Engine/Source/Runtime/MovieSceneCoreTypes/Classes/MovieSceneVectorSection.h b/Engine/Source/Runtime/MovieSceneCoreTypes/Classes/MovieSceneVectorSection.h
--- a/Engine/Source/Runtime/MovieSceneCoreTypes/Classes/MovieSceneVectorSection.h
+++ b/Engine/Source/Runtime/MovieSceneCoreTypes/Classes/MovieSceneVectorSection.h
@@ -23,9 +23,10 @@ public:
 	FVector4 Eval( float Position, const FVector4& DefaultVector ) const;
 
 	/** MovieSceneSection interface */
-	virtual void MoveSection( float DeltaPosition ) override;
-	virtual void DilateSection( float DilationFactor, float Origin ) override;
-	
+	virtual void MoveSection(float DeltaPosition, TSet<FKeyHandle>& KeyHandles) override;
+	virtual void DilateSection(float DilationFactor, float Origin, TSet<FKeyHandle>& KeyHandles) override;
+	virtual void GetKeyHandles(TSet<FKeyHandle>& KeyHandles) const override;
+
 	/** 
 	 * Adds a key to the section
 	 *
diff --git a/Engine/Source/Runtime/MovieSceneCoreTypes/Private/MovieScene3DTransformSection.cpp b/Engine/Source/Runtime/MovieSceneCoreTypes/Private/MovieScene3DTransformSection.cpp
--- a/Engine/Source/Runtime/MovieSceneCoreTypes/Private/MovieScene3DTransformSection.cpp
+++ b/Engine/Source/Runtime/MovieSceneCoreTypes/Private/MovieScene3DTransformSection.cpp
@@ -10,32 +10,61 @@ UMovieScene3DTransformSection::UMovieScene3DTransformSection( const FObjectIniti
 
 }
 
-void UMovieScene3DTransformSection::MoveSection( float DeltaTime )
+void UMovieScene3DTransformSection::MoveSection( float DeltaTime, TSet<FKeyHandle>& KeyHandles )
 {
-	Super::MoveSection( DeltaTime );
+	Super::MoveSection( DeltaTime, KeyHandles );
 
 	// Move all the curves in this section
 	for( int32 Axis = 0; Axis < 3; ++Axis )
 	{
-		Translation[Axis].ShiftCurve( DeltaTime );
-		Rotation[Axis].ShiftCurve( DeltaTime );
-		Scale[Axis].ShiftCurve( DeltaTime );
+		Translation[Axis].ShiftCurve( DeltaTime, KeyHandles );
+		Rotation[Axis].ShiftCurve( DeltaTime, KeyHandles );
+		Scale[Axis].ShiftCurve( DeltaTime, KeyHandles );
 	}
 }
 
-void UMovieScene3DTransformSection::DilateSection( float DilationFactor, float Origin )
+void UMovieScene3DTransformSection::DilateSection( float DilationFactor, float Origin, TSet<FKeyHandle>& KeyHandles )
 {
-	Super::DilateSection(DilationFactor, Origin);
+	Super::DilateSection(DilationFactor, Origin, KeyHandles);
 
 	for( int32 Axis = 0; Axis < 3; ++Axis )
 	{
-		Translation[Axis].ScaleCurve( Origin, DilationFactor );
-		Rotation[Axis].ScaleCurve( Origin, DilationFactor );
-		Scale[Axis].ScaleCurve( Origin, DilationFactor );
+		Translation[Axis].ScaleCurve( Origin, DilationFactor, KeyHandles );
+		Rotation[Axis].ScaleCurve( Origin, DilationFactor, KeyHandles );
+		Scale[Axis].ScaleCurve( Origin, DilationFactor, KeyHandles );
 	}
 }
 
-	
+void UMovieScene3DTransformSection::GetKeyHandles(TSet<FKeyHandle>& KeyHandles) const
+{
+	for (int32 Axis = 0; Axis < 3; ++Axis)
+	{
+		for (auto It(Translation[Axis].GetKeyHandleIterator()); It; ++It)
+		{
+			float Time = Translation[Axis].GetKeyTime(It.Key());
+			if (IsTimeWithinSection(Time))
+			{
+				KeyHandles.Add(It.Key());
+			}
+		}
+		for (auto It(Rotation[Axis].GetKeyHandleIterator()); It; ++It)
+		{
+			float Time = Rotation[Axis].GetKeyTime(It.Key());
+			if (IsTimeWithinSection(Time))
+			{
+				KeyHandles.Add(It.Key());
+			}
+		}
+		for (auto It(Scale[Axis].GetKeyHandleIterator()); It; ++It)
+		{
+			float Time = Scale[Axis].GetKeyTime(It.Key());
+			if (IsTimeWithinSection(Time))
+			{
+				KeyHandles.Add(It.Key());
+			}
+		}
+	}
+}	
 
 bool UMovieScene3DTransformSection::EvalTranslation( float Time, FVector& OutTranslation ) const
 {
diff --git a/Engine/Source/Runtime/MovieSceneCoreTypes/Private/MovieSceneAnimationSection.cpp b/Engine/Source/Runtime/MovieSceneCoreTypes/Private/MovieSceneAnimationSection.cpp
--- a/Engine/Source/Runtime/MovieSceneCoreTypes/Private/MovieSceneAnimationSection.cpp
+++ b/Engine/Source/Runtime/MovieSceneCoreTypes/Private/MovieSceneAnimationSection.cpp
@@ -11,19 +11,23 @@ UMovieSceneAnimationSection::UMovieSceneAnimationSection( const FObjectInitializ
 	AnimationDilationFactor = 1.f;
 }
 
-void UMovieSceneAnimationSection::MoveSection( float DeltaTime )
+void UMovieSceneAnimationSection::MoveSection( float DeltaTime, TSet<FKeyHandle>& KeyHandles )
 {
-	Super::MoveSection(DeltaTime);
+	Super::MoveSection(DeltaTime, KeyHandles);
 
 	AnimationStartTime += DeltaTime;
 }
 
-void UMovieSceneAnimationSection::DilateSection( float DilationFactor, float Origin )
+void UMovieSceneAnimationSection::DilateSection( float DilationFactor, float Origin, TSet<FKeyHandle>& KeyHandles )
 {
 	AnimationStartTime = (AnimationStartTime - Origin) * DilationFactor + Origin;
 	AnimationDilationFactor *= DilationFactor;
 
-	Super::DilateSection(DilationFactor, Origin);
+	Super::DilateSection(DilationFactor, Origin, KeyHandles);
+}
+
+void UMovieSceneAnimationSection::GetKeyHandles(TSet<FKeyHandle>& KeyHandles) const
+{
 }
 
 void UMovieSceneAnimationSection::GetSnapTimes(TArray<float>& OutSnapTimes, bool bGetSectionBorders) const
diff --git a/Engine/Source/Runtime/MovieSceneCoreTypes/Private/MovieSceneAudioSection.cpp b/Engine/Source/Runtime/MovieSceneCoreTypes/Private/MovieSceneAudioSection.cpp
--- a/Engine/Source/Runtime/MovieSceneCoreTypes/Private/MovieSceneAudioSection.cpp
+++ b/Engine/Source/Runtime/MovieSceneCoreTypes/Private/MovieSceneAudioSection.cpp
@@ -25,19 +25,23 @@ TRange<float> UMovieSceneAudioSection::GetAudioTrueRange() const
 	return !Sound ? TRange<float>::Empty() : TRange<float>(AudioStartTime, AudioStartTime + Sound->GetDuration() * AudioDilationFactor);
 }
 
-void UMovieSceneAudioSection::MoveSection( float DeltaTime )
+void UMovieSceneAudioSection::MoveSection( float DeltaTime, TSet<FKeyHandle>& KeyHandles )
 {
-	Super::MoveSection(DeltaTime);
+	Super::MoveSection(DeltaTime, KeyHandles);
 
 	AudioStartTime += DeltaTime;
 }
 
-void UMovieSceneAudioSection::DilateSection( float DilationFactor, float Origin )
+void UMovieSceneAudioSection::DilateSection( float DilationFactor, float Origin, TSet<FKeyHandle>& KeyHandles )
 {
 	AudioStartTime = (AudioStartTime - Origin) * DilationFactor + Origin;
 	AudioDilationFactor *= DilationFactor;
 
-	Super::DilateSection(DilationFactor, Origin);
+	Super::DilateSection(DilationFactor, Origin, KeyHandles);
+}
+
+void UMovieSceneAudioSection::GetKeyHandles(TSet<FKeyHandle>& KeyHandles) const
+{
 }
 
 void UMovieSceneAudioSection::GetSnapTimes(TArray<float>& OutSnapTimes, bool bGetSectionBorders) const
diff --git a/Engine/Source/Runtime/MovieSceneCoreTypes/Private/MovieSceneBoolSection.cpp b/Engine/Source/Runtime/MovieSceneCoreTypes/Private/MovieSceneBoolSection.cpp
--- a/Engine/Source/Runtime/MovieSceneCoreTypes/Private/MovieSceneBoolSection.cpp
+++ b/Engine/Source/Runtime/MovieSceneCoreTypes/Private/MovieSceneBoolSection.cpp
@@ -14,19 +14,33 @@ bool UMovieSceneBoolSection::Eval( float Position ) const
 	return !!BoolCurve.Evaluate(Position);
 }
 
-void UMovieSceneBoolSection::MoveSection( float DeltaPosition )
+void UMovieSceneBoolSection::MoveSection( float DeltaPosition, TSet<FKeyHandle>& KeyHandles )
 {
-	Super::MoveSection( DeltaPosition );
+	Super::MoveSection( DeltaPosition, KeyHandles );
 
-	BoolCurve.ShiftCurve(DeltaPosition);
+	BoolCurve.ShiftCurve(DeltaPosition, KeyHandles);
 }
 
-void UMovieSceneBoolSection::DilateSection( float DilationFactor, float Origin )
+void UMovieSceneBoolSection::DilateSection( float DilationFactor, float Origin, TSet<FKeyHandle>& KeyHandles )
 {
-	Super::DilateSection(DilationFactor, Origin);
+	float StartTime = GetStartTime();
+	float EndTime = GetEndTime();
+
+	Super::DilateSection(DilationFactor, Origin, KeyHandles);
 	
-	BoolCurve.ScaleCurve(Origin, DilationFactor);
+	BoolCurve.ScaleCurve(Origin, DilationFactor, KeyHandles);
+}
 
+void UMovieSceneBoolSection::GetKeyHandles(TSet<FKeyHandle>& KeyHandles) const
+{
+	for (auto It(BoolCurve.GetKeyHandleIterator()); It; ++It)
+	{
+		float Time = BoolCurve.GetKeyTime(It.Key());
+		if (IsTimeWithinSection(Time))
+		{
+			KeyHandles.Add(It.Key());
+		}
+	}
 }
 
 void UMovieSceneBoolSection::AddKey( float Time, bool Value )
diff --git a/Engine/Source/Runtime/MovieSceneCoreTypes/Private/MovieSceneByteSection.cpp b/Engine/Source/Runtime/MovieSceneCoreTypes/Private/MovieSceneByteSection.cpp
--- a/Engine/Source/Runtime/MovieSceneCoreTypes/Private/MovieSceneByteSection.cpp
+++ b/Engine/Source/Runtime/MovieSceneCoreTypes/Private/MovieSceneByteSection.cpp
@@ -14,19 +14,28 @@ uint8 UMovieSceneByteSection::Eval( float Position ) const
 	return !!ByteCurve.Evaluate(Position);
 }
 
-void UMovieSceneByteSection::MoveSection( float DeltaPosition )
+void UMovieSceneByteSection::MoveSection( float DeltaPosition, TSet<FKeyHandle>& KeyHandles )
 {
-	Super::MoveSection( DeltaPosition );
+	Super::MoveSection( DeltaPosition, KeyHandles );
 
-	ByteCurve.ShiftCurve(DeltaPosition);
+	ByteCurve.ShiftCurve(DeltaPosition, KeyHandles);
 }
 
-void UMovieSceneByteSection::DilateSection( float DilationFactor, float Origin )
+void UMovieSceneByteSection::DilateSection( float DilationFactor, float Origin, TSet<FKeyHandle>& KeyHandles )
 {
-	Super::DilateSection(DilationFactor, Origin);
+	Super::DilateSection(DilationFactor, Origin, KeyHandles);
 	
-	ByteCurve.ScaleCurve(Origin, DilationFactor);
+	ByteCurve.ScaleCurve(Origin, DilationFactor, KeyHandles);
+}
 
+void UMovieSceneByteSection::GetKeyHandles(TSet<FKeyHandle>& KeyHandles) const
+{
+	for (auto It(ByteCurve.GetKeyHandleIterator()); It; ++It)
+	{
+		float Time = ByteCurve.GetKeyTime(It.Key());
+		if (IsTimeWithinSection(Time))
+			KeyHandles.Add(It.Key());
+	}
 }
 
 void UMovieSceneByteSection::AddKey( float Time, uint8 Value )
diff --git a/Engine/Source/Runtime/MovieSceneCoreTypes/Private/MovieSceneColorSection.cpp b/Engine/Source/Runtime/MovieSceneCoreTypes/Private/MovieSceneColorSection.cpp
--- a/Engine/Source/Runtime/MovieSceneCoreTypes/Private/MovieSceneColorSection.cpp
+++ b/Engine/Source/Runtime/MovieSceneCoreTypes/Private/MovieSceneColorSection.cpp
@@ -9,26 +9,59 @@ UMovieSceneColorSection::UMovieSceneColorSection( const FObjectInitializer& Obje
 {
 }
 
-void UMovieSceneColorSection::MoveSection( float DeltaTime )
+void UMovieSceneColorSection::MoveSection( float DeltaTime, TSet<FKeyHandle>& KeyHandles )
 {
-	Super::MoveSection( DeltaTime );
+	Super::MoveSection( DeltaTime, KeyHandles );
 
 	// Move all the curves in this section
-	RedCurve.ShiftCurve(DeltaTime);
-	GreenCurve.ShiftCurve(DeltaTime);
-	BlueCurve.ShiftCurve(DeltaTime);
-	AlphaCurve.ShiftCurve(DeltaTime);
+	RedCurve.ShiftCurve(DeltaTime, KeyHandles);
+	GreenCurve.ShiftCurve(DeltaTime, KeyHandles);
+	BlueCurve.ShiftCurve(DeltaTime, KeyHandles);
+	AlphaCurve.ShiftCurve(DeltaTime, KeyHandles);
 }
 
-void UMovieSceneColorSection::DilateSection( float DilationFactor, float Origin )
+void UMovieSceneColorSection::DilateSection( float DilationFactor, float Origin, TSet<FKeyHandle>& KeyHandles )
 {
-	Super::DilateSection(DilationFactor, Origin);
+	Super::DilateSection(DilationFactor, Origin, KeyHandles);
 
-	RedCurve.ScaleCurve(Origin, DilationFactor);
-	GreenCurve.ScaleCurve(Origin, DilationFactor);
-	BlueCurve.ScaleCurve(Origin, DilationFactor);
-	AlphaCurve.ScaleCurve(Origin, DilationFactor);
+	RedCurve.ScaleCurve(Origin, DilationFactor, KeyHandles);
+	GreenCurve.ScaleCurve(Origin, DilationFactor, KeyHandles);
+	BlueCurve.ScaleCurve(Origin, DilationFactor, KeyHandles);
+	AlphaCurve.ScaleCurve(Origin, DilationFactor, KeyHandles);
+}
 
+void UMovieSceneColorSection::GetKeyHandles(TSet<FKeyHandle>& KeyHandles) const
+{
+	for (auto It(RedCurve.GetKeyHandleIterator()); It; ++It)
+	{
+		float Time = RedCurve.GetKeyTime(It.Key());
+		if (IsTimeWithinSection(Time))
+		{
+			KeyHandles.Add(It.Key());
+		}
+	}
+	for (auto It(GreenCurve.GetKeyHandleIterator()); It; ++It)
+	{
+		float Time = GreenCurve.GetKeyTime(It.Key());
+		if (IsTimeWithinSection(Time))
+		{
+			KeyHandles.Add(It.Key());
+		}
+	}
+	for (auto It(BlueCurve.GetKeyHandleIterator()); It; ++It)
+	{
+		float Time = BlueCurve.GetKeyTime(It.Key());
+		if (IsTimeWithinSection(Time))
+		{
+			KeyHandles.Add(It.Key());
+		}
+	}
+	for (auto It(AlphaCurve.GetKeyHandleIterator()); It; ++It)
+	{
+		float Time = AlphaCurve.GetKeyTime(It.Key());
+		if (IsTimeWithinSection(Time))
+			KeyHandles.Add(It.Key());
+	}
 }
 
 FLinearColor UMovieSceneColorSection::Eval( float Position, const FLinearColor& DefaultColor ) const
@@ -89,4 +122,4 @@ void UMovieSceneColorSection::AddKeyToNamedCurve(float Time, const FColorKey& Ke
 	{
 		AddKeyToCurve(AlphaCurve, Time, Key.Value.A);
 	}
-}
\ No newline at end of file
+}
diff --git a/Engine/Source/Runtime/MovieSceneCoreTypes/Private/MovieSceneFloatSection.cpp b/Engine/Source/Runtime/MovieSceneCoreTypes/Private/MovieSceneFloatSection.cpp
--- a/Engine/Source/Runtime/MovieSceneCoreTypes/Private/MovieSceneFloatSection.cpp
+++ b/Engine/Source/Runtime/MovieSceneCoreTypes/Private/MovieSceneFloatSection.cpp
@@ -14,19 +14,31 @@ float UMovieSceneFloatSection::Eval( float Position ) const
 	return FloatCurve.Eval( Position );
 }
 
-void UMovieSceneFloatSection::MoveSection( float DeltaPosition )
+void UMovieSceneFloatSection::MoveSection( float DeltaPosition, TSet<FKeyHandle>& KeyHandles )
 {
-	Super::MoveSection( DeltaPosition );
+	Super::MoveSection( DeltaPosition, KeyHandles );
 
 	// Move the curve
-	FloatCurve.ShiftCurve(DeltaPosition );
+	FloatCurve.ShiftCurve(DeltaPosition, KeyHandles);
 }
 
-void UMovieSceneFloatSection::DilateSection( float DilationFactor, float Origin )
+void UMovieSceneFloatSection::DilateSection( float DilationFactor, float Origin, TSet<FKeyHandle>& KeyHandles )
 {	
-	Super::DilateSection(DilationFactor, Origin);
+	Super::DilateSection(DilationFactor, Origin, KeyHandles);
 	
-	FloatCurve.ScaleCurve(Origin, DilationFactor);
+	FloatCurve.ScaleCurve(Origin, DilationFactor, KeyHandles);
+}
+
+void UMovieSceneFloatSection::GetKeyHandles(TSet<FKeyHandle>& KeyHandles) const
+{
+	for (auto It(FloatCurve.GetKeyHandleIterator()); It; ++It)
+	{
+		float Time = FloatCurve.GetKeyTime(It.Key());
+		if (IsTimeWithinSection(Time))
+		{
+			KeyHandles.Add(It.Key());
+		}
+	}
 }
 
 void UMovieSceneFloatSection::AddKey( float Time, float Value )
diff --git a/Engine/Source/Runtime/MovieSceneCoreTypes/Private/MovieSceneVectorSection.cpp b/Engine/Source/Runtime/MovieSceneCoreTypes/Private/MovieSceneVectorSection.cpp
--- a/Engine/Source/Runtime/MovieSceneCoreTypes/Private/MovieSceneVectorSection.cpp
+++ b/Engine/Source/Runtime/MovieSceneCoreTypes/Private/MovieSceneVectorSection.cpp
@@ -81,25 +81,40 @@ void UMovieSceneVectorSection::AddKeyToNamedCurve(float Time, FName CurveName, c
 	}
 }
 
-void UMovieSceneVectorSection::MoveSection( float DeltaTime )
+void UMovieSceneVectorSection::MoveSection( float DeltaTime, TSet<FKeyHandle>& KeyHandles )
 {
 	check(ChannelsUsed >= 2 && ChannelsUsed <= 4);
 
-	Super::MoveSection( DeltaTime );
+	Super::MoveSection( DeltaTime, KeyHandles );
 
 	for (int32 i = 0; i < ChannelsUsed; ++i)
 	{
-		Curves[i].ShiftCurve(DeltaTime);
+		Curves[i].ShiftCurve(DeltaTime, KeyHandles);
 	}
 }
 
-void UMovieSceneVectorSection::DilateSection( float DilationFactor, float Origin )
+void UMovieSceneVectorSection::DilateSection( float DilationFactor, float Origin, TSet<FKeyHandle>& KeyHandles )
 {
 	check(ChannelsUsed >= 2 && ChannelsUsed <= 4);
-	Super::DilateSection(DilationFactor, Origin);
+	Super::DilateSection(DilationFactor, Origin, KeyHandles);
 
 	for (int32 i = 0; i < ChannelsUsed; ++i)
 	{
-		Curves[i].ScaleCurve(Origin, DilationFactor);
+		Curves[i].ScaleCurve(Origin, DilationFactor, KeyHandles);
+	}
+}
+
+void UMovieSceneVectorSection::GetKeyHandles(TSet<FKeyHandle>& KeyHandles) const
+{
+	for (int32 i = 0; i < ChannelsUsed; ++i)
+	{
+		for (auto It(Curves[i].GetKeyHandleIterator()); It; ++It)
+		{
+			float Time = Curves[i].GetKeyTime(It.Key());
+			if (IsTimeWithinSection(Time))
+			{
+				KeyHandles.Add(It.Key());
+			}
+		}
 	}
 }
diff --git a/Engine/Source/Runtime/UMG/Private/Animation/MovieScene2DTransformSection.cpp b/Engine/Source/Runtime/UMG/Private/Animation/MovieScene2DTransformSection.cpp
--- a/Engine/Source/Runtime/UMG/Private/Animation/MovieScene2DTransformSection.cpp
+++ b/Engine/Source/Runtime/UMG/Private/Animation/MovieScene2DTransformSection.cpp
@@ -10,32 +10,72 @@ UMovieScene2DTransformSection::UMovieScene2DTransformSection( const FObjectIniti
 
 }
 
-void UMovieScene2DTransformSection::MoveSection( float DeltaTime )
+void UMovieScene2DTransformSection::MoveSection( float DeltaTime, TSet<FKeyHandle>& KeyHandles )
 {
-	Super::MoveSection( DeltaTime );
+	Super::MoveSection( DeltaTime, KeyHandles );
 
-	Rotation.ShiftCurve( DeltaTime );
+	Rotation.ShiftCurve( DeltaTime, KeyHandles );
 
 	// Move all the curves in this section
 	for( int32 Axis = 0; Axis < 2; ++Axis )
 	{
-		Translation[Axis].ShiftCurve( DeltaTime );
-		Scale[Axis].ShiftCurve( DeltaTime );
-		Shear[Axis].ShiftCurve( DeltaTime );
+		Translation[Axis].ShiftCurve( DeltaTime, KeyHandles );
+		Scale[Axis].ShiftCurve( DeltaTime, KeyHandles );
+		Shear[Axis].ShiftCurve( DeltaTime, KeyHandles );
 	}
 }
 
-void UMovieScene2DTransformSection::DilateSection( float DilationFactor, float Origin )
+void UMovieScene2DTransformSection::DilateSection( float DilationFactor, float Origin, TSet<FKeyHandle>& KeyHandles )
 {
-	Super::DilateSection(DilationFactor, Origin);
+	Super::DilateSection(DilationFactor, Origin, KeyHandles);
 
-	Rotation.ScaleCurve( Origin, DilationFactor );
+	Rotation.ScaleCurve( Origin, DilationFactor, KeyHandles);
 
 	for( int32 Axis = 0; Axis < 2; ++Axis )
 	{
-		Translation[Axis].ScaleCurve( Origin, DilationFactor );
-		Scale[Axis].ScaleCurve( Origin, DilationFactor );
-		Shear[Axis].ScaleCurve( Origin, DilationFactor  );
+		Translation[Axis].ScaleCurve( Origin, DilationFactor, KeyHandles );
+		Scale[Axis].ScaleCurve( Origin, DilationFactor, KeyHandles );
+		Shear[Axis].ScaleCurve( Origin, DilationFactor, KeyHandles );
+	}
+}
+
+void UMovieScene2DTransformSection::GetKeyHandles(TSet<FKeyHandle>& KeyHandles) const
+{
+	for (auto It(Rotation.GetKeyHandleIterator()); It; ++It)
+	{
+		float Time = Rotation.GetKeyTime(It.Key());
+		if (IsTimeWithinSection(Time))
+		{
+			KeyHandles.Add(It.Key());
+		}
+	}
+
+	for (int32 Axis = 0; Axis < 2; ++Axis)
+	{
+		for (auto It(Translation[Axis].GetKeyHandleIterator()); It; ++It)
+		{
+			float Time = Translation[Axis].GetKeyTime(It.Key());
+			if (IsTimeWithinSection(Time))
+			{
+				KeyHandles.Add(It.Key());
+			}
+		}
+		for (auto It(Scale[Axis].GetKeyHandleIterator()); It; ++It)
+		{
+			float Time = Scale[Axis].GetKeyTime(It.Key());
+			if (IsTimeWithinSection(Time))
+			{
+				KeyHandles.Add(It.Key());
+			}
+		}
+		for (auto It(Shear[Axis].GetKeyHandleIterator()); It; ++It)
+		{
+			float Time = Shear[Axis].GetKeyTime(It.Key());
+			if (IsTimeWithinSection(Time))
+			{
+				KeyHandles.Add(It.Key());
+			}
+		}
 	}
 }
 
@@ -153,4 +193,4 @@ void UMovieScene2DTransformSection::AddKey( float Time, const struct F2DTransfor
 	{
 		AddKeyToNamedCurve(Time, TransformKey);
 	}
-}
\ No newline at end of file
+}
diff --git a/Engine/Source/Runtime/UMG/Private/Animation/MovieSceneMarginSection.cpp b/Engine/Source/Runtime/UMG/Private/Animation/MovieSceneMarginSection.cpp
--- a/Engine/Source/Runtime/UMG/Private/Animation/MovieSceneMarginSection.cpp
+++ b/Engine/Source/Runtime/UMG/Private/Animation/MovieSceneMarginSection.cpp
@@ -9,25 +9,61 @@ UMovieSceneMarginSection::UMovieSceneMarginSection( const FObjectInitializer& Ob
 {
 }
 
-void UMovieSceneMarginSection::MoveSection( float DeltaTime )
+void UMovieSceneMarginSection::MoveSection( float DeltaTime, TSet<FKeyHandle>& KeyHandles )
 {
-	Super::MoveSection( DeltaTime );
+	Super::MoveSection( DeltaTime, KeyHandles );
 
 	// Move all the curves in this section
-	LeftCurve.ShiftCurve(DeltaTime);
-	TopCurve.ShiftCurve(DeltaTime);
-	RightCurve.ShiftCurve(DeltaTime);
-	BottomCurve.ShiftCurve(DeltaTime);
+	LeftCurve.ShiftCurve(DeltaTime, KeyHandles);
+	TopCurve.ShiftCurve(DeltaTime, KeyHandles);
+	RightCurve.ShiftCurve(DeltaTime, KeyHandles);
+	BottomCurve.ShiftCurve(DeltaTime, KeyHandles);
 }
 
-void UMovieSceneMarginSection::DilateSection( float DilationFactor, float Origin )
+void UMovieSceneMarginSection::DilateSection( float DilationFactor, float Origin, TSet<FKeyHandle>& KeyHandles )
 {
-	Super::DilateSection(DilationFactor, Origin);
+	Super::DilateSection(DilationFactor, Origin, KeyHandles);
 
-	LeftCurve.ScaleCurve(Origin, DilationFactor);
-	TopCurve.ScaleCurve(Origin, DilationFactor);
-	RightCurve.ScaleCurve(Origin, DilationFactor);
-	BottomCurve.ScaleCurve(Origin, DilationFactor);
+	LeftCurve.ScaleCurve(Origin, DilationFactor, KeyHandles);
+	TopCurve.ScaleCurve(Origin, DilationFactor, KeyHandles);
+	RightCurve.ScaleCurve(Origin, DilationFactor, KeyHandles);
+	BottomCurve.ScaleCurve(Origin, DilationFactor, KeyHandles);
+}
+
+void UMovieSceneMarginSection::GetKeyHandles(TSet<FKeyHandle>& KeyHandles) const
+{
+	for (auto It(LeftCurve.GetKeyHandleIterator()); It; ++It)
+	{
+		float Time = LeftCurve.GetKeyTime(It.Key());
+		if (IsTimeWithinSection(Time))
+		{
+			KeyHandles.Add(It.Key());
+		}
+	}
+	for (auto It(TopCurve.GetKeyHandleIterator()); It; ++It)
+	{
+		float Time = TopCurve.GetKeyTime(It.Key());
+		if (IsTimeWithinSection(Time))
+		{
+			KeyHandles.Add(It.Key());
+		}
+	}
+	for (auto It(RightCurve.GetKeyHandleIterator()); It; ++It)
+	{
+		float Time = RightCurve.GetKeyTime(It.Key());
+		if (IsTimeWithinSection(Time))
+		{
+			KeyHandles.Add(It.Key());
+		}
+	}
+	for (auto It(BottomCurve.GetKeyHandleIterator()); It; ++It)
+	{
+		float Time = BottomCurve.GetKeyTime(It.Key());
+		if (IsTimeWithinSection(Time))
+		{
+			KeyHandles.Add(It.Key());
+		}
+	}
 }
 
 FMargin UMovieSceneMarginSection::Eval( float Position, const FMargin& DefaultValue ) const
diff --git a/Engine/Source/Runtime/UMG/Public/Animation/MovieScene2DTransformSection.h b/Engine/Source/Runtime/UMG/Public/Animation/MovieScene2DTransformSection.h
--- a/Engine/Source/Runtime/UMG/Public/Animation/MovieScene2DTransformSection.h
+++ b/Engine/Source/Runtime/UMG/Public/Animation/MovieScene2DTransformSection.h
@@ -18,8 +18,9 @@ class UMovieScene2DTransformSection : public UMovieSceneSection
 public:
 
 	/** MovieSceneSection interface */
-	virtual void MoveSection( float DeltaPosition ) override;
-	virtual void DilateSection( float DilationFactor, float Origin ) override;
+	virtual void MoveSection(float DeltaPosition, TSet<FKeyHandle>& KeyHandles) override;
+	virtual void DilateSection(float DilationFactor, float Origin, TSet<FKeyHandle>& KeyHandles) override;
+	virtual void GetKeyHandles(TSet<FKeyHandle>& KeyHandles) const override;
 
 	UMG_API FRichCurve& GetTranslationCurve( EAxis::Type Axis );
 
diff --git a/Engine/Source/Runtime/UMG/Public/Animation/MovieSceneMarginSection.h b/Engine/Source/Runtime/UMG/Public/Animation/MovieSceneMarginSection.h
--- a/Engine/Source/Runtime/UMG/Public/Animation/MovieSceneMarginSection.h
+++ b/Engine/Source/Runtime/UMG/Public/Animation/MovieSceneMarginSection.h
@@ -17,8 +17,9 @@ class UMovieSceneMarginSection : public UMovieSceneSection
 	GENERATED_UCLASS_BODY()
 public:
 	/** UMovieSceneSection interface */
-	virtual void MoveSection( float DeltaPosition ) override;
-	virtual void DilateSection( float DilationFactor, float Origin ) override;
+	virtual void MoveSection(float DeltaPosition, TSet<FKeyHandle>& KeyHandles) override;
+	virtual void DilateSection(float DilationFactor, float Origin, TSet<FKeyHandle>& KeyHandles) override;
+	virtual void GetKeyHandles(TSet<FKeyHandle>& KeyHandles) const override;
 
 	/**
 	 * Updates this section
@@ -74,6 +75,7 @@ public:
 	 */
 	FRichCurve& GetBottomCurve() { return BottomCurve; }
 	const FRichCurve& GetBottomCurve() const { return BottomCurve; }
+
 private:
 	void AddKeyToNamedCurve( float Time, const FMarginKey& MarginKey );
 
