UMG/Sequencer - Disable auto-key when no animation is selected.

diff --git a/Engine/Source/Editor/Sequencer/Private/Sequencer.cpp b/Engine/Source/Editor/Sequencer/Private/Sequencer.cpp
--- a/Engine/Source/Editor/Sequencer/Private/Sequencer.cpp
+++ b/Engine/Source/Editor/Sequencer/Private/Sequencer.cpp
@@ -498,11 +498,16 @@ FAnimatedRange FSequencer::GetViewRange() const
 	return AnimatedRange;
 }
 
-bool FSequencer::IsAutoKeyEnabled() const 
+bool FSequencer::GetAutoKeyEnabled() const 
 {
 	return bAllowAutoKey;
 }
 
+void FSequencer::SetAutoKeyEnabled(bool bAutoKeyEnabled)
+{
+	bAllowAutoKey = bAutoKeyEnabled;
+}
+
 bool FSequencer::IsRecordingLive() const 
 {
 	return PlaybackState == EMovieScenePlayerStatus::Recording && GIsPlayInEditorWorld;
diff --git a/Engine/Source/Editor/Sequencer/Private/Sequencer.h b/Engine/Source/Editor/Sequencer/Private/Sequencer.h
--- a/Engine/Source/Editor/Sequencer/Private/Sequencer.h
+++ b/Engine/Source/Editor/Sequencer/Private/Sequencer.h
@@ -75,7 +75,8 @@ public:
 	TSharedRef<FMovieSceneInstance> GetInstanceForSubMovieSceneSection( UMovieSceneSection& SubMovieSceneSection ) const override;
 	virtual void AddNewShot(FGuid CameraGuid) override;
 	virtual void AddAnimation(FGuid ObjectGuid, class UAnimSequence* AnimSequence) override;
-	virtual bool IsAutoKeyEnabled() const override;
+	virtual bool GetAutoKeyEnabled() const override;
+	virtual void SetAutoKeyEnabled(bool bAutoKeyEnabled) override;
 	virtual bool IsRecordingLive() const override;
 	virtual float GetCurrentLocalTime(UMovieScene& MovieScene) override;
 	virtual float GetGlobalTime() override;
diff --git a/Engine/Source/Editor/Sequencer/Private/SequencerObjectChangeListener.cpp b/Engine/Source/Editor/Sequencer/Private/SequencerObjectChangeListener.cpp
--- a/Engine/Source/Editor/Sequencer/Private/SequencerObjectChangeListener.cpp
+++ b/Engine/Source/Editor/Sequencer/Private/SequencerObjectChangeListener.cpp
@@ -152,7 +152,7 @@ void FSequencerObjectChangeListener::KeyProperty(FKeyPropertyParams KeyPropertyP
 void FSequencerObjectChangeListener::OnObjectPreEditChange( UObject* Object, const FEditPropertyChain& PropertyChain )
 {
 	// We only care if we are in auto-key mode and not attempting to change properties of a CDO (which cannot be animated)
-	if( Sequencer.IsValid() && Sequencer.Pin()->IsAutoKeyEnabled() && !Object->HasAnyFlags(RF_ClassDefaultObject) && PropertyChain.GetActiveMemberNode() )
+	if( Sequencer.IsValid() && Sequencer.Pin()->GetAutoKeyEnabled() && !Object->HasAnyFlags(RF_ClassDefaultObject) && PropertyChain.GetActiveMemberNode() )
 	{
 		// Register with the property editor module that we'd like to know about animated float properties that change
 		FPropertyEditorModule& PropertyEditor = FModuleManager::Get().LoadModuleChecked<FPropertyEditorModule>( "PropertyEditor" );
@@ -194,7 +194,7 @@ void FSequencerObjectChangeListener::OnObjectPostEditChange( UObject* Object, FP
 		bool bShouldPropagateChanges = bIsObjectValid;
 
 		// We only care if we are in auto-key mode and not attempting to change properties of a CDO (which cannot be animated)
-		if( Sequencer.IsValid() && Sequencer.Pin()->IsAutoKeyEnabled() && bIsObjectValid && !Object->HasAnyFlags(RF_ClassDefaultObject) )
+		if( Sequencer.IsValid() && Sequencer.Pin()->GetAutoKeyEnabled() && bIsObjectValid && !Object->HasAnyFlags(RF_ClassDefaultObject) )
 		{
 			TSharedPtr< IPropertyChangeListener > Listener;
 			ActivePropertyChangeListeners.RemoveAndCopyValue( Object, Listener );
diff --git a/Engine/Source/Editor/Sequencer/Public/ISequencer.h b/Engine/Source/Editor/Sequencer/Public/ISequencer.h
--- a/Engine/Source/Editor/Sequencer/Public/ISequencer.h
+++ b/Engine/Source/Editor/Sequencer/Public/ISequencer.h
@@ -56,7 +56,10 @@ public:
 	virtual void AddSubMovieScene( UMovieScene* SubMovieScene ) = 0;
 
 	/** @return Returns whether auto-key is enabled in this sequencer */
-	virtual bool IsAutoKeyEnabled() const = 0;
+	virtual bool GetAutoKeyEnabled() const = 0;
+
+	/** Sets whether autokey is enabled in this sequencer. */
+	virtual void SetAutoKeyEnabled(bool bAutoKeyEnabled) = 0;
 
 	/** @return Returns whether sequencer is currently recording live data from simulated actors */
 	virtual bool IsRecordingLive() const = 0;
diff --git a/Engine/Source/Editor/Sequencer/Public/MovieSceneTrackEditor.h b/Engine/Source/Editor/Sequencer/Public/MovieSceneTrackEditor.h
--- a/Engine/Source/Editor/Sequencer/Public/MovieSceneTrackEditor.h
+++ b/Engine/Source/Editor/Sequencer/Public/MovieSceneTrackEditor.h
@@ -49,7 +49,7 @@ public:
 	{
 		// @todo sequencer livecapture: This turns on "auto key" for the purpose of capture keys for actor state
 		// during PIE sessions when record mode is active.
-		return Sequencer.Pin()->IsRecordingLive() || Sequencer.Pin()->IsAutoKeyEnabled();
+		return Sequencer.Pin()->IsRecordingLive() || Sequencer.Pin()->GetAutoKeyEnabled();
 	}
 
 	void NotifyMovieSceneDataChanged()
diff --git a/Engine/Source/Editor/UMGEditor/Private/Details/SWidgetDetailsView.cpp b/Engine/Source/Editor/UMGEditor/Private/Details/SWidgetDetailsView.cpp
--- a/Engine/Source/Editor/UMGEditor/Private/Details/SWidgetDetailsView.cpp
+++ b/Engine/Source/Editor/UMGEditor/Private/Details/SWidgetDetailsView.cpp
@@ -467,7 +467,7 @@ void SWidgetDetailsView::HandleIsVariableChanged(ECheckBoxState CheckState)
 void SWidgetDetailsView::NotifyPreChange(FEditPropertyChain* PropertyAboutToChange)
 {
 	// During auto-key do not migrate values
-	if( !BlueprintEditor.Pin()->GetSequencer()->IsAutoKeyEnabled() )
+	if( !BlueprintEditor.Pin()->GetSequencer()->GetAutoKeyEnabled() )
 	{
 		TSharedPtr<FWidgetBlueprintEditor> Editor = BlueprintEditor.Pin();
 
@@ -480,7 +480,7 @@ void SWidgetDetailsView::NotifyPostChange(const FPropertyChangedEvent& PropertyC
 {
 	const static FName DesignerRebuildName("DesignerRebuild");
 
-	if ( PropertyChangedEvent.ChangeType != EPropertyChangeType::Interactive && !BlueprintEditor.Pin()->GetSequencer()->IsAutoKeyEnabled() )
+	if ( PropertyChangedEvent.ChangeType != EPropertyChangeType::Interactive && !BlueprintEditor.Pin()->GetSequencer()->GetAutoKeyEnabled() )
 	{
 		TSharedPtr<FWidgetBlueprintEditor> Editor = BlueprintEditor.Pin();
 
diff --git a/Engine/Source/Editor/UMGEditor/Private/WidgetBlueprintEditor.cpp b/Engine/Source/Editor/UMGEditor/Private/WidgetBlueprintEditor.cpp
--- a/Engine/Source/Editor/UMGEditor/Private/WidgetBlueprintEditor.cpp
+++ b/Engine/Source/Editor/UMGEditor/Private/WidgetBlueprintEditor.cpp
@@ -572,6 +572,7 @@ void FWidgetBlueprintEditor::ChangeViewedAnimation( UWidgetAnimation& InAnimatio
 		{
 			// Disable sequencer from interaction
 			Sequencer->GetSequencerWidget()->SetEnabled(false);
+			Sequencer->SetAutoKeyEnabled(false);
 			NoAnimationTextBlockPin->SetVisibility(EVisibility::Visible);
 			SequencerOverlayPin->SetVisibility( EVisibility::HitTestInvisible );
 		}
