Sequencer: Add ability to expand/collapse nodes from a hotkey.

'o' will expand or collapse the selected node (or all nodes if non selected)
'shift-O' will expand or collapse including its descendant nodes

diff --git a/Engine/Source/Editor/Sequencer/Private/SSequencer.cpp b/Engine/Source/Editor/Sequencer/Private/SSequencer.cpp
--- a/Engine/Source/Editor/Sequencer/Private/SSequencer.cpp
+++ b/Engine/Source/Editor/Sequencer/Private/SSequencer.cpp
@@ -811,7 +811,7 @@ void SSequencer::OnAssetsDropped( const FAssetDragDropOp& DragDropOp )
 		}
 	}
 
-	const TSet< TSharedRef<const FSequencerDisplayNode> >* SelectedNodes = Sequencer.Pin()->GetSelection()->GetSelectedOutlinerNodes();
+	const TSet< TSharedRef<FSequencerDisplayNode> >* SelectedNodes = Sequencer.Pin()->GetSelection()->GetSelectedOutlinerNodes();
 	FGuid TargetObjectGuid;
 	// if exactly one object node is selected, we have a target object guid
 	TSharedPtr<const FSequencerDisplayNode> DisplayNode;
@@ -976,7 +976,7 @@ FText SSequencer::GetShotSectionTitle(UMovieSceneSection* ShotSection) const
 
 void SSequencer::DeleteSelectedNodes()
 {
-	const TSet< TSharedRef<const FSequencerDisplayNode> >* SelectedNodes = Sequencer.Pin()->GetSelection()->GetSelectedOutlinerNodes();
+	const TSet< TSharedRef<FSequencerDisplayNode> >* SelectedNodes = Sequencer.Pin()->GetSelection()->GetSelectedOutlinerNodes();
 
 	if( SelectedNodes->Num() > 0 )
 	{
@@ -989,12 +989,53 @@ void SSequencer::DeleteSelectedNodes()
 			if( SelectedNode->IsVisible() )
 			{
 				// Delete everything in the entire node
-				SequencerRef.OnRequestNodeDeleted( SelectedNode );
+				TSharedRef<const FSequencerDisplayNode> NodeToBeDeleted = StaticCastSharedRef<const FSequencerDisplayNode>(SelectedNode);
+				SequencerRef.OnRequestNodeDeleted( NodeToBeDeleted );
 			}
 		}
 	}
 }
 
+void SSequencer::ExpandCollapseNode(TSharedRef<FSequencerDisplayNode> Node, bool bDescendants, bool bExpand)
+{
+	if (Node->IsExpanded() != bExpand)
+	{
+		Node->ToggleExpansion();
+	}
+
+	if (bDescendants)
+	{
+		for (auto ChildNode : Node->GetChildNodes())
+		{
+			ExpandCollapseNode(ChildNode, bDescendants, bExpand);
+		}
+	}
+}
+
+void SSequencer::ToggleExpandCollapseSelectedNodes(bool bDescendants)
+{
+	const TSet< TSharedRef<FSequencerDisplayNode> >* SelectedNodes = Sequencer.Pin()->GetSelection()->GetSelectedOutlinerNodes();
+
+	TSet< TSharedRef<FSequencerDisplayNode> > Nodes(*SelectedNodes);
+
+	if (Nodes.Num() == 0)
+	{
+		TSet<TSharedRef<FSequencerDisplayNode>> RootNodes(SequencerNodeTree->GetRootNodes());
+		Nodes = RootNodes;
+	}
+
+	if (Nodes.Num() > 0)
+	{
+		auto It = Nodes.CreateConstIterator();
+		bool bExpand = !(*It).Get().IsExpanded();
+
+		for (auto Node : Nodes)
+		{
+			ExpandCollapseNode(Node, bDescendants, bExpand);
+		}
+	}
+}
+
 EVisibility SSequencer::GetBreadcrumbTrailVisibility() const
 {
 	return Sequencer.Pin()->IsLevelEditorSequencer() ? EVisibility::Visible : EVisibility::Collapsed;
diff --git a/Engine/Source/Editor/Sequencer/Private/SSequencer.h b/Engine/Source/Editor/Sequencer/Private/SSequencer.h
--- a/Engine/Source/Editor/Sequencer/Private/SSequencer.h
+++ b/Engine/Source/Editor/Sequencer/Private/SSequencer.h
@@ -88,6 +88,10 @@ public:
 	/** Deletes selected nodes out of the sequencer node tree */
 	void DeleteSelectedNodes();
 
+	/** Expand or collapse selected nodes out of the sequencer node tree */
+	void ToggleExpandCollapseSelectedNodes(bool bDescendants = false);
+	void ExpandCollapseNode(TSharedRef<FSequencerDisplayNode> SelectedNode, bool bDescendants, bool bExpand);
+
 private:
 	/** Empty active timer to ensure Slate ticks during Sequencer playback */
 	EActiveTimerReturnType EnsureSlateTickDuringPlayback(double InCurrentTime, float InDeltaTime);	
diff --git a/Engine/Source/Editor/Sequencer/Private/Sequencer.cpp b/Engine/Source/Editor/Sequencer/Private/Sequencer.cpp
--- a/Engine/Source/Editor/Sequencer/Private/Sequencer.cpp
+++ b/Engine/Source/Editor/Sequencer/Private/Sequencer.cpp
@@ -149,7 +149,6 @@ FSequencer::~FSequencer()
 	GEditor->UnregisterForUndo( this );
 
 	DestroySpawnablesForAllMovieScenes();
-	
 	ObjectBindingManager.Reset();
 
 	TrackEditors.Empty();
@@ -1382,6 +1381,16 @@ void FSequencer::StepBackward()
 	OnStepBackward();
 }
 
+void FSequencer::ToggleExpandCollapseNodes()
+{
+	SequencerWidget->ToggleExpandCollapseSelectedNodes();
+}
+
+void FSequencer::ToggleExpandCollapseNodesAndDescendants()
+{
+	SequencerWidget->ToggleExpandCollapseSelectedNodes(true);
+}
+
 void FSequencer::SetKey()
 {
 	USelection* CurrentSelection = GEditor->GetSelectedActors();
@@ -1431,6 +1440,14 @@ void FSequencer::BindSequencerCommands()
 		Commands.StepBackward,
 		FExecuteAction::CreateSP( this, &FSequencer::StepBackward ) );
 
+	SequencerCommandBindings->MapAction(
+		Commands.ToggleExpandCollapseNodes,
+		FExecuteAction::CreateSP(this, &FSequencer::ToggleExpandCollapseNodes));
+
+	SequencerCommandBindings->MapAction(
+		Commands.ToggleExpandCollapseNodesAndDescendants,
+		FExecuteAction::CreateSP(this, &FSequencer::ToggleExpandCollapseNodesAndDescendants));
+
 	SequencerCommandBindings->MapAction(
 		Commands.SetKey,
 		FExecuteAction::CreateSP( this, &FSequencer::SetKey ) );
diff --git a/Engine/Source/Editor/Sequencer/Private/Sequencer.h b/Engine/Source/Editor/Sequencer/Private/Sequencer.h
--- a/Engine/Source/Editor/Sequencer/Private/Sequencer.h
+++ b/Engine/Source/Editor/Sequencer/Private/Sequencer.h
@@ -322,6 +322,12 @@ protected:
 	void StepForward();
 	void StepBackward();
 
+	/** Expand or collapse selected nodes */
+	void ToggleExpandCollapseNodes();
+
+	/** Expand or collapse selected nodes and descendants*/
+	void ToggleExpandCollapseNodesAndDescendants();
+
 	/** Manually sets a key for the selected objects at the current time */
 	void SetKey();
 
diff --git a/Engine/Source/Editor/Sequencer/Private/SequencerCommands.cpp b/Engine/Source/Editor/Sequencer/Private/SequencerCommands.cpp
--- a/Engine/Source/Editor/Sequencer/Private/SequencerCommands.cpp
+++ b/Engine/Source/Editor/Sequencer/Private/SequencerCommands.cpp
@@ -13,6 +13,9 @@ void FSequencerCommands::RegisterCommands()
 	UI_COMMAND( StepForward, "Step Forward", "Step the timeline forward.", EUserInterfaceActionType::Button, FInputChord(EKeys::Right) );
 	UI_COMMAND( StepBackward, "Step Backward", "Step the timeline backward", EUserInterfaceActionType::Button, FInputChord(EKeys::Left) );
 
+	UI_COMMAND( ToggleExpandCollapseNodes, "Toggle Expand/Collapse Nodes", "Toggle expand or collapse selected nodes.", EUserInterfaceActionType::Button, FInputChord(EKeys::O) );
+	UI_COMMAND( ToggleExpandCollapseNodesAndDescendants, "Toggle Expand/Collapse Nodes and Descendants", "Toggle expand or collapse selected nodes and their descendants.", EUserInterfaceActionType::Button, FInputChord(EModifierKey::Shift, EKeys::O) );
+
 	UI_COMMAND( SetKey, "Set Key", "Sets a key at the current time for the selected actor.", EUserInterfaceActionType::Button, FInputChord(EKeys::K) );
 
 	UI_COMMAND( ToggleAutoKeyEnabled, "Auto Key", "Enables and disables auto keying.", EUserInterfaceActionType::ToggleButton, FInputChord() );
diff --git a/Engine/Source/Editor/Sequencer/Private/SequencerCommands.h b/Engine/Source/Editor/Sequencer/Private/SequencerCommands.h
--- a/Engine/Source/Editor/Sequencer/Private/SequencerCommands.h
+++ b/Engine/Source/Editor/Sequencer/Private/SequencerCommands.h
@@ -30,6 +30,12 @@ public:
 	/** Step backward */
 	TSharedPtr< FUICommandInfo > StepBackward;
 
+	/** Expand/collapse nodes */
+	TSharedPtr< FUICommandInfo > ToggleExpandCollapseNodes;
+
+	/** Expand/collapse nodes and descendants */
+	TSharedPtr< FUICommandInfo > ToggleExpandCollapseNodesAndDescendants;
+
 	/** Sets a key at the current time for the selected actor */
 	TSharedPtr< FUICommandInfo > SetKey;
 
diff --git a/Engine/Source/Editor/Sequencer/Private/SequencerSelection.cpp b/Engine/Source/Editor/Sequencer/Private/SequencerSelection.cpp
--- a/Engine/Source/Editor/Sequencer/Private/SequencerSelection.cpp
+++ b/Engine/Source/Editor/Sequencer/Private/SequencerSelection.cpp
@@ -19,7 +19,7 @@ const TSet<TWeakObjectPtr<UMovieSceneSection>>* FSequencerSelection::GetSelected
 	return &SelectedSections;
 }
 
-const TSet<TSharedRef<const FSequencerDisplayNode>>* FSequencerSelection::GetSelectedOutlinerNodes() const
+const TSet<TSharedRef<FSequencerDisplayNode>>* FSequencerSelection::GetSelectedOutlinerNodes() const
 {
 	return &SelectedOutlinerNodes;
 }
@@ -58,7 +58,7 @@ void FSequencerSelection::AddToSelection(UMovieSceneSection* Section)
 	OnSectionSelectionChanged.Broadcast();
 }
 
-void FSequencerSelection::AddToSelection(TSharedRef<const FSequencerDisplayNode> OutlinerNode)
+void FSequencerSelection::AddToSelection(TSharedRef<FSequencerDisplayNode> OutlinerNode)
 {
 	SelectedOutlinerNodes.Add(OutlinerNode);
 	ActiveSelection = EActiveSelection::OutlinerNode;
diff --git a/Engine/Source/Editor/Sequencer/Private/SequencerSelection.h b/Engine/Source/Editor/Sequencer/Private/SequencerSelection.h
--- a/Engine/Source/Editor/Sequencer/Private/SequencerSelection.h
+++ b/Engine/Source/Editor/Sequencer/Private/SequencerSelection.h
@@ -33,7 +33,7 @@ public:
 	const TSet<TWeakObjectPtr<UMovieSceneSection>>* GetSelectedSections() const;
 
 	/** Gets a set of the selected outliner nodes. */
-	const TSet<TSharedRef<const FSequencerDisplayNode>>* GetSelectedOutlinerNodes() const;
+	const TSet<TSharedRef<FSequencerDisplayNode>>* GetSelectedOutlinerNodes() const;
 
 	/** 
 	 * Gets the currently active selection.  This is used to determine which selection actions
@@ -48,7 +48,7 @@ public:
 	void AddToSelection(UMovieSceneSection* Section);
 
 	/** Adds an outliner node to the selection */
-	void AddToSelection(TSharedRef<const FSequencerDisplayNode> OutlinerNode);
+	void AddToSelection(TSharedRef<FSequencerDisplayNode> OutlinerNode);
 
 	/** Returns whether or not the key is selected. */
 	bool IsSelected(FSelectedKey Key) const;
@@ -83,7 +83,7 @@ public:
 private:
 	TSet<FSelectedKey> SelectedKeys;
 	TSet<TWeakObjectPtr<UMovieSceneSection>> SelectedSections;
-	TSet<TSharedRef<const FSequencerDisplayNode>> SelectedOutlinerNodes;
+	TSet<TSharedRef<FSequencerDisplayNode>> SelectedOutlinerNodes;
 
 	FOnSelectionChanged OnKeySelectionChanged;
 	FOnSelectionChanged OnSectionSelectionChanged;
