Sequencer - Keep animation outliner selection consistent across undos.

diff --git a/Engine/Source/Editor/Sequencer/Private/SSequencer.cpp b/Engine/Source/Editor/Sequencer/Private/SSequencer.cpp
--- a/Engine/Source/Editor/Sequencer/Private/SSequencer.cpp
+++ b/Engine/Source/Editor/Sequencer/Private/SSequencer.cpp
@@ -613,8 +613,30 @@ EActiveTimerReturnType SSequencer::EnsureSlateTickDuringPlayback(double InCurren
 	return EActiveTimerReturnType::Stop;
 }
 
+void RestoreSelectionState(const TArray<TSharedRef<FSequencerDisplayNode>>& DisplayNodes, TSet<FString>& SelectedPathNames, FSequencerSelection* SequencerSelection)
+{
+	for (TSharedRef<FSequencerDisplayNode> DisplayNode : DisplayNodes)
+	{
+		if (SelectedPathNames.Contains(DisplayNode->GetPathName()))
+		{
+			SequencerSelection->AddToSelection(DisplayNode);
+		}
+		for (TSharedRef<FSequencerDisplayNode> ChildDisplayNode : DisplayNode->GetChildNodes())
+		{
+			RestoreSelectionState(DisplayNode->GetChildNodes(), SelectedPathNames, SequencerSelection);
+		}
+	}
+}
+
 void SSequencer::UpdateLayoutTree()
-{	
+{
+	// Cache the selected path names so selection can be restored after the update.
+	TSet<FString> SelectedPathNames;
+	for (TSharedRef<const FSequencerDisplayNode> SelectedDisplayNode : Sequencer.Pin()->GetSelection()->GetSelectedOutlinerNodes()->Array())
+	{
+		SelectedPathNames.Add(SelectedDisplayNode->GetPathName());
+	}
+
 	// Update the node tree
 	SequencerNodeTree->Update();
 
@@ -623,6 +645,9 @@ void SSequencer::UpdateLayoutTree()
 
 	CurveEditor->SetSequencerNodeTree(SequencerNodeTree);
 
+	// Restore the selection state.
+	RestoreSelectionState(SequencerNodeTree->GetRootNodes(), SelectedPathNames, Sequencer.Pin()->GetSelection());
+
 	SequencerNodeTree->UpdateCachedVisibilityBasedOnShotFiltersChanged();
 }
 
diff --git a/Engine/Source/Editor/UMGEditor/Private/TabFactory/AnimationTabSummoner.cpp b/Engine/Source/Editor/UMGEditor/Private/TabFactory/AnimationTabSummoner.cpp
--- a/Engine/Source/Editor/UMGEditor/Private/TabFactory/AnimationTabSummoner.cpp
+++ b/Engine/Source/Editor/UMGEditor/Private/TabFactory/AnimationTabSummoner.cpp
@@ -372,7 +372,6 @@ private:
 
 	void OnSelectionChanged( TSharedPtr<FWidgetAnimationListItem> InSelectedItem, ESelectInfo::Type SelectionInfo )
 	{
-		UWidgetBlueprint* WidgetBlueprint = BlueprintEditor.Pin()->GetWidgetBlueprintObj();
 		UWidgetAnimation* WidgetAnimation;
 		if( InSelectedItem.IsValid() )
 		{
@@ -382,7 +381,12 @@ private:
 		{
 			WidgetAnimation = UWidgetAnimation::GetNullAnimation();
 		}
-		BlueprintEditor.Pin()->ChangeViewedAnimation( *WidgetAnimation );
+
+		const UWidgetAnimation* CurrentWidgetAnimation = BlueprintEditor.Pin()->RefreshCurrentAnimation();
+		if (WidgetAnimation != CurrentWidgetAnimation)
+		{
+			BlueprintEditor.Pin()->ChangeViewedAnimation( *WidgetAnimation );
+		}
 	}
 
 	TSharedPtr<SWidget> OnContextMenuOpening() const
