Sequencer: Don't clear key selection if invoking right click menu.

diff --git a/Engine/Source/Editor/Sequencer/Private/SSection.cpp b/Engine/Source/Editor/Sequencer/Private/SSection.cpp
--- a/Engine/Source/Editor/Sequencer/Private/SSection.cpp
+++ b/Engine/Source/Editor/Sequencer/Private/SSection.cpp
@@ -726,7 +726,10 @@ void SSection::HandleKeySelection( const FSelectedKey& Key, const FPointerEvent&
 		// we are selecting due to drag and the key being dragged is not selected or control is not down
 		bool bShouldClearSelectionDueToDrag =  bSelectDueToDrag ? !GetSequencer().GetSelection()->IsSelected( Key ) : true;
 		
-		if( !MouseEvent.IsControlDown() && bShouldClearSelectionDueToDrag )
+		// Keep key selection if right clicking to bring up a menu and the current key is selected
+		bool bKeepKeySelection = MouseEvent.GetEffectingButton() == EKeys::RightMouseButton && GetSequencer().GetSelection()->IsSelected( Key );
+
+		if( (!MouseEvent.IsControlDown() && bShouldClearSelectionDueToDrag) && !bKeepKeySelection )
 		{
 			GetSequencer().GetSelection()->EmptySelectedKeys();
 		}
