Sequencer: Synchronize sequencer display node selection so that nodes get selected when their corresponding runtime objects are selected in the world/level. Also, clear actor selection when clicking on empty space in the sequencer outliner view.

diff --git a/Engine/Source/Editor/Sequencer/Private/SAnimationOutlinerView.cpp b/Engine/Source/Editor/Sequencer/Private/SAnimationOutlinerView.cpp
--- a/Engine/Source/Editor/Sequencer/Private/SAnimationOutlinerView.cpp
+++ b/Engine/Source/Editor/Sequencer/Private/SAnimationOutlinerView.cpp
@@ -460,12 +460,19 @@ void SAnimationOutlinerView::OnSelectionChanged( TSharedPtr<FSequencerDisplayNod
 			const bool bDeselectBSP = true;
 			const bool bWarnAboutTooManyActors = false;
 
+			const FScopedTransaction Transaction(NSLOCTEXT("UnrealEd", "ClickingOnActors", "Clicking on Actors"));
+
 			// Clear selection
 			GEditor->SelectNone(bNotifySelectionChanged,bDeselectBSP,bWarnAboutTooManyActors);
 			GEditor->GetSelectedActors()->BeginBatchSelectOperation();
 
 			// Select each actor
 			bool bActorSelected = false;
+
+			// Mark that the user is selecting so that the UI doesn't respond to the selection changes in the following block
+			TSharedRef<SSequencer> SequencerWidget = StaticCastSharedRef<SSequencer>(Sequencer->GetSequencerWidget());
+			SequencerWidget->SetUserIsSelecting(true);
+			
 			for( int32 ObjectIndex = 0; ObjectIndex < RuntimeObjects.Num(); ++ObjectIndex )
 			{
 				AActor* Actor = Cast<AActor>( RuntimeObjects[ObjectIndex] );
@@ -485,6 +492,9 @@ void SAnimationOutlinerView::OnSelectionChanged( TSharedPtr<FSequencerDisplayNod
 			{
 				GEditor->NoteSelectionChange();
 			}
+
+			// Unlock the selection so that the sequencer widget can now respond to selection changes in the level
+			SequencerWidget->SetUserIsSelecting(false);
 		}
 	}
 }
diff --git a/Engine/Source/Editor/Sequencer/Private/SSequencer.cpp b/Engine/Source/Editor/Sequencer/Private/SSequencer.cpp
--- a/Engine/Source/Editor/Sequencer/Private/SSequencer.cpp
+++ b/Engine/Source/Editor/Sequencer/Private/SSequencer.cpp
@@ -33,6 +33,47 @@
 #define LOCTEXT_NAMESPACE "Sequencer"
 
 
+class SSequencerScrollBox : public SScrollBox
+{
+public:
+	void Construct( const FArguments& InArgs, TSharedRef<FSequencer> InSequencer);
+
+	virtual FReply OnMouseButtonDown(const FGeometry& MyGeometry, const FPointerEvent& MouseEvent) override;
+
+	/** The sequencer which owns this widget. */
+	TWeakPtr<FSequencer> Sequencer;
+};
+
+void SSequencerScrollBox::Construct(const FArguments& InArgs, TSharedRef<FSequencer> InSequencer)
+{
+	Sequencer = InSequencer;
+	SScrollBox::Construct(InArgs);
+}
+
+FReply SSequencerScrollBox::OnMouseButtonDown(const FGeometry& MyGeometry, const FPointerEvent& MouseEvent)
+{
+	// Clear the selection if no sequencer display nodes were clicked on
+	if (MouseEvent.GetEffectingButton() == EKeys::LeftMouseButton)
+	{
+		if (Sequencer.Pin()->GetSelection()->GetSelectedOutlinerNodes()->Num() != 0)
+		{
+			Sequencer.Pin()->GetSelection()->EmptySelectedOutlinerNodes();
+
+			if (Sequencer.Pin()->IsLevelEditorSequencer())
+			{
+				const bool bNotifySelectionChanged = false;
+				const bool bDeselectBSP = true;
+				const bool bWarnAboutTooManyActors = false;
+
+				const FScopedTransaction Transaction(NSLOCTEXT("UnrealEd", "ClickingOnActors", "Clicking on Actors"));
+				GEditor->SelectNone(bNotifySelectionChanged, bDeselectBSP, bWarnAboutTooManyActors);
+				return FReply::Handled();
+			}
+		}
+	}
+	return FReply::Unhandled();
+}
+
 /**
  * The shot filter overlay displays the overlay needed to filter out widgets based on
  * which shots are actively in use.
@@ -147,6 +188,9 @@ void SSequencer::Construct( const FArguments& InArgs, TSharedRef< class FSequenc
 {
 	Sequencer = InSequencer;
 	bIsActiveTimerRegistered = false;
+	bUserIsSelecting = false;
+
+	USelection::SelectionChangedEvent.AddSP(this, &SSequencer::OnActorSelectionChanged);
 
 	// Create a node tree which contains a tree of movie scene data to display in the sequence
 	SequencerNodeTree = MakeShareable( new FSequencerNodeTree( InSequencer.Get() ) );
@@ -251,7 +295,7 @@ void SSequencer::Construct( const FArguments& InArgs, TSharedRef< class FSequenc
 
 							+ SOverlay::Slot()
 							[
-								SNew( SScrollBox )
+								SNew( SSequencerScrollBox, Sequencer.Pin().ToSharedRef() )
 								.ExternalScrollbar(ScrollBar)
 
 								+ SScrollBox::Slot()
@@ -967,6 +1011,42 @@ void SSequencer::OnActorsDropped( FActorDragDropGraphEdOp& DragDropOp )
 	Sequencer.Pin()->OnActorsDropped( DragDropOp.Actors );
 }
 
+void SSequencer::OnActorSelectionChanged(UObject* obj)
+{
+	if (!Sequencer.Pin()->IsLevelEditorSequencer())
+	{
+		return;
+	}
+
+	// If the user is selecting within the sequencer, ignore
+	if (bUserIsSelecting)
+	{
+		return;
+	}
+
+	TSet<TSharedRef<FSequencerDisplayNode>> RootNodes(SequencerNodeTree->GetRootNodes());
+
+	// Clear selection
+	Sequencer.Pin()->GetSelection()->EmptySelectedOutlinerNodes();
+
+	// Select the nodes that have runtime objects that are selected in the level
+	for (auto Node : RootNodes)
+	{
+		TSharedRef<FObjectBindingNode> ObjectBindingNode = StaticCastSharedRef<FObjectBindingNode>(Node);
+		TArray<UObject*> RuntimeObjects;
+		Sequencer.Pin()->GetRuntimeObjects( Sequencer.Pin()->GetFocusedMovieSceneInstance(), ObjectBindingNode->GetObjectBinding(), RuntimeObjects );
+		
+		for (int32 RuntimeIndex = 0; RuntimeIndex < RuntimeObjects.Num(); ++RuntimeIndex )
+		{
+			if (GEditor->GetSelectedActors()->IsSelected(RuntimeObjects[RuntimeIndex]))
+			{
+				Sequencer.Pin()->GetSelection()->AddToSelection(Node);
+				break;
+			}
+		}
+	}
+}
+
 void SSequencer::OnCrumbClicked(const FSequencerBreadcrumb& Item)
 {
 	if (Item.BreadcrumbType != FSequencerBreadcrumb::ShotType)
diff --git a/Engine/Source/Editor/Sequencer/Private/SSequencer.h b/Engine/Source/Editor/Sequencer/Private/SSequencer.h
--- a/Engine/Source/Editor/Sequencer/Private/SSequencer.h
+++ b/Engine/Source/Editor/Sequencer/Private/SSequencer.h
@@ -92,6 +92,10 @@ public:
 	void ToggleExpandCollapseSelectedNodes(bool bDescendants = false);
 	void ExpandCollapseNode(TSharedRef<FSequencerDisplayNode> SelectedNode, bool bDescendants, bool bExpand);
 
+	/** Whether the user is selecting. Ignore selection changes from the level when the user is selecting. */
+	void SetUserIsSelecting(bool bUserIsSelectingIn) { bUserIsSelecting = bUserIsSelectingIn; }
+	bool UserIsSelecting() { return bUserIsSelecting; }
+
 private:
 	/** Empty active timer to ensure Slate ticks during Sequencer playback */
 	EActiveTimerReturnType EnsureSlateTickDuringPlayback(double InCurrentTime, float InDeltaTime);	
@@ -190,6 +194,12 @@ private:
 	 */
 	void OnActorsDropped( class FActorDragDropGraphEdOp& DragDropOp ); 
 	
+	/**
+	* Delegate used when actor selection changes in the level
+	*
+	*/	
+	void OnActorSelectionChanged( UObject* obj );
+
 	/** Called when a breadcrumb is clicked on in the sequencer */
 	void OnCrumbClicked(const FSequencerBreadcrumb& Item);
 
@@ -226,6 +236,8 @@ private:
 	float ColumnFillCoefficients[2];
 	/** Whether the active timer is currently registered */
 	bool bIsActiveTimerRegistered;
+	/** Whether the user is selecting. Ignore selection changes from the level when the user is selecting. */
+	bool bUserIsSelecting;
 
 	FOnGetAddMenuContent OnGetAddMenuContent;
 };
