Fixed sequencer track lanes not linking their heights correctly

diff --git a/Engine/Source/Editor/Sequencer/Private/SequencerTrackLaneFactory.cpp b/Engine/Source/Editor/Sequencer/Private/SequencerTrackLaneFactory.cpp
--- a/Engine/Source/Editor/Sequencer/Private/SequencerTrackLaneFactory.cpp
+++ b/Engine/Source/Editor/Sequencer/Private/SequencerTrackLaneFactory.cpp
@@ -23,6 +23,7 @@ struct FWidgetSizeLinker
 				if (MinSize.IsSet())
 				{
 					MinSize->X = FMath::Min(MinSize->X, DesiredSize.X);
+					MinSize->Y = FMath::Min(MinSize->Y, DesiredSize.Y);
 				}
 				else
 				{
@@ -47,6 +48,7 @@ struct FWidgetSizeLinker
 				if (MaxSize.IsSet())
 				{
 					MaxSize->X = FMath::Max(MaxSize->X, DesiredSize.X);
+					MaxSize->Y = FMath::Max(MaxSize->Y, DesiredSize.Y);
 				}
 				else
 				{
