Sequencer: Fix crash on shutdown when receiving selection changed.

diff --git a/Engine/Source/Editor/Sequencer/Private/SSequencer.cpp b/Engine/Source/Editor/Sequencer/Private/SSequencer.cpp
--- a/Engine/Source/Editor/Sequencer/Private/SSequencer.cpp
+++ b/Engine/Source/Editor/Sequencer/Private/SSequencer.cpp
@@ -631,6 +631,7 @@ TSharedRef<SWidget> SSequencer::MakeTransportControls()
 
 SSequencer::~SSequencer()
 {
+	USelection::SelectionChangedEvent.RemoveAll(this);
 }
 
 void SSequencer::RegisterActiveTimerForPlayback()
@@ -1013,6 +1014,11 @@ void SSequencer::OnActorsDropped( FActorDragDropGraphEdOp& DragDropOp )
 
 void SSequencer::OnActorSelectionChanged(UObject* obj)
 {
+	if (!Sequencer.IsValid())
+	{
+		return;
+	}
+
 	if (!Sequencer.Pin()->IsLevelEditorSequencer())
 	{
 		return;
