Fix Sequencer spawnables being treated as possessables

diff --git a/Engine/Source/Editor/Sequencer/Private/Sequencer.cpp b/Engine/Source/Editor/Sequencer/Private/Sequencer.cpp
--- a/Engine/Source/Editor/Sequencer/Private/Sequencer.cpp
+++ b/Engine/Source/Editor/Sequencer/Private/Sequencer.cpp
@@ -579,51 +579,55 @@ FGuid FSequencer::GetHandleToObject( UObject* Object )
 
 	FGuid ObjectGuid = ObjectBindingManager->FindGuidForObject( *FocusedMovieScene, *Object );
 
-	// Check here for spawnable otherwise spawnables get recreated as possessables, which doesn't make sense
-	if (ObjectGuid.IsValid() && !FocusedMovieScene->FindSpawnable(ObjectGuid))
+	// Check here for spawnable otherwise spawnables get recreated as possessables, which doesn't make sens
+	FMovieSceneSpawnable* Spawnable = FocusedMovieScene->FindSpawnable(ObjectGuid);
+	if( !Spawnable )
 	{
-		// Make sure that the possessable is still valid, if it's not remove the binding so new one 
-		// can be created.  This can happen due to undo.
-		FMovieScenePossessable* Possessable = FocusedMovieScene->FindPossessable(ObjectGuid);
-		if (Possessable == nullptr )
+		if(ObjectGuid.IsValid())
 		{
-			ObjectBindingManager->UnbindPossessableObjects(ObjectGuid);
-			ObjectGuid.Invalidate();
+			// Make sure that the possessable is still valid, if it's not remove the binding so new one 
+			// can be created.  This can happen due to undo.
+			FMovieScenePossessable* Possessable = FocusedMovieScene->FindPossessable(ObjectGuid);
+			if(Possessable == nullptr)
+			{
+				ObjectBindingManager->UnbindPossessableObjects(ObjectGuid);
+				ObjectGuid.Invalidate();
+			}
 		}
-	}
-	
-	bool bPossessableAdded = false;
-	
-	// If the object guid was not found attempt to add it
-	// Note: Only possessed actors can be added like this
-	if( !ObjectGuid.IsValid() && ObjectBindingManager->CanPossessObject( *Object ) )
-	{
-		// @todo sequencer: Undo doesn't seem to be working at all
-		const FScopedTransaction Transaction( LOCTEXT("UndoPossessingObject", "Possess Object with MovieScene") );
-		
-		// Possess the object!
+
+		bool bPossessableAdded = false;
+
+		// If the object guid was not found attempt to add it
+		// Note: Only possessed actors can be added like this
+		if(!ObjectGuid.IsValid() && ObjectBindingManager->CanPossessObject(*Object))
 		{
-			// Create a new possessable
-			FocusedMovieScene->Modify();
+			// @todo sequencer: Undo doesn't seem to be working at all
+			const FScopedTransaction Transaction(LOCTEXT("UndoPossessingObject", "Possess Object with MovieScene"));
 
-			ObjectGuid = FocusedMovieScene->AddPossessable( Object->GetName(), Object->GetClass() );
-			
-			if ( IsShotFilteringOn() )
+			// Possess the object!
 			{
-				AddUnfilterableObject(ObjectGuid);
+				// Create a new possessable
+				FocusedMovieScene->Modify();
+
+				ObjectGuid = FocusedMovieScene->AddPossessable(Object->GetName(), Object->GetClass());
+
+				if(IsShotFilteringOn())
+				{
+					AddUnfilterableObject(ObjectGuid);
+				}
+
+				ObjectBindingManager->BindPossessableObject(ObjectGuid, *Object);
+
+				bPossessableAdded = true;
 			}
-			
-			ObjectBindingManager->BindPossessableObject( ObjectGuid, *Object );
-			
-			bPossessableAdded = true;
 		}
-	}
-	
-	if( bPossessableAdded )
-	{
-		SpawnOrDestroyPuppetObjects( GetFocusedMovieSceneInstance() );
-			
-		NotifyMovieSceneDataChanged();
+
+		if(bPossessableAdded)
+		{
+			SpawnOrDestroyPuppetObjects(GetFocusedMovieSceneInstance());
+
+			NotifyMovieSceneDataChanged();
+		}
 	}
 	
 	return ObjectGuid;
