Added auto scrolling to the sequencer timeline view
This addresses UE-17006 - Clicking and dragging on the scrub bar allows you to scroll time out of view

diff --git a/Engine/Source/Editor/Sequencer/Private/SSequencer.cpp b/Engine/Source/Editor/Sequencer/Private/SSequencer.cpp
--- a/Engine/Source/Editor/Sequencer/Private/SSequencer.cpp
+++ b/Engine/Source/Editor/Sequencer/Private/SSequencer.cpp
@@ -203,6 +203,8 @@ void SSequencer::Construct( const FArguments& InArgs, TSharedRef< class FSequenc
 	TimeSliderArgs.ViewRange = InArgs._ViewRange;
 	TimeSliderArgs.OnViewRangeChanged = InArgs._OnViewRangeChanged;
 	TimeSliderArgs.ScrubPosition = InArgs._ScrubPosition;
+	TimeSliderArgs.OnBeginScrubberMovement = InArgs._OnBeginScrubbing;
+	TimeSliderArgs.OnEndScrubberMovement = InArgs._OnEndScrubbing;
 	TimeSliderArgs.OnScrubPositionChanged = InArgs._OnScrubPositionChanged;
 
 	TSharedRef<FSequencerTimeSliderController> TimeSliderController( new FSequencerTimeSliderController( TimeSliderArgs ) );
diff --git a/Engine/Source/Editor/Sequencer/Private/SSequencer.h b/Engine/Source/Editor/Sequencer/Private/SSequencer.h
--- a/Engine/Source/Editor/Sequencer/Private/SSequencer.h
+++ b/Engine/Source/Editor/Sequencer/Private/SSequencer.h
@@ -58,6 +58,10 @@ public:
 		SLATE_ATTRIBUTE( float, ScrubPosition )
 		/** Called when the user changes the view range */
 		SLATE_EVENT( FOnViewRangeChanged, OnViewRangeChanged )
+		/** Called when the user has begun scrubbing */
+		SLATE_EVENT( FSimpleDelegate, OnBeginScrubbing )
+		/** Called when the user has finished scrubbing */
+		SLATE_EVENT( FSimpleDelegate, OnEndScrubbing )
 		/** Called when the user changes the scrub position */
 		SLATE_EVENT( FOnScrubPositionChanged, OnScrubPositionChanged )
 		/** Called to populate the add combo button in the toolbar. */
diff --git a/Engine/Source/Editor/Sequencer/Private/Sequencer.cpp b/Engine/Source/Editor/Sequencer/Private/Sequencer.cpp
--- a/Engine/Source/Editor/Sequencer/Private/Sequencer.cpp
+++ b/Engine/Source/Editor/Sequencer/Private/Sequencer.cpp
@@ -94,6 +94,8 @@ void FSequencer::InitSequencer( const FSequencerInitParams& InitParams, const TA
 		SequencerWidget = SNew( SSequencer, SharedThis( this ) )
 			.ViewRange( this, &FSequencer::GetViewRange )
 			.ScrubPosition( this, &FSequencer::OnGetScrubPosition )
+			.OnBeginScrubbing( this, &FSequencer::OnBeginScrubbing )
+			.OnEndScrubbing( this, &FSequencer::OnEndScrubbing )
 			.OnScrubPositionChanged( this, &FSequencer::OnScrubPositionChanged )
 			.OnViewRangeChanged( this, &FSequencer::OnViewRangeChanged )
 			.OnGetAddMenuContent(InitParams.ViewParams.OnGetAddMenuContent);
@@ -198,6 +200,15 @@ void FSequencer::Tick(float InDeltaTime)
 		SequencerWidget->UpdateLayoutTree();
 		bNeedTreeRefresh = false;
 	}
+	
+	// Animate the autoscroll offset if it's set
+	if (AutoscrollOffset.IsSet())
+	{
+		static const float AutoScrollFactor = 0.1f;
+		float Offset = AutoscrollOffset.GetValue() * AutoScrollFactor;
+		OnViewRangeChanged(TRange<float>(TargetViewRange.GetLowerBoundValue() + Offset, TargetViewRange.GetUpperBoundValue() + Offset), EViewRangeInterpolation::Immediate);
+		SetGlobalTime(GetGlobalTime() + Offset);
+	}
 
 	float NewTime = GetGlobalTime() + InDeltaTime;
 	if (PlaybackState == EMovieScenePlayerStatus::Playing ||
@@ -560,18 +571,56 @@ void FSequencer::SetGlobalTime( float NewTime )
 {
 	float LastTime = ScrubPosition;
 
+	if (true/*bAutoScroll*/ && PlaybackState != EMovieScenePlayerStatus::Scrubbing)
+	{
+		float RangeOffset = CalculateAutoscrollEncroachment(NewTime).Get(0.f);
+			
+		// When not scrubbing, we auto scroll the view range immediately
+		if (RangeOffset != 0.f)
+		{
+			OnViewRangeChanged(
+				TRange<float>(TargetViewRange.GetLowerBoundValue() + RangeOffset, TargetViewRange.GetUpperBoundValue() + RangeOffset),
+				EViewRangeInterpolation::Immediate
+				);
+		}
+	}
+
 	// Update the position
 	ScrubPosition = NewTime;
 
 	RootMovieSceneInstance->Update( ScrubPosition, LastTime, *this );
 }
 
+TOptional<float> FSequencer::CalculateAutoscrollEncroachment(float NewTime) const
+{
+	enum class EDirection { Positive, Negative };
+	const EDirection Movement = NewTime - ScrubPosition >= 0 ? EDirection::Positive : EDirection::Negative;
+
+	const TRange<float> CurrentRange = GetViewRange();
+	const float RangeMin = CurrentRange.GetLowerBoundValue(), RangeMax = CurrentRange.GetUpperBoundValue();
+	const float AutoScrollThreshold = (RangeMax - RangeMin) * 0.1f;
+
+	if (Movement == EDirection::Negative && NewTime < RangeMin + AutoScrollThreshold)
+	{
+		// Scrolling backwards in time, and have hit the threshold
+		return NewTime - (RangeMin + AutoScrollThreshold);
+	}
+	else if (Movement == EDirection::Positive && ScrubPosition > RangeMax - AutoScrollThreshold)
+	{
+		// Scrolling forwards in time, and have hit the threshold
+		return NewTime - (RangeMax - AutoScrollThreshold);
+	}
+	else
+	{
+		return TOptional<float>();
+	}
+}
+
 void FSequencer::SetPerspectiveViewportPossessionEnabled(bool bEnabled)
 {
 	bPerspectiveViewportPossessionEnabled = bEnabled;
 }
 
-
 FGuid FSequencer::GetHandleToObject( UObject* Object )
 {
 	TSharedRef<FMovieSceneInstance> FocusedMovieSceneInstance = GetFocusedMovieSceneInstance();
@@ -935,20 +984,37 @@ void FSequencer::OnViewRangeChanged( TRange<float> NewViewRange, EViewRangeInter
 
 void FSequencer::OnScrubPositionChanged( float NewScrubPosition, bool bScrubbing )
 {
-	if (bScrubbing)
-	{
-		PlaybackState =
-			PlaybackState == EMovieScenePlayerStatus::BeginningScrubbing || PlaybackState == EMovieScenePlayerStatus::Scrubbing ?
-			EMovieScenePlayerStatus::Scrubbing : EMovieScenePlayerStatus::BeginningScrubbing;
-	}
-	else
+	if (PlaybackState == EMovieScenePlayerStatus::Scrubbing)
 	{
-		PlaybackState = EMovieScenePlayerStatus::Stopped;
+		if (!bScrubbing)
+		{
+			OnEndScrubbing();
+		}
+		else if (true/*bAutoScroll*/)
+		{
+			// When scrubbing, we animate auto-scrolled scrub position in Tick()
+			AutoscrollOffset = CalculateAutoscrollEncroachment(NewScrubPosition);
+			if (AutoscrollOffset.IsSet())
+			{
+				return;
+			}
+		}
 	}
 
 	SetGlobalTime( NewScrubPosition );
 }
 
+void FSequencer::OnBeginScrubbing()
+{
+	PlaybackState = EMovieScenePlayerStatus::Scrubbing;
+}
+
+void FSequencer::OnEndScrubbing()
+{
+	PlaybackState = EMovieScenePlayerStatus::Stopped;
+	AutoscrollOffset.Reset();
+}
+
 void FSequencer::OnToggleAutoKey()
 {
 	bAllowAutoKey = !bAllowAutoKey;
diff --git a/Engine/Source/Editor/Sequencer/Private/Sequencer.h b/Engine/Source/Editor/Sequencer/Private/Sequencer.h
--- a/Engine/Source/Editor/Sequencer/Private/Sequencer.h
+++ b/Engine/Source/Editor/Sequencer/Private/Sequencer.h
@@ -296,6 +296,11 @@ protected:
 	 */
 	void OnViewRangeChanged( TRange<float> NewViewRange, EViewRangeInterpolation Interpolation = EViewRangeInterpolation::Animated );
 
+	/**
+	 * Calculates the amount of encroachment the specified time has into the autoscroll region, if any
+	 */
+	TOptional<float> CalculateAutoscrollEncroachment(float NewTime) const;
+
 	/**
 	 * Called when the scrub position is changed by the user
 	 * This will stop any playback from happening
@@ -304,6 +309,12 @@ protected:
 	 */
 	void OnScrubPositionChanged( float NewScrubPosition, bool bScrubbing );
 
+	/** Called when the user has begun scrubbing */
+	void OnBeginScrubbing();
+
+	/** Called when the user has finished scrubbing */
+	void OnEndScrubbing();
+
 	/**
 	 * Called when auto-key is toggled by a user
 	 *
@@ -399,6 +410,8 @@ private:
 	TRange<float> TargetViewRange;
 	/** The last time range that was viewed */
 	TRange<float> LastViewRange;
+	/** The amount of autoscroll pan offset that is currently being applied */
+	TOptional<float> AutoscrollOffset;
 
 	/** Zoom smoothing curves */
 	FCurveSequence ZoomAnimation;
diff --git a/Engine/Source/Runtime/MovieScene/Public/IMovieScenePlayer.h b/Engine/Source/Runtime/MovieScene/Public/IMovieScenePlayer.h
--- a/Engine/Source/Runtime/MovieScene/Public/IMovieScenePlayer.h
+++ b/Engine/Source/Runtime/MovieScene/Public/IMovieScenePlayer.h
@@ -9,7 +9,6 @@ namespace EMovieScenePlayerStatus
 		Stopped,
 		Playing,
 		Recording,
-		BeginningScrubbing,
 		Scrubbing,
 		MAX
 	};
