Added sequencer setting to control auto-scrolling

diff --git a/Engine/Source/Editor/EditorStyle/Private/SlateEditorStyle.cpp b/Engine/Source/Editor/EditorStyle/Private/SlateEditorStyle.cpp
--- a/Engine/Source/Editor/EditorStyle/Private/SlateEditorStyle.cpp
+++ b/Engine/Source/Editor/EditorStyle/Private/SlateEditorStyle.cpp
@@ -1592,6 +1592,8 @@ void FSlateEditorStyle::FStyle::SetupGeneralStyles()
 		Set( "Sequencer.KeyMark", new IMAGE_BRUSH("Sequencer/KeyMark", FVector2D(3,21), FLinearColor::White, ESlateBrushTileType::NoTile ) );
 		Set( "Sequencer.ToggleIsSnapEnabled", new IMAGE_BRUSH( "Icons/icon_Sequencer_ToggleSnap_40x", Icon40x40 ) );
 		Set( "Sequencer.ToggleIsSnapEnabled.Small", new IMAGE_BRUSH( "Icons/icon_Sequencer_ToggleSnap_40x", Icon16x16 ) );
+		Set( "Sequencer.ToggleAutoScroll", new IMAGE_BRUSH( "Icons/icon_Sequencer_ToggleAutoScroll_40x", Icon40x40 ) );
+		Set( "Sequencer.ToggleAutoScroll.Small", new IMAGE_BRUSH( "Icons/icon_Sequencer_ToggleAutoScroll_16x", Icon16x16 ) );
 
 		const FButtonStyle DetailsKeyButton = FButtonStyle(NoBorder)
 			.SetNormal( IMAGE_BRUSH("Sequencer/AddKey_Details", FVector2D(11,11) )  )
diff --git a/Engine/Source/Editor/Sequencer/Private/SSequencer.cpp b/Engine/Source/Editor/Sequencer/Private/SSequencer.cpp
--- a/Engine/Source/Editor/Sequencer/Private/SSequencer.cpp
+++ b/Engine/Source/Editor/Sequencer/Private/SSequencer.cpp
@@ -493,6 +493,10 @@ TSharedRef<SWidget> SSequencer::MakeToolBar()
 		{
 			ToolBarBuilder.AddToolBarButton( FSequencerCommands::Get().ToggleCleanView );
 		}
+
+		ToolBarBuilder.SetLabelVisibility(EVisibility::Collapsed);
+		ToolBarBuilder.AddToolBarButton( FSequencerCommands::Get().ToggleAutoScroll, NAME_None, TAttribute<FText>( FText::GetEmpty() ) );
+		ToolBarBuilder.SetLabelVisibility(EVisibility::Visible);
 	}
 	ToolBarBuilder.EndSection();
 
diff --git a/Engine/Source/Editor/Sequencer/Private/Sequencer.cpp b/Engine/Source/Editor/Sequencer/Private/Sequencer.cpp
--- a/Engine/Source/Editor/Sequencer/Private/Sequencer.cpp
+++ b/Engine/Source/Editor/Sequencer/Private/Sequencer.cpp
@@ -151,6 +151,7 @@ FSequencer::FSequencer()
 	: SequencerCommandBindings( new FUICommandList )
 	, TargetViewRange(0.f, 5.f)
 	, LastViewRange(0.f, 5.f)
+	, bAutoScrollEnabled(true)
 	, PlaybackState( EMovieScenePlayerStatus::Stopped )
 	, ScrubPosition( 0.0f )
 	, bLoopingEnabled( false )
@@ -159,7 +160,7 @@ FSequencer::FSequencer()
 	, bIsEditingWithinLevelEditor( false )
 	, bNeedTreeRefresh( false )
 {
-
+	
 }
 
 FSequencer::~FSequencer()
@@ -571,7 +572,7 @@ void FSequencer::SetGlobalTime( float NewTime )
 {
 	float LastTime = ScrubPosition;
 
-	if (true/*bAutoScroll*/ && PlaybackState != EMovieScenePlayerStatus::Scrubbing)
+	if (bAutoScrollEnabled && PlaybackState != EMovieScenePlayerStatus::Scrubbing)
 	{
 		float RangeOffset = CalculateAutoscrollEncroachment(NewTime).Get(0.f);
 			
@@ -990,7 +991,7 @@ void FSequencer::OnScrubPositionChanged( float NewScrubPosition, bool bScrubbing
 		{
 			OnEndScrubbing();
 		}
-		else if (true/*bAutoScroll*/)
+		else if (bAutoScrollEnabled)
 		{
 			// When scrubbing, we animate auto-scrolled scrub position in Tick()
 			AutoscrollOffset = CalculateAutoscrollEncroachment(NewScrubPosition);
@@ -1020,6 +1021,12 @@ void FSequencer::OnToggleAutoKey()
 	bAllowAutoKey = !bAllowAutoKey;
 }
 
+void FSequencer::OnToggleAutoScroll()
+{
+	bAutoScrollEnabled = !bAutoScrollEnabled;
+	GetMutableDefault<USequencerSettings>()->SetAutoScrollEnabled(bAutoScrollEnabled);
+}
+
 FGuid FSequencer::AddSpawnableForAssetOrClass( UObject* Object, UObject* CounterpartGamePreviewObject )
 {
 	FGuid NewSpawnableGuid;
@@ -1573,6 +1580,13 @@ void FSequencer::BindSequencerCommands()
 		FCanExecuteAction::CreateLambda( []{ return true; } ),
 		FIsActionChecked::CreateSP( this, &FSequencer::OnGetAllowAutoKey ) );
 
+	bAutoScrollEnabled = Settings->GetAutoScrollEnabled();
+	SequencerCommandBindings->MapAction(
+		Commands.ToggleAutoScroll,
+		FExecuteAction::CreateSP( this, &FSequencer::OnToggleAutoScroll ),
+		FCanExecuteAction::CreateLambda( []{ return true; } ),
+		FIsActionChecked::CreateSP( this, &FSequencer::GetAutoScrollEnabled ) );
+
 	SequencerCommandBindings->MapAction(
 		Commands.ToggleCleanView,
 		FExecuteAction::CreateLambda( [Settings]{ Settings->SetIsUsingCleanView( !Settings->GetIsUsingCleanView() ); } ),
diff --git a/Engine/Source/Editor/Sequencer/Private/Sequencer.h b/Engine/Source/Editor/Sequencer/Private/Sequencer.h
--- a/Engine/Source/Editor/Sequencer/Private/Sequencer.h
+++ b/Engine/Source/Editor/Sequencer/Private/Sequencer.h
@@ -322,6 +322,10 @@ protected:
 	 */
 	void OnToggleAutoKey();
 
+	/** Called to toggle auto-scroll on and off */
+	void OnToggleAutoScroll();
+	bool GetAutoScrollEnabled() const { return bAutoScrollEnabled; }
+
 	/** Called via UEditorEngine::GetActorRecordingStateEvent to check to see whether we need to record actor state */
 	void GetActorRecordingState( bool& bIsRecording /* In+Out */ ) const;
 	
@@ -412,6 +416,8 @@ private:
 	TRange<float> LastViewRange;
 	/** The amount of autoscroll pan offset that is currently being applied */
 	TOptional<float> AutoscrollOffset;
+	/** Whether or not we are allowing autoscroll */
+	bool bAutoScrollEnabled;
 
 	/** Zoom smoothing curves */
 	FCurveSequence ZoomAnimation;
diff --git a/Engine/Source/Editor/Sequencer/Private/SequencerCommands.cpp b/Engine/Source/Editor/Sequencer/Private/SequencerCommands.cpp
--- a/Engine/Source/Editor/Sequencer/Private/SequencerCommands.cpp
+++ b/Engine/Source/Editor/Sequencer/Private/SequencerCommands.cpp
@@ -20,6 +20,8 @@ void FSequencerCommands::RegisterCommands()
 
 	UI_COMMAND( ToggleAutoKeyEnabled, "Auto Key", "Enables and disables auto keying.", EUserInterfaceActionType::ToggleButton, FInputChord() );
 
+	UI_COMMAND( ToggleAutoScroll, "Auto Scroll", "Toggle auto-scroll: When enabled, automatically scrolls the sequencer view to keep the current time visible.", EUserInterfaceActionType::ToggleButton, FInputChord(EModifierKey::Shift, EKeys::S) );
+
 	UI_COMMAND( ToggleIsSnapEnabled, "Enable Snapping", "Enables and disables snapping.", EUserInterfaceActionType::ToggleButton, FInputChord() );
 
 	UI_COMMAND( ToggleSnapKeyTimesToInterval, "Snap to the interval", "Snap keys to the interval.", EUserInterfaceActionType::ToggleButton, FInputChord() );
diff --git a/Engine/Source/Editor/Sequencer/Private/SequencerCommands.h b/Engine/Source/Editor/Sequencer/Private/SequencerCommands.h
--- a/Engine/Source/Editor/Sequencer/Private/SequencerCommands.h
+++ b/Engine/Source/Editor/Sequencer/Private/SequencerCommands.h
@@ -65,6 +65,9 @@ public:
 
 	/** Toggles whether the "clean view" is enabled in the level editor. */
 	TSharedPtr< FUICommandInfo > ToggleCleanView;
+	
+	/** Turns auto scroll on and off. */
+	TSharedPtr< FUICommandInfo > ToggleAutoScroll;
 
 	/** Toggles whether or not the curve editor should be shown. */
 	TSharedPtr< FUICommandInfo > ToggleShowCurveEditor;
diff --git a/Engine/Source/Editor/Sequencer/Private/SequencerSettings.cpp b/Engine/Source/Editor/Sequencer/Private/SequencerSettings.cpp
--- a/Engine/Source/Editor/Sequencer/Private/SequencerSettings.cpp
+++ b/Engine/Source/Editor/Sequencer/Private/SequencerSettings.cpp
@@ -15,6 +15,7 @@ USequencerSettings::USequencerSettings( const FObjectInitializer& ObjectInitiali
 	CurveValueSnapInterval = 10.0f;
 	bSnapCurveValueToInterval = true;
 	bIsUsingCleanView = false;
+	bAutoScrollEnabled = true;
 	bShowCurveEditor = false;
 	bShowCurveEditorCurveToolTips = true;
 	CurveVisibility = ESequencerCurveVisibility::AllCurves;
@@ -160,6 +161,21 @@ void USequencerSettings::SetIsUsingCleanView(bool InbIsUsingCleanView )
 	}
 }
 
+bool USequencerSettings::GetAutoScrollEnabled() const
+{
+	return bAutoScrollEnabled;
+}
+
+void USequencerSettings::SetAutoScrollEnabled(bool bInAutoScrollEnabled)
+{
+	if (bAutoScrollEnabled != bInAutoScrollEnabled)
+	{
+		bAutoScrollEnabled = bInAutoScrollEnabled;
+		SaveConfig();
+	}
+}
+
+
 bool USequencerSettings::GetShowCurveEditor() const
 {
 	return bShowCurveEditor;
diff --git a/Engine/Source/Editor/Sequencer/Private/SequencerSettings.h b/Engine/Source/Editor/Sequencer/Private/SequencerSettings.h
--- a/Engine/Source/Editor/Sequencer/Private/SequencerSettings.h
+++ b/Engine/Source/Editor/Sequencer/Private/SequencerSettings.h
@@ -79,6 +79,11 @@ public:
 	/** Sets whether or not the 'Clean View' is enabled. In 'Clean View' mode only global tracks are displayed when no filter is applied. */
 	void SetIsUsingCleanView(bool InbIsUsingCleanView);
 
+	/** Gets whether or not auto-scroll is enabled. */
+	bool GetAutoScrollEnabled() const;
+	/** Sets whether or not auto-scroll is enabled. */
+	void SetAutoScrollEnabled(bool bInAutoScrollEnabled);
+
 	/** Gets whether or not the curve editor should be shown. */
 	bool GetShowCurveEditor() const;
 	/** Sets whether or not the curve editor should be shown. */
@@ -134,6 +139,9 @@ protected:
 	UPROPERTY( config )
 	bool bIsUsingCleanView;
 
+	UPROPERTY( config )
+	bool bAutoScrollEnabled;
+
 	UPROPERTY( config )
 	bool bShowCurveEditor;
 
