Sequencer: Bind to play movie scene node if it's not bound when getting runtime objects. This fixes the startup problem of opening an existing level and having the possessables move around.

diff --git a/Engine/Source/Editor/Sequencer/Private/SequencerActorBindingManager.cpp b/Engine/Source/Editor/Sequencer/Private/SequencerActorBindingManager.cpp
--- a/Engine/Source/Editor/Sequencer/Private/SequencerActorBindingManager.cpp
+++ b/Engine/Source/Editor/Sequencer/Private/SequencerActorBindingManager.cpp
@@ -291,10 +291,18 @@ void FSequencerActorBindingManager::GetRuntimeObjects( const TSharedRef<FMovieSc
 		OutRuntimeObjects.Reset();
 		OutRuntimeObjects.Add( FoundSpawnedObject );
 	}
-	else if( PlayMovieSceneNode.IsValid() )
+	else
 	{
 		// Possessable
-		OutRuntimeObjects = PlayMovieSceneNode->FindBoundObjects( ObjectGuid );
+		if (!PlayMovieSceneNode.IsValid())
+		{
+			BindToPlayMovieSceneNode(false);
+		}
+	
+		if (PlayMovieSceneNode.IsValid())
+		{
+			OutRuntimeObjects = PlayMovieSceneNode->FindBoundObjects( ObjectGuid );
+		}
 	}
 }
 
