Sequencer: Get actor name for outliner display.

diff --git a/Engine/Source/Editor/NiagaraEditor/Private/NiagaraSequencer.h b/Engine/Source/Editor/NiagaraEditor/Private/NiagaraSequencer.h
--- a/Engine/Source/Editor/NiagaraEditor/Private/NiagaraSequencer.h
+++ b/Engine/Source/Editor/NiagaraEditor/Private/NiagaraSequencer.h
@@ -235,6 +235,6 @@ public:
 	virtual void UnbindPossessableObjects(const FGuid& PossessableGuid) override	{};
 	virtual void GetRuntimeObjects(const TSharedRef<FMovieSceneInstance>& MovieSceneInstance, const FGuid& ObjectGuid, TArray<UObject*>& OutRuntimeObjects) const override	{}
 
-	virtual bool TryGetObjectBindingDisplayName(const FGuid& ObjectGuid, FText& DisplayName) const override	{ return false; }
+	virtual bool TryGetObjectBindingDisplayName(const TSharedRef<FMovieSceneInstance>& MovieSceneInstance, const FGuid& ObjectGuid, FText& DisplayName) const override	{ return false; }
 private:
 };
\ No newline at end of file
diff --git a/Engine/Source/Editor/Sequencer/Private/SequencerActorBindingManager.cpp b/Engine/Source/Editor/Sequencer/Private/SequencerActorBindingManager.cpp
--- a/Engine/Source/Editor/Sequencer/Private/SequencerActorBindingManager.cpp
+++ b/Engine/Source/Editor/Sequencer/Private/SequencerActorBindingManager.cpp
@@ -306,6 +306,22 @@ void FSequencerActorBindingManager::GetRuntimeObjects( const TSharedRef<FMovieSc
 	}
 }
 
+bool FSequencerActorBindingManager::TryGetObjectBindingDisplayName(const TSharedRef<FMovieSceneInstance>& MovieSceneInstance, const FGuid& ObjectGuid, FText& DisplayName) const
+{
+	TArray< UObject*> RuntimeObjects;
+	GetRuntimeObjects(MovieSceneInstance, ObjectGuid, RuntimeObjects);
+	for (int32 ObjIndex = 0; ObjIndex < RuntimeObjects.Num(); ++ObjIndex)
+	{
+		AActor* Actor = Cast<AActor>(RuntimeObjects[ObjIndex]);
+		if (Actor)
+		{
+			DisplayName = FText::FromString(Actor->GetActorLabel());
+			return true;
+		}
+	}
+	return false;
+}
+
 UK2Node_PlayMovieScene* FSequencerActorBindingManager::FindPlayMovieSceneNodeInLevelScript( const UMovieScene* MovieScene ) const
 {
 	// Grab the world object for this editor
diff --git a/Engine/Source/Editor/Sequencer/Private/SequencerActorBindingManager.h b/Engine/Source/Editor/Sequencer/Private/SequencerActorBindingManager.h
--- a/Engine/Source/Editor/Sequencer/Private/SequencerActorBindingManager.h
+++ b/Engine/Source/Editor/Sequencer/Private/SequencerActorBindingManager.h
@@ -86,7 +86,7 @@ public:
 	virtual void BindPossessableObject( const FGuid& PossessableGuid, UObject& PossessedObject ) override;
 	virtual void UnbindPossessableObjects( const FGuid& PossessableGuid ) override;
 	virtual void GetRuntimeObjects( const TSharedRef<FMovieSceneInstance>& MovieSceneInstance, const FGuid& ObjectGuid, TArray<UObject*>& OutRuntimeObjects ) const override;
-	virtual bool TryGetObjectBindingDisplayName(const FGuid& ObjectGuid, FText& DisplayName) const override { return false; }
+	virtual bool TryGetObjectBindingDisplayName(const TSharedRef<FMovieSceneInstance>& MovieSceneInstance, const FGuid& ObjectGuid, FText& DisplayName) const override;
 
 protected:
 	
diff --git a/Engine/Source/Editor/Sequencer/Private/SequencerDisplayNode.cpp b/Engine/Source/Editor/Sequencer/Private/SequencerDisplayNode.cpp
--- a/Engine/Source/Editor/Sequencer/Private/SequencerDisplayNode.cpp
+++ b/Engine/Source/Editor/Sequencer/Private/SequencerDisplayNode.cpp
@@ -476,7 +476,7 @@ void FTrackNode::FixRowIndices()
 FText FObjectBindingNode::GetDisplayName() const
 {
 	FText DisplayName;
-	return GetSequencer().GetObjectBindingManager()->TryGetObjectBindingDisplayName(ObjectBinding, DisplayName) ?
+	return GetSequencer().GetObjectBindingManager()->TryGetObjectBindingDisplayName(GetSequencer().GetFocusedMovieSceneInstance(), ObjectBinding, DisplayName) ?
 		DisplayName : DefaultDisplayName;
 }
 
diff --git a/Engine/Source/Editor/Sequencer/Public/ISequencerObjectBindingManager.h b/Engine/Source/Editor/Sequencer/Public/ISequencerObjectBindingManager.h
--- a/Engine/Source/Editor/Sequencer/Public/ISequencerObjectBindingManager.h
+++ b/Engine/Source/Editor/Sequencer/Public/ISequencerObjectBindingManager.h
@@ -61,5 +61,5 @@ public:
 	 * @param DisplayName	the display name for the object binding.
 	 * @returns true if DisplayName has been set to a valid display name, otherwise false.
 	 */
-	virtual bool TryGetObjectBindingDisplayName( const FGuid& ObjectGuid, FText& DisplayName ) const = 0;
+	virtual bool TryGetObjectBindingDisplayName( const TSharedRef<FMovieSceneInstance>& MovieSceneInstance, const FGuid& ObjectGuid, FText& DisplayName ) const = 0;
 };
diff --git a/Engine/Source/Editor/UMGEditor/Private/Animation/UMGSequencerObjectBindingManager.cpp b/Engine/Source/Editor/UMGEditor/Private/Animation/UMGSequencerObjectBindingManager.cpp
--- a/Engine/Source/Editor/UMGEditor/Private/Animation/UMGSequencerObjectBindingManager.cpp
+++ b/Engine/Source/Editor/UMGEditor/Private/Animation/UMGSequencerObjectBindingManager.cpp
@@ -132,7 +132,7 @@ void FUMGSequencerObjectBindingManager::GetRuntimeObjects( const TSharedRef<FMov
 	}
 }
 
-bool FUMGSequencerObjectBindingManager::TryGetObjectBindingDisplayName(const FGuid& ObjectGuid, FText& DisplayName) const
+bool FUMGSequencerObjectBindingManager::TryGetObjectBindingDisplayName(const TSharedRef<FMovieSceneInstance>& MovieSceneInstance, const FGuid& ObjectGuid, FText& DisplayName) const
 {
 	// TODO: This gets called every frame for every bound object and could be a potential performance issue for a really complicated animation.
 	TArray<TWeakObjectPtr<UObject>> BindingObjects;
diff --git a/Engine/Source/Editor/UMGEditor/Private/Animation/UMGSequencerObjectBindingManager.h b/Engine/Source/Editor/UMGEditor/Private/Animation/UMGSequencerObjectBindingManager.h
--- a/Engine/Source/Editor/UMGEditor/Private/Animation/UMGSequencerObjectBindingManager.h
+++ b/Engine/Source/Editor/UMGEditor/Private/Animation/UMGSequencerObjectBindingManager.h
@@ -25,7 +25,7 @@ public:
 	virtual void BindPossessableObject( const FGuid& PossessableGuid, UObject& PossessedObject ) override;
 	virtual void UnbindPossessableObjects( const FGuid& PossessableGuid ) override;
 	virtual void GetRuntimeObjects( const TSharedRef<FMovieSceneInstance>& MovieSceneInstance, const FGuid& ObjectGuid, TArray<UObject*>& OutRuntimeObjects ) const override;
-	virtual bool TryGetObjectBindingDisplayName(const FGuid& ObjectGuid, FText& DisplayName) const override;
+	virtual bool TryGetObjectBindingDisplayName(const TSharedRef<FMovieSceneInstance>& MovieSceneInstance, const FGuid& ObjectGuid, FText& DisplayName) const override;
 	
 	/** @return true if the current animation is valid */
 	bool HasValidWidgetAnimation() const;
