Sequencer: Fix crash when the movie scene isn't found in the InstanceToPuppetObjectMap.

diff --git a/Engine/Source/Editor/Sequencer/Private/SequencerActorBindingManager.cpp b/Engine/Source/Editor/Sequencer/Private/SequencerActorBindingManager.cpp
--- a/Engine/Source/Editor/Sequencer/Private/SequencerActorBindingManager.cpp
+++ b/Engine/Source/Editor/Sequencer/Private/SequencerActorBindingManager.cpp
@@ -630,9 +630,13 @@ FGuid FSequencerActorBindingManager::FindSpawnableGuidForPuppetObject( UObject*
 
 UObject* FSequencerActorBindingManager::FindPuppetObjectForSpawnableGuid( const TSharedRef<FMovieSceneInstance>& MovieSceneInstance, const FGuid& Guid ) const
 {
-	const TArray< TSharedRef<FPuppetActorInfo> >& PuppetObjects = InstanceToPuppetObjectsMap.FindChecked( MovieSceneInstance );
+	const TArray< TSharedRef<FPuppetActorInfo> >* PuppetObjects = InstanceToPuppetObjectsMap.Find(MovieSceneInstance);
+	if (PuppetObjects == NULL)
+	{
+		return NULL;
+	}
 
-	for( auto PuppetObjectIter( PuppetObjects.CreateConstIterator() ); PuppetObjectIter; ++PuppetObjectIter )
+	for( auto PuppetObjectIter( PuppetObjects->CreateConstIterator() ); PuppetObjectIter; ++PuppetObjectIter )
 	{
 		if( PuppetObjectIter->Get().GetType() == EPuppetObjectType::Actor )
 		{
