Sequencer: Update the viewport when scrubbing and realtime is off.

Tested objects/characters with transform and animation tracks. Tested with realtime turned off in the viewport, also with a camera locked to the viewport and the little camera preview window. Tested with audio scrubbing.

diff --git a/Engine/Source/Editor/Sequencer/Private/Sequencer.cpp b/Engine/Source/Editor/Sequencer/Private/Sequencer.cpp
--- a/Engine/Source/Editor/Sequencer/Private/Sequencer.cpp
+++ b/Engine/Source/Editor/Sequencer/Private/Sequencer.cpp
@@ -590,6 +590,17 @@ void FSequencer::SetGlobalTime( float NewTime )
 	ScrubPosition = NewTime;
 
 	RootMovieSceneInstance->Update( ScrubPosition, LastTime, *this );
+
+	for(FLevelEditorViewportClient* LevelVC : GEditor->LevelViewportClients)
+	{
+		if (LevelVC && LevelVC->IsPerspective())
+		{
+			if (!LevelVC->IsRealtime())
+			{
+				LevelVC->Invalidate();
+			}
+		}
+	}
 }
 
 TOptional<float> FSequencer::CalculateAutoscrollEncroachment(float NewTime) const
