Various minor sequencer sub-movie scene improvements

- Fixed a crash when adding a new sub-movie scene.  Crashed calling SaveState before the movie scene instance was created
- Added a save button to sequencer to save the focued movie scene so you dont have to go back and forth between the content browser while working on multiple sub-movie scenes
- Fixed traveling back up to parent movie scenes causing the time range in Sequencer to shift

diff --git a/Engine/Source/Editor/Sequencer/Private/SSequencer.cpp b/Engine/Source/Editor/Sequencer/Private/SSequencer.cpp
--- a/Engine/Source/Editor/Sequencer/Private/SSequencer.cpp
+++ b/Engine/Source/Editor/Sequencer/Private/SSequencer.cpp
@@ -487,6 +487,42 @@ TSharedRef<SWidget> SSequencer::MakeToolBar()
 				]);
 		}
 
+		if( Sequencer.Pin()->IsLevelEditorSequencer() )
+		{
+			ToolBarBuilder.AddWidget
+			(
+				SNew( SButton )
+				.ButtonStyle(FEditorStyle::Get(), "FlatButton")
+				.ToolTipText( LOCTEXT( "SaveDirtyPackagesTooltip", "Saves the current Movie Scene" ) )
+				.ContentPadding(FMargin(6, 2))
+				.ForegroundColor( FSlateColor::UseForeground() )
+				.OnClicked( this, &SSequencer::OnSaveMovieSceneClicked )
+				[
+					SNew( SHorizontalBox )
+					+ SHorizontalBox::Slot()
+					.VAlign(VAlign_Center)
+					.AutoWidth()
+					[
+						SNew(STextBlock)
+						.Font(FEditorStyle::Get().GetFontStyle("FontAwesome.11"))
+						.Text(FText::FromString(FString(TEXT("\xf0c7"))) /*fa-floppy-o*/)
+						.ShadowOffset( FVector2D( 1,1) )
+					]
+					+ SHorizontalBox::Slot()
+					.AutoWidth()
+					.VAlign(VAlign_Center)
+					.Padding(4, 1, 0, 0)
+					[
+						SNew( STextBlock )
+						.Text( LOCTEXT( "SaveMovieScene", "Save" ) )
+						.ShadowOffset( FVector2D( 1,1) )
+					]
+				]
+			);
+
+			ToolBarBuilder.AddSeparator();
+		}
+
 		ToolBarBuilder.AddToolBarButton( FSequencerCommands::Get().ToggleAutoKeyEnabled );
 
 		if ( Sequencer.Pin()->IsLevelEditorSequencer() )
@@ -1133,6 +1169,13 @@ void SSequencer::ExpandCollapseNode(TSharedRef<FSequencerDisplayNode> Node, bool
 	}
 }
 
+FReply SSequencer::OnSaveMovieSceneClicked()
+{
+	Sequencer.Pin()->SaveCurrentMovieScene();
+
+	return FReply::Unhandled();
+}
+
 void SSequencer::ToggleExpandCollapseSelectedNodes(bool bDescendants)
 {
 	const TSet< TSharedRef<FSequencerDisplayNode> >* SelectedNodes = Sequencer.Pin()->GetSelection()->GetSelectedOutlinerNodes();
diff --git a/Engine/Source/Editor/Sequencer/Private/SSequencer.h b/Engine/Source/Editor/Sequencer/Private/SSequencer.h
--- a/Engine/Source/Editor/Sequencer/Private/SSequencer.h
+++ b/Engine/Source/Editor/Sequencer/Private/SSequencer.h
@@ -100,6 +100,8 @@ public:
 	void SetUserIsSelecting(bool bUserIsSelectingIn) { bUserIsSelecting = bUserIsSelectingIn; }
 	bool UserIsSelecting() { return bUserIsSelecting; }
 
+	/** Called when the save button is clicked */
+	FReply OnSaveMovieSceneClicked();
 private:
 	/** Empty active timer to ensure Slate ticks during Sequencer playback */
 	EActiveTimerReturnType EnsureSlateTickDuringPlayback(double InCurrentTime, float InDeltaTime);	
diff --git a/Engine/Source/Editor/Sequencer/Private/Sequencer.cpp b/Engine/Source/Editor/Sequencer/Private/Sequencer.cpp
--- a/Engine/Source/Editor/Sequencer/Private/Sequencer.cpp
+++ b/Engine/Source/Editor/Sequencer/Private/Sequencer.cpp
@@ -591,16 +591,7 @@ void FSequencer::SetGlobalTime( float NewTime )
 
 	RootMovieSceneInstance->Update( ScrubPosition, LastTime, *this );
 
-	for(FLevelEditorViewportClient* LevelVC : GEditor->LevelViewportClients)
-	{
-		if (LevelVC && LevelVC->IsPerspective())
-		{
-			if (!LevelVC->IsRealtime())
-			{
-				LevelVC->Invalidate();
-			}
-		}
-	}
+	GEditor->RedrawLevelEditingViewports();
 }
 
 TOptional<float> FSequencer::CalculateAutoscrollEncroachment(float NewTime) const
@@ -1300,6 +1291,17 @@ void FSequencer::OnPostSaveWorld(uint32 SaveFlags, class UWorld* World, bool bSu
 	SetGlobalTime(GetGlobalTime());
 }
 
+void FSequencer::SaveCurrentMovieScene()
+{
+	UPackage* MovieScenePackage = GetFocusedMovieScene()->GetOutermost();
+
+	TArray<UPackage*> PackagesToSave;
+
+	PackagesToSave.Add( MovieScenePackage );
+
+	FEditorFileUtils::PromptForCheckoutAndSave( PackagesToSave, false, false );
+}
+
 void FSequencer::OnSectionSelectionChanged()
 {
 	for (TWeakObjectPtr<UMovieSceneSection> SelectedSection : *Selection.GetSelectedSections())
diff --git a/Engine/Source/Editor/Sequencer/Private/Sequencer.h b/Engine/Source/Editor/Sequencer/Private/Sequencer.h
--- a/Engine/Source/Editor/Sequencer/Private/Sequencer.h
+++ b/Engine/Source/Editor/Sequencer/Private/Sequencer.h
@@ -39,7 +39,9 @@ struct FSequencerInitParams
 class FSequencer : public ISequencer, public FGCObject, public FEditorUndoClient, public FTickableEditorObject
 { 
 
-public:
+public:	
+	static bool IsSequencerEnabled();
+
 	/**
 	 * Initializes sequencer
 	 *
@@ -92,7 +94,7 @@ public:
 	virtual void NotifyMovieSceneDataChanged() override;
 	virtual void UpdateRuntimeInstances() override;
 	virtual void AddSubMovieScene(UMovieScene* SubMovieScene) override;
-	virtual void FilterToShotSections(const TArray< TWeakObjectPtr<class UMovieSceneSection> >& ShotSections, bool bZoomToShotBounds = true) override;
+	virtual void FilterToShotSections(const TArray< TWeakObjectPtr<class UMovieSceneSection> >& ShotSections, bool bZoomToShotBounds = false) override;
 	virtual void FilterToSelectedShotSections(bool bZoomToShotBounds = true) override;
 	virtual bool CanKeyProperty(FCanKeyPropertyParams CanKeyPropertyParams) const override;
 	virtual void KeyProperty(FKeyPropertyParams KeyPropertyParams) override;
@@ -252,8 +254,9 @@ public:
 
 	/** @return Whether or not this sequencer is used in the level editor */
 	bool IsLevelEditorSequencer() const { return bIsEditingWithinLevelEditor; }
-	
-	static bool IsSequencerEnabled();
+
+	/** Called to save the current movie scene */
+	void SaveCurrentMovieScene();
 protected:
 	/**
 	 * Reset data about a movie scene when pushing or popping a movie scene
@@ -366,7 +369,6 @@ protected:
 
 	/** Called after the world has been saved. The sequencer updates to the animated state. */
 	void OnPostSaveWorld(uint32 SaveFlags, class UWorld* World, bool bSuccess);
-
 private:
 	TMap< TWeakObjectPtr<UMovieSceneSection>, TSharedRef<FMovieSceneInstance> > MovieSceneSectionToInstanceMap;
 
diff --git a/Engine/Source/Runtime/MovieSceneTracks/Private/TrackInstances/SubMovieSceneTrackInstance.cpp b/Engine/Source/Runtime/MovieSceneTracks/Private/TrackInstances/SubMovieSceneTrackInstance.cpp
--- a/Engine/Source/Runtime/MovieSceneTracks/Private/TrackInstances/SubMovieSceneTrackInstance.cpp
+++ b/Engine/Source/Runtime/MovieSceneTracks/Private/TrackInstances/SubMovieSceneTrackInstance.cpp
@@ -91,9 +91,11 @@ void FSubMovieSceneTrackInstance::SaveState(const TArray<UObject*>& RuntimeObjec
 	{
 		USubMovieSceneSection* Section = CastChecked<USubMovieSceneSection>( AllSections[SectionIndex] );
 
-		TSharedPtr<FMovieSceneInstance>& Instance = SubMovieSceneInstances.FindChecked( Section );
-
-		Instance->SaveState();
+		TSharedPtr<FMovieSceneInstance> Instance = SubMovieSceneInstances.FindRef( Section );
+		if( Instance.IsValid() )
+		{
+			Instance->SaveState();
+		}
 	}
 }
 
@@ -105,8 +107,10 @@ void FSubMovieSceneTrackInstance::RestoreState(const TArray<UObject*>& RuntimeOb
 	{
 		USubMovieSceneSection* Section = CastChecked<USubMovieSceneSection>( AllSections[SectionIndex] );
 
-		TSharedPtr<FMovieSceneInstance>& Instance = SubMovieSceneInstances.FindChecked( Section );
-
-		Instance->RestoreState();
+		TSharedPtr<FMovieSceneInstance> Instance = SubMovieSceneInstances.FindRef( Section );
+		if( Instance.IsValid() )
+		{
+			Instance->RestoreState();
+		}
 	}
 }
\ No newline at end of file
