Sequencer: Outliner selection now respects control (toggle) and shift (select inclusive).

diff --git a/Engine/Source/Editor/Sequencer/Private/SAnimationOutlinerView.cpp b/Engine/Source/Editor/Sequencer/Private/SAnimationOutlinerView.cpp
--- a/Engine/Source/Editor/Sequencer/Private/SAnimationOutlinerView.cpp
+++ b/Engine/Source/Editor/Sequencer/Private/SAnimationOutlinerView.cpp
@@ -264,17 +264,94 @@ FReply SAnimationOutlinerTreeNode::OnAddKeyClicked()
 	return FReply::Handled();
 }
 
+void SAnimationOutlinerTreeNode::GetAllDescendantNodes(TSharedPtr<FSequencerDisplayNode> RootNode, TArray<TSharedRef<FSequencerDisplayNode> >& AllNodes)
+{
+	if (!RootNode.IsValid())
+	{
+		return;
+	}
+
+	AllNodes.Add(RootNode.ToSharedRef());
+
+	const FSequencerDisplayNode* RootNodeC = RootNode.Get();
+
+	for (TSharedRef<FSequencerDisplayNode> ChildNode : RootNodeC->GetChildNodes())
+	{
+		AllNodes.Add(ChildNode);
+
+		GetAllDescendantNodes(ChildNode, AllNodes);
+	}
+}
+
 FReply SAnimationOutlinerTreeNode::OnMouseButtonDown( const FGeometry& MyGeometry, const FPointerEvent& MouseEvent )
 {
 	if( MouseEvent.GetEffectingButton() == EKeys::LeftMouseButton && DisplayNode->IsSelectable() )
 	{
 		FSequencer& Sequencer = DisplayNode->GetSequencer();
-		bool bSelected = Sequencer.GetSelection()->IsSelected(DisplayNode.ToSharedRef());
 
-		if( MouseEvent.IsControlDown() )
+		TArray<TSharedPtr<FSequencerDisplayNode> > AffectedNodes;
+		AffectedNodes.Add(DisplayNode.ToSharedRef());
+
+		if (MouseEvent.IsShiftDown())
 		{
-			// Select the node if we were clicked on
-			Sequencer.GetSelection()->AddToSelection(DisplayNode.ToSharedRef());
+			FSequencerNodeTree& ParentTree = DisplayNode->GetParentTree();
+
+			const TArray< TSharedRef<FSequencerDisplayNode> >RootNodes = ParentTree.GetRootNodes();
+
+			// Get all nodes in order
+			TArray<TSharedRef<FSequencerDisplayNode> > AllNodes;
+			for (int32 i = 0; i < RootNodes.Num(); ++i)
+			{
+				GetAllDescendantNodes(RootNodes[i], AllNodes);
+			}
+
+			int32 FirstIndexToSelect = INT32_MAX;
+			int32 LastIndexToSelect = INT32_MIN;
+
+			for (int32 ChildIndex = 0; ChildIndex < AllNodes.Num(); ++ChildIndex)
+			{
+				TSharedRef<FSequencerDisplayNode> ChildNode = AllNodes[ChildIndex];
+
+				if (ChildNode == DisplayNode.ToSharedRef() || Sequencer.GetSelection()->IsSelected(ChildNode))
+				{
+					if (ChildIndex < FirstIndexToSelect)
+					{
+						FirstIndexToSelect = ChildIndex;
+					}
+					if (ChildIndex > LastIndexToSelect)
+					{
+						LastIndexToSelect = ChildIndex;
+					}
+				}
+			}
+
+			if (FirstIndexToSelect != INT32_MAX && LastIndexToSelect != INT32_MIN)
+			{
+				for (int32 ChildIndex = FirstIndexToSelect; ChildIndex <= LastIndexToSelect; ++ChildIndex)
+				{
+					TSharedRef<FSequencerDisplayNode> ChildNode = AllNodes[ChildIndex];
+
+					if (!Sequencer.GetSelection()->IsSelected(ChildNode))
+					{
+						Sequencer.GetSelection()->AddToSelection(ChildNode);
+						AffectedNodes.Add(ChildNode);
+					}
+				}
+			}
+		}
+		else if( MouseEvent.IsControlDown() )
+		{
+			bool bSelected = Sequencer.GetSelection()->IsSelected(DisplayNode.ToSharedRef());
+
+			// Toggle selection when control is down
+			if (bSelected)
+			{
+				Sequencer.GetSelection()->RemoveFromSelection(DisplayNode.ToSharedRef());
+			}
+			else
+			{
+				Sequencer.GetSelection()->AddToSelection(DisplayNode.ToSharedRef());
+			}
 		}
 		else
 		{
@@ -283,7 +360,7 @@ FReply SAnimationOutlinerTreeNode::OnMouseButtonDown( const FGeometry& MyGeometr
 			Sequencer.GetSelection()->AddToSelection(DisplayNode.ToSharedRef());
 		}
 
-		OnSelectionChanged.ExecuteIfBound( DisplayNode );
+		OnSelectionChanged.ExecuteIfBound( AffectedNodes );
 		return FReply::Handled();
 	}
 
@@ -442,61 +519,93 @@ void SAnimationOutlinerView::OnArrangeChildren( const FGeometry& AllottedGeometr
 	}
 }
 
-
-void SAnimationOutlinerView::OnSelectionChanged( TSharedPtr<FSequencerDisplayNode> AffectedNode )
+void SAnimationOutlinerView::OnSelectionChanged( TArray<TSharedPtr<FSequencerDisplayNode> > AffectedNodes )
 {
+	TArray<TSharedPtr<FSequencerDisplayNode> > ObjectNodes;
+
+	for (int32 NodeIdx = 0; NodeIdx < AffectedNodes.Num(); ++NodeIdx)
+	{
+		if( AffectedNodes[NodeIdx]->GetType() == ESequencerNode::Object )
+		{
+			ObjectNodes.Add(AffectedNodes[NodeIdx]);
+		}
+	}
+
+	if (!ObjectNodes.Num())
+	{
+		return;
+	}
+
+	if (!Sequencer->IsLevelEditorSequencer())
+	{
+		return;
+	}
+
+	// Mark that the user is selecting so that the UI doesn't respond to the selection changes in the following block
+	TSharedRef<SSequencer> SequencerWidget = StaticCastSharedRef<SSequencer>(Sequencer->GetSequencerWidget());
+	SequencerWidget->SetUserIsSelecting(true);
+
+	const FScopedTransaction Transaction(NSLOCTEXT("UnrealEd", "ClickingOnActors", "Clicking on Actors"));
+	bool bActorSelected = false;
+
+	const FModifierKeysState ModifierKeys = FSlateApplication::Get().GetModifierKeys();
+	bool IsControlDown = ModifierKeys.IsControlDown();
+	bool IsShiftDown = ModifierKeys.IsShiftDown();
+
+	const bool bNotifySelectionChanged = false;
+	const bool bDeselectBSP = true;
+	const bool bWarnAboutTooManyActors = false;
+
+	GEditor->GetSelectedActors()->Modify();
+
+	if (!IsControlDown && !IsShiftDown)
+	{
+		GEditor->SelectNone(bNotifySelectionChanged, bDeselectBSP, bWarnAboutTooManyActors);
+	}
+
+	GEditor->GetSelectedActors()->BeginBatchSelectOperation();
+
 	// Select objects bound to the object node
-	if( AffectedNode->GetType() == ESequencerNode::Object )
+	for (int32 ObjectIdx = 0; ObjectIdx < ObjectNodes.Num(); ++ObjectIdx)
 	{
-		const TSharedPtr<const FObjectBindingNode> ObjectNode = StaticCastSharedPtr<const FObjectBindingNode>( AffectedNode );
+		const TSharedPtr<const FObjectBindingNode> ObjectNode = StaticCastSharedPtr<const FObjectBindingNode>( ObjectNodes[ObjectIdx] );
 
 		// Get the bound objects
 		TArray<UObject*> RuntimeObjects;
 		Sequencer->GetRuntimeObjects( Sequencer->GetFocusedMovieSceneInstance(), ObjectNode->GetObjectBinding(), RuntimeObjects );
 		
-		if( RuntimeObjects.Num() > 0 && Sequencer->IsLevelEditorSequencer() )
+		if( RuntimeObjects.Num() > 0 )
 		{
-			const bool bNotifySelectionChanged = false;
-			const bool bDeselectBSP = true;
-			const bool bWarnAboutTooManyActors = false;
-
-			const FScopedTransaction Transaction(NSLOCTEXT("UnrealEd", "ClickingOnActors", "Clicking on Actors"));
-
-			// Clear selection
-			GEditor->SelectNone(bNotifySelectionChanged,bDeselectBSP,bWarnAboutTooManyActors);
-			GEditor->GetSelectedActors()->BeginBatchSelectOperation();
-
-			// Select each actor
-			bool bActorSelected = false;
-
-			// Mark that the user is selecting so that the UI doesn't respond to the selection changes in the following block
-			TSharedRef<SSequencer> SequencerWidget = StaticCastSharedRef<SSequencer>(Sequencer->GetSequencerWidget());
-			SequencerWidget->SetUserIsSelecting(true);
-			
-			for( int32 ObjectIndex = 0; ObjectIndex < RuntimeObjects.Num(); ++ObjectIndex )
+			// Select each actor			
+			for( int32 ActorIdx = 0; ActorIdx < RuntimeObjects.Num(); ++ActorIdx )
 			{
-				AActor* Actor = Cast<AActor>( RuntimeObjects[ObjectIndex] );
+				AActor* Actor = Cast<AActor>( RuntimeObjects[ActorIdx] );
 
 				if( Actor )
 				{
-					const bool bSelectActor = true;
+					bool bSelectActor = true;
+
+					if (IsControlDown)
+					{
+						bSelectActor = Sequencer->GetSelection()->IsSelected(ObjectNodes[ObjectIdx].ToSharedRef());
+					}
 
 					GEditor->SelectActor(Actor, bSelectActor, bNotifySelectionChanged );
 					bActorSelected = true;
 				}
 			}
+		}
+	}
 
-			GEditor->GetSelectedActors()->EndBatchSelectOperation();
+	GEditor->GetSelectedActors()->EndBatchSelectOperation();
 
-			if( bActorSelected )
-			{
-				GEditor->NoteSelectionChange();
-			}
-
-			// Unlock the selection so that the sequencer widget can now respond to selection changes in the level
-			SequencerWidget->SetUserIsSelecting(false);
-		}
+	if( bActorSelected )
+	{
+		GEditor->NoteSelectionChange();
 	}
+
+	// Unlock the selection so that the sequencer widget can now respond to selection changes in the level
+	SequencerWidget->SetUserIsSelecting(false);
 }
 
 #undef LOCTEXT_NAMESPACE
\ No newline at end of file
diff --git a/Engine/Source/Editor/Sequencer/Private/SAnimationOutlinerView.h b/Engine/Source/Editor/Sequencer/Private/SAnimationOutlinerView.h
--- a/Engine/Source/Editor/Sequencer/Private/SAnimationOutlinerView.h
+++ b/Engine/Source/Editor/Sequencer/Private/SAnimationOutlinerView.h
@@ -10,7 +10,7 @@ class SAnimationOutlinerTreeNode : public SCompoundWidget
 {
 public:
 	/** A delegate called when the widget is selected */
-	DECLARE_DELEGATE_OneParam( FOnNodeSelectionChanged, TSharedPtr<class FSequencerDisplayNode> )
+	DECLARE_DELEGATE_OneParam( FOnNodeSelectionChanged, TArray<TSharedPtr<class FSequencerDisplayNode> >)
 
 	SLATE_BEGIN_ARGS( SAnimationOutlinerTreeNode ) {}
 		/** Called when the widget is selected */	
@@ -72,6 +72,9 @@ private:
 	/** Handles the add key button being clicked. */
 	FReply OnAddKeyClicked();
 
+	/** Get all descendant nodes from the given root node */
+	void GetAllDescendantNodes(TSharedPtr<FSequencerDisplayNode> RootNode, TArray<TSharedRef<FSequencerDisplayNode> >& AllNodes);
+
 private:
 	/** Layout node the widget is visualizing */
 	TSharedPtr<FSequencerDisplayNode> DisplayNode;
@@ -121,8 +124,7 @@ private:
 	 *
 	 * @param AffectedNode	The node that was selected
 	 */
-	void OnSelectionChanged( TSharedPtr<FSequencerDisplayNode> AffectedNode );
-
+	void OnSelectionChanged( TArray<TSharedPtr<FSequencerDisplayNode> > AffectedNodes );
 private:
 	/** Internal sequencer interface */
 	FSequencer* Sequencer;
diff --git a/Engine/Source/Editor/Sequencer/Private/SequencerDisplayNode.h b/Engine/Source/Editor/Sequencer/Private/SequencerDisplayNode.h
--- a/Engine/Source/Editor/Sequencer/Private/SequencerDisplayNode.h
+++ b/Engine/Source/Editor/Sequencer/Private/SequencerDisplayNode.h
@@ -137,6 +137,9 @@ public:
 	/** Gets the sequencer that owns this node */
 	FSequencer& GetSequencer() const { return ParentTree.GetSequencer(); }
 	
+	/** Gets the parent tree that this node is in */
+	FSequencerNodeTree& GetParentTree() const { return ParentTree; }
+
 	/** Gets all the key area nodes recursively, including this node if applicable */
 	virtual void GetChildKeyAreaNodesRecursively(TArray< TSharedRef<class FSectionKeyAreaNode> >& OutNodes) const;
 
diff --git a/Engine/Source/Editor/Sequencer/Private/SequencerSelection.cpp b/Engine/Source/Editor/Sequencer/Private/SequencerSelection.cpp
--- a/Engine/Source/Editor/Sequencer/Private/SequencerSelection.cpp
+++ b/Engine/Source/Editor/Sequencer/Private/SequencerSelection.cpp
@@ -65,6 +65,27 @@ void FSequencerSelection::AddToSelection(TSharedRef<FSequencerDisplayNode> Outli
 	OnOutlinerNodeSelectionChanged.Broadcast();
 }
 
+void FSequencerSelection::RemoveFromSelection(FSelectedKey Key)
+{
+	SelectedKeys.Remove(Key);
+	ActiveSelection = EActiveSelection::KeyAndSection;
+	OnKeySelectionChanged.Broadcast();
+}
+
+void FSequencerSelection::RemoveFromSelection(UMovieSceneSection* Section)
+{
+	SelectedSections.Remove(Section);
+	ActiveSelection = EActiveSelection::KeyAndSection;
+	OnSectionSelectionChanged.Broadcast();
+}
+
+void FSequencerSelection::RemoveFromSelection(TSharedRef<FSequencerDisplayNode> OutlinerNode)
+{
+	SelectedOutlinerNodes.Remove(OutlinerNode);
+	ActiveSelection = EActiveSelection::OutlinerNode;
+	OnOutlinerNodeSelectionChanged.Broadcast();
+}
+
 bool FSequencerSelection::IsSelected(FSelectedKey Key) const
 {
 	return SelectedKeys.Contains(Key);
diff --git a/Engine/Source/Editor/Sequencer/Private/SequencerSelection.h b/Engine/Source/Editor/Sequencer/Private/SequencerSelection.h
--- a/Engine/Source/Editor/Sequencer/Private/SequencerSelection.h
+++ b/Engine/Source/Editor/Sequencer/Private/SequencerSelection.h
@@ -50,6 +50,15 @@ public:
 	/** Adds an outliner node to the selection */
 	void AddToSelection(TSharedRef<FSequencerDisplayNode> OutlinerNode);
 
+	/** Removes a key from the selection */
+	void RemoveFromSelection(FSelectedKey Key);
+
+	/** Removes a section from the selection */
+	void RemoveFromSelection(UMovieSceneSection* Section);
+
+	/** Removes an outliner node from the selection */
+	void RemoveFromSelection(TSharedRef<FSequencerDisplayNode> OutlinerNode);
+
 	/** Returns whether or not the key is selected. */
 	bool IsSelected(FSelectedKey Key) const;
 
