Sequencer: Child nodes should not be expanded when created. This makes things a bit more presentable when objects are first added and keyed.

diff --git a/Engine/Source/Editor/Sequencer/Private/SequencerDisplayNode.cpp b/Engine/Source/Editor/Sequencer/Private/SequencerDisplayNode.cpp
--- a/Engine/Source/Editor/Sequencer/Private/SequencerDisplayNode.cpp
+++ b/Engine/Source/Editor/Sequencer/Private/SequencerDisplayNode.cpp
@@ -182,6 +182,10 @@ TSharedRef<FTrackNode> FSequencerDisplayNode::AddSectionAreaNode( FName SectionN
 	{
 		// No existing node found make a new one
 		SectionNode = MakeShareable( new FTrackNode( SectionName, AssociatedTrack, SharedThis( this ), ParentTree ) );
+		if (SectionNode.Get()->IsExpanded())
+		{
+			SectionNode.Get()->ToggleExpansion();
+		}
 		ChildNodes.Add( SectionNode.ToSharedRef() );
 	}
 
