Sequencer: Update outliner selection tobe in sync with scene selection when the tree is rebuilt.

diff --git a/Engine/Source/Editor/Sequencer/Private/SSequencer.cpp b/Engine/Source/Editor/Sequencer/Private/SSequencer.cpp
--- a/Engine/Source/Editor/Sequencer/Private/SSequencer.cpp
+++ b/Engine/Source/Editor/Sequencer/Private/SSequencer.cpp
@@ -739,6 +739,9 @@ void SSequencer::UpdateLayoutTree()
 	// Restore the selection state.
 	RestoreSelectionState(SequencerNodeTree->GetRootNodes(), SelectedPathNames, Sequencer.Pin()->GetSelection());
 
+	// Update to actor selection.
+	OnActorSelectionChanged(NULL);
+
 	SequencerNodeTree->UpdateCachedVisibilityBasedOnShotFiltersChanged();
 }
 
@@ -1058,7 +1061,7 @@ void SSequencer::OnActorsDropped( FActorDragDropGraphEdOp& DragDropOp )
 	Sequencer.Pin()->OnActorsDropped( DragDropOp.Actors );
 }
 
-void SSequencer::OnActorSelectionChanged(UObject* obj)
+void SSequencer::OnActorSelectionChanged(UObject*)
 {
 	if (!Sequencer.IsValid())
 	{
