Sequencer: Sort root nodes alphabetically.

diff --git a/Engine/Source/Editor/Sequencer/Private/SequencerNodeTree.cpp b/Engine/Source/Editor/Sequencer/Private/SequencerNodeTree.cpp
--- a/Engine/Source/Editor/Sequencer/Private/SequencerNodeTree.cpp
+++ b/Engine/Source/Editor/Sequencer/Private/SequencerNodeTree.cpp
@@ -69,7 +69,29 @@ void FSequencerNodeTree::Update()
 	{
 		bool operator()( const TSharedRef<FSequencerDisplayNode>& A, const TSharedRef<FSequencerDisplayNode>& B ) const
 		{
-			return A->GetType() == ESequencerNode::Object ? false : true;
+			if (A->GetType() == ESequencerNode::Object)
+			{
+				if (B->GetType() == ESequencerNode::Object)
+				{
+					return A->GetDisplayName().ToString() < B->GetDisplayName().ToString();
+				}
+				else
+				{
+					return false; // ie. master tracks should be first in line
+				}
+			}
+			if (A->GetType() != ESequencerNode::Object)
+			{
+				if (B->GetType() != ESequencerNode::Object)
+				{
+					return A->GetDisplayName().ToString() < B->GetDisplayName().ToString();
+				}
+				else
+				{
+					return true; // ie. master tracks should be first in line
+				}
+			}
+			return false;
 		}
 	};
 
