-Fixed a crash when deleting a track from Sequencer.
-Replaced return by pointer with return by reference for Sequencer selection sets.

diff --git a/Engine/Source/Editor/Sequencer/Private/SAnimationOutlinerView.cpp b/Engine/Source/Editor/Sequencer/Private/SAnimationOutlinerView.cpp
--- a/Engine/Source/Editor/Sequencer/Private/SAnimationOutlinerView.cpp
+++ b/Engine/Source/Editor/Sequencer/Private/SAnimationOutlinerView.cpp
@@ -288,6 +288,7 @@ FReply SAnimationOutlinerTreeNode::OnMouseButtonDown( const FGeometry& MyGeometr
 	if( MouseEvent.GetEffectingButton() == EKeys::LeftMouseButton && DisplayNode->IsSelectable() )
 	{
 		FSequencer& Sequencer = DisplayNode->GetSequencer();
+		bool bSelected = Sequencer.GetSelection().IsSelected(DisplayNode.ToSharedRef());
 
 		TArray<TSharedPtr<FSequencerDisplayNode> > AffectedNodes;
 		AffectedNodes.Add(DisplayNode.ToSharedRef());
@@ -312,7 +313,7 @@ FReply SAnimationOutlinerTreeNode::OnMouseButtonDown( const FGeometry& MyGeometr
 			{
 				TSharedRef<FSequencerDisplayNode> ChildNode = AllNodes[ChildIndex];
 
-				if (ChildNode == DisplayNode.ToSharedRef() || Sequencer.GetSelection()->IsSelected(ChildNode))
+				if (ChildNode == DisplayNode.ToSharedRef() || Sequencer.GetSelection().IsSelected(ChildNode))
 				{
 					if (ChildIndex < FirstIndexToSelect)
 					{
@@ -331,9 +332,9 @@ FReply SAnimationOutlinerTreeNode::OnMouseButtonDown( const FGeometry& MyGeometr
 				{
 					TSharedRef<FSequencerDisplayNode> ChildNode = AllNodes[ChildIndex];
 
-					if (!Sequencer.GetSelection()->IsSelected(ChildNode))
+					if (!Sequencer.GetSelection().IsSelected(ChildNode))
 					{
-						Sequencer.GetSelection()->AddToSelection(ChildNode);
+						Sequencer.GetSelection().AddToSelection(ChildNode);
 						AffectedNodes.Add(ChildNode);
 					}
 				}
@@ -341,23 +342,23 @@ FReply SAnimationOutlinerTreeNode::OnMouseButtonDown( const FGeometry& MyGeometr
 		}
 		else if( MouseEvent.IsControlDown() )
 		{
-			bool bSelected = Sequencer.GetSelection()->IsSelected(DisplayNode.ToSharedRef());
+			bool bSelected = Sequencer.GetSelection().IsSelected(DisplayNode.ToSharedRef());
 
 			// Toggle selection when control is down
 			if (bSelected)
 			{
-				Sequencer.GetSelection()->RemoveFromSelection(DisplayNode.ToSharedRef());
+				Sequencer.GetSelection().RemoveFromSelection(DisplayNode.ToSharedRef());
 			}
 			else
 			{
-				Sequencer.GetSelection()->AddToSelection(DisplayNode.ToSharedRef());
+				Sequencer.GetSelection().AddToSelection(DisplayNode.ToSharedRef());
 			}
 		}
 		else
 		{
 			// Deselect the other nodes and select this node.
-			Sequencer.GetSelection()->EmptySelectedOutlinerNodes();
-			Sequencer.GetSelection()->AddToSelection(DisplayNode.ToSharedRef());
+			Sequencer.GetSelection().EmptySelectedOutlinerNodes();
+			Sequencer.GetSelection().AddToSelection(DisplayNode.ToSharedRef());
 		}
 
 		OnSelectionChanged.ExecuteIfBound( AffectedNodes );
@@ -403,10 +404,10 @@ const FSlateBrush* SAnimationOutlinerTreeNode::GetNodeBorderImage() const
 {
 	// Display a highlight when the node is selected
 	FSequencer& Sequencer = DisplayNode->GetSequencer();
-	const bool bIsSelected = Sequencer.GetSelection()->IsSelected(DisplayNode.ToSharedRef());
+	const bool bIsSelected = Sequencer.GetSelection().IsSelected(DisplayNode.ToSharedRef());
 	if (bIsSelected)
 	{
-		if (Sequencer.GetSelection()->GetActiveSelection() == FSequencerSelection::EActiveSelection::OutlinerNode)
+		if (Sequencer.GetSelection().GetActiveSelection() == FSequencerSelection::EActiveSelection::OutlinerNode)
 		{
 			return SelectedBrush;
 		}
@@ -587,7 +588,7 @@ void SAnimationOutlinerView::OnSelectionChanged( TArray<TSharedPtr<FSequencerDis
 
 					if (IsControlDown)
 					{
-						bSelectActor = Sequencer->GetSelection()->IsSelected(ObjectNodes[ObjectIdx].ToSharedRef());
+						bSelectActor = Sequencer->GetSelection().IsSelected(ObjectNodes[ObjectIdx].ToSharedRef());
 					}
 
 					GEditor->SelectActor(Actor, bSelectActor, bNotifySelectionChanged );
diff --git a/Engine/Source/Editor/Sequencer/Private/SSection.cpp b/Engine/Source/Editor/Sequencer/Private/SSection.cpp
--- a/Engine/Source/Editor/Sequencer/Private/SSection.cpp
+++ b/Engine/Source/Editor/Sequencer/Private/SSection.cpp
@@ -179,7 +179,7 @@ void SSection::CreateDragOperation( const FGeometry& MyGeometry, const FPointerE
 
 	if( bKeysUnderMouse )
 	{
-		DragOperation = MakeShareable( new FMoveKeys( GetSequencer(), GetSequencer().GetSelection()->GetSelectedKeys(), PressedKey ) );
+		DragOperation = MakeShareable( new FMoveKeys( GetSequencer(), GetSequencer().GetSelection().GetSelectedKeys(), PressedKey ) );
 	}
 	else
 	{
@@ -321,8 +321,8 @@ void SSection::PaintKeys( const FGeometry& AllottedGeometry, const FSlateRect& M
 
 				FSelectedKey TestKey( SectionObject, KeyArea, KeyHandle );
 
-				bool bSelected = Sequencer.GetSelection()->IsSelected( TestKey );
-				bool bActive = Sequencer.GetSelection()->GetActiveSelection() == FSequencerSelection::EActiveSelection::KeyAndSection;
+				bool bSelected = Sequencer.GetSelection().IsSelected( TestKey );
+				bool bActive = Sequencer.GetSelection().GetActiveSelection() == FSequencerSelection::EActiveSelection::KeyAndSection;
 
 				if( TestKey == PressedKey )
 				{
@@ -365,9 +365,9 @@ void SSection::DrawSectionBorders( const FGeometry& AllottedGeometry, const FSla
 {
 	UMovieSceneSection* SectionObject = SectionInterface->GetSectionObject();
 
-	FSequencerSelection* Selection = ParentSectionArea->GetSequencer().GetSelection();
-	const bool bSelected = Selection->IsSelected(SectionObject);
-	const bool bActive = Selection->GetActiveSelection() == FSequencerSelection::EActiveSelection::KeyAndSection;
+	FSequencerSelection& Selection = ParentSectionArea->GetSequencer().GetSelection();
+	const bool bSelected = Selection.IsSelected(SectionObject);
+	const bool bActive = Selection.GetActiveSelection() == FSequencerSelection::EActiveSelection::KeyAndSection;
 
 	static const FName SelectionColorName("SelectionColor");
 	static const FName SelectionInactiveColorName("SelectionColorInactive");
@@ -619,7 +619,7 @@ FReply SSection::OnMouseMove( const FGeometry& MyGeometry, const FPointerEvent&
 					if( PressedKey.IsValid() )
 					{
 						// Clear selected sections when beginning to drag keys
-						GetSequencer().GetSelection()->EmptySelectedSections();
+						GetSequencer().GetSelection().EmptySelectedSections();
 
 						bool bSelectDueToDrag = true;
 						HandleKeySelection( PressedKey, MouseEvent, bSelectDueToDrag );
@@ -630,7 +630,7 @@ FReply SSection::OnMouseMove( const FGeometry& MyGeometry, const FPointerEvent&
 					else
 					{
 						// Clear selected keys when beginning to drag a section
-						GetSequencer().GetSelection()->EmptySelectedKeys();
+						GetSequencer().GetSelection().EmptySelectedKeys();
 
 						HandleSectionSelection( MouseEvent );
 
@@ -724,16 +724,16 @@ void SSection::HandleKeySelection( const FSelectedKey& Key, const FPointerEvent&
 	{
 		// Clear previous key selection if:
 		// we are selecting due to drag and the key being dragged is not selected or control is not down
-		bool bShouldClearSelectionDueToDrag =  bSelectDueToDrag ? !GetSequencer().GetSelection()->IsSelected( Key ) : true;
+		bool bShouldClearSelectionDueToDrag =  bSelectDueToDrag ? !GetSequencer().GetSelection().IsSelected( Key ) : true;
 		
 		// Keep key selection if right clicking to bring up a menu and the current key is selected
-		bool bKeepKeySelection = MouseEvent.GetEffectingButton() == EKeys::RightMouseButton && GetSequencer().GetSelection()->IsSelected( Key );
+		bool bKeepKeySelection = MouseEvent.GetEffectingButton() == EKeys::RightMouseButton && GetSequencer().GetSelection().IsSelected( Key );
 
 		if( (!MouseEvent.IsControlDown() && bShouldClearSelectionDueToDrag) && !bKeepKeySelection )
 		{
-			GetSequencer().GetSelection()->EmptySelectedKeys();
+			GetSequencer().GetSelection().EmptySelectedKeys();
 		}
-		GetSequencer().GetSelection()->AddToSelection( Key );
+		GetSequencer().GetSelection().AddToSelection( Key );
 	}
 }
 
@@ -741,11 +741,11 @@ void SSection::HandleSectionSelection( const FPointerEvent& MouseEvent )
 {
 	if( !MouseEvent.IsControlDown() )
 	{
-		GetSequencer().GetSelection()->EmptySelectedSections();
+		GetSequencer().GetSelection().EmptySelectedSections();
 	}
 
 	// handle selecting sections 
 	UMovieSceneSection* Section = SectionInterface->GetSectionObject();
-	GetSequencer().GetSelection()->AddToSelection(Section);
+	GetSequencer().GetSelection().AddToSelection(Section);
 }
 
diff --git a/Engine/Source/Editor/Sequencer/Private/SSequencer.cpp b/Engine/Source/Editor/Sequencer/Private/SSequencer.cpp
--- a/Engine/Source/Editor/Sequencer/Private/SSequencer.cpp
+++ b/Engine/Source/Editor/Sequencer/Private/SSequencer.cpp
@@ -55,9 +55,9 @@ FReply SSequencerScrollBox::OnMouseButtonDown(const FGeometry& MyGeometry, const
 	// Clear the selection if no sequencer display nodes were clicked on
 	if (MouseEvent.GetEffectingButton() == EKeys::LeftMouseButton)
 	{
-		if (Sequencer.Pin()->GetSelection()->GetSelectedOutlinerNodes()->Num() != 0)
+		if (Sequencer.Pin()->GetSelection().GetSelectedOutlinerNodes().Num() != 0)
 		{
-			Sequencer.Pin()->GetSelection()->EmptySelectedOutlinerNodes();
+			Sequencer.Pin()->GetSelection().EmptySelectedOutlinerNodes();
 
 			if (Sequencer.Pin()->IsLevelEditorSequencer())
 			{
@@ -192,7 +192,7 @@ void SSequencer::Construct( const FArguments& InArgs, TSharedRef< class FSequenc
 
 	USelection::SelectionChangedEvent.AddSP(this, &SSequencer::OnActorSelectionChanged);
 
-	GetMutableDefault<USequencerSettings>()->GetOnShowCurveEditorChanged()->AddSP(this, &SSequencer::OnCurveEditorVisibilityChanged);
+	GetMutableDefault<USequencerSettings>()->GetOnShowCurveEditorChanged().AddSP(this, &SSequencer::OnCurveEditorVisibilityChanged);
 
 	// Create a node tree which contains a tree of movie scene data to display in the sequence
 	SequencerNodeTree = MakeShareable( new FSequencerNodeTree( InSequencer.Get() ) );
@@ -677,7 +677,7 @@ TSharedRef<SWidget> SSequencer::MakeTransportControls()
 SSequencer::~SSequencer()
 {
 	USelection::SelectionChangedEvent.RemoveAll(this);
-	GetMutableDefault<USequencerSettings>()->GetOnShowCurveEditorChanged()->RemoveAll(this);
+	GetMutableDefault<USequencerSettings>()->GetOnShowCurveEditorChanged().RemoveAll(this);
 }
 
 void SSequencer::RegisterActiveTimerForPlayback()
@@ -704,13 +704,13 @@ EActiveTimerReturnType SSequencer::EnsureSlateTickDuringPlayback(double InCurren
 	return EActiveTimerReturnType::Stop;
 }
 
-void RestoreSelectionState(const TArray<TSharedRef<FSequencerDisplayNode>>& DisplayNodes, TSet<FString>& SelectedPathNames, FSequencerSelection* SequencerSelection)
+void RestoreSelectionState(const TArray<TSharedRef<FSequencerDisplayNode>>& DisplayNodes, TSet<FString>& SelectedPathNames, FSequencerSelection& SequencerSelection)
 {
 	for (TSharedRef<FSequencerDisplayNode> DisplayNode : DisplayNodes)
 	{
 		if (SelectedPathNames.Contains(DisplayNode->GetPathName()))
 		{
-			SequencerSelection->AddToSelection(DisplayNode);
+			SequencerSelection.AddToSelection(DisplayNode);
 		}
 		for (TSharedRef<FSequencerDisplayNode> ChildDisplayNode : DisplayNode->GetChildNodes())
 		{
@@ -723,7 +723,7 @@ void SSequencer::UpdateLayoutTree()
 {
 	// Cache the selected path names so selection can be restored after the update.
 	TSet<FString> SelectedPathNames;
-	for (TSharedRef<const FSequencerDisplayNode> SelectedDisplayNode : Sequencer.Pin()->GetSelection()->GetSelectedOutlinerNodes()->Array())
+	for (TSharedRef<const FSequencerDisplayNode> SelectedDisplayNode : Sequencer.Pin()->GetSelection().GetSelectedOutlinerNodes().Array())
 	{
 		SelectedPathNames.Add(SelectedDisplayNode->GetPathName());
 	}
@@ -930,13 +930,13 @@ void SSequencer::OnAssetsDropped( const FAssetDragDropOp& DragDropOp )
 		}
 	}
 
-	const TSet< TSharedRef<FSequencerDisplayNode> >* SelectedNodes = Sequencer.Pin()->GetSelection()->GetSelectedOutlinerNodes();
+	const TSet< TSharedRef<FSequencerDisplayNode> >& SelectedNodes = Sequencer.Pin()->GetSelection().GetSelectedOutlinerNodes();
 	FGuid TargetObjectGuid;
 	// if exactly one object node is selected, we have a target object guid
 	TSharedPtr<const FSequencerDisplayNode> DisplayNode;
-	if (SelectedNodes->Num() == 1)
+	if (SelectedNodes.Num() == 1)
 	{
-		for (TSharedRef<const FSequencerDisplayNode> SelectedNode : *SelectedNodes )
+		for (TSharedRef<const FSequencerDisplayNode> SelectedNode : SelectedNodes )
 		{
 			DisplayNode = SelectedNode;
 		}
@@ -1082,7 +1082,7 @@ void SSequencer::OnActorSelectionChanged(UObject*)
 	TSet<TSharedRef<FSequencerDisplayNode>> RootNodes(SequencerNodeTree->GetRootNodes());
 
 	// Clear selection
-	Sequencer.Pin()->GetSelection()->EmptySelectedOutlinerNodes();
+	Sequencer.Pin()->GetSelection().EmptySelectedOutlinerNodes();
 
 	// Select the nodes that have runtime objects that are selected in the level
 	for (auto Node : RootNodes)
@@ -1095,7 +1095,7 @@ void SSequencer::OnActorSelectionChanged(UObject*)
 		{
 			if (GEditor->GetSelectedActors()->IsSelected(RuntimeObjects[RuntimeIndex]))
 			{
-				Sequencer.Pin()->GetSelection()->AddToSelection(Node);
+				Sequencer.Pin()->GetSelection().AddToSelection(Node);
 				break;
 			}
 		}
@@ -1136,15 +1136,15 @@ FText SSequencer::GetShotSectionTitle(UMovieSceneSection* ShotSection) const
 
 void SSequencer::DeleteSelectedNodes()
 {
-	const TSet< TSharedRef<FSequencerDisplayNode> >* SelectedNodes = Sequencer.Pin()->GetSelection()->GetSelectedOutlinerNodes();
+	TSet< TSharedRef<FSequencerDisplayNode> > SelectedNodesCopy = Sequencer.Pin()->GetSelection().GetSelectedOutlinerNodes();
 
-	if( SelectedNodes->Num() > 0 )
+	if( SelectedNodesCopy.Num() > 0 )
 	{
 		const FScopedTransaction Transaction( LOCTEXT("UndoDeletingObject", "Delete Object from MovieScene") );
 
 		FSequencer& SequencerRef = *Sequencer.Pin();
 
-		for( auto SelectedNode : *SelectedNodes )
+		for( const TSharedRef<FSequencerDisplayNode>& SelectedNode : SelectedNodesCopy )
 		{
 			if( SelectedNode->IsVisible() )
 			{
@@ -1181,9 +1181,9 @@ FReply SSequencer::OnSaveMovieSceneClicked()
 
 void SSequencer::ToggleExpandCollapseSelectedNodes(bool bDescendants)
 {
-	const TSet< TSharedRef<FSequencerDisplayNode> >* SelectedNodes = Sequencer.Pin()->GetSelection()->GetSelectedOutlinerNodes();
+	const TSet< TSharedRef<FSequencerDisplayNode> >& SelectedNodes = Sequencer.Pin()->GetSelection().GetSelectedOutlinerNodes();
 
-	TSet< TSharedRef<FSequencerDisplayNode> > Nodes(*SelectedNodes);
+	TSet< TSharedRef<FSequencerDisplayNode> > Nodes(SelectedNodes);
 
 	if (Nodes.Num() == 0)
 	{
diff --git a/Engine/Source/Editor/Sequencer/Private/SSequencerCurveEditor.cpp b/Engine/Source/Editor/Sequencer/Private/SSequencerCurveEditor.cpp
--- a/Engine/Source/Editor/Sequencer/Private/SSequencerCurveEditor.cpp
+++ b/Engine/Source/Editor/Sequencer/Private/SSequencerCurveEditor.cpp
@@ -11,8 +11,8 @@ void SSequencerCurveEditor::Construct( const FArguments& InArgs, TSharedRef<FSeq
 {
 	Sequencer = InSequencer;
 	TimeSliderController = InTimeSliderController;
-	NodeTreeSelectionChangedHandle = Sequencer.Pin()->GetSelection()->GetOnOutlinerNodeSelectionChanged()->AddSP(this, &SSequencerCurveEditor::NodeTreeSelectionChanged);
-	GetMutableDefault<USequencerSettings>()->GetOnCurveVisibilityChanged()->AddSP( this, &SSequencerCurveEditor::SequencerCurveVisibilityChanged );
+	NodeTreeSelectionChangedHandle = Sequencer.Pin()->GetSelection().GetOnOutlinerNodeSelectionChanged().AddSP(this, &SSequencerCurveEditor::NodeTreeSelectionChanged);
+	GetMutableDefault<USequencerSettings>()->GetOnCurveVisibilityChanged().AddSP( this, &SSequencerCurveEditor::SequencerCurveVisibilityChanged );
 
 	// @todo sequencer: switch to lambda capture expressions when support is introduced?
 	auto ViewRange = InArgs._ViewRange;
@@ -157,7 +157,7 @@ SSequencerCurveEditor::~SSequencerCurveEditor()
 {
 	if ( Sequencer.IsValid() )
 	{
-		Sequencer.Pin()->GetSelection()->GetOnOutlinerNodeSelectionChanged()->Remove( NodeTreeSelectionChangedHandle );
+		Sequencer.Pin()->GetSelection().GetOnOutlinerNodeSelectionChanged().Remove( NodeTreeSelectionChangedHandle );
 	}
 }
 
diff --git a/Engine/Source/Editor/Sequencer/Private/SSequencerSectionAreaView.cpp b/Engine/Source/Editor/Sequencer/Private/SSequencerSectionAreaView.cpp
--- a/Engine/Source/Editor/Sequencer/Private/SSequencerSectionAreaView.cpp
+++ b/Engine/Source/Editor/Sequencer/Private/SSequencerSectionAreaView.cpp
@@ -137,7 +137,7 @@ FReply SSequencerSectionAreaView::OnMouseButtonDown( const FGeometry& MyGeometry
 	// Clear selected sections
 	if( !MouseEvent.IsControlDown() )
 	{
-		GetSequencer().GetSelection()->EmptySelectedSections();
+		GetSequencer().GetSelection().EmptySelectedSections();
 	}
 
 	return FReply::Handled();
diff --git a/Engine/Source/Editor/Sequencer/Private/SectionDragOperations.cpp b/Engine/Source/Editor/Sequencer/Private/SectionDragOperations.cpp
--- a/Engine/Source/Editor/Sequencer/Private/SectionDragOperations.cpp
+++ b/Engine/Source/Editor/Sequencer/Private/SectionDragOperations.cpp
@@ -398,13 +398,13 @@ bool FMoveSection::IsSectionAtNewLocationValid(UMovieSceneSection* InSection, in
 
 void FMoveKeys::OnBeginDrag(const FVector2D& LocalMousePos, TSharedPtr<FTrackNode> SequencerNode)
 {
-	check( SelectedKeys->Num() > 0 )
+	check( SelectedKeys.Num() > 0 )
 
 	// Begin an editor transaction and mark the section as transactional so it's state will be saved
 	GEditor->BeginTransaction( NSLOCTEXT("Sequencer", "MoveKeysTransaction", "Move Keys") );
 
 	TSet<UMovieSceneSection*> ModifiedSections;
-	for( FSelectedKey SelectedKey : *SelectedKeys )
+	for( FSelectedKey SelectedKey : SelectedKeys )
 	{
 		UMovieSceneSection* OwningSection = SelectedKey.Section;
 
@@ -447,7 +447,7 @@ void FMoveKeys::OnDrag( const FPointerEvent& MouseEvent, const FVector2D& LocalM
 				GetKeySnapTimes(OutSnapTimes, SequencerNode);
 
 				TArray<float> InitialTimes;
-				for ( FSelectedKey SelectedKey : *SelectedKeys )
+				for ( FSelectedKey SelectedKey : SelectedKeys )
 				{
 					InitialTimes.Add(SelectedKey.KeyArea->GetKeyTime(SelectedKey.KeyHandle.GetValue()) + DistanceMoved);
 				}
@@ -466,7 +466,7 @@ void FMoveKeys::OnDrag( const FPointerEvent& MouseEvent, const FVector2D& LocalM
 			}
 		}
 
-		for( FSelectedKey SelectedKey : *SelectedKeys )
+		for( FSelectedKey SelectedKey : SelectedKeys )
 		{
 			UMovieSceneSection* Section = SelectedKey.Section;
 
@@ -511,7 +511,7 @@ void FMoveKeys::GetKeySnapTimes(TArray<float>& OutSnapTimes, TSharedPtr<FTrackNo
 			{
 				FKeyHandle KeyHandle = KeyHandles[KeyIndex];
 				bool bKeyIsSnappable = true;
-				for ( FSelectedKey SelectedKey : *SelectedKeys )
+				for ( FSelectedKey SelectedKey : SelectedKeys )
 				{
 					if (SelectedKey.KeyArea == KeyArea && SelectedKey.KeyHandle.GetValue() == KeyHandle)
 					{
diff --git a/Engine/Source/Editor/Sequencer/Private/SectionDragOperations.h b/Engine/Source/Editor/Sequencer/Private/SectionDragOperations.h
--- a/Engine/Source/Editor/Sequencer/Private/SectionDragOperations.h
+++ b/Engine/Source/Editor/Sequencer/Private/SectionDragOperations.h
@@ -139,7 +139,7 @@ private:
 class FMoveKeys : public FSequencerDragOperation
 {
 public:
-	FMoveKeys( FSequencer& Sequencer,  const TSet<FSelectedKey>* InSelectedKeys, FSelectedKey& PressedKey )
+	FMoveKeys( FSequencer& Sequencer,  const TSet<FSelectedKey>& InSelectedKeys, FSelectedKey& PressedKey )
 		: FSequencerDragOperation(Sequencer)
 		, SelectedKeys( InSelectedKeys )
 		, DraggedKey( PressedKey )
@@ -156,7 +156,7 @@ protected:
 
 private:
 	/** The selected keys being moved. */
-	const TSet<FSelectedKey>* SelectedKeys;
+	const TSet<FSelectedKey>& SelectedKeys;
 	/** The exact key that we're dragging */
 	FSelectedKey DraggedKey;
 };
diff --git a/Engine/Source/Editor/Sequencer/Private/Sequencer.cpp b/Engine/Source/Editor/Sequencer/Private/Sequencer.cpp
--- a/Engine/Source/Editor/Sequencer/Private/Sequencer.cpp
+++ b/Engine/Source/Editor/Sequencer/Private/Sequencer.cpp
@@ -423,7 +423,7 @@ void FSequencer::DeleteSelectedKeys()
 {
 	FScopedTransaction DeleteKeysTransaction( NSLOCTEXT("Sequencer", "DeleteSelectedKeys_Transaction", "Delete Selected Keys") );
 	bool bAnythingRemoved = false;
-	TArray<FSelectedKey> SelectedKeysArray = Selection.GetSelectedKeys()->Array();
+	TArray<FSelectedKey> SelectedKeysArray = Selection.GetSelectedKeys().Array();
 
 	for ( const FSelectedKey& Key : SelectedKeysArray )
 	{
@@ -1304,7 +1304,7 @@ void FSequencer::SaveCurrentMovieScene()
 
 void FSequencer::OnSectionSelectionChanged()
 {
-	for (TWeakObjectPtr<UMovieSceneSection> SelectedSection : *Selection.GetSelectedSections())
+	for (TWeakObjectPtr<UMovieSceneSection> SelectedSection : Selection.GetSelectedSections())
 	{
 		// if we select something, consider it unfilterable until we change shot filters
 		UnfilterableSections.AddUnique(TWeakObjectPtr<UMovieSceneSection>(SelectedSection));
@@ -1314,7 +1314,7 @@ void FSequencer::OnSectionSelectionChanged()
 void FSequencer::ZoomToSelectedSections()
 {
 	TArray< TRange<float> > Bounds;
-	for (TWeakObjectPtr<UMovieSceneSection> SelectedSection : *Selection.GetSelectedSections())
+	for (TWeakObjectPtr<UMovieSceneSection> SelectedSection : Selection.GetSelectedSections())
 	{
 		Bounds.Add(SelectedSection->GetRange());
 	}
@@ -1393,7 +1393,7 @@ void FSequencer::FilterToShotSections(const TArray< TWeakObjectPtr<class UMovieS
 void FSequencer::FilterToSelectedShotSections(bool bZoomToShotBounds)
 {
 	TArray< TWeakObjectPtr<UMovieSceneSection> > SelectedShotSections;
-	for (TWeakObjectPtr<UMovieSceneSection> SelectedSection : *Selection.GetSelectedSections())
+	for (TWeakObjectPtr<UMovieSceneSection> SelectedSection : Selection.GetSelectedSections())
 	{
 		if (SelectedSection->IsA<UMovieSceneShotSection>())
 		{
@@ -1418,12 +1418,12 @@ TSharedRef<ISequencerObjectBindingManager> FSequencer::GetObjectBindingManager()
 	return ObjectBindingManager.ToSharedRef();
 }
 
-FSequencerSelection* FSequencer::GetSelection()
+FSequencerSelection& FSequencer::GetSelection()
 {
-	return &Selection;
+	return Selection;
 }
 
-TArray< TWeakObjectPtr<UMovieSceneSection> > FSequencer::GetFilteringShotSections() const
+const TArray< TWeakObjectPtr<UMovieSceneSection> >& FSequencer::GetFilteringShotSections() const
 {
 	return FilteringShots;
 }
@@ -1474,7 +1474,7 @@ void FSequencer::DeleteSelectedItems()
 	if (Selection.GetActiveSelection() == FSequencerSelection::EActiveSelection::KeyAndSection)
 	{
 		DeleteSelectedKeys();
-		for (TWeakObjectPtr<UMovieSceneSection> SelectedSection : *Selection.GetSelectedSections())
+		for (TWeakObjectPtr<UMovieSceneSection> SelectedSection : Selection.GetSelectedSections())
 		{
 			DeleteSection(SelectedSection.Get());
 		}
diff --git a/Engine/Source/Editor/Sequencer/Private/Sequencer.h b/Engine/Source/Editor/Sequencer/Private/Sequencer.h
--- a/Engine/Source/Editor/Sequencer/Private/Sequencer.h
+++ b/Engine/Source/Editor/Sequencer/Private/Sequencer.h
@@ -99,7 +99,7 @@ public:
 	virtual bool CanKeyProperty(FCanKeyPropertyParams CanKeyPropertyParams) const override;
 	virtual void KeyProperty(FKeyPropertyParams KeyPropertyParams) override;
 	virtual TSharedRef<ISequencerObjectBindingManager> GetObjectBindingManager() const override;
-	virtual FSequencerSelection* GetSelection() override;
+	virtual FSequencerSelection& GetSelection() override;
 
 	/** @return The current view range */
 	FAnimatedRange GetViewRange() const;
@@ -182,7 +182,7 @@ public:
 	/**
 	 * Gets all shots that are filtering currently
 	 */
-	TArray< TWeakObjectPtr<UMovieSceneSection> > GetFilteringShotSections() const;
+	const TArray< TWeakObjectPtr<UMovieSceneSection> >& GetFilteringShotSections() const;
 
 	/**
 	 * Checks to see if an object is unfilterable
diff --git a/Engine/Source/Editor/Sequencer/Private/SequencerCurveOwner.cpp b/Engine/Source/Editor/Sequencer/Private/SequencerCurveOwner.cpp
--- a/Engine/Source/Editor/Sequencer/Private/SequencerCurveOwner.cpp
+++ b/Engine/Source/Editor/Sequencer/Private/SequencerCurveOwner.cpp
@@ -62,7 +62,7 @@ FSequencerCurveOwner::FSequencerCurveOwner( TSharedPtr<FSequencerNodeTree> InSeq
 					bAddCurve = true;
 					break;
 				case ESequencerCurveVisibility::SelectedCurves:
-					bAddCurve = KeyAreaNode->GetSequencer().GetSelection()->IsSelected(KeyAreaNode);
+					bAddCurve = KeyAreaNode->GetSequencer().GetSelection().IsSelected(KeyAreaNode);
 					break;
 				case ESequencerCurveVisibility::AnimatedCurves:
 					bAddCurve = RichCurve->GetNumKeys() > 0;
diff --git a/Engine/Source/Editor/Sequencer/Private/SequencerNodeTree.cpp b/Engine/Source/Editor/Sequencer/Private/SequencerNodeTree.cpp
--- a/Engine/Source/Editor/Sequencer/Private/SequencerNodeTree.cpp
+++ b/Engine/Source/Editor/Sequencer/Private/SequencerNodeTree.cpp
@@ -15,7 +15,7 @@ void FSequencerNodeTree::Empty()
 {
 	RootNodes.Empty();
 	ObjectBindingMap.Empty();
-	Sequencer.GetSelection()->EmptySelectedOutlinerNodes();
+	Sequencer.GetSelection().EmptySelectedOutlinerNodes();
 	EditorMap.Empty();
 	FilteredNodes.Empty();
 }
diff --git a/Engine/Source/Editor/Sequencer/Private/SequencerSelection.cpp b/Engine/Source/Editor/Sequencer/Private/SequencerSelection.cpp
--- a/Engine/Source/Editor/Sequencer/Private/SequencerSelection.cpp
+++ b/Engine/Source/Editor/Sequencer/Private/SequencerSelection.cpp
@@ -9,19 +9,19 @@ FSequencerSelection::FSequencerSelection()
 	ActiveSelection = EActiveSelection::None;
 }
 
-const TSet<FSelectedKey>* FSequencerSelection::GetSelectedKeys() const
+const TSet<FSelectedKey>& FSequencerSelection::GetSelectedKeys() const
 {
-	return &SelectedKeys;
+	return SelectedKeys;
 }
 
-const TSet<TWeakObjectPtr<UMovieSceneSection>>* FSequencerSelection::GetSelectedSections() const
+const TSet<TWeakObjectPtr<UMovieSceneSection>>& FSequencerSelection::GetSelectedSections() const
 {
-	return &SelectedSections;
+	return SelectedSections;
 }
 
-const TSet<TSharedRef<FSequencerDisplayNode>>* FSequencerSelection::GetSelectedOutlinerNodes() const
+const TSet<TSharedRef<FSequencerDisplayNode>>& FSequencerSelection::GetSelectedOutlinerNodes() const
 {
-	return &SelectedOutlinerNodes;
+	return SelectedOutlinerNodes;
 }
 
 FSequencerSelection::EActiveSelection FSequencerSelection::GetActiveSelection() const
@@ -29,19 +29,19 @@ FSequencerSelection::EActiveSelection FSequencerSelection::GetActiveSelection()
 	return ActiveSelection;
 }
 
-FSequencerSelection::FOnSelectionChanged* FSequencerSelection::GetOnKeySelectionChanged()
+FSequencerSelection::FOnSelectionChanged& FSequencerSelection::GetOnKeySelectionChanged()
 {
-	return &OnKeySelectionChanged;
+	return OnKeySelectionChanged;
 }
 
-FSequencerSelection::FOnSelectionChanged* FSequencerSelection::GetOnSectionSelectionChanged()
+FSequencerSelection::FOnSelectionChanged& FSequencerSelection::GetOnSectionSelectionChanged()
 {
-	return &OnSectionSelectionChanged;
+	return OnSectionSelectionChanged;
 }
 
-FSequencerSelection::FOnSelectionChanged* FSequencerSelection::GetOnOutlinerNodeSelectionChanged()
+FSequencerSelection::FOnSelectionChanged& FSequencerSelection::GetOnOutlinerNodeSelectionChanged()
 {
-	return &OnOutlinerNodeSelectionChanged;
+	return OnOutlinerNodeSelectionChanged;
 }
 
 void FSequencerSelection::AddToSelection(FSelectedKey Key)
diff --git a/Engine/Source/Editor/Sequencer/Private/SequencerSelection.h b/Engine/Source/Editor/Sequencer/Private/SequencerSelection.h
--- a/Engine/Source/Editor/Sequencer/Private/SequencerSelection.h
+++ b/Engine/Source/Editor/Sequencer/Private/SequencerSelection.h
@@ -27,13 +27,13 @@ public:
 	FSequencerSelection();
 
 	/** Gets a set of the selected keys. */
-	const TSet<FSelectedKey>* GetSelectedKeys() const;
+	const TSet<FSelectedKey>& GetSelectedKeys() const;
 
 	/** Gets a set of the selected sections */
-	const TSet<TWeakObjectPtr<UMovieSceneSection>>* GetSelectedSections() const;
+	const TSet<TWeakObjectPtr<UMovieSceneSection>>& GetSelectedSections() const;
 
 	/** Gets a set of the selected outliner nodes. */
-	const TSet<TSharedRef<FSequencerDisplayNode>>* GetSelectedOutlinerNodes() const;
+	const TSet<TSharedRef<FSequencerDisplayNode>>& GetSelectedOutlinerNodes() const;
 
 	/** 
 	 * Gets the currently active selection.  This is used to determine which selection actions
@@ -81,13 +81,13 @@ public:
 	void EmptySelectedOutlinerNodes();
 
 	/** Gets a multicast delegate which is called when the key selection changes. */
-	FOnSelectionChanged* GetOnKeySelectionChanged();
+	FOnSelectionChanged& GetOnKeySelectionChanged();
 
 	/** Gets a multicast delegate which is called when the section selection changes. */
-	FOnSelectionChanged* GetOnSectionSelectionChanged();
+	FOnSelectionChanged& GetOnSectionSelectionChanged();
 
 	/** Gets a multicast delegate which is called when the outliner node selection changes. */
-	FOnSelectionChanged* GetOnOutlinerNodeSelectionChanged();
+	FOnSelectionChanged& GetOnOutlinerNodeSelectionChanged();
 	
 private:
 	TSet<FSelectedKey> SelectedKeys;
diff --git a/Engine/Source/Editor/Sequencer/Private/SequencerSettings.cpp b/Engine/Source/Editor/Sequencer/Private/SequencerSettings.cpp
--- a/Engine/Source/Editor/Sequencer/Private/SequencerSettings.cpp
+++ b/Engine/Source/Editor/Sequencer/Private/SequencerSettings.cpp
@@ -227,12 +227,12 @@ float USequencerSettings::SnapTimeToInterval( float InTimeValue ) const
 		: InTimeValue;
 }
 
-USequencerSettings::FOnCurveVisibilityChanged* USequencerSettings::GetOnCurveVisibilityChanged()
+USequencerSettings::FOnCurveVisibilityChanged& USequencerSettings::GetOnCurveVisibilityChanged()
 {
-	return &OnCurveVisibilityChanged;
+	return OnCurveVisibilityChanged;
 }
 
-USequencerSettings::FOnShowCurveEditorChanged* USequencerSettings::GetOnShowCurveEditorChanged()
+USequencerSettings::FOnShowCurveEditorChanged& USequencerSettings::GetOnShowCurveEditorChanged()
 {
-	return &OnShowCurveEditorChanged;
+	return OnShowCurveEditorChanged;
 }
diff --git a/Engine/Source/Editor/Sequencer/Private/SequencerSettings.h b/Engine/Source/Editor/Sequencer/Private/SequencerSettings.h
--- a/Engine/Source/Editor/Sequencer/Private/SequencerSettings.h
+++ b/Engine/Source/Editor/Sequencer/Private/SequencerSettings.h
@@ -103,10 +103,10 @@ public:
 	float SnapTimeToInterval(float InTimeValue) const;
 
 	/** Gets the multicast delegate which is run whenever the CurveVisibility value changes. */
-	FOnCurveVisibilityChanged* GetOnCurveVisibilityChanged();
+	FOnCurveVisibilityChanged& GetOnCurveVisibilityChanged();
 
 	/** Gets the multicast delegate which is run whenever the curve editor is shown/hidden. */
-	FOnShowCurveEditorChanged* GetOnShowCurveEditorChanged();
+	FOnShowCurveEditorChanged& GetOnShowCurveEditorChanged();
 
 protected:
 	UPROPERTY( config )
diff --git a/Engine/Source/Editor/Sequencer/Public/ISequencer.h b/Engine/Source/Editor/Sequencer/Public/ISequencer.h
--- a/Engine/Source/Editor/Sequencer/Public/ISequencer.h
+++ b/Engine/Source/Editor/Sequencer/Public/ISequencer.h
@@ -132,5 +132,5 @@ public:
 
 	virtual TSharedRef<ISequencerObjectBindingManager> GetObjectBindingManager() const = 0;
 
-	virtual FSequencerSelection* GetSelection() = 0;
+	virtual FSequencerSelection& GetSelection() = 0;
 };
