Fixed shadowed variable

diff --git a/Engine/Source/Editor/Sequencer/Private/SAnimationOutlinerView.cpp b/Engine/Source/Editor/Sequencer/Private/SAnimationOutlinerView.cpp
--- a/Engine/Source/Editor/Sequencer/Private/SAnimationOutlinerView.cpp
+++ b/Engine/Source/Editor/Sequencer/Private/SAnimationOutlinerView.cpp
@@ -342,8 +342,6 @@ FReply SAnimationOutlinerTreeNode::OnMouseButtonDown( const FGeometry& MyGeometr
 		}
 		else if( MouseEvent.IsControlDown() )
 		{
-			bool bSelected = Sequencer.GetSelection().IsSelected(DisplayNode.ToSharedRef());
-
 			// Toggle selection when control is down
 			if (bSelected)
 			{
