Sequencer: Back out disabling expansion of child tracks because it's also invoked when the outliner tree is rebuilt. So, if you expand child tracks and then perform an operation that causes the outliner tree to be rebuilt (ie. drag in another object), the child tracks will get collapsed.

diff --git a/Engine/Source/Editor/Sequencer/Private/SequencerDisplayNode.cpp b/Engine/Source/Editor/Sequencer/Private/SequencerDisplayNode.cpp
--- a/Engine/Source/Editor/Sequencer/Private/SequencerDisplayNode.cpp
+++ b/Engine/Source/Editor/Sequencer/Private/SequencerDisplayNode.cpp
@@ -182,10 +182,6 @@ TSharedRef<FTrackNode> FSequencerDisplayNode::AddSectionAreaNode( FName SectionN
 	{
 		// No existing node found make a new one
 		SectionNode = MakeShareable( new FTrackNode( SectionName, AssociatedTrack, SharedThis( this ), ParentTree ) );
-		if (SectionNode.Get()->IsExpanded())
-		{
-			SectionNode.Get()->ToggleExpansion();
-		}
 		ChildNodes.Add( SectionNode.ToSharedRef() );
 	}
 
