Sequencer: Add support for track editors to customize buttons in the outliner next to the property pulldown menu.

diff --git a/Engine/Source/Editor/EditorStyle/Private/SlateEditorStyle.cpp b/Engine/Source/Editor/EditorStyle/Private/SlateEditorStyle.cpp
--- a/Engine/Source/Editor/EditorStyle/Private/SlateEditorStyle.cpp
+++ b/Engine/Source/Editor/EditorStyle/Private/SlateEditorStyle.cpp
@@ -1594,6 +1594,8 @@ void FSlateEditorStyle::FStyle::SetupGeneralStyles()
 		Set( "Sequencer.ToggleIsSnapEnabled.Small", new IMAGE_BRUSH( "Icons/icon_Sequencer_ToggleSnap_40x", Icon16x16 ) );
 		Set( "Sequencer.ToggleAutoScroll", new IMAGE_BRUSH( "Icons/icon_Sequencer_ToggleAutoScroll_40x", Icon40x40 ) );
 		Set( "Sequencer.ToggleAutoScroll.Small", new IMAGE_BRUSH( "Icons/icon_Sequencer_ToggleAutoScroll_16x", Icon16x16 ) );
+		Set( "Sequencer.LockCamera", new IMAGE_BRUSH( "Icons/icon_ToggleActorPilotCameraView_16x", Icon16x16 ) );
+		Set( "Sequencer.UnlockCamera", new IMAGE_BRUSH( "Icons/icon_ToggleActorPilotCameraView_16x", Icon16x16, FLinearColor(1.f, 1.f, 1.f, 0.5f) ) );
 
 		const FButtonStyle DetailsKeyButton = FButtonStyle(NoBorder)
 			.SetNormal( IMAGE_BRUSH("Sequencer/AddKey_Details", FVector2D(11,11) )  )
diff --git a/Engine/Source/Editor/MovieSceneTools/Private/TrackEditors/ParticleTrackEditor.h b/Engine/Source/Editor/MovieSceneTools/Private/TrackEditors/ParticleTrackEditor.h
--- a/Engine/Source/Editor/MovieSceneTools/Private/TrackEditors/ParticleTrackEditor.h
+++ b/Engine/Source/Editor/MovieSceneTools/Private/TrackEditors/ParticleTrackEditor.h
@@ -6,7 +6,7 @@
 /**
  * Tools for particle tracks
  */
-class FParticleTrackEditor : public FMovieSceneTrackEditor, public TSharedFromThis<FParticleTrackEditor>
+class FParticleTrackEditor : public FMovieSceneTrackEditor
 {
 public:
 	/**
diff --git a/Engine/Source/Editor/Sequencer/Private/Sequencer.cpp b/Engine/Source/Editor/Sequencer/Private/Sequencer.cpp
--- a/Engine/Source/Editor/Sequencer/Private/Sequencer.cpp
+++ b/Engine/Source/Editor/Sequencer/Private/Sequencer.cpp
@@ -1683,4 +1683,12 @@ void FSequencer::BuildObjectBindingContextMenu(FMenuBuilder& MenuBuilder, const
 	}
 }
 
+void FSequencer::BuildObjectBindingEditButtons(TSharedPtr<SHorizontalBox> EditBox, const FGuid& ObjectBinding, const UClass* ObjectClass)
+{
+	for (int32 i = 0; i < TrackEditors.Num(); ++i)
+	{
+		TrackEditors[i]->BuildObjectBindingEditButtons(EditBox, ObjectBinding, ObjectClass);
+	}
+}
+
 #undef LOCTEXT_NAMESPACE
diff --git a/Engine/Source/Editor/Sequencer/Private/Sequencer.h b/Engine/Source/Editor/Sequencer/Private/Sequencer.h
--- a/Engine/Source/Editor/Sequencer/Private/Sequencer.h
+++ b/Engine/Source/Editor/Sequencer/Private/Sequencer.h
@@ -221,6 +221,15 @@ public:
 	 */
 	void BuildObjectBindingContextMenu(class FMenuBuilder& MenuBuilder, const FGuid& ObjectBinding, const class UClass* ObjectClass);
 
+	/**
+	 * Builds up the edit buttons for object binding nodes in the outliner
+	 * 
+	 * @param EditBox	    The edit box to add things to
+	 * @param ObjectBinding	The object binding of the selected node
+	 * @param ObjectClass	The class of the selected object
+	 */
+	void BuildObjectBindingEditButtons(TSharedPtr<SHorizontalBox> EditBox, const FGuid& ObjectBinding, const class UClass* ObjectClass);
+
 	/** IMovieScenePlayer interface */
 	virtual void GetRuntimeObjects( TSharedRef<FMovieSceneInstance> MovieSceneInstance, const FGuid& ObjectHandle, TArray< UObject* >& OutObjects ) const override;
 	virtual void UpdatePreviewViewports(UObject* ObjectToViewThrough) const override;
diff --git a/Engine/Source/Editor/Sequencer/Private/SequencerDisplayNode.cpp b/Engine/Source/Editor/Sequencer/Private/SequencerDisplayNode.cpp
--- a/Engine/Source/Editor/Sequencer/Private/SequencerDisplayNode.cpp
+++ b/Engine/Source/Editor/Sequencer/Private/SequencerDisplayNode.cpp
@@ -492,9 +492,32 @@ bool FObjectBindingNode::GetShotFilteredVisibilityToCache() const
 	return !GetSequencer().IsShotFilteringOn() || GetSequencer().IsObjectUnfilterable(ObjectBinding) || HasVisibleChildren();
 }
 
+const UClass* FObjectBindingNode::GetClassForObjectBinding()
+{
+	FSequencer& ParentSequencer = GetSequencer();
+
+	UMovieScene* MovieScene = GetSequencer().GetFocusedMovieScene();
+
+	FMovieSceneSpawnable* Spawnable = MovieScene->FindSpawnable(ObjectBinding);
+	FMovieScenePossessable* Possessable = MovieScene->FindPossessable(ObjectBinding);
+	
+	// should exist, but also shouldn't be both a spawnable and a possessable
+	check((Spawnable != NULL) ^ (Possessable != NULL));
+	check((NULL == Spawnable) || (NULL != Spawnable->GetClass())  );
+	const UClass* ObjectClass = Spawnable ? Spawnable->GetClass()->GetSuperClass() : Possessable->GetPossessedObjectClass();
+
+	return ObjectClass;
+}
+
 TSharedRef<SWidget> FObjectBindingNode::GenerateEditWidgetForOutliner()
 {
-	return
+	// Create a container edit box
+	TSharedPtr<class SHorizontalBox> EditBox;
+	SAssignNew(EditBox, SHorizontalBox);	
+
+	// Add the property combo box
+	EditBox.Get()->AddSlot()
+	[
 		SNew(SComboButton)
 		.ButtonStyle(FEditorStyle::Get(), "FlatButton.Light")
 		.OnGetMenuContent(this, &FObjectBindingNode::OnGetAddPropertyTrackMenuContent)
@@ -520,22 +543,19 @@ TSharedRef<SWidget> FObjectBindingNode::GenerateEditWidgetForOutliner()
 				.Font(FEditorStyle::GetFontStyle("Sequencer.AnimationOutliner.RegularFont"))
 				.Text(LOCTEXT("AddPropertyButton", "Property"))
 			]
-		];
+		]
+	];
+
+	const UClass* ObjectClass = GetClassForObjectBinding();
+
+	GetSequencer().BuildObjectBindingEditButtons(EditBox, ObjectBinding, ObjectClass);
+
+	return EditBox.ToSharedRef();
 }
 
 TSharedPtr<SWidget> FObjectBindingNode::OnSummonContextMenu(const FGeometry& MyGeometry, const FPointerEvent& MouseEvent)
 {
-	FSequencer& ParentSequencer = GetSequencer();
-
-	UMovieScene* MovieScene = GetSequencer().GetFocusedMovieScene();
-
-	FMovieSceneSpawnable* Spawnable = MovieScene->FindSpawnable(ObjectBinding);
-	FMovieScenePossessable* Possessable = MovieScene->FindPossessable(ObjectBinding);
-	
-	// should exist, but also shouldn't be both a spawnable and a possessable
-	check((Spawnable != NULL) ^ (Possessable != NULL));
-	check((NULL == Spawnable) || (NULL != Spawnable->GetClass())  );
-	const UClass* ObjectClass = Spawnable ? Spawnable->GetClass()->GetSuperClass() : Possessable->GetPossessedObjectClass();
+	const UClass* ObjectClass = GetClassForObjectBinding();
 
 	// @todo sequencer replace with UI Commands instead of faking it
 	const bool bShouldCloseWindowAfterMenuSelection = true;
diff --git a/Engine/Source/Editor/Sequencer/Private/SequencerDisplayNode.h b/Engine/Source/Editor/Sequencer/Private/SequencerDisplayNode.h
--- a/Engine/Source/Editor/Sequencer/Private/SequencerDisplayNode.h
+++ b/Engine/Source/Editor/Sequencer/Private/SequencerDisplayNode.h
@@ -390,6 +390,9 @@ private:
 
 	void AddTrackForProperty(TArray<UProperty*> PropertyPath);
 
+	/** Get class for object binding */
+	const UClass* GetClassForObjectBinding();
+
 	/** The binding to live objects */
 	FGuid ObjectBinding;
 	/** The default display name of the object which is used if the binding manager doesn't provide one. */
diff --git a/Engine/Source/Editor/Sequencer/Public/MovieSceneTrackEditor.h b/Engine/Source/Editor/Sequencer/Public/MovieSceneTrackEditor.h
--- a/Engine/Source/Editor/Sequencer/Public/MovieSceneTrackEditor.h
+++ b/Engine/Source/Editor/Sequencer/Public/MovieSceneTrackEditor.h
@@ -15,7 +15,7 @@ DECLARE_DELEGATE_OneParam(FOnKeyProperty, float)
  * Base class for handling key and section drawing and manipulation for a UMovieSceneTrack class
  * @todo Sequencer - Interface needs cleanup
  */
-class FMovieSceneTrackEditor
+class FMovieSceneTrackEditor : public TSharedFromThis<FMovieSceneTrackEditor>
 {
 public:
 	/** Constructor */
@@ -183,6 +183,15 @@ public:
 	 */
 	virtual void BuildObjectBindingContextMenu(FMenuBuilder& MenuBuilder, const FGuid& ObjectBinding, const UClass* ObjectClass) {}
 
+	/**
+	 * Builds up the object binding edit buttons for the outliner
+	 *
+	 * @param EditBox	    The edit box to add buttons to
+	 * @param ObjectBinding	The object binding this is for
+	 * @param ObjectClass	The class of the object this is for
+	 */
+	virtual void BuildObjectBindingEditButtons(TSharedPtr<SHorizontalBox> EditBox, const FGuid& ObjectBinding, const UClass* ObjectClass) {}
+
 private:
 	/** The sequencer bound to this handler.  Used to access movie scene and time info during auto-key */
 	TWeakPtr<ISequencer> Sequencer;
