Fix deprecation warnings

diff --git a/Engine/Source/Editor/Sequencer/Private/Sequencer.cpp b/Engine/Source/Editor/Sequencer/Private/Sequencer.cpp
--- a/Engine/Source/Editor/Sequencer/Private/Sequencer.cpp
+++ b/Engine/Source/Editor/Sequencer/Private/Sequencer.cpp
@@ -739,7 +739,7 @@ void FSequencer::UpdatePreviewViewports(UObject* ObjectToViewThrough) const
 
 	for(FLevelEditorViewportClient* LevelVC : GEditor->LevelViewportClients)
 	{
-		if(LevelVC && LevelVC->IsPerspective() && LevelVC->AllowMatineePreview())
+		if(LevelVC && LevelVC->IsPerspective() && LevelVC->AllowsCinematicPreview())
 		{
 			LevelVC->SetMatineeActorLock(Cast<AActor>(ObjectToViewThrough));
 		}
diff --git a/Engine/Source/Editor/Sequencer/Private/SequencerAssetEditor.cpp b/Engine/Source/Editor/Sequencer/Private/SequencerAssetEditor.cpp
--- a/Engine/Source/Editor/Sequencer/Private/SequencerAssetEditor.cpp
+++ b/Engine/Source/Editor/Sequencer/Private/SequencerAssetEditor.cpp
@@ -245,7 +245,7 @@ EVisibility FSequencerAssetEditor::GetTransportControlVisibility(TWeakPtr<ILevel
 	TSharedPtr<ILevelViewport> LevelViewportPin = LevelViewport.Pin();
 
 	FLevelEditorViewportClient& ViewportClient = LevelViewportPin->GetLevelViewportClient();
-	return (ViewportClient.IsPerspective() && ViewportClient.AllowMatineePreview()) ? EVisibility::Visible : EVisibility::Collapsed;
+	return (ViewportClient.IsPerspective() && ViewportClient.AllowsCinematicPreview()) ? EVisibility::Visible : EVisibility::Collapsed;
 }
 
 
