Adjust padding on sequencer track rows

diff --git a/Engine/Source/Editor/Sequencer/Private/SequencerTrackLaneFactory.cpp b/Engine/Source/Editor/Sequencer/Private/SequencerTrackLaneFactory.cpp
--- a/Engine/Source/Editor/Sequencer/Private/SequencerTrackLaneFactory.cpp
+++ b/Engine/Source/Editor/Sequencer/Private/SequencerTrackLaneFactory.cpp
@@ -95,11 +95,14 @@ void FSequencerTrackLaneFactory::GenerateWidgetForNode( FSequencerDisplayNode& N
 	TSharedRef<SWidget> OutlinerWidget = Node.GenerateContainerWidgetForOutliner();
 	TSharedRef<SWidget> TrackWidget = Node.GenerateWidgetForSectionArea( FAnimatedRange::WrapAttribute(ViewRange) );
 
-	// Link the height of the widets using a size linker
+	// Link the height of the widgets using a size linker
 	FWidgetSizeLinker Link;
 	Link.Widgets.Add(OutlinerWidget);
 	Link.Widgets.Add(TrackWidget);
 
+
+	FMargin Padding = Node.GetType() == ESequencerNode::Object ? FMargin( 0.0f, 5.0f, 0.0f, 0.0f ) : FMargin( 0 );
+
 	auto GetMaxSize = [=]() -> float {
 		return Link.GetMaxSize().Y;
 	};
@@ -109,6 +112,7 @@ void FSequencerTrackLaneFactory::GenerateWidgetForNode( FSequencerDisplayNode& N
 
 	TrackOutliner->AddSlot()
 	.AutoHeight()
+	.Padding( Padding )
 	[
 		SNew(SBox)
 		.HeightOverride_Lambda(GetMaxSize)
@@ -119,6 +123,7 @@ void FSequencerTrackLaneFactory::GenerateWidgetForNode( FSequencerDisplayNode& N
 
 	TrackArea->AddSlot()
 	.AutoHeight()
+	.Padding( Padding )
 	[
 		SNew(SBox)
 		.HeightOverride_Lambda(GetMaxSize)
