Fixed a number of shutdown and map load crashes in Sequencer
- Do not cache off the world in the actor binding manager
- Ensure tracks are fully updated when a map changes
- Properly clean up spawnable actors when a world is torn down

diff --git a/Engine/Source/Editor/Sequencer/Private/Sequencer.cpp b/Engine/Source/Editor/Sequencer/Private/Sequencer.cpp
--- a/Engine/Source/Editor/Sequencer/Private/Sequencer.cpp
+++ b/Engine/Source/Editor/Sequencer/Private/Sequencer.cpp
@@ -441,6 +441,8 @@ void FSequencer::OnMapChanged( class UWorld* NewWorld, EMapChangeType::Type MapC
 	if( ( MapChangeType == EMapChangeType::LoadMap || MapChangeType == EMapChangeType::NewMap ) )
 	{
 		SpawnOrDestroyPuppetObjects( GetFocusedMovieSceneInstance() );
+
+		NotifyMovieSceneDataChanged();
 	}
 }
 
diff --git a/Engine/Source/Editor/Sequencer/Private/SequencerActorBindingManager.cpp b/Engine/Source/Editor/Sequencer/Private/SequencerActorBindingManager.cpp
--- a/Engine/Source/Editor/Sequencer/Private/SequencerActorBindingManager.cpp
+++ b/Engine/Source/Editor/Sequencer/Private/SequencerActorBindingManager.cpp
@@ -13,12 +13,12 @@
 #include "SNotificationList.h"
 #include "NotificationManager.h"
 #include "Engine/Selection.h"
+#include "LevelEditor.h"
 
 #define LOCTEXT_NAMESPACE "Sequencer"
 
-FSequencerActorBindingManager::FSequencerActorBindingManager( UWorld* InWorld, TSharedRef<ISequencerObjectChangeListener> InObjectChangeListener, TSharedRef<FSequencer> InSequencer )
-	: ActorWorld( InWorld )
-	, PlayMovieSceneNode( nullptr )
+FSequencerActorBindingManager::FSequencerActorBindingManager( TSharedRef<ISequencerObjectChangeListener> InObjectChangeListener, TSharedRef<FSequencer> InSequencer )
+	: PlayMovieSceneNode( nullptr )
 	, Sequencer( InSequencer )
 	, ObjectChangeListener( InObjectChangeListener )
 {
@@ -89,6 +89,7 @@ void FSequencerActorBindingManager::SpawnOrDestroyObjectsForInstance( TSharedRef
 
 	UMovieScene* MovieScene = MovieSceneInstance->GetMovieScene();
 
+	UWorld* ActorWorld = GetWorld();
 	// Remove any puppet objects that we no longer need
 	if( MovieScene )
 	{
@@ -166,6 +167,7 @@ void FSequencerActorBindingManager::SpawnOrDestroyObjectsForInstance( TSharedRef
 
 	if( !bDestroyAll && MovieScene )
 	{
+
 		for( auto SpawnableIndex = 0; SpawnableIndex < MovieScene->GetSpawnableCount(); ++SpawnableIndex )
 		{
 			FMovieSceneSpawnable& Spawnable = MovieScene->GetSpawnable( SpawnableIndex );
@@ -334,6 +336,8 @@ UK2Node_PlayMovieScene* FSequencerActorBindingManager::FindPlayMovieSceneNodeInL
 	// Grab the world object for this editor
 	check( MovieScene != NULL );
 	
+	UWorld* ActorWorld = GetWorld();
+
 	// Search all levels in the specified world
 	for( TArray<ULevel*>::TConstIterator LevelIter( ActorWorld->GetLevels().CreateConstIterator() ); LevelIter; ++LevelIter )
 	{
@@ -369,6 +373,8 @@ UK2Node_PlayMovieScene* FSequencerActorBindingManager::CreateNewPlayMovieSceneNo
 	// Grab the world object for this editor
 	check( MovieScene != NULL );
 	
+	UWorld* ActorWorld = GetWorld();
+
 	ULevel* Level = ActorWorld->GetCurrentLevel();
 	check( Level != NULL );
 	
@@ -484,6 +490,17 @@ UK2Node_PlayMovieScene* FSequencerActorBindingManager::BindToPlayMovieSceneNode(
 	return PlayMovieSceneNode.Get();
 }
 
+UWorld* FSequencerActorBindingManager::GetWorld() const
+{
+	if( FModuleManager::Get().IsModuleLoaded("LevelEditor") )
+	{
+		FLevelEditorModule& LevelEditorModule = FModuleManager::Get().GetModuleChecked<FLevelEditorModule>("LevelEditor");
+		return LevelEditorModule.GetFirstLevelEditor()->GetWorld();
+	}
+
+	return nullptr;
+}
+
 void FSequencerActorBindingManager::OnPlayMovieSceneBindingsChanged()
 {
 	const bool bDestroyAll = false;
diff --git a/Engine/Source/Editor/Sequencer/Private/SequencerActorBindingManager.h b/Engine/Source/Editor/Sequencer/Private/SequencerActorBindingManager.h
--- a/Engine/Source/Editor/Sequencer/Private/SequencerActorBindingManager.h
+++ b/Engine/Source/Editor/Sequencer/Private/SequencerActorBindingManager.h
@@ -70,7 +70,7 @@ public:
 	 *
 	 * @param	InWorld	The world to spawn actors in
 	 */
-	FSequencerActorBindingManager( UWorld* InWorld, TSharedRef<ISequencerObjectChangeListener> InObjectChangeListener, TSharedRef<FSequencer> InSequencer );
+	FSequencerActorBindingManager( TSharedRef<ISequencerObjectChangeListener> InObjectChangeListener, TSharedRef<FSequencer> InSequencer );
 
 
 	/** Destructor */
@@ -148,9 +148,11 @@ protected:
 	 */
 	class UK2Node_PlayMovieScene* BindToPlayMovieSceneNode( const bool bCreateIfNotFound ) const;
 	
+	/**
+	 * @return The world that possessed actors belong to or where spawned actors should go
+	 */
+	UWorld* GetWorld() const;
 private:
-	TWeakObjectPtr< UWorld > ActorWorld;
-	
 	/** 'PlayMovieScene' node in level script that we're currently associating with the active MovieScene.  This node contains the bindings
 	 information that we need to be able to preview possessed actors in the level */
 	mutable TWeakObjectPtr< UK2Node_PlayMovieScene > PlayMovieSceneNode;
diff --git a/Engine/Source/Editor/Sequencer/Private/SequencerAssetEditor.cpp b/Engine/Source/Editor/Sequencer/Private/SequencerAssetEditor.cpp
--- a/Engine/Source/Editor/Sequencer/Private/SequencerAssetEditor.cpp
+++ b/Engine/Source/Editor/Sequencer/Private/SequencerAssetEditor.cpp
@@ -73,7 +73,7 @@ void FSequencerAssetEditor::InitSequencerAssetEditor( const EToolkitMode::Type M
 	SequencerInitParams.ViewParams = InViewParams;
 	SequencerInitParams.ObjectChangeListener = MakeShareable( new FSequencerObjectChangeListener( Sequencer.ToSharedRef(), bEditWithinLevelEditor ) );
 	
-	SequencerInitParams.ObjectBindingManager = MakeShareable( new FSequencerActorBindingManager( InitToolkitHost->GetWorld(), SequencerInitParams.ObjectChangeListener.ToSharedRef(), Sequencer.ToSharedRef() ) );
+	SequencerInitParams.ObjectBindingManager = MakeShareable( new FSequencerActorBindingManager( SequencerInitParams.ObjectChangeListener.ToSharedRef(), Sequencer.ToSharedRef() ) );
 
 	SequencerInitParams.RootMovieScene = InRootMovieScene;
 	
diff --git a/Engine/Source/Editor/UnrealEd/Private/EditorServer.cpp b/Engine/Source/Editor/UnrealEd/Private/EditorServer.cpp
--- a/Engine/Source/Editor/UnrealEd/Private/EditorServer.cpp
+++ b/Engine/Source/Editor/UnrealEd/Private/EditorServer.cpp
@@ -1797,6 +1797,15 @@ void UEditorEngine::CheckForWorldGCLeaks( UWorld* NewWorld, UPackage* WorldPacka
 
 void UEditorEngine::EditorDestroyWorld( FWorldContext & Context, const FText& CleanseText, UWorld* NewWorld )
 {
+	if( FModuleManager::Get().IsModuleLoaded("LevelEditor") )
+	{
+		FLevelEditorModule& LevelEditor = FModuleManager::GetModuleChecked<FLevelEditorModule>("LevelEditor");
+
+		// Notify level editors of the map change
+		LevelEditor.BroadcastMapChanged( Context.World(), EMapChangeType::TearDownWorld );
+	}
+
+
 	UWorld* ContextWorld = Context.World();
 
 	if (ContextWorld == NULL )
@@ -2510,7 +2519,7 @@ bool UEditorEngine::Map_Load(const TCHAR* Str, FOutputDevice& Ar)
 					{
 						FLevelEditorModule& LevelEditor = FModuleManager::GetModuleChecked<FLevelEditorModule>("LevelEditor");
 
-						// Notify slate level editors of the map change
+						// Notify level editors of the map change
 						LevelEditor.BroadcastMapChanged( Context.World(), EMapChangeType::LoadMap );
 					}
 
