Sequencer: Add key to nearest section and not all sections.

diff --git a/Engine/Source/Editor/Sequencer/Private/SAnimationOutlinerView.cpp b/Engine/Source/Editor/Sequencer/Private/SAnimationOutlinerView.cpp
--- a/Engine/Source/Editor/Sequencer/Private/SAnimationOutlinerView.cpp
+++ b/Engine/Source/Editor/Sequencer/Private/SAnimationOutlinerView.cpp
@@ -7,6 +7,7 @@
 #include "Sequencer.h"
 #include "MovieScene.h"
 #include "MovieSceneSection.h"
+#include "MovieSceneCommonHelpers.h"
 #include "Engine/Selection.h"
 #include "ISequencerObjectBindingManager.h"
 #include "IKeyArea.h"
@@ -253,14 +254,32 @@ FReply SAnimationOutlinerTreeNode::OnAddKeyClicked()
 	TSet<TSharedPtr<IKeyArea>> KeyAreas;
 	GetAllKeyAreas(DisplayNode, KeyAreas);
 
+	TArray<UMovieSceneSection*> KeyAreaSections;
+	for (TSharedPtr<IKeyArea> KeyArea : KeyAreas)
+	{
+		UMovieSceneSection* OwningSection = KeyArea->GetOwningSection();
+		KeyAreaSections.Add(OwningSection);
+	}
+
+	UMovieSceneSection* NearestSection = MovieSceneHelpers::FindNearestSectionAtTime(KeyAreaSections, CurrentTime);
+	if (!NearestSection)
+	{
+		return FReply::Unhandled();
+	}
+
 	FScopedTransaction Transaction(LOCTEXT("AddKeys", "Add keys at current time"));
 	for (TSharedPtr<IKeyArea> KeyArea : KeyAreas)
 	{
 		UMovieSceneSection* OwningSection = KeyArea->GetOwningSection();
-		OwningSection->SetFlags(RF_Transactional);
-		OwningSection->Modify();
-		KeyArea->AddKeyUnique(CurrentTime);
+		if (OwningSection == NearestSection)
+		{
+			OwningSection->SetFlags(RF_Transactional);
+			OwningSection->Modify();
+			KeyArea->AddKeyUnique(CurrentTime);
+			break;
+		}
 	}
+
 	return FReply::Handled();
 }
 
