Fixed the section selection highlight in sequencer so that it properly matches the contour of a section

diff --git a/Engine/Source/Editor/EditorStyle/Private/SlateEditorStyle.cpp b/Engine/Source/Editor/EditorStyle/Private/SlateEditorStyle.cpp
--- a/Engine/Source/Editor/EditorStyle/Private/SlateEditorStyle.cpp
+++ b/Engine/Source/Editor/EditorStyle/Private/SlateEditorStyle.cpp
@@ -1596,6 +1596,7 @@ void FSlateEditorStyle::FStyle::SetupGeneralStyles()
 		Set( "Sequencer.ToggleAutoScroll.Small", new IMAGE_BRUSH( "Icons/icon_Sequencer_ToggleAutoScroll_16x", Icon16x16 ) );
 		Set( "Sequencer.LockCamera", new IMAGE_BRUSH( "Icons/icon_ToggleActorPilotCameraView_16x", Icon16x16 ) );
 		Set( "Sequencer.UnlockCamera", new IMAGE_BRUSH( "Icons/icon_ToggleActorPilotCameraView_16x", Icon16x16, FLinearColor(1.f, 1.f, 1.f, 0.5f) ) );
+		Set( "Sequencer.Section.SelectionBorder", new BORDER_BRUSH( TEXT("Sequencer/SectionHighlight"), FMargin(7.f/16.f) ) );
 
 		const FButtonStyle DetailsKeyButton = FButtonStyle(NoBorder)
 			.SetNormal( IMAGE_BRUSH("Sequencer/AddKey_Details", FVector2D(11,11) )  )
diff --git a/Engine/Source/Editor/Sequencer/Private/SSection.cpp b/Engine/Source/Editor/Sequencer/Private/SSection.cpp
--- a/Engine/Source/Editor/Sequencer/Private/SSection.cpp
+++ b/Engine/Source/Editor/Sequencer/Private/SSection.cpp
@@ -219,8 +219,7 @@ int32 SSection::OnPaint( const FPaintArgs& Args, const FGeometry& AllottedGeomet
 
 	// Ask the interface to draw the section
 	int32 PostSectionLayer = SectionInterface->OnPaintSection( SectionGeometry, SectionClipRect, OutDrawElements, LayerId, bParentEnabled );
-
-	// @todo Sequencer - Temp indicators of the area of a section which can be dragged.
+	
 	DrawSectionBorders(AllottedGeometry, MyClippingRect, OutDrawElements, PostSectionLayer );
 
 	PaintKeys( SectionGeometry, MyClippingRect, OutDrawElements, PostSectionLayer, InWidgetStyle );
@@ -411,13 +410,13 @@ void SSection::DrawSectionBorders( const FGeometry& AllottedGeometry, const FSla
 	// draw selection box
 	if(bSelected)
 	{
-		static const FName PlainBorder("PlainBorder");
+		static const FName SelectionBorder("Sequencer.Section.SelectionBorder");
 
 		FSlateDrawElement::MakeBox(
 			OutDrawElements,
-			LayerId,
+			LayerId+1,
 			AllottedGeometry.ToPaintGeometry(),
-			FEditorStyle::GetBrush(PlainBorder),
+			FEditorStyle::GetBrush(SelectionBorder),
 			MyClippingRect,
 			ESlateDrawEffect::None,
 			bActive ? SelectionColor : SelectionInactiveColor
@@ -667,6 +666,8 @@ FReply SSection::OnMouseMove( const FGeometry& MyGeometry, const FPointerEvent&
 
 void SSection::OnMouseLeave( const FPointerEvent& MouseEvent )
 {
+	SCompoundWidget::OnMouseLeave( MouseEvent );
+
 	if( !HasMouseCapture() )
 	{
 		ResetHoveredState();
