MovieScene: Made puppet types private

diff --git a/Engine/Source/Editor/Sequencer/Private/SequencerActorBindingManager.cpp b/Engine/Source/Editor/Sequencer/Private/SequencerActorBindingManager.cpp
--- a/Engine/Source/Editor/Sequencer/Private/SequencerActorBindingManager.cpp
+++ b/Engine/Source/Editor/Sequencer/Private/SequencerActorBindingManager.cpp
@@ -18,6 +18,58 @@
 
 #define LOCTEXT_NAMESPACE "Sequencer"
 
+
+namespace EPuppetObjectType
+{
+	enum Type
+	{
+		/** Puppet actor */
+		Actor,
+	};
+}
+
+
+/**
+ * Stores the relationship between a spawned puppet object (e.g. actor) and it's data descriptor (e.g. blueprint)
+ */
+class FBasePuppetInfo
+{	
+
+public:
+
+	/** @return Gets the type of puppet */
+	virtual EPuppetObjectType::Type GetType() const = 0;
+};
+
+
+/**
+ * Puppet actor info
+ */
+struct FPuppetActorInfo : public FBasePuppetInfo
+{	
+
+public:
+
+	virtual ~FPuppetActorInfo() {}
+
+	/** @return Gets the type of puppet */
+	virtual EPuppetObjectType::Type GetType() const override
+	{
+		return EPuppetObjectType::Actor;
+	}
+
+
+public:
+
+	/** Spawnable guid */
+	FGuid SpawnableGuid;
+
+	/** The actual puppet actor that we are actively managing.  These actors actually live in the editor level
+	    just like any other actor, but we'll directly manage their lifetimes while in the editor */
+	TWeakObjectPtr< class AActor > PuppetActor;	
+};
+
+
 FSequencerActorBindingManager::FSequencerActorBindingManager( TSharedRef<ISequencerObjectChangeListener> InObjectChangeListener, TSharedRef<FSequencer> InSequencer )
 	: PlayMovieSceneNode( nullptr )
 	, Sequencer( InSequencer )
diff --git a/Engine/Source/Editor/Sequencer/Private/SequencerActorBindingManager.h b/Engine/Source/Editor/Sequencer/Private/SequencerActorBindingManager.h
--- a/Engine/Source/Editor/Sequencer/Private/SequencerActorBindingManager.h
+++ b/Engine/Source/Editor/Sequencer/Private/SequencerActorBindingManager.h
@@ -5,56 +5,8 @@
 #include "Editor/Sequencer/Public/ISequencerObjectBindingManager.h"
 
 class FSequencer;
+struct FPuppetActorInfo;
 
-namespace EPuppetObjectType
-{
-	enum Type
-	{
-		/** Puppet actor */
-		Actor,
-	};
-}
-
-
-/**
- * Stores the relationship between a spawned puppet object (e.g. actor) and it's data descriptor (e.g. blueprint)
- */
-class FBasePuppetInfo
-{	
-
-public:
-
-	/** @return Gets the type of puppet */
-	virtual EPuppetObjectType::Type GetType() const = 0;
-};
-
-
-/**
- * Puppet actor info
- */
-struct FPuppetActorInfo : public FBasePuppetInfo
-{	
-
-public:
-
-	virtual ~FPuppetActorInfo() {}
-
-	/** @return Gets the type of puppet */
-	virtual EPuppetObjectType::Type GetType() const override
-	{
-		return EPuppetObjectType::Actor;
-	}
-
-
-public:
-
-	/** Spawnable guid */
-	FGuid SpawnableGuid;
-
-	/** The actual puppet actor that we are actively managing.  These actors actually live in the editor level
-	    just like any other actor, but we'll directly manage their lifetimes while in the editor */
-	TWeakObjectPtr< class AActor > PuppetActor;	
-};
 
 
 /**
