MovieScene: Started to move refactored/cleaned up legacy code into binding manager replacement

diff --git a/Engine/Source/Editor/Sequencer/Private/SequencerBindingManager.cpp b/Engine/Source/Editor/Sequencer/Private/SequencerBindingManager.cpp
--- a/Engine/Source/Editor/Sequencer/Private/SequencerBindingManager.cpp
+++ b/Engine/Source/Editor/Sequencer/Private/SequencerBindingManager.cpp
@@ -5,49 +5,533 @@
 #include "MovieSceneInstance.h"
 #include "SequencerBindingManager.h"
 
+// @todo sequencer: gmp: refactor these dependencies into another module
+#include "ModuleManager.h"
+#include "LevelEditor.h"
+
+
+/* Private types
+ *****************************************************************************/
+
+/**
+ * Puppet actor info
+ */
+struct FSequencerActorInfo
+{
+	/** Unique object identifier. */
+	FMovieSceneObjectId ObjectId;
+
+	/**
+	 * The actual puppet actor that we are actively managing.
+	 * These actors actually live in the editor level just like any other actor,
+	 * but we'll directly manage their lifetimes while in the editor.
+	 */
+	TWeakObjectPtr<AActor> Actor;	
+};
+
 
 /* USequencerBindingManager structors
  *****************************************************************************/
 
 USequencerBindingManager::USequencerBindingManager(const FObjectInitializer& ObjectInitializer)
 	: Super(ObjectInitializer)
-{ }
+	//, ObjectChangeListener(InObjectChangeListener)
+{
+	// @todo sequencer: gmp: initialize ObjectChangeListener
+	//ObjectChangeListener->GetOnPropagateObjectChanges().AddUObject(this, &USequencerBindingManager::HandlePropagateObjectChanges);
+}
+
+
+USequencerBindingManager::~USequencerBindingManager()
+{
+	if (ObjectChangeListener.IsValid())
+	{
+		ObjectChangeListener.Pin()->GetOnPropagateObjectChanges().RemoveAll(this);
+	}
+}
 
 
 /* IMovieSceneBindingManager interface
  *****************************************************************************/
 
-void USequencerBindingManager::BindPossessableObject(const FMovieSceneObjectId& PossessableGuid, UObject& PossessedObject)
+bool USequencerBindingManager::AllowsSpawnableObjects() const
 {
+	return true;
+}
 
+void USequencerBindingManager::BindPossessableObject(const FMovieSceneObjectId& ObjectId, UObject& PossessedObject)
+{
+	// @todo sequencer: gmp: bind possessable object
 }
 
 
 bool USequencerBindingManager::CanPossessObject(UObject& Object) const
 {
+	return Object.IsA<AActor>();
+}
+
+
+void USequencerBindingManager::DestroyAllSpawnedObjects()
+{
+	const bool bDestroyAll = true;
+
+	for (auto It = InstanceToActorInfos.CreateConstIterator(); It; ++It)
+	{
+		SpawnOrDestroyObjectsForInstance(It.Key().Pin().ToSharedRef(), bDestroyAll);
+	}
+
+	InstanceToActorInfos.Empty();
+}
+
+
+FMovieSceneObjectId USequencerBindingManager::FindIdForObject(const UMovieScene& MovieScene, UObject& Object) const
+{
+	FMovieSceneObjectId Result = FindIdForSpawnableObject(Object);
+
+	if (Result.IsValid())
+	{
+		return Result;
+	}
+
+	const bool bIsGamePreviewObject = !!(Object.GetOutermost()->PackageFlags & PKG_PlayInEditor);
+
+	if (bIsGamePreviewObject)
+	{
+		// someone is asking for a handle to an object from a game preview session,
+		// probably because they want to capture keys during live simulation.
+		//
+		// Check to see if we already have a puppet that was generated from a
+		// recording of this game preview object.
+		//
+		// @todo sequencer livecapture: We could support recalling counterpart by full name instead of weak pointer, to allow "overdubbing" of previously recorded actors, when the new actors in the current play session have the same path name
+		// @todo sequencer livecapture: Ideally we could capture from editor-world actors that are "puppeteered" as well (real time)
+		const FMovieSceneSpawnable* FoundSpawnable = MovieScene.FindSpawnableForCounterpart(&Object);
+		
+		if (FoundSpawnable != nullptr)
+		{
+			Result = FoundSpawnable->GetObjectId();
+		}
+	}
+	else
+	{
+		// @todo sequencer: gmp: update bindings
+/*		BindToPlayMovieSceneNode(false);
+
+		// when editing within the level editor, make sure we're bound to the level
+		// script node which contains data about possessables.
+		if (PlayMovieSceneNode.IsValid())
+		{
+			ObjectGuid = PlayMovieSceneNode->FindGuidForObject(&Object);
+		}*/
+	}
+
+	return Result;
+}
+
+
+void USequencerBindingManager::GetRuntimeObjects(const TSharedRef<FMovieSceneInstance>& MovieSceneInstance, const FMovieSceneObjectId& ObjectId, TArray<UObject*>& OutRuntimeObjects) const
+{
+	UObject* Object = FindActorInfoById(MovieSceneInstance, ObjectId);
+
+	if (Object)
+	{
+		// spawnable object
+		OutRuntimeObjects.Reset();
+		OutRuntimeObjects.Add(Object);
+	}
+	else
+	{
+		// possessable object
+		// @todo sequenceR: gmp: bind object?
+/*		if (!PlayMovieSceneNode.IsValid())
+		{
+			BindToPlayMovieSceneNode(false);
+		}
+	
+		if (PlayMovieSceneNode.IsValid())
+		{
+			OutRuntimeObjects = PlayMovieSceneNode->FindBoundObjects( ObjectGuid );
+		}*/
+	}
+}
+
+
+void USequencerBindingManager::RemoveMovieSceneInstance(const TSharedRef<FMovieSceneInstance>& MovieSceneInstance)
+{
+	SpawnOrDestroyObjectsForInstance(MovieSceneInstance, true);
+	InstanceToActorInfos.Remove(MovieSceneInstance);
+}
+
+
+void USequencerBindingManager::SpawnOrDestroyObjectsForInstance(TSharedRef<FMovieSceneInstance> MovieSceneInstance, bool bDestroyAll)
+{
+	bool bAnyLevelActorsChanged = false;
+
+	// Get the list of puppet objects for the movie scene
+	TArray<TSharedRef<FSequencerActorInfo>>& ActorInfos = InstanceToActorInfos.FindOrAdd(MovieSceneInstance);
+	UMovieScene* MovieScene = MovieSceneInstance->GetMovieScene();
+	UWorld* ActorWorld = GetWorld();
+
+	// remove any puppet objects that we no longer need
+	if (MovieScene != nullptr)
+	{
+		for (auto ActorInfoIndex = 0; ActorInfoIndex < ActorInfos.Num(); ++ActorInfoIndex)
+		{
+			TSharedRef<FSequencerActorInfo> ActorInfo = ActorInfos[ActorInfoIndex];
+			bool bShouldDestroyActor = true;
+
+			// figure out if we still need this puppet actor
+			if (!bDestroyAll)
+			{
+				for (auto SpawnableIndex = 0; SpawnableIndex < MovieScene->GetSpawnableCount(); ++SpawnableIndex)
+				{
+					auto& Spawnable = MovieScene->GetSpawnable(SpawnableIndex);
+
+					if (Spawnable.GetObjectId() == ActorInfo->ObjectId)
+					{
+						bShouldDestroyActor = false;
+						break;
+					}
+				}
+			}
+
+			if (!bShouldDestroyActor)
+			{
+				// Actor is no longer valid (probably the world was destroyed)
+				continue;
+			}
+
+			AActor* Actor = ActorInfo->Actor.Get();
+				
+			if (Actor == nullptr)
+			{
+				continue;
+			}
+
+			UWorld* PuppetWorld = Actor->GetWorld();
+
+			if (ensure(PuppetWorld != nullptr))
+			{
+				// destroy this actor
+				// Ignored unless called while game is running
+				const bool bNetForce = false;
+
+				// We don't want to dirty the level for puppet actor changes
+				const bool bShouldModifyLevel = false;
+
+				if (!bAnyLevelActorsChanged)
+				{
+					DeselectAllActors();
+				}
+
+				// Actor should never be selected in the editor at this point.  We took care of that up above.
+				ensure(!Actor->IsSelected());
+
+				const bool bWasDestroyed = PuppetWorld->DestroyActor(Actor, bNetForce, bShouldModifyLevel);
+
+				if (bWasDestroyed)
+				{
+					bAnyLevelActorsChanged = true;
+					ActorInfos.RemoveAt(ActorInfoIndex--);
+				}
+				else
+				{
+					// @todo sequencer: At least one puppet couldn't be cleaned up!
+				}
+			}
+		}
+	}
+
+	if( !bDestroyAll && MovieScene )
+	{
+		for (auto SpawnableIndex = 0; SpawnableIndex < MovieScene->GetSpawnableCount(); ++SpawnableIndex)
+		{
+			FMovieSceneSpawnable& Spawnable = MovieScene->GetSpawnable( SpawnableIndex );
+
+			// Must have a valid world for us to be able to do this
+			if (ActorWorld == nullptr )
+			{
+				continue;
+			}
+
+			// Do we already have a puppet for this spawnable?
+			bool bIsAlreadySpawned = false;
+
+			for (auto PuppetIndex = 0; PuppetIndex < ActorInfos.Num(); ++PuppetIndex)
+			{
+				auto& ActorInfo = ActorInfos[PuppetIndex];
+
+				if (ActorInfo->ObjectId == Spawnable.GetObjectId())
+				{
+					bIsAlreadySpawned = true;
+
+					break;
+				}
+			}
+
+			if (bIsAlreadySpawned)
+			{
+				continue;
+			}
+
+			UClass* GeneratedClass = Spawnable.GetClass();
+
+			if ((GeneratedClass == nullptr) || !GeneratedClass->IsChildOf(AActor::StaticClass()))
+			{
+				continue;
+			}
+
+			AActor* ActorCDO = CastChecked<AActor>(GeneratedClass->ClassDefaultObject);
+
+			const FVector SpawnLocation = ActorCDO->GetRootComponent()->RelativeLocation;
+			const FRotator SpawnRotation = ActorCDO->GetRootComponent()->RelativeRotation;
+
+			// @todo sequencer: We should probably spawn these in a specific sub-level!
+			// World->CurrentLevel = ???;
+
+			const FName ActorName = NAME_None;
+
+			// Override the object flags so that RF_Transactional is not set.  Puppet actors are never transactional
+			// @todo sequencer: These actors need to avoid any transaction history.  However, RF_Transactional can currently be set on objects on the fly!
+			const EObjectFlags ObjectFlags = RF_Transient;		// NOTE: We are omitting RF_Transactional intentionally
+
+			// @todo sequencer livecapture: Consider using SetPlayInEditorWorld() and RestoreEditorWorld() here instead
+						
+			// @todo sequencer actors: We need to make sure puppet objects aren't copied into PIE/SIE sessions!  They should be omitted from that duplication!
+
+			// Spawn the puppet actor
+			FActorSpawnParameters SpawnInfo;
+			{
+				SpawnInfo.Name = ActorName;
+				SpawnInfo.ObjectFlags = ObjectFlags;
+			}
+
+			AActor* NewActor = ActorWorld->SpawnActor(GeneratedClass, &SpawnLocation, &SpawnRotation, SpawnInfo);
+				
+			if (NewActor != nullptr)
+			{
+				// @todo sequencer: We're naming the actor based off of the spawnable's name.  Is that really what we want?
+				FActorLabelUtilities::SetActorLabelUnique(NewActor, Spawnable.GetName());
+
+				// Keep track of this actor
+				TSharedRef<FSequencerActorInfo> NewActorInfo(new FSequencerActorInfo());
+				{
+					NewActorInfo->ObjectId = Spawnable.GetObjectId();
+					NewActorInfo->Actor = NewActor;
+					ActorInfos.Add(NewActorInfo);
+				}
+			}
+			else
+			{
+				// Actor failed to spawn
+				// @todo sequencer: What should we do when this happens to one or more actors?
+			}
+		}
+	}
+}
+
+
+bool USequencerBindingManager::TryGetObjectBindingDisplayName(const TSharedRef<FMovieSceneInstance>& MovieSceneInstance, const FMovieSceneObjectId& ObjectId, FText& OutDisplayName) const
+{
+	TArray<UObject*> RuntimeObjects;
+	GetRuntimeObjects(MovieSceneInstance, ObjectId, RuntimeObjects);
+
+	for (int32 ObjIndex = 0; ObjIndex < RuntimeObjects.Num(); ++ObjIndex)
+	{
+		AActor* Actor = Cast<AActor>(RuntimeObjects[ObjIndex]);
+
+		if (Actor != nullptr)
+		{
+			OutDisplayName = FText::FromString(Actor->GetActorLabel());
+
+			return true;
+		}
+	}
+
 	return false;
 }
 
 
-FMovieSceneObjectId USequencerBindingManager::FindGuidForObject(const UMovieScene& MovieScene, UObject& Object) const
+void USequencerBindingManager::UnbindPossessableObjects(const FMovieSceneObjectId& ObjectId)
+{
+	// @todo sequencer: gmp: unbind object
+}
+
+
+/* USequencerBindingManager implementation
+ *****************************************************************************/
+
+void USequencerBindingManager::DeselectAllActors()
+{
+	TArray<AActor*> ActorsToDeselect;
+
+	// find actors to deselect
+	for (auto ActorInfosPair : InstanceToActorInfos)
+	{
+		for (const TSharedRef<FSequencerActorInfo>& ActorInfo : ActorInfosPair.Value)
+		{
+			AActor* Actor = ActorInfo->Actor.Get();
+
+			if (Actor != nullptr)
+			{
+				ActorsToDeselect.Add(Actor);
+			}
+		}
+	}
+
+	if (ActorsToDeselect.Num() == 0)
+	{
+		return;
+	}
+
+	// deselect actors
+	USelection* Selection = GEditor->GetSelectedActors();
+	{
+		Selection->BeginBatchSelectOperation();
+		Selection->MarkBatchDirty();
+
+		for (AActor* Actor : ActorsToDeselect)
+		{
+			GEditor->GetSelectedActors()->Deselect(Actor);
+		}
+
+		Selection->EndBatchSelectOperation();
+	}
+}
+
+
+FMovieSceneObjectId USequencerBindingManager::FindIdForSpawnableObject(UObject& Object) const
 {
+	for (auto ActorInfosPair : InstanceToActorInfos)
+	{
+		for (const TSharedRef<FSequencerActorInfo>& ActorInfo : ActorInfosPair.Value)
+		{
+			AActor* Actor = ActorInfo->Actor.Get();
+
+			if (Actor == &Object)
+			{
+				return ActorInfo->ObjectId;
+			}
+		}
+	}
+
+	// not found
 	return FMovieSceneObjectId();
 }
 
 
-void USequencerBindingManager::GetRuntimeObjects(const TSharedRef<FMovieSceneInstance>& MovieSceneInstance, const FGuid& ObjectGuid, TArray<UObject*>& OutRuntimeObjects) const
+UObject* USequencerBindingManager::FindActorInfoById(const TSharedRef<FMovieSceneInstance>& MovieSceneInstance, const FMovieSceneObjectId& ObjectId) const
 {
+	const TArray<TSharedRef<FSequencerActorInfo>>* ActorInfos = InstanceToActorInfos.Find(MovieSceneInstance);
+
+	if (ActorInfos == nullptr)
+	{
+		return nullptr;
+	}
 
+	for (TSharedRef<FSequencerActorInfo> ActorInfo : *ActorInfos)
+	{
+		if (ActorInfo->ObjectId == ObjectId)
+		{
+			return ActorInfo->Actor.Get();
+		}
+	}
+
+	// not found
+	return nullptr;
 }
 
 
-bool USequencerBindingManager::TryGetObjectBindingDisplayName(const TSharedRef<FMovieSceneInstance>& MovieSceneInstance, const FGuid& ObjectGuid, FText& DisplayName) const
+UWorld* USequencerBindingManager::GetWorld() const
 {
-	return false;
+	if (FModuleManager::Get().IsModuleLoaded("LevelEditor"))
+	{
+		FLevelEditorModule& LevelEditorModule = FModuleManager::Get().GetModuleChecked<FLevelEditorModule>("LevelEditor");
+
+		return LevelEditorModule.GetFirstLevelEditor()->GetWorld();
+	}
+
+	return nullptr;
+}
+
+
+/* USequencerBindingManager internals
+ *****************************************************************************/
+
+void USequencerBindingManager::HandlePropagateObjectChanges(UObject* Object)
+{
+	if (Object == nullptr)
+	{
+		return;
+	}
+
+	// @todo sequencer major: Many other types of changes to the puppet actor may occur which should be propagated (or disallowed?):
+	//		- Deleting the puppet actor (Del key) should probably delete the reference actor and remove it from the MovieScene
+	//		- Attachment changes
+	//		- Replace/Convert operations?
+
+	// We only care about actor objects for this type of spawner.  Note that sometimes we can get component objects passed in, but we'll
+	// rely on the caller also calling OnObjectPropertyChange for the outer Actor if a component property changes.
+	if (!Object->IsA<AActor>())
+	{
+		return;
+	}
+
+	for (auto ActorInfosPair : InstanceToActorInfos)
+	{
+		// is this an object that we care about?
+		for (const TSharedRef<FSequencerActorInfo>& ActorInfo : ActorInfosPair.Value)
+		{
+			AActor* Actor = ActorInfo->Actor.Get();
+
+			if (Actor == Object)
+			{
+				// @todo sequencer: propagation: Don't propagate changes that are being auto-keyed or key-adjusted!
+
+				// Our puppet actor was modified.
+				PropagateActorChanges(ActorInfo);
+			}
+		}
+	}
 }
 
 
-void USequencerBindingManager::UnbindPossessableObjects(const FMovieSceneObjectId& PossessableGuid)
+void USequencerBindingManager::PropagateActorChanges(const TSharedRef<FSequencerActorInfo> ActorInfo)
 {
+	AActor* Actor = ActorInfo->Actor.Get();
+	AActor* TargetActor = nullptr;
+
+	// find the spawnable for this puppet actor
+	FMovieSceneSpawnable* FoundSpawnable = nullptr;
+
+	// @todo sequencer: gmp: add horrible sequencer dependencies to binding manager
+/*	TArray<UMovieScene*> MovieScenesBeingEdited = Sequencer.Pin()->GetMovieScenesBeingEdited();
+
+	for (UMovieScene* MovieScene : MovieScenesBeingEdited)
+	{
+		FoundSpawnable = MovieScene->FindSpawnable(ActorInfo->ObjectId);
+		
+		if (FoundSpawnable != nullptr)
+		{
+			break;
+		}
+	}
+*/	
+	if (ensure((FoundSpawnable != nullptr) && (Actor != nullptr)))
+	{
+		UClass* ActorClass = FoundSpawnable->GetClass();
+
+		// the proxy actor's class should always be the BP that it was spawned from
+		UClass* SpawnedActorClass = Actor->GetClass();
+		check(ActorClass == SpawnedActorClass);
+
+		// we'll be copying properties into the class default object of the Blueprint's generated class
+		TargetActor = CastChecked<AActor>(ActorClass->GetDefaultObject());
+	}
 
+	if ((Actor != nullptr) && (TargetActor != nullptr))
+	{
+		// @todo sequencer: gmp: add horrible sequencer dependencies to binding manager
+//		Sequencer.Pin()->CopyActorProperties(Actor, TargetActor);
+	}
 }
diff --git a/Engine/Source/Editor/Sequencer/Private/SequencerBindingManager.h b/Engine/Source/Editor/Sequencer/Private/SequencerBindingManager.h
--- a/Engine/Source/Editor/Sequencer/Private/SequencerBindingManager.h
+++ b/Engine/Source/Editor/Sequencer/Private/SequencerBindingManager.h
@@ -6,8 +6,13 @@
 #include "SequencerBindingManager.generated.h"
 
 
+struct FSequencerActorInfo;
+
+
 /**
  * Movie scene binding manager used by Sequencer.
+ *
+ * @todo squencer: gmp: probably want to refactor this to not depend on Actors/move into own module
  */
 UCLASS(BlueprintType, MinimalAPI)
 class USequencerBindingManager
@@ -16,14 +21,70 @@ class USequencerBindingManager
 {
 	GENERATED_UCLASS_BODY()
 
+	/** Virtual destructor. */
+	virtual ~USequencerBindingManager();
+
 public:
 
 	// IMovieSceneBindingManager interface
 
-	virtual void BindPossessableObject(const FMovieSceneObjectId& PossessableGuid, UObject& PossessedObject) override;
+	virtual bool AllowsSpawnableObjects() const override;
+	virtual void BindPossessableObject(const FMovieSceneObjectId& ObjectId, UObject& PossessedObject) override;
 	virtual bool CanPossessObject( UObject& Object ) const override;
-	virtual FMovieSceneObjectId FindGuidForObject(const UMovieScene& MovieScene, UObject& Object) const override;
-	virtual void GetRuntimeObjects(const TSharedRef<FMovieSceneInstance>& MovieSceneInstance, const FGuid& ObjectGuid, TArray<UObject*>& OutRuntimeObjects) const override;
-	virtual bool TryGetObjectBindingDisplayName(const TSharedRef<FMovieSceneInstance>& MovieSceneInstance, const FGuid& ObjectGuid, FText& DisplayName) const override;
-	virtual void UnbindPossessableObjects(const FMovieSceneObjectId& PossessableGuid) override;
+	virtual void DestroyAllSpawnedObjects() override;
+	virtual FMovieSceneObjectId FindIdForObject(const UMovieScene& MovieScene, UObject& Object) const override;
+	virtual void GetRuntimeObjects(const TSharedRef<FMovieSceneInstance>& MovieSceneInstance, const FMovieSceneObjectId& ObjectId, TArray<UObject*>& OutRuntimeObjects) const override;
+	virtual void RemoveMovieSceneInstance(const TSharedRef<FMovieSceneInstance>& MovieSceneInstance) override;
+	virtual void SpawnOrDestroyObjectsForInstance(TSharedRef<FMovieSceneInstance> MovieSceneInstance, bool bDestroyAll) override;
+	virtual bool TryGetObjectBindingDisplayName(const TSharedRef<FMovieSceneInstance>& MovieSceneInstance, const FMovieSceneObjectId& ObjectId, FText& OutDisplayName) const override;
+	virtual void UnbindPossessableObjects(const FMovieSceneObjectId& ObjectId) override;
+
+protected:
+
+	/**
+	 * Deselect all proxy actors in the Editor.
+	 *
+	 * @todo sequencer: gmp: remove Editor dependencies from this class
+	 */
+	void DeselectAllActors();
+
+	/**
+	 * Find the unique identifier for a given proxy object.
+	 *
+	 * @param Object The proxy object to find the identifier for.
+	 * @return The object's identifier.
+	 * @see FindPuppetObjectForId
+	 */
+	FMovieSceneObjectId FindIdForSpawnableObject(UObject& Object) const;
+
+	/**
+	 * Find the puppet actor for a given object identifier.
+	 *
+	 * @param ObjectId The object identifier.
+	 * @return The corresponding object, or nullptr if not found.
+	 */
+	UObject* FindActorInfoById(const TSharedRef<FMovieSceneInstance>& MovieSceneInstance, const FMovieSceneObjectId& ObjectId) const;
+
+	/**
+	 * Get the world that possessed actors belong to or where spawned actors should go.
+	 *
+	 * @return The world object.
+	 */
+	UWorld* GetWorld() const;
+
+private:
+
+	/** Handles changes of non-auto-keyed properties on any object. */
+	void HandlePropagateObjectChanges(UObject* Object);
+
+	/** Refresh either the reference actor or the data descriptor BP after its puppet proxy actor has been changed. */
+	void PropagateActorChanges(const TSharedRef<FSequencerActorInfo> ActorInfo);
+
+private:
+
+	/** List of puppet objects that we've spawned and are actively managing. */
+	TMap<TWeakPtr<FMovieSceneInstance>, TArray<TSharedRef<FSequencerActorInfo>>> InstanceToActorInfos;
+
+	/** The object change listener. */
+	TWeakPtr<ISequencerObjectChangeListener> ObjectChangeListener;
 };
diff --git a/Engine/Source/Editor/UMGEditor/Private/Animation/UMGBindingManager.cpp b/Engine/Source/Editor/UMGEditor/Private/Animation/UMGBindingManager.cpp
--- a/Engine/Source/Editor/UMGEditor/Private/Animation/UMGBindingManager.cpp
+++ b/Engine/Source/Editor/UMGEditor/Private/Animation/UMGBindingManager.cpp
@@ -17,9 +17,15 @@ UUMGBindingManager::UUMGBindingManager(const FObjectInitializer& ObjectInitializ
 /* IMovieSceneBindingManager interface
  *****************************************************************************/
 
-void UUMGBindingManager::BindPossessableObject(const FMovieSceneObjectId& PossessableGuid, UObject& PossessedObject)
+bool UUMGBindingManager::AllowsSpawnableObjects() const
 {
+	return false;
+}
+
 
+void UUMGBindingManager::BindPossessableObject(const FMovieSceneObjectId& ObjectId, UObject& PossessedObject)
+{
+	// @todo gmp: sequencer: bind possessable object
 }
 
 
@@ -29,25 +35,43 @@ bool UUMGBindingManager::CanPossessObject(UObject& Object) const
 }
 
 
-FMovieSceneObjectId UUMGBindingManager::FindGuidForObject(const UMovieScene& MovieScene, UObject& Object) const
+void UUMGBindingManager::DestroyAllSpawnedObjects()
+{
+
+}
+
+
+FMovieSceneObjectId UUMGBindingManager::FindIdForObject(const UMovieScene& MovieScene, UObject& Object) const
 {
 	return FMovieSceneObjectId();
 }
 
 
-void UUMGBindingManager::GetRuntimeObjects(const TSharedRef<FMovieSceneInstance>& MovieSceneInstance, const FGuid& ObjectGuid, TArray<UObject*>& OutRuntimeObjects) const
+void UUMGBindingManager::GetRuntimeObjects(const TSharedRef<FMovieSceneInstance>& MovieSceneInstance, const FMovieSceneObjectId& ObjectId, TArray<UObject*>& OutRuntimeObjects) const
 {
 
 }
 
 
-bool UUMGBindingManager::TryGetObjectBindingDisplayName(const TSharedRef<FMovieSceneInstance>& MovieSceneInstance, const FGuid& ObjectGuid, FText& DisplayName) const
+void UUMGBindingManager::RemoveMovieSceneInstance(const TSharedRef<FMovieSceneInstance>& MovieSceneInstance)
 {
-	return false;
+
+}
+
+
+void UUMGBindingManager::SpawnOrDestroyObjectsForInstance(TSharedRef<FMovieSceneInstance> MovieSceneInstance, bool bDestroyAll)
+{
+
 }
 
 
-void UUMGBindingManager::UnbindPossessableObjects(const FMovieSceneObjectId& PossessableGuid)
+bool UUMGBindingManager::TryGetObjectBindingDisplayName(const TSharedRef<FMovieSceneInstance>& MovieSceneInstance, const FMovieSceneObjectId& ObjectId, FText& OutDisplayName) const
 {
+	return false;
+}
 
+
+void UUMGBindingManager::UnbindPossessableObjects(const FMovieSceneObjectId& ObjectId)
+{
+	// @todo gmp: sequencer: bind possessable object
 }
diff --git a/Engine/Source/Editor/UMGEditor/Private/Animation/UMGBindingManager.h b/Engine/Source/Editor/UMGEditor/Private/Animation/UMGBindingManager.h
--- a/Engine/Source/Editor/UMGEditor/Private/Animation/UMGBindingManager.h
+++ b/Engine/Source/Editor/UMGEditor/Private/Animation/UMGBindingManager.h
@@ -20,10 +20,14 @@ public:
 
 	// IMovieSceneBindingManager interface
 
-	virtual void BindPossessableObject(const FMovieSceneObjectId& PossessableGuid, UObject& PossessedObject) override;
+	virtual bool AllowsSpawnableObjects() const override;
+	virtual void BindPossessableObject(const FMovieSceneObjectId& ObjectId, UObject& PossessedObject) override;
 	virtual bool CanPossessObject( UObject& Object ) const override;
-	virtual FMovieSceneObjectId FindGuidForObject(const UMovieScene& MovieScene, UObject& Object) const override;
-	virtual void GetRuntimeObjects(const TSharedRef<FMovieSceneInstance>& MovieSceneInstance, const FGuid& ObjectGuid, TArray<UObject*>& OutRuntimeObjects) const override;
-	virtual bool TryGetObjectBindingDisplayName(const TSharedRef<FMovieSceneInstance>& MovieSceneInstance, const FGuid& ObjectGuid, FText& DisplayName) const override;
-	virtual void UnbindPossessableObjects(const FMovieSceneObjectId& PossessableGuid) override;
+	virtual void DestroyAllSpawnedObjects() override;
+	virtual FMovieSceneObjectId FindIdForObject(const UMovieScene& MovieScene, UObject& Object) const override;
+	virtual void GetRuntimeObjects(const TSharedRef<FMovieSceneInstance>& MovieSceneInstance, const FMovieSceneObjectId& ObjectId, TArray<UObject*>& OutRuntimeObjects) const override;
+	virtual void RemoveMovieSceneInstance(const TSharedRef<FMovieSceneInstance>& MovieSceneInstance) override;
+	virtual void SpawnOrDestroyObjectsForInstance(TSharedRef<FMovieSceneInstance> MovieSceneInstance, bool bDestroyAll) override;
+	virtual bool TryGetObjectBindingDisplayName(const TSharedRef<FMovieSceneInstance>& MovieSceneInstance, const FMovieSceneObjectId& ObjectId, FText& OutDisplayName) const override;
+	virtual void UnbindPossessableObjects(const FMovieSceneObjectId& ObjectId) override;
 };
diff --git a/Engine/Source/Runtime/MovieScene/Private/MovieSceneInstance.cpp b/Engine/Source/Runtime/MovieScene/Private/MovieSceneInstance.cpp
--- a/Engine/Source/Runtime/MovieScene/Private/MovieSceneInstance.cpp
+++ b/Engine/Source/Runtime/MovieScene/Private/MovieSceneInstance.cpp
@@ -3,12 +3,20 @@
 #include "MovieScenePrivatePCH.h"
 #include "MovieSceneInstance.h"
 
+
 FMovieSceneInstance::FMovieSceneInstance( UMovieScene& InMovieScene )
 	: MovieScene( &InMovieScene )
 {
 	TimeRange = MovieScene->GetTimeRange();
 }
 
+
+TScriptInterface<UMovieSceneBindingManager> FMovieSceneInstance::GetBindingManager() const
+{
+	return MovieScene.Get()->GetBindingManager();
+}
+
+
 void FMovieSceneInstance::SaveState()
 {
 	TMap<FGuid, FMovieSceneObjectBindingInstance>::TIterator ObjectIt = ObjectBindingInstances.CreateIterator();
@@ -23,6 +31,7 @@ void FMovieSceneInstance::SaveState()
 	}
 }
 
+
 void FMovieSceneInstance::RestoreState()
 {
 	TMap<FGuid, FMovieSceneObjectBindingInstance>::TIterator ObjectIt = ObjectBindingInstances.CreateIterator();
@@ -37,6 +46,7 @@ void FMovieSceneInstance::RestoreState()
 	}
 }
 
+
 void FMovieSceneInstance::Update( float Position, float LastPosition, class IMovieScenePlayer& Player )
 {
 	// Update  shot track
@@ -65,6 +75,7 @@ void FMovieSceneInstance::Update( float Position, float LastPosition, class IMov
 	}
 }
 
+
 void FMovieSceneInstance::RefreshInstance( IMovieScenePlayer& Player )
 {
 	TimeRange = MovieScene->GetTimeRange();
@@ -181,4 +192,4 @@ void FMovieSceneInstance::RefreshInstanceMap( const TArray<UMovieSceneTrack*>& T
 			It.RemoveCurrent();
 		}
 	}
-}
\ No newline at end of file
+}
diff --git a/Engine/Source/Runtime/MovieScene/Public/MovieScene.h b/Engine/Source/Runtime/MovieScene/Public/MovieScene.h
--- a/Engine/Source/Runtime/MovieScene/Public/MovieScene.h
+++ b/Engine/Source/Runtime/MovieScene/Public/MovieScene.h
@@ -2,6 +2,7 @@
 
 #pragma once
 
+#include "MovieSceneObjectId.h"
 #include "MovieScene.generated.h"
 
 
@@ -25,14 +26,12 @@ struct FMovieSceneSpawnable
 public:
 
 	/** FMovieSceneSpawnable default constructor */
-	FMovieSceneSpawnable()
-	{
-	}
+	FMovieSceneSpawnable() { }
 
 	/** FMovieSceneSpawnable initialization constructor */
 	FMovieSceneSpawnable( const FString& InitName, UClass* InitClass, UObject* InitCounterpartGamePreviewObject );
 
-	/** @return Returns the guid for this possessable */
+	/** @return Returns the guid for this spawnable */
 	const FGuid& GetGuid() const
 	{
 		return Guid;
@@ -44,6 +43,14 @@ public:
 		return Name;
 	}
 
+	/**
+	 * Get the unique object identifier.
+	 */
+	const FMovieSceneObjectId& GetObjectId() const
+	{
+		return ObjectId;
+	}
+
 	/** @return Returns the blueprint associated with this spawnable */
 	UClass* GetClass() { return GeneratedClass; }
 
@@ -75,6 +82,9 @@ private:
 	    used in the editor when capturing keyframe data from a live simulation */
 	// @todo sequencer data: Should be editor only
 	FWeakObjectPtr CounterpartGamePreviewObject;
+
+	// @todo sequencer: gmp: refactor sequencer object IDs.
+	FMovieSceneObjectId ObjectId;
 };
 
 /**
@@ -133,9 +143,9 @@ private:
 	// possessable object available.
 	UPROPERTY()
 	UClass* PossessedObjectClass;
-
 };
 
+
 /**
  * Editor only data that needs to be saved between sessions for editing but has no runtime purpose
  */
@@ -422,6 +432,16 @@ public:
 
 public:
 
+	/**
+	 * Get the movie scene's object binding manager.
+	 *
+	 * @return The scene's binding manager.
+	 */
+	TScriptInterface<UMovieSceneBindingManager> GetBindingManager() const
+	{
+		return BindingManager;
+	}
+
 	/**
 	 * @return All object bindings.
 	 */
diff --git a/Engine/Source/Runtime/MovieScene/Public/MovieSceneBindingManager.h b/Engine/Source/Runtime/MovieScene/Public/MovieSceneBindingManager.h
--- a/Engine/Source/Runtime/MovieScene/Public/MovieSceneBindingManager.h
+++ b/Engine/Source/Runtime/MovieScene/Public/MovieSceneBindingManager.h
@@ -2,6 +2,7 @@
 
 #pragma once
 
+#include "MovieSceneObjectId.h"
 #include "MovieSceneBindingManager.generated.h"
 
 
@@ -10,75 +11,6 @@ class UMovieScene;
 class UObject;
 
 
-/**
- * Identifies an object (or embedded component) bound in a movie scene.
- */
-USTRUCT()
-struct FMovieSceneObjectId
-{
-	GENERATED_USTRUCT_BODY()
-
-	/** Unique identifier of a bound object. */
-	UPROPERTY()
-	FGuid ObjectGuid;
-
-	/** Name of the bound object's component object being used (optional). */
-	UPROPERTY()
-	FName ComponentName;
-
-public:
-
-	/** Default constructor. */
-	FMovieSceneObjectId() { }
-
-	/** Creates and initializes a new instance from an object ID. */
-	FMovieSceneObjectId(const FGuid& InObjectGuid)
-		: ObjectGuid(InObjectGuid)
-	{ }
-
-	/** Creates and initializes a new instance from an object and component ID. */
-	FMovieSceneObjectId(const FGuid& InObjectGuid, const FName& InComponentName)
-		: ObjectGuid(InObjectGuid)
-		, ComponentName(InComponentName)
-	{ }
-
-public:
-
-	/**
-	 * Compares two object identifiers for equality.
-	 *
-	 * @param X The first identifier to compare.
-	 * @param Y The second identifier to compare.
-	 * @return true if the identifiers are equal, false otherwise.
-	 */
-	friend bool operator==( const FMovieSceneObjectId& X, const FMovieSceneObjectId& Y )
-	{
-		return ((X.ObjectGuid == Y.ObjectGuid) && (X.ComponentName == X.ComponentName));
-	}
-
-	/**
-	 * Compares two object identifiers for inequality.
-	 *
-	 * @param X The first identifier to compare.
-	 * @param Y The second identifier to compare.
-	 * @return true if the identifier are not equal, false otherwise.
-	 */
-	friend bool operator!=( const FMovieSceneObjectId& X, const FMovieSceneObjectId& Y )
-	{
-		return ((X.ObjectGuid != Y.ObjectGuid) || (X.ComponentName != X.ComponentName));
-	}
-
-	/**
-	 * Checks whether this object identifier is valid or not.
-	 *
-	 * @return true if valid, false otherwise.
-	 */
-	bool IsValid() const
-	{
-		return ObjectGuid.IsValid();
-	}
-};
-
 
 /**
  * Interface for movie scene object binding managers.
@@ -98,32 +30,58 @@ class IMovieSceneBindingManager
 {
 	GENERATED_IINTERFACE_BODY()
 
+	/**
+	 * Whether objects can be spawned at run-time.
+	 *
+	 * @return true if objects can be spawned by sequencer, false if only existing objects can be possessed.
+	 */
+	virtual bool AllowsSpawnableObjects() const PURE_VIRTUAL(IMovieSceneBindingManager::AllowsSpawnableObjects, return false;);
+
 	/**
 	 * Called when Sequencer has created an object binding for a possessable object
 	 * 
-	 * @param PossessableGuid The guid used to map to the possessable object.  Note the guid can be bound to multiple objects at once
+	 * @param ObjectId The guid used to map to the possessable object.  Note the guid can be bound to multiple objects at once
 	 * @param PossessedObject The runtime object which was possessed.
+	 * @see UnbindPossessableObjects
 	 */
-	virtual void BindPossessableObject(const FMovieSceneObjectId& PossessableGuid, UObject& PossessedObject) PURE_VIRTUAL(IMovieSceneBindingManager::BindPossessableObject,);
+	virtual void BindPossessableObject(const FMovieSceneObjectId& ObjectId, UObject& PossessedObject) PURE_VIRTUAL(IMovieSceneBindingManager::BindPossessableObject,);
 	
 	/**
-	 * @todo sequencer: add documentation
+	 * @todo gmp: sequencer: add documentation.
 	 */
 	virtual bool CanPossessObject( UObject& Object ) const PURE_VIRTUAL(IMovieSceneBindingManager::CanPossessObject, return false;);
 
 	/**
-	 * Finds an existing guid for a bound object
+	 * @todo gmp: sequencer: add documentation
+	 */
+	virtual void DestroyAllSpawnedObjects() PURE_VIRTUAL(IMovieSceneBindingManager::DestroyAllSpawnedObjects,);
+
+	/**
+	 * Finds an existing guid for a bound object.
 	 *
-	 * @param MovieScene The movie scene that may contain the binding
-	 * @param The object to get the guid for
-	 * @return The object identifier or an invalid guid if not found
+	 * @param MovieScene The movie scene that may contain the binding.
+	 * @param The object to get the guid for.
+	 * @return The object identifier or an invalid guid if not found.
+	 */
+	virtual FMovieSceneObjectId FindIdForObject(const UMovieScene& MovieScene, UObject& Object) const PURE_VIRTUAL(IMovieSceneBindingManager::FindGuidForObject, return FGuid(););
+
+	/**
+	 * @todo gmp: sequencer: add documentation
 	 */
-	virtual FMovieSceneObjectId FindGuidForObject(const UMovieScene& MovieScene, UObject& Object) const PURE_VIRTUAL(IMovieSceneBindingManager::FindGuidForObject, return FGuid(););
+	virtual void GetRuntimeObjects(const TSharedRef<FMovieSceneInstance>& MovieSceneInstance, const FMovieSceneObjectId& ObjectId, TArray<UObject*>& OutRuntimeObjects) const PURE_VIRTUAL(IMovieSceneBindingManager::GetRuntimeObjects,);
 
 	/**
-	 * @todo sequencer: add documentation
+	 * @todo gmp: sequencer: add documentation
+	 */
+	virtual void RemoveMovieSceneInstance(const TSharedRef<FMovieSceneInstance>& MovieSceneInstance) PURE_VIRTUAL(IMovieSceneBindingManager::RemoveMovieSceneInstance,);
+
+	/**
+	 * Called to spawn or destroy objects for a movie instance.
+	 *
+	 * @param MovieSceneInstance The instance to spawn or destroy objects for.
+	 * @param bDestroyAll If true, destroy all spawned objects for the instance, if false only destroy unused objects.
 	 */
-	virtual void GetRuntimeObjects(const TSharedRef<FMovieSceneInstance>& MovieSceneInstance, const FGuid& ObjectGuid, TArray<UObject*>& OutRuntimeObjects) const PURE_VIRTUAL(IMovieSceneBindingManager::GetRuntimeObjects,);
+	virtual void SpawnOrDestroyObjectsForInstance( TSharedRef<FMovieSceneInstance> MovieSceneInstance, bool bDestroyAll ) PURE_VIRTUAL(IMovieSceneBindingManager::SpawnOrDestroyObjectsForInstance,);
 
 	/**
 	 * 
@@ -133,12 +91,13 @@ class IMovieSceneBindingManager
 	 * @param DisplayName The display name for the object binding.
 	 * @returns true if DisplayName has been set to a valid display name, otherwise false.
 	 */
-	virtual bool TryGetObjectBindingDisplayName(const TSharedRef<FMovieSceneInstance>& MovieSceneInstance, const FGuid& ObjectGuid, FText& DisplayName) const PURE_VIRTUAL(IMovieSceneBindingManager::TryGetObjectBindingDisplayName, return false;);
+	virtual bool TryGetObjectBindingDisplayName(const TSharedRef<FMovieSceneInstance>& MovieSceneInstance, const FMovieSceneObjectId& ObjectId, FText& OutDisplayName) const PURE_VIRTUAL(IMovieSceneBindingManager::TryGetObjectBindingDisplayName, return false;);
 
 	/**
-	 * Unbinds all possessable objects from the provided guid
+	 * Unbinds all possessable objects from the provided GUID.
 	 *
-	 * @param PossessableGuid The guid bound to possessable objects that should be removed
+	 * @param ObjectId The guid bound to possessable objects that should be removed.
+	 * @see BindPossessableObject
 	 */
-	virtual void UnbindPossessableObjects(const FMovieSceneObjectId& PossessableGuid) PURE_VIRTUAL(IMovieSceneBindingManager::UnbindPossessableObjects,);
+	virtual void UnbindPossessableObjects(const FMovieSceneObjectId& ObjectId) PURE_VIRTUAL(IMovieSceneBindingManager::UnbindPossessableObjects,);
 };
diff --git a/Engine/Source/Runtime/MovieScene/Public/MovieSceneInstance.h b/Engine/Source/Runtime/MovieScene/Public/MovieSceneInstance.h
--- a/Engine/Source/Runtime/MovieScene/Public/MovieSceneInstance.h
+++ b/Engine/Source/Runtime/MovieScene/Public/MovieSceneInstance.h
@@ -2,71 +2,110 @@
 
 #pragma once
 
+
+class UMovieSceneBindingManager;
 class UMovieSceneTrack;
 
-typedef TMap< TWeakObjectPtr<UMovieSceneTrack>, TSharedPtr<class IMovieSceneTrackInstance> > FMovieSceneInstanceMap;
-class FMovieSceneInstance : public TSharedFromThis<FMovieSceneInstance>
+
+typedef TMap<TWeakObjectPtr<UMovieSceneTrack>, TSharedPtr<class IMovieSceneTrackInstance>> FMovieSceneInstanceMap;
+
+
+class FMovieSceneInstance
+	: public TSharedFromThis<FMovieSceneInstance>
 {
 public:
-	MOVIESCENE_API FMovieSceneInstance( UMovieScene& InMovieScene );
 
 	/**
-	* Save state of the objects that this movie scene controls.
-	*/
-	MOVIESCENE_API void SaveState ();
+	 * Creates and initializes a new instance from a movie scene.
+	 *
+	 * @param InMovieScene The movie scene that this instance represents.
+	 */
+	MOVIESCENE_API FMovieSceneInstance(UMovieScene& InMovieScene);
+
+	/** Save state of the objects that this movie scene controls. */
+	MOVIESCENE_API void SaveState();
+
+	/** Restore state of the objects that this movie scene controls. */
+	MOVIESCENE_API void RestoreState();
 
 	/**
-	* Restore state of the objects that this movie scene controls.
-	*/
-	MOVIESCENE_API void RestoreState ();
+	 * Updates this movie scene.
+	 *
+	 * @param Position The local playback position.
+	 * @param Player Movie scene player interface for interaction with runtime data.
+	 */
+	MOVIESCENE_API void Update(float Position, float LastPosition, class IMovieScenePlayer& Player);
 
 	/**
-	 * Updates this movie scene
+	 * Refreshes the existing instance.
 	 *
-	 * @param Position The local playback position
-	 * @param Player Movie scene player interface for interaction with runtime data
+	 * Called when something significant about movie scene data changes (like adding or removing a track).
+	 * Instantiates all new tracks found and removes instances for tracks that no longer exist.
 	 */
-	MOVIESCENE_API void Update( float Position, float LastPosition, class IMovieScenePlayer& Player );
+	MOVIESCENE_API void RefreshInstance(class IMovieScenePlayer& Player);
 
 	/**
-	 * Refreshes the existing instance.  Called when something significant about moviescene data changes (like adding or removing a track)
-	 * Instances all new tracks found and removes instances for tracks that no longer exist
+	 * Get the movie scene's object binding manager.
+	 *
+	 * @return The scene's binding manager.
 	 */
-	MOVIESCENE_API void RefreshInstance( class IMovieScenePlayer& Player );
+	TScriptInterface<UMovieSceneBindingManager> GetBindingManager() const;
 
-	/** @return The movie scene being instanced */
-	UMovieScene* GetMovieScene() { return MovieScene.Get(); }
+	/**
+	 * Get the movie scene associated with this instance.
+	 *
+	 * @return The movie scene object.
+	 */
+	UMovieScene* GetMovieScene()
+	{
+		return MovieScene.Get();
+	}
 
 	/** 
-	 * @return The time range for the movie scene 
-	 * Note this is not necessarily how long the instance will be playing.  The parent could have changed the length 
+	 * Get the time range for the movie scene.
+	 *
+	 * Note: This is not necessarily how long the instance will be playing.
+	 * The parent could have changed the length.
+	 *
+	 * @return The time range.
 	 */
-	TRange<float> GetMovieSceneTimeRange() const { return TimeRange; }
+	TRange<float> GetMovieSceneTimeRange() const
+	{
+		return TimeRange;
+	}
+
+protected:
+
+	void RefreshInstanceMap(const TArray<UMovieSceneTrack*>& Tracks, const TArray<UObject*>& RuntimeObjects, FMovieSceneInstanceMap& TrackInstances, class IMovieScenePlayer& Player);
 
 private:
-	void RefreshInstanceMap( const TArray<UMovieSceneTrack*>& Tracks, const TArray<UObject*>& RuntimeObjects, FMovieSceneInstanceMap& TrackInstances, class IMovieScenePlayer& Player );
-private:
+
 	/** A paring of a runtime object guid to a runtime data and track instances animating the runtime objects */
 	struct FMovieSceneObjectBindingInstance
 	{
 		/** Runtime object guid for looking up runtime UObjects */
 		FGuid ObjectGuid;
+
 		/** Actual runtime objects */	
 		TArray<UObject*> RuntimeObjects;
+
 		/** Instances that animate the runtime objects */
 		FMovieSceneInstanceMap TrackInstances;
 	};
 
+
 	/** MovieScene that is instanced */
 	TWeakObjectPtr<UMovieScene> MovieScene;
+
 	/** The shot track instance map */
 	TSharedPtr<class IMovieSceneTrackInstance> ShotTrackInstance;
+
 	/** All Master track instances */
 	FMovieSceneInstanceMap MasterTrackInstances;
+
 	/** All object binding instances */
 	TMap<FGuid, FMovieSceneObjectBindingInstance> ObjectBindingInstances;
+
 	/** Cached time range for the movie scene */
 	TRange<float> TimeRange;
 };
-
-
diff --git a/Engine/Source/Runtime/MovieScene/Public/MovieSceneObjectId.h b/Engine/Source/Runtime/MovieScene/Public/MovieSceneObjectId.h
new file mode 100644
--- /dev/null
+++ b/Engine/Source/Runtime/MovieScene/Public/MovieSceneObjectId.h
@@ -0,0 +1,75 @@
+// Copyright 1998-2015 Epic Games, Inc. All Rights Reserved.
+
+#pragma once
+
+#include "MovieSceneObjectId.generated.h"
+
+
+/**
+ * Identifies an object (or embedded component) bound in a movie scene.
+ */
+USTRUCT()
+struct FMovieSceneObjectId
+{
+	GENERATED_USTRUCT_BODY()
+
+	/** Unique identifier of a bound object. */
+	UPROPERTY()
+	FGuid ObjectGuid;
+
+	/** Name of the bound object's component object being used (optional). */
+	UPROPERTY()
+	FName ComponentName;
+
+public:
+
+	/** Default constructor. */
+	FMovieSceneObjectId() { }
+
+	/** Creates and initializes a new instance from an object ID. */
+	FMovieSceneObjectId(const FGuid& InObjectGuid)
+		: ObjectGuid(InObjectGuid)
+	{ }
+
+	/** Creates and initializes a new instance from an object and component ID. */
+	FMovieSceneObjectId(const FGuid& InObjectGuid, const FName& InComponentName)
+		: ObjectGuid(InObjectGuid)
+		, ComponentName(InComponentName)
+	{ }
+
+public:
+
+	/**
+	 * Compares two object identifiers for equality.
+	 *
+	 * @param X The first identifier to compare.
+	 * @param Y The second identifier to compare.
+	 * @return true if the identifiers are equal, false otherwise.
+	 */
+	friend bool operator==( const FMovieSceneObjectId& X, const FMovieSceneObjectId& Y )
+	{
+		return ((X.ObjectGuid == Y.ObjectGuid) && (X.ComponentName == X.ComponentName));
+	}
+
+	/**
+	 * Compares two object identifiers for inequality.
+	 *
+	 * @param X The first identifier to compare.
+	 * @param Y The second identifier to compare.
+	 * @return true if the identifier are not equal, false otherwise.
+	 */
+	friend bool operator!=( const FMovieSceneObjectId& X, const FMovieSceneObjectId& Y )
+	{
+		return ((X.ObjectGuid != Y.ObjectGuid) || (X.ComponentName != X.ComponentName));
+	}
+
+	/**
+	 * Checks whether this object identifier is valid or not.
+	 *
+	 * @return true if valid, false otherwise.
+	 */
+	bool IsValid() const
+	{
+		return ObjectGuid.IsValid();
+	}
+};
