Sequencer: Add setting for setting the global time when dragging a keyframe.

diff --git a/Engine/Source/Editor/Sequencer/Private/SSequencer.cpp b/Engine/Source/Editor/Sequencer/Private/SSequencer.cpp
--- a/Engine/Source/Editor/Sequencer/Private/SSequencer.cpp
+++ b/Engine/Source/Editor/Sequencer/Private/SSequencer.cpp
@@ -623,6 +623,7 @@ TSharedRef<SWidget> SSequencer::MakeSnapMenu()
 	MenuBuilder.BeginSection( "PlayTimeSnapping", LOCTEXT( "SnappingMenuPlayTimeHeader", "Play Time Snapping" ) );
 	{
 		MenuBuilder.AddMenuEntry( FSequencerCommands::Get().ToggleSnapPlayTimeToInterval );
+		MenuBuilder.AddMenuEntry( FSequencerCommands::Get().ToggleSnapPlayTimeToDraggedKey );
 	}
 	MenuBuilder.EndSection();
 
diff --git a/Engine/Source/Editor/Sequencer/Private/SectionDragOperations.cpp b/Engine/Source/Editor/Sequencer/Private/SectionDragOperations.cpp
--- a/Engine/Source/Editor/Sequencer/Private/SectionDragOperations.cpp
+++ b/Engine/Source/Editor/Sequencer/Private/SectionDragOperations.cpp
@@ -475,6 +475,8 @@ void FMoveKeys::OnDrag( const FPointerEvent& MouseEvent, const FVector2D& LocalM
 			}
 		}
 
+		float PrevNewKeyTime = FLT_MAX;
+
 		for( FSelectedKey SelectedKey : SelectedKeys )
 		{
 			UMovieSceneSection* Section = SelectedKey.Section;
@@ -498,6 +500,21 @@ void FMoveKeys::OnDrag( const FPointerEvent& MouseEvent, const FVector2D& LocalM
 			{
 				Section->SetStartTime( NewKeyTime );
 			}
+
+			if (PrevNewKeyTime == FLT_MAX)
+			{
+				PrevNewKeyTime = NewKeyTime;
+			}
+			else if (!FMath::IsNearlyEqual(NewKeyTime, PrevNewKeyTime))
+			{
+				PrevNewKeyTime = -FLT_MAX;
+			}
+		}
+
+		// Snap the play time to the new dragged key time if all the keyframes were dragged to the same time
+		if (Settings->GetSnapPlayTimeToDraggedKey() && PrevNewKeyTime != FLT_MAX && PrevNewKeyTime != -FLT_MAX)
+		{
+			Sequencer.SetGlobalTime(PrevNewKeyTime);
 		}
 	}
 }
diff --git a/Engine/Source/Editor/Sequencer/Private/Sequencer.cpp b/Engine/Source/Editor/Sequencer/Private/Sequencer.cpp
--- a/Engine/Source/Editor/Sequencer/Private/Sequencer.cpp
+++ b/Engine/Source/Editor/Sequencer/Private/Sequencer.cpp
@@ -1662,6 +1662,12 @@ void FSequencer::BindSequencerCommands()
 		FCanExecuteAction::CreateLambda( []{ return true; } ),
 		FIsActionChecked::CreateLambda( [Settings]{ return Settings->GetSnapPlayTimeToInterval(); } ) );
 
+	SequencerCommandBindings->MapAction(
+		Commands.ToggleSnapPlayTimeToDraggedKey,
+		FExecuteAction::CreateLambda( [Settings]{ Settings->SetSnapPlayTimeToDraggedKey( !Settings->GetSnapPlayTimeToDraggedKey() ); } ),
+		FCanExecuteAction::CreateLambda( []{ return true; } ),
+		FIsActionChecked::CreateLambda( [Settings]{ return Settings->GetSnapPlayTimeToDraggedKey(); } ) );
+
 	SequencerCommandBindings->MapAction(
 		Commands.ToggleSnapCurveValueToInterval,
 		FExecuteAction::CreateLambda( [Settings]{ Settings->SetSnapCurveValueToInterval( !Settings->GetSnapCurveValueToInterval() ); } ),
diff --git a/Engine/Source/Editor/Sequencer/Private/SequencerCommands.cpp b/Engine/Source/Editor/Sequencer/Private/SequencerCommands.cpp
--- a/Engine/Source/Editor/Sequencer/Private/SequencerCommands.cpp
+++ b/Engine/Source/Editor/Sequencer/Private/SequencerCommands.cpp
@@ -31,6 +31,7 @@ void FSequencerCommands::RegisterCommands()
 	UI_COMMAND( ToggleSnapSectionTimesToSections, "Snap to other sections", "Snap sections to other sections.", EUserInterfaceActionType::ToggleButton, FInputChord() );
 
 	UI_COMMAND( ToggleSnapPlayTimeToInterval, "Snap to the interval while scrubbing", "Snap the play time to the interval while scrubbing.", EUserInterfaceActionType::ToggleButton, FInputChord() );
+	UI_COMMAND( ToggleSnapPlayTimeToDraggedKey, "Snap to the dragged key", "Snap the play time to the dragged key.", EUserInterfaceActionType::ToggleButton, FInputChord() );
 
 	UI_COMMAND(ToggleSnapCurveValueToInterval, "Snap curve key values", "Snap curve keys to the value snapping interval.", EUserInterfaceActionType::ToggleButton, FInputChord() );
 
diff --git a/Engine/Source/Editor/Sequencer/Private/SequencerCommands.h b/Engine/Source/Editor/Sequencer/Private/SequencerCommands.h
--- a/Engine/Source/Editor/Sequencer/Private/SequencerCommands.h
+++ b/Engine/Source/Editor/Sequencer/Private/SequencerCommands.h
@@ -60,6 +60,9 @@ public:
 	/** Toggles whether or not the play time should snap to the selected interval. */
 	TSharedPtr< FUICommandInfo > ToggleSnapPlayTimeToInterval;
 
+	/** Toggles whether or not the play time should snap to the dragged key. */
+	TSharedPtr< FUICommandInfo > ToggleSnapPlayTimeToDraggedKey;
+
 	/** Toggles whether or not to snap curve values to the interval. */
 	TSharedPtr< FUICommandInfo > ToggleSnapCurveValueToInterval;
 
diff --git a/Engine/Source/Editor/Sequencer/Private/SequencerSettings.cpp b/Engine/Source/Editor/Sequencer/Private/SequencerSettings.cpp
--- a/Engine/Source/Editor/Sequencer/Private/SequencerSettings.cpp
+++ b/Engine/Source/Editor/Sequencer/Private/SequencerSettings.cpp
@@ -12,6 +12,7 @@ USequencerSettings::USequencerSettings( const FObjectInitializer& ObjectInitiali
 	bSnapSectionTimesToInterval = true;
 	bSnapSectionTimesToSections = true;
 	bSnapPlayTimeToInterval = true;
+	bSnapPlayTimeToDraggedKey = true;
 	CurveValueSnapInterval = 10.0f;
 	bSnapCurveValueToInterval = true;
 	bIsUsingCleanView = false;
@@ -119,6 +120,20 @@ void USequencerSettings::SetSnapPlayTimeToInterval(bool InbSnapPlayTimeToInterva
 	}
 }
 
+bool USequencerSettings::GetSnapPlayTimeToDraggedKey() const
+{
+	return bSnapPlayTimeToDraggedKey;
+}
+
+void USequencerSettings::SetSnapPlayTimeToDraggedKey(bool InbSnapPlayTimeToDraggedKey)
+{
+	if ( bSnapPlayTimeToDraggedKey != InbSnapPlayTimeToDraggedKey )
+	{
+		bSnapPlayTimeToDraggedKey = InbSnapPlayTimeToDraggedKey;
+		SaveConfig();
+	}
+}
+
 float USequencerSettings::GetCurveValueSnapInterval() const
 {
 	return CurveValueSnapInterval;
diff --git a/Engine/Source/Editor/Sequencer/Private/SequencerSettings.h b/Engine/Source/Editor/Sequencer/Private/SequencerSettings.h
--- a/Engine/Source/Editor/Sequencer/Private/SequencerSettings.h
+++ b/Engine/Source/Editor/Sequencer/Private/SequencerSettings.h
@@ -64,6 +64,11 @@ public:
 	/** Sets whether or not to snap the play time to the interval while scrubbing. */
 	void SetSnapPlayTimeToInterval(bool InbSnapPlayTimeToInterval);
 
+	/** Gets whether or not to snap the play time to the dragged key. */
+	bool GetSnapPlayTimeToDraggedKey() const;
+	/** Sets whether or not to snap the play time to the dragged key. */
+	void SetSnapPlayTimeToDraggedKey(bool InbSnapPlayTimeToDraggedKey);
+
 	/** Gets the snapping interval for curve values. */
 	float GetCurveValueSnapInterval() const;
 	/** Sets the snapping interval for curve values. */
@@ -130,6 +135,9 @@ protected:
 	UPROPERTY( config )
 	bool bSnapPlayTimeToInterval;
 
+	UPROPERTY( config )
+	bool bSnapPlayTimeToDraggedKey;
+
 	UPROPERTY( config )
 	float CurveValueSnapInterval;
 
