Sequencer: Double click on a section that's currently expanded to the view range now sets the view range back to before the double click. This allows you to get out of the section view range easily.

diff --git a/Engine/Source/Editor/Sequencer/Private/Sequencer.cpp b/Engine/Source/Editor/Sequencer/Private/Sequencer.cpp
--- a/Engine/Source/Editor/Sequencer/Private/Sequencer.cpp
+++ b/Engine/Source/Editor/Sequencer/Private/Sequencer.cpp
@@ -150,6 +150,7 @@ FSequencer::FSequencer()
 	: SequencerCommandBindings( new FUICommandList )
 	, TargetViewRange(0.f, 5.f)
 	, LastViewRange(0.f, 5.f)
+	, ViewRangeBeforeZoom(TRange<float>::Empty())
 	, bAutoScrollEnabled(true)
 	, PlaybackState( EMovieScenePlayerStatus::Stopped )
 	, ScrubPosition( 0.0f )
@@ -1351,7 +1352,19 @@ void FSequencer::ZoomToSelectedSections()
 
 	if (!BoundsHull.IsEmpty() && !BoundsHull.IsDegenerate())
 	{
-		OnViewRangeChanged(BoundsHull);
+		// Zoom back to last view range if already expanded
+		if (!ViewRangeBeforeZoom.IsEmpty() &&
+			FMath::IsNearlyEqual(BoundsHull.GetLowerBoundValue(), GetViewRange().GetLowerBoundValue()) &&
+			FMath::IsNearlyEqual(BoundsHull.GetUpperBoundValue(), GetViewRange().GetUpperBoundValue()))
+		{
+			OnViewRangeChanged(ViewRangeBeforeZoom);
+		}
+		else
+		{
+			ViewRangeBeforeZoom = GetViewRange();
+
+			OnViewRangeChanged(BoundsHull);
+		}
 	}
 }
 
diff --git a/Engine/Source/Editor/Sequencer/Private/Sequencer.h b/Engine/Source/Editor/Sequencer/Private/Sequencer.h
--- a/Engine/Source/Editor/Sequencer/Private/Sequencer.h
+++ b/Engine/Source/Editor/Sequencer/Private/Sequencer.h
@@ -418,6 +418,8 @@ private:
 	TRange<float> TargetViewRange;
 	/** The last time range that was viewed */
 	TRange<float> LastViewRange;
+	/** The view range before zooming */
+	TRange<float> ViewRangeBeforeZoom;
 	/** The amount of autoscroll pan offset that is currently being applied */
 	TOptional<float> AutoscrollOffset;
 	/** Whether or not we are allowing autoscroll */
