Sequencer settings are now stored per sequencer-type

diff --git a/Engine/Source/Editor/MovieSceneTools/Private/AssetTypeActions_MovieScene.cpp b/Engine/Source/Editor/MovieSceneTools/Private/AssetTypeActions_MovieScene.cpp
--- a/Engine/Source/Editor/MovieSceneTools/Private/AssetTypeActions_MovieScene.cpp
+++ b/Engine/Source/Editor/MovieSceneTools/Private/AssetTypeActions_MovieScene.cpp
@@ -27,7 +27,7 @@ void FAssetTypeActions_MovieScene::OpenAssetEditor( const TArray<UObject*>& InOb
 			// @todo sequencer: Only allow users to create new MovieScenes if that feature is turned on globally.
 			if( FParse::Param( FCommandLine::Get(), TEXT( "Sequencer" ) ) )
 			{
-				FSequencerViewParams ViewParams;
+				FSequencerViewParams ViewParams( TEXT( "MovieSceneEditorSettings" ) );
 				FModuleManager::LoadModuleChecked< ISequencerModule >( "Sequencer" ).CreateSequencerAssetEditor( Mode, ViewParams, EditWithinLevelEditor, MovieScene, true );
 			}
 		}
diff --git a/Engine/Source/Editor/NiagaraEditor/Private/NiagaraEffectEditor.cpp b/Engine/Source/Editor/NiagaraEditor/Private/NiagaraEffectEditor.cpp
--- a/Engine/Source/Editor/NiagaraEditor/Private/NiagaraEffectEditor.cpp
+++ b/Engine/Source/Editor/NiagaraEditor/Private/NiagaraEffectEditor.cpp
@@ -81,7 +81,7 @@ void FNiagaraEffectEditor::InitNiagaraEffectEditor(const EToolkitMode::Type Mode
 
 	if (!Sequencer.IsValid())
 	{
-		FSequencerViewParams ViewParams;
+		FSequencerViewParams ViewParams( TEXT( "NiagaraSequencerSettings" ) );
 		ViewParams.InitalViewRange = TRange<float>(-0.02f, 3.2f);
 		ViewParams.InitialScrubPosition = 0;
 
diff --git a/Engine/Source/Editor/Sequencer/Private/SSequencer.cpp b/Engine/Source/Editor/Sequencer/Private/SSequencer.cpp
--- a/Engine/Source/Editor/Sequencer/Private/SSequencer.cpp
+++ b/Engine/Source/Editor/Sequencer/Private/SSequencer.cpp
@@ -191,8 +191,9 @@ void SSequencer::Construct( const FArguments& InArgs, TSharedRef< class FSequenc
 	bUserIsSelecting = false;
 
 	USelection::SelectionChangedEvent.AddSP(this, &SSequencer::OnActorSelectionChanged);
+	Settings = InSequencer->GetSettings();
 
-	GetMutableDefault<USequencerSettings>()->GetOnShowCurveEditorChanged().AddSP(this, &SSequencer::OnCurveEditorVisibilityChanged);
+	Settings->GetOnShowCurveEditorChanged().AddSP(this, &SSequencer::OnCurveEditorVisibilityChanged);
 
 	// Create a node tree which contains a tree of movie scene data to display in the sequence
 	SequencerNodeTree = MakeShareable( new FSequencerNodeTree( InSequencer.Get() ) );
@@ -206,6 +207,7 @@ void SSequencer::Construct( const FArguments& InArgs, TSharedRef< class FSequenc
 	TimeSliderArgs.OnBeginScrubberMovement = InArgs._OnBeginScrubbing;
 	TimeSliderArgs.OnEndScrubberMovement = InArgs._OnEndScrubbing;
 	TimeSliderArgs.OnScrubPositionChanged = InArgs._OnScrubPositionChanged;
+	TimeSliderArgs.Settings = Settings;
 
 	TSharedRef<FSequencerTimeSliderController> TimeSliderController( new FSequencerTimeSliderController( TimeSliderArgs ) );
 	
@@ -678,7 +680,7 @@ TSharedRef<SWidget> SSequencer::MakeTransportControls()
 SSequencer::~SSequencer()
 {
 	USelection::SelectionChangedEvent.RemoveAll(this);
-	GetMutableDefault<USequencerSettings>()->GetOnShowCurveEditorChanged().RemoveAll(this);
+	Settings->GetOnShowCurveEditorChanged().RemoveAll(this);
 }
 
 void SSequencer::RegisterActiveTimerForPlayback()
@@ -792,22 +794,22 @@ void SSequencer::OnOutlinerSearchChanged( const FText& Filter )
 
 float SSequencer::OnGetTimeSnapInterval() const
 {
-	return GetDefault<USequencerSettings>()->GetTimeSnapInterval();
+	return Settings->GetTimeSnapInterval();
 }
 
 void SSequencer::OnTimeSnapIntervalChanged( float InInterval )
 {
-	GetMutableDefault<USequencerSettings>()->SetTimeSnapInterval(InInterval);
+	Settings->SetTimeSnapInterval(InInterval);
 }
 
 float SSequencer::OnGetValueSnapInterval() const
 {
-	return GetDefault<USequencerSettings>()->GetCurveValueSnapInterval();
+	return Settings->GetCurveValueSnapInterval();
 }
 
 void SSequencer::OnValueSnapIntervalChanged( float InInterval )
 {
-	GetMutableDefault<USequencerSettings>()->SetCurveValueSnapInterval( InInterval );
+	Settings->SetCurveValueSnapInterval( InInterval );
 }
 
 
@@ -1211,13 +1213,13 @@ EVisibility SSequencer::GetBreadcrumbTrailVisibility() const
 
 EVisibility SSequencer::GetCurveEditorToolBarVisibility() const
 {
-	return GetDefault<USequencerSettings>()->GetShowCurveEditor() ? EVisibility::Visible : EVisibility::Collapsed;
+	return Settings->GetShowCurveEditor() ? EVisibility::Visible : EVisibility::Collapsed;
 }
 
 float SSequencer::GetOutlinerSpacerFill() const
 {
 	const float Column1Coeff = GetColumnFillCoefficient(1);
-	return GetDefault<USequencerSettings>()->GetShowCurveEditor() ? Column1Coeff / (1 - Column1Coeff) : 0.f;
+	return Settings->GetShowCurveEditor() ? Column1Coeff / (1 - Column1Coeff) : 0.f;
 }
 
 void SSequencer::OnColumnFillCoefficientChanged(float FillCoefficient, int32 ColumnIndex)
@@ -1227,12 +1229,12 @@ void SSequencer::OnColumnFillCoefficientChanged(float FillCoefficient, int32 Col
 
 EVisibility SSequencer::GetTrackAreaVisibility() const
 {
-	return GetDefault<USequencerSettings>()->GetShowCurveEditor() ? EVisibility::Collapsed : EVisibility::Visible;
+	return Settings->GetShowCurveEditor() ? EVisibility::Collapsed : EVisibility::Visible;
 }
 
 EVisibility SSequencer::GetCurveEditorVisibility() const
 {
-	return GetDefault<USequencerSettings>()->GetShowCurveEditor() ? EVisibility::Visible : EVisibility::Collapsed;
+	return Settings->GetShowCurveEditor() ? EVisibility::Visible : EVisibility::Collapsed;
 }
 
 void SSequencer::OnCurveEditorVisibilityChanged()
@@ -1240,7 +1242,7 @@ void SSequencer::OnCurveEditorVisibilityChanged()
 	if (CurveEditor.IsValid())
 	{
 		// Only zoom horizontally if the editor is visible
-		CurveEditor->SetZoomToFit(true, GetDefault<USequencerSettings>()->GetShowCurveEditor());
+		CurveEditor->SetZoomToFit(true, Settings->GetShowCurveEditor());
 	}
 }
 
diff --git a/Engine/Source/Editor/Sequencer/Private/SSequencer.h b/Engine/Source/Editor/Sequencer/Private/SSequencer.h
--- a/Engine/Source/Editor/Sequencer/Private/SSequencer.h
+++ b/Engine/Source/Editor/Sequencer/Private/SSequencer.h
@@ -45,7 +45,7 @@ struct FSequencerBreadcrumb
 /**
  * Main sequencer UI widget
  */
-class SSequencer : public SCompoundWidget
+class SSequencer : public SCompoundWidget, public FGCObject
 {
 public:
 	DECLARE_DELEGATE_OneParam( FOnToggleBoolOption, bool )
@@ -72,6 +72,11 @@ public:
 	void Construct( const FArguments& InArgs, TSharedRef< class FSequencer > InSequencer );
 
 	~SSequencer();
+	
+	virtual void AddReferencedObjects( FReferenceCollector& Collector )
+	{
+		Collector.AddReferencedObject( Settings );
+	}
 
 	virtual bool SupportsKeyboardFocus() const override { return true; }
 
@@ -241,6 +246,8 @@ private:
 	TSharedPtr< class SBreadcrumbTrail<FSequencerBreadcrumb> > BreadcrumbTrail;
 	/** The main sequencer interface */
 	TWeakPtr<FSequencer> Sequencer;
+	/** Cached settings provided to the sequencer itself on creation */
+	USequencerSettings* Settings;
 	/** The fill coefficients of each column in the grid. */
 	float ColumnFillCoefficients[2];
 	/** Whether the active timer is currently registered */
diff --git a/Engine/Source/Editor/Sequencer/Private/SSequencerCurveEditor.cpp b/Engine/Source/Editor/Sequencer/Private/SSequencerCurveEditor.cpp
--- a/Engine/Source/Editor/Sequencer/Private/SSequencerCurveEditor.cpp
+++ b/Engine/Source/Editor/Sequencer/Private/SSequencerCurveEditor.cpp
@@ -10,9 +10,11 @@
 void SSequencerCurveEditor::Construct( const FArguments& InArgs, TSharedRef<FSequencer> InSequencer, TSharedRef<ITimeSliderController> InTimeSliderController )
 {
 	Sequencer = InSequencer;
+	Settings = InSequencer->GetSettings();
+
 	TimeSliderController = InTimeSliderController;
 	NodeTreeSelectionChangedHandle = Sequencer.Pin()->GetSelection().GetOnOutlinerNodeSelectionChanged().AddSP(this, &SSequencerCurveEditor::NodeTreeSelectionChanged);
-	GetMutableDefault<USequencerSettings>()->GetOnCurveVisibilityChanged().AddSP( this, &SSequencerCurveEditor::SequencerCurveVisibilityChanged );
+	Settings->GetOnCurveVisibilityChanged().AddSP( this, &SSequencerCurveEditor::SequencerCurveVisibilityChanged );
 
 	// @todo sequencer: switch to lambda capture expressions when support is introduced?
 	auto ViewRange = InArgs._ViewRange;
@@ -24,7 +26,7 @@ void SSequencerCurveEditor::Construct( const FArguments& InArgs, TSharedRef<FSeq
 		.ViewMaxInput_Lambda( [=]{ return ViewRange.Get().GetUpperBoundValue(); } )
 		.OnSetInputViewRange_Lambda( [=](float InLowerBound, float InUpperBound){ OnViewRangeChanged.ExecuteIfBound(TRange<float>(InLowerBound, InUpperBound), EViewRangeInterpolation::Immediate); } )
 		.HideUI( false )
-		.ZoomToFitHorizontal( GetDefault<USequencerSettings>()->GetShowCurveEditor() )
+		.ZoomToFitHorizontal( Settings->GetShowCurveEditor() )
 		.ShowCurveSelector( false )
 		.ShowZoomButtons( false )
 		.ShowInputGridNumbers( false )
@@ -114,34 +116,33 @@ void SSequencerCurveEditor::UpdateCurveOwner()
 
 bool SSequencerCurveEditor::GetCurveSnapEnabled() const
 {
-	const USequencerSettings* Settings = GetDefault<USequencerSettings>();
 	return Settings->GetIsSnapEnabled();
 }
 
 float SSequencerCurveEditor::GetCurveTimeSnapInterval() const
 {
-	return GetDefault<USequencerSettings>()->GetSnapKeyTimesToInterval()
-		? GetDefault<USequencerSettings>()->GetTimeSnapInterval()
+	return Settings->GetSnapKeyTimesToInterval()
+		? Settings->GetTimeSnapInterval()
 		: 0;
 }
 
 float SSequencerCurveEditor::GetCurveValueSnapInterval() const
 {
-	return GetDefault<USequencerSettings>()->GetSnapCurveValueToInterval()
-		? GetDefault<USequencerSettings>()->GetCurveValueSnapInterval()
+	return Settings->GetSnapCurveValueToInterval()
+		? Settings->GetCurveValueSnapInterval()
 		: 0;
 }
 
 bool SSequencerCurveEditor::GetShowCurveEditorCurveToolTips() const
 {
-	return GetDefault<USequencerSettings>()->GetShowCurveEditorCurveToolTips();
+	return Settings->GetShowCurveEditorCurveToolTips();
 }
 
 void SSequencerCurveEditor::NodeTreeSelectionChanged()
 {
 	if (SequencerNodeTree.IsValid())
 	{
-		if (GetDefault<USequencerSettings>()->GetCurveVisibility() == ESequencerCurveVisibility::SelectedCurves)
+		if (Settings->GetCurveVisibility() == ESequencerCurveVisibility::SelectedCurves)
 		{
 			UpdateCurveOwner();
 		}
diff --git a/Engine/Source/Editor/Sequencer/Private/SSequencerCurveEditor.h b/Engine/Source/Editor/Sequencer/Private/SSequencerCurveEditor.h
--- a/Engine/Source/Editor/Sequencer/Private/SSequencerCurveEditor.h
+++ b/Engine/Source/Editor/Sequencer/Private/SSequencerCurveEditor.h
@@ -60,6 +60,8 @@ private:
 
 	/** The sequencer which owns this widget. */
 	TWeakPtr<FSequencer> Sequencer;
+	/** Cached user-supplied settings object */
+	USequencerSettings* Settings; 
 	/** The class responsible for time sliding on the curve editor */
 	TSharedPtr<class ITimeSliderController> TimeSliderController;
 	/** The sequencer node tree which contains the key area nodes which supply the curves to edit. */
diff --git a/Engine/Source/Editor/Sequencer/Private/SectionDragOperations.h b/Engine/Source/Editor/Sequencer/Private/SectionDragOperations.h
--- a/Engine/Source/Editor/Sequencer/Private/SectionDragOperations.h
+++ b/Engine/Source/Editor/Sequencer/Private/SectionDragOperations.h
@@ -11,7 +11,7 @@ public:
 	FSequencerDragOperation( FSequencer& InSequencer )
 		: Sequencer(InSequencer)
 	{
-		Settings = GetDefault<USequencerSettings>();
+		Settings = Sequencer.GetSettings();
 	}
 
 	virtual ~FSequencerDragOperation(){}
diff --git a/Engine/Source/Editor/Sequencer/Private/Sequencer.cpp b/Engine/Source/Editor/Sequencer/Private/Sequencer.cpp
--- a/Engine/Source/Editor/Sequencer/Private/Sequencer.cpp
+++ b/Engine/Source/Editor/Sequencer/Private/Sequencer.cpp
@@ -56,6 +56,8 @@ void FSequencer::InitSequencer( const FSequencerInitParams& InitParams, const TA
 	if( IsSequencerEnabled() )
 	{
 		bIsEditingWithinLevelEditor = InitParams.bEditWithinLevelEditor;
+		
+		Settings = USequencerSettingsContainer::GetOrCreate(*InitParams.ViewParams.UniqueName);
 
 		// If this sequencer edits the level, close out the existing active sequencer and mark this as the active sequencer.
 		if (bIsEditingWithinLevelEditor)
@@ -742,6 +744,8 @@ void FSequencer::RemoveMovieSceneInstance( UMovieSceneSection& MovieSceneSection
 
 void FSequencer::AddReferencedObjects( FReferenceCollector& Collector )
 {
+	Collector.AddReferencedObject( Settings );
+
 	for( int32 MovieSceneIndex = 0; MovieSceneIndex < MovieSceneStack.Num(); ++MovieSceneIndex )
 	{
 		UMovieScene* Scene = MovieSceneStack[MovieSceneIndex]->GetMovieScene();
@@ -996,7 +1000,7 @@ void FSequencer::OnToggleAutoKey()
 void FSequencer::OnToggleAutoScroll()
 {
 	bAutoScrollEnabled = !bAutoScrollEnabled;
-	GetMutableDefault<USequencerSettings>()->SetAutoScrollEnabled(bAutoScrollEnabled);
+	Settings->SetAutoScrollEnabled(bAutoScrollEnabled);
 }
 
 FGuid FSequencer::AddSpawnableForAssetOrClass( UObject* Object, UObject* CounterpartGamePreviewObject )
@@ -1618,8 +1622,6 @@ void FSequencer::BindSequencerCommands()
 		Commands.SetKey,
 		FExecuteAction::CreateSP( this, &FSequencer::SetKey ) );
 
-	USequencerSettings* Settings = GetMutableDefault<USequencerSettings>();
-
 	SequencerCommandBindings->MapAction(
 		Commands.ToggleAutoKeyEnabled,
 		FExecuteAction::CreateSP( this, &FSequencer::OnToggleAutoKey ),
@@ -1635,87 +1637,87 @@ void FSequencer::BindSequencerCommands()
 
 	SequencerCommandBindings->MapAction(
 		Commands.ToggleCleanView,
-		FExecuteAction::CreateLambda( [Settings]{ Settings->SetIsUsingCleanView( !Settings->GetIsUsingCleanView() ); } ),
+		FExecuteAction::CreateLambda( [this]{ Settings->SetIsUsingCleanView( !Settings->GetIsUsingCleanView() ); } ),
 		FCanExecuteAction::CreateLambda( []{ return true; } ),
-		FIsActionChecked::CreateLambda( [Settings]{ return Settings->GetIsUsingCleanView(); } ) );
+		FIsActionChecked::CreateLambda( [this]{ return Settings->GetIsUsingCleanView(); } ) );
 
 	SequencerCommandBindings->MapAction(
 		Commands.ToggleIsSnapEnabled,
-		FExecuteAction::CreateLambda( [Settings]{ Settings->SetIsSnapEnabled( !Settings->GetIsSnapEnabled() ); } ),
+		FExecuteAction::CreateLambda( [this]{ Settings->SetIsSnapEnabled( !Settings->GetIsSnapEnabled() ); } ),
 		FCanExecuteAction::CreateLambda( []{ return true; } ),
-		FIsActionChecked::CreateLambda( [Settings]{ return Settings->GetIsSnapEnabled(); } ) );
+		FIsActionChecked::CreateLambda( [this]{ return Settings->GetIsSnapEnabled(); } ) );
 
 	SequencerCommandBindings->MapAction(
 		Commands.ToggleSnapKeyTimesToInterval,
-		FExecuteAction::CreateLambda( [Settings]{ Settings->SetSnapKeyTimesToInterval( !Settings->GetSnapKeyTimesToInterval() ); } ),
+		FExecuteAction::CreateLambda( [this]{ Settings->SetSnapKeyTimesToInterval( !Settings->GetSnapKeyTimesToInterval() ); } ),
 		FCanExecuteAction::CreateLambda( []{ return true; } ),
-		FIsActionChecked::CreateLambda( [Settings]{ return Settings->GetSnapKeyTimesToInterval(); } ) );
+		FIsActionChecked::CreateLambda( [this]{ return Settings->GetSnapKeyTimesToInterval(); } ) );
 
 	SequencerCommandBindings->MapAction(
 		Commands.ToggleSnapKeyTimesToKeys,
-		FExecuteAction::CreateLambda( [Settings]{ Settings->SetSnapKeyTimesToKeys( !Settings->GetSnapKeyTimesToKeys() ); } ),
+		FExecuteAction::CreateLambda( [this]{ Settings->SetSnapKeyTimesToKeys( !Settings->GetSnapKeyTimesToKeys() ); } ),
 		FCanExecuteAction::CreateLambda( []{ return true; } ),
-		FIsActionChecked::CreateLambda( [Settings]{ return Settings->GetSnapKeyTimesToKeys(); } ) );
+		FIsActionChecked::CreateLambda( [this]{ return Settings->GetSnapKeyTimesToKeys(); } ) );
 
 	SequencerCommandBindings->MapAction(
 		Commands.ToggleSnapSectionTimesToInterval,
-		FExecuteAction::CreateLambda( [Settings]{ Settings->SetSnapSectionTimesToInterval( !Settings->GetSnapSectionTimesToInterval() ); } ),
+		FExecuteAction::CreateLambda( [this]{ Settings->SetSnapSectionTimesToInterval( !Settings->GetSnapSectionTimesToInterval() ); } ),
 		FCanExecuteAction::CreateLambda( []{ return true; } ),
-		FIsActionChecked::CreateLambda( [Settings]{ return Settings->GetSnapSectionTimesToInterval(); } ) );
+		FIsActionChecked::CreateLambda( [this]{ return Settings->GetSnapSectionTimesToInterval(); } ) );
 
 	SequencerCommandBindings->MapAction(
 		Commands.ToggleSnapSectionTimesToSections,
-		FExecuteAction::CreateLambda( [Settings]{ Settings->SetSnapSectionTimesToSections( !Settings->GetSnapSectionTimesToSections() ); } ),
+		FExecuteAction::CreateLambda( [this]{ Settings->SetSnapSectionTimesToSections( !Settings->GetSnapSectionTimesToSections() ); } ),
 		FCanExecuteAction::CreateLambda( []{ return true; } ),
-		FIsActionChecked::CreateLambda( [Settings]{ return Settings->GetSnapSectionTimesToSections(); } ) );
+		FIsActionChecked::CreateLambda( [this]{ return Settings->GetSnapSectionTimesToSections(); } ) );
 
 	SequencerCommandBindings->MapAction(
 		Commands.ToggleSnapPlayTimeToInterval,
-		FExecuteAction::CreateLambda( [Settings]{ Settings->SetSnapPlayTimeToInterval( !Settings->GetSnapPlayTimeToInterval() ); } ),
+		FExecuteAction::CreateLambda( [this]{ Settings->SetSnapPlayTimeToInterval( !Settings->GetSnapPlayTimeToInterval() ); } ),
 		FCanExecuteAction::CreateLambda( []{ return true; } ),
-		FIsActionChecked::CreateLambda( [Settings]{ return Settings->GetSnapPlayTimeToInterval(); } ) );
+		FIsActionChecked::CreateLambda( [this]{ return Settings->GetSnapPlayTimeToInterval(); } ) );
 
 	SequencerCommandBindings->MapAction(
 		Commands.ToggleSnapPlayTimeToDraggedKey,
-		FExecuteAction::CreateLambda( [Settings]{ Settings->SetSnapPlayTimeToDraggedKey( !Settings->GetSnapPlayTimeToDraggedKey() ); } ),
+		FExecuteAction::CreateLambda( [this]{ Settings->SetSnapPlayTimeToDraggedKey( !Settings->GetSnapPlayTimeToDraggedKey() ); } ),
 		FCanExecuteAction::CreateLambda( []{ return true; } ),
-		FIsActionChecked::CreateLambda( [Settings]{ return Settings->GetSnapPlayTimeToDraggedKey(); } ) );
+		FIsActionChecked::CreateLambda( [this]{ return Settings->GetSnapPlayTimeToDraggedKey(); } ) );
 
 	SequencerCommandBindings->MapAction(
 		Commands.ToggleSnapCurveValueToInterval,
-		FExecuteAction::CreateLambda( [Settings]{ Settings->SetSnapCurveValueToInterval( !Settings->GetSnapCurveValueToInterval() ); } ),
+		FExecuteAction::CreateLambda( [this]{ Settings->SetSnapCurveValueToInterval( !Settings->GetSnapCurveValueToInterval() ); } ),
 		FCanExecuteAction::CreateLambda( []{ return true; } ),
-		FIsActionChecked::CreateLambda( [Settings]{ return Settings->GetSnapCurveValueToInterval(); } ) );
+		FIsActionChecked::CreateLambda( [this]{ return Settings->GetSnapCurveValueToInterval(); } ) );
 
 	SequencerCommandBindings->MapAction(
 		Commands.ToggleShowCurveEditor,
-		FExecuteAction::CreateLambda( [Settings]{ Settings->SetShowCurveEditor(!Settings->GetShowCurveEditor()); } ),
+		FExecuteAction::CreateLambda( [this]{ Settings->SetShowCurveEditor(!Settings->GetShowCurveEditor()); } ),
 		FCanExecuteAction::CreateLambda( []{ return true; } ),
-		FIsActionChecked::CreateLambda( [Settings]{ return Settings->GetShowCurveEditor(); } ) );
+		FIsActionChecked::CreateLambda( [this]{ return Settings->GetShowCurveEditor(); } ) );
 
 	SequencerCommandBindings->MapAction(
 		Commands.ToggleShowCurveEditorCurveToolTips,
-		FExecuteAction::CreateLambda( [Settings]{ Settings->SetShowCurveEditorCurveToolTips( !Settings->GetShowCurveEditorCurveToolTips() ); } ),
+		FExecuteAction::CreateLambda( [this]{ Settings->SetShowCurveEditorCurveToolTips( !Settings->GetShowCurveEditorCurveToolTips() ); } ),
 		FCanExecuteAction::CreateLambda( []{ return true; } ),
-		FIsActionChecked::CreateLambda( [Settings]{ return Settings->GetShowCurveEditorCurveToolTips(); } ) );
+		FIsActionChecked::CreateLambda( [this]{ return Settings->GetShowCurveEditorCurveToolTips(); } ) );
 
 	SequencerCommandBindings->MapAction(
 		Commands.SetAllCurveVisibility,
-		FExecuteAction::CreateLambda( [Settings]{ Settings->SetCurveVisibility( ESequencerCurveVisibility::AllCurves ); } ),
+		FExecuteAction::CreateLambda( [this]{ Settings->SetCurveVisibility( ESequencerCurveVisibility::AllCurves ); } ),
 		FCanExecuteAction::CreateLambda( []{ return true; } ),
-		FIsActionChecked::CreateLambda( [Settings]{ return Settings->GetCurveVisibility() == ESequencerCurveVisibility::AllCurves; } ) );
+		FIsActionChecked::CreateLambda( [this]{ return Settings->GetCurveVisibility() == ESequencerCurveVisibility::AllCurves; } ) );
 
 	SequencerCommandBindings->MapAction(
 		Commands.SetSelectedCurveVisibility,
-		FExecuteAction::CreateLambda( [Settings]{ Settings->SetCurveVisibility( ESequencerCurveVisibility::SelectedCurves ); } ),
+		FExecuteAction::CreateLambda( [this]{ Settings->SetCurveVisibility( ESequencerCurveVisibility::SelectedCurves ); } ),
 		FCanExecuteAction::CreateLambda( []{ return true; } ),
-		FIsActionChecked::CreateLambda( [Settings]{ return Settings->GetCurveVisibility() == ESequencerCurveVisibility::SelectedCurves; } ) );
+		FIsActionChecked::CreateLambda( [this]{ return Settings->GetCurveVisibility() == ESequencerCurveVisibility::SelectedCurves; } ) );
 
 	SequencerCommandBindings->MapAction(
 		Commands.SetAnimatedCurveVisibility,
-		FExecuteAction::CreateLambda( [Settings]{ Settings->SetCurveVisibility( ESequencerCurveVisibility::AnimatedCurves ); } ),
+		FExecuteAction::CreateLambda( [this]{ Settings->SetCurveVisibility( ESequencerCurveVisibility::AnimatedCurves ); } ),
 		FCanExecuteAction::CreateLambda( []{ return true; } ),
-		FIsActionChecked::CreateLambda( [Settings]{ return Settings->GetCurveVisibility() == ESequencerCurveVisibility::AnimatedCurves; } ) );
+		FIsActionChecked::CreateLambda( [this]{ return Settings->GetCurveVisibility() == ESequencerCurveVisibility::AnimatedCurves; } ) );
 
 	for (int32 i = 0; i < TrackEditors.Num(); ++i)
 	{
diff --git a/Engine/Source/Editor/Sequencer/Private/Sequencer.h b/Engine/Source/Editor/Sequencer/Private/Sequencer.h
--- a/Engine/Source/Editor/Sequencer/Private/Sequencer.h
+++ b/Engine/Source/Editor/Sequencer/Private/Sequencer.h
@@ -104,6 +104,9 @@ public:
 	/** @return The current view range */
 	FAnimatedRange GetViewRange() const;
 
+	/** Access the user-supplied settings object */
+	USequencerSettings* GetSettings() const { return Settings; }
+
 	bool IsPerspectiveViewportPosessionEnabled() const { return bPerspectiveViewportPossessionEnabled; }
 
 	/**
@@ -375,6 +378,9 @@ protected:
 	void UpdatePreviewLevelViewportClientFromCameraCut( FLevelEditorViewportClient& InViewportClient, UObject* InCameraObject, bool bNewCameraCut ) const;
 
 private:
+	/** User-supplied settings object for this sequencer */
+	USequencerSettings* Settings;
+
 	TMap< TWeakObjectPtr<UMovieSceneSection>, TSharedRef<FMovieSceneInstance> > MovieSceneSectionToInstanceMap;
 
 	/** Command list for sequencer commands */
diff --git a/Engine/Source/Editor/Sequencer/Private/SequencerCurveOwner.cpp b/Engine/Source/Editor/Sequencer/Private/SequencerCurveOwner.cpp
--- a/Engine/Source/Editor/Sequencer/Private/SequencerCurveOwner.cpp
+++ b/Engine/Source/Editor/Sequencer/Private/SequencerCurveOwner.cpp
@@ -56,7 +56,7 @@ FSequencerCurveOwner::FSequencerCurveOwner( TSharedPtr<FSequencerNodeTree> InSeq
 			if ( RichCurve != nullptr )
 			{
 				bool bAddCurve = false;
-				switch ( GetDefault<USequencerSettings>()->GetCurveVisibility() )
+				switch ( SequencerNodeTree->GetSequencer().GetSettings()->GetCurveVisibility() )
 				{
 				case ESequencerCurveVisibility::AllCurves:
 					bAddCurve = true;
diff --git a/Engine/Source/Editor/Sequencer/Private/SequencerDisplayNode.cpp b/Engine/Source/Editor/Sequencer/Private/SequencerDisplayNode.cpp
--- a/Engine/Source/Editor/Sequencer/Private/SequencerDisplayNode.cpp
+++ b/Engine/Source/Editor/Sequencer/Private/SequencerDisplayNode.cpp
@@ -294,7 +294,7 @@ bool FSequencerDisplayNode::IsVisible() const
 	// If shot filtering is off on clean view is on, node must be pinned
 	return bCachedShotFilteredVisibility &&
 		(ParentTree.HasActiveFilter() ? ParentTree.IsNodeFiltered(AsShared()) : IsParentExpandedOrIsARootNode()) &&
-		(GetSequencer().IsShotFilteringOn() || !GetDefault<USequencerSettings>()->GetIsUsingCleanView() || bNodeIsPinned);
+		(GetSequencer().IsShotFilteringOn() || !GetSequencer().GetSettings()->GetIsUsingCleanView() || bNodeIsPinned);
 }
 
 bool FSequencerDisplayNode::HasVisibleChildren() const
diff --git a/Engine/Source/Editor/Sequencer/Private/SequencerSettings.cpp b/Engine/Source/Editor/Sequencer/Private/SequencerSettings.cpp
--- a/Engine/Source/Editor/Sequencer/Private/SequencerSettings.cpp
+++ b/Engine/Source/Editor/Sequencer/Private/SequencerSettings.cpp
@@ -2,6 +2,27 @@
 
 #include "SequencerPrivatePCH.h"
 
+USequencerSettings* USequencerSettingsContainer::GetOrCreate(const TCHAR* InName)
+{
+	static const TCHAR* SettingsContainerName = TEXT("SequencerSettingsContainer");
+
+	auto* Outer = FindObject<USequencerSettingsContainer>(GetTransientPackage(), SettingsContainerName);
+	if (!Outer)
+	{
+		Outer = NewObject<USequencerSettingsContainer>(GetTransientPackage(), USequencerSettingsContainer::StaticClass(), SettingsContainerName);
+		Outer->AddToRoot();
+	}
+	
+	USequencerSettings* Inst = FindObject<USequencerSettings>( Outer, InName );
+	if (!Inst)
+	{
+		Inst = NewObject<USequencerSettings>( Outer, USequencerSettings::StaticClass(), InName );
+		Inst->LoadConfig();
+	}
+
+	return Inst;
+}
+
 USequencerSettings::USequencerSettings( const FObjectInitializer& ObjectInitializer )
 	: Super( ObjectInitializer )
 {
diff --git a/Engine/Source/Editor/Sequencer/Private/SequencerSettings.h b/Engine/Source/Editor/Sequencer/Private/SequencerSettings.h
--- a/Engine/Source/Editor/Sequencer/Private/SequencerSettings.h
+++ b/Engine/Source/Editor/Sequencer/Private/SequencerSettings.h
@@ -19,8 +19,19 @@ namespace ESequencerCurveVisibility
 	};
 }
 
+/** Empty class used to house multiple named USequencerSettings */
+UCLASS()
+class USequencerSettingsContainer : public UObject
+{
+public:
+	GENERATED_BODY()
+
+	/** Get or create a settings object for the specified name */
+	static USequencerSettings* GetOrCreate(const TCHAR* InName);
+};
+
 /** Serializable options for sequencer. */
-UCLASS(config=EditorPerProjectUserSettings)
+UCLASS(config=EditorPerProjectUserSettings, PerObjectConfig)
 class USequencerSettings : public UObject
 {
 public:
diff --git a/Engine/Source/Editor/Sequencer/Private/SequencerTrackLaneFactory.cpp b/Engine/Source/Editor/Sequencer/Private/SequencerTrackLaneFactory.cpp
--- a/Engine/Source/Editor/Sequencer/Private/SequencerTrackLaneFactory.cpp
+++ b/Engine/Source/Editor/Sequencer/Private/SequencerTrackLaneFactory.cpp
@@ -87,13 +87,15 @@ void FSequencerTrackLaneFactory::Repopulate(const FSequencerNodeTree& NodeTree)
 	}
 }
 
-void FSequencerTrackLaneFactory::GenerateWidgetForNode( FSequencerDisplayNode& Node)
+void FSequencerTrackLaneFactory::GenerateWidgetForNode( FSequencerDisplayNode& Node )
 {
+	auto* Settings = Sequencer->GetSettings();
+
 	TAttribute<FAnimatedRange> ViewRange;
 	ViewRange.Bind(TAttribute<FAnimatedRange>::FGetter::CreateSP(Sequencer, &FSequencer::GetViewRange));
 
 	TSharedRef<SWidget> OutlinerWidget = Node.GenerateContainerWidgetForOutliner();
-	TSharedRef<SWidget> TrackWidget = Node.GenerateWidgetForSectionArea( FAnimatedRange::WrapAttribute(ViewRange) );
+	TSharedRef<SWidget> TrackWidget = Node.GenerateWidgetForSectionArea( FAnimatedRange::WrapAttribute( ViewRange ) );
 
 	// Link the height of the widgets using a size linker
 	FWidgetSizeLinker Link;
@@ -106,8 +108,8 @@ void FSequencerTrackLaneFactory::GenerateWidgetForNode( FSequencerDisplayNode& N
 	auto GetMaxSize = [=]() -> float {
 		return Link.GetMaxSize().Y;
 	};
-	auto IsVisible = []() -> EVisibility{
-		return GetDefault<USequencerSettings>()->GetShowCurveEditor() ? EVisibility::Hidden : EVisibility::Visible;
+	auto IsVisible = [=]() -> EVisibility{
+		return Settings->GetShowCurveEditor() ? EVisibility::Hidden : EVisibility::Visible;
 	};
 
 	TrackOutliner->AddSlot()
diff --git a/Engine/Source/Editor/Sequencer/Private/TimeSliderController.cpp b/Engine/Source/Editor/Sequencer/Private/TimeSliderController.cpp
--- a/Engine/Source/Editor/Sequencer/Private/TimeSliderController.cpp
+++ b/Engine/Source/Editor/Sequencer/Private/TimeSliderController.cpp
@@ -340,10 +340,9 @@ FReply FSequencerTimeSliderController::OnMouseButtonUp( TSharedRef<SWidget> Widg
 			FVector2D CursorPos = MyGeometry.AbsoluteToLocal(MouseEvent.GetLastScreenSpacePosition());
 			float NewValue = RangeToScreen.LocalXToInput(CursorPos.X);
 
-			const USequencerSettings* Settings = GetDefault<USequencerSettings>();
-			if ( Settings->GetIsSnapEnabled() && Settings->GetSnapPlayTimeToInterval() )
+			if ( TimeSliderArgs.Settings->GetIsSnapEnabled() && TimeSliderArgs.Settings->GetSnapPlayTimeToInterval() )
 			{
-				NewValue = Settings->SnapTimeToInterval( NewValue );
+				NewValue = TimeSliderArgs.Settings->SnapTimeToInterval( NewValue );
 			}
 
 			CommitScrubPosition( NewValue, /*bIsScrubbing=*/false );
@@ -405,10 +404,9 @@ FReply FSequencerTimeSliderController::OnMouseMove( TSharedRef<SWidget> WidgetOw
 				FVector2D CursorPos = MyGeometry.AbsoluteToLocal( MouseEvent.GetLastScreenSpacePosition() );
 				float NewValue = RangeToScreen.LocalXToInput( CursorPos.X );
 
-				const USequencerSettings* Settings = GetDefault<USequencerSettings>();
-				if ( Settings->GetIsSnapEnabled() && Settings->GetSnapPlayTimeToInterval() )
+				if ( TimeSliderArgs.Settings->GetIsSnapEnabled() && TimeSliderArgs.Settings->GetSnapPlayTimeToInterval() )
 				{
-					NewValue = Settings->SnapTimeToInterval(NewValue);
+					NewValue = TimeSliderArgs.Settings->SnapTimeToInterval(NewValue);
 				}
 
 				CommitScrubPosition( NewValue, /*bIsScrubbing=*/true );
diff --git a/Engine/Source/Editor/Sequencer/Public/ISequencerModule.h b/Engine/Source/Editor/Sequencer/Public/ISequencerModule.h
--- a/Engine/Source/Editor/Sequencer/Public/ISequencerModule.h
+++ b/Engine/Source/Editor/Sequencer/Public/ISequencerModule.h
@@ -28,9 +28,13 @@ struct FSequencerViewParams
 
 	FOnGetAddMenuContent OnGetAddMenuContent;
 
-	FSequencerViewParams()
+	/** Unique name for the sequencer */
+	FString UniqueName;
+
+	FSequencerViewParams(FString InName = FString())
 		: InitalViewRange( 0.0f, 5.0f )
 		, InitialScrubPosition( 0.0f )
+		, UniqueName(MoveTemp(InName))
 	{}
 };
 
diff --git a/Engine/Source/Editor/SequencerWidgets/Public/ITimeSlider.h b/Engine/Source/Editor/SequencerWidgets/Public/ITimeSlider.h
--- a/Engine/Source/Editor/SequencerWidgets/Public/ITimeSlider.h
+++ b/Engine/Source/Editor/SequencerWidgets/Public/ITimeSlider.h
@@ -2,6 +2,8 @@
 
 #pragma once
 
+class USequencerSettings;
+
 /** Enum specifying how to interpolate to a new view range */
 enum class EViewRangeInterpolation
 {
@@ -60,6 +62,7 @@ struct FTimeSliderArgs
 		, OnEndScrubberMovement()
 		, OnViewRangeChanged()
 		, AllowZoom(true)
+		, Settings(nullptr)
 	{}
 
 	/** The scrub position */
@@ -80,6 +83,8 @@ struct FTimeSliderArgs
 	FOnViewRangeChanged OnViewRangeChanged;
 	/** If we are allowed to zoom */
 	bool AllowZoom;
+	/** User-supplied settings object */
+	USequencerSettings* Settings;
 };
 
 class ITimeSliderController
diff --git a/Engine/Source/Editor/UMGEditor/Private/WidgetBlueprintEditor.cpp b/Engine/Source/Editor/UMGEditor/Private/WidgetBlueprintEditor.cpp
--- a/Engine/Source/Editor/UMGEditor/Private/WidgetBlueprintEditor.cpp
+++ b/Engine/Source/Editor/UMGEditor/Private/WidgetBlueprintEditor.cpp
@@ -541,7 +541,7 @@ TSharedPtr<ISequencer>& FWidgetBlueprintEditor::GetSequencer()
 {
 	if(!Sequencer.IsValid())
 	{
-		FSequencerViewParams ViewParams;
+		FSequencerViewParams ViewParams( TEXT( "UMGSequencerSettings" ) );
 		ViewParams.InitalViewRange = TRange<float>(-0.02f, 3.2f);
 		ViewParams.InitialScrubPosition = 0;
 		ViewParams.OnGetAddMenuContent = FOnGetAddMenuContent::CreateSP(this, &FWidgetBlueprintEditor::OnGetAnimationAddMenuContent);
