Sequencer: Fix a bug with adding a key to the nearest section (needs to add to all key area's nearest owning sections)

diff --git a/Engine/Source/Editor/Sequencer/Private/SAnimationOutlinerView.cpp b/Engine/Source/Editor/Sequencer/Private/SAnimationOutlinerView.cpp
--- a/Engine/Source/Editor/Sequencer/Private/SAnimationOutlinerView.cpp
+++ b/Engine/Source/Editor/Sequencer/Private/SAnimationOutlinerView.cpp
@@ -276,7 +276,6 @@ FReply SAnimationOutlinerTreeNode::OnAddKeyClicked()
 			OwningSection->SetFlags(RF_Transactional);
 			OwningSection->Modify();
 			KeyArea->AddKeyUnique(CurrentTime);
-			break;
 		}
 	}
 
