Fixed use of undefined type in a header

diff --git a/Engine/Source/Editor/Sequencer/Private/SectionDragOperations.cpp b/Engine/Source/Editor/Sequencer/Private/SectionDragOperations.cpp
--- a/Engine/Source/Editor/Sequencer/Private/SectionDragOperations.cpp
+++ b/Engine/Source/Editor/Sequencer/Private/SectionDragOperations.cpp
@@ -14,6 +14,12 @@
 /** How many pixels near the mouse has to be before snapping occurs */
 const float PixelSnapWidth = 10.f;
 
+FSequencerDragOperation::FSequencerDragOperation( FSequencer& InSequencer )
+	: Sequencer(InSequencer)
+{
+	Settings = Sequencer.GetSettings();
+}
+
 void FSequencerDragOperation::BeginTransaction( UMovieSceneSection& Section, const FText& TransactionDesc )
 {
 	// Begin an editor transaction and mark the section as transactional so it's state will be saved
diff --git a/Engine/Source/Editor/Sequencer/Private/SectionDragOperations.h b/Engine/Source/Editor/Sequencer/Private/SectionDragOperations.h
--- a/Engine/Source/Editor/Sequencer/Private/SectionDragOperations.h
+++ b/Engine/Source/Editor/Sequencer/Private/SectionDragOperations.h
@@ -2,17 +2,15 @@
 
 #pragma once
 
+class USequencerSettings;
+
 /**
  * An drag operation that occurs on a section or key
  */
 class FSequencerDragOperation
 {
 public:
-	FSequencerDragOperation( FSequencer& InSequencer )
-		: Sequencer(InSequencer)
-	{
-		Settings = Sequencer.GetSettings();
-	}
+	FSequencerDragOperation( FSequencer& InSequencer );
 
 	virtual ~FSequencerDragOperation(){}
 	void BeginTransaction( UMovieSceneSection& Section, const FText& TransactionDesc );
