Sequencer: Tweak unzoom to use KINDA_SMALL_NUMBER.

diff --git a/Engine/Source/Editor/Sequencer/Private/Sequencer.cpp b/Engine/Source/Editor/Sequencer/Private/Sequencer.cpp
--- a/Engine/Source/Editor/Sequencer/Private/Sequencer.cpp
+++ b/Engine/Source/Editor/Sequencer/Private/Sequencer.cpp
@@ -1358,8 +1358,8 @@ void FSequencer::ZoomToSelectedSections()
 	{
 		// Zoom back to last view range if already expanded
 		if (!ViewRangeBeforeZoom.IsEmpty() &&
-			FMath::IsNearlyEqual(BoundsHull.GetLowerBoundValue(), GetViewRange().GetLowerBoundValue()) &&
-			FMath::IsNearlyEqual(BoundsHull.GetUpperBoundValue(), GetViewRange().GetUpperBoundValue()))
+			FMath::IsNearlyEqual(BoundsHull.GetLowerBoundValue(), GetViewRange().GetLowerBoundValue(), KINDA_SMALL_NUMBER) &&
+			FMath::IsNearlyEqual(BoundsHull.GetUpperBoundValue(), GetViewRange().GetUpperBoundValue(), KINDA_SMALL_NUMBER))
 		{
 			OnViewRangeChanged(ViewRangeBeforeZoom);
 		}
