Sequencer: Fix tooltip typo.

diff --git a/Engine/Source/Editor/Sequencer/Private/SequencerCommands.cpp b/Engine/Source/Editor/Sequencer/Private/SequencerCommands.cpp
--- a/Engine/Source/Editor/Sequencer/Private/SequencerCommands.cpp
+++ b/Engine/Source/Editor/Sequencer/Private/SequencerCommands.cpp
@@ -38,7 +38,7 @@ void FSequencerCommands::RegisterCommands()
 	UI_COMMAND( ToggleCleanView, "Clean View", "Enable 'Clean View' mode which displays only global tracks when no filter is applied.", EUserInterfaceActionType::ToggleButton, FInputChord() );
 	
 	UI_COMMAND( ToggleShowCurveEditor, "Curve Editor", "Show the animation keys in a curve editor.", EUserInterfaceActionType::ToggleButton, FInputChord() );
-	UI_COMMAND( ToggleShowCurveEditorCurveToolTips, "Curve Tool Tips", "Show a tool tip with name and values when hoving over a curve.", EUserInterfaceActionType::ToggleButton, FInputChord() );
+	UI_COMMAND( ToggleShowCurveEditorCurveToolTips, "Curve Tool Tips", "Show a tool tip with name and values when hovering over a curve.", EUserInterfaceActionType::ToggleButton, FInputChord() );
 	UI_COMMAND( SetAllCurveVisibility, "All Curves", "Show all curves in the curve editor.", EUserInterfaceActionType::RadioButton, FInputChord() );
 	UI_COMMAND( SetSelectedCurveVisibility, "Selected Curves", "Show only selected curves in the curve editor.", EUserInterfaceActionType::RadioButton, FInputChord() );
 	UI_COMMAND( SetAnimatedCurveVisibility, "Animated Curves", "Show only curves with keyframes in the curve editor.", EUserInterfaceActionType::RadioButton, FInputChord() );
