Sequencer: Sequencer editor mode to process hotkeys.

diff --git a/Engine/Source/Editor/Sequencer/Private/Sequencer.cpp b/Engine/Source/Editor/Sequencer/Private/Sequencer.cpp
--- a/Engine/Source/Editor/Sequencer/Private/Sequencer.cpp
+++ b/Engine/Source/Editor/Sequencer/Private/Sequencer.cpp
@@ -2,6 +2,7 @@
 
 #include "SequencerPrivatePCH.h"
 #include "Sequencer.h"
+#include "SequencerEdMode.h"
 #include "ISequencerObjectBindingManager.h"
 #include "MovieScene.h"
 #include "Engine/LevelScriptBlueprint.h"
@@ -131,8 +132,6 @@ void FSequencer::InitSequencer( const FSequencerInitParams& InitParams, const TA
 			TrackEditors.Add( TrackEditor );
 		}
 
-
-
 		ZoomAnimation = FCurveSequence();
 		ZoomCurve = ZoomAnimation.AddCurve(0.f, 0.2f, ECurveEaseFunction::QuadIn);
 		OverlayAnimation = FCurveSequence();
@@ -144,6 +143,16 @@ void FSequencer::InitSequencer( const FSequencerInitParams& InitParams, const TA
 		// NOTE: Could fill in asset editor commands here!
 
 		BindSequencerCommands();
+
+		if( GLevelEditorModeTools().IsModeActive( FSequencerEdMode::EM_SequencerMode ) )
+		{
+			GLevelEditorModeTools().DeactivateMode( FSequencerEdMode::EM_SequencerMode );
+		}
+
+		GLevelEditorModeTools().ActivateMode( FSequencerEdMode::EM_SequencerMode );
+
+		FSequencerEdMode* SequencerEdMode = (FSequencerEdMode*)(GLevelEditorModeTools().GetActiveMode(FSequencerEdMode::EM_SequencerMode));
+		SequencerEdMode->SetSequencer(ActiveSequencer.Get());
 	}
 }
 
@@ -190,6 +199,11 @@ FSequencer::OnClose()
 		FEditorDelegates::PostSaveWorld.RemoveAll(this);
 
 		ActiveSequencer.Reset();
+
+		if( GLevelEditorModeTools().IsModeActive( FSequencerEdMode::EM_SequencerMode ) )
+		{
+			GLevelEditorModeTools().DeactivateMode( FSequencerEdMode::EM_SequencerMode );
+		}
 	}
 }
 
diff --git a/Engine/Source/Editor/Sequencer/Private/SequencerEdMode.cpp b/Engine/Source/Editor/Sequencer/Private/SequencerEdMode.cpp
new file mode 100644
--- /dev/null
+++ b/Engine/Source/Editor/Sequencer/Private/SequencerEdMode.cpp
@@ -0,0 +1,71 @@
+// Copyright 1998-2015 Epic Games, Inc. All Rights Reserved.
+
+#include "SequencerPrivatePCH.h"
+#include "SequencerEdMode.h"
+#include "Sequencer.h"
+
+const FEditorModeID FSequencerEdMode::EM_SequencerMode(TEXT("EM_SequencerMode"));
+
+FSequencerEdMode::FSequencerEdMode() :
+	Sequencer(NULL)
+{
+	FSequencerEdModeTool* SequencerEdModeTool = new FSequencerEdModeTool();
+
+	Tools.Add( SequencerEdModeTool );
+	SetCurrentTool( SequencerEdModeTool );
+}
+
+FSequencerEdMode::~FSequencerEdMode()
+{
+}
+
+void FSequencerEdMode::Enter()
+{
+	FEdMode::Enter();
+}
+
+void FSequencerEdMode::Exit()
+{
+	Sequencer = NULL;
+
+	FEdMode::Exit();
+}
+
+bool FSequencerEdMode::IsCompatibleWith(FEditorModeID OtherModeID) const
+{
+	// Compatible with all modes so that we can take over with the sequencer hotkeys
+	return true;
+}
+
+bool FSequencerEdMode::InputKey( FEditorViewportClient* ViewportClient, FViewport* Viewport, FKey Key, EInputEvent Event )
+{
+	if (Sequencer == NULL)
+	{
+		return false;
+	}
+
+	if (Event != IE_Released)
+	{
+		FModifierKeysState KeyState = FSlateApplication::Get().GetModifierKeys();
+
+		if (Sequencer->GetCommandBindings().Get()->ProcessCommandBindings(Key, KeyState, (Event == IE_Repeat) ))
+		{
+			return true;
+		}
+	}
+
+	return FEdMode::InputKey(ViewportClient, Viewport, Key, Event);
+}
+
+FSequencerEdModeTool::FSequencerEdModeTool()
+{
+}
+
+FSequencerEdModeTool::~FSequencerEdModeTool()
+{
+}
+
+bool FSequencerEdModeTool::InputKey(FEditorViewportClient* ViewportClient, FViewport* Viewport, FKey Key, EInputEvent Event)
+{
+	return FModeTool::InputKey(ViewportClient, Viewport, Key, Event);
+}
diff --git a/Engine/Source/Editor/Sequencer/Private/SequencerEdMode.h b/Engine/Source/Editor/Sequencer/Private/SequencerEdMode.h
new file mode 100644
--- /dev/null
+++ b/Engine/Source/Editor/Sequencer/Private/SequencerEdMode.h
@@ -0,0 +1,47 @@
+// Copyright 1998-2015 Epic Games, Inc. All Rights Reserved.
+
+#pragma once
+#include "EditorModeTools.h"
+
+class FSequencer;
+
+/**
+ * FSequencerEdMode is the editor mode for additional drawing and handling sequencer hotkeys in the editor
+ */
+class FSequencerEdMode : public FEdMode
+{
+public:
+	static const FEditorModeID EM_SequencerMode;
+
+public:
+	FSequencerEdMode();
+	virtual ~FSequencerEdMode();
+
+	/* FEdMode interface */
+	virtual void Enter() override;
+	virtual void Exit() override;
+	virtual bool IsCompatibleWith(FEditorModeID OtherModeID) const override;
+	virtual bool InputKey( FEditorViewportClient* ViewportClient, FViewport* Viewport, FKey Key, EInputEvent Event ) override;
+
+	void SetSequencer(FSequencer* InSequencer) { Sequencer = InSequencer; }
+
+private:
+	FSequencer* Sequencer;
+};
+
+/**
+ * FSequencerEdMode is the editor mode tool for additional drawing and handling sequencer hotkeys in the editor
+ */
+class FSequencerEdModeTool : public FModeTool
+{
+public:
+	FSequencerEdModeTool();
+	virtual ~FSequencerEdModeTool();
+
+	virtual FString GetName() const override { return TEXT("Sequencer Edit"); }
+
+	/**
+	 * @return		true if the key was handled by this editor mode tool.
+	 */
+	virtual bool InputKey(FEditorViewportClient* ViewportClient, FViewport* Viewport, FKey Key, EInputEvent Event) override;
+};
diff --git a/Engine/Source/Editor/Sequencer/Private/SequencerModule.cpp b/Engine/Source/Editor/Sequencer/Private/SequencerModule.cpp
--- a/Engine/Source/Editor/Sequencer/Private/SequencerModule.cpp
+++ b/Engine/Source/Editor/Sequencer/Private/SequencerModule.cpp
@@ -6,6 +6,7 @@
 #include "Toolkits/ToolkitManager.h"
 #include "SequencerCommands.h"
 #include "SequencerAssetEditor.h"
+#include "SequencerEdMode.h"
 #include "SequencerObjectChangeListener.h"
 
 // We disable the deprecation warnings here because otherwise it'll complain about us
@@ -80,6 +81,12 @@ class FSequencerModule : public ISequencerModule
 		if (GIsEditor)
 		{
 			FSequencerCommands::Register();
+
+			FEditorModeRegistry::Get().RegisterMode<FSequencerEdMode>(
+				FSequencerEdMode::EM_SequencerMode,
+				NSLOCTEXT("Sequencer", "SequencerEditMode", "Sequencer Mode"),
+				FSlateIcon(),
+				false);
 		}
 	}
 
@@ -88,6 +95,8 @@ class FSequencerModule : public ISequencerModule
 		if (GIsEditor)
 		{
 			FSequencerCommands::Unregister();
+
+			FEditorModeRegistry::Get().UnregisterMode(FSequencerEdMode::EM_SequencerMode);
 		}
 	}
 private:
