MovieScene: Continued to refactor MovieScene and Sequencer

diff --git a/Engine/Source/Editor/MovieSceneTools/Private/TrackEditors/PathTrackEditor.cpp b/Engine/Source/Editor/MovieSceneTools/Private/TrackEditors/PathTrackEditor.cpp
--- a/Engine/Source/Editor/MovieSceneTools/Private/TrackEditors/PathTrackEditor.cpp
+++ b/Engine/Source/Editor/MovieSceneTools/Private/TrackEditors/PathTrackEditor.cpp
@@ -200,7 +200,7 @@ void F3DPathTrackEditor::BuildObjectBindingContextMenu(FMenuBuilder& MenuBuilder
 		}
 		else
 		{
-			AActor* ActorWithSplineComponent = NULL;
+			AActor* ActorWithSplineComponent = nullptr;
 			FFormatNamedArguments Args;
 			MenuBuilder.AddMenuEntry(
 				FText::Format( NSLOCTEXT("Sequencer", "AddPath", "Add Path"), Args),
@@ -280,14 +280,9 @@ void F3DPathTrackEditor::SetPath(UMovieSceneSection* Section, AActor* ActorWithS
 	const FScopedTransaction Transaction(NSLOCTEXT("Sequencer", "UndoSetPath", "Set Path"));
 
 	UMovieScene3DPathSection* PathSection = (UMovieScene3DPathSection*)(Section);
+	FGuid SplineId = FindOrCreateHandleToObject(ActorWithSplineComponent);
 
-	//@todo sequencer - replace this with movie scene object id stuff
-	FGuid SplineGuid = FindOrCreateHandleToObject(ActorWithSplineComponent);
-
-	FMovieSceneObjectId SplineId;
-	SplineId.ObjectGuid = SplineGuid;
-
-	if (SplineGuid.IsValid())
+	if (SplineId.IsValid())
 	{
 		PathSection->SetPathId(SplineId);
 	}
@@ -295,25 +290,22 @@ void F3DPathTrackEditor::SetPath(UMovieSceneSection* Section, AActor* ActorWithS
 
 void F3DPathTrackEditor::AddKeyInternal( float KeyTime, const TArray<UObject*> Objects, UObject* AdditionalAsset)
 {
-	AActor* ActorWithSplineComponent = NULL;
+	AActor* ActorWithSplineComponent = nullptr;
 	
-	if (AdditionalAsset != NULL)
+	if (AdditionalAsset != nullptr)
 	{
 		ActorWithSplineComponent = Cast<AActor>(AdditionalAsset);
 	}
 
-	//@todo sequencer - replace this with movie scene object id stuff
-	FGuid SplineGuid;
-	if (ActorWithSplineComponent != NULL)
+	FGuid SplineId;
+
+	if (ActorWithSplineComponent != nullptr)
 	{
-		SplineGuid = FindOrCreateHandleToObject(ActorWithSplineComponent);
+		SplineId = FindOrCreateHandleToObject(ActorWithSplineComponent);
 	}
 
-	if (SplineGuid.IsValid())
+	if (SplineId.IsValid())
 	{
-		FMovieSceneObjectId SplineId;
-		SplineId.ObjectGuid = SplineGuid;
-
 		for( int32 ObjectIndex = 0; ObjectIndex < Objects.Num(); ++ObjectIndex )
 		{
 			UObject* Object = Objects[ObjectIndex];
diff --git a/Engine/Source/Editor/Sequencer/Private/SequencerBinding.cpp b/Engine/Source/Editor/Sequencer/Private/SequencerBinding.cpp
new file mode 100644
--- /dev/null
+++ b/Engine/Source/Editor/Sequencer/Private/SequencerBinding.cpp
@@ -0,0 +1,46 @@
+// Copyright 1998-2015 Epic Games, Inc. All Rights Reserved.
+
+#include "SequencerPrivatePCH.h"
+#include "SequencerBinding.h"
+
+
+/* FSequencerBinding interface
+ *****************************************************************************/
+
+UObject* FSequencerBinding::GetBoundObject() const
+{
+	// get component-less object
+	if (ComponentName.IsEmpty())
+	{
+		return Object;
+	}
+
+	// get cached component
+	if ((CachedComponent != nullptr) && (CachedComponent->GetName() == ComponentName))
+	{
+		return CachedComponent;
+	}
+
+	CachedComponent = nullptr;
+
+	// find & cache component
+	AActor* Actor = Cast<AActor>(Object);
+
+	if (Actor == nullptr)
+	{
+		return Object;
+	}
+
+	for (UActorComponent* ActorComponent : Actor->GetComponents())
+	{
+		if (ActorComponent->GetName() == ComponentName)
+		{
+			CachedComponent = ActorComponent;
+
+			return ActorComponent;
+		}
+	}
+
+	// component not found
+	return nullptr;
+}
diff --git a/Engine/Source/Editor/Sequencer/Private/SequencerBinding.h b/Engine/Source/Editor/Sequencer/Private/SequencerBinding.h
new file mode 100644
--- /dev/null
+++ b/Engine/Source/Editor/Sequencer/Private/SequencerBinding.h
@@ -0,0 +1,78 @@
+// Copyright 1998-2015 Epic Games, Inc. All Rights Reserved.
+
+#pragma once
+
+#include "SequencerBinding.generated.h"
+
+
+/**
+ * Structure for object bindings.
+ */
+USTRUCT()
+struct FSequencerBinding
+{
+	GENERATED_USTRUCT_BODY(FSequencerBinding)
+
+public:
+
+	/** Creates and initializes a new instance. */
+	FSequencerBinding()
+		: Object(nullptr)
+		, CachedComponent(nullptr)
+	{ }
+
+	/**
+	 * Creates and initializes an object binding.
+	 *
+	 * @param InObject The object to be bound.
+	 */
+	FSequencerBinding(UObject* InObject)
+		: Object(InObject)
+		, CachedComponent(nullptr)
+	{ }
+
+	/**
+	 * Creates and initializes a component binding.
+	 *
+	 * @param InObject The object that owns the component.
+	 * @param InComponentName The component to be bound.
+	 */
+	FSequencerBinding(UObject* InObject, FString InComponentName)
+		: Object(InObject)
+		, ComponentName(InComponentName)
+		, CachedComponent(nullptr)
+	{ }
+
+	/**
+	 * Compares two bindings for equality.
+	 *
+	 * @param X The first binding to compare.
+	 * @param Y The second binding to compare.
+	 * @return true if the bindings refer to the same object, false otherwise.
+	 */
+	friend bool operator==(const FSequencerBinding& X, const FSequencerBinding& Y)
+	{
+		return (X.Object == Y.Object) && (X.ComponentName == Y.ComponentName);
+	}
+
+	/**
+	 * Gets a pointer to the bound object.
+	 *
+	 * @return The bound object.
+	 */
+	SEQUENCER_API UObject* GetBoundObject() const;
+
+private:
+
+	/** The object being bound, usually an Actor. */
+	UPROPERTY()
+	UObject* Object; // @todo sequencer: gmp: this needs to become a persistent object pointer
+
+	/** Optional name of the ActorComponent being bound. */
+	UPROPERTY()
+	FString ComponentName;
+
+	/** Cached pointer to the Actor component, if bound. */
+	UPROPERTY(transient)
+	mutable UObject* CachedComponent;
+};
diff --git a/Engine/Source/Editor/Sequencer/Private/SequencerBindingManager.cpp b/Engine/Source/Editor/Sequencer/Private/SequencerBindingManager.cpp
--- a/Engine/Source/Editor/Sequencer/Private/SequencerBindingManager.cpp
+++ b/Engine/Source/Editor/Sequencer/Private/SequencerBindingManager.cpp
@@ -3,6 +3,7 @@
 #include "SequencerPrivatePCH.h"
 #include "MovieScene.h"
 #include "MovieSceneInstance.h"
+#include "SequencerBinding.h"
 #include "SequencerBindingManager.h"
 
 // @todo sequencer: gmp: refactor these dependencies into another module
@@ -10,26 +11,6 @@
 #include "LevelEditor.h"
 
 
-/* Private types
- *****************************************************************************/
-
-/**
- * Puppet actor info
- */
-struct FSequencerActorInfo
-{
-	/** Unique object identifier. */
-	FMovieSceneObjectId ObjectId;
-
-	/**
-	 * The actual puppet actor that we are actively managing.
-	 * These actors actually live in the editor level just like any other actor,
-	 * but we'll directly manage their lifetimes while in the editor.
-	 */
-	TWeakObjectPtr<AActor> Actor;	
-};
-
-
 /* USequencerBindingManager structors
  *****************************************************************************/
 
@@ -59,9 +40,24 @@ bool USequencerBindingManager::AllowsSpawnableObjects() const
 	return true;
 }
 
-void USequencerBindingManager::BindPossessableObject(const FMovieSceneObjectId& ObjectId, UObject& PossessedObject)
+
+void USequencerBindingManager::BindPossessableObject(const FGuid& ObjectId, UObject& PossessedObject)
 {
-	// @todo sequencer: gmp: bind possessable object
+	FString ComponentName;
+	UActorComponent* Component = Cast<UActorComponent>(&PossessedObject);
+
+	if (Component != nullptr)
+	{
+		// @todo sequencer: gmp: need support for UStruct keys in TMap
+		Bindings.Add(ObjectId.ToString(), FSequencerBinding(PossessedObject.GetOuter(), ComponentName));
+		//Bindings.Add(ObjectId, FSequencerBinding(PossessedObject.GetOuter(), ComponentName));
+	}
+	else
+	{
+		// @todo sequencer: gmp: need support for UStruct keys in TMap
+		Bindings.Add(ObjectId.ToString(), FSequencerBinding(&PossessedObject));
+		//Bindings.Add(ObjectId, FSequencerBinding(Actor));
+	}
 }
 
 
@@ -71,32 +67,41 @@ bool USequencerBindingManager::CanPossessObject(UObject& Object) const
 }
 
 
-void USequencerBindingManager::DestroyAllSpawnedObjects()
+void USequencerBindingManager::ClearBindings()
 {
-	const bool bDestroyAll = true;
+	Bindings.Empty();
+}
 
-	for (auto It = InstanceToActorInfos.CreateConstIterator(); It; ++It)
-	{
-		SpawnOrDestroyObjectsForInstance(It.Key().Pin().ToSharedRef(), bDestroyAll);
-	}
 
-	InstanceToActorInfos.Empty();
+void USequencerBindingManager::DestroyAllSpawnedObjects(UMovieScene& MovieScene)
+{
+	const bool DestroyAll = true;
+	SpawnOrDestroyObjects(&MovieScene, DestroyAll);
+	SpawnedActors.Empty();
 }
 
 
-FMovieSceneObjectId USequencerBindingManager::FindIdForObject(const UMovieScene& MovieScene, UObject& Object) const
+FGuid USequencerBindingManager::FindObjectId(const UMovieScene& MovieScene, UObject& Object) const
 {
-	FMovieSceneObjectId Result = FindIdForSpawnableObject(Object);
-
-	if (Result.IsValid())
+	// search bindings (possessables)
+	for (auto BindingPair : Bindings)
 	{
-		return Result;
+		const FSequencerBinding& Binding = BindingPair.Value;
+
+		if (Binding.GetBoundObject() == &Object)
+		{
+			FGuid Result; FGuid::Parse(BindingPair.Key, Result); return Result;
+			//return BindingPair.Key;
+		}
 	}
 
+	// search proxies (spawnables)
 	const bool bIsGamePreviewObject = !!(Object.GetOutermost()->PackageFlags & PKG_PlayInEditor);
 
 	if (bIsGamePreviewObject)
 	{
+		// @todo sequencer: gmp: i don't want this class to know the MovieScene that owns it
+		//
 		// someone is asking for a handle to an object from a game preview session,
 		// probably because they want to capture keys during live simulation.
 		//
@@ -109,141 +114,122 @@ FMovieSceneObjectId USequencerBindingManager::FindIdForObject(const UMovieScene&
 		
 		if (FoundSpawnable != nullptr)
 		{
-			Result = FoundSpawnable->GetObjectId();
+			return FoundSpawnable->GetGuid();
 		}
 	}
 	else
 	{
-		// @todo sequencer: gmp: update bindings
-/*		BindToPlayMovieSceneNode(false);
-
-		// when editing within the level editor, make sure we're bound to the level
-		// script node which contains data about possessables.
-		if (PlayMovieSceneNode.IsValid())
+		for (auto SpawnedActorPair : SpawnedActors)
 		{
-			ObjectGuid = PlayMovieSceneNode->FindGuidForObject(&Object);
-		}*/
+			if (SpawnedActorPair.Value.Get() == &Object)
+			{
+				return SpawnedActorPair.Key;
+			}
+		}
 	}
 
-	return Result;
+	return FGuid();
 }
 
 
-void USequencerBindingManager::GetRuntimeObjects(const TSharedRef<FMovieSceneInstance>& MovieSceneInstance, const FMovieSceneObjectId& ObjectId, TArray<UObject*>& OutRuntimeObjects) const
+UObject* USequencerBindingManager::FindObject(const FGuid& ObjectId) const
 {
-	UObject* Object = FindActorInfoById(MovieSceneInstance, ObjectId);
+	// search bindings (possessables)
+	// @todo sequencer: gmp: need support for UStruct keys in TMap
+	const FSequencerBinding* Binding = Bindings.Find(ObjectId.ToString());
+	//const FSequencerBinding* Binding = Bindings.Find(ObjectId);
 
-	if (Object)
+	if (Binding != nullptr)
 	{
-		// spawnable object
-		OutRuntimeObjects.Reset();
-		OutRuntimeObjects.Add(Object);
+		return Binding->GetBoundObject();
 	}
-	else
+
+	// search proxies (spawnables)
+	TWeakObjectPtr<AActor> Actor = SpawnedActors.FindRef(ObjectId);
+
+	if (Actor.IsValid())
 	{
-		// possessable object
-		// @todo sequenceR: gmp: bind object?
-/*		if (!PlayMovieSceneNode.IsValid())
-		{
-			BindToPlayMovieSceneNode(false);
-		}
-	
-		if (PlayMovieSceneNode.IsValid())
-		{
-			OutRuntimeObjects = PlayMovieSceneNode->FindBoundObjects( ObjectGuid );
-		}*/
+		return Actor.Get();
 	}
-}
 
-
-void USequencerBindingManager::RemoveMovieSceneInstance(const TSharedRef<FMovieSceneInstance>& MovieSceneInstance)
-{
-	SpawnOrDestroyObjectsForInstance(MovieSceneInstance, true);
-	InstanceToActorInfos.Remove(MovieSceneInstance);
+	// not found
+	return nullptr;
 }
 
 
-void USequencerBindingManager::SpawnOrDestroyObjectsForInstance(TSharedRef<FMovieSceneInstance> MovieSceneInstance, bool bDestroyAll)
+void USequencerBindingManager::SpawnOrDestroyObjects(UMovieScene* MovieScene, bool DestroyAll)
 {
 	bool bAnyLevelActorsChanged = false;
 
-	// Get the list of puppet objects for the movie scene
-	TArray<TSharedRef<FSequencerActorInfo>>& ActorInfos = InstanceToActorInfos.FindOrAdd(MovieSceneInstance);
-	UMovieScene* MovieScene = MovieSceneInstance->GetMovieScene();
-	UWorld* ActorWorld = GetWorld();
-
-	// remove any puppet objects that we no longer need
+	// remove any proxy actors that we no longer need
 	if (MovieScene != nullptr)
 	{
-		for (auto ActorInfoIndex = 0; ActorInfoIndex < ActorInfos.Num(); ++ActorInfoIndex)
+		for (auto SpawnedActorPair : SpawnedActors)
 		{
-			TSharedRef<FSequencerActorInfo> ActorInfo = ActorInfos[ActorInfoIndex];
-			bool bShouldDestroyActor = true;
+			bool ShouldDestroyActor = true;
 
-			// figure out if we still need this puppet actor
-			if (!bDestroyAll)
+			if (!DestroyAll)
 			{
+				// figure out if we still need this proxy actor
 				for (auto SpawnableIndex = 0; SpawnableIndex < MovieScene->GetSpawnableCount(); ++SpawnableIndex)
 				{
 					auto& Spawnable = MovieScene->GetSpawnable(SpawnableIndex);
 
-					if (Spawnable.GetObjectId() == ActorInfo->ObjectId)
+					if (Spawnable.GetGuid() == SpawnedActorPair.Key)
 					{
-						bShouldDestroyActor = false;
+						ShouldDestroyActor = false;
 						break;
 					}
 				}
 			}
 
-			if (!bShouldDestroyActor)
+			if (!ShouldDestroyActor)
 			{
 				// Actor is no longer valid (probably the world was destroyed)
 				continue;
 			}
 
-			AActor* Actor = ActorInfo->Actor.Get();
+			AActor* Actor = SpawnedActorPair.Value.Get();
 				
 			if (Actor == nullptr)
 			{
 				continue;
 			}
 
-			UWorld* PuppetWorld = Actor->GetWorld();
+			// destroy Actor
+			UWorld* World = Actor->GetWorld();
 
-			if (ensure(PuppetWorld != nullptr))
+			if (ensure(World != nullptr))
 			{
-				// destroy this actor
-				// Ignored unless called while game is running
-				const bool bNetForce = false;
-
-				// We don't want to dirty the level for puppet actor changes
-				const bool bShouldModifyLevel = false;
-
 				if (!bAnyLevelActorsChanged)
 				{
 					DeselectAllActors();
 				}
 
-				// Actor should never be selected in the editor at this point.  We took care of that up above.
+				// Actor should never be selected in the Editor at this point (we took care of that up above)
 				ensure(!Actor->IsSelected());
 
-				const bool bWasDestroyed = PuppetWorld->DestroyActor(Actor, bNetForce, bShouldModifyLevel);
+				const bool bNetForce = false;
+				const bool bShouldModifyLevel = false;
+				const bool bWasDestroyed = World->DestroyActor(Actor, bNetForce, bShouldModifyLevel);
 
 				if (bWasDestroyed)
 				{
 					bAnyLevelActorsChanged = true;
-					ActorInfos.RemoveAt(ActorInfoIndex--);
+					SpawnedActors.Remove(SpawnedActorPair.Key);
 				}
 				else
 				{
-					// @todo sequencer: At least one puppet couldn't be cleaned up!
+					// @todo sequencer: At least one proxy actor couldn't be cleaned up!
 				}
 			}
 		}
 	}
 
-	if( !bDestroyAll && MovieScene )
+	if (!DestroyAll && (MovieScene != nullptr))
 	{
+		UWorld* ActorWorld = GetWorld();
+
 		for (auto SpawnableIndex = 0; SpawnableIndex < MovieScene->GetSpawnableCount(); ++SpawnableIndex)
 		{
 			FMovieSceneSpawnable& Spawnable = MovieScene->GetSpawnable( SpawnableIndex );
@@ -254,22 +240,8 @@ void USequencerBindingManager::SpawnOrDestroyObjectsForInstance(TSharedRef<FMovi
 				continue;
 			}
 
-			// Do we already have a puppet for this spawnable?
-			bool bIsAlreadySpawned = false;
-
-			for (auto PuppetIndex = 0; PuppetIndex < ActorInfos.Num(); ++PuppetIndex)
-			{
-				auto& ActorInfo = ActorInfos[PuppetIndex];
-
-				if (ActorInfo->ObjectId == Spawnable.GetObjectId())
-				{
-					bIsAlreadySpawned = true;
-
-					break;
-				}
-			}
-
-			if (bIsAlreadySpawned)
+			// Do we already have a proxy for this spawnable?
+			if (SpawnedActors.Contains(Spawnable.GetGuid()))
 			{
 				continue;
 			}
@@ -312,14 +284,7 @@ void USequencerBindingManager::SpawnOrDestroyObjectsForInstance(TSharedRef<FMovi
 			{
 				// @todo sequencer: We're naming the actor based off of the spawnable's name.  Is that really what we want?
 				FActorLabelUtilities::SetActorLabelUnique(NewActor, Spawnable.GetName());
-
-				// Keep track of this actor
-				TSharedRef<FSequencerActorInfo> NewActorInfo(new FSequencerActorInfo());
-				{
-					NewActorInfo->ObjectId = Spawnable.GetObjectId();
-					NewActorInfo->Actor = NewActor;
-					ActorInfos.Add(NewActorInfo);
-				}
+				SpawnedActors.Add(Spawnable.GetGuid(), NewActor);
 			}
 			else
 			{
@@ -331,30 +296,33 @@ void USequencerBindingManager::SpawnOrDestroyObjectsForInstance(TSharedRef<FMovi
 }
 
 
-bool USequencerBindingManager::TryGetObjectBindingDisplayName(const TSharedRef<FMovieSceneInstance>& MovieSceneInstance, const FMovieSceneObjectId& ObjectId, FText& OutDisplayName) const
+bool USequencerBindingManager::TryGetObjectDisplayName(const FGuid& ObjectId, FText& OutDisplayName) const
 {
-	TArray<UObject*> RuntimeObjects;
-	GetRuntimeObjects(MovieSceneInstance, ObjectId, RuntimeObjects);
+	UObject* Object = FindObject(ObjectId);
 
-	for (int32 ObjIndex = 0; ObjIndex < RuntimeObjects.Num(); ++ObjIndex)
+	if (Object == nullptr)
 	{
-		AActor* Actor = Cast<AActor>(RuntimeObjects[ObjIndex]);
+		return false;
+	}
 
-		if (Actor != nullptr)
-		{
-			OutDisplayName = FText::FromString(Actor->GetActorLabel());
+	AActor* Actor = Cast<AActor>(Object);
 
-			return true;
-		}
+	if (Actor == nullptr)
+	{
+		return false;
 	}
 
-	return false;
+	OutDisplayName = FText::FromString(Actor->GetActorLabel());
+
+	return true;
 }
 
 
-void USequencerBindingManager::UnbindPossessableObjects(const FMovieSceneObjectId& ObjectId)
+void USequencerBindingManager::UnbindPossessableObjects(const FGuid& ObjectId)
 {
-	// @todo sequencer: gmp: unbind object
+	// @todo sequencer: gmp: need support for UStruct keys in TMap
+	Bindings.Remove(ObjectId.ToString());
+	//Bindings.Remove(ObjectId);
 }
 
 
@@ -363,82 +331,23 @@ void USequencerBindingManager::UnbindPossessableObjects(const FMovieSceneObjectI
 
 void USequencerBindingManager::DeselectAllActors()
 {
-	TArray<AActor*> ActorsToDeselect;
-
-	// find actors to deselect
-	for (auto ActorInfosPair : InstanceToActorInfos)
-	{
-		for (const TSharedRef<FSequencerActorInfo>& ActorInfo : ActorInfosPair.Value)
-		{
-			AActor* Actor = ActorInfo->Actor.Get();
-
-			if (Actor != nullptr)
-			{
-				ActorsToDeselect.Add(Actor);
-			}
-		}
-	}
-
-	if (ActorsToDeselect.Num() == 0)
-	{
-		return;
-	}
-
-	// deselect actors
 	USelection* Selection = GEditor->GetSelectedActors();
 	{
 		Selection->BeginBatchSelectOperation();
 		Selection->MarkBatchDirty();
 
-		for (AActor* Actor : ActorsToDeselect)
-		{
-			GEditor->GetSelectedActors()->Deselect(Actor);
-		}
-
-		Selection->EndBatchSelectOperation();
-	}
-}
-
-
-FMovieSceneObjectId USequencerBindingManager::FindIdForSpawnableObject(UObject& Object) const
-{
-	for (auto ActorInfosPair : InstanceToActorInfos)
-	{
-		for (const TSharedRef<FSequencerActorInfo>& ActorInfo : ActorInfosPair.Value)
+		for (auto SpawnedActorPair : SpawnedActors)
 		{
-			AActor* Actor = ActorInfo->Actor.Get();
+			AActor* SpawnedActor = SpawnedActorPair.Value.Get();
 
-			if (Actor == &Object)
+			if (SpawnedActor != nullptr)
 			{
-				return ActorInfo->ObjectId;
+				GEditor->GetSelectedActors()->Deselect(SpawnedActor);
 			}
 		}
-	}
 
-	// not found
-	return FMovieSceneObjectId();
-}
-
-
-UObject* USequencerBindingManager::FindActorInfoById(const TSharedRef<FMovieSceneInstance>& MovieSceneInstance, const FMovieSceneObjectId& ObjectId) const
-{
-	const TArray<TSharedRef<FSequencerActorInfo>>* ActorInfos = InstanceToActorInfos.Find(MovieSceneInstance);
-
-	if (ActorInfos == nullptr)
-	{
-		return nullptr;
-	}
-
-	for (TSharedRef<FSequencerActorInfo> ActorInfo : *ActorInfos)
-	{
-		if (ActorInfo->ObjectId == ObjectId)
-		{
-			return ActorInfo->Actor.Get();
-		}
+		Selection->EndBatchSelectOperation();
 	}
-
-	// not found
-	return nullptr;
 }
 
 
@@ -477,33 +386,30 @@ void USequencerBindingManager::HandlePropagateObjectChanges(UObject* Object)
 		return;
 	}
 
-	for (auto ActorInfosPair : InstanceToActorInfos)
+	for (auto SpawnedActorPair : SpawnedActors)
 	{
-		// is this an object that we care about?
-		for (const TSharedRef<FSequencerActorInfo>& ActorInfo : ActorInfosPair.Value)
-		{
-			AActor* Actor = ActorInfo->Actor.Get();
+		AActor* SpawnedActor = SpawnedActorPair.Value.Get();
 
-			if (Actor == Object)
-			{
-				// @todo sequencer: propagation: Don't propagate changes that are being auto-keyed or key-adjusted!
-
-				// Our puppet actor was modified.
-				PropagateActorChanges(ActorInfo);
-			}
+		if (SpawnedActor == Object)
+		{
+			// @todo sequencer: propagation: Don't propagate changes that are being auto-keyed or key-adjusted!
+			PropagateActorChanges(SpawnedActorPair.Key, SpawnedActor);
 		}
 	}
 }
 
 
-void USequencerBindingManager::PropagateActorChanges(const TSharedRef<FSequencerActorInfo> ActorInfo)
+void USequencerBindingManager::PropagateActorChanges(const FGuid& ObjectId, AActor* Actor)
 {
-	AActor* Actor = ActorInfo->Actor.Get();
-	AActor* TargetActor = nullptr;
+	if (Actor == nullptr)
+	{
+		return;
+	}
 
-	// find the spawnable for this puppet actor
+	AActor* TargetActor = nullptr;
 	FMovieSceneSpawnable* FoundSpawnable = nullptr;
 
+	// find the spawnable for this puppet actor
 	// @todo sequencer: gmp: add horrible sequencer dependencies to binding manager
 /*	TArray<UMovieScene*> MovieScenesBeingEdited = Sequencer.Pin()->GetMovieScenesBeingEdited();
 
diff --git a/Engine/Source/Editor/Sequencer/Private/SequencerBindingManager.h b/Engine/Source/Editor/Sequencer/Private/SequencerBindingManager.h
--- a/Engine/Source/Editor/Sequencer/Private/SequencerBindingManager.h
+++ b/Engine/Source/Editor/Sequencer/Private/SequencerBindingManager.h
@@ -3,16 +3,12 @@
 #pragma once
 
 #include "MovieSceneBindingManager.h"
+#include "SequencerBinding.h"
 #include "SequencerBindingManager.generated.h"
 
 
-struct FSequencerActorInfo;
-
-
 /**
  * Movie scene binding manager used by Sequencer.
- *
- * @todo squencer: gmp: probably want to refactor this to not depend on Actors/move into own module
  */
 UCLASS(BlueprintType, MinimalAPI)
 class USequencerBindingManager
@@ -29,15 +25,15 @@ public:
 	// IMovieSceneBindingManager interface
 
 	virtual bool AllowsSpawnableObjects() const override;
-	virtual void BindPossessableObject(const FMovieSceneObjectId& ObjectId, UObject& PossessedObject) override;
-	virtual bool CanPossessObject( UObject& Object ) const override;
-	virtual void DestroyAllSpawnedObjects() override;
-	virtual FMovieSceneObjectId FindIdForObject(const UMovieScene& MovieScene, UObject& Object) const override;
-	virtual void GetRuntimeObjects(const TSharedRef<FMovieSceneInstance>& MovieSceneInstance, const FMovieSceneObjectId& ObjectId, TArray<UObject*>& OutRuntimeObjects) const override;
-	virtual void RemoveMovieSceneInstance(const TSharedRef<FMovieSceneInstance>& MovieSceneInstance) override;
-	virtual void SpawnOrDestroyObjectsForInstance(TSharedRef<FMovieSceneInstance> MovieSceneInstance, bool bDestroyAll) override;
-	virtual bool TryGetObjectBindingDisplayName(const TSharedRef<FMovieSceneInstance>& MovieSceneInstance, const FMovieSceneObjectId& ObjectId, FText& OutDisplayName) const override;
-	virtual void UnbindPossessableObjects(const FMovieSceneObjectId& ObjectId) override;
+	virtual void BindPossessableObject(const FGuid& ObjectId, UObject& PossessedObject) override;
+	virtual bool CanPossessObject(UObject& Object) const override;
+	virtual void ClearBindings() override;
+	virtual void DestroyAllSpawnedObjects(UMovieScene& MovieScene) override;
+	virtual FGuid FindObjectId(const UMovieScene& MovieScene, UObject& Object) const override;
+	virtual UObject* FindObject(const FGuid& ObjectId) const override;
+	virtual void SpawnOrDestroyObjects(UMovieScene* MovieScene, bool DestroyAll) override;
+	virtual bool TryGetObjectDisplayName(const FGuid& ObjectId, FText& OutDisplayName) const override;
+	virtual void UnbindPossessableObjects(const FGuid& ObjectId) override;
 
 protected:
 
@@ -46,31 +42,14 @@ protected:
 	 *
 	 * @todo sequencer: gmp: remove Editor dependencies from this class
 	 */
-	void DeselectAllActors();
-
-	/**
-	 * Find the unique identifier for a given proxy object.
-	 *
-	 * @param Object The proxy object to find the identifier for.
-	 * @return The object's identifier.
-	 * @see FindPuppetObjectForId
-	 */
-	FMovieSceneObjectId FindIdForSpawnableObject(UObject& Object) const;
-
-	/**
-	 * Find the puppet actor for a given object identifier.
-	 *
-	 * @param ObjectId The object identifier.
-	 * @return The corresponding object, or nullptr if not found.
-	 */
-	UObject* FindActorInfoById(const TSharedRef<FMovieSceneInstance>& MovieSceneInstance, const FMovieSceneObjectId& ObjectId) const;
+	SEQUENCER_API void DeselectAllActors();
 
 	/**
 	 * Get the world that possessed actors belong to or where spawned actors should go.
 	 *
 	 * @return The world object.
 	 */
-	UWorld* GetWorld() const;
+	SEQUENCER_API UWorld* GetWorld() const;
 
 private:
 
@@ -78,13 +57,19 @@ private:
 	void HandlePropagateObjectChanges(UObject* Object);
 
 	/** Refresh either the reference actor or the data descriptor BP after its puppet proxy actor has been changed. */
-	void PropagateActorChanges(const TSharedRef<FSequencerActorInfo> ActorInfo);
+	void PropagateActorChanges(const FGuid& ObjectId, AActor* Actor);
 
 private:
 
-	/** List of puppet objects that we've spawned and are actively managing. */
-	TMap<TWeakPtr<FMovieSceneInstance>, TArray<TSharedRef<FSequencerActorInfo>>> InstanceToActorInfos;
+	/** Collection of object bindings. */
+	// @todo sequencer: gmp: need support for UStruct keys in TMap
+	UPROPERTY()
+	TMap<FString, FSequencerBinding> Bindings;
+	//TMap<FGuid, FSequencerBinding> Bindings;
 
 	/** The object change listener. */
 	TWeakPtr<ISequencerObjectChangeListener> ObjectChangeListener;
+
+	/** Collection of spawned proxy actors. */
+	TMap<FGuid, TWeakObjectPtr<AActor>> SpawnedActors;
 };
diff --git a/Engine/Source/Editor/UMGEditor/Private/Animation/UMGBindingManager.cpp b/Engine/Source/Editor/UMGEditor/Private/Animation/UMGBindingManager.cpp
--- a/Engine/Source/Editor/UMGEditor/Private/Animation/UMGBindingManager.cpp
+++ b/Engine/Source/Editor/UMGEditor/Private/Animation/UMGBindingManager.cpp
@@ -2,7 +2,6 @@
 
 #include "UMGEditorPrivatePCH.h"
 #include "MovieScene.h"
-#include "MovieSceneInstance.h"
 #include "UMGBindingManager.h"
 #include "Animation/WidgetAnimation.h"
 #include "WidgetBlueprint.h"
@@ -55,13 +54,13 @@ void UUMGBindingManager::InitPreviewObjects()
 
 		if (FoundSlot == nullptr)
 		{
-			ObjectIdToPreviewObjects.Add(Binding.AnimationId, FoundObject);
-			PreviewObjectToId.Add(FoundObject, Binding.AnimationId);
+			IdToPreviewObjects.Add(Binding.AnimationId, FoundObject);
+			PreviewObjectToIds.Add(FoundObject, Binding.AnimationId);
 		}
 		else
 		{
-			ObjectIdToSlotContentPreviewObjects.Add(Binding.AnimationId, FoundSlot->Content);
-			SlotContentPreviewObjectToId.Add(FoundSlot->Content, Binding.AnimationId);
+			IdToSlotContentPreviewObjects.Add(Binding.AnimationId, FoundSlot->Content);
+			SlotContentPreviewObjectToIds.Add(FoundSlot->Content, Binding.AnimationId);
 		}
 	}
 }
@@ -94,7 +93,7 @@ bool UUMGBindingManager::AllowsSpawnableObjects() const
 }
 
 
-void UUMGBindingManager::BindPossessableObject(const FMovieSceneObjectId& ObjectId, UObject& PossessedObject)
+void UUMGBindingManager::BindPossessableObject(const FGuid& ObjectId, UObject& PossessedObject)
 {
 	if (WidgetAnimation == nullptr)
 	{
@@ -105,8 +104,8 @@ void UUMGBindingManager::BindPossessableObject(const FMovieSceneObjectId& Object
 
 	if ((PossessedSlot != nullptr) && (PossessedSlot->Content != nullptr))
 	{
-		SlotContentPreviewObjectToId.Add(PossessedSlot->Content, ObjectId);
-		ObjectIdToSlotContentPreviewObjects.Add(ObjectId, PossessedSlot->Content);
+		SlotContentPreviewObjectToIds.Add(PossessedSlot->Content, ObjectId);
+		IdToSlotContentPreviewObjects.Add(ObjectId, PossessedSlot->Content);
 
 		// Save the name of the widget containing the slots. This is the object
 		// to look up that contains the slot itself (the thing we are animating).
@@ -121,8 +120,8 @@ void UUMGBindingManager::BindPossessableObject(const FMovieSceneObjectId& Object
 	}
 	else if (PossessedSlot == nullptr)
 	{
-		PreviewObjectToId.Add(&PossessedObject, ObjectId);
-		ObjectIdToPreviewObjects.Add(ObjectId, &PossessedObject);
+		PreviewObjectToIds.Add(&PossessedObject, ObjectId);
+		IdToPreviewObjects.Add(ObjectId, &PossessedObject);
 
 		FWidgetAnimationBinding NewBinding;
 		{
@@ -157,98 +156,102 @@ bool UUMGBindingManager::CanPossessObject(UObject& Object) const
 }
 
 
-FMovieSceneObjectId UUMGBindingManager::FindIdForObject(const UMovieScene& MovieScene, UObject& Object) const
+void UUMGBindingManager::ClearBindings()
+{
+	 // @todo sequencer: clear UMG bindings?
+}
+
+
+FGuid UUMGBindingManager::FindObjectId(const UMovieScene& MovieScene, UObject& Object) const
 {
 	UPanelSlot* Slot = Cast<UPanelSlot>(&Object);
 
 	if (Slot != nullptr)
 	{
 		// slot guids are tracked by their content.
-		return PreviewObjectToId.FindRef(Slot->Content);
+		return PreviewObjectToIds.FindRef(Slot->Content);
 	}
 
-	return PreviewObjectToId.FindRef(&Object);
+	return PreviewObjectToIds.FindRef(&Object);
 }
 
 
-void UUMGBindingManager::GetRuntimeObjects(const TSharedRef<FMovieSceneInstance>& MovieSceneInstance, const FMovieSceneObjectId& ObjectId, TArray<UObject*>& OutRuntimeObjects) const
+UObject* UUMGBindingManager::FindObject(const FGuid& ObjectId) const
 {
-	TArray<TWeakObjectPtr<UObject>> PreviewObjects;
-	ObjectIdToPreviewObjects.MultiFind(ObjectId, PreviewObjects);
+	TWeakObjectPtr<UObject> PreviewObject = IdToPreviewObjects.FindRef(ObjectId);
 
-	for (TWeakObjectPtr<UObject>& PreviewObject : PreviewObjects)
+	if (PreviewObject.IsValid())
 	{
-		OutRuntimeObjects.Add(PreviewObject.Get());
+		return PreviewObject.Get();
 	}
 
-	TArray<TWeakObjectPtr<UObject>> SlotContentPreviewObjects;
-	ObjectIdToSlotContentPreviewObjects.MultiFind(ObjectId, SlotContentPreviewObjects);
+	TWeakObjectPtr<UObject> SlotContentPreviewObject = IdToSlotContentPreviewObjects.FindRef(ObjectId);
 
-	for (TWeakObjectPtr<UObject>& SlotContentPreviewObject : SlotContentPreviewObjects)
+	if (SlotContentPreviewObject.IsValid())
 	{
-		UWidget* ContentWidget = Cast<UWidget>(SlotContentPreviewObject.Get());
-		OutRuntimeObjects.Add(ContentWidget->Slot);
+		return SlotContentPreviewObject.Get();
 	}
+	
+	return nullptr;
 }
 
 
-bool UUMGBindingManager::TryGetObjectBindingDisplayName(const TSharedRef<FMovieSceneInstance>& MovieSceneInstance, const FMovieSceneObjectId& ObjectId, FText& OutDisplayName) const
+bool UUMGBindingManager::TryGetObjectDisplayName(const FGuid& ObjectId, FText& OutDisplayName) const
 {
-	// TODO: This gets called every frame for every bound object and could be a potential performance issue for a really complicated animation.
-	TArray<TWeakObjectPtr<UObject>> BindingObjects;
-	ObjectIdToPreviewObjects.MultiFind(ObjectId, BindingObjects);
-	
-	TArray<TWeakObjectPtr<UObject>> SlotContentBindingObjects;
-	ObjectIdToSlotContentPreviewObjects.MultiFind(ObjectId, SlotContentBindingObjects);
-	
-	if ((BindingObjects.Num() == 0) && (SlotContentBindingObjects.Num() == 0))
-	{
-		OutDisplayName = LOCTEXT("NoBoundObjects", "No bound objects");
-	}
-	else if (BindingObjects.Num() + SlotContentBindingObjects.Num() > 1)
-	{
-		OutDisplayName = LOCTEXT("Multiple bound objects", "Multilple bound objects");
-	}
-	else if (BindingObjects.Num() == 1)
+	// TODO: This gets called every frame for every bound object and could
+	// be a potential performance issue for a really complicated animation.
+
+	TWeakObjectPtr<UObject> PreviewObject = IdToPreviewObjects.FindRef(ObjectId);
+
+	if (PreviewObject.IsValid())
 	{
-		OutDisplayName = FText::FromString(BindingObjects[0].Get()->GetName());
+		OutDisplayName = FText::FromString(PreviewObject.Get()->GetName());
+		return true;
 	}
-	else // SlotContentBindingObjects.Num() == 1
+
+	TWeakObjectPtr<UObject> SlotContentPreviewObject = IdToSlotContentPreviewObjects.FindRef(ObjectId);
+
+	if (SlotContentPreviewObject.IsValid())
 	{
-		UWidget* SlotContent = Cast<UWidget>(SlotContentBindingObjects[0].Get());
+		UWidget* SlotContent = Cast<UWidget>(SlotContentPreviewObject.Get());
 		FText PanelName = SlotContent->Slot != nullptr && SlotContent->Slot->Parent != nullptr
 			? FText::FromString(SlotContent->Slot->Parent->GetName())
 			: LOCTEXT("InvalidPanel", "Invalid Panel");
 		FText ContentName = FText::FromString(SlotContent->GetName());
 		OutDisplayName = FText::Format(LOCTEXT("SlotObject", "{0} ({1} Slot)"), ContentName, PanelName);
+
+		return true;
 	}
 
-	return true;
+	return false;
 }
 
 
-void UUMGBindingManager::UnbindPossessableObjects(const FMovieSceneObjectId& ObjectId)
+void UUMGBindingManager::UnbindPossessableObjects(const FGuid& ObjectId)
 {
+	// unbind widgets
 	TArray<TWeakObjectPtr<UObject>> PreviewObjects;
-	ObjectIdToPreviewObjects.MultiFind(ObjectId, PreviewObjects);
+	IdToPreviewObjects.MultiFind(ObjectId, PreviewObjects);
 
 	for (TWeakObjectPtr<UObject>& PreviewObject : PreviewObjects)
 	{
-		PreviewObjectToId.Remove(PreviewObject);
+		PreviewObjectToIds.Remove(PreviewObject);
 	}
 
-	ObjectIdToPreviewObjects.Remove(ObjectId);
+	IdToPreviewObjects.Remove(ObjectId);
 
+	// unbind slots
 	TArray<TWeakObjectPtr<UObject>> SlotContentPreviewObjects;
-	ObjectIdToSlotContentPreviewObjects.MultiFind(ObjectId, SlotContentPreviewObjects);
+	IdToSlotContentPreviewObjects.MultiFind(ObjectId, SlotContentPreviewObjects);
 
 	for (TWeakObjectPtr<UObject>& SlotContentPreviewObject : SlotContentPreviewObjects)
 	{
-		SlotContentPreviewObjectToId.Remove(SlotContentPreviewObject);
+		SlotContentPreviewObjectToIds.Remove(SlotContentPreviewObject);
 	}
 
-	ObjectIdToSlotContentPreviewObjects.Remove(ObjectId);
+	IdToSlotContentPreviewObjects.Remove(ObjectId);
 
+	// remove animation bindings
 	if ((WidgetAnimation != nullptr) && (WidgetBlueprintEditor != nullptr))
 	{
 		UWidgetBlueprint* WidgetBlueprint = WidgetBlueprintEditor->GetWidgetBlueprintObj();
@@ -266,10 +269,10 @@ void UUMGBindingManager::UnbindPossessableObjects(const FMovieSceneObjectId& Obj
 
 void UUMGBindingManager::HandleWidgetPreviewUpdated()
 {
-	PreviewObjectToId.Empty();
-	ObjectIdToPreviewObjects.Empty();
-	ObjectIdToSlotContentPreviewObjects.Empty();
-	SlotContentPreviewObjectToId.Empty();
+	PreviewObjectToIds.Empty();
+	IdToPreviewObjects.Empty();
+	IdToSlotContentPreviewObjects.Empty();
+	SlotContentPreviewObjectToIds.Empty();
 
 	InitPreviewObjects();
 
diff --git a/Engine/Source/Editor/UMGEditor/Private/Animation/UMGBindingManager.h b/Engine/Source/Editor/UMGEditor/Private/Animation/UMGBindingManager.h
--- a/Engine/Source/Editor/UMGEditor/Private/Animation/UMGBindingManager.h
+++ b/Engine/Source/Editor/UMGEditor/Private/Animation/UMGBindingManager.h
@@ -62,15 +62,15 @@ public:
 	// IMovieSceneBindingManager interface
 
 	virtual bool AllowsSpawnableObjects() const override;
-	virtual void BindPossessableObject(const FMovieSceneObjectId& ObjectId, UObject& PossessedObject) override;
-	virtual bool CanPossessObject( UObject& Object ) const override;
-	virtual void DestroyAllSpawnedObjects() override { }
-	virtual FMovieSceneObjectId FindIdForObject(const UMovieScene& MovieScene, UObject& Object) const override;
-	virtual void GetRuntimeObjects(const TSharedRef<FMovieSceneInstance>& MovieSceneInstance, const FMovieSceneObjectId& ObjectId, TArray<UObject*>& OutRuntimeObjects) const override;
-	virtual void RemoveMovieSceneInstance(const TSharedRef<FMovieSceneInstance>& MovieSceneInstance) override { }
-	virtual void SpawnOrDestroyObjectsForInstance(TSharedRef<FMovieSceneInstance> MovieSceneInstance, bool bDestroyAll) override { }
-	virtual bool TryGetObjectBindingDisplayName(const TSharedRef<FMovieSceneInstance>& MovieSceneInstance, const FMovieSceneObjectId& ObjectId, FText& OutDisplayName) const override;
-	virtual void UnbindPossessableObjects(const FMovieSceneObjectId& ObjectId) override;
+	virtual void BindPossessableObject(const FGuid& ObjectId, UObject& PossessedObject) override;
+	virtual bool CanPossessObject(UObject& Object) const override;
+	virtual void ClearBindings() override;
+	virtual void DestroyAllSpawnedObjects(UMovieScene& MovieScene) override { }
+	virtual FGuid FindObjectId(const UMovieScene& MovieScene, UObject& Object) const override;
+	virtual UObject* FindObject(const FGuid& ObjectId) const override;
+	virtual void SpawnOrDestroyObjects(UMovieScene* MovieScene, bool DestroyAll) override { }
+	virtual bool TryGetObjectDisplayName(const FGuid& ObjectId, FText& OutDisplayName) const override;
+	virtual void UnbindPossessableObjects(const FGuid& ObjectId) override;
 
 private:
 
@@ -79,11 +79,11 @@ private:
 
 private:
 
-	/** Mapping of preview objects to sequencer guids */
-	TMap<TWeakObjectPtr<UObject>, FMovieSceneObjectId> PreviewObjectToId;
-	TMultiMap<FMovieSceneObjectId, TWeakObjectPtr<UObject>> ObjectIdToPreviewObjects;
-	TMultiMap<FMovieSceneObjectId, TWeakObjectPtr<UObject>> ObjectIdToSlotContentPreviewObjects;
-	TMap<TWeakObjectPtr<UObject>, FMovieSceneObjectId> SlotContentPreviewObjectToId;
+	/** Mapping of preview objects to sequencer GUIDs */
+	TMap<TWeakObjectPtr<UObject>, FGuid> PreviewObjectToIds;
+	TMultiMap<FGuid, TWeakObjectPtr<UObject>> IdToPreviewObjects;
+	TMultiMap<FGuid, TWeakObjectPtr<UObject>> IdToSlotContentPreviewObjects;
+	TMap<TWeakObjectPtr<UObject>, FGuid> SlotContentPreviewObjectToIds;
 
 	/** The current widget animation. */
 	UPROPERTY(transient)
diff --git a/Engine/Source/Runtime/MovieScene/Public/MovieScene.h b/Engine/Source/Runtime/MovieScene/Public/MovieScene.h
--- a/Engine/Source/Runtime/MovieScene/Public/MovieScene.h
+++ b/Engine/Source/Runtime/MovieScene/Public/MovieScene.h
@@ -2,7 +2,6 @@
 
 #pragma once
 
-#include "MovieSceneObjectId.h"
 #include "MovieScene.generated.h"
 
 
@@ -43,14 +42,6 @@ public:
 		return Name;
 	}
 
-	/**
-	 * Get the unique object identifier.
-	 */
-	const FMovieSceneObjectId& GetObjectId() const
-	{
-		return ObjectId;
-	}
-
 	/** @return Returns the blueprint associated with this spawnable */
 	UClass* GetClass() { return GeneratedClass; }
 
@@ -82,9 +73,6 @@ private:
 	    used in the editor when capturing keyframe data from a live simulation */
 	// @todo sequencer data: Should be editor only
 	FWeakObjectPtr CounterpartGamePreviewObject;
-
-	// @todo sequencer: gmp: refactor sequencer object IDs.
-	FMovieSceneObjectId ObjectId;
 };
 
 /**
diff --git a/Engine/Source/Runtime/MovieScene/Public/MovieSceneBindingManager.h b/Engine/Source/Runtime/MovieScene/Public/MovieSceneBindingManager.h
--- a/Engine/Source/Runtime/MovieScene/Public/MovieSceneBindingManager.h
+++ b/Engine/Source/Runtime/MovieScene/Public/MovieSceneBindingManager.h
@@ -2,7 +2,6 @@
 
 #pragma once
 
-#include "MovieSceneObjectId.h"
 #include "MovieSceneBindingManager.generated.h"
 
 
@@ -44,44 +43,47 @@ class IMovieSceneBindingManager
 	 * @param PossessedObject The runtime object which was possessed.
 	 * @see UnbindPossessableObjects
 	 */
-	virtual void BindPossessableObject(const FMovieSceneObjectId& ObjectId, UObject& PossessedObject) PURE_VIRTUAL(IMovieSceneBindingManager::BindPossessableObject,);
+	virtual void BindPossessableObject(const FGuid& ObjectId, UObject& PossessedObject) PURE_VIRTUAL(IMovieSceneBindingManager::BindPossessableObject,);
 	
 	/**
 	 * @todo gmp: sequencer: add documentation.
 	 */
-	virtual bool CanPossessObject( UObject& Object ) const PURE_VIRTUAL(IMovieSceneBindingManager::CanPossessObject, return false;);
+	virtual bool CanPossessObject(UObject& Object) const PURE_VIRTUAL(IMovieSceneBindingManager::CanPossessObject, return false;);
 
-	/**
-	 * @todo gmp: sequencer: add documentation
-	 */
-	virtual void DestroyAllSpawnedObjects() PURE_VIRTUAL(IMovieSceneBindingManager::DestroyAllSpawnedObjects,);
+	/** Remove all bindings from this manager. */
+	virtual void ClearBindings() PURE_VIRTUAL(IMovieSceneBindingManager::ClearBindings,);
 
 	/**
-	 * Finds an existing guid for a bound object.
+	 * Destroy all objects that were spawned by this binding manager.
 	 *
-	 * @param MovieScene The movie scene that may contain the binding.
-	 * @param The object to get the guid for.
-	 * @return The object identifier or an invalid guid if not found.
+	 * @param MovieScene The movie scene that owns this manager.
 	 */
-	virtual FMovieSceneObjectId FindIdForObject(const UMovieScene& MovieScene, UObject& Object) const PURE_VIRTUAL(IMovieSceneBindingManager::FindGuidForObject, return FGuid(););
+	virtual void DestroyAllSpawnedObjects(UMovieScene& MovieScene) PURE_VIRTUAL(IMovieSceneBindingManager::DestroyAllSpawnedObjects,);
 
 	/**
-	 * @todo gmp: sequencer: add documentation
+	 * Find the identifier for a possessed or spawned object.
+	 *
+	 * @param The scene that owns this binding manager.
+	 * @param The object to get the identifier for.
+	 * @return The object identifier, or an invalid GUID if not found.
 	 */
-	virtual void GetRuntimeObjects(const TSharedRef<FMovieSceneInstance>& MovieSceneInstance, const FMovieSceneObjectId& ObjectId, TArray<UObject*>& OutRuntimeObjects) const PURE_VIRTUAL(IMovieSceneBindingManager::GetRuntimeObjects,);
+	virtual FGuid FindObjectId(const UMovieScene& MovieScene, UObject& Object) const PURE_VIRTUAL(IMovieSceneBindingManager::FindGuidForObject, return FGuid(););
 
 	/**
-	 * @todo gmp: sequencer: add documentation
+	 * Gets the possessed or spawned object for the specified identifier.
+	 *
+	 * @param ObjectId The unique identifier of the object.
+	 * @return The object, or nullptr if not found.
 	 */
-	virtual void RemoveMovieSceneInstance(const TSharedRef<FMovieSceneInstance>& MovieSceneInstance) PURE_VIRTUAL(IMovieSceneBindingManager::RemoveMovieSceneInstance,);
+	virtual UObject* FindObject(const FGuid& ObjectId) const PURE_VIRTUAL(IMovieSceneBindingManager::FindObject, return nullptr;);
 
 	/**
 	 * Called to spawn or destroy objects for a movie instance.
 	 *
-	 * @param MovieSceneInstance The instance to spawn or destroy objects for.
-	 * @param bDestroyAll If true, destroy all spawned objects for the instance, if false only destroy unused objects.
+	 * @param MovieScene The movie scene to spawn or destroy objects for.
+	 * @param DestroyAll If true, destroy all spawned objects for the instance, if false only destroy unused objects.
 	 */
-	virtual void SpawnOrDestroyObjectsForInstance( TSharedRef<FMovieSceneInstance> MovieSceneInstance, bool bDestroyAll ) PURE_VIRTUAL(IMovieSceneBindingManager::SpawnOrDestroyObjectsForInstance,);
+	virtual void SpawnOrDestroyObjects(UMovieScene* MovieScene, bool DestroyAll) PURE_VIRTUAL(IMovieSceneBindingManager::SpawnOrDestroyObjects,);
 
 	/**
 	 * 
@@ -91,7 +93,7 @@ class IMovieSceneBindingManager
 	 * @param DisplayName The display name for the object binding.
 	 * @returns true if DisplayName has been set to a valid display name, otherwise false.
 	 */
-	virtual bool TryGetObjectBindingDisplayName(const TSharedRef<FMovieSceneInstance>& MovieSceneInstance, const FMovieSceneObjectId& ObjectId, FText& OutDisplayName) const PURE_VIRTUAL(IMovieSceneBindingManager::TryGetObjectBindingDisplayName, return false;);
+	virtual bool TryGetObjectDisplayName(const FGuid& ObjectId, FText& OutDisplayName) const PURE_VIRTUAL(IMovieSceneBindingManager::TryGetObjectDisplayName, return false;);
 
 	/**
 	 * Unbinds all possessable objects from the provided GUID.
@@ -99,5 +101,5 @@ class IMovieSceneBindingManager
 	 * @param ObjectId The guid bound to possessable objects that should be removed.
 	 * @see BindPossessableObject
 	 */
-	virtual void UnbindPossessableObjects(const FMovieSceneObjectId& ObjectId) PURE_VIRTUAL(IMovieSceneBindingManager::UnbindPossessableObjects,);
+	virtual void UnbindPossessableObjects(const FGuid& ObjectId) PURE_VIRTUAL(IMovieSceneBindingManager::UnbindPossessableObjects,);
 };
diff --git a/Engine/Source/Runtime/MovieScene/Public/MovieSceneObjectId.h b/Engine/Source/Runtime/MovieScene/Public/MovieSceneObjectId.h
deleted file mode 100644
diff --git a/Engine/Source/Runtime/MovieSceneTracks/Private/Sections/MovieScene3DPathSection.cpp b/Engine/Source/Runtime/MovieSceneTracks/Private/Sections/MovieScene3DPathSection.cpp
--- a/Engine/Source/Runtime/MovieSceneTracks/Private/Sections/MovieScene3DPathSection.cpp
+++ b/Engine/Source/Runtime/MovieSceneTracks/Private/Sections/MovieScene3DPathSection.cpp
@@ -15,14 +15,14 @@ UMovieScene3DPathSection::UMovieScene3DPathSection( const FObjectInitializer& Ob
 {
 }
 
-void UMovieScene3DPathSection::SetPathId(const FMovieSceneObjectId& InPathId)
+void UMovieScene3DPathSection::SetPathId(const FGuid& InPathId)
 {
 	Modify();
 
 	PathId = InPathId;
 }
 
-FMovieSceneObjectId UMovieScene3DPathSection::GetPathId() const
+FGuid UMovieScene3DPathSection::GetPathId() const
 {
 	return PathId;
 }
@@ -144,7 +144,7 @@ void UMovieScene3DPathSection::GetKeyHandles(TSet<FKeyHandle>& KeyHandles) const
 	}
 }
 
-void UMovieScene3DPathSection::AddPath( float Time, float SequenceEndTime, const FMovieSceneObjectId& InPathId )
+void UMovieScene3DPathSection::AddPath( float Time, float SequenceEndTime, const FGuid& InPathId )
 {
 	Modify();
 	PathId = InPathId;
diff --git a/Engine/Source/Runtime/MovieSceneTracks/Private/TrackInstances/MovieScene3DPathTrackInstance.cpp b/Engine/Source/Runtime/MovieSceneTracks/Private/TrackInstances/MovieScene3DPathTrackInstance.cpp
--- a/Engine/Source/Runtime/MovieSceneTracks/Private/TrackInstances/MovieScene3DPathTrackInstance.cpp
+++ b/Engine/Source/Runtime/MovieSceneTracks/Private/TrackInstances/MovieScene3DPathTrackInstance.cpp
@@ -63,10 +63,11 @@ void FMovieScene3DPathTrackInstance::Update( float Position, float LastPosition,
 			(FirstPathSection == NULL || FirstPathSection->GetRowIndex() > PathSection->GetRowIndex()))
 		{
 			TArray<UObject*> PathObjects;
-			FMovieSceneObjectId PathId = PathSection->GetPathId();
+			FGuid PathId = PathSection->GetPathId();
+
 			if (PathId.IsValid())
 			{
-				Player.GetRuntimeObjects( Player.GetRootMovieSceneInstance(), PathId.ObjectGuid, PathObjects);
+				Player.GetRuntimeObjects( Player.GetRootMovieSceneInstance(), PathId, PathObjects);
 
 				for (int32 PathObjectIndex = 0; PathObjectIndex < PathObjects.Num(); ++PathObjectIndex)
 				{
diff --git a/Engine/Source/Runtime/MovieSceneTracks/Private/Tracks/MovieScene3DPathTrack.cpp b/Engine/Source/Runtime/MovieSceneTracks/Private/Tracks/MovieScene3DPathTrack.cpp
--- a/Engine/Source/Runtime/MovieSceneTracks/Private/Tracks/MovieScene3DPathTrack.cpp
+++ b/Engine/Source/Runtime/MovieSceneTracks/Private/Tracks/MovieScene3DPathTrack.cpp
@@ -54,7 +54,7 @@ TRange<float> UMovieScene3DPathTrack::GetSectionBoundaries() const
 	return TRange<float>::Hull(Bounds);
 }
 
-void UMovieScene3DPathTrack::AddPath(float KeyTime, float PathEndTime, const FMovieSceneObjectId& PathId)
+void UMovieScene3DPathTrack::AddPath(float KeyTime, float PathEndTime, const FGuid& PathId)
 {
 	// add the section
 	UMovieScene3DPathSection* NewSection = NewObject<UMovieScene3DPathSection>(this);
diff --git a/Engine/Source/Runtime/MovieSceneTracks/Public/Sections/MovieScene3DPathSection.h b/Engine/Source/Runtime/MovieSceneTracks/Public/Sections/MovieScene3DPathSection.h
--- a/Engine/Source/Runtime/MovieSceneTracks/Public/Sections/MovieScene3DPathSection.h
+++ b/Engine/Source/Runtime/MovieSceneTracks/Public/Sections/MovieScene3DPathSection.h
@@ -2,7 +2,6 @@
 
 #pragma once
 
-#include "MovieSceneObjectId.h"
 #include "MovieSceneSection.h"
 #include "MovieScene3DPathSection.generated.h"
 
@@ -36,10 +35,10 @@ public:
 	virtual void GetKeyHandles(TSet<FKeyHandle>& KeyHandles) const override;
 
 	/** Sets the path id for this path section */
-	virtual void SetPathId(const FMovieSceneObjectId& InId);
+	virtual void SetPathId(const FGuid& InId);
 
 	/** Gets the path id for this path section */
-	virtual FMovieSceneObjectId GetPathId() const;
+	virtual FGuid GetPathId() const;
 
 	/**
 	 * Evaluates the path track
@@ -55,7 +54,7 @@ public:
 	 * @param SequenceEndTime   The time at the end of the sequence, by default the path is set to end at this time
 	 * @param InPathId The id to the path
 	 */
-	void AddPath( float Time, float SequenceEndTime, const FMovieSceneObjectId& InPathId);
+	void AddPath( float Time, float SequenceEndTime, const FGuid& InPathId);
 
 	/** 
 	 * Returns the timing curve
@@ -67,7 +66,7 @@ public:
 private:
 	/** The path possessable that this path uses */
 	UPROPERTY()
-	FMovieSceneObjectId PathId;
+	FGuid PathId;
 
 	/** Timing Curve */
 	UPROPERTY()
diff --git a/Engine/Source/Runtime/MovieSceneTracks/Public/Tracks/MovieScene3DPathTrack.h b/Engine/Source/Runtime/MovieSceneTracks/Public/Tracks/MovieScene3DPathTrack.h
--- a/Engine/Source/Runtime/MovieSceneTracks/Public/Tracks/MovieScene3DPathTrack.h
+++ b/Engine/Source/Runtime/MovieSceneTracks/Public/Tracks/MovieScene3DPathTrack.h
@@ -31,7 +31,7 @@ public:
 	 * @param PathEndTime       Set the path to end at this time
 	 * @param PathId			The id to the spline
 	 */
-	virtual void AddPath( float Time, float PathEndTime, const FMovieSceneObjectId& PathId );
+	virtual void AddPath( float Time, float PathEndTime, const FGuid& PathId );
 
 private:
 	/** List of all path sections */
diff --git a/Engine/Source/Runtime/UMG/Public/Animation/WidgetAnimation.h b/Engine/Source/Runtime/UMG/Public/Animation/WidgetAnimation.h
--- a/Engine/Source/Runtime/UMG/Public/Animation/WidgetAnimation.h
+++ b/Engine/Source/Runtime/UMG/Public/Animation/WidgetAnimation.h
@@ -3,7 +3,6 @@
 #pragma once
 
 #include "IMovieScenePlayer.h"
-#include "MovieSceneObjectId.h"
 #include "WidgetAnimation.generated.h"
 
 class UMovieScene;
@@ -26,7 +25,7 @@ struct FWidgetAnimationBinding
 
 	// @todo sequencer: gmp: UMG integration
 	//UPROPERTY()
-	FMovieSceneObjectId AnimationId;
+	FGuid AnimationId;
 
 public:
 	/**
