MovieScene: Renamed 'binding managers' to 'object managers'

diff --git a/Engine/Config/BaseEngine.ini b/Engine/Config/BaseEngine.ini
--- a/Engine/Config/BaseEngine.ini
+++ b/Engine/Config/BaseEngine.ini
@@ -491,7 +491,7 @@ bUseStreamingPause=false
 +ActiveStructRedirects=(OldStructName="MovieSceneSpawnable",NewStructName="/Script/MovieScene.MovieSceneSpawnable")
 +ActiveStructRedirects=(OldStructName="MovieScenePossessable",NewStructName="/Script/MovieScene.MovieScenePossessable")
 +ActiveStructRedirects=(OldStructName="MovieSceneEditorData",NewStructName="/Script/MovieScene.MovieSceneEditorData")
-+ActiveStructRedirects=(OldStructName="MovieSceneObjectBinding",NewStructName="/Script/MovieScene.MovieSceneObjectBinding")
++ActiveStructRedirects=(OldStructName="MovieSceneObjectBinding",NewStructName="/Script/MovieScene.MovieSceneBinding")
 +ActiveStructRedirects=(OldStructName="MovieSceneBoundObject",NewStructName="/Script/MovieScene.MovieSceneBoundObject")
 
 ; MovieSceneCoreTypes moved to MovieSceneTracks
diff --git a/Engine/Source/Editor/Sequencer/Private/SequencerBinding.cpp b/Engine/Source/Editor/Sequencer/Private/SequencerBinding.cpp
deleted file mode 100644
diff --git a/Engine/Source/Editor/Sequencer/Private/SequencerBinding.h b/Engine/Source/Editor/Sequencer/Private/SequencerBinding.h
deleted file mode 100644
diff --git a/Engine/Source/Editor/Sequencer/Private/SequencerBindingManager.cpp b/Engine/Source/Editor/Sequencer/Private/SequencerBindingManager.cpp
deleted file mode 100644
diff --git a/Engine/Source/Editor/Sequencer/Private/SequencerBindingManager.h b/Engine/Source/Editor/Sequencer/Private/SequencerBindingManager.h
deleted file mode 100644
diff --git a/Engine/Source/Editor/Sequencer/Private/SequencerMovieSceneObjectManager.cpp b/Engine/Source/Editor/Sequencer/Private/SequencerMovieSceneObjectManager.cpp
new file mode 100644
--- /dev/null
+++ b/Engine/Source/Editor/Sequencer/Private/SequencerMovieSceneObjectManager.cpp
@@ -0,0 +1,437 @@
+// Copyright 1998-2015 Epic Games, Inc. All Rights Reserved.
+
+#include "SequencerPrivatePCH.h"
+#include "MovieScene.h"
+#include "MovieSceneInstance.h"
+#include "SequencerMovieSceneObjectManager.h"
+#include "SequencerPossessedObject.h"
+
+// @todo sequencer: gmp: refactor these dependencies into another module
+#include "ModuleManager.h"
+#include "LevelEditor.h"
+
+
+/* USequencerMovieSceneObjectManager structors
+ *****************************************************************************/
+
+USequencerMovieSceneObjectManager::USequencerMovieSceneObjectManager(const FObjectInitializer& ObjectInitializer)
+	: Super(ObjectInitializer)
+	//, ObjectChangeListener(InObjectChangeListener)
+{
+	// @todo sequencer: gmp: initialize ObjectChangeListener
+	//ObjectChangeListener->GetOnPropagateObjectChanges().AddUObject(this, &USequencerMovieSceneObjectManager::HandlePropagateObjectChanges);
+}
+
+
+USequencerMovieSceneObjectManager::~USequencerMovieSceneObjectManager()
+{
+	if (ObjectChangeListener.IsValid())
+	{
+		ObjectChangeListener.Pin()->GetOnPropagateObjectChanges().RemoveAll(this);
+	}
+}
+
+
+/* IMovieSceneObjectManager interface
+ *****************************************************************************/
+
+bool USequencerMovieSceneObjectManager::AllowsSpawnableObjects() const
+{
+	return true;
+}
+
+
+void USequencerMovieSceneObjectManager::BindPossessableObject(const FGuid& ObjectId, UObject& PossessedObject)
+{
+	FString ComponentName;
+	UActorComponent* Component = Cast<UActorComponent>(&PossessedObject);
+
+	if (Component != nullptr)
+	{
+		// @todo sequencer: gmp: need support for UStruct keys in TMap
+		PossessedObjects.Add(ObjectId.ToString(), FSequencerPossessedObject(PossessedObject.GetOuter(), ComponentName));
+		//PossessedObjects.Add(ObjectId, FSequencerPossessedObject(PossessedObject.GetOuter(), ComponentName));
+	}
+	else
+	{
+		// @todo sequencer: gmp: need support for UStruct keys in TMap
+		PossessedObjects.Add(ObjectId.ToString(), FSequencerPossessedObject(&PossessedObject));
+		//PossessedObjects.Add(ObjectId, FSequencerPossessedObject(Actor));
+	}
+}
+
+
+bool USequencerMovieSceneObjectManager::CanPossessObject(UObject& Object) const
+{
+	return Object.IsA<AActor>();
+}
+
+
+void USequencerMovieSceneObjectManager::DestroyAllSpawnedObjects(UMovieScene& MovieScene)
+{
+	const bool DestroyAll = true;
+	SpawnOrDestroyObjects(&MovieScene, DestroyAll);
+	SpawnedActors.Empty();
+}
+
+
+FGuid USequencerMovieSceneObjectManager::FindObjectId(const UMovieScene& MovieScene, UObject& Object) const
+{
+	// search bindings (possessables)
+	for (auto PossessedObjectPair : PossessedObjects)
+	{
+		const FSequencerPossessedObject& PossessedObject = PossessedObjectPair.Value;
+
+		if (PossessedObject.GetObject() == &Object)
+		{
+			FGuid Result; FGuid::Parse(PossessedObjectPair.Key, Result); return Result;
+			//return PossessedObjectsPair.Key;
+		}
+	}
+
+	// search proxies (spawnables)
+	const bool bIsGamePreviewObject = !!(Object.GetOutermost()->PackageFlags & PKG_PlayInEditor);
+
+	if (bIsGamePreviewObject)
+	{
+		// @todo sequencer: gmp: i don't want this class to know the MovieScene that owns it
+		//
+		// someone is asking for a handle to an object from a game preview session,
+		// probably because they want to capture keys during live simulation.
+		//
+		// Check to see if we already have a puppet that was generated from a
+		// recording of this game preview object.
+		//
+		// @todo sequencer livecapture: We could support recalling counterpart by full name instead of weak pointer, to allow "overdubbing" of previously recorded actors, when the new actors in the current play session have the same path name
+		// @todo sequencer livecapture: Ideally we could capture from editor-world actors that are "puppeteered" as well (real time)
+		const FMovieSceneSpawnable* FoundSpawnable = MovieScene.FindSpawnableForCounterpart(&Object);
+		
+		if (FoundSpawnable != nullptr)
+		{
+			return FoundSpawnable->GetGuid();
+		}
+	}
+	else
+	{
+		for (auto SpawnedActorPair : SpawnedActors)
+		{
+			if (SpawnedActorPair.Value.Get() == &Object)
+			{
+				return SpawnedActorPair.Key;
+			}
+		}
+	}
+
+	return FGuid();
+}
+
+
+UObject* USequencerMovieSceneObjectManager::FindObject(const FGuid& ObjectId) const
+{
+	// search bindings (possessables)
+	// @todo sequencer: gmp: need support for UStruct keys in TMap
+	const FSequencerPossessedObject* PossessedObject = PossessedObjects.Find(ObjectId.ToString());
+	//const FSequencerPossessedObject* PossessedObject = PossessedObjects.Find(ObjectId);
+
+	if (PossessedObject != nullptr)
+	{
+		return PossessedObject->GetObject();
+	}
+
+	// search proxies (spawnables)
+	TWeakObjectPtr<AActor> Actor = SpawnedActors.FindRef(ObjectId);
+
+	if (Actor.IsValid())
+	{
+		return Actor.Get();
+	}
+
+	// not found
+	return nullptr;
+}
+
+
+void USequencerMovieSceneObjectManager::SpawnOrDestroyObjects(UMovieScene* MovieScene, bool DestroyAll)
+{
+	bool bAnyLevelActorsChanged = false;
+
+	// remove any proxy actors that we no longer need
+	if (MovieScene != nullptr)
+	{
+		for (auto SpawnedActorPair : SpawnedActors)
+		{
+			bool ShouldDestroyActor = true;
+
+			if (!DestroyAll)
+			{
+				// figure out if we still need this proxy actor
+				for (auto SpawnableIndex = 0; SpawnableIndex < MovieScene->GetSpawnableCount(); ++SpawnableIndex)
+				{
+					auto& Spawnable = MovieScene->GetSpawnable(SpawnableIndex);
+
+					if (Spawnable.GetGuid() == SpawnedActorPair.Key)
+					{
+						ShouldDestroyActor = false;
+						break;
+					}
+				}
+			}
+
+			if (!ShouldDestroyActor)
+			{
+				// Actor is no longer valid (probably the world was destroyed)
+				continue;
+			}
+
+			AActor* Actor = SpawnedActorPair.Value.Get();
+				
+			if (Actor == nullptr)
+			{
+				continue;
+			}
+
+			// destroy Actor
+			UWorld* World = Actor->GetWorld();
+
+			if (ensure(World != nullptr))
+			{
+				if (!bAnyLevelActorsChanged)
+				{
+					DeselectAllActors();
+				}
+
+				// Actor should never be selected in the Editor at this point (we took care of that up above)
+				ensure(!Actor->IsSelected());
+
+				const bool bNetForce = false;
+				const bool bShouldModifyLevel = false;
+				const bool bWasDestroyed = World->DestroyActor(Actor, bNetForce, bShouldModifyLevel);
+
+				if (bWasDestroyed)
+				{
+					bAnyLevelActorsChanged = true;
+					SpawnedActors.Remove(SpawnedActorPair.Key);
+				}
+				else
+				{
+					// @todo sequencer: At least one proxy actor couldn't be cleaned up!
+				}
+			}
+		}
+	}
+
+	if (!DestroyAll && (MovieScene != nullptr))
+	{
+		UWorld* ActorWorld = GetWorld();
+
+		for (auto SpawnableIndex = 0; SpawnableIndex < MovieScene->GetSpawnableCount(); ++SpawnableIndex)
+		{
+			FMovieSceneSpawnable& Spawnable = MovieScene->GetSpawnable( SpawnableIndex );
+
+			// Must have a valid world for us to be able to do this
+			if (ActorWorld == nullptr )
+			{
+				continue;
+			}
+
+			// Do we already have a proxy for this spawnable?
+			if (SpawnedActors.Contains(Spawnable.GetGuid()))
+			{
+				continue;
+			}
+
+			UClass* GeneratedClass = Spawnable.GetClass();
+
+			if ((GeneratedClass == nullptr) || !GeneratedClass->IsChildOf(AActor::StaticClass()))
+			{
+				continue;
+			}
+
+			AActor* ActorCDO = CastChecked<AActor>(GeneratedClass->ClassDefaultObject);
+
+			const FVector SpawnLocation = ActorCDO->GetRootComponent()->RelativeLocation;
+			const FRotator SpawnRotation = ActorCDO->GetRootComponent()->RelativeRotation;
+
+			// @todo sequencer: We should probably spawn these in a specific sub-level!
+			// World->CurrentLevel = ???;
+
+			const FName ActorName = NAME_None;
+
+			// Override the object flags so that RF_Transactional is not set.  Puppet actors are never transactional
+			// @todo sequencer: These actors need to avoid any transaction history.  However, RF_Transactional can currently be set on objects on the fly!
+			const EObjectFlags ObjectFlags = RF_Transient;		// NOTE: We are omitting RF_Transactional intentionally
+
+			// @todo sequencer livecapture: Consider using SetPlayInEditorWorld() and RestoreEditorWorld() here instead
+						
+			// @todo sequencer actors: We need to make sure puppet objects aren't copied into PIE/SIE sessions!  They should be omitted from that duplication!
+
+			// Spawn the puppet actor
+			FActorSpawnParameters SpawnInfo;
+			{
+				SpawnInfo.Name = ActorName;
+				SpawnInfo.ObjectFlags = ObjectFlags;
+			}
+
+			AActor* NewActor = ActorWorld->SpawnActor(GeneratedClass, &SpawnLocation, &SpawnRotation, SpawnInfo);
+				
+			if (NewActor != nullptr)
+			{
+				// @todo sequencer: We're naming the actor based off of the spawnable's name.  Is that really what we want?
+				FActorLabelUtilities::SetActorLabelUnique(NewActor, Spawnable.GetName());
+				SpawnedActors.Add(Spawnable.GetGuid(), NewActor);
+			}
+			else
+			{
+				// Actor failed to spawn
+				// @todo sequencer: What should we do when this happens to one or more actors?
+			}
+		}
+	}
+}
+
+
+bool USequencerMovieSceneObjectManager::TryGetObjectDisplayName(const FGuid& ObjectId, FText& OutDisplayName) const
+{
+	UObject* Object = FindObject(ObjectId);
+
+	if (Object == nullptr)
+	{
+		return false;
+	}
+
+	AActor* Actor = Cast<AActor>(Object);
+
+	if (Actor == nullptr)
+	{
+		return false;
+	}
+
+	OutDisplayName = FText::FromString(Actor->GetActorLabel());
+
+	return true;
+}
+
+
+void USequencerMovieSceneObjectManager::UnbindPossessableObjects(const FGuid& ObjectId)
+{
+	// @todo sequencer: gmp: need support for UStruct keys in TMap
+	PossessedObjects.Remove(ObjectId.ToString());
+	//PossessedObjects.Remove(ObjectId);
+}
+
+
+/* USequencerMovieSceneObjectManager implementation
+ *****************************************************************************/
+
+void USequencerMovieSceneObjectManager::DeselectAllActors()
+{
+	USelection* Selection = GEditor->GetSelectedActors();
+	{
+		Selection->BeginBatchSelectOperation();
+		Selection->MarkBatchDirty();
+
+		for (auto SpawnedActorPair : SpawnedActors)
+		{
+			AActor* SpawnedActor = SpawnedActorPair.Value.Get();
+
+			if (SpawnedActor != nullptr)
+			{
+				GEditor->GetSelectedActors()->Deselect(SpawnedActor);
+			}
+		}
+
+		Selection->EndBatchSelectOperation();
+	}
+}
+
+
+UWorld* USequencerMovieSceneObjectManager::GetWorld() const
+{
+	if (FModuleManager::Get().IsModuleLoaded("LevelEditor"))
+	{
+		FLevelEditorModule& LevelEditorModule = FModuleManager::Get().GetModuleChecked<FLevelEditorModule>("LevelEditor");
+
+		return LevelEditorModule.GetFirstLevelEditor()->GetWorld();
+	}
+
+	return nullptr;
+}
+
+
+/* USequencerMovieSceneObjectManager internals
+ *****************************************************************************/
+
+void USequencerMovieSceneObjectManager::HandlePropagateObjectChanges(UObject* Object)
+{
+	if (Object == nullptr)
+	{
+		return;
+	}
+
+	// @todo sequencer major: Many other types of changes to the puppet actor may occur which should be propagated (or disallowed?):
+	//		- Deleting the puppet actor (Del key) should probably delete the reference actor and remove it from the MovieScene
+	//		- Attachment changes
+	//		- Replace/Convert operations?
+
+	// We only care about actor objects for this type of spawner.  Note that sometimes we can get component objects passed in, but we'll
+	// rely on the caller also calling OnObjectPropertyChange for the outer Actor if a component property changes.
+	if (!Object->IsA<AActor>())
+	{
+		return;
+	}
+
+	for (auto SpawnedActorPair : SpawnedActors)
+	{
+		AActor* SpawnedActor = SpawnedActorPair.Value.Get();
+
+		if (SpawnedActor == Object)
+		{
+			// @todo sequencer: propagation: Don't propagate changes that are being auto-keyed or key-adjusted!
+			PropagateActorChanges(SpawnedActorPair.Key, SpawnedActor);
+		}
+	}
+}
+
+
+void USequencerMovieSceneObjectManager::PropagateActorChanges(const FGuid& ObjectId, AActor* Actor)
+{
+	if (Actor == nullptr)
+	{
+		return;
+	}
+
+	AActor* TargetActor = nullptr;
+	FMovieSceneSpawnable* FoundSpawnable = nullptr;
+
+	// find the spawnable for this puppet actor
+	// @todo sequencer: gmp: add horrible sequencer dependencies to binding manager
+/*	TArray<UMovieScene*> MovieScenesBeingEdited = Sequencer.Pin()->GetMovieScenesBeingEdited();
+
+	for (UMovieScene* MovieScene : MovieScenesBeingEdited)
+	{
+		FoundSpawnable = MovieScene->FindSpawnable(ActorInfo->ObjectId);
+		
+		if (FoundSpawnable != nullptr)
+		{
+			break;
+		}
+	}
+*/	
+	if (ensure((FoundSpawnable != nullptr) && (Actor != nullptr)))
+	{
+		UClass* ActorClass = FoundSpawnable->GetClass();
+
+		// the proxy actor's class should always be the BP that it was spawned from
+		UClass* SpawnedActorClass = Actor->GetClass();
+		check(ActorClass == SpawnedActorClass);
+
+		// we'll be copying properties into the class default object of the Blueprint's generated class
+		TargetActor = CastChecked<AActor>(ActorClass->GetDefaultObject());
+	}
+
+	if ((Actor != nullptr) && (TargetActor != nullptr))
+	{
+		// @todo sequencer: gmp: add horrible sequencer dependencies to binding manager
+//		Sequencer.Pin()->CopyActorProperties(Actor, TargetActor);
+	}
+}
diff --git a/Engine/Source/Editor/Sequencer/Private/SequencerMovieSceneObjectManager.h b/Engine/Source/Editor/Sequencer/Private/SequencerMovieSceneObjectManager.h
new file mode 100644
--- /dev/null
+++ b/Engine/Source/Editor/Sequencer/Private/SequencerMovieSceneObjectManager.h
@@ -0,0 +1,75 @@
+// Copyright 1998-2015 Epic Games, Inc. All Rights Reserved.
+
+#pragma once
+
+#include "MovieSceneObjectManager.h"
+#include "SequencerPossessedObject.h"
+#include "SequencerMovieSceneObjectManager.generated.h"
+
+
+/**
+ * Movie scene object manager used by Sequencer.
+ */
+UCLASS(BlueprintType, MinimalAPI)
+class USequencerMovieSceneObjectManager
+	: public UObject
+	, public IMovieSceneObjectManager
+{
+	GENERATED_UCLASS_BODY()
+
+	/** Virtual destructor. */
+	virtual ~USequencerMovieSceneObjectManager();
+
+public:
+
+	// IMovieSceneObjectManager interface
+
+	virtual bool AllowsSpawnableObjects() const override;
+	virtual void BindPossessableObject(const FGuid& ObjectId, UObject& PossessedObject) override;
+	virtual bool CanPossessObject(UObject& Object) const override;
+	virtual void DestroyAllSpawnedObjects(UMovieScene& MovieScene) override;
+	virtual FGuid FindObjectId(const UMovieScene& MovieScene, UObject& Object) const override;
+	virtual UObject* FindObject(const FGuid& ObjectId) const override;
+	virtual void SpawnOrDestroyObjects(UMovieScene* MovieScene, bool DestroyAll) override;
+	virtual bool TryGetObjectDisplayName(const FGuid& ObjectId, FText& OutDisplayName) const override;
+	virtual void UnbindPossessableObjects(const FGuid& ObjectId) override;
+
+protected:
+
+	/**
+	 * Deselect all proxy actors in the Editor.
+	 *
+	 * @todo sequencer: remove Editor dependencies from this class
+	 */
+	SEQUENCER_API void DeselectAllActors();
+
+	/**
+	 * Get the world that possessed actors belong to or where spawned actors should go.
+	 *
+	 * @return The world object.
+	 * @todo sequencer: remove Editor dependencies from this class
+	 */
+	SEQUENCER_API UWorld* GetWorld() const;
+
+private:
+
+	/** Handles changes of non-auto-keyed properties on any object. */
+	void HandlePropagateObjectChanges(UObject* Object);
+
+	/** Refresh either the reference actor or the data descriptor BP after its puppet proxy actor has been changed. */
+	void PropagateActorChanges(const FGuid& ObjectId, AActor* Actor);
+
+private:
+
+	/** Collection of possessed objects. */
+	// @todo sequencer: gmp: need support for UStruct keys in TMap
+	UPROPERTY()
+	TMap<FString, FSequencerPossessedObject> PossessedObjects;
+	//TMap<FGuid, FSequencerBinding> Bindings;
+
+	/** The object change listener. */
+	TWeakPtr<ISequencerObjectChangeListener> ObjectChangeListener;
+
+	/** Collection of spawned proxy actors. */
+	TMap<FGuid, TWeakObjectPtr<AActor>> SpawnedActors;
+};
diff --git a/Engine/Source/Editor/Sequencer/Private/SequencerNodeTree.cpp b/Engine/Source/Editor/Sequencer/Private/SequencerNodeTree.cpp
--- a/Engine/Source/Editor/Sequencer/Private/SequencerNodeTree.cpp
+++ b/Engine/Source/Editor/Sequencer/Private/SequencerNodeTree.cpp
@@ -45,10 +45,10 @@ void FSequencerNodeTree::Update()
 	}
 
 
-	const TArray<FMovieSceneObjectBinding>& ObjectBindings = MovieScene->GetObjectBindings();
+	const TArray<FMovieSceneBinding>& Bindings = MovieScene->GetBindings();
 
 	// Make nodes for all object bindings
-	for( const FMovieSceneObjectBinding& Binding : ObjectBindings )
+	for( const FMovieSceneBinding& Binding : Bindings )
 	{
 		TSharedRef<FObjectBindingNode> ObjectBindingNode = AddObjectBinding( Binding.GetName(), Binding.GetObjectGuid(), NewRootNodes );
 
diff --git a/Engine/Source/Editor/Sequencer/Private/SequencerPossessedObject.cpp b/Engine/Source/Editor/Sequencer/Private/SequencerPossessedObject.cpp
new file mode 100644
--- /dev/null
+++ b/Engine/Source/Editor/Sequencer/Private/SequencerPossessedObject.cpp
@@ -0,0 +1,53 @@
+// Copyright 1998-2015 Epic Games, Inc. All Rights Reserved.
+
+#include "SequencerPrivatePCH.h"
+#include "SequencerPossessedObject.h"
+
+
+/* FSequencerPossessedObject interface
+ *****************************************************************************/
+
+UObject* FSequencerPossessedObject::GetObject() const
+{
+	// get component-less object
+	if (ComponentName.IsEmpty())
+	{
+		return ObjectOrOwner.Get();
+	}
+
+	// get cached component
+	if ((CachedComponent != nullptr) && (CachedComponent->GetName() == ComponentName))
+	{
+		return CachedComponent;
+	}
+
+	CachedComponent = nullptr;
+
+	// find & cache component
+	UObject* Object = ObjectOrOwner.Get();
+
+	if (Object == nullptr)
+	{
+		return nullptr;
+	}
+
+	AActor* Owner = Cast<AActor>(Object);
+
+	if (Owner == nullptr)
+	{
+		return Object;
+	}
+
+	for (UActorComponent* ActorComponent : Owner->GetComponents())
+	{
+		if (ActorComponent->GetName() == ComponentName)
+		{
+			CachedComponent = ActorComponent;
+
+			return ActorComponent;
+		}
+	}
+
+	// component not found
+	return nullptr;
+}
diff --git a/Engine/Source/Editor/Sequencer/Private/SequencerPossessedObject.h b/Engine/Source/Editor/Sequencer/Private/SequencerPossessedObject.h
new file mode 100644
--- /dev/null
+++ b/Engine/Source/Editor/Sequencer/Private/SequencerPossessedObject.h
@@ -0,0 +1,78 @@
+// Copyright 1998-2015 Epic Games, Inc. All Rights Reserved.
+
+#pragma once
+
+#include "SequencerPossessedObject.generated.h"
+
+
+/**
+ * Structure for movie scene objects in Sequencer.
+ */
+USTRUCT()
+struct FSequencerPossessedObject
+{
+	GENERATED_USTRUCT_BODY(FSequencerPossessedObject)
+
+public:
+
+	/** Creates and initializes a new instance. */
+	FSequencerPossessedObject()
+		: ObjectOrOwner(nullptr)
+		, CachedComponent(nullptr)
+	{ }
+
+	/**
+	 * Creates and initializes a new instance from an object.
+	 *
+	 * @param InObject The object to be bound.
+	 */
+	FSequencerPossessedObject(UObject* InObject)
+		: ObjectOrOwner(InObject)
+		, CachedComponent(nullptr)
+	{ }
+
+	/**
+	 * Creates and initializes a new instance from an object component.
+	 *
+	 * @param InOwner The object that owns the component.
+	 * @param InComponentName The component to be bound.
+	 */
+	FSequencerPossessedObject(UObject* InOwner, FString InComponentName)
+		: ObjectOrOwner(InOwner)
+		, ComponentName(InComponentName)
+		, CachedComponent(nullptr)
+	{ }
+
+	/**
+	 * Compares two bindings for equality.
+	 *
+	 * @param X The first binding to compare.
+	 * @param Y The second binding to compare.
+	 * @return true if the bindings refer to the same object, false otherwise.
+	 */
+	friend bool operator==(const FSequencerPossessedObject& X, const FSequencerPossessedObject& Y)
+	{
+		return (X.ObjectOrOwner == Y.ObjectOrOwner) && (X.ComponentName == Y.ComponentName);
+	}
+
+	/**
+	 * Gets a pointer to the possessed object.
+	 *
+	 * @return The object (usually an Actor or an ActorComponent).
+	 */
+	SEQUENCER_API UObject* GetObject() const;
+
+private:
+
+	/** The object or the owner of the object being possessed. */
+	UPROPERTY()
+	TLazyObjectPtr<UObject> ObjectOrOwner;
+
+	/** Optional name of an ActorComponent. */
+	UPROPERTY()
+	FString ComponentName;
+
+	/** Cached pointer to the Actor component (only if ComponentName is set). */
+	UPROPERTY(transient)
+	mutable UObject* CachedComponent;
+};
diff --git a/Engine/Source/Editor/UMGEditor/Private/Animation/UMGBindingManager.cpp b/Engine/Source/Editor/UMGEditor/Private/Animation/UMGBindingManager.cpp
deleted file mode 100644
diff --git a/Engine/Source/Editor/UMGEditor/Private/Animation/UMGBindingManager.h b/Engine/Source/Editor/UMGEditor/Private/Animation/UMGBindingManager.h
deleted file mode 100644
diff --git a/Engine/Source/Editor/UMGEditor/Private/Animation/UMGMovieSceneObjectManager.cpp b/Engine/Source/Editor/UMGEditor/Private/Animation/UMGMovieSceneObjectManager.cpp
new file mode 100644
--- /dev/null
+++ b/Engine/Source/Editor/UMGEditor/Private/Animation/UMGMovieSceneObjectManager.cpp
@@ -0,0 +1,280 @@
+// Copyright 1998-2015 Epic Games, Inc. All Rights Reserved.
+
+#include "UMGEditorPrivatePCH.h"
+#include "MovieScene.h"
+#include "UMGMovieSceneObjectManager.h"
+#include "Animation/WidgetAnimation.h"
+#include "WidgetBlueprint.h"
+#include "WidgetBlueprintEditor.h"
+
+
+#define LOCTEXT_NAMESPACE "UUMGBindingManager"
+
+
+/* UUMGBindingManager structors
+ *****************************************************************************/
+
+UUMGMovieSceneObjectManager::UUMGMovieSceneObjectManager(const FObjectInitializer& ObjectInitializer)
+	: Super(ObjectInitializer)
+	, WidgetAnimation(nullptr)
+	, WidgetBlueprintEditor(nullptr)
+{ }
+
+
+/* UUMGBindingManager interface
+ *****************************************************************************/
+
+void UUMGMovieSceneObjectManager::InitPreviewObjects()
+{
+	if ((WidgetAnimation == nullptr) || (WidgetBlueprintEditor == nullptr))
+	{
+		return;
+	}
+
+	// Populate preview object to guid map
+	UUserWidget* PreviewWidget = WidgetBlueprintEditor->GetPreview();
+
+	if (PreviewWidget == nullptr)
+	{
+		return;
+	}
+
+	UWidgetTree* WidgetTree = PreviewWidget->WidgetTree;
+
+	for (const FWidgetAnimationBinding& Binding : WidgetAnimation->AnimationBindings)
+	{
+		UObject* FoundObject = Binding.FindRuntimeObject(*WidgetTree);
+
+		if (FoundObject == nullptr)
+		{
+			continue;
+		}
+
+		UPanelSlot* FoundSlot = Cast<UPanelSlot>(FoundObject);
+
+		if (FoundSlot == nullptr)
+		{
+			IdToPreviewObjects.Add(Binding.AnimationId, FoundObject);
+			PreviewObjectToIds.Add(FoundObject, Binding.AnimationId);
+		}
+		else
+		{
+			IdToSlotContentPreviewObjects.Add(Binding.AnimationId, FoundSlot->Content);
+			SlotContentPreviewObjectToIds.Add(FoundSlot->Content, Binding.AnimationId);
+		}
+	}
+}
+
+
+void UUMGMovieSceneObjectManager::BindWidgetBlueprintEditor(FWidgetBlueprintEditor& InWidgetBlueprintEditor, UWidgetAnimation& InAnimation)
+{
+	WidgetAnimation = &InAnimation;
+	WidgetBlueprintEditor = &InWidgetBlueprintEditor;
+	WidgetBlueprintEditor->GetOnWidgetPreviewUpdated().AddUObject(this, &UUMGMovieSceneObjectManager::HandleWidgetPreviewUpdated);
+}
+
+
+void UUMGMovieSceneObjectManager::UnbindWidgetBlueprintEditor()
+{
+	if (WidgetBlueprintEditor != nullptr)
+	{
+		WidgetBlueprintEditor->GetOnWidgetPreviewUpdated().RemoveAll(this);
+	}
+}
+
+
+
+/* IMovieSceneObjectManager interface
+ *****************************************************************************/
+
+bool UUMGMovieSceneObjectManager::AllowsSpawnableObjects() const
+{
+	return false;
+}
+
+
+void UUMGMovieSceneObjectManager::BindPossessableObject(const FGuid& ObjectId, UObject& PossessedObject)
+{
+	if (WidgetAnimation == nullptr)
+	{
+		return;
+	}
+
+	UPanelSlot* PossessedSlot = Cast<UPanelSlot>(&PossessedObject);
+
+	if ((PossessedSlot != nullptr) && (PossessedSlot->Content != nullptr))
+	{
+		SlotContentPreviewObjectToIds.Add(PossessedSlot->Content, ObjectId);
+		IdToSlotContentPreviewObjects.Add(ObjectId, PossessedSlot->Content);
+
+		// Save the name of the widget containing the slots. This is the object
+		// to look up that contains the slot itself (the thing we are animating).
+		FWidgetAnimationBinding NewBinding;
+		{
+			NewBinding.AnimationId = ObjectId;
+			NewBinding.SlotWidgetName = PossessedSlot->GetFName();
+			NewBinding.WidgetName = PossessedSlot->Content->GetFName();
+		}
+
+		WidgetAnimation->AnimationBindings.Add(NewBinding);
+	}
+	else if (PossessedSlot == nullptr)
+	{
+		PreviewObjectToIds.Add(&PossessedObject, ObjectId);
+		IdToPreviewObjects.Add(ObjectId, &PossessedObject);
+
+		FWidgetAnimationBinding NewBinding;
+		{
+			NewBinding.AnimationId = ObjectId;
+			NewBinding.WidgetName = PossessedObject.GetFName();
+		}
+
+		WidgetAnimation->AnimationBindings.Add(NewBinding);
+	}
+}
+
+
+bool UUMGMovieSceneObjectManager::CanPossessObject(UObject& Object) const
+{
+	if (WidgetBlueprintEditor == nullptr)
+	{
+		return false;
+	}
+
+	UPanelSlot* Slot = Cast<UPanelSlot>(&Object);
+
+	if ((Slot != nullptr) && (Slot->Content == nullptr))
+	{
+		// can't possess empty slots.
+		return false;
+	}
+
+	UWidgetBlueprint* WidgetBlueprint = WidgetBlueprintEditor->GetWidgetBlueprintObj();
+	UUserWidget* PreviewWidget = WidgetBlueprintEditor->GetPreview();
+
+	return (Object.IsA<UVisual>() && Object.IsIn(PreviewWidget));
+}
+
+
+FGuid UUMGMovieSceneObjectManager::FindObjectId(const UMovieScene& MovieScene, UObject& Object) const
+{
+	UPanelSlot* Slot = Cast<UPanelSlot>(&Object);
+
+	if (Slot != nullptr)
+	{
+		// slot guids are tracked by their content.
+		return PreviewObjectToIds.FindRef(Slot->Content);
+	}
+
+	return PreviewObjectToIds.FindRef(&Object);
+}
+
+
+UObject* UUMGMovieSceneObjectManager::FindObject(const FGuid& ObjectId) const
+{
+	TWeakObjectPtr<UObject> PreviewObject = IdToPreviewObjects.FindRef(ObjectId);
+
+	if (PreviewObject.IsValid())
+	{
+		return PreviewObject.Get();
+	}
+
+	TWeakObjectPtr<UObject> SlotContentPreviewObject = IdToSlotContentPreviewObjects.FindRef(ObjectId);
+
+	if (SlotContentPreviewObject.IsValid())
+	{
+		return SlotContentPreviewObject.Get();
+	}
+	
+	return nullptr;
+}
+
+
+bool UUMGMovieSceneObjectManager::TryGetObjectDisplayName(const FGuid& ObjectId, FText& OutDisplayName) const
+{
+	// TODO: This gets called every frame for every bound object and could
+	// be a potential performance issue for a really complicated animation.
+
+	TWeakObjectPtr<UObject> PreviewObject = IdToPreviewObjects.FindRef(ObjectId);
+
+	if (PreviewObject.IsValid())
+	{
+		OutDisplayName = FText::FromString(PreviewObject.Get()->GetName());
+		return true;
+	}
+
+	TWeakObjectPtr<UObject> SlotContentPreviewObject = IdToSlotContentPreviewObjects.FindRef(ObjectId);
+
+	if (SlotContentPreviewObject.IsValid())
+	{
+		UWidget* SlotContent = Cast<UWidget>(SlotContentPreviewObject.Get());
+		FText PanelName = SlotContent->Slot != nullptr && SlotContent->Slot->Parent != nullptr
+			? FText::FromString(SlotContent->Slot->Parent->GetName())
+			: LOCTEXT("InvalidPanel", "Invalid Panel");
+		FText ContentName = FText::FromString(SlotContent->GetName());
+		OutDisplayName = FText::Format(LOCTEXT("SlotObject", "{0} ({1} Slot)"), ContentName, PanelName);
+
+		return true;
+	}
+
+	return false;
+}
+
+
+void UUMGMovieSceneObjectManager::UnbindPossessableObjects(const FGuid& ObjectId)
+{
+	// unbind widgets
+	TArray<TWeakObjectPtr<UObject>> PreviewObjects;
+	IdToPreviewObjects.MultiFind(ObjectId, PreviewObjects);
+
+	for (TWeakObjectPtr<UObject>& PreviewObject : PreviewObjects)
+	{
+		PreviewObjectToIds.Remove(PreviewObject);
+	}
+
+	IdToPreviewObjects.Remove(ObjectId);
+
+	// unbind slots
+	TArray<TWeakObjectPtr<UObject>> SlotContentPreviewObjects;
+	IdToSlotContentPreviewObjects.MultiFind(ObjectId, SlotContentPreviewObjects);
+
+	for (TWeakObjectPtr<UObject>& SlotContentPreviewObject : SlotContentPreviewObjects)
+	{
+		SlotContentPreviewObjectToIds.Remove(SlotContentPreviewObject);
+	}
+
+	IdToSlotContentPreviewObjects.Remove(ObjectId);
+
+	// remove animation bindings
+	if ((WidgetAnimation != nullptr) && (WidgetBlueprintEditor != nullptr))
+	{
+		UWidgetBlueprint* WidgetBlueprint = WidgetBlueprintEditor->GetWidgetBlueprintObj();
+
+		WidgetAnimation->Modify();
+		WidgetAnimation->AnimationBindings.RemoveAll([&](const FWidgetAnimationBinding& Binding) {
+			return Binding.AnimationId == ObjectId;
+		});
+	}
+}
+
+
+/* UUMGBindingManager event handlers
+ *****************************************************************************/
+
+void UUMGMovieSceneObjectManager::HandleWidgetPreviewUpdated()
+{
+	PreviewObjectToIds.Empty();
+	IdToPreviewObjects.Empty();
+	IdToSlotContentPreviewObjects.Empty();
+	SlotContentPreviewObjectToIds.Empty();
+
+	InitPreviewObjects();
+
+	if (WidgetBlueprintEditor != nullptr)
+	{
+		WidgetBlueprintEditor->GetSequencer()->UpdateRuntimeInstances();
+	}
+}
+
+
+#undef LOCTEXT_NAMESPACE
diff --git a/Engine/Source/Editor/UMGEditor/Private/Animation/UMGMovieSceneObjectManager.h b/Engine/Source/Editor/UMGEditor/Private/Animation/UMGMovieSceneObjectManager.h
new file mode 100644
--- /dev/null
+++ b/Engine/Source/Editor/UMGEditor/Private/Animation/UMGMovieSceneObjectManager.h
@@ -0,0 +1,93 @@
+// Copyright 1998-2015 Epic Games, Inc. All Rights Reserved.
+
+#pragma once
+
+#include "MovieSceneObjectManager.h"
+#include "UMGMovieSceneObjectManager.generated.h"
+
+
+class FWidgetBlueprintEditor;
+class UObject;
+class UWidgetAnimation;
+
+
+/**
+ * Movie scene object manager used by Sequencer.
+ */
+UCLASS(BlueprintType, MinimalAPI)
+class UUMGMovieSceneObjectManager
+	: public UObject
+	, public IMovieSceneObjectManager
+{
+	GENERATED_UCLASS_BODY()
+
+public:
+
+	/**
+	 * Get the current widget animation.
+	 *
+	 * @return Widget animation object, or nullptr if not valid.
+	 * @see HasValidWidgetAnimation
+	 */
+	UWidgetAnimation* GetWidgetAnimation()
+	{
+		return WidgetAnimation;
+	}
+
+	/** Check whether the current animation is valid.
+	 *
+	 * @return true if the animation is valid, false otherwise.
+	 * @see GetWidgetAnimation
+	 */
+	bool HasValidWidgetAnimation() const;
+
+	/** Rebuild the mappings to live preview objects. */
+	void InitPreviewObjects();
+
+	/**
+	 * Bind the object manager to the given widget BP editor.
+	 *
+	 * @param InWidgetBlueprintEditor The widget BP editor to bind to.
+	 * @param InAnimation The current animation to set.
+	 */
+	// @todo sequencer: gmp: call this from some appropriate place
+	void BindWidgetBlueprintEditor(FWidgetBlueprintEditor& InWidgetBlueprintEditor, UWidgetAnimation& InAnimation);
+
+	/** Unbind from the widget BP editor. */
+	// @todo sequencer: gmp: call this from some appropriate place
+	void UnbindWidgetBlueprintEditor();
+
+public:
+
+	// IMovieSceneObjectManager interface
+
+	virtual bool AllowsSpawnableObjects() const override;
+	virtual void BindPossessableObject(const FGuid& ObjectId, UObject& PossessedObject) override;
+	virtual bool CanPossessObject(UObject& Object) const override;
+	virtual void DestroyAllSpawnedObjects(UMovieScene& MovieScene) override { }
+	virtual FGuid FindObjectId(const UMovieScene& MovieScene, UObject& Object) const override;
+	virtual UObject* FindObject(const FGuid& ObjectId) const override;
+	virtual void SpawnOrDestroyObjects(UMovieScene* MovieScene, bool DestroyAll) override { }
+	virtual bool TryGetObjectDisplayName(const FGuid& ObjectId, FText& OutDisplayName) const override;
+	virtual void UnbindPossessableObjects(const FGuid& ObjectId) override;
+
+private:
+
+	/** Handles widget preview updates in the current widget BP editor. */
+	void HandleWidgetPreviewUpdated();
+
+private:
+
+	/** Mapping of preview objects to sequencer GUIDs */
+	TMap<TWeakObjectPtr<UObject>, FGuid> PreviewObjectToIds;
+	TMultiMap<FGuid, TWeakObjectPtr<UObject>> IdToPreviewObjects;
+	TMultiMap<FGuid, TWeakObjectPtr<UObject>> IdToSlotContentPreviewObjects;
+	TMap<TWeakObjectPtr<UObject>, FGuid> SlotContentPreviewObjectToIds;
+
+	/** The current widget animation. */
+	UPROPERTY(transient)
+	UWidgetAnimation* WidgetAnimation;
+
+	/** Pointer to the widget BP editor. */
+	FWidgetBlueprintEditor* WidgetBlueprintEditor;
+};
diff --git a/Engine/Source/Runtime/MovieScene/Private/MovieScene.cpp b/Engine/Source/Runtime/MovieScene/Private/MovieScene.cpp
--- a/Engine/Source/Runtime/MovieScene/Private/MovieScene.cpp
+++ b/Engine/Source/Runtime/MovieScene/Private/MovieScene.cpp
@@ -39,7 +39,7 @@ FGuid UMovieScene::AddSpawnable( const FString& Name, UBlueprint* Blueprint, UOb
 	Spawnables.Add( NewSpawnable );
 
 	// Add a new binding so that tracks can be added to it
-	new (ObjectBindings) FMovieSceneObjectBinding( NewSpawnable.GetGuid(), NewSpawnable.GetName() );
+	new (ObjectBindings) FMovieSceneBinding( NewSpawnable.GetGuid(), NewSpawnable.GetName() );
 
 	return NewSpawnable.GetGuid();
 }
@@ -65,7 +65,7 @@ bool UMovieScene::RemoveSpawnable( const FGuid& Guid )
 					Blueprint->MarkPendingKill();
 				}
 
-				RemoveObjectBinding( Guid );
+				RemoveBinding( Guid );
 
 				// Found it!
 				Spawnables.RemoveAt( SpawnableIter.GetIndex() );
@@ -138,7 +138,7 @@ FGuid UMovieScene::AddPossessable( const FString& Name, UClass* Class )
 	Possessables.Add( NewPossessable );
 
 	// Add a new binding so that tracks can be added to it
-	new (ObjectBindings) FMovieSceneObjectBinding( NewPossessable.GetGuid(), NewPossessable.GetName() );
+	new (ObjectBindings) FMovieSceneBinding( NewPossessable.GetGuid(), NewPossessable.GetName() );
 
 	return NewPossessable.GetGuid();
 }
@@ -159,7 +159,7 @@ bool UMovieScene::RemovePossessable( const FGuid& PossessableGuid )
 			// Found it!
 			Possessables.RemoveAt( PossesableIter.GetIndex() );
 
-			RemoveObjectBinding( PossessableGuid );
+			RemoveBinding( PossessableGuid );
 
 			bAnythingRemoved = true;
 			break;
@@ -224,8 +224,8 @@ TArray<UMovieSceneSection*> UMovieScene::GetAllSections() const
 	// Add all object binding sections
 	for( int32 ObjectBindingIndex = 0; ObjectBindingIndex < ObjectBindings.Num(); ++ObjectBindingIndex )
 	{
-		const FMovieSceneObjectBinding& ObjectBinding = ObjectBindings[ObjectBindingIndex];
-		const TArray<UMovieSceneTrack*>& Tracks = ObjectBinding.GetTracks();
+		const FMovieSceneBinding& Binding = ObjectBindings[ObjectBindingIndex];
+		const TArray<UMovieSceneTrack*>& Tracks = Binding.GetTracks();
 
 		for( int32 TrackIndex = 0; TrackIndex < Tracks.Num(); ++TrackIndex )
 		{
@@ -236,7 +236,7 @@ TArray<UMovieSceneSection*> UMovieScene::GetAllSections() const
 	return OutSections;
 }
 
-void UMovieScene::RemoveObjectBinding( const FGuid& Guid )
+void UMovieScene::RemoveBinding( const FGuid& Guid )
 {
 	// Update each type
 	for( int32 BindingIndex = 0; BindingIndex < ObjectBindings.Num(); ++BindingIndex )
@@ -258,11 +258,11 @@ UMovieSceneTrack* UMovieScene::FindTrack( TSubclassOf<UMovieSceneTrack> TrackCla
 	
 	for( int32 ObjectBindingIndex = 0; ObjectBindingIndex < ObjectBindings.Num(); ++ObjectBindingIndex )
 	{
-		const FMovieSceneObjectBinding& ObjectBinding = ObjectBindings[ObjectBindingIndex];
+		const FMovieSceneBinding& Binding = ObjectBindings[ObjectBindingIndex];
 
-		if( ObjectBinding.GetObjectGuid() == ObjectGuid ) 
+		if( Binding.GetObjectGuid() == ObjectGuid ) 
 		{
-			const TArray<UMovieSceneTrack*>& Tracks = ObjectBinding.GetTracks();
+			const TArray<UMovieSceneTrack*>& Tracks = Binding.GetTracks();
 			for( int32 TrackIndex = 0; TrackIndex < Tracks.Num(); ++TrackIndex )
 			{
 				UMovieSceneTrack* Track = Tracks[TrackIndex];
@@ -286,15 +286,15 @@ class UMovieSceneTrack* UMovieScene::AddTrack( TSubclassOf<UMovieSceneTrack> Tra
 
 	for( int32 ObjectBindingIndex = 0; ObjectBindingIndex < ObjectBindings.Num(); ++ObjectBindingIndex )
 	{
-		FMovieSceneObjectBinding& ObjectBinding = ObjectBindings[ObjectBindingIndex];
+		FMovieSceneBinding& Binding = ObjectBindings[ObjectBindingIndex];
 
-		if( ObjectBinding.GetObjectGuid() == ObjectGuid ) 
+		if( Binding.GetObjectGuid() == ObjectGuid ) 
 		{
 			Modify();
 
 			CreatedType = NewObject<UMovieSceneTrack>(this, TrackClass, NAME_None, RF_Transactional);
 			
-			ObjectBinding.AddTrack( *CreatedType );
+			Binding.AddTrack( *CreatedType );
 		}
 	}
 
@@ -307,9 +307,9 @@ bool UMovieScene::RemoveTrack( UMovieSceneTrack* Track )
 	bool bAnythingRemoved = false;
 	for( int32 ObjectBindingIndex = 0; ObjectBindingIndex < ObjectBindings.Num(); ++ObjectBindingIndex )
 	{
-		FMovieSceneObjectBinding& ObjectBinding = ObjectBindings[ObjectBindingIndex];
+		FMovieSceneBinding& Binding = ObjectBindings[ObjectBindingIndex];
 
-		if( ObjectBinding.RemoveTrack( *Track ) )
+		if( Binding.RemoveTrack( *Track ) )
 		{
 			bAnythingRemoved = true;
 			// The track was removed from the current binding, stop searching now as it cannot exist in any other binding
diff --git a/Engine/Source/Runtime/MovieScene/Private/MovieSceneBinding.cpp b/Engine/Source/Runtime/MovieScene/Private/MovieSceneBinding.cpp
--- a/Engine/Source/Runtime/MovieScene/Private/MovieSceneBinding.cpp
+++ b/Engine/Source/Runtime/MovieScene/Private/MovieSceneBinding.cpp
@@ -8,7 +8,7 @@
 /* FMovieSceneBinding interface
  *****************************************************************************/
 
-TRange<float> FMovieSceneObjectBinding::GetTimeRange() const
+TRange<float> FMovieSceneBinding::GetTimeRange() const
 {
 	TArray<TRange<float>> Bounds;
 
@@ -21,13 +21,13 @@ TRange<float> FMovieSceneObjectBinding::GetTimeRange() const
 }
 
 
-void FMovieSceneObjectBinding::AddTrack(UMovieSceneTrack& NewTrack)
+void FMovieSceneBinding::AddTrack(UMovieSceneTrack& NewTrack)
 {
 	Tracks.Add( &NewTrack );
 }
 
 
-bool FMovieSceneObjectBinding::RemoveTrack(UMovieSceneTrack& Track)
+bool FMovieSceneBinding::RemoveTrack(UMovieSceneTrack& Track)
 {
 	return (Tracks.RemoveSingle(&Track) != 0);
 }
diff --git a/Engine/Source/Runtime/MovieScene/Private/MovieSceneBindingManager.cpp b/Engine/Source/Runtime/MovieScene/Private/MovieSceneBindingManager.cpp
deleted file mode 100644
diff --git a/Engine/Source/Runtime/MovieScene/Private/MovieSceneInstance.cpp b/Engine/Source/Runtime/MovieScene/Private/MovieSceneInstance.cpp
--- a/Engine/Source/Runtime/MovieScene/Private/MovieSceneInstance.cpp
+++ b/Engine/Source/Runtime/MovieScene/Private/MovieSceneInstance.cpp
@@ -11,9 +11,9 @@ FMovieSceneInstance::FMovieSceneInstance( UMovieScene& InMovieScene )
 }
 
 
-TScriptInterface<UMovieSceneBindingManager> FMovieSceneInstance::GetBindingManager() const
+TScriptInterface<UMovieSceneObjectManager> FMovieSceneInstance::GetObjectManager() const
 {
-	return MovieScene.Get()->GetBindingManager();
+	return MovieScene.Get()->GetObjectManager();
 }
 
 
@@ -113,10 +113,10 @@ void FMovieSceneInstance::RefreshInstance( IMovieScenePlayer& Player )
 
 	TSet< FGuid > FoundObjectBindings;
 	// Get all tracks for each object binding and create instances for them if needed
-	const TArray<FMovieSceneObjectBinding>& ObjectBindings = MovieScene->GetObjectBindings();
+	const TArray<FMovieSceneBinding>& ObjectBindings = MovieScene->GetBindings();
 	for( int32 BindingIndex = 0; BindingIndex < ObjectBindings.Num(); ++BindingIndex )
 	{
-		const FMovieSceneObjectBinding& ObjectBinding = ObjectBindings[BindingIndex];
+		const FMovieSceneBinding& ObjectBinding = ObjectBindings[BindingIndex];
 
 		// Create an instance for this object binding
 		FMovieSceneObjectBindingInstance& BindingInstance = ObjectBindingInstances.FindOrAdd( ObjectBinding.GetObjectGuid() );
diff --git a/Engine/Source/Runtime/MovieScene/Private/MovieSceneObjectManager.cpp b/Engine/Source/Runtime/MovieScene/Private/MovieSceneObjectManager.cpp
new file mode 100644
--- /dev/null
+++ b/Engine/Source/Runtime/MovieScene/Private/MovieSceneObjectManager.cpp
@@ -0,0 +1,12 @@
+// Copyright 1998-2015 Epic Games, Inc. All Rights Reserved.
+
+#include "MovieScenePrivatePCH.h"
+#include "MovieSceneObjectManager.h"
+
+
+/* UMovieSceneBindingManager structors
+ *****************************************************************************/
+
+UMovieSceneObjectManager::UMovieSceneObjectManager(const FObjectInitializer& ObjectInitializer)
+	: Super(ObjectInitializer)
+{ }
diff --git a/Engine/Source/Runtime/MovieScene/Public/MovieScene.h b/Engine/Source/Runtime/MovieScene/Public/MovieScene.h
--- a/Engine/Source/Runtime/MovieScene/Public/MovieScene.h
+++ b/Engine/Source/Runtime/MovieScene/Public/MovieScene.h
@@ -7,7 +7,7 @@
 
 
 class UBlueprint;
-class UMovieSceneBindingManager;
+class UMovieSceneObjectManager;
 class UMovieSceneSection;
 class UMovieSceneTrack;
 
@@ -354,21 +354,21 @@ public:
 public:
 
 	/**
-	 * Get the movie scene's object binding manager.
-	 *
-	 * @return The scene's binding manager.
+	 * @return All object bindings.
 	 */
-	TScriptInterface<UMovieSceneBindingManager> GetBindingManager() const
+	const TArray<FMovieSceneBinding>& GetBindings() const
 	{
-		return BindingManager;
+		return ObjectBindings;
 	}
 
 	/**
-	 * @return All object bindings.
+	 * Get the movie scene's object manager.
+	 *
+	 * @return The object manager.
 	 */
-	const TArray<FMovieSceneObjectBinding>& GetObjectBindings() const
+	TScriptInterface<UMovieSceneObjectManager> GetObjectManager() const
 	{
-		return ObjectBindings;
+		return BindingManager;
 	}
 
 	/**
@@ -396,17 +396,17 @@ public:
 protected:
 
 	/**
-	 * Removes animation data bound to a guid.
+	 * Removes animation data bound to a GUID.
 	 *
 	 * @param Guid The guid bound to animation data to remove
 	 */
-	void RemoveObjectBinding(const FGuid& Guid);
+	void RemoveBinding(const FGuid& Guid);
 
 private:
 
 	/** The object binding manager. */
 	UPROPERTY()
-	TScriptInterface<UMovieSceneBindingManager> BindingManager;
+	TScriptInterface<UMovieSceneObjectManager> BindingManager;
 
 	/**
 	 * Data-only blueprints for all of the objects that we we're able to spawn.
@@ -422,7 +422,7 @@ private:
 
 	/** Tracks bound to possessed or spawned objects */
 	UPROPERTY()
-	TArray<FMovieSceneObjectBinding> ObjectBindings;
+	TArray<FMovieSceneBinding> ObjectBindings;
 
 	/** Master tracks which are not bound to spawned or possessed objects */
 	UPROPERTY()
diff --git a/Engine/Source/Runtime/MovieScene/Public/MovieSceneBinding.h b/Engine/Source/Runtime/MovieScene/Public/MovieSceneBinding.h
--- a/Engine/Source/Runtime/MovieScene/Public/MovieSceneBinding.h
+++ b/Engine/Source/Runtime/MovieScene/Public/MovieSceneBinding.h
@@ -12,12 +12,12 @@ class UMovieSceneTrack;
  * A set of tracks bound to runtime objects
  */
 USTRUCT()
-struct FMovieSceneObjectBinding
+struct FMovieSceneBinding
 {
 	GENERATED_USTRUCT_BODY()
 
 	/** Default constructor. */
-	FMovieSceneObjectBinding() { }
+	FMovieSceneBinding() { }
 
 	/**
 	 * Creates and initializes a new instance.
@@ -26,7 +26,7 @@ struct FMovieSceneObjectBinding
 	 * @param InBindingName
 	 * @param InTracks
 	 */
-	FMovieSceneObjectBinding(const FGuid& InObjectGuid, const FString& InBindingName, const TArray<UMovieSceneTrack*>& InTracks)
+	FMovieSceneBinding(const FGuid& InObjectGuid, const FString& InBindingName, const TArray<UMovieSceneTrack*>& InTracks)
 		: ObjectGuid(InObjectGuid)
 		, BindingName(InBindingName)
 		, Tracks(InTracks)
@@ -38,7 +38,7 @@ struct FMovieSceneObjectBinding
 	 * @param InObjectGuid
 	 * @param InBindingName
 	 */
-	FMovieSceneObjectBinding(const FGuid& InObjectGuid, const FString& InBindingName)
+	FMovieSceneBinding(const FGuid& InObjectGuid, const FString& InBindingName)
 		: ObjectGuid(InObjectGuid)
 		, BindingName(InBindingName)
 	{ }
diff --git a/Engine/Source/Runtime/MovieScene/Public/MovieSceneBindingManager.h b/Engine/Source/Runtime/MovieScene/Public/MovieSceneBindingManager.h
deleted file mode 100644
diff --git a/Engine/Source/Runtime/MovieScene/Public/MovieSceneInstance.h b/Engine/Source/Runtime/MovieScene/Public/MovieSceneInstance.h
--- a/Engine/Source/Runtime/MovieScene/Public/MovieSceneInstance.h
+++ b/Engine/Source/Runtime/MovieScene/Public/MovieSceneInstance.h
@@ -3,7 +3,7 @@
 #pragma once
 
 
-class UMovieSceneBindingManager;
+class UMovieSceneObjectManager;
 class UMovieSceneTrack;
 
 
@@ -44,13 +44,6 @@ public:
 	 */
 	MOVIESCENE_API void RefreshInstance(class IMovieScenePlayer& Player);
 
-	/**
-	 * Get the movie scene's object binding manager.
-	 *
-	 * @return The scene's binding manager.
-	 */
-	TScriptInterface<UMovieSceneBindingManager> GetBindingManager() const;
-
 	/**
 	 * Get the movie scene associated with this instance.
 	 *
@@ -74,6 +67,13 @@ public:
 		return TimeRange;
 	}
 
+	/**
+	 * Get the movie scene's object manager.
+	 *
+	 * @return The object manager.
+	 */
+	TScriptInterface<UMovieSceneObjectManager> GetObjectManager() const;
+
 protected:
 
 	void RefreshInstanceMap(const TArray<UMovieSceneTrack*>& Tracks, const TArray<UObject*>& RuntimeObjects, FMovieSceneInstanceMap& TrackInstances, class IMovieScenePlayer& Player);
diff --git a/Engine/Source/Runtime/MovieScene/Public/MovieSceneObjectManager.h b/Engine/Source/Runtime/MovieScene/Public/MovieSceneObjectManager.h
new file mode 100644
--- /dev/null
+++ b/Engine/Source/Runtime/MovieScene/Public/MovieSceneObjectManager.h
@@ -0,0 +1,102 @@
+// Copyright 1998-2015 Epic Games, Inc. All Rights Reserved.
+
+#pragma once
+
+#include "MovieSceneObjectManager.generated.h"
+
+
+class FMovieSceneInstance;
+class UMovieScene;
+class UObject;
+
+
+
+/**
+ * Interface for movie scene object managers.
+ */
+UINTERFACE(MinimalAPI, meta=(CannotImplementInterfaceInBlueprint))
+class UMovieSceneObjectManager
+	: public UInterface
+{
+	GENERATED_UINTERFACE_BODY()
+};
+
+
+/**
+ * Interface for movie scene object managers.
+ */
+class IMovieSceneObjectManager
+{
+	GENERATED_IINTERFACE_BODY()
+
+	/**
+	 * Whether objects can be spawned at run-time.
+	 *
+	 * @return true if objects can be spawned by sequencer, false if only existing objects can be possessed.
+	 */
+	virtual bool AllowsSpawnableObjects() const PURE_VIRTUAL(IMovieSceneObjectManager::AllowsSpawnableObjects, return false;);
+
+	/**
+	 * Called when Sequencer has created an object binding for a possessable object
+	 * 
+	 * @param ObjectId The guid used to map to the possessable object.  Note the guid can be bound to multiple objects at once
+	 * @param PossessedObject The runtime object which was possessed.
+	 * @see UnbindPossessableObjects
+	 */
+	virtual void BindPossessableObject(const FGuid& ObjectId, UObject& PossessedObject) PURE_VIRTUAL(IMovieSceneObjectManager::BindPossessableObject,);
+	
+	/**
+	 * @todo gmp: sequencer: add documentation.
+	 */
+	virtual bool CanPossessObject(UObject& Object) const PURE_VIRTUAL(IMovieSceneObjectManager::CanPossessObject, return false;);
+
+	/**
+	 * Destroy all objects that were spawned by this binding manager.
+	 *
+	 * @param MovieScene The movie scene that owns this manager.
+	 */
+	virtual void DestroyAllSpawnedObjects(UMovieScene& MovieScene) PURE_VIRTUAL(IMovieSceneObjectManager::DestroyAllSpawnedObjects,);
+
+	/**
+	 * Find the identifier for a possessed or spawned object.
+	 *
+	 * @param The scene that owns this binding manager.
+	 * @param The object to get the identifier for.
+	 * @return The object identifier, or an invalid GUID if not found.
+	 */
+	virtual FGuid FindObjectId(const UMovieScene& MovieScene, UObject& Object) const PURE_VIRTUAL(IMovieSceneObjectManager::FindGuidForObject, return FGuid(););
+
+	/**
+	 * Gets the possessed or spawned object for the specified identifier.
+	 *
+	 * @param ObjectId The unique identifier of the object.
+	 * @return The object, or nullptr if not found.
+	 */
+	virtual UObject* FindObject(const FGuid& ObjectId) const PURE_VIRTUAL(IMovieSceneObjectManager::FindObject, return nullptr;);
+
+	/**
+	 * Called to spawn or destroy objects for a movie instance.
+	 *
+	 * @param MovieScene The movie scene to spawn or destroy objects for.
+	 * @param DestroyAll If true, destroy all spawned objects for the instance, if false only destroy unused objects.
+	 */
+	virtual void SpawnOrDestroyObjects(UMovieScene* MovieScene, bool DestroyAll) PURE_VIRTUAL(IMovieSceneObjectManager::SpawnOrDestroyObjects,);
+
+	/**
+	 * 
+	 * Tries to get a display name for the binding represented by the guid.
+	 *
+	 * @param ObjectGuid The guid for the object binding.
+	 * @param DisplayName The display name for the object binding.
+	 * @returns true if DisplayName has been set to a valid display name, otherwise false.
+	 */
+	virtual bool TryGetObjectDisplayName(const FGuid& ObjectId, FText& OutDisplayName) const PURE_VIRTUAL(IMovieSceneObjectManager::TryGetObjectDisplayName, return false;);
+
+	/**
+	 * Unbinds all possessable objects from the provided GUID.
+	 *
+	 * @param ObjectId The guid bound to possessable objects that should be removed.
+	 * @see BindPossessableObject
+	 */
+	virtual void UnbindPossessableObjects(const FGuid& ObjectId) PURE_VIRTUAL(IMovieSceneObjectManager::UnbindPossessableObjects,);
+};
