Sequencer: Set snap play to dragged key to false by default.

diff --git a/Engine/Source/Editor/Sequencer/Private/SequencerSettings.cpp b/Engine/Source/Editor/Sequencer/Private/SequencerSettings.cpp
--- a/Engine/Source/Editor/Sequencer/Private/SequencerSettings.cpp
+++ b/Engine/Source/Editor/Sequencer/Private/SequencerSettings.cpp
@@ -33,7 +33,7 @@ USequencerSettings::USequencerSettings( const FObjectInitializer& ObjectInitiali
 	bSnapSectionTimesToInterval = true;
 	bSnapSectionTimesToSections = true;
 	bSnapPlayTimeToInterval = true;
-	bSnapPlayTimeToDraggedKey = true;
+	bSnapPlayTimeToDraggedKey = false;
 	CurveValueSnapInterval = 10.0f;
 	bSnapCurveValueToInterval = true;
 	bIsUsingCleanView = false;
