Sequencer: Fix track selection so that the tracks aren't deselect on undo.

diff --git a/Engine/Source/Editor/Sequencer/Private/SSequencer.cpp b/Engine/Source/Editor/Sequencer/Private/SSequencer.cpp
--- a/Engine/Source/Editor/Sequencer/Private/SSequencer.cpp
+++ b/Engine/Source/Editor/Sequencer/Private/SSequencer.cpp
@@ -1084,9 +1084,6 @@ void SSequencer::OnActorSelectionChanged(UObject*)
 
 	TSet<TSharedRef<FSequencerDisplayNode>> RootNodes(SequencerNodeTree->GetRootNodes());
 
-	// Clear selection
-	Sequencer.Pin()->GetSelection().EmptySelectedOutlinerNodes();
-
 	// Select the nodes that have runtime objects that are selected in the level
 	for (auto Node : RootNodes)
 	{
@@ -1094,14 +1091,27 @@ void SSequencer::OnActorSelectionChanged(UObject*)
 		TArray<UObject*> RuntimeObjects;
 		Sequencer.Pin()->GetRuntimeObjects( Sequencer.Pin()->GetFocusedMovieSceneInstance(), ObjectBindingNode->GetObjectBinding(), RuntimeObjects );
 		
+		bool bDoSelect = false;
 		for (int32 RuntimeIndex = 0; RuntimeIndex < RuntimeObjects.Num(); ++RuntimeIndex )
 		{
 			if (GEditor->GetSelectedActors()->IsSelected(RuntimeObjects[RuntimeIndex]))
 			{
-				Sequencer.Pin()->GetSelection().AddToSelection(Node);
+				bDoSelect = true;
 				break;
 			}
 		}
+
+		if (Sequencer.Pin()->GetSelection().IsSelected(Node) != bDoSelect)
+		{
+			if (bDoSelect)
+			{
+				Sequencer.Pin()->GetSelection().AddToSelection(Node);
+			}
+			else
+			{
+				Sequencer.Pin()->GetSelection().RemoveFromSelection(Node);
+			}
+		}
 	}
 }
 
