Fixed compile error

diff --git a/Engine/Source/Editor/Sequencer/Private/SSequencerCurveEditorToolBar.cpp b/Engine/Source/Editor/Sequencer/Private/SSequencerCurveEditorToolBar.cpp
--- a/Engine/Source/Editor/Sequencer/Private/SSequencerCurveEditorToolBar.cpp
+++ b/Engine/Source/Editor/Sequencer/Private/SSequencerCurveEditorToolBar.cpp
@@ -128,4 +128,6 @@ TSharedRef<SWidget> SSequencerCurveEditorToolBar::MakeCurveEditorCurveOptionsMen
 		FNewMenuDelegate::CreateStatic( &FExtrapolationMenus::MakePostInfinityExtrapSubMenu ) );
 	
 	return MenuBuilder.MakeWidget();
-}
\ No newline at end of file
+}
+
+#undef LOCTEXT_NAMESPACE
