Sequencer: Changed autokey to a setting.

diff --git a/Engine/Source/Editor/Sequencer/Private/Sequencer.cpp b/Engine/Source/Editor/Sequencer/Private/Sequencer.cpp
--- a/Engine/Source/Editor/Sequencer/Private/Sequencer.cpp
+++ b/Engine/Source/Editor/Sequencer/Private/Sequencer.cpp
@@ -166,7 +166,6 @@ FSequencer::FSequencer()
 	, PlaybackState( EMovieScenePlayerStatus::Stopped )
 	, ScrubPosition( 0.0f )
 	, bLoopingEnabled( false )
-	, bAllowAutoKey( false )
 	, bPerspectiveViewportPossessionEnabled( true )
 	, bIsEditingWithinLevelEditor( false )
 	, bNeedTreeRefresh( false )
@@ -527,12 +526,12 @@ FAnimatedRange FSequencer::GetViewRange() const
 
 bool FSequencer::GetAutoKeyEnabled() const 
 {
-	return bAllowAutoKey;
+	return Settings->GetAutoKeyEnabled();
 }
 
 void FSequencer::SetAutoKeyEnabled(bool bAutoKeyEnabled)
 {
-	bAllowAutoKey = bAutoKeyEnabled;
+	Settings->SetAutoKeyEnabled(bAutoKeyEnabled);
 }
 
 bool FSequencer::IsRecordingLive() const 
@@ -1006,11 +1005,6 @@ void FSequencer::OnEndScrubbing()
 	AutoscrollOffset.Reset();
 }
 
-void FSequencer::OnToggleAutoKey()
-{
-	bAllowAutoKey = !bAllowAutoKey;
-}
-
 void FSequencer::OnToggleAutoScroll()
 {
 	bAutoScrollEnabled = !bAutoScrollEnabled;
@@ -1638,9 +1632,9 @@ void FSequencer::BindSequencerCommands()
 
 	SequencerCommandBindings->MapAction(
 		Commands.ToggleAutoKeyEnabled,
-		FExecuteAction::CreateSP( this, &FSequencer::OnToggleAutoKey ),
+		FExecuteAction::CreateLambda( [this]{ Settings->SetAutoKeyEnabled( !Settings->GetAutoKeyEnabled() ); } ),
 		FCanExecuteAction::CreateLambda( []{ return true; } ),
-		FIsActionChecked::CreateSP( this, &FSequencer::OnGetAllowAutoKey ) );
+		FIsActionChecked::CreateLambda( [this]{ return Settings->GetAutoKeyEnabled(); } ) );
 
 	bAutoScrollEnabled = Settings->GetAutoScrollEnabled();
 	SequencerCommandBindings->MapAction(
diff --git a/Engine/Source/Editor/Sequencer/Private/Sequencer.h b/Engine/Source/Editor/Sequencer/Private/Sequencer.h
--- a/Engine/Source/Editor/Sequencer/Private/Sequencer.h
+++ b/Engine/Source/Editor/Sequencer/Private/Sequencer.h
@@ -296,9 +296,6 @@ protected:
 	/** @return The current scrub position */
 	virtual float OnGetScrubPosition() const { return ScrubPosition; }
 
-	/** @return Whether or not we currently allow autokeying */
-	virtual bool OnGetAllowAutoKey() const { return bAllowAutoKey; }
-
 	/**
 	 * Called when the view range is changed by the user
 	 *
@@ -325,13 +322,6 @@ protected:
 	/** Called when the user has finished scrubbing */
 	void OnEndScrubbing();
 
-	/**
-	 * Called when auto-key is toggled by a user
-	 *
-	 * @param bInAllowAutoKey	The new auto key state
-	 */
-	void OnToggleAutoKey();
-
 	/** Called to toggle auto-scroll on and off */
 	void OnToggleAutoScroll();
 	bool GetAutoScrollEnabled() const { return bAutoScrollEnabled; }
@@ -448,9 +438,6 @@ private:
 	/** Whether looping while playing is enabled for this sequencer */
 	bool bLoopingEnabled;
 
-	/** Whether or not we are allowing autokey */
-	bool bAllowAutoKey;
-
 	bool bPerspectiveViewportPossessionEnabled;
 
 	/** True if this sequencer is being edited within the level editor */
diff --git a/Engine/Source/Editor/Sequencer/Private/SequencerSettings.cpp b/Engine/Source/Editor/Sequencer/Private/SequencerSettings.cpp
--- a/Engine/Source/Editor/Sequencer/Private/SequencerSettings.cpp
+++ b/Engine/Source/Editor/Sequencer/Private/SequencerSettings.cpp
@@ -26,6 +26,7 @@ USequencerSettings* USequencerSettingsContainer::GetOrCreate(const TCHAR* InName
 USequencerSettings::USequencerSettings( const FObjectInitializer& ObjectInitializer )
 	: Super( ObjectInitializer )
 {
+	bAutoKeyEnabled = false;
 	bIsSnapEnabled = true;
 	TimeSnapInterval = .05f;
 	bSnapKeyTimesToInterval = true;
@@ -42,6 +43,20 @@ USequencerSettings::USequencerSettings( const FObjectInitializer& ObjectInitiali
 	bShowCurveEditorCurveToolTips = true;
 }
 
+bool USequencerSettings::GetAutoKeyEnabled() const
+{
+	return bAutoKeyEnabled;
+}
+
+void USequencerSettings::SetAutoKeyEnabled(bool InbAutoKeyEnabled)
+{
+	if ( bAutoKeyEnabled != InbAutoKeyEnabled )
+	{
+		bAutoKeyEnabled = InbAutoKeyEnabled;
+		SaveConfig();
+	}
+}
+
 bool USequencerSettings::GetIsSnapEnabled() const
 {
 	return bIsSnapEnabled;
diff --git a/Engine/Source/Editor/Sequencer/Private/SequencerSettings.h b/Engine/Source/Editor/Sequencer/Private/SequencerSettings.h
--- a/Engine/Source/Editor/Sequencer/Private/SequencerSettings.h
+++ b/Engine/Source/Editor/Sequencer/Private/SequencerSettings.h
@@ -22,6 +22,11 @@ class USequencerSettings : public UObject
 public:
 	GENERATED_UCLASS_BODY()
 
+	/** Gets whether or not auto key is enabled. */
+	bool GetAutoKeyEnabled() const;
+	/** Sets whether or not auto key is enabled. */
+	void SetAutoKeyEnabled(bool InbAutoKeyEnabled);
+
 	/** Gets whether or not snapping is enabled. */
 	bool GetIsSnapEnabled() const;
 	/** Sets whether or not snapping is enabled. */
@@ -96,6 +101,9 @@ public:
 	float SnapTimeToInterval(float InTimeValue) const;
 
 protected:
+	UPROPERTY( config )
+	bool bAutoKeyEnabled;
+
 	UPROPERTY( config )
 	bool bIsSnapEnabled;
 
