Sequencer: Fix auto framing happening when the curve editor isn't visible.

diff --git a/Engine/Source/Editor/Sequencer/Private/SSequencer.cpp b/Engine/Source/Editor/Sequencer/Private/SSequencer.cpp
--- a/Engine/Source/Editor/Sequencer/Private/SSequencer.cpp
+++ b/Engine/Source/Editor/Sequencer/Private/SSequencer.cpp
@@ -193,6 +193,8 @@ void SSequencer::Construct( const FArguments& InArgs, TSharedRef< class FSequenc
 	USelection::SelectionChangedEvent.AddSP(this, &SSequencer::OnActorSelectionChanged);
 	Settings = InSequencer->GetSettings();
 
+	Settings->GetOnShowCurveEditorChanged().AddSP(this, &SSequencer::OnCurveEditorVisibilityChanged);
+
 	// Create a node tree which contains a tree of movie scene data to display in the sequence
 	SequencerNodeTree = MakeShareable( new FSequencerNodeTree( InSequencer.Get() ) );
 
@@ -650,6 +652,7 @@ TSharedRef<SWidget> SSequencer::MakeTransportControls()
 SSequencer::~SSequencer()
 {
 	USelection::SelectionChangedEvent.RemoveAll(this);
+	Settings->GetOnShowCurveEditorChanged().RemoveAll(this);
 }
 
 void SSequencer::RegisterActiveTimerForPlayback()
@@ -1214,4 +1217,13 @@ EVisibility SSequencer::GetCurveEditorVisibility() const
 	return Settings->GetShowCurveEditor() ? EVisibility::Visible : EVisibility::Collapsed;
 }
 
+void SSequencer::OnCurveEditorVisibilityChanged()
+{
+	if (CurveEditor.IsValid())
+	{
+		// Only zoom horizontally if the editor is visible
+		CurveEditor->SetAllowAutoFrame(Settings->GetShowCurveEditor());
+	}
+}
+
 #undef LOCTEXT_NAMESPACE
diff --git a/Engine/Source/Editor/Sequencer/Private/SSequencer.h b/Engine/Source/Editor/Sequencer/Private/SSequencer.h
--- a/Engine/Source/Editor/Sequencer/Private/SSequencer.h
+++ b/Engine/Source/Editor/Sequencer/Private/SSequencer.h
@@ -226,6 +226,9 @@ private:
 	/** Called when a column fill percentage is changed by a splitter slot. */
 	void OnColumnFillCoefficientChanged(float FillCoefficient, int32 ColumnIndex);
 
+	/** Called when the curve editor is shown or hidden */
+	void OnCurveEditorVisibilityChanged();
+
 private:
 
 	/** Section area widget */
diff --git a/Engine/Source/Editor/Sequencer/Private/SequencerSettings.cpp b/Engine/Source/Editor/Sequencer/Private/SequencerSettings.cpp
--- a/Engine/Source/Editor/Sequencer/Private/SequencerSettings.cpp
+++ b/Engine/Source/Editor/Sequencer/Private/SequencerSettings.cpp
@@ -236,6 +236,7 @@ void USequencerSettings::SetShowCurveEditor(bool InbShowCurveEditor)
 	if (bShowCurveEditor != InbShowCurveEditor)
 	{
 		bShowCurveEditor = InbShowCurveEditor;
+		OnShowCurveEditorChanged.Broadcast();
 		SaveConfig();
 	}
 }
@@ -260,3 +261,9 @@ float USequencerSettings::SnapTimeToInterval( float InTimeValue ) const
 		? FMath::RoundToInt( InTimeValue / TimeSnapInterval ) * TimeSnapInterval
 		: InTimeValue;
 }
+
+USequencerSettings::FOnShowCurveEditorChanged& USequencerSettings::GetOnShowCurveEditorChanged()
+{
+	return OnShowCurveEditorChanged;
+}
+
diff --git a/Engine/Source/Editor/Sequencer/Private/SequencerSettings.h b/Engine/Source/Editor/Sequencer/Private/SequencerSettings.h
--- a/Engine/Source/Editor/Sequencer/Private/SequencerSettings.h
+++ b/Engine/Source/Editor/Sequencer/Private/SequencerSettings.h
@@ -22,6 +22,8 @@ class USequencerSettings : public UObject
 public:
 	GENERATED_UCLASS_BODY()
 
+	DECLARE_MULTICAST_DELEGATE( FOnShowCurveEditorChanged );
+
 	/** Gets whether or not auto key is enabled. */
 	bool GetAutoKeyEnabled() const;
 	/** Sets whether or not auto key is enabled. */
@@ -100,6 +102,9 @@ public:
 	/** Snaps a time value in seconds to the currently selected interval. */
 	float SnapTimeToInterval(float InTimeValue) const;
 
+	/** Gets the multicast delegate which is run whenever the curve editor is shown/hidden. */
+	FOnShowCurveEditorChanged& GetOnShowCurveEditorChanged();
+
 protected:
 	UPROPERTY( config )
 	bool bAutoKeyEnabled;
@@ -145,4 +150,6 @@ protected:
 
 	UPROPERTY( config )
 	bool bShowCurveEditorCurveToolTips;
+
+	FOnShowCurveEditorChanged OnShowCurveEditorChanged;
 };
\ No newline at end of file
diff --git a/Engine/Source/Editor/UnrealEd/Private/SCurveEditor.cpp b/Engine/Source/Editor/UnrealEd/Private/SCurveEditor.cpp
--- a/Engine/Source/Editor/UnrealEd/Private/SCurveEditor.cpp
+++ b/Engine/Source/Editor/UnrealEd/Private/SCurveEditor.cpp
@@ -52,6 +52,7 @@ void SCurveEditor::Construct(const FArguments& InArgs)
 	GridColor = InArgs._GridColor;
 
 	bIsUsingSlider = false;
+	bAllowAutoFrame = true;
 
 	// if editor size is set, use it, otherwise, use default value
 	if (DesiredSize.Get().IsZero())
@@ -1158,7 +1159,7 @@ void SCurveEditor::SetCurveOwner(FCurveOwnerInterface* InCurveOwner, bool bCanEd
 
 	ValidateSelection();
 
-	if (Settings->GetAutoFrameCurveEditor())
+	if (Settings->GetAutoFrameCurveEditor() && GetAllowAutoFrame())
 	{
 		if( bZoomToFitVertical )
 		{
diff --git a/Engine/Source/Editor/UnrealEd/Public/SCurveEditor.h b/Engine/Source/Editor/UnrealEd/Public/SCurveEditor.h
--- a/Engine/Source/Editor/UnrealEd/Public/SCurveEditor.h
+++ b/Engine/Source/Editor/UnrealEd/Public/SCurveEditor.h
@@ -233,6 +233,10 @@ public:
 	/** Gets a list of the commands handled by this control */
 	UNREALED_API TSharedPtr<FUICommandList> GetCommands();
 
+	/** Gets or sets whether autoframing is allowed */
+	UNREALED_API bool GetAllowAutoFrame() { return bAllowAutoFrame; }
+	UNREALED_API void SetAllowAutoFrame(bool bInAllowAutoFrame) { bAllowAutoFrame = bInAllowAutoFrame; }
+
 private:
 	/** Used to track a key and the curve that owns it */
 	struct FSelectedCurveKey
@@ -629,6 +633,9 @@ private:
 	/** Gradient editor */
 	TSharedPtr<class SColorGradientEditor> GradientViewer;
 
+	/** Flag to allow auto framing */
+	bool bAllowAutoFrame;
+
 protected:
 	/** Minimum input of view range  */
 	TAttribute<float>	ViewMinInput;
