Sequencer: Step to next/previous key.

Hotkeys are:
  Previous key : ,
  Next key : .
  Previous camera key: alt-,
  Next camera key: alt-.

Refactor GetAllKeyAreas to SequencerCommonHelpers

diff --git a/Engine/Source/Editor/Sequencer/Private/SAnimationOutlinerView.cpp b/Engine/Source/Editor/Sequencer/Private/SAnimationOutlinerView.cpp
--- a/Engine/Source/Editor/Sequencer/Private/SAnimationOutlinerView.cpp
+++ b/Engine/Source/Editor/Sequencer/Private/SAnimationOutlinerView.cpp
@@ -144,45 +144,7 @@ void SAnimationOutlinerTreeNode::Construct( const FArguments& InArgs, TSharedRef
 	SetVisibility( TAttribute<EVisibility>( this, &SAnimationOutlinerTreeNode::GetNodeVisibility ) );
 }
 
-void GetAllKeyAreas(TSharedPtr<FSequencerDisplayNode> DisplayNode, TSet<TSharedPtr<IKeyArea>>& KeyAreas)
-{
-	TArray<TSharedPtr<FSequencerDisplayNode>> NodesToCheck;
-	NodesToCheck.Add(DisplayNode);
-	while (NodesToCheck.Num() > 0)
-	{
-		TSharedPtr<FSequencerDisplayNode> NodeToCheck = NodesToCheck[0];
-		NodesToCheck.RemoveAt(0);
 
-		if (NodeToCheck->GetType() == ESequencerNode::Track)
-		{
-			TSharedPtr<FTrackNode> TrackNode = StaticCastSharedPtr<FTrackNode>(NodeToCheck);
-			TArray<TSharedRef<FSectionKeyAreaNode>> KeyAreaNodes;
-			TrackNode->GetChildKeyAreaNodesRecursively(KeyAreaNodes);
-			for (TSharedRef<FSectionKeyAreaNode> KeyAreaNode : KeyAreaNodes)
-			{
-				for (TSharedPtr<IKeyArea> KeyArea : KeyAreaNode->GetAllKeyAreas())
-				{
-					KeyAreas.Add(KeyArea);
-				}
-			}
-		}
-		else
-		{
-			if (NodeToCheck->GetType() == ESequencerNode::KeyArea)
-			{
-				TSharedPtr<FSectionKeyAreaNode> KeyAreaNode = StaticCastSharedPtr<FSectionKeyAreaNode>(NodeToCheck);
-				for (TSharedPtr<IKeyArea> KeyArea : KeyAreaNode->GetAllKeyAreas())
-				{
-					KeyAreas.Add(KeyArea);
-				}
-			}
-			for (TSharedRef<FSequencerDisplayNode> ChildNode : NodeToCheck->GetChildNodes())
-			{
-				NodesToCheck.Add(ChildNode);
-			}
-		}
-	}
-}
 
 FReply SAnimationOutlinerTreeNode::OnPreviousKeyClicked()
 {
@@ -193,7 +155,7 @@ FReply SAnimationOutlinerTreeNode::OnPreviousKeyClicked()
 	bool PreviousKeyFound = false;
 
 	TSet<TSharedPtr<IKeyArea>> KeyAreas;
-	GetAllKeyAreas(DisplayNode, KeyAreas);
+	SequencerHelpers::GetAllKeyAreas(DisplayNode, KeyAreas);
 	for (TSharedPtr<IKeyArea> KeyArea : KeyAreas)
 	{
 		for (FKeyHandle& KeyHandle : KeyArea->GetUnsortedKeyHandles())
@@ -224,7 +186,7 @@ FReply SAnimationOutlinerTreeNode::OnNextKeyClicked()
 	bool NextKeyFound = false;
 
 	TSet<TSharedPtr<IKeyArea>> KeyAreas;
-	GetAllKeyAreas(DisplayNode, KeyAreas);
+	SequencerHelpers::GetAllKeyAreas(DisplayNode, KeyAreas);
 	for (TSharedPtr<IKeyArea> KeyArea : KeyAreas)
 	{
 		for (FKeyHandle& KeyHandle : KeyArea->GetUnsortedKeyHandles())
@@ -252,7 +214,7 @@ FReply SAnimationOutlinerTreeNode::OnAddKeyClicked()
 	float CurrentTime = Sequencer.GetCurrentLocalTime(*Sequencer.GetFocusedMovieScene());
 
 	TSet<TSharedPtr<IKeyArea>> KeyAreas;
-	GetAllKeyAreas(DisplayNode, KeyAreas);
+	SequencerHelpers::GetAllKeyAreas(DisplayNode, KeyAreas);
 
 	TArray<UMovieSceneSection*> KeyAreaSections;
 	for (TSharedPtr<IKeyArea> KeyArea : KeyAreas)
diff --git a/Engine/Source/Editor/Sequencer/Private/SSequencer.cpp b/Engine/Source/Editor/Sequencer/Private/SSequencer.cpp
--- a/Engine/Source/Editor/Sequencer/Private/SSequencer.cpp
+++ b/Engine/Source/Editor/Sequencer/Private/SSequencer.cpp
@@ -1162,6 +1162,130 @@ FReply SSequencer::OnSaveMovieSceneClicked()
 	return FReply::Unhandled();
 }
 
+void SSequencer::StepToNextKey()
+{
+	StepToKey(true, false);
+}
+
+void SSequencer::StepToPreviousKey()
+{
+	StepToKey(false, false);
+}
+
+void SSequencer::StepToNextCameraKey()
+{
+	StepToKey(true, true);
+}
+
+void SSequencer::StepToPreviousCameraKey()
+{
+	StepToKey(false, true);
+}
+
+void SSequencer::StepToKey(bool bStepToNextKey, bool bCameraOnly)
+{
+	TSet< TSharedRef<FSequencerDisplayNode> > Nodes;
+
+	if (bCameraOnly)
+	{
+		TSet<TSharedRef<FSequencerDisplayNode>> RootNodes(SequencerNodeTree->GetRootNodes());
+
+		TSet<TWeakObjectPtr<AActor> > LockedActors;
+		for (int32 i = 0; i < GEditor->LevelViewportClients.Num(); ++i)
+		{		
+			FLevelEditorViewportClient* LevelVC = GEditor->LevelViewportClients[i];
+			if (LevelVC && LevelVC->IsPerspective() && LevelVC->GetViewMode() != VMI_Unknown)
+			{
+				TWeakObjectPtr<AActor> ActorLock = LevelVC->GetActiveActorLock();
+				if (ActorLock.IsValid())
+				{
+					LockedActors.Add(ActorLock);
+				}
+			}
+		}
+
+		for (auto RootNode : RootNodes)
+		{
+			TSharedRef<FObjectBindingNode> ObjectBindingNode = StaticCastSharedRef<FObjectBindingNode>(RootNode);
+			TArray<UObject*> RuntimeObjects;
+			Sequencer.Pin()->GetRuntimeObjects( Sequencer.Pin()->GetFocusedMovieSceneInstance(), ObjectBindingNode->GetObjectBinding(), RuntimeObjects );
+		
+			for (int32 RuntimeIndex = 0; RuntimeIndex < RuntimeObjects.Num(); ++RuntimeIndex )
+			{
+				AActor* RuntimeActor = Cast<AActor>(RuntimeObjects[RuntimeIndex]);
+				if (RuntimeActor != NULL && LockedActors.Contains(RuntimeActor))
+				{
+					Nodes.Add(RootNode);
+				}
+			}
+		}
+	}
+	else
+	{
+		const TSet< TSharedRef<FSequencerDisplayNode> >& SelectedNodes = Sequencer.Pin()->GetSelection().GetSelectedOutlinerNodes();
+		Nodes = SelectedNodes;
+
+		if (Nodes.Num() == 0)
+		{
+			TSet<TSharedRef<FSequencerDisplayNode>> RootNodes(SequencerNodeTree->GetRootNodes());
+			for (auto RootNode : RootNodes)
+			{
+				Nodes.Add(RootNode);
+
+				SequencerHelpers::GetDescendantNodes(RootNode, Nodes);
+			}
+		}
+	}
+		
+	if (Nodes.Num() > 0)
+	{
+		float ClosestKeyDistance = MAX_FLT;
+		float CurrentTime = Sequencer.Pin()->GetCurrentLocalTime(*Sequencer.Pin()->GetFocusedMovieScene());
+		float StepToTime = 0;
+		bool StepToKeyFound = false;
+
+		auto It = Nodes.CreateConstIterator();
+		bool bExpand = !(*It).Get().IsExpanded();
+
+		for (auto Node : Nodes)
+		{
+			TSet<TSharedPtr<IKeyArea>> KeyAreas;
+			SequencerHelpers::GetAllKeyAreas(Node, KeyAreas);
+
+			for (TSharedPtr<IKeyArea> KeyArea : KeyAreas)
+			{
+				for (FKeyHandle& KeyHandle : KeyArea->GetUnsortedKeyHandles())
+				{
+					float KeyTime = KeyArea->GetKeyTime(KeyHandle);
+					if (bStepToNextKey)
+					{
+						if (KeyTime > CurrentTime && KeyTime - CurrentTime < ClosestKeyDistance)
+						{
+							StepToTime = KeyTime;
+							ClosestKeyDistance = KeyTime - CurrentTime;
+							StepToKeyFound = true;
+						}
+					}
+					else
+					{
+						if (KeyTime < CurrentTime && CurrentTime - KeyTime < ClosestKeyDistance)
+						{
+							StepToTime = KeyTime;
+							ClosestKeyDistance = CurrentTime - KeyTime;
+							StepToKeyFound = true;
+						}
+					}
+				}
+			}
+		}
+
+		if (StepToKeyFound)
+		{
+			Sequencer.Pin()->SetGlobalTime(StepToTime);
+		}
+	}
+}
+
 void SSequencer::ToggleExpandCollapseSelectedNodes(bool bDescendants)
 {
 	const TSet< TSharedRef<FSequencerDisplayNode> >& SelectedNodes = Sequencer.Pin()->GetSelection().GetSelectedOutlinerNodes();
diff --git a/Engine/Source/Editor/Sequencer/Private/SSequencer.h b/Engine/Source/Editor/Sequencer/Private/SSequencer.h
--- a/Engine/Source/Editor/Sequencer/Private/SSequencer.h
+++ b/Engine/Source/Editor/Sequencer/Private/SSequencer.h
@@ -97,6 +97,13 @@ public:
 	/** Deletes selected nodes out of the sequencer node tree */
 	void DeleteSelectedNodes();
 
+	/** Step to next and previous keyframes */
+	void StepToNextKey();
+	void StepToPreviousKey();
+	void StepToNextCameraKey();
+	void StepToPreviousCameraKey();
+	void StepToKey(bool bStepToNextKey, bool bCameraOnly);
+
 	/** Expand or collapse selected nodes out of the sequencer node tree */
 	void ToggleExpandCollapseSelectedNodes(bool bDescendants = false);
 	void ExpandCollapseNode(TSharedRef<FSequencerDisplayNode> SelectedNode, bool bDescendants, bool bExpand);
diff --git a/Engine/Source/Editor/Sequencer/Private/Sequencer.cpp b/Engine/Source/Editor/Sequencer/Private/Sequencer.cpp
--- a/Engine/Source/Editor/Sequencer/Private/Sequencer.cpp
+++ b/Engine/Source/Editor/Sequencer/Private/Sequencer.cpp
@@ -1559,6 +1559,26 @@ void FSequencer::StepBackward()
 	OnStepBackward();
 }
 
+void FSequencer::StepToNextKey()
+{
+	SequencerWidget->StepToNextKey();
+}
+
+void FSequencer::StepToPreviousKey()
+{
+	SequencerWidget->StepToPreviousKey();
+}
+
+void FSequencer::StepToNextCameraKey()
+{
+	SequencerWidget->StepToNextCameraKey();
+}
+
+void FSequencer::StepToPreviousCameraKey()
+{
+	SequencerWidget->StepToPreviousCameraKey();
+}
+
 void FSequencer::ToggleExpandCollapseNodes()
 {
 	SequencerWidget->ToggleExpandCollapseSelectedNodes();
@@ -1618,6 +1638,22 @@ void FSequencer::BindSequencerCommands()
 		Commands.StepBackward,
 		FExecuteAction::CreateSP( this, &FSequencer::StepBackward ) );
 
+	SequencerCommandBindings->MapAction(
+		Commands.StepToNextKey,
+		FExecuteAction::CreateSP( this, &FSequencer::StepToNextKey ) );
+
+	SequencerCommandBindings->MapAction(
+		Commands.StepToPreviousKey,
+		FExecuteAction::CreateSP( this, &FSequencer::StepToPreviousKey ) );
+
+	SequencerCommandBindings->MapAction(
+		Commands.StepToNextCameraKey,
+		FExecuteAction::CreateSP( this, &FSequencer::StepToNextCameraKey ) );
+
+	SequencerCommandBindings->MapAction(
+		Commands.StepToPreviousCameraKey,
+		FExecuteAction::CreateSP( this, &FSequencer::StepToPreviousCameraKey ) );
+
 	SequencerCommandBindings->MapAction(
 		Commands.ToggleExpandCollapseNodes,
 		FExecuteAction::CreateSP(this, &FSequencer::ToggleExpandCollapseNodes));
diff --git a/Engine/Source/Editor/Sequencer/Private/Sequencer.h b/Engine/Source/Editor/Sequencer/Private/Sequencer.h
--- a/Engine/Source/Editor/Sequencer/Private/Sequencer.h
+++ b/Engine/Source/Editor/Sequencer/Private/Sequencer.h
@@ -338,6 +338,10 @@ protected:
 	void Rewind();
 	void StepForward();
 	void StepBackward();
+	void StepToNextKey();
+	void StepToPreviousKey();
+	void StepToNextCameraKey();
+	void StepToPreviousCameraKey();
 
 	/** Expand or collapse selected nodes */
 	void ToggleExpandCollapseNodes();
diff --git a/Engine/Source/Editor/Sequencer/Private/SequencerCommands.cpp b/Engine/Source/Editor/Sequencer/Private/SequencerCommands.cpp
--- a/Engine/Source/Editor/Sequencer/Private/SequencerCommands.cpp
+++ b/Engine/Source/Editor/Sequencer/Private/SequencerCommands.cpp
@@ -12,6 +12,10 @@ void FSequencerCommands::RegisterCommands()
 	UI_COMMAND( Rewind, "Rewind", "Rewind the timeline.", EUserInterfaceActionType::Button, FInputChord(EKeys::Up) );
 	UI_COMMAND( StepForward, "Step Forward", "Step the timeline forward.", EUserInterfaceActionType::Button, FInputChord(EKeys::Right) );
 	UI_COMMAND( StepBackward, "Step Backward", "Step the timeline backward", EUserInterfaceActionType::Button, FInputChord(EKeys::Left) );
+	UI_COMMAND( StepToNextKey, "Step to Next Key", "Step to the next key", EUserInterfaceActionType::Button, FInputChord(EKeys::Period) );
+	UI_COMMAND( StepToPreviousKey, "Step to Previous Key", "Step to the previous key", EUserInterfaceActionType::Button, FInputChord(EKeys::Comma) );
+	UI_COMMAND( StepToNextCameraKey, "Step to Next Camera Key", "Step to the next camera key", EUserInterfaceActionType::Button, FInputChord(EModifierKey::Alt, EKeys::Period) );
+	UI_COMMAND( StepToPreviousCameraKey, "Step to Previous Camera Key", "Step to the previous camera key", EUserInterfaceActionType::Button, FInputChord(EModifierKey::Alt, EKeys::Comma) );
 
 	UI_COMMAND( ToggleExpandCollapseNodes, "Toggle Expand/Collapse Nodes", "Toggle expand or collapse selected nodes.", EUserInterfaceActionType::Button, FInputChord(EKeys::O) );
 	UI_COMMAND( ToggleExpandCollapseNodesAndDescendants, "Toggle Expand/Collapse Nodes and Descendants", "Toggle expand or collapse selected nodes and their descendants.", EUserInterfaceActionType::Button, FInputChord(EModifierKey::Shift, EKeys::O) );
diff --git a/Engine/Source/Editor/Sequencer/Private/SequencerCommands.h b/Engine/Source/Editor/Sequencer/Private/SequencerCommands.h
--- a/Engine/Source/Editor/Sequencer/Private/SequencerCommands.h
+++ b/Engine/Source/Editor/Sequencer/Private/SequencerCommands.h
@@ -30,6 +30,18 @@ public:
 	/** Step backward */
 	TSharedPtr< FUICommandInfo > StepBackward;
 
+	/** Step to next key */
+	TSharedPtr< FUICommandInfo > StepToNextKey;
+
+	/** Step to previous key */
+	TSharedPtr< FUICommandInfo > StepToPreviousKey;
+
+	/** Step to next camera key */
+	TSharedPtr< FUICommandInfo > StepToNextCameraKey;
+
+	/** Step to previous camera key */
+	TSharedPtr< FUICommandInfo > StepToPreviousCameraKey;
+
 	/** Expand/collapse nodes */
 	TSharedPtr< FUICommandInfo > ToggleExpandCollapseNodes;
 
diff --git a/Engine/Source/Editor/Sequencer/Private/SequencerCommonHelpers.cpp b/Engine/Source/Editor/Sequencer/Private/SequencerCommonHelpers.cpp
new file mode 100644
--- /dev/null
+++ b/Engine/Source/Editor/Sequencer/Private/SequencerCommonHelpers.cpp
@@ -0,0 +1,55 @@
+// Copyright 1998-2015 Epic Games, Inc. All Rights Reserved.
+
+#include "SequencerPrivatePCH.h"
+#include "SequencerCommonHelpers.h"
+
+void SequencerHelpers::GetAllKeyAreas(TSharedPtr<FSequencerDisplayNode> DisplayNode, TSet<TSharedPtr<IKeyArea>>& KeyAreas)
+{
+	TArray<TSharedPtr<FSequencerDisplayNode>> NodesToCheck;
+	NodesToCheck.Add(DisplayNode);
+	while (NodesToCheck.Num() > 0)
+	{
+		TSharedPtr<FSequencerDisplayNode> NodeToCheck = NodesToCheck[0];
+		NodesToCheck.RemoveAt(0);
+
+		if (NodeToCheck->GetType() == ESequencerNode::Track)
+		{
+			TSharedPtr<FTrackNode> TrackNode = StaticCastSharedPtr<FTrackNode>(NodeToCheck);
+			TArray<TSharedRef<FSectionKeyAreaNode>> KeyAreaNodes;
+			TrackNode->GetChildKeyAreaNodesRecursively(KeyAreaNodes);
+			for (TSharedRef<FSectionKeyAreaNode> KeyAreaNode : KeyAreaNodes)
+			{
+				for (TSharedPtr<IKeyArea> KeyArea : KeyAreaNode->GetAllKeyAreas())
+				{
+					KeyAreas.Add(KeyArea);
+				}
+			}
+		}
+		else
+		{
+			if (NodeToCheck->GetType() == ESequencerNode::KeyArea)
+			{
+				TSharedPtr<FSectionKeyAreaNode> KeyAreaNode = StaticCastSharedPtr<FSectionKeyAreaNode>(NodeToCheck);
+				for (TSharedPtr<IKeyArea> KeyArea : KeyAreaNode->GetAllKeyAreas())
+				{
+					KeyAreas.Add(KeyArea);
+				}
+			}
+			for (TSharedRef<FSequencerDisplayNode> ChildNode : NodeToCheck->GetChildNodes())
+			{
+				NodesToCheck.Add(ChildNode);
+			}
+		}
+	}
+}
+
+void SequencerHelpers::GetDescendantNodes(TSharedRef<FSequencerDisplayNode> DisplayNode, TSet<TSharedRef<FSequencerDisplayNode>>& Nodes)
+{
+	for (auto ChildNode : DisplayNode.Get().GetChildNodes())
+	{
+		Nodes.Add(ChildNode);
+
+		GetDescendantNodes(ChildNode, Nodes);
+	}
+}
+
diff --git a/Engine/Source/Editor/Sequencer/Private/SequencerCommonHelpers.h b/Engine/Source/Editor/Sequencer/Private/SequencerCommonHelpers.h
new file mode 100644
--- /dev/null
+++ b/Engine/Source/Editor/Sequencer/Private/SequencerCommonHelpers.h
@@ -0,0 +1,18 @@
+// Copyright 1998-2015 Epic Games, Inc. All Rights Reserved.
+
+#pragma once
+
+
+class SequencerHelpers
+{
+public:
+	/**
+	 * Gets the key areas from the requested node
+	 */
+	static void GetAllKeyAreas(TSharedPtr<FSequencerDisplayNode> DisplayNode, TSet<TSharedPtr<IKeyArea>>& KeyAreas);
+
+	/**
+	* Get descendant nodes
+	*/
+	static void GetDescendantNodes(TSharedRef<FSequencerDisplayNode> DisplayNode, TSet<TSharedRef<FSequencerDisplayNode> >& Nodes);
+};
\ No newline at end of file
diff --git a/Engine/Source/Editor/Sequencer/Private/SequencerPrivatePCH.h b/Engine/Source/Editor/Sequencer/Private/SequencerPrivatePCH.h
--- a/Engine/Source/Editor/Sequencer/Private/SequencerPrivatePCH.h
+++ b/Engine/Source/Editor/Sequencer/Private/SequencerPrivatePCH.h
@@ -23,6 +23,7 @@
 #include "SAnimationOutlinerView.h"
 #include "MovieSceneInstance.h"
 #include "SequencerSettings.h"
+#include "SequencerCommonHelpers.h"
 #include "SequencerCurveOwner.h"
 #include "SSequencerCurveEditor.h"
 #include "SSequencerCurveEditorToolBar.h"
