Fixed non-unity

diff --git a/Engine/Source/Editor/Sequencer/Private/SSequencer.cpp b/Engine/Source/Editor/Sequencer/Private/SSequencer.cpp
--- a/Engine/Source/Editor/Sequencer/Private/SSequencer.cpp
+++ b/Engine/Source/Editor/Sequencer/Private/SSequencer.cpp
@@ -27,8 +27,8 @@
 #include "SNumericDropDown.h"
 #include "EditorWidgetsModule.h"
 #include "SequencerTrackLaneFactory.h"
-
 #include "SSequencerSplitterOverlay.h"
+#include "IKeyArea.h"
 
 #define LOCTEXT_NAMESPACE "Sequencer"
 
