Sequencer: Fix allow auto frame on initialization.

diff --git a/Engine/Source/Editor/Sequencer/Private/SSequencer.cpp b/Engine/Source/Editor/Sequencer/Private/SSequencer.cpp
--- a/Engine/Source/Editor/Sequencer/Private/SSequencer.cpp
+++ b/Engine/Source/Editor/Sequencer/Private/SSequencer.cpp
@@ -242,6 +242,7 @@ void SSequencer::Construct( const FArguments& InArgs, TSharedRef< class FSequenc
 		.Visibility( this, &SSequencer::GetCurveEditorVisibility )
 		.OnViewRangeChanged( InArgs._OnViewRangeChanged )
 		.ViewRange( InArgs._ViewRange );
+	CurveEditor->SetAllowAutoFrame(Settings->GetShowCurveEditor());
 
 	const int32 				Column0	= 0,	Column1	= 1;
 	const int32 Row0	= 0,
