MovieScene: Merging non-functional refactor changes/fixes

diff --git a/Engine/Source/Editor/Sequencer/Private/Sequencer.cpp b/Engine/Source/Editor/Sequencer/Private/Sequencer.cpp
--- a/Engine/Source/Editor/Sequencer/Private/Sequencer.cpp
+++ b/Engine/Source/Editor/Sequencer/Private/Sequencer.cpp
@@ -38,10 +38,12 @@
 #include "Engine/Selection.h"
 #include "IMenu.h"
 
+
 #define LOCTEXT_NAMESPACE "Sequencer"
 
 DEFINE_LOG_CATEGORY(LogSequencer);
 
+
 bool FSequencer::IsSequencerEnabled()
 {
 	return true;
@@ -169,9 +171,7 @@ FSequencer::FSequencer()
 	, bPerspectiveViewportPossessionEnabled( true )
 	, bIsEditingWithinLevelEditor( false )
 	, bNeedTreeRefresh( false )
-{
-	
-}
+{ }
 
 FSequencer::~FSequencer()
 {
@@ -440,8 +440,7 @@ void FSequencer::SpawnOrDestroyPuppetObjects( TSharedRef<FMovieSceneInstance> Mo
 {
 	if( ObjectBindingManager->AllowsSpawnableObjects() )
 	{
-		const bool bDestroyAll = false;
-		ObjectBindingManager->SpawnOrDestroyObjectsForInstance( MovieSceneInstance, bDestroyAll );
+		ObjectBindingManager->SpawnOrDestroyObjectsForInstance( MovieSceneInstance, false /*bDestroyAll*/ );
 	}
 }
 
@@ -468,7 +467,8 @@ void FSequencer::OnActorsDropped( const TArray<TWeakObjectPtr<AActor> >& Actors
 	for( TWeakObjectPtr<AActor> WeakActor : Actors )
 	{
 		AActor* Actor = WeakActor.Get();
-		if( Actor != NULL )
+
+		if (Actor != nullptr)
 		{
 			// Grab the MovieScene that is currently focused.  We'll add our Blueprint as an inner of the
 			// MovieScene asset.
@@ -616,54 +616,56 @@ FGuid FSequencer::GetHandleToObject( UObject* Object )
 
 	// Check here for spawnable otherwise spawnables get recreated as possessables, which doesn't make sens
 	FMovieSceneSpawnable* Spawnable = FocusedMovieScene->FindSpawnable(ObjectGuid);
-	if( !Spawnable )
-	{
-		if(ObjectGuid.IsValid())
-		{
-			// Make sure that the possessable is still valid, if it's not remove the binding so new one 
-			// can be created.  This can happen due to undo.
-			FMovieScenePossessable* Possessable = FocusedMovieScene->FindPossessable(ObjectGuid);
-			if(Possessable == nullptr)
-			{
-				ObjectBindingManager->UnbindPossessableObjects(ObjectGuid);
-				ObjectGuid.Invalidate();
-			}
-		}
 
-		bool bPossessableAdded = false;
+	if (Spawnable)
+	{
+		return ObjectGuid;
+	}
 
-		// If the object guid was not found attempt to add it
-		// Note: Only possessed actors can be added like this
-		if(!ObjectGuid.IsValid() && ObjectBindingManager->CanPossessObject(*Object))
+	if (ObjectGuid.IsValid())
+	{
+		// Make sure that the possessable is still valid, if it's not remove the binding so new one 
+		// can be created.  This can happen due to undo.
+		FMovieScenePossessable* Possessable = FocusedMovieScene->FindPossessable(ObjectGuid);
+		
+		if (Possessable == nullptr)
 		{
-			// @todo sequencer: Undo doesn't seem to be working at all
-			const FScopedTransaction Transaction(LOCTEXT("UndoPossessingObject", "Possess Object with MovieScene"));
+			ObjectBindingManager->UnbindPossessableObjects(ObjectGuid);
+			ObjectGuid.Invalidate();
+		}
+	}
 
-			// Possess the object!
-			{
-				// Create a new possessable
-				FocusedMovieScene->Modify();
+	bool bPossessableAdded = false;
 
-				ObjectGuid = FocusedMovieScene->AddPossessable(Object->GetName(), Object->GetClass());
+	// If the object guid was not found attempt to add it
+	// Note: Only possessed actors can be added like this
+	if (!ObjectGuid.IsValid() && ObjectBindingManager->CanPossessObject(*Object))
+	{
+		// @todo sequencer: Undo doesn't seem to be working at all
+		const FScopedTransaction Transaction(LOCTEXT("UndoPossessingObject", "Possess Object with MovieScene"));
 
-				if(IsShotFilteringOn())
-				{
-					AddUnfilterableObject(ObjectGuid);
-				}
+		// Possess the object!
+		{
+			// Create a new possessable
+			FocusedMovieScene->Modify();
 
-				ObjectBindingManager->BindPossessableObject(ObjectGuid, *Object);
+			ObjectGuid = FocusedMovieScene->AddPossessable(Object->GetName(), Object->GetClass());
 
-				bPossessableAdded = true;
+			if (IsShotFilteringOn())
+			{
+				AddUnfilterableObject(ObjectGuid);
 			}
-		}
 
-		if(bPossessableAdded)
-		{
-			SpawnOrDestroyPuppetObjects(GetFocusedMovieSceneInstance());
-
-			NotifyMovieSceneDataChanged();
+			ObjectBindingManager->BindPossessableObject(ObjectGuid, *Object);
+			bPossessableAdded = true;
 		}
 	}
+
+	if (bPossessableAdded)
+	{
+		SpawnOrDestroyPuppetObjects(GetFocusedMovieSceneInstance());
+		NotifyMovieSceneDataChanged();
+	}
 	
 	return ObjectGuid;
 }
@@ -687,7 +689,7 @@ void FSequencer::GetRuntimeObjects( TSharedRef<FMovieSceneInstance> MovieSceneIn
 	{	
 		// First, try to find spawnable puppet objects for the specified Guid
 		UObject* FoundPuppetObject = ObjectSpawner->FindPuppetObjectForSpawnableGuid( MovieSceneInstance, ObjectHandle );
-		if( FoundPuppetObject != NULL )
+		if( FoundPuppetObject != nullptr )
 		{
 			// Found a puppet object!
 			OutObjects.Reset();
@@ -696,7 +698,7 @@ void FSequencer::GetRuntimeObjects( TSharedRef<FMovieSceneInstance> MovieSceneIn
 		else
 		{
 			// No puppets were found for spawnables, so now we'll check for possessed actors
-			if( PlayMovieSceneNode != NULL )
+			if( PlayMovieSceneNode != nullptr )
 			{
 				// @todo Sequencer SubMovieScene: Figure out how to instance possessables
 				OutObjects = PlayMovieSceneNode->FindBoundObjects( ObjectHandle );
@@ -1015,7 +1017,7 @@ FGuid FSequencer::AddSpawnableForAssetOrClass( UObject* Object, UObject* Counter
 {
 	FGuid NewSpawnableGuid;
 	
-	if( ObjectBindingManager->AllowsSpawnableObjects() )
+	if (ObjectBindingManager->AllowsSpawnableObjects())
 	{
 		// Grab the MovieScene that is currently focused.  We'll add our Blueprint as an inner of the
 		// MovieScene asset.
@@ -1036,10 +1038,10 @@ FGuid FSequencer::AddSpawnableForAssetOrClass( UObject* Object, UObject* Counter
 		const FString NewSpawnableName = AssetName.ToString();		// @todo sequencer: Need UI to allow user to rename these slots
 
 		// Create our new blueprint!
-		UBlueprint* NewBlueprint = NULL;
+		UBlueprint* NewBlueprint = nullptr;
 		{
 			// @todo sequencer: Add support for forcing specific factories for an asset
-			UActorFactory* FactoryToUse = NULL;
+			UActorFactory* FactoryToUse = nullptr;
 			if( bIsActorClass )
 			{
 				// Placing an actor class directly::
@@ -1051,7 +1053,7 @@ FGuid FSequencer::AddSpawnableForAssetOrClass( UObject* Object, UObject* Counter
 				FactoryToUse = FActorFactoryAssetProxy::GetFactoryForAssetObject( Object );
 			}
 
-			if( FactoryToUse != NULL )
+			if( FactoryToUse != nullptr )
 			{
 				// Create the blueprint
 				NewBlueprint = FactoryToUse->CreateBlueprint( Object, OwnerMovieScene, BlueprintName );
@@ -1063,9 +1065,9 @@ FGuid FSequencer::AddSpawnableForAssetOrClass( UObject* Object, UObject* Counter
 			}
 		}
 
-		if( ensure( NewBlueprint != NULL ) )
+		if (ensure(NewBlueprint != nullptr))
 		{
-			if( NewBlueprint->GeneratedClass != NULL && FBlueprintEditorUtils::IsActorBased( NewBlueprint ) )
+			if ((NewBlueprint->GeneratedClass != nullptr) && FBlueprintEditorUtils::IsActorBased(NewBlueprint))
 			{
 				AActor* ActorCDO = CastChecked< AActor >( NewBlueprint->GeneratedClass->ClassDefaultObject );
 
@@ -1073,7 +1075,7 @@ FGuid FSequencer::AddSpawnableForAssetOrClass( UObject* Object, UObject* Counter
 				// @todo sequencer livecapture: This isn't really good enough to handle complex actors.  The dynamically-spawned actor could have components that
 				// were created in its construction script or via straight-up C++ code.  Instead what we should probably do is duplicate the PIE actor and generate
 				// our CDO from that duplicate.  It could get pretty complex though.
-				if( CounterpartGamePreviewObject != NULL )
+				if (CounterpartGamePreviewObject != nullptr)
 				{
 					AActor* CounterpartGamePreviewActor = CastChecked< AActor >( CounterpartGamePreviewObject );
 					CopyActorProperties( CounterpartGamePreviewActor, ActorCDO );
@@ -1207,7 +1209,7 @@ void FSequencer::OnRequestNodeDeleted( TSharedRef<const FSequencerDisplayNode>&
 void FSequencer::PlaceActorInFrontOfCamera( AActor* ActorCDO )
 {
 	// Place the actor in front of the active perspective camera if we have one
-	if( GCurrentLevelEditingViewportClient != NULL && GCurrentLevelEditingViewportClient->IsPerspective() )
+	if ((GCurrentLevelEditingViewportClient != nullptr) && GCurrentLevelEditingViewportClient->IsPerspective())
 	{
 		// Don't allow this when the active viewport is showing a simulation/PIE level
 		const bool bIsViewportShowingPIEWorld = ( GCurrentLevelEditingViewportClient->GetWorld()->GetOutermost()->PackageFlags & PKG_PlayInEditor ) != 0;
diff --git a/Engine/Source/Editor/Sequencer/Private/Sequencer.h b/Engine/Source/Editor/Sequencer/Private/Sequencer.h
--- a/Engine/Source/Editor/Sequencer/Private/Sequencer.h
+++ b/Engine/Source/Editor/Sequencer/Private/Sequencer.h
@@ -6,11 +6,13 @@
 #include "SelectedKey.h"
 #include "EditorUndoClient.h"
 
+
 class UMovieScene;
 class IToolkitHost;
 class ISequencerObjectBindingManager;
 class ISequencerObjectChangeListener;
 
+
 /** Parameters for initializing a Sequencer */
 struct FSequencerInitParams
 {
@@ -33,13 +35,18 @@ struct FSequencerInitParams
 	bool bEditWithinLevelEditor;
 };
 
+
 /**
- * Sequencer is the editing tool for MovieScene assets
+ * Sequencer is the editing tool for MovieScene assets.
  */
-class FSequencer : public ISequencer, public FGCObject, public FEditorUndoClient, public FTickableEditorObject
-{ 
+class FSequencer
+	: public ISequencer
+	, public FGCObject
+	, public FEditorUndoClient
+	, public FTickableEditorObject
+{
+public:
 
-public:	
 	static bool IsSequencerEnabled();
 
 	/**
@@ -51,9 +58,7 @@ public:
 	 */
 	void InitSequencer( const FSequencerInitParams& InitParams, const TArray<FOnCreateTrackEditor>& TrackEditorDelegates );
 
-	/** 
-	* Closes the sequencer
-	*/
+	/** Close the sequencer. */
 	void OnClose();
 
 	/** Constructor */
@@ -62,16 +67,24 @@ public:
 	/** Destructor */
 	virtual ~FSequencer();
 
-	/** FGCObject interface */
+public:
+
+	// FGCObject interface
+
 	virtual void AddReferencedObjects(FReferenceCollector& Collector) override;
 
+public:
+
 	// FTickableEditorObject interface
+
 	virtual void Tick(float DeltaTime) override;
 	virtual bool IsTickable() const override { return true; }
 	virtual TStatId GetStatId() const override { RETURN_QUICK_DECLARE_CYCLE_STAT(FSequencer, STATGROUP_Tickables); };
-	// End FTickableEditorObject interface
 
-	/** ISequencer interface */
+public:
+
+	// ISequencer interface
+
 	virtual TSharedRef<SWidget> GetSequencerWidget() const override { return SequencerWidget.ToSharedRef(); }
 	virtual UMovieScene* GetRootMovieScene() const override;
 	virtual UMovieScene* GetFocusedMovieScene() const override;
@@ -104,6 +117,8 @@ public:
 	/** @return The current view range */
 	virtual FAnimatedRange GetViewRange() const override;
 
+public:
+
 	/** Access the user-supplied settings object */
 	USequencerSettings* GetSettings() const { return Settings; }
 
@@ -253,6 +268,7 @@ public:
 	FReply OnStepToEnd();
 	FReply OnStepToBeginning();
 	FReply OnToggleLooping();
+
 	bool IsLooping() const;
 	EPlaybackMode::Type GetPlaybackMode() const;
 
@@ -264,12 +280,11 @@ public:
 
 	/** Called to save the current movie scene */
 	void SaveCurrentMovieScene();
+
 protected:
-	/**
-	 * Reset data about a movie scene when pushing or popping a movie scene
-	 */
-	void ResetPerMovieSceneData();
 
+	/** Reset data about a movie scene when pushing or popping a movie scene. */
+	void ResetPerMovieSceneData();
 
 	/**
 	 * Destroys spawnables for all movie scenes in the stack
@@ -279,7 +294,6 @@ protected:
 	/** Sets the actor CDO such that it is placed in front of the active perspective viewport camera, if we have one */
 	static void PlaceActorInFrontOfCamera( AActor* ActorCDO );
 
-
 	/**
 	 * Gets the far time boundaries of the currently edited movie scene
 	 * If the scene has shots, it only takes the shot section boundaries
@@ -324,7 +338,16 @@ protected:
 
 	/** Called to toggle auto-scroll on and off */
 	void OnToggleAutoScroll();
-	bool GetAutoScrollEnabled() const { return bAutoScrollEnabled; }
+
+	/**
+	 * Whether auto-scroll is enabled.
+	 *
+	 * @return true if auto-scroll is enabled, false otherwise.
+	 */
+	bool GetAutoScrollEnabled() const
+	{
+		return bAutoScrollEnabled;
+	}
 
 	/** Called via UEditorEngine::GetActorRecordingStateEvent to check to see whether we need to record actor state */
 	void GetActorRecordingState( bool& bIsRecording /* In+Out */ ) const;
@@ -372,6 +395,7 @@ protected:
 	void UpdatePreviewLevelViewportClientFromCameraCut( FLevelEditorViewportClient& InViewportClient, UObject* InCameraObject, bool bNewCameraCut ) const;
 
 private:
+
 	/** User-supplied settings object for this sequencer */
 	USequencerSettings* Settings;
 
@@ -416,18 +440,23 @@ private:
 
 	/** The time range target to be viewed */
 	TRange<float> TargetViewRange;
+
 	/** The last time range that was viewed */
 	TRange<float> LastViewRange;
+
 	/** The view range before zooming */
 	TRange<float> ViewRangeBeforeZoom;
+
 	/** The amount of autoscroll pan offset that is currently being applied */
 	TOptional<float> AutoscrollOffset;
+
 	/** Whether or not we are allowing autoscroll */
 	bool bAutoScrollEnabled;
 
 	/** Zoom smoothing curves */
 	FCurveSequence ZoomAnimation;
 	FCurveHandle ZoomCurve;
+
 	/** Overlay fading curves */
 	FCurveSequence OverlayAnimation;
 	FCurveHandle OverlayCurve;
diff --git a/Engine/Source/Editor/Sequencer/Private/SequencerAssetEditor.cpp b/Engine/Source/Editor/Sequencer/Private/SequencerAssetEditor.cpp
--- a/Engine/Source/Editor/Sequencer/Private/SequencerAssetEditor.cpp
+++ b/Engine/Source/Editor/Sequencer/Private/SequencerAssetEditor.cpp
@@ -14,14 +14,18 @@
 #include "SequencerObjectChangeListener.h"
 #include "SDockTab.h"
 
+
 #define LOCTEXT_NAMESPACE "Sequencer"
 
 const FName FSequencerAssetEditor::SequencerMainTabId( TEXT( "Sequencer_SequencerMain" ) );
 
+
 namespace SequencerDefs
 {
 	static const FName SequencerAppIdentifier( TEXT( "SequencerApp" ) );
 }
+
+
 void FSequencerAssetEditor::RegisterTabSpawners(const TSharedRef<class FTabManager>& TabManager)
 {
 	if( FSequencer::IsSequencerEnabled() && !IsWorldCentricAssetEditor() )
@@ -35,6 +39,7 @@ void FSequencerAssetEditor::RegisterTabSpawners(const TSharedRef<class FTabManag
 
 }
 
+
 void FSequencerAssetEditor::UnregisterTabSpawners(const TSharedRef<class FTabManager>& TabManager)
 {
 	if( FSequencer::IsSequencerEnabled() && !IsWorldCentricAssetEditor() )
@@ -47,19 +52,20 @@ void FSequencerAssetEditor::UnregisterTabSpawners(const TSharedRef<class FTabMan
 	LevelEditorModule.AttachSequencer( SNullWidget::NullWidget, nullptr );
 }
 
+
 void FSequencerAssetEditor::InitSequencerAssetEditor( const EToolkitMode::Type Mode, const FSequencerViewParams& InViewParams, const TSharedPtr< class IToolkitHost >& InitToolkitHost, UMovieScene* InRootMovieScene, const TArray<FOnCreateTrackEditor>& TrackEditorDelegates, bool bEditWithinLevelEditor )
 {
 	{
 		const TSharedRef<FTabManager::FLayout> StandaloneDefaultLayout = FTabManager::NewLayout("Standalone_Sequencer_Layout")
-		->AddArea
-		(
-			FTabManager::NewPrimaryArea()
-			->Split
+			->AddArea
 			(
-				FTabManager::NewStack()
-				->AddTab(SequencerMainTabId, ETabState::OpenedTab)
-			)
-		);
+				FTabManager::NewPrimaryArea()
+					->Split
+					(
+						FTabManager::NewStack()
+						->AddTab(SequencerMainTabId, ETabState::OpenedTab)
+					)
+			);
 
 		const bool bCreateDefaultStandaloneMenu = true;
 		const bool bCreateDefaultToolbar = false;
@@ -70,20 +76,18 @@ void FSequencerAssetEditor::InitSequencerAssetEditor( const EToolkitMode::Type M
 	Sequencer = MakeShareable(new FSequencer);
 	
 	FSequencerInitParams SequencerInitParams;
-	SequencerInitParams.ViewParams = InViewParams;
-	SequencerInitParams.ObjectChangeListener = MakeShareable( new FSequencerObjectChangeListener( Sequencer.ToSharedRef(), bEditWithinLevelEditor ) );
-	
-	SequencerInitParams.ObjectBindingManager = MakeShareable( new FSequencerActorBindingManager( SequencerInitParams.ObjectChangeListener.ToSharedRef(), Sequencer.ToSharedRef() ) );
+	{
+		SequencerInitParams.ViewParams = InViewParams;
+		SequencerInitParams.ObjectChangeListener = MakeShareable(new FSequencerObjectChangeListener(Sequencer.ToSharedRef(), bEditWithinLevelEditor));
+		SequencerInitParams.ObjectBindingManager = MakeShareable(new FSequencerActorBindingManager(SequencerInitParams.ObjectChangeListener.ToSharedRef(), Sequencer.ToSharedRef()));
+		SequencerInitParams.RootMovieScene = InRootMovieScene;
+		SequencerInitParams.bEditWithinLevelEditor = bEditWithinLevelEditor;
+		SequencerInitParams.ToolkitHost = InitToolkitHost;
+	}
 
-	SequencerInitParams.RootMovieScene = InRootMovieScene;
-	
-	SequencerInitParams.bEditWithinLevelEditor = bEditWithinLevelEditor;
-	
-	SequencerInitParams.ToolkitHost = InitToolkitHost;
-	
 	Sequencer->InitSequencer( SequencerInitParams, TrackEditorDelegates );
 
-	if(bEditWithinLevelEditor)
+	if (bEditWithinLevelEditor)
 	{
 		// @todo remove when world-centric mode is added
 		FLevelEditorModule& LevelEditorModule = FModuleManager::LoadModuleChecked<FLevelEditorModule>("LevelEditor");
@@ -96,16 +100,19 @@ void FSequencerAssetEditor::InitSequencerAssetEditor( const EToolkitMode::Type M
 	}
 }
 
+
 TSharedRef<ISequencer> FSequencerAssetEditor::GetSequencerInterface() const
 { 
 	return Sequencer.ToSharedRef(); 
 }
 
+
 FSequencerAssetEditor::FSequencerAssetEditor()
 {
 
 }
 
+
 FSequencerAssetEditor::~FSequencerAssetEditor()
 {
 	Sequencer->OnClose();
@@ -139,6 +146,7 @@ FName FSequencerAssetEditor::GetToolkitFName() const
 	return SequencerName;
 }
 
+
 FText FSequencerAssetEditor::GetBaseToolkitName() const
 {
 	return LOCTEXT("AppLabel", "Sequencer");
diff --git a/Engine/Source/Editor/Sequencer/Private/SequencerModule.cpp b/Engine/Source/Editor/Sequencer/Private/SequencerModule.cpp
--- a/Engine/Source/Editor/Sequencer/Private/SequencerModule.cpp
+++ b/Engine/Source/Editor/Sequencer/Private/SequencerModule.cpp
@@ -9,6 +9,7 @@
 #include "SequencerEdMode.h"
 #include "SequencerObjectChangeListener.h"
 
+
 // We disable the deprecation warnings here because otherwise it'll complain about us
 // implementing RegisterTrackEditor and UnRegisterTrackEditor.  We know
 // that, but we only want it to complain if *others* implement or call these functions.
@@ -16,37 +17,39 @@
 // These macros should be removed when those functions are removed.
 PRAGMA_DISABLE_DEPRECATION_WARNINGS
 
+
 /**
  * SequencerModule implementation (private)
  */
-class FSequencerModule : public ISequencerModule
+class FSequencerModule
+	: public ISequencerModule
 {
+public:
+
+	// ISequencerModule interface
 
-	/** ISequencerModule interface */
 	virtual TSharedPtr<ISequencer> CreateSequencer( UMovieScene* RootMovieScene, const FSequencerViewParams& InViewParams, TSharedRef<ISequencerObjectBindingManager> ObjectBindingManager ) override
 	{
 		TSharedRef< FSequencer > Sequencer = MakeShareable(new FSequencer);
 		
 		FSequencerInitParams SequencerInitParams;
-		SequencerInitParams.ViewParams = InViewParams;
-		SequencerInitParams.ObjectChangeListener = MakeShareable( new FSequencerObjectChangeListener( Sequencer, false ) );
-		
-		SequencerInitParams.ObjectBindingManager = ObjectBindingManager;
-		
-		SequencerInitParams.RootMovieScene = RootMovieScene;
-		
-		SequencerInitParams.bEditWithinLevelEditor = false;
-		
-		SequencerInitParams.ToolkitHost = nullptr;
+		{
+			SequencerInitParams.ViewParams = InViewParams;
+			SequencerInitParams.ObjectChangeListener = MakeShareable( new FSequencerObjectChangeListener( Sequencer, false ) );
+			SequencerInitParams.ObjectBindingManager = ObjectBindingManager;
+			SequencerInitParams.RootMovieScene = RootMovieScene;
+			SequencerInitParams.bEditWithinLevelEditor = false;
+			SequencerInitParams.ToolkitHost = nullptr;
+		}
 		
-		Sequencer->InitSequencer( SequencerInitParams, TrackEditorDelegates );
+		Sequencer->InitSequencer(SequencerInitParams, TrackEditorDelegates);
 		
 		return Sequencer;
 	}
 	
 	virtual TSharedPtr<ISequencer> CreateSequencerAssetEditor( const EToolkitMode::Type Mode, const FSequencerViewParams& InViewParams, const TSharedPtr< class IToolkitHost >& InitToolkitHost, UMovieScene* InRootMovieScene, bool bEditWithinLevelEditor ) override
 	{
-		TSharedRef< FSequencerAssetEditor > SequencerAssetEditor = MakeShareable(new FSequencerAssetEditor);
+		TSharedRef<FSequencerAssetEditor> SequencerAssetEditor = MakeShareable(new FSequencerAssetEditor);
 
 		SequencerAssetEditor->InitSequencerAssetEditor( Mode, InViewParams, InitToolkitHost, InRootMovieScene, TrackEditorDelegates, bEditWithinLevelEditor );
 		return SequencerAssetEditor->GetSequencerInterface();
@@ -86,12 +89,14 @@ class FSequencerModule : public ISequencerModule
 			FEditorModeRegistry::Get().UnregisterMode(FSequencerEdMode::EM_SequencerMode);
 		}
 	}
+
 private:
+
 	/** List of auto-key handler delegates sequencers will execute when they are created */
 	TArray< FOnCreateTrackEditor > TrackEditorDelegates;
 };
 
-PRAGMA_ENABLE_DEPRECATION_WARNINGS
 
+PRAGMA_ENABLE_DEPRECATION_WARNINGS
 
-IMPLEMENT_MODULE( FSequencerModule, Sequencer );
+IMPLEMENT_MODULE(FSequencerModule, Sequencer);
diff --git a/Engine/Source/Editor/UMGEditor/Private/WidgetBlueprintEditor.cpp b/Engine/Source/Editor/UMGEditor/Private/WidgetBlueprintEditor.cpp
--- a/Engine/Source/Editor/UMGEditor/Private/WidgetBlueprintEditor.cpp
+++ b/Engine/Source/Editor/UMGEditor/Private/WidgetBlueprintEditor.cpp
@@ -767,19 +767,23 @@ void GetBindableObjects(UWidget* RootWidget, TArray<FObjectAndDisplayName>& Bind
 TSharedRef<SWidget> FWidgetBlueprintEditor::OnGetAnimationAddMenuContent(TSharedRef<ISequencer> InSequencer)
 {
 	TArray<FObjectAndDisplayName> BindableObjects;
-	GetBindableObjects(GetPreview()->GetRootWidget(), BindableObjects);
-	BindableObjects.Sort();
+	{
+		GetBindableObjects(GetPreview()->GetRootWidget(), BindableObjects);
+		BindableObjects.Sort();
+	}
+
+	FMenuBuilder AddMenuBuilder(true, nullptr);
 
-	FMenuBuilder AddMenuBuilder(true, NULL);
 	for (FObjectAndDisplayName& BindableObject : BindableObjects)
 	{
 		FGuid BoundObjectGuid = SequencerObjectBindingManager->FindGuidForObject(*InSequencer->GetFocusedMovieScene(), *BindableObject.Object);
-		if (BoundObjectGuid.IsValid() == false)
+		if (!BoundObjectGuid.IsValid())
 		{
 			FUIAction AddMenuAction(FExecuteAction::CreateSP(this, &FWidgetBlueprintEditor::AddObjectToAnimation, BindableObject.Object));
 			AddMenuBuilder.AddMenuEntry(BindableObject.DisplayName, FText(), FSlateIcon(), AddMenuAction);
 		}
 	}
+
 	return AddMenuBuilder.MakeWidget();
 }
 
diff --git a/Engine/Source/Runtime/MovieScene/Private/MovieScene.cpp b/Engine/Source/Runtime/MovieScene/Private/MovieScene.cpp
--- a/Engine/Source/Runtime/MovieScene/Private/MovieScene.cpp
+++ b/Engine/Source/Runtime/MovieScene/Private/MovieScene.cpp
@@ -4,11 +4,8 @@
 #include "MovieScene.h"
 
 
-
-UMovieScene::UMovieScene( const FObjectInitializer& ObjectInitializer )
-	: Super( ObjectInitializer )
-{ }
-
+/* UMovieScene interface
+ *****************************************************************************/
 
 #if WITH_EDITOR
 
@@ -227,20 +224,6 @@ TArray<UMovieSceneSection*> UMovieScene::GetAllSections() const
 }
 
 
-void UMovieScene::RemoveBinding( const FGuid& Guid )
-{
-	// Update each type
-	for( int32 BindingIndex = 0; BindingIndex < ObjectBindings.Num(); ++BindingIndex )
-	{
-		if( ObjectBindings[BindingIndex].GetObjectGuid() == Guid )
-		{
-			ObjectBindings.RemoveAt( BindingIndex );
-			break;
-		}
-	}
-}
-
-
 UMovieSceneTrack* UMovieScene::FindTrack( TSubclassOf<UMovieSceneTrack> TrackClass, const FGuid& ObjectGuid, FName UniqueTrackName ) const
 {
 	UMovieSceneTrack* FoundTrack = nullptr;
@@ -394,3 +377,20 @@ bool UMovieScene::IsAMasterTrack(const UMovieSceneTrack* Track) const
 
 	return false;
 }
+
+
+/* UMovieScene implementation
+ *****************************************************************************/
+
+void UMovieScene::RemoveBinding(const FGuid& Guid)
+{
+	// update each type
+	for (int32 BindingIndex = 0; BindingIndex < ObjectBindings.Num(); ++BindingIndex)
+	{
+		if (ObjectBindings[BindingIndex].GetObjectGuid() == Guid)
+		{
+			ObjectBindings.RemoveAt(BindingIndex);
+			break;
+		}
+	}
+}
diff --git a/Engine/Source/Runtime/MovieScene/Private/RuntimeMovieScenePlayer.cpp b/Engine/Source/Runtime/MovieScene/Private/RuntimeMovieScenePlayer.cpp
--- a/Engine/Source/Runtime/MovieScene/Private/RuntimeMovieScenePlayer.cpp
+++ b/Engine/Source/Runtime/MovieScene/Private/RuntimeMovieScenePlayer.cpp
@@ -1,17 +1,18 @@
 // Copyright 1998-2015 Epic Games, Inc. All Rights Reserved.
 
 #include "MovieScenePrivatePCH.h"
-#include "RuntimeMovieScenePlayer.h"
+#include "MovieSceneBindings.h"
 #include "MovieSceneInstance.h"
+#include "RuntimeMovieScenePlayer.h"
+
 
+/* URuntimeMovieScenePlayer static functions
+ *****************************************************************************/
 
-URuntimeMovieScenePlayer* URuntimeMovieScenePlayer::CreateRuntimeMovieScenePlayer( ULevel* Level, UMovieSceneBindings* MovieSceneBindings )
+URuntimeMovieScenePlayer* URuntimeMovieScenePlayer::CreateRuntimeMovieScenePlayer(ULevel* Level, UMovieSceneBindings* MovieSceneBindings)
 {	
-	URuntimeMovieScenePlayer* NewRuntimeMovieScenePlayer = NewObject<URuntimeMovieScenePlayer>(
-		(UObject*)GetTransientPackage(), 
-		NAME_None, 
-		RF_Transient );
-	check( NewRuntimeMovieScenePlayer != nullptr );
+	URuntimeMovieScenePlayer* NewRuntimeMovieScenePlayer = NewObject<URuntimeMovieScenePlayer>((UObject*)GetTransientPackage(), NAME_None, RF_Transient);
+	check(NewRuntimeMovieScenePlayer != nullptr);
 
 	// Associate the player with its world
 	NewRuntimeMovieScenePlayer->World = Level->OwningWorld;
@@ -30,17 +31,55 @@ URuntimeMovieScenePlayer* URuntimeMovieScenePlayer::CreateRuntimeMovieScenePlaye
 }
 
 
-URuntimeMovieScenePlayer::URuntimeMovieScenePlayer( const FObjectInitializer& ObjectInitializer )
-	: Super( ObjectInitializer )
+/* URuntimeMovieScenePlayer structors
+ *****************************************************************************/
+
+URuntimeMovieScenePlayer::URuntimeMovieScenePlayer(const FObjectInitializer& ObjectInitializer)
+	: Super(ObjectInitializer)
+	, bIsPlaying(false)
+	, MovieSceneBindings(nullptr)
+	, TimeCursorPosition(0.0f)
+{ }
+
+
+/* URuntimeMovieScenePlayer interface
+ *****************************************************************************/
+
+bool URuntimeMovieScenePlayer::IsPlaying() const
 {
+	return bIsPlaying;
+}
+
 
-	MovieSceneBindings = nullptr;
-	TimeCursorPosition = 0.0f;
+void URuntimeMovieScenePlayer::Pause()
+{
 	bIsPlaying = false;
 }
 
 
-void URuntimeMovieScenePlayer::SpawnActorsForMovie( TSharedRef<FMovieSceneInstance> MovieSceneInstance  )
+void URuntimeMovieScenePlayer::Play()
+{
+	// Init the root movie scene instance
+	if (MovieSceneBindings != nullptr)
+	{
+		UMovieScene* MovieScene = MovieSceneBindings->GetRootMovieScene();
+	
+		// @todo Sequencer playback: Should we recreate the instance every time?
+		RootMovieSceneInstance = MakeShareable(new FMovieSceneInstance(*MovieScene));
+
+		// @odo Sequencer Should we spawn actors here?
+		SpawnActorsForMovie(RootMovieSceneInstance.ToSharedRef());
+		RootMovieSceneInstance->RefreshInstance(*this);
+	}
+	
+	bIsPlaying = true;
+}
+
+
+/* IMovieScenePlayer interface
+ *****************************************************************************/
+
+void URuntimeMovieScenePlayer::SpawnActorsForMovie(TSharedRef<FMovieSceneInstance> MovieSceneInstance)
 {
 	UWorld* WorldPtr = World.Get();
 	if( WorldPtr != nullptr && MovieSceneBindings != nullptr )
@@ -92,7 +131,7 @@ void URuntimeMovieScenePlayer::SpawnActorsForMovie( TSharedRef<FMovieSceneInstan
 }
 
 
-void URuntimeMovieScenePlayer::DestroyActorsForMovie( TSharedRef<FMovieSceneInstance> MovieSceneInstance )
+void URuntimeMovieScenePlayer::DestroyActorsForMovie(TSharedRef<FMovieSceneInstance> MovieSceneInstance)
 {
 	UWorld* WorldPtr = World.Get();
 	if( WorldPtr != nullptr && MovieSceneBindings != nullptr )
@@ -120,62 +159,13 @@ void URuntimeMovieScenePlayer::DestroyActorsForMovie( TSharedRef<FMovieSceneInst
 }
 
 
-bool URuntimeMovieScenePlayer::IsPlaying() const
-{
-	return bIsPlaying;
-}
-
-
-void URuntimeMovieScenePlayer::Play()
-{
-	// Init the root movie scene instance
-	if( MovieSceneBindings != nullptr )
-	{
-		UMovieScene* MovieScene = MovieSceneBindings->GetRootMovieScene();
-	
-		// @todo Sequencer playback: Should we recreate the instance every time?
-		RootMovieSceneInstance = MakeShareable( new FMovieSceneInstance( *MovieScene ) );
-
-		// @odo Sequencer Should we spawn actors here?
-		SpawnActorsForMovie( RootMovieSceneInstance.ToSharedRef() );
-
-		RootMovieSceneInstance->RefreshInstance( *this );
-	}
-	
-	bIsPlaying = true;
-}
-
-
-void URuntimeMovieScenePlayer::Pause()
-{
-	bIsPlaying = false;
-}
-
-
 void URuntimeMovieScenePlayer::SetMovieSceneBindings( UMovieSceneBindings* NewMovieSceneBindings )
 {
 	MovieSceneBindings = NewMovieSceneBindings;
 }
 
-
-void URuntimeMovieScenePlayer::Tick( const float DeltaSeconds )
-{
-	float LastTimePosition = TimeCursorPosition;
-
-	if( bIsPlaying )
-	{
-		TimeCursorPosition += DeltaSeconds;
-	}
-
-	// Update the movie scene!
-	if( MovieSceneBindings != nullptr && RootMovieSceneInstance.IsValid() )
-	{
-		RootMovieSceneInstance->Update( TimeCursorPosition, LastTimePosition, *this );
-	}
-}
-
 	
-void URuntimeMovieScenePlayer::GetRuntimeObjects( TSharedRef<FMovieSceneInstance> MovieSceneInstance, const FGuid& ObjectHandle, TArray< UObject* >& OutObjects ) const
+void URuntimeMovieScenePlayer::GetRuntimeObjects(TSharedRef<FMovieSceneInstance> MovieSceneInstance, const FGuid& ObjectHandle, TArray<UObject*>& OutObjects) const
 {
 	// @todo sequencer runtime: Add support for individually spawning actors on demand when first requested?
 	//    This may be important to reduce the up-front hitch when spawning actors for the entire movie, but
@@ -237,14 +227,39 @@ EMovieScenePlayerStatus::Type URuntimeMovieScenePlayer::GetPlaybackStatus() cons
 }
 
 
-void URuntimeMovieScenePlayer::AddOrUpdateMovieSceneInstance( UMovieSceneSection& MovieSceneSection, TSharedRef<FMovieSceneInstance> InstanceToAdd )
+void URuntimeMovieScenePlayer::AddOrUpdateMovieSceneInstance(UMovieSceneSection& MovieSceneSection, TSharedRef<FMovieSceneInstance> InstanceToAdd)
 {
 	SpawnActorsForMovie( InstanceToAdd );
 }
 
 
-void URuntimeMovieScenePlayer::RemoveMovieSceneInstance( UMovieSceneSection& MovieSceneSection, TSharedRef<FMovieSceneInstance> InstanceToRemove )
+void URuntimeMovieScenePlayer::RemoveMovieSceneInstance(UMovieSceneSection& MovieSceneSection, TSharedRef<FMovieSceneInstance> InstanceToRemove)
 {
 	const bool bDestroyAll = true;
 	DestroyActorsForMovie( InstanceToRemove );
 }
+
+
+TSharedRef<FMovieSceneInstance> URuntimeMovieScenePlayer::GetRootMovieSceneInstance() const
+{
+	return RootMovieSceneInstance.ToSharedRef();
+}
+
+
+/* IRuntimeMovieScenePlayerInterface interface
+ *****************************************************************************/
+
+void URuntimeMovieScenePlayer::Tick(const float DeltaSeconds)
+{
+	float LastTimePosition = TimeCursorPosition;
+
+	if (bIsPlaying)
+	{
+		TimeCursorPosition += DeltaSeconds;
+	}
+
+	if(RootMovieSceneInstance.IsValid())
+	{
+		RootMovieSceneInstance->Update(TimeCursorPosition, LastTimePosition, *this);
+	}
+}
diff --git a/Engine/Source/Runtime/MovieScene/Public/MovieScene.h b/Engine/Source/Runtime/MovieScene/Public/MovieScene.h
--- a/Engine/Source/Runtime/MovieScene/Public/MovieScene.h
+++ b/Engine/Source/Runtime/MovieScene/Public/MovieScene.h
@@ -38,7 +38,7 @@ UCLASS()
 class MOVIESCENE_API UMovieScene
 	: public UObject
 {
-	GENERATED_UCLASS_BODY()
+	GENERATED_BODY()
 
 public:
 
diff --git a/Engine/Source/Runtime/MovieScene/Public/RuntimeMovieScenePlayer.h b/Engine/Source/Runtime/MovieScene/Public/RuntimeMovieScenePlayer.h
--- a/Engine/Source/Runtime/MovieScene/Public/RuntimeMovieScenePlayer.h
+++ b/Engine/Source/Runtime/MovieScene/Public/RuntimeMovieScenePlayer.h
@@ -4,18 +4,26 @@
 
 #include "IMovieScenePlayer.h"
 #include "MovieScene.h"
-#include "MovieSceneBindings.h"
 #include "Runtime/Engine/Classes/MovieScene/RuntimeMovieScenePlayerInterface.h"
-
 #include "RuntimeMovieScenePlayer.generated.h"
 
 
+class FMovieSceneInstance;
+class ULevel;
+class UMovieSceneBindings;
+
+
 /**
- * RuntimeMovieScenePlayer is used to actually "play" a MovieScene asset at runtime.  It keeps track of playback
- * state and provides functions for manipulating a MovieScene while its playing.
+ * RuntimeMovieScenePlayer is used to actually "play" a MovieScene asset at runtime.
+ *
+ * This class keeps track of playback state and provides functions for manipulating
+ * a MovieScene while its playing.
  */
-UCLASS( MinimalAPI )
-class URuntimeMovieScenePlayer : public UObject, public IMovieScenePlayer, public IRuntimeMovieScenePlayerInterface
+UCLASS(MinimalAPI)
+class URuntimeMovieScenePlayer
+	: public UObject
+	, public IMovieScenePlayer
+	, public IRuntimeMovieScenePlayerInterface
 {
 	GENERATED_UCLASS_BODY()
 
@@ -25,59 +33,55 @@ public:
 	 * Static:  Creates a RuntimeMovieScenePlayer.  This function is designed to be called through the expanded PlayMovieScene
 	 * Kismet node, and not called directly through visual script.
 	 *
-	 * @param	Level	The level to play this MovieScene in.  This level will "own" the instance of the returned RuntimeMovieScenePlayer.
-	 * @param	MovieSceneBindings	The MovieScene to play, along with the bindings to any actors in the level that may be possessed during playback.
-	 *
-	 * @return	Returns the newly created RuntimeMovieScenePlayer instance for this MovieScene
+	 * @param Level The level to play this MovieScene in.  This level will "own" the instance of the returned RuntimeMovieScenePlayer.
+	 * @param MovieSceneBindings The MovieScene to play, along with the bindings to any actors in the level that may be possessed during playback.
+	 * @return The newly created RuntimeMovieScenePlayer instance for this MovieScene
 	 */
 	UFUNCTION(BlueprintCallable, Category="Game|Cinematic")
-	static URuntimeMovieScenePlayer* CreateRuntimeMovieScenePlayer( class ULevel* Level, class UMovieSceneBindings* InMovieSceneBindings );
+	static URuntimeMovieScenePlayer* CreateRuntimeMovieScenePlayer(ULevel* Level, UMovieSceneBindings* InMovieSceneBindings);
 
-	/**
-	 * Starts playback from the current time cursor position.  Called by the expanded PlayMovieScene node.
-	 */
+public:
+
+	/** Start playback from the current time cursor position. */
 	UFUNCTION(BlueprintCallable, Category="Game|Cinematic")
 	void Play();
 
-	/**
-	 * Pauses playback.  Called by the expanded PlayMovieScene node.
-	 */
+	/** Pause playback. */
 	UFUNCTION(BlueprintCallable, Category="Game|Cinematic")
 	void Pause();
 
-	/** @return Returns true if the sequence is actively playing, or false if paused or if Play() was never called */
+	/** Check whether the sequence is actively playing. */
 	bool IsPlaying() const;
 
-
 protected:
 
 	/** Sets the bindings for this RuntimeMovieScenePlayer.  Should be called right after construction. */
-	void SetMovieSceneBindings( UMovieSceneBindings* NewMovieSceneBindings );
+	void SetMovieSceneBindings(UMovieSceneBindings* NewMovieSceneBindings);
 
-	/** IMovieScenePlayer interface */
-	virtual void SpawnActorsForMovie( TSharedRef<FMovieSceneInstance> MovieSceneInstance );
-	virtual void DestroyActorsForMovie( TSharedRef<FMovieSceneInstance> MovieSceneInstance );
-	virtual void GetRuntimeObjects( TSharedRef<FMovieSceneInstance> MovieSceneInstance, const FGuid& ObjectHandle, TArray< UObject* >& OutObjects ) const override;
+protected:
+
+	// IMovieScenePlayer interface
+
+	virtual void SpawnActorsForMovie(TSharedRef<FMovieSceneInstance> MovieSceneInstance);
+	virtual void DestroyActorsForMovie(TSharedRef<FMovieSceneInstance> MovieSceneInstance);
+	virtual void GetRuntimeObjects(TSharedRef<FMovieSceneInstance> MovieSceneInstance, const FGuid& ObjectHandle, TArray<UObject*>& OutObjects) const override;
 	virtual void UpdateCameraCut(UObject* ObjectToViewThrough, bool bNewCameraCut) const override;
 	virtual EMovieScenePlayerStatus::Type GetPlaybackStatus() const override;
-	virtual void AddOrUpdateMovieSceneInstance( class UMovieSceneSection& MovieSceneSection, TSharedRef<FMovieSceneInstance> InstanceToAdd ) override;
-	virtual void RemoveMovieSceneInstance( class UMovieSceneSection& MovieSceneSection, TSharedRef<FMovieSceneInstance> InstanceToRemove ) override;
-	virtual TSharedRef<FMovieSceneInstance> GetRootMovieSceneInstance() const override { return RootMovieSceneInstance.ToSharedRef(); }
-	/** IRuntimeMovieScenePlayerInterface */
-	virtual void Tick( const float DeltaSeconds ) override;
+	virtual void AddOrUpdateMovieSceneInstance(UMovieSceneSection& MovieSceneSection, TSharedRef<FMovieSceneInstance> InstanceToAdd) override;
+	virtual void RemoveMovieSceneInstance(UMovieSceneSection& MovieSceneSection, TSharedRef<FMovieSceneInstance> InstanceToRemove) override;
+	virtual TSharedRef<FMovieSceneInstance> GetRootMovieSceneInstance() const override;
+
+protected:
+
+	// IRuntimeMovieScenePlayerInterface
 
+	virtual void Tick(const float DeltaSeconds) override;
 
 private:
-	struct FSpawnedActorInfo
-	{
-		/** Identifier that maps this actor to a movie scene spawnable */
-		FGuid RuntimeGuid;
-		/** The spawned actor */
-		TWeakObjectPtr<AActor> SpawnedActor;
-	};
 
-	/** The world this player will spawn actors in, if needed */
-	TWeakObjectPtr< UWorld > World;
+	/** Whether we're currently playing. If false, then sequence playback is paused or was never started. */
+	UPROPERTY()
+	bool bIsPlaying;
 
 	/** The MovieScene we're playing as well as bindings for this MovieScene.  May contain sub-scenes. */
 	UPROPERTY()
@@ -87,15 +91,23 @@ private:
 	UPROPERTY()
 	double TimeCursorPosition;
 
-	/** Whether we're currently playing.  If false, then sequence playback is paused or was never started. */
-	UPROPERTY()
-	bool bIsPlaying;
+private:
 
-	/** The root movie scene instance to update when playing. */
-	TSharedPtr<class FMovieSceneInstance> RootMovieSceneInstance;
+	struct FSpawnedActorInfo
+	{
+		/** Identifier that maps this actor to a movie scene spawnable */
+		FGuid RuntimeGuid;
+
+		/** The spawned actor */
+		TWeakObjectPtr<AActor> SpawnedActor;
+	};
 
 	/** Maps spawnable GUIDs to their spawned actor in the world */
-	TMap< TWeakPtr<FMovieSceneInstance>, TArray< FSpawnedActorInfo > > InstanceToSpawnedActorMap;
-};
+	TMap<TWeakPtr<FMovieSceneInstance>, TArray<FSpawnedActorInfo>> InstanceToSpawnedActorMap;
 
+	/** The root movie scene instance to update when playing. */
+	TSharedPtr<FMovieSceneInstance> RootMovieSceneInstance;
 
+	/** The world this player will spawn actors in, if needed */
+	TWeakObjectPtr<UWorld> World;
+};
diff --git a/Engine/Source/Runtime/Niagara/Niagara.Build.cs b/Engine/Source/Runtime/Niagara/Niagara.Build.cs
--- a/Engine/Source/Runtime/Niagara/Niagara.Build.cs
+++ b/Engine/Source/Runtime/Niagara/Niagara.Build.cs
@@ -21,7 +21,8 @@ public class Niagara : ModuleRules
 
         PublicDependencyModuleNames.AddRange(
             new string[] {
-                "RHI"
+                "MovieScene",
+                "RHI",
             }
         );
 
diff --git a/Engine/Source/Runtime/UMG/Private/Animation/WidgetAnimation.cpp b/Engine/Source/Runtime/UMG/Private/Animation/WidgetAnimation.cpp
--- a/Engine/Source/Runtime/UMG/Private/Animation/WidgetAnimation.cpp
+++ b/Engine/Source/Runtime/UMG/Private/Animation/WidgetAnimation.cpp
@@ -10,6 +10,7 @@
 
 UWidgetAnimation::UWidgetAnimation(const FObjectInitializer& ObjectInitializer)
 	: Super(ObjectInitializer)
+	, MovieScene(nullptr)
 { }
 
 
diff --git a/Engine/Source/Runtime/UMG/Public/Animation/WidgetAnimation.h b/Engine/Source/Runtime/UMG/Public/Animation/WidgetAnimation.h
--- a/Engine/Source/Runtime/UMG/Public/Animation/WidgetAnimation.h
+++ b/Engine/Source/Runtime/UMG/Public/Animation/WidgetAnimation.h
@@ -8,7 +8,6 @@
 
 
 class UMovieScene;
-class UWidgetTree;
 
 
 UCLASS(BlueprintType, MinimalAPI)
@@ -20,12 +19,29 @@ class UWidgetAnimation
 public:
 
 #if WITH_EDITOR
+	/**
+	 * Get a placeholder animation.
+	 *
+	 * @return Placeholder animation.
+	 */
 	static UMG_API UWidgetAnimation* GetNullAnimation();
 #endif
 
+	/**
+	 * Get the start time of this animation.
+	 *
+	 * @return Start time in seconds.
+	 * @see GetEndTime
+	 */
 	UFUNCTION(BlueprintCallable, Category="Animation")
 	UMG_API float GetStartTime() const;
 
+	/**
+	 * Get the end time of this animation.
+	 *
+	 * @return End time in seconds.
+	 * @see GetStartTime
+	 */
 	UFUNCTION(BlueprintCallable, Category="Animation")
 	UMG_API float GetEndTime() const;
 
