Sequencer: Fix undo issues in curve editor - don't reframe and don't drop selection.

diff --git a/Engine/Source/Editor/Sequencer/Private/SSequencer.cpp b/Engine/Source/Editor/Sequencer/Private/SSequencer.cpp
--- a/Engine/Source/Editor/Sequencer/Private/SSequencer.cpp
+++ b/Engine/Source/Editor/Sequencer/Private/SSequencer.cpp
@@ -704,19 +704,26 @@ void SSequencer::UpdateLayoutTree()
 		SelectedPathNames.Add(SelectedDisplayNode->GetPathName());
 	}
 
+	// Suspend broadcasting selection changes because we don't want unnecessary rebuilds.
+	Sequencer.Pin()->GetSelection().SuspendBroadcast();
+
 	// Update the node tree
 	SequencerNodeTree->Update();
 
 	FSequencerTrackLaneFactory Factory(TrackOutliner.ToSharedRef(), TrackArea.ToSharedRef(), Sequencer.Pin().ToSharedRef() );
 	Factory.Repopulate( *SequencerNodeTree );
 
-	CurveEditor->SetSequencerNodeTree(SequencerNodeTree);
-
 	// Restore the selection state.
 	RestoreSelectionState(SequencerNodeTree->GetRootNodes(), SelectedPathNames, Sequencer.Pin()->GetSelection());	// Update to actor selection.
 	OnActorSelectionChanged(NULL);
 
+	// This must come after the selection state has been restored so that the curve editor is populated with the correctly selected nodes
+	CurveEditor->SetSequencerNodeTree(SequencerNodeTree);
+
 	SequencerNodeTree->UpdateCachedVisibilityBasedOnShotFiltersChanged();
+
+	// Continue broadcasting selection changes
+	Sequencer.Pin()->GetSelection().ResumeBroadcast();
 }
 
 void SSequencer::UpdateBreadcrumbs(const TArray< TWeakObjectPtr<class UMovieSceneSection> >& FilteringShots)
diff --git a/Engine/Source/Editor/Sequencer/Private/SSequencerCurveEditor.cpp b/Engine/Source/Editor/Sequencer/Private/SSequencerCurveEditor.cpp
--- a/Engine/Source/Editor/Sequencer/Private/SSequencerCurveEditor.cpp
+++ b/Engine/Source/Editor/Sequencer/Private/SSequencerCurveEditor.cpp
@@ -110,7 +110,43 @@ void SSequencerCurveEditor::SetSequencerNodeTree( TSharedPtr<FSequencerNodeTree>
 
 void SSequencerCurveEditor::UpdateCurveOwner()
 {
-	CurveOwner = MakeShareable( new FSequencerCurveOwner( SequencerNodeTree, GetSettings()->GetCurveVisibility() ) );
+	FSequencerCurveOwner* NewCurveOwner = new FSequencerCurveOwner( SequencerNodeTree, GetSettings()->GetCurveVisibility() );
+
+	bool bAllFound = false;
+	if (NewCurveOwner != NULL && CurveOwner.IsValid())
+	{
+		TSet<FName> NewCurveOwnerCurveNames;
+		if (NewCurveOwner != NULL)
+		{
+			for (auto Curve : NewCurveOwner->GetCurves())
+			{
+				NewCurveOwnerCurveNames.Add(Curve.CurveName);
+			}
+		}
+
+		TSet<FName> CurveOwnerCurveNames;
+		if (CurveOwner.Get() != NULL)
+		{
+			if (CurveOwner.Get()->GetCurves().Num() == NewCurveOwner->GetCurves().Num())
+			{
+				bAllFound = true;
+				for (auto Curve : CurveOwner.Get()->GetCurves())
+				{
+					if (!NewCurveOwnerCurveNames.Contains(Curve.CurveName))
+					{
+						bAllFound = false;
+					}
+				}
+			}
+		}
+	}
+
+	if (bAllFound)
+	{
+		return;
+	}
+
+	CurveOwner = MakeShareable( NewCurveOwner );
 	SetCurveOwner( CurveOwner.Get() );
 }
 
diff --git a/Engine/Source/Editor/Sequencer/Private/SequencerSelection.cpp b/Engine/Source/Editor/Sequencer/Private/SequencerSelection.cpp
--- a/Engine/Source/Editor/Sequencer/Private/SequencerSelection.cpp
+++ b/Engine/Source/Editor/Sequencer/Private/SequencerSelection.cpp
@@ -7,6 +7,7 @@
 FSequencerSelection::FSequencerSelection()
 {
 	ActiveSelection = EActiveSelection::None;
+	bSuspendBroadcast = false;
 }
 
 const TSet<FSelectedKey>& FSequencerSelection::GetSelectedKeys() const
@@ -48,42 +49,60 @@ void FSequencerSelection::AddToSelection(FSelectedKey Key)
 {
 	SelectedKeys.Add(Key);
 	ActiveSelection = EActiveSelection::KeyAndSection;
-	OnKeySelectionChanged.Broadcast();
+	if (!bSuspendBroadcast)
+	{
+		OnKeySelectionChanged.Broadcast();
+	}
 }
 
 void FSequencerSelection::AddToSelection(UMovieSceneSection* Section)
 {
 	SelectedSections.Add(Section);
 	ActiveSelection = EActiveSelection::KeyAndSection;
-	OnSectionSelectionChanged.Broadcast();
+	if (!bSuspendBroadcast)
+	{
+		OnSectionSelectionChanged.Broadcast();
+	}
 }
 
 void FSequencerSelection::AddToSelection(TSharedRef<FSequencerDisplayNode> OutlinerNode)
 {
 	SelectedOutlinerNodes.Add(OutlinerNode);
 	ActiveSelection = EActiveSelection::OutlinerNode;
-	OnOutlinerNodeSelectionChanged.Broadcast();
+	if (!bSuspendBroadcast)
+	{
+		OnOutlinerNodeSelectionChanged.Broadcast();
+	}
 }
 
 void FSequencerSelection::RemoveFromSelection(FSelectedKey Key)
 {
 	SelectedKeys.Remove(Key);
 	ActiveSelection = EActiveSelection::KeyAndSection;
-	OnKeySelectionChanged.Broadcast();
+	if (!bSuspendBroadcast)
+	{
+		OnKeySelectionChanged.Broadcast();
+	}
 }
 
 void FSequencerSelection::RemoveFromSelection(UMovieSceneSection* Section)
 {
 	SelectedSections.Remove(Section);
 	ActiveSelection = EActiveSelection::KeyAndSection;
-	OnSectionSelectionChanged.Broadcast();
+	if (!bSuspendBroadcast)
+	{
+		OnSectionSelectionChanged.Broadcast();
+	}
 }
 
 void FSequencerSelection::RemoveFromSelection(TSharedRef<FSequencerDisplayNode> OutlinerNode)
 {
 	SelectedOutlinerNodes.Remove(OutlinerNode);
 	ActiveSelection = EActiveSelection::OutlinerNode;
-	OnOutlinerNodeSelectionChanged.Broadcast();
+	if (!bSuspendBroadcast)
+	{
+		OnOutlinerNodeSelectionChanged.Broadcast();
+	}
 }
 
 bool FSequencerSelection::IsSelected(FSelectedKey Key) const
@@ -111,17 +130,26 @@ void FSequencerSelection::Empty()
 void FSequencerSelection::EmptySelectedKeys()
 {
 	SelectedKeys.Empty();
-	OnKeySelectionChanged.Broadcast();
+	if (!bSuspendBroadcast)
+	{
+		OnKeySelectionChanged.Broadcast();
+	}
 }
 
 void FSequencerSelection::EmptySelectedSections()
 {
 	SelectedSections.Empty();
-	OnSectionSelectionChanged.Broadcast();
+	if (!bSuspendBroadcast)
+	{
+		OnSectionSelectionChanged.Broadcast();
+	}
 }
 
 void FSequencerSelection::EmptySelectedOutlinerNodes()
 {
 	SelectedOutlinerNodes.Empty();
-	OnOutlinerNodeSelectionChanged.Broadcast();
+	if (!bSuspendBroadcast)
+	{
+		OnOutlinerNodeSelectionChanged.Broadcast();
+	}
 }
\ No newline at end of file
diff --git a/Engine/Source/Editor/Sequencer/Private/SequencerSelection.h b/Engine/Source/Editor/Sequencer/Private/SequencerSelection.h
--- a/Engine/Source/Editor/Sequencer/Private/SequencerSelection.h
+++ b/Engine/Source/Editor/Sequencer/Private/SequencerSelection.h
@@ -88,6 +88,10 @@ public:
 
 	/** Gets a multicast delegate which is called when the outliner node selection changes. */
 	FOnSelectionChanged& GetOnOutlinerNodeSelectionChanged();
+
+	/** Suspend or resume broadcast of selection changing  */
+	void SuspendBroadcast() { bSuspendBroadcast = true; }
+	void ResumeBroadcast() { bSuspendBroadcast = false; }
 	
 private:
 	TSet<FSelectedKey> SelectedKeys;
@@ -99,4 +103,6 @@ private:
 	FOnSelectionChanged OnOutlinerNodeSelectionChanged;
 
 	EActiveSelection ActiveSelection;
+
+	bool bSuspendBroadcast;
 };
\ No newline at end of file
