Sequencer: Fix delete in sequencer so that it can fallback to the generic delete if there are no sequencer items to delete (ie. keys, sections, outliner nodes).

diff --git a/Engine/Source/Editor/Sequencer/Private/Sequencer.cpp b/Engine/Source/Editor/Sequencer/Private/Sequencer.cpp
--- a/Engine/Source/Editor/Sequencer/Private/Sequencer.cpp
+++ b/Engine/Source/Editor/Sequencer/Private/Sequencer.cpp
@@ -412,6 +412,8 @@ void FSequencer::DeleteSection(class UMovieSceneSection* Section)
 		// Full refresh required just in case the last section was removed from any track.
 		NotifyMovieSceneDataChanged();
 	}
+
+	Selection.RemoveFromSelection(Section);
 }
 
 void FSequencer::DeleteSelectedKeys()
@@ -434,6 +436,8 @@ void FSequencer::DeleteSelectedKeys()
 	{
 		UpdateRuntimeInstances();
 	}
+
+	Selection.EmptySelectedKeys();
 }
 
 void FSequencer::SpawnOrDestroyPuppetObjects( TSharedRef<FMovieSceneInstance> MovieSceneInstance )
@@ -1536,6 +1540,12 @@ void FSequencer::DeleteSelectedItems()
 	}
 }
 
+bool FSequencer::CanDeleteSelectedItems()
+{
+	// Only delete sequencer items if the delete key was pressed in the sequencer widget and not in, for example, the viewport.
+	return SequencerWidget->HasKeyboardFocus();
+}
+
 void FSequencer::TogglePlay()
 {
 	OnPlay();
@@ -1618,7 +1628,8 @@ void FSequencer::BindSequencerCommands()
 
 	SequencerCommandBindings->MapAction(
 		FGenericCommands::Get().Delete,
-		FExecuteAction::CreateSP( this, &FSequencer::DeleteSelectedItems ) );
+		FExecuteAction::CreateSP( this, &FSequencer::DeleteSelectedItems ),
+		FCanExecuteAction::CreateSP(this, &FSequencer::CanDeleteSelectedItems ));
 
 	SequencerCommandBindings->MapAction(
 		Commands.TogglePlay,
diff --git a/Engine/Source/Editor/Sequencer/Private/Sequencer.h b/Engine/Source/Editor/Sequencer/Private/Sequencer.h
--- a/Engine/Source/Editor/Sequencer/Private/Sequencer.h
+++ b/Engine/Source/Editor/Sequencer/Private/Sequencer.h
@@ -354,6 +354,7 @@ protected:
 	
 	/** Called when a user executes the delete command to delete sections or keys */
 	void DeleteSelectedItems();
+	bool CanDeleteSelectedItems();
 	
 	/** Transport controls */
 	void TogglePlay();
