Sequencer - Fix add track drop-down menus and add sub object support to the outliner and drop-downs.

diff --git a/Engine/Source/Editor/MovieSceneTools/Private/K2Node_PlayMovieScene.cpp b/Engine/Source/Editor/MovieSceneTools/Private/K2Node_PlayMovieScene.cpp
--- a/Engine/Source/Editor/MovieSceneTools/Private/K2Node_PlayMovieScene.cpp
+++ b/Engine/Source/Editor/MovieSceneTools/Private/K2Node_PlayMovieScene.cpp
@@ -240,6 +240,11 @@ void UK2Node_PlayMovieScene::BindPossessableToObjects(const FGuid& PossessableGu
 	FBlueprintEditorUtils::MarkBlueprintAsStructurallyModified( this->GetBlueprint() );
 }
 
+void UK2Node_PlayMovieScene::UnbindPossessable( const FGuid& PossessableGuid )
+{
+	// @todo Sequencer: Should we remove the pin and unlink?
+	MovieSceneBindings->RemoveBinding(PossessableGuid);
+}
 
 TArray< FMovieSceneBoundObjectInfo > UK2Node_PlayMovieScene::FindBoundObjectInfos( const FGuid& Guid )
 {
diff --git a/Engine/Source/Editor/MovieSceneTools/Public/K2Node_PlayMovieScene.h b/Engine/Source/Editor/MovieSceneTools/Public/K2Node_PlayMovieScene.h
--- a/Engine/Source/Editor/MovieSceneTools/Public/K2Node_PlayMovieScene.h
+++ b/Engine/Source/Editor/MovieSceneTools/Public/K2Node_PlayMovieScene.h
@@ -76,6 +76,9 @@ class UK2Node_PlayMovieScene : public UK2Node
 	 */
 	virtual void BindPossessableToObjects( const FGuid& PossessableGuid, const TArray< FMovieSceneBoundObjectInfo >& ObjectInfos );
 
+	/** Unbinds all possessable objects for the specified guid */
+	virtual void UnbindPossessable(const FGuid& PossessableGuid);
+
 	/**
 	 * Given a guid for a possessable, attempts to find the object bound to that guid
 	 *
diff --git a/Engine/Source/Editor/NiagaraEditor/Private/NiagaraSequencer.h b/Engine/Source/Editor/NiagaraEditor/Private/NiagaraSequencer.h
--- a/Engine/Source/Editor/NiagaraEditor/Private/NiagaraSequencer.h
+++ b/Engine/Source/Editor/NiagaraEditor/Private/NiagaraSequencer.h
@@ -246,6 +246,7 @@ public:
 	virtual void BindPossessableObject(const FGuid& PossessableGuid, UObject& PossessedObject) override	{};
 	virtual void UnbindPossessableObjects(const FGuid& PossessableGuid) override	{};
 	virtual void GetRuntimeObjects(const TSharedRef<FMovieSceneInstance>& MovieSceneInstance, const FGuid& ObjectGuid, TArray<UObject*>& OutRuntimeObjects) const override	{}
+	virtual UObject* GetParentObject( UObject* Object ) const override { return nullptr; }
 
 	virtual bool TryGetObjectBindingDisplayName(const TSharedRef<FMovieSceneInstance>& MovieSceneInstance, const FGuid& ObjectGuid, FText& DisplayName) const override	{ return false; }
 private:
diff --git a/Engine/Source/Editor/Sequencer/Private/SAnimationOutlinerView.cpp b/Engine/Source/Editor/Sequencer/Private/SAnimationOutlinerView.cpp
--- a/Engine/Source/Editor/Sequencer/Private/SAnimationOutlinerView.cpp
+++ b/Engine/Source/Editor/Sequencer/Private/SAnimationOutlinerView.cpp
@@ -26,8 +26,8 @@ void SAnimationOutlinerTreeNode::Construct( const FArguments& InArgs, TSharedRef
 	ExpandedBrush = FEditorStyle::GetBrush( "TreeArrow_Expanded" );
 	CollapsedBrush = FEditorStyle::GetBrush( "TreeArrow_Collapsed" );
 
-	// Choose the font.  If the node is a root node, we show a larger font for it.
-	FSlateFontInfo NodeFont = Node->GetParent().IsValid() ? 
+	// Choose the font.  If the node is a root node or an object node, we show a larger font for it.
+	FSlateFontInfo NodeFont = Node->GetParent().IsValid() && Node->GetType() != ESequencerNode::Object ? 
 		FEditorStyle::GetFontStyle("Sequencer.AnimationOutliner.RegularFont") :
 		FEditorStyle::GetFontStyle("Sequencer.AnimationOutliner.BoldFont");
 
@@ -82,7 +82,6 @@ void SAnimationOutlinerTreeNode::Construct( const FArguments& InArgs, TSharedRef
 		[
 			SNew(SButton)
 			.ButtonStyle(FEditorStyle::Get(), "FlatButton")
-			.Visibility(this, &SAnimationOutlinerTreeNode::GetKeyButtonVisibility)
 			.ToolTipText(LOCTEXT("PreviousKeyButton", "Set the time to the previous key"))
 			.OnClicked(this, &SAnimationOutlinerTreeNode::OnPreviousKeyClicked)
 			.ContentPadding(0)
@@ -91,6 +90,7 @@ void SAnimationOutlinerTreeNode::Construct( const FArguments& InArgs, TSharedRef
 				.TextStyle(FEditorStyle::Get(), "NormalText.Important")
 				.Font(FEditorStyle::Get().GetFontStyle("FontAwesome.7"))
 				.Text(FText::FromString(FString(TEXT("\xf060"))) /*fa-arrow-left*/)
+				.ColorAndOpacity(this, &SAnimationOutlinerTreeNode::GetKeyButtonColorAndOpacity)
 			]
 		]
 		// Add key slot
@@ -100,7 +100,6 @@ void SAnimationOutlinerTreeNode::Construct( const FArguments& InArgs, TSharedRef
 		[
 			SNew(SButton)
 			.ButtonStyle(FEditorStyle::Get(), "FlatButton")
-			.Visibility(this, &SAnimationOutlinerTreeNode::GetKeyButtonVisibility)
 			.ToolTipText(LOCTEXT("AddKeyButton", "Add a new key at the current time"))
 			.OnClicked(this, &SAnimationOutlinerTreeNode::OnAddKeyClicked)
 			.ContentPadding(0)
@@ -109,6 +108,7 @@ void SAnimationOutlinerTreeNode::Construct( const FArguments& InArgs, TSharedRef
 				.TextStyle(FEditorStyle::Get(), "NormalText.Important")
 				.Font(FEditorStyle::Get().GetFontStyle("FontAwesome.7"))
 				.Text(FText::FromString(FString(TEXT("\xf055"))) /*fa-plus-circle*/)
+				.ColorAndOpacity(this, &SAnimationOutlinerTreeNode::GetKeyButtonColorAndOpacity)
 			]
 		]
 		// Next key slot
@@ -118,7 +118,6 @@ void SAnimationOutlinerTreeNode::Construct( const FArguments& InArgs, TSharedRef
 		[
 			SNew(SButton)
 			.ButtonStyle(FEditorStyle::Get(), "FlatButton")
-			.Visibility(this, &SAnimationOutlinerTreeNode::GetKeyButtonVisibility)
 			.ToolTipText(LOCTEXT("NextKeyButton", "Set the time to the next key"))
 			.OnClicked(this, &SAnimationOutlinerTreeNode::OnNextKeyClicked)
 			.ContentPadding(0)
@@ -127,6 +126,7 @@ void SAnimationOutlinerTreeNode::Construct( const FArguments& InArgs, TSharedRef
 				.TextStyle(FEditorStyle::Get(), "NormalText.Important")
 				.Font(FEditorStyle::Get().GetFontStyle("FontAwesome.7"))
 				.Text(FText::FromString(FString(TEXT("\xf061"))) /*fa-arrow-right*/)
+				.ColorAndOpacity(this, &SAnimationOutlinerTreeNode::GetKeyButtonColorAndOpacity)
 			]
 		];
 
@@ -416,9 +416,11 @@ EVisibility SAnimationOutlinerTreeNode::GetExpanderVisibility() const
 	return DisplayNode->GetNumChildren() > 0 ? EVisibility::Visible : EVisibility::Hidden;
 }
 
-EVisibility SAnimationOutlinerTreeNode::GetKeyButtonVisibility() const
+FSlateColor SAnimationOutlinerTreeNode::GetKeyButtonColorAndOpacity() const
 {
-	return DisplayNode->GetType() == ESequencerNode::Object ? EVisibility::Hidden : EVisibility::Visible;
+	return IsHovered()
+		? FLinearColor(1, 1, 1, 1)
+		: FLinearColor(.25f, .25f, .25f, 1);
 }
 
 FText SAnimationOutlinerTreeNode::GetDisplayName() const
diff --git a/Engine/Source/Editor/Sequencer/Private/SAnimationOutlinerView.h b/Engine/Source/Editor/Sequencer/Private/SAnimationOutlinerView.h
--- a/Engine/Source/Editor/Sequencer/Private/SAnimationOutlinerView.h
+++ b/Engine/Source/Editor/Sequencer/Private/SAnimationOutlinerView.h
@@ -54,9 +54,9 @@ private:
 	EVisibility GetExpanderVisibility() const;
 
 	/**
-	 * @return the visibility of the previous, next, and add key buttons. 
+	 * @return the color of the previous, next, and add key buttons. 
 	 */
-	EVisibility GetKeyButtonVisibility() const;
+	FSlateColor GetKeyButtonColorAndOpacity() const;
 
 	/**
 	 * @return The display name for this node.
diff --git a/Engine/Source/Editor/Sequencer/Private/SSequencer.cpp b/Engine/Source/Editor/Sequencer/Private/SSequencer.cpp
--- a/Engine/Source/Editor/Sequencer/Private/SSequencer.cpp
+++ b/Engine/Source/Editor/Sequencer/Private/SSequencer.cpp
@@ -474,7 +474,7 @@ TSharedRef<SWidget> SSequencer::MakeToolBar()
 					[
 						SNew(STextBlock)
 						.TextStyle(FEditorStyle::Get(), "NormalText.Important")
-						.Text(LOCTEXT("AddButton", "Track"))
+						.Text(LOCTEXT("AddButton", "Add"))
 					]
 
 					+ SHorizontalBox::Slot()
diff --git a/Engine/Source/Editor/Sequencer/Private/SSequencer.h b/Engine/Source/Editor/Sequencer/Private/SSequencer.h
--- a/Engine/Source/Editor/Sequencer/Private/SSequencer.h
+++ b/Engine/Source/Editor/Sequencer/Private/SSequencer.h
@@ -8,6 +8,8 @@ namespace SequencerLayoutConstants
 	const float IndentAmount = 25.0f;
 	/** Padding between each node */
 	const float NodePadding = 3.0f;
+	/** Padding between each node */
+	const float ObjectNodePadding = 5.0f;
 	/** Height of each object node */
 	const float ObjectNodeHeight = 20.0f;
 	/** Height of each section area if there are no sections (note: section areas may be larger than this if they have children. This is the height of a section area with no children or all children hidden) */
diff --git a/Engine/Source/Editor/Sequencer/Private/Sequencer.cpp b/Engine/Source/Editor/Sequencer/Private/Sequencer.cpp
--- a/Engine/Source/Editor/Sequencer/Private/Sequencer.cpp
+++ b/Engine/Source/Editor/Sequencer/Private/Sequencer.cpp
@@ -663,6 +663,13 @@ FGuid FSequencer::GetHandleToObject( UObject* Object )
 			ObjectBindingManager->BindPossessableObject(ObjectGuid, *Object);
 			bPossessableAdded = true;
 		}
+
+		// If we're adding a possessable, generate handles to it's parent objects too.
+		UObject* ParentObject = ObjectBindingManager->GetParentObject(Object);
+		if ( ParentObject != nullptr )
+		{
+			GetHandleToObject(ParentObject);
+		}
 	}
 
 	if (bPossessableAdded)
diff --git a/Engine/Source/Editor/Sequencer/Private/SequencerActorBindingManager.cpp b/Engine/Source/Editor/Sequencer/Private/SequencerActorBindingManager.cpp
--- a/Engine/Source/Editor/Sequencer/Private/SequencerActorBindingManager.cpp
+++ b/Engine/Source/Editor/Sequencer/Private/SequencerActorBindingManager.cpp
@@ -367,7 +367,10 @@ void FSequencerActorBindingManager::BindPossessableObject( const FGuid& Possessa
 
 void FSequencerActorBindingManager::UnbindPossessableObjects( const FGuid& PossessableGuid )
 {
-	// @todo Sequencer: Should we remove the pin and unlink?
+	if (PlayMovieSceneNode.IsValid())
+	{
+		PlayMovieSceneNode->UnbindPossessable(PossessableGuid);
+	}
 }
 
 void FSequencerActorBindingManager::GetRuntimeObjects( const TSharedRef<FMovieSceneInstance>& MovieSceneInstance, const FGuid& ObjectGuid, TArray<UObject*>& OutRuntimeObjects ) const
@@ -430,6 +433,16 @@ bool FSequencerActorBindingManager::TryGetObjectBindingDisplayName(const TShared
 	return false;
 }
 
+UObject* FSequencerActorBindingManager::GetParentObject( UObject* Object ) const
+{
+	UActorComponent* Component = Cast<UActorComponent>(Object);
+	if (Component != nullptr)
+	{
+		return Component->GetOwner();
+	}
+	return nullptr;
+}
+
 UK2Node_PlayMovieScene* FSequencerActorBindingManager::FindPlayMovieSceneNodeInLevelScript( const UMovieScene* MovieScene ) const
 {
 	// Grab the world object for this editor
diff --git a/Engine/Source/Editor/Sequencer/Private/SequencerActorBindingManager.h b/Engine/Source/Editor/Sequencer/Private/SequencerActorBindingManager.h
--- a/Engine/Source/Editor/Sequencer/Private/SequencerActorBindingManager.h
+++ b/Engine/Source/Editor/Sequencer/Private/SequencerActorBindingManager.h
@@ -40,6 +40,7 @@ public:
 	virtual void UnbindPossessableObjects( const FGuid& PossessableGuid ) override;
 	virtual void GetRuntimeObjects( const TSharedRef<FMovieSceneInstance>& MovieSceneInstance, const FGuid& ObjectGuid, TArray<UObject*>& OutRuntimeObjects ) const override;
 	virtual bool TryGetObjectBindingDisplayName(const TSharedRef<FMovieSceneInstance>& MovieSceneInstance, const FGuid& ObjectGuid, FText& DisplayName) const override;
+	virtual UObject* GetParentObject( UObject* Object ) const override;
 
 protected:
 	
diff --git a/Engine/Source/Editor/Sequencer/Private/SequencerAssetEditor.cpp b/Engine/Source/Editor/Sequencer/Private/SequencerAssetEditor.cpp
--- a/Engine/Source/Editor/Sequencer/Private/SequencerAssetEditor.cpp
+++ b/Engine/Source/Editor/Sequencer/Private/SequencerAssetEditor.cpp
@@ -109,7 +109,11 @@ TSharedRef<ISequencer> FSequencerAssetEditor::GetSequencerInterface() const
 
 FSequencerAssetEditor::FSequencerAssetEditor()
 {
-
+	// Register sequencer menu extenders.
+	ISequencerModule& SequencerModule = FModuleManager::Get().LoadModuleChecked<ISequencerModule>( "Sequencer" );
+	int32 NewIndex = SequencerModule.GetMenuExtensibilityManager()->GetExtenderDelegates().Add(
+		FAssetEditorExtender::CreateRaw( this, &FSequencerAssetEditor::GetContextSensitiveSequencerExtender ) );
+	SequencerExtenderHandle = SequencerModule.GetMenuExtensibilityManager()->GetExtenderDelegates()[NewIndex].GetHandle();
 }
 
 
@@ -123,6 +127,13 @@ FSequencerAssetEditor::~FSequencerAssetEditor()
 		auto& LevelEditorModule = FModuleManager::LoadModuleChecked< FLevelEditorModule >( TEXT( "LevelEditor" ) );
 		LevelEditorModule.OnMapChanged().RemoveAll( this );
 	}
+
+	// Un-Register sequencer menu extenders.
+	ISequencerModule& SequencerModule = FModuleManager::Get().LoadModuleChecked<ISequencerModule>( "Sequencer" );
+	SequencerModule.GetMenuExtensibilityManager()->GetExtenderDelegates().RemoveAll( [this]( const FAssetEditorExtender& Extender )
+	{
+		return SequencerExtenderHandle == Extender.GetHandle();
+	} );
 }
 
 
@@ -164,5 +175,43 @@ FString FSequencerAssetEditor::GetWorldCentricTabPrefix() const
 	return LOCTEXT("WorldCentricTabPrefix", "Sequencer ").ToString();
 }
 
+TSharedRef<FExtender> FSequencerAssetEditor::GetContextSensitiveSequencerExtender( const TSharedRef<FUICommandList> CommandList, const TArray<UObject*> ContextSensitiveObjects )
+{
+	TSharedRef<FExtender> AddTrackMenuExtender( new FExtender() );
+	AddTrackMenuExtender->AddMenuExtension(
+		SequencerMenuExtensionPoints::AddTrackMenu_PropertiesSection,
+		EExtensionHook::Before,
+		CommandList,
+		FMenuExtensionDelegate::CreateRaw( this, &FSequencerAssetEditor::ExtendSequencerAddTrackMenu, ContextSensitiveObjects ) );
+	return AddTrackMenuExtender;
+}
+
+void FSequencerAssetEditor::ExtendSequencerAddTrackMenu( FMenuBuilder& AddTrackMenuBuilder, TArray<UObject*> ContextObjects )
+{
+	if ( ContextObjects.Num() == 1 )
+	{
+		AActor* Actor = Cast<AActor>( ContextObjects[0] );
+		if ( Actor != nullptr )
+		{
+			AddTrackMenuBuilder.BeginSection( "Components", LOCTEXT( "ComponentsSection", "Components" ) );
+			{
+				for ( UActorComponent* Component : Actor->GetComponents() )
+				{
+					FUIAction AddComponentAction( FExecuteAction::CreateSP( this, &FSequencerAssetEditor::AddComponentTrack, Component ) );
+					FText AddComponentLabel = FText::FromString(Component->GetName());
+					FText AddComponentToolTip = FText::Format( LOCTEXT( "ComponentToolTipFormat", "Add {0} component" ), FText::FromString( Component->GetName() ) );
+					AddTrackMenuBuilder.AddMenuEntry( AddComponentLabel, AddComponentToolTip, FSlateIcon(), AddComponentAction );
+				}
+			}
+			AddTrackMenuBuilder.EndSection();
+		}
+	}
+}
+
+void FSequencerAssetEditor::AddComponentTrack( UActorComponent* Component )
+{
+	Sequencer->GetHandleToObject( Component );
+}
+
 
 #undef LOCTEXT_NAMESPACE
diff --git a/Engine/Source/Editor/Sequencer/Private/SequencerAssetEditor.h b/Engine/Source/Editor/Sequencer/Private/SequencerAssetEditor.h
--- a/Engine/Source/Editor/Sequencer/Private/SequencerAssetEditor.h
+++ b/Engine/Source/Editor/Sequencer/Private/SequencerAssetEditor.h
@@ -42,6 +42,10 @@ private:
 	EVisibility GetTransportControlVisibility(TWeakPtr<ILevelViewport> LevelViewport) const;
 	TSharedRef<SDockTab> SpawnTab_SequencerMain(const FSpawnTabArgs& Args);	
 
+	TSharedRef<FExtender> GetContextSensitiveSequencerExtender( const TSharedRef<FUICommandList> CommandList, const TArray<UObject*> ContextSensitiveObjects );
+	void ExtendSequencerAddTrackMenu( FMenuBuilder& AddTrackMenuBuilder, TArray<UObject*> ContextObjects );
+	void AddComponentTrack( UActorComponent* Component );
+
 private:
 	TSharedPtr<FSequencer> Sequencer;
 
@@ -50,4 +54,5 @@ private:
 	/**	The tab ids for all the tabs used */
 	static const FName SequencerMainTabId;
 
+	FDelegateHandle SequencerExtenderHandle;
 };
diff --git a/Engine/Source/Editor/Sequencer/Private/SequencerDisplayNode.cpp b/Engine/Source/Editor/Sequencer/Private/SequencerDisplayNode.cpp
--- a/Engine/Source/Editor/Sequencer/Private/SequencerDisplayNode.cpp
+++ b/Engine/Source/Editor/Sequencer/Private/SequencerDisplayNode.cpp
@@ -140,6 +140,11 @@ FSequencerDisplayNode::FSequencerDisplayNode( FName InNodeName, TSharedPtr<FSequ
 	bExpanded = ParentTree.GetSavedExpansionState( *this );
 }
 
+void FSequencerDisplayNode::AddObjectBindingNode(TSharedRef<FObjectBindingNode> ObjectBindingNode)
+{
+	ChildNodes.Add(ObjectBindingNode);
+}
+
 TSharedRef<FSectionCategoryNode> FSequencerDisplayNode::AddCategoryNode( FName CategoryName, const FText& DisplayLabel )
 {
 	TSharedPtr<FSectionCategoryNode> CategoryNode;
@@ -541,7 +546,7 @@ TSharedRef<SWidget> FObjectBindingNode::GenerateEditWidgetForOutliner()
 			[
 				SNew(STextBlock)
 				.Font(FEditorStyle::GetFontStyle("Sequencer.AnimationOutliner.RegularFont"))
-				.Text(LOCTEXT("AddPropertyButton", "Property"))
+				.Text(LOCTEXT("AddTrackButton", "Track"))
 			]
 		]
 	];
@@ -614,17 +619,25 @@ TSharedRef<SWidget> FObjectBindingNode::OnGetAddPropertyTrackMenuContent()
 		return A.Num() < B.Num();
 	});
 
-	FMenuBuilder AddTrackMenuBuilder(true, NULL);
-	for (TArray<UProperty*>& KeyablePropertyPath : KeyablePropertyPaths)
+	ISequencerModule& SequencerModule = FModuleManager::GetModuleChecked<ISequencerModule>( "Sequencer" );
+	TSharedRef<FUICommandList> CommandList(new FUICommandList);
+	FMenuBuilder AddTrackMenuBuilder(true, nullptr, SequencerModule.GetMenuExtensibilityManager()->GetAllExtenders(CommandList, BoundObjects));
+
+	AddTrackMenuBuilder.BeginSection( SequencerMenuExtensionPoints::AddTrackMenu_PropertiesSection, LOCTEXT("PropertiesMenuHeader" , "Properties"));
 	{
-		FUIAction AddTrackMenuAction(FExecuteAction::CreateSP(this, &FObjectBindingNode::AddTrackForProperty, KeyablePropertyPath));
-		TArray<FString> PropertyNames;
-		for (UProperty* Property : KeyablePropertyPath)
+		for ( TArray<UProperty*>& KeyablePropertyPath : KeyablePropertyPaths )
 		{
-			PropertyNames.Add(Property->GetName());
+			FUIAction AddTrackMenuAction( FExecuteAction::CreateSP( this, &FObjectBindingNode::AddTrackForProperty, KeyablePropertyPath ) );
+			TArray<FString> PropertyNames;
+			for ( UProperty* Property : KeyablePropertyPath )
+			{
+				PropertyNames.Add( Property->GetName() );
+			}
+			AddTrackMenuBuilder.AddMenuEntry( FText::FromString( FString::Join( PropertyNames, TEXT( "." ) ) ), FText(), FSlateIcon(), AddTrackMenuAction );
 		}
-		AddTrackMenuBuilder.AddMenuEntry(FText::FromString(FString::Join(PropertyNames, TEXT("."))), FText(), FSlateIcon(), AddTrackMenuAction);
 	}
+	AddTrackMenuBuilder.EndSection();
+
 	return AddTrackMenuBuilder.MakeWidget();
 }
 
@@ -633,7 +646,7 @@ void FObjectBindingNode::AddTrackForProperty(TArray<UProperty*> PropertyPath)
 	FSequencer& Sequencer = GetSequencer();
 
 	TArray<UObject*> BoundObjects;
-	Sequencer.GetObjectBindingManager()->GetRuntimeObjects(Sequencer.GetRootMovieSceneInstance(), ObjectBinding, BoundObjects);
+	Sequencer.GetObjectBindingManager()->GetRuntimeObjects(Sequencer.GetFocusedMovieSceneInstance(), ObjectBinding, BoundObjects);
 
 	TArray<UObject*> KeyableBoundObjects;
 	for (UObject* BoundObject : BoundObjects)
diff --git a/Engine/Source/Editor/Sequencer/Private/SequencerDisplayNode.h b/Engine/Source/Editor/Sequencer/Private/SequencerDisplayNode.h
--- a/Engine/Source/Editor/Sequencer/Private/SequencerDisplayNode.h
+++ b/Engine/Source/Editor/Sequencer/Private/SequencerDisplayNode.h
@@ -36,6 +36,13 @@ public:
 	 */
 	FSequencerDisplayNode( FName InNodeName, TSharedPtr<FSequencerDisplayNode> InParentNode, FSequencerNodeTree& InParentTree );
 
+	/**
+	* Adds an object binding node to this node.
+	*
+	* @param ObjectBindingNode The node to add
+	*/
+	void AddObjectBindingNode( TSharedRef<FObjectBindingNode> ObjectBindingNode );
+
 	/**
 	 * Adds a category to this node
 	 * 
@@ -129,6 +136,19 @@ public:
 	 */
 	const TArray< TSharedRef<FSequencerDisplayNode> >& GetChildNodes() const { return ChildNodes; }
 
+	/**
+	 * Sorts the child nodes with the supplied predicate
+	 */
+	template <class PREDICATE_CLASS>
+	void SortChildNodes(const PREDICATE_CLASS& Predicate)
+	{
+		ChildNodes.StableSort(Predicate);
+		for (TSharedRef<FSequencerDisplayNode> ChildNode : ChildNodes)
+		{
+			ChildNode->SortChildNodes(Predicate);
+		}
+	}
+
 	/**
 	 * @return The parent of this node                                                              
 	 */
diff --git a/Engine/Source/Editor/Sequencer/Private/SequencerModule.cpp b/Engine/Source/Editor/Sequencer/Private/SequencerModule.cpp
--- a/Engine/Source/Editor/Sequencer/Private/SequencerModule.cpp
+++ b/Engine/Source/Editor/Sequencer/Private/SequencerModule.cpp
@@ -78,6 +78,9 @@ public:
 				FSlateIcon(),
 				false);
 		}
+
+		MenuExtensibilityManager = MakeShareable( new FExtensibilityManager );
+		ToolBarExtensibilityManager = MakeShareable( new FExtensibilityManager );
 	}
 
 	virtual void ShutdownModule() override
@@ -90,10 +93,16 @@ public:
 		}
 	}
 
+	virtual TSharedPtr<FExtensibilityManager> GetMenuExtensibilityManager() const override { return MenuExtensibilityManager; }
+	virtual TSharedPtr<FExtensibilityManager> GetToolBarExtensibilityManager() const override { return ToolBarExtensibilityManager; }
+
 private:
 
 	/** List of auto-key handler delegates sequencers will execute when they are created */
 	TArray< FOnCreateTrackEditor > TrackEditorDelegates;
+
+	TSharedPtr<FExtensibilityManager> MenuExtensibilityManager;
+	TSharedPtr<FExtensibilityManager> ToolBarExtensibilityManager;
 };
 
 
diff --git a/Engine/Source/Editor/Sequencer/Private/SequencerNodeTree.cpp b/Engine/Source/Editor/Sequencer/Private/SequencerNodeTree.cpp
--- a/Engine/Source/Editor/Sequencer/Private/SequencerNodeTree.cpp
+++ b/Engine/Source/Editor/Sequencer/Private/SequencerNodeTree.cpp
@@ -10,6 +10,7 @@
 #include "MovieSceneTrackEditor.h"
 #include "SectionLayoutBuilder.h"
 #include "ISequencerSection.h"
+#include "ISequencerObjectBindingManager.h"
 
 void FSequencerNodeTree::Empty()
 {
@@ -47,10 +48,17 @@ void FSequencerNodeTree::Update()
 
 	const TArray<FMovieSceneBinding>& Bindings = MovieScene->GetBindings();
 
+	TMap<FGuid, const FMovieSceneBinding*> GuidToBindingMap;
+	for (const FMovieSceneBinding& Binding : Bindings)
+	{
+		GuidToBindingMap.Add(Binding.GetObjectGuid(), &Binding);
+	}
+
 	// Make nodes for all object bindings
+	TArray< TSharedRef<FSequencerDisplayNode> > NewObjectNodes;
 	for( const FMovieSceneBinding& Binding : Bindings )
 	{
-		TSharedRef<FObjectBindingNode> ObjectBindingNode = AddObjectBinding( Binding.GetName(), Binding.GetObjectGuid(), NewRootNodes );
+		TSharedRef<FObjectBindingNode> ObjectBindingNode = AddObjectBinding( Binding.GetName(), Binding.GetObjectGuid(), GuidToBindingMap, NewObjectNodes );
 
 		const TArray<UMovieSceneTrack*>& Tracks = Binding.GetTracks();
 
@@ -64,41 +72,31 @@ void FSequencerNodeTree::Update()
 			TSharedRef<FTrackNode> SectionAreaNode = ObjectBindingNode->AddSectionAreaNode( SectionName, TrackRef );
 			MakeSectionInterfaces( TrackRef, SectionAreaNode );
 		}
-
 	}
 
-	struct FRootNodeSorter
+	struct FObjectNodeSorter
 	{
 		bool operator()( const TSharedRef<FSequencerDisplayNode>& A, const TSharedRef<FSequencerDisplayNode>& B ) const
 		{
-			if (A->GetType() == ESequencerNode::Object)
+			if (A->GetType() == ESequencerNode::Object && B->GetType() != ESequencerNode::Object)
 			{
-				if (B->GetType() == ESequencerNode::Object)
-				{
-					return A->GetDisplayName().ToString() < B->GetDisplayName().ToString();
-				}
-				else
-				{
-					return false; // ie. master tracks should be first in line
-				}
+				return true;
 			}
-			if (A->GetType() != ESequencerNode::Object)
+			if (A->GetType() != ESequencerNode::Object && B->GetType() == ESequencerNode::Object)
 			{
-				if (B->GetType() != ESequencerNode::Object)
-				{
-					return A->GetDisplayName().ToString() < B->GetDisplayName().ToString();
-				}
-				else
-				{
-					return true; // ie. master tracks should be first in line
-				}
+				return false;
 			}
-			return false;
+			return A->GetDisplayName().ToString() < B->GetDisplayName().ToString();
 		}
 	};
 
-	// Sort so that master tracks appear before object tracks
-	NewRootNodes.Sort( FRootNodeSorter() );
+	NewObjectNodes.Sort( FObjectNodeSorter() );
+	for (TSharedRef<FSequencerDisplayNode> NewObjectNode : NewObjectNodes)
+	{
+		NewObjectNode->SortChildNodes(FObjectNodeSorter());
+	}
+
+	NewRootNodes.Append(NewObjectNodes);
 
 	// Look for a shot track.  It will always come first if it exists
 	UMovieSceneTrack* ShotTrack = MovieScene->GetShotTrack();
@@ -174,21 +172,62 @@ const TArray< TSharedRef<FSequencerDisplayNode> >& FSequencerNodeTree::GetRootNo
 }
 
 
-TSharedRef<FObjectBindingNode> FSequencerNodeTree::AddObjectBinding( const FString& ObjectName, const FGuid& ObjectBinding, TArray< TSharedRef<FSequencerDisplayNode> >& OutNodeList )
+TSharedRef<FObjectBindingNode> FSequencerNodeTree::AddObjectBinding(const FString& ObjectName, const FGuid& ObjectBinding, TMap<FGuid, const FMovieSceneBinding*>& GuidToBindingMap, TArray< TSharedRef<FSequencerDisplayNode> >& OutNodeList)
 {
-	// The node name is the object guid
-	FName ObjectNodeName = *ObjectBinding.ToString();
-	TSharedPtr<FSequencerDisplayNode> ParentNode = NULL;
-
-	TSharedRef< FObjectBindingNode > ObjectNode = MakeShareable( new FObjectBindingNode( ObjectNodeName, ObjectName, ObjectBinding, ParentNode, *this ) );
+	TSharedPtr<FObjectBindingNode> ObjectNode;
+	TSharedPtr<FObjectBindingNode>* FoundObjectNode = ObjectBindingMap.Find(ObjectBinding);
+	if (FoundObjectNode != nullptr)
+	{
+		ObjectNode = *FoundObjectNode;
+	}
+	else
+	{
+		// The node name is the object guid
+		FName ObjectNodeName = *ObjectBinding.ToString();
+
+		// Try to get the parent object node if there is one.
+		TSharedPtr<FObjectBindingNode> ParentNode;
+		TArray<UObject*> RuntimeObjects;
+		TSharedRef<ISequencerObjectBindingManager> BindingManager(Sequencer.GetObjectBindingManager());
+		BindingManager->GetRuntimeObjects(Sequencer.GetFocusedMovieSceneInstance(), ObjectBinding, RuntimeObjects);
+		if ( RuntimeObjects.Num() == 1)
+		{
+			UObject* ParentObject = BindingManager->GetParentObject(RuntimeObjects[0]);
+			if (ParentObject != nullptr)
+			{
+				FGuid ParentBinding = BindingManager->FindGuidForObject(*Sequencer.GetFocusedMovieScene(), *ParentObject);
+				TSharedPtr<FObjectBindingNode>* FoundParentNode = ObjectBindingMap.Find( ParentBinding );
+				if ( FoundParentNode != nullptr )
+				{
+					ParentNode = *FoundParentNode;
+				}
+				else
+				{
+					const FMovieSceneBinding** FoundParentMovieSceneBinding = GuidToBindingMap.Find( ParentBinding );
+					if ( FoundParentMovieSceneBinding != nullptr )
+					{
+						ParentNode = AddObjectBinding( (*FoundParentMovieSceneBinding)->GetName(), ParentBinding, GuidToBindingMap, OutNodeList );
+					}
+				}
+			}
+		}
 
-	// Object binding nodes are always root nodes
-	OutNodeList.Add( ObjectNode );
+		// Create the node.
+		ObjectNode = MakeShareable( new FObjectBindingNode( ObjectNodeName, ObjectName, ObjectBinding, ParentNode, *this ) );
+		if (ParentNode.IsValid())
+		{
+			ParentNode->AddObjectBindingNode(ObjectNode.ToSharedRef());
+		}
+		else
+		{
+			OutNodeList.Add( ObjectNode.ToSharedRef() );
+		}
 
-	// Map the guid to the object binding node for fast lookup later
-	ObjectBindingMap.Add( ObjectBinding, ObjectNode );
+		// Map the guid to the object binding node for fast lookup later
+		ObjectBindingMap.Add( ObjectBinding, ObjectNode );
+	}
 
-	return ObjectNode;
+	return ObjectNode.ToSharedRef();
 }
 
 void FSequencerNodeTree::SaveExpansionState( const FSequencerDisplayNode& Node, bool bExpanded )
diff --git a/Engine/Source/Editor/Sequencer/Private/SequencerNodeTree.h b/Engine/Source/Editor/Sequencer/Private/SequencerNodeTree.h
--- a/Engine/Source/Editor/Sequencer/Private/SequencerNodeTree.h
+++ b/Engine/Source/Editor/Sequencer/Private/SequencerNodeTree.h
@@ -6,6 +6,7 @@ class UMovieSceneSection;
 class FSequencerDisplayNode;
 class FObjectBindingNode;
 class UMovieSceneTrack;
+struct FMovieSceneBinding;
 
 /**
  * Represents a tree of sequencer display nodes, used to populate the Sequencer UI with MovieScene data
@@ -89,11 +90,9 @@ private:
 	void MakeSectionInterfaces( UMovieSceneTrack& Track, TSharedRef<class FTrackNode>& SectionAreaNode );
 
 	/**
-	 * Creates a new root object binding node
-	 * 
-	 * @param ObjectBinding	The object binding guid
+	 * Creates a new object binding node and any parent binding nodes.
 	 */
-	TSharedRef<FObjectBindingNode> AddObjectBinding( const FString& ObjectName, const FGuid& ObjectBinding, TArray< TSharedRef<FSequencerDisplayNode> >& OutNodeList );
+	TSharedRef<FObjectBindingNode> AddObjectBinding( const FString& ObjectName, const FGuid& ObjectBinding, TMap<FGuid, const FMovieSceneBinding*>& GuidToBindingMap, TArray< TSharedRef<FSequencerDisplayNode> >& OutNodeList );
 private:
 	/** Tools for building movie scene section layouts.  One tool for each track */
 	TMap< class UMovieSceneTrack*, TSharedPtr<FMovieSceneTrackEditor> > EditorMap;
diff --git a/Engine/Source/Editor/Sequencer/Private/SequencerTrackLaneFactory.cpp b/Engine/Source/Editor/Sequencer/Private/SequencerTrackLaneFactory.cpp
--- a/Engine/Source/Editor/Sequencer/Private/SequencerTrackLaneFactory.cpp
+++ b/Engine/Source/Editor/Sequencer/Private/SequencerTrackLaneFactory.cpp
@@ -103,7 +103,9 @@ void FSequencerTrackLaneFactory::GenerateWidgetForNode( FSequencerDisplayNode& N
 	Link.Widgets.Add(TrackWidget);
 
 
-	FMargin Padding = Node.GetType() == ESequencerNode::Object ? FMargin( 0.0f, 5.0f, 0.0f, 0.0f ) : FMargin( 0 );
+	FMargin Padding = Node.GetType() == ESequencerNode::Object 
+		? FMargin( 0.0f, SequencerLayoutConstants::ObjectNodePadding, 0.0f, 0.0f )
+		: FMargin( 0.0f, SequencerLayoutConstants::NodePadding, 0.0f, 0.0f );
 
 	auto GetMaxSize = [=]() -> float {
 		return Link.GetMaxSize().Y;
@@ -142,8 +144,8 @@ void FSequencerTrackLaneFactory::GenerateLayoutNodeWidgetsRecursive( const TArra
 	{
 		FSequencerDisplayNode& Node = *SharedNode;
 
-		// The section area encompasses multiple logical child nodes so we do not generate children that arent section areas (they would not be top level nodes in that case)
-		if( Node.GetType() == ESequencerNode::Track )
+		// Sections generate their own child nodes, so only generate nodes for tracks and objects.
+		if( Node.GetType() == ESequencerNode::Track || Node.GetType() == ESequencerNode::Object)
 		{
 			GenerateWidgetForNode( Node );
 			GenerateLayoutNodeWidgetsRecursive( Node.GetChildNodes() );
diff --git a/Engine/Source/Editor/Sequencer/Public/ISequencerModule.h b/Engine/Source/Editor/Sequencer/Public/ISequencerModule.h
--- a/Engine/Source/Editor/Sequencer/Public/ISequencerModule.h
+++ b/Engine/Source/Editor/Sequencer/Public/ISequencerModule.h
@@ -10,6 +10,11 @@
 
 class ISequencerObjectBindingManager;
 
+namespace SequencerMenuExtensionPoints
+{
+	static const FName AddTrackMenu_PropertiesSection("AddTrackMenu_PropertiesSection");
+}
+
 /**
  * A delegate which will create an auto-key handler
  */
@@ -82,5 +87,11 @@ public:
 	 * @param InHandle	Handle to the delegate to unregister
 	 */
 	virtual void UnRegisterTrackEditor_Handle( FDelegateHandle InHandle ) = 0;
+
+	/** Gets the extensibility manager for menus */
+	virtual TSharedPtr<FExtensibilityManager> GetMenuExtensibilityManager() const = 0;
+
+	/** Gets the extensibility manager for toolbars */
+	virtual TSharedPtr<FExtensibilityManager> GetToolBarExtensibilityManager() const = 0;
 };
 
diff --git a/Engine/Source/Editor/Sequencer/Public/ISequencerObjectBindingManager.h b/Engine/Source/Editor/Sequencer/Public/ISequencerObjectBindingManager.h
--- a/Engine/Source/Editor/Sequencer/Public/ISequencerObjectBindingManager.h
+++ b/Engine/Source/Editor/Sequencer/Public/ISequencerObjectBindingManager.h
@@ -64,4 +64,7 @@ public:
 	 * @returns true if DisplayName has been set to a valid display name, otherwise false.
 	 */
 	virtual bool TryGetObjectBindingDisplayName( const TSharedRef<FMovieSceneInstance>& MovieSceneInstance, const FGuid& ObjectGuid, FText& DisplayName ) const = 0;
+
+	/** Gets the logical parent object for the supplied object which is not necessarily it's outer.  If the object has no logical parent returns nullptr. */
+	virtual UObject* GetParentObject(UObject* Object) const = 0;
 };
diff --git a/Engine/Source/Editor/UMGEditor/Private/Animation/UMGSequencerObjectBindingManager.cpp b/Engine/Source/Editor/UMGEditor/Private/Animation/UMGSequencerObjectBindingManager.cpp
--- a/Engine/Source/Editor/UMGEditor/Private/Animation/UMGSequencerObjectBindingManager.cpp
+++ b/Engine/Source/Editor/UMGEditor/Private/Animation/UMGSequencerObjectBindingManager.cpp
@@ -5,6 +5,7 @@
 #include "WidgetBlueprintEditor.h"
 #include "ISequencer.h"
 #include "MovieScene.h"
+#include "MovieSceneInstance.h"
 #include "WidgetBlueprintEditor.h"
 #include "Components/PanelSlot.h"
 #include "Components/Visual.h"
@@ -151,8 +152,8 @@ bool FUMGSequencerObjectBindingManager::TryGetObjectBindingDisplayName(const TSh
 	{
 		if (BindingObjects[0].IsValid())
 		{
-		DisplayName = FText::FromString(BindingObjects[0].Get()->GetName());
-	}
+			DisplayName = FText::FromString(BindingObjects[0].Get()->GetName());
+		}
 		else
 		{
 			DisplayName = LOCTEXT("InvalidObject", "Invalid Object");
@@ -162,13 +163,22 @@ bool FUMGSequencerObjectBindingManager::TryGetObjectBindingDisplayName(const TSh
 	{
 		if ( SlotContentBindingObjects[0].IsValid())
 		{
-		UWidget* SlotContent = Cast<UWidget>(SlotContentBindingObjects[0].Get());
-		FText PanelName = SlotContent->Slot != nullptr && SlotContent->Slot->Parent != nullptr
-			? FText::FromString(SlotContent->Slot->Parent->GetName())
-			: LOCTEXT("InvalidPanel", "Invalid Panel");
-		FText ContentName = FText::FromString(SlotContent->GetName());
-		DisplayName = FText::Format(LOCTEXT("SlotObject", "{0} ({1} Slot)"), ContentName, PanelName);
-	}
+			UWidget* SlotContent = Cast<UWidget>(SlotContentBindingObjects[0].Get());
+			if (FindGuidForObject(*MovieSceneInstance->GetMovieScene(), *SlotContent).IsValid())
+			{
+				DisplayName = SlotContent->GetParent() != nullptr
+					? FText::Format(LOCTEXT("SlotWithParentFormat", "{0} Slot"), FText::FromString(SlotContent->GetParent()->GetName()))
+					: FText::FromString(SlotContentBindingObjects[0]->GetClass()->GetName());
+			}
+			else
+			{
+				FText PanelName = SlotContent->Slot != nullptr && SlotContent->Slot->Parent != nullptr
+					? FText::FromString(SlotContent->Slot->Parent->GetName())
+					: LOCTEXT("InvalidPanel", "Invalid Panel");
+				FText ContentName = FText::FromString(SlotContent->GetName());
+				DisplayName = FText::Format(LOCTEXT("SlotObject", "{0} ({1} Slot)"), ContentName, PanelName);
+			}
+		}
 		else
 		{
 			DisplayName = LOCTEXT("InvalidSlotContent", "Invalid slot object");
@@ -177,6 +187,19 @@ bool FUMGSequencerObjectBindingManager::TryGetObjectBindingDisplayName(const TSh
 	return true;
 }
 
+UObject* FUMGSequencerObjectBindingManager::GetParentObject( UObject* Object ) const
+{
+	UPanelSlot* Slot = Cast<UPanelSlot>(Object);
+	if (Slot != nullptr)
+	{
+
+		// The slot is actually the child of the panel widget in the hierarchy, but we want it to show up as
+		// a sub-object of the widget it contains in the timeline so we return the content's GUID.
+		return Slot->Content;
+	}
+	return nullptr;
+}
+
 bool FUMGSequencerObjectBindingManager::HasValidWidgetAnimation() const
 {
 	UWidgetBlueprint* WidgetBlueprint = WidgetBlueprintEditor.GetWidgetBlueprintObj();
diff --git a/Engine/Source/Editor/UMGEditor/Private/Animation/UMGSequencerObjectBindingManager.h b/Engine/Source/Editor/UMGEditor/Private/Animation/UMGSequencerObjectBindingManager.h
--- a/Engine/Source/Editor/UMGEditor/Private/Animation/UMGSequencerObjectBindingManager.h
+++ b/Engine/Source/Editor/UMGEditor/Private/Animation/UMGSequencerObjectBindingManager.h
@@ -27,6 +27,7 @@ public:
 	virtual void UnbindPossessableObjects( const FGuid& PossessableGuid ) override;
 	virtual void GetRuntimeObjects( const TSharedRef<FMovieSceneInstance>& MovieSceneInstance, const FGuid& ObjectGuid, TArray<UObject*>& OutRuntimeObjects ) const override;
 	virtual bool TryGetObjectBindingDisplayName(const TSharedRef<FMovieSceneInstance>& MovieSceneInstance, const FGuid& ObjectGuid, FText& DisplayName) const override;
+	virtual UObject* GetParentObject( UObject* Object ) const override;
 	
 	/** @return true if the current animation is valid */
 	bool HasValidWidgetAnimation() const;
diff --git a/Engine/Source/Editor/UMGEditor/Private/WidgetBlueprintEditor.cpp b/Engine/Source/Editor/UMGEditor/Private/WidgetBlueprintEditor.cpp
--- a/Engine/Source/Editor/UMGEditor/Private/WidgetBlueprintEditor.cpp
+++ b/Engine/Source/Editor/UMGEditor/Private/WidgetBlueprintEditor.cpp
@@ -32,6 +32,12 @@ FWidgetBlueprintEditor::FWidgetBlueprintEditor()
 	, bIsRealTime(true)
 {
 	PreviewScene.GetWorld()->bBegunPlay = false;
+
+	// Register sequencer menu extenders.
+	ISequencerModule& SequencerModule = FModuleManager::Get().LoadModuleChecked<ISequencerModule>( "Sequencer" );
+	int32 NewIndex = SequencerModule.GetMenuExtensibilityManager()->GetExtenderDelegates().Add(
+		FAssetEditorExtender::CreateRaw( this, &FWidgetBlueprintEditor::GetContextSensitiveSequencerExtender ));
+	SequencerExtenderHandle = SequencerModule.GetMenuExtensibilityManager()->GetExtenderDelegates()[NewIndex].GetHandle();
 }
 
 FWidgetBlueprintEditor::~FWidgetBlueprintEditor()
@@ -48,6 +54,13 @@ FWidgetBlueprintEditor::~FWidgetBlueprintEditor()
 	Sequencer.Reset();
 
 	SequencerObjectBindingManager.Reset();
+
+	// Un-Register sequencer menu extenders.
+	ISequencerModule& SequencerModule = FModuleManager::Get().LoadModuleChecked<ISequencerModule>("Sequencer");
+	SequencerModule.GetMenuExtensibilityManager()->GetExtenderDelegates().RemoveAll([this]( const FAssetEditorExtender& Extender )
+	{
+		return SequencerExtenderHandle == Extender.GetHandle();
+	});
 }
 
 void FWidgetBlueprintEditor::InitWidgetBlueprintEditor(const EToolkitMode::Type Mode, const TSharedPtr< IToolkitHost >& InitToolkitHost, const TArray<UBlueprint*>& InBlueprints, bool bShouldOpenInDefaultsMode)
@@ -793,4 +806,39 @@ void FWidgetBlueprintEditor::AddObjectToAnimation(UObject* ObjectToAnimate)
 	Sequencer->GetHandleToObject(ObjectToAnimate);
 }
 
+TSharedRef<FExtender> FWidgetBlueprintEditor::GetContextSensitiveSequencerExtender( const TSharedRef<FUICommandList> CommandList, const TArray<UObject*> ContextSensitiveObjects )
+{
+	TSharedRef<FExtender> AddTrackMenuExtender( new FExtender() );
+	AddTrackMenuExtender->AddMenuExtension(
+		SequencerMenuExtensionPoints::AddTrackMenu_PropertiesSection,
+		EExtensionHook::Before,
+		CommandList,
+		FMenuExtensionDelegate::CreateRaw( this, &FWidgetBlueprintEditor::ExtendSequencerAddTrackMenu, ContextSensitiveObjects ) );
+	return AddTrackMenuExtender;
+}
+
+void FWidgetBlueprintEditor::ExtendSequencerAddTrackMenu( FMenuBuilder& AddTrackMenuBuilder, TArray<UObject*> ContextObjects )
+{
+	if ( ContextObjects.Num() == 1 )
+	{
+		UWidget* Widget = Cast<UWidget>( ContextObjects[0] );
+		if ( Widget != nullptr && Widget->GetParent() != nullptr && Widget->Slot != nullptr )
+		{
+			AddTrackMenuBuilder.BeginSection( "Slot", LOCTEXT( "SlotSection", "Slot" ) );
+			{
+				FUIAction AddSlotAction( FExecuteAction::CreateRaw( this, &FWidgetBlueprintEditor::AddSlotTrack, Widget->Slot ) );
+				FText AddSlotLabel = FText::Format(LOCTEXT("SlotLabelFormat", "{0} Slot"), FText::FromString(Widget->GetParent()->GetName()));
+				FText AddSlotToolTip = FText::Format(LOCTEXT("SlotToolTipFormat", "Add {0} slot"), FText::FromString( Widget->GetParent()->GetName()));
+				AddTrackMenuBuilder.AddMenuEntry(AddSlotLabel, AddSlotToolTip, FSlateIcon(), AddSlotAction);
+			}
+			AddTrackMenuBuilder.EndSection();
+		}
+	}
+}
+
+void FWidgetBlueprintEditor::AddSlotTrack( UPanelSlot* Slot )
+{
+	GetSequencer()->GetHandleToObject( Slot );
+}
+
 #undef LOCTEXT_NAMESPACE
diff --git a/Engine/Source/Editor/UMGEditor/Private/WidgetBlueprintEditor.h b/Engine/Source/Editor/UMGEditor/Private/WidgetBlueprintEditor.h
--- a/Engine/Source/Editor/UMGEditor/Private/WidgetBlueprintEditor.h
+++ b/Engine/Source/Editor/UMGEditor/Private/WidgetBlueprintEditor.h
@@ -181,6 +181,15 @@ private:
 	/** Adds the supplied UObject to the current animation. */
 	void AddObjectToAnimation(UObject* ObjectToAnimate);
 
+	/** Gets the extender to use for sequencers context sensitive menus and toolbars. */
+	TSharedRef<FExtender> GetContextSensitiveSequencerExtender( const TSharedRef<FUICommandList> CommandList, const TArray<UObject*> ContextSensitiveObjects );
+
+	/** Extends the sequencer add track menu. */
+	void ExtendSequencerAddTrackMenu( FMenuBuilder& AddTrackMenuBuilder, TArray<UObject*> ContextObjects );
+
+	/** Add an animation track for the supplied slot to the current animation. */
+	void AddSlotTrack( UPanelSlot* Slot );
+
 private:
 	/** The preview scene that owns the preview GUI */
 	FPreviewScene PreviewScene;
@@ -234,4 +243,6 @@ private:
 	bool bIsRealTime;
 
 	TArray< TFunction<void()> > QueuedDesignerActions;
+
+	FDelegateHandle SequencerExtenderHandle;
 };
diff --git a/Engine/Source/Runtime/MovieScene/Private/MovieSceneBindings.cpp b/Engine/Source/Runtime/MovieScene/Private/MovieSceneBindings.cpp
--- a/Engine/Source/Runtime/MovieScene/Private/MovieSceneBindings.cpp
+++ b/Engine/Source/Runtime/MovieScene/Private/MovieSceneBindings.cpp
@@ -43,6 +43,17 @@ FMovieSceneBoundObject& UMovieSceneBindings::AddBinding(const FGuid& PosessableG
 	return BoundObjects.Last();
 }
 
+void UMovieSceneBindings::RemoveBinding( const FGuid& PosessableGuid )
+{
+	for ( int32 i = 0; i < BoundObjects.Num(); i++ )
+	{
+		if ( BoundObjects[i].GetPossessableGuid() == PosessableGuid )
+		{
+			BoundObjects.RemoveAt( i );
+			break;
+		}
+	}
+}
 
 TArray< FMovieSceneBoundObjectInfo > UMovieSceneBindings::FindBoundObjects( const FGuid& Guid ) const
 {
diff --git a/Engine/Source/Runtime/MovieScene/Public/MovieSceneBindings.h b/Engine/Source/Runtime/MovieScene/Public/MovieSceneBindings.h
--- a/Engine/Source/Runtime/MovieScene/Public/MovieSceneBindings.h
+++ b/Engine/Source/Runtime/MovieScene/Public/MovieSceneBindings.h
@@ -136,6 +136,13 @@ public:
 	 */
 	virtual FMovieSceneBoundObject& AddBinding( const FGuid& PosessableGuid, const TArray< FMovieSceneBoundObjectInfo >& ObjectInfos );
 
+	/**
+	* Removes a binding between the specified possessable from a MovieScene asset, to the specified objects
+	*
+	* @param	PossessableGuid		The Guid of the possessable from the MovieScene asset to remove
+	*/
+	virtual void RemoveBinding( const FGuid& PosessableGuid );
+
 	/**
 	 * Given a guid for a possessable, attempts to find the object bound to that guid
 	 *
