MovieScene: Merging some more non-functional refactor changes

diff --git a/Engine/Source/Editor/Sequencer/Private/SAnimationOutlinerView.cpp b/Engine/Source/Editor/Sequencer/Private/SAnimationOutlinerView.cpp
--- a/Engine/Source/Editor/Sequencer/Private/SAnimationOutlinerView.cpp
+++ b/Engine/Source/Editor/Sequencer/Private/SAnimationOutlinerView.cpp
@@ -9,7 +9,6 @@
 #include "MovieSceneSection.h"
 #include "MovieSceneCommonHelpers.h"
 #include "Engine/Selection.h"
-#include "ISequencerObjectBindingManager.h"
 #include "IKeyArea.h"
 
 #define LOCTEXT_NAMESPACE "AnimationOutliner"
diff --git a/Engine/Source/Editor/Sequencer/Private/Sequencer.cpp b/Engine/Source/Editor/Sequencer/Private/Sequencer.cpp
--- a/Engine/Source/Editor/Sequencer/Private/Sequencer.cpp
+++ b/Engine/Source/Editor/Sequencer/Private/Sequencer.cpp
@@ -465,42 +465,35 @@ void FSequencer::OnMapChanged( class UWorld* NewWorld, EMapChangeType::Type MapC
 void FSequencer::OnActorsDropped( const TArray<TWeakObjectPtr<AActor> >& Actors )
 {
 	bool bPossessableAdded = false;
-	for( TWeakObjectPtr<AActor> WeakActor : Actors )
+	UMovieScene* OwnerMovieScene = GetFocusedMovieScene();
+
+	for (TWeakObjectPtr<AActor> WeakActor : Actors)
 	{
 		AActor* Actor = WeakActor.Get();
 
 		if (Actor != nullptr)
-		{
-			// Grab the MovieScene that is currently focused.  We'll add our Blueprint as an inner of the
-			// MovieScene asset.
-			UMovieScene* OwnerMovieScene = GetFocusedMovieScene();
-			
+		{		
 			// @todo sequencer: Undo doesn't seem to be working at all
-			const FScopedTransaction Transaction( LOCTEXT("UndoPossessingObject", "Possess Object with MovieScene") );
+			const FScopedTransaction Transaction(LOCTEXT("UndoPossessingObject", "Possess Object with MovieScene"));
 			
-			// Possess the object!
-			{
-				// Create a new possessable
-				OwnerMovieScene->Modify();
+			// possess the object!
+			OwnerMovieScene->Modify();
 				
-				const FGuid PossessableGuid = OwnerMovieScene->AddPossessable( Actor->GetActorLabel(), Actor->GetClass() );
+			const FGuid PossessableGuid = OwnerMovieScene->AddPossessable(Actor->GetActorLabel(), Actor->GetClass());
 				 
-				if ( IsShotFilteringOn() )
-				{
-					AddUnfilterableObject(PossessableGuid);
-				}
-				
-				ObjectBindingManager->BindPossessableObject( PossessableGuid, *Actor );
-				
-				bPossessableAdded = true;
+			if (IsShotFilteringOn())
+			{
+				AddUnfilterableObject(PossessableGuid);
 			}
+				
+			ObjectBindingManager->BindPossessableObject(PossessableGuid, *Actor);			
+			bPossessableAdded = true;
 		}
 	}
 	
-	if( bPossessableAdded )
+	if (bPossessableAdded)
 	{
-		SpawnOrDestroyPuppetObjects( GetFocusedMovieSceneInstance() );
-		
+		SpawnOrDestroyPuppetObjects(GetFocusedMovieSceneInstance());
 		NotifyMovieSceneDataChanged();
 	}
 }
