MovieScene: Fixed comments and formatting

diff --git a/Engine/Source/Editor/Sequencer/Public/ISequencerModule.h b/Engine/Source/Editor/Sequencer/Public/ISequencerModule.h
--- a/Engine/Source/Editor/Sequencer/Public/ISequencerModule.h
+++ b/Engine/Source/Editor/Sequencer/Public/ISequencerModule.h
@@ -4,94 +4,104 @@
 
 #include "ModuleInterface.h"
 #include "ISequencer.h"
-
 #include "Toolkits/AssetEditorToolkit.h"
 #include "Toolkits/IToolkit.h"
 
+
 class ISequencerObjectBindingManager;
 
+
 namespace SequencerMenuExtensionPoints
 {
 	static const FName AddTrackMenu_PropertiesSection("AddTrackMenu_PropertiesSection");
 }
 
-/**
- * A delegate which will create an auto-key handler
- */
-DECLARE_DELEGATE_RetVal_OneParam( TSharedRef<class FMovieSceneTrackEditor>, FOnCreateTrackEditor, TSharedRef<ISequencer> );
 
-DECLARE_DELEGATE_RetVal_OneParam( TSharedRef<SWidget>, FOnGetAddMenuContent, TSharedRef<ISequencer> );
+/** A delegate which will create an auto-key handler. */
+DECLARE_DELEGATE_RetVal_OneParam(TSharedRef<class FMovieSceneTrackEditor>, FOnCreateTrackEditor, TSharedRef<ISequencer>);
 
-/** Parameters for initializing a Sequencer */
+/** A delegate that is executed when adding menu content. */
+DECLARE_DELEGATE_RetVal_OneParam(TSharedRef<SWidget>, FOnGetAddMenuContent, TSharedRef<ISequencer>);
+
+
+/**
+ * Sequencer initialization parameters.
+ */
 struct FSequencerViewParams
 {
-	/** The initial time range of the view */
+	/** The initial time range of the view. */
 	TRange<float> InitalViewRange;
 
-	/** Initial Scrub Position */
+	/** Initial Scrub Position. */
 	float InitialScrubPosition;
 
 	FOnGetAddMenuContent OnGetAddMenuContent;
 
-	/** Unique name for the sequencer */
+	/** Unique name for the sequencer. */
 	FString UniqueName;
 
 	FSequencerViewParams(FString InName = FString())
-		: InitalViewRange( 0.0f, 5.0f )
-		, InitialScrubPosition( 0.0f )
+		: InitalViewRange(0.0f, 5.0f)
+		, InitialScrubPosition(0.0f)
 		, UniqueName(MoveTemp(InName))
-	{}
+	{ }
 };
 
+
 /**
- * The public interface of SequencerModule
+ * Interface for the Sequencer module.
  */
-class ISequencerModule : public IModuleInterface
+class ISequencerModule
+	: public IModuleInterface
 {
-
 public:
 
 	/**
-	 * Creates a new instance of a standalone sequencer that can be added to other UIs
+	 * Creates a new instance of a standalone sequencer that can be added to other UIs.
 	 *
-	 * @param 	InRootMovieScene	The movie scene to edit
-	 * @param	ViewParams			Parameters for how to view sequencer UI
-	 * @return	Interface to the new editor
+	 * @param InRootMovieScene The movie scene to edit.
+	 * @param ViewParams Parameters for how to view sequencer UI.
+	 * @return Interface to the new editor.
 	 */
-	virtual TSharedPtr<ISequencer> CreateSequencer( UMovieScene* InRootMovieScene, const FSequencerViewParams& InViewParams, TSharedRef<ISequencerObjectBindingManager> ObjectBindingManager ) = 0;
+	virtual TSharedPtr<ISequencer> CreateSequencer(UMovieScene* InRootMovieScene, const FSequencerViewParams& InViewParams, TSharedRef<ISequencerObjectBindingManager> ObjectBindingManager) = 0;
 
 	/**
 	 * Creates a new instance of a Sequencer, the editor for MovieScene assets in an asset editor
 	 *
-	 * @param	Mode					Mode that this editor should operate in
-	 * @param	ViewParams				Parameters for how to view sequencer UI
-	 * @param	InitToolkitHost			When Mode is WorldCentric, this is the level editor instance to spawn this editor within
-	 * @param	ObjectToEdit			The object to start editing
-	 *
-	 * @return	Interface to the new editor
+	 * @param Mode Mode that this editor should operate in.
+	 * @param ViewParams Parameters for how to view sequencer UI.
+	 * @param InitToolkitHost When Mode is WorldCentric, this is the level editor instance to spawn this editor within.
+	 * @param ObjectToEdit The object to start editing.
+	 * @return Interface to the new editor.
 	 */
-	virtual TSharedPtr<ISequencer> CreateSequencerAssetEditor( const EToolkitMode::Type Mode, const FSequencerViewParams& InViewParams, const TSharedPtr< class IToolkitHost >& InitToolkitHost, UMovieScene* InRootMovieScene, bool bEditWithinLevelEditor ) = 0;
+	virtual TSharedPtr<ISequencer> CreateSequencerAssetEditor(const EToolkitMode::Type Mode, const FSequencerViewParams& InViewParams, const TSharedPtr< class IToolkitHost >& InitToolkitHost, UMovieScene* InRootMovieScene, bool bEditWithinLevelEditor) = 0;
 
 	/** 
-	 * Registers a delegate that will create an editor for a track in each sequencer 
-	 *
-	 * @param InOnCreateTrackEditor	Delegate to register
+	 * Registers a delegate that will create an editor for a track in each sequencer.
 	 *
+	 * @param InOnCreateTrackEditor	Delegate to register.
 	 * @return A handle to the newly-added delegate.
 	 */
-	virtual FDelegateHandle RegisterTrackEditor_Handle( FOnCreateTrackEditor InOnCreateTrackEditor ) = 0;
+	virtual FDelegateHandle RegisterTrackEditor_Handle(FOnCreateTrackEditor InOnCreateTrackEditor) = 0;
 
 	/** 
 	 * Unregisters a previously registered delegate for creating a track editor
 	 *
 	 * @param InHandle	Handle to the delegate to unregister
 	 */
-	virtual void UnRegisterTrackEditor_Handle( FDelegateHandle InHandle ) = 0;
+	virtual void UnRegisterTrackEditor_Handle(FDelegateHandle InHandle) = 0;
 
-	/** Gets the extensibility manager for menus */
+	/**
+	 * Get the extensibility manager for menus.
+	 *
+	 * @return Menu extensibility manager.
+	 */
 	virtual TSharedPtr<FExtensibilityManager> GetMenuExtensibilityManager() const = 0;
 
-	/** Gets the extensibility manager for toolbars */
+	/**
+	 * Get the extensibility manager for toolbars.
+	 *
+	 * @return Toolbar extensibility manager.
+	 */
 	virtual TSharedPtr<FExtensibilityManager> GetToolBarExtensibilityManager() const = 0;
 };
-
