MovieScene: Fixed comments and formatting

diff --git a/Engine/Source/Editor/Sequencer/Public/ISequencer.h b/Engine/Source/Editor/Sequencer/Public/ISequencer.h
--- a/Engine/Source/Editor/Sequencer/Public/ISequencer.h
+++ b/Engine/Source/Editor/Sequencer/Public/ISequencer.h
@@ -7,16 +7,23 @@
 #include "Editor/SequencerWidgets/Public/ITimeSlider.h"
 #include "KeyPropertyParams.h"
 
-class UMovieScene;
+
 class FSequencerSelection;
+class ISequencerObjectBindingManager;
+class UAnimSequence;
+class UMovieScene;
+class UMovieSceneSection;
+
 
 /**
- * Sequencer public interface
+ * Interface for sequencers.
  */
-class ISequencer : public IMovieScenePlayer, public TSharedFromThis<ISequencer>
+class ISequencer
+	: public IMovieScenePlayer
+	, public TSharedFromThis<ISequencer>
 {
-
 public:
+
 	/** @return Widget used to display the sequencer */
 	virtual TSharedRef<SWidget> GetSequencerWidget() const = 0;
 
@@ -32,29 +39,29 @@ public:
 	virtual TSharedRef<FMovieSceneInstance> GetFocusedMovieSceneInstance() const = 0;
 
 	/** Resets sequencer with a new RootMovieScene */
-	virtual void ResetToNewRootMovieScene( UMovieScene& NewRoot, TSharedRef<class ISequencerObjectBindingManager> NewObjectBindingManager ) = 0;
+	virtual void ResetToNewRootMovieScene(UMovieScene& NewRoot, TSharedRef<ISequencerObjectBindingManager> NewObjectBindingManager) = 0;
 
 	/**
-	 * Focuses a sub-movie scene (MovieScene within a MovieScene) in the sequencer
+	 * Focuses a sub-movie scene (MovieScene within a MovieScene) in the sequencer.
 	 * 
-	 * @param SubMovieSceneInstance	Sub-MovieScene instance to focus
+	 * @param SubMovieSceneInstance	Sub-MovieScene instance to focus.
 	 */
-	virtual void FocusSubMovieScene( TSharedRef<FMovieSceneInstance> SubMovieSceneInstance ) = 0;
+	virtual void FocusSubMovieScene(TSharedRef<FMovieSceneInstance> SubMovieSceneInstance) = 0;
 	
 	/**
-	 * Given a sub-movie scene section, returns the instance of the movie scene for that section
+	 * Given a sub-movie scene section, returns the instance of the movie scene for that section.
 	 *
-	 * @param SubMovieSceneSection	The sub-movie scene section containing a movie scene to get the instance for
-	 * @return The instance for the sub-moviescene
+	 * @param SubMovieSceneSection	The sub-movie scene section containing a movie scene to get the instance for.
+	 * @return The instance for the sub-movie scene
 	 */
-	virtual TSharedRef<FMovieSceneInstance> GetInstanceForSubMovieSceneSection( UMovieSceneSection& SubMovieSceneSection ) const = 0;
+	virtual TSharedRef<FMovieSceneInstance> GetInstanceForSubMovieSceneSection(UMovieSceneSection& SubMovieSceneSection) const = 0;
 
 	/**
 	 * Adds a movie scene as a section inside the current movie scene
 	 * 
-	 * @param SubMovieScene The moviescene to add
+	 * @param SubMovieScene The movie scene to add.
 	 */
-	virtual void AddSubMovieScene( UMovieScene* SubMovieScene ) = 0;
+	virtual void AddSubMovieScene(UMovieScene* SubMovieScene) = 0;
 
 	/** @return Returns whether auto-key is enabled in this sequencer */
 	virtual bool GetAutoKeyEnabled() const = 0;
@@ -66,60 +73,71 @@ public:
 	virtual bool IsRecordingLive() const = 0;
 
 	/**
-	 * Sets which shot sections need to be filtered to, or none to disable shot filtering
+	 * Sets which shot sections need to be filtered to, or none to disable shot filtering.
 	 *
-	 * @param ShotSections					The shot sections to filter to
-	 * @param bZoomToShotBounds				Whether or not to zoom time to the bounds of the shots 
+	 * @param ShotSections The shot sections to filter.
+	 * @param bZoomToShotBounds Whether or not to zoom time to the bounds of the shots.
+	 * @see FilterToSelectedShotSections
 	 */
-	virtual void FilterToShotSections(const TArray< TWeakObjectPtr<class UMovieSceneSection> >& ShotSections, bool bZoomToShotBounds = true ) = 0;
+	virtual void FilterToShotSections(const TArray<TWeakObjectPtr<UMovieSceneSection>>& ShotSections, bool bZoomToShotBounds = true) = 0;
 	
 	/**
-	 * Sets shot filtering to current selection
+	 * Sets shot filtering to current selection.
 	 *
-	 * @param bZoomToShotBounds				Whether or not to zoom time to the bounds of the shots 
+	 * @param bZoomToShotBounds Whether or not to zoom time to the bounds of the shots.
+	 * @see FilterToShotSections
 	 */
-	virtual void FilterToSelectedShotSections(bool bZoomToShotBounds = true ) = 0;
+	virtual void FilterToSelectedShotSections(bool bZoomToShotBounds = true) = 0;
 	
 	/** Adds a new shot key with the specified camera guid */
-	virtual void AddNewShot( FGuid CameraGuid ) = 0;
+	virtual void AddNewShot(FGuid CameraGuid) = 0;
 
 	/** Adds a new animation to the specified skeletal mesh object guid */
-	virtual void AddAnimation( FGuid ObjectGuid, class UAnimSequence* AnimSequence ) = 0;
+	virtual void AddAnimation(FGuid ObjectGuid, UAnimSequence* AnimSequence) = 0;
 
 	/**
 	 * Gets the current time of the time slider relative to the passed in movie scene
 	 *
-	 * @param MovieScene	The movie scene to get the local time for.  The movie scene must exist in the sequencer
+	 * @param MovieScene The movie scene to get the local time for (must exist in this sequencer).
 	 */
-	virtual float GetCurrentLocalTime( UMovieScene& MovieScene ) = 0;
+	virtual float GetCurrentLocalTime(UMovieScene& MovieScene) = 0;
 	
 	/**
 	 * Gets the global time.
+	 *
+	 * @return Global time in seconds.
+	 * @see SetGlobalTime
 	 */
 	virtual float GetGlobalTime() = 0;
 
 	/**
 	 * Sets the global position to the time.
 	 *
-	 * @param Time The global time to set to.
+	 * @param Time The global time to set.
+	 * @see GetGlobalTime
 	 */
 	virtual void SetGlobalTime(float Time) = 0;
 
 	/** @return The current view range */
-	virtual FAnimatedRange GetViewRange() const { return FAnimatedRange(); }
+	virtual FAnimatedRange GetViewRange() const
+	{
+		return FAnimatedRange();
+	}
 
 	/**
 	 * Sets whether perspective viewport hijacking is enabled.
+	 *
+	 * @param bEnabled true if the viewport should be enabled, false if it should be disabled.
 	 */
 	virtual void SetPerspectiveViewportPossessionEnabled(bool bEnabled) = 0;
 
 	/**
 	 * Gets a handle to runtime information about the object being manipulated by a movie scene
 	 * 
-	 * @param Object	The object to get a handle for
-	 * @return The handle to the object
+	 * @param Object The object to get a handle for.
+	 * @return The handle to the object.
 	 */
-	virtual FGuid GetHandleToObject( UObject* Object ) = 0;
+	virtual FGuid GetHandleToObject(UObject* Object) = 0;
 
 	/**
 	 * @return Returns the object change listener for sequencer instance
