Miscellaneous sequencer improvements

- Sequencer tree is now an STreeView which gives us proper virtualized scroll, and allows removal of lots of manually implemented tree-management code.
  - Selection and expansion is now handled by STreeView itself.
  - Track lanes are now also virtualized vertically, so we should only see section area widgets when they're actually on screen. Positioning these elements is a frame behind the tree rows.
- Added Traversal methods for more expressive traversal of the sequencer tree (See FSequencerDisplayNode::Traverse(Visible)_(Child|Parent)First).
- Removed outliner backgrtounds and replaced with more subtle separator lines.
- Added Marquee selection for sequencer keys. Currently this is not implented when LMB dragging on a SSection.
- Added RMB panning for the sequencer track area. This currently has no inertial scroll, but could be added in future.

diff --git a/Engine/Source/Editor/MovieSceneTools/Private/TrackEditors/TransformTrackEditor.cpp b/Engine/Source/Editor/MovieSceneTools/Private/TrackEditors/TransformTrackEditor.cpp
--- a/Engine/Source/Editor/MovieSceneTools/Private/TrackEditors/TransformTrackEditor.cpp
+++ b/Engine/Source/Editor/MovieSceneTools/Private/TrackEditors/TransformTrackEditor.cpp
@@ -406,6 +406,7 @@ void F3DTransformTrackEditor::BuildObjectBindingEditButtons(TSharedPtr<SHorizont
 		.IsChecked(this, &F3DTransformTrackEditor::IsCameraLocked, CameraActor)
 		.OnCheckStateChanged(this, &F3DTransformTrackEditor::OnLockCameraClicked, CameraActor)
 		.ToolTipText(this, &F3DTransformTrackEditor::GetLockCameraToolTip, CameraActor)
+		.ForegroundColor(FLinearColor::White)
 		.CheckedImage(FEditorStyle::GetBrush("Sequencer.LockCamera"))
 		.CheckedHoveredImage(FEditorStyle::GetBrush("Sequencer.LockCamera"))
 		.CheckedPressedImage(FEditorStyle::GetBrush("Sequencer.LockCamera"))
diff --git a/Engine/Source/Editor/Sequencer/Private/SAnimationOutlinerTreeNode.cpp b/Engine/Source/Editor/Sequencer/Private/SAnimationOutlinerTreeNode.cpp
new file mode 100644
--- /dev/null
+++ b/Engine/Source/Editor/Sequencer/Private/SAnimationOutlinerTreeNode.cpp
@@ -0,0 +1,524 @@
+// Copyright 1998-2015 Epic Games, Inc. All Rights Reserved.
+
+#include "SequencerPrivatePCH.h"
+
+#include "SAnimationOutlinerTreeNode.h"
+#include "ScopedTransaction.h"
+#include "Sequencer.h"
+#include "MovieScene.h"
+#include "MovieSceneSection.h"
+#include "MovieSceneCommonHelpers.h"
+#include "Engine/Selection.h"
+#include "IKeyArea.h"
+#include "SSequencerTreeView.h"
+
+#define LOCTEXT_NAMESPACE "AnimationOutliner"
+
+
+void SAnimationOutlinerTreeNode::Construct( const FArguments& InArgs, TSharedRef<FSequencerDisplayNode> Node, const TSharedRef<SSequencerTreeViewRow>& InTableRow, FSequencer* InSequencer  )
+{
+	Sequencer = InSequencer;
+	DisplayNode = Node;
+
+	SelectedBrush = FEditorStyle::GetBrush( "Sequencer.AnimationOutliner.SelectionBorder" );
+	SelectedBrushInactive = FEditorStyle::GetBrush("Sequencer.AnimationOutliner.SelectionBorderInactive");
+	NotSelectedBrush = FEditorStyle::GetBrush( "NoBorder" );
+	TableRowStyle = &FEditorStyle::Get().GetWidgetStyle<FTableRowStyle>("TableView.Row");
+
+	// Choose the font.  If the node is a root node or an object node, we show a larger font for it.
+	FSlateFontInfo NodeFont = Node->GetParent().IsValid() && Node->GetType() != ESequencerNode::Object ? 
+		FEditorStyle::GetFontStyle("Sequencer.AnimationOutliner.RegularFont") :
+		FEditorStyle::GetFontStyle("Sequencer.AnimationOutliner.BoldFont");
+
+	TSharedRef<SWidget> TextWidget = 
+		SNew( STextBlock )
+		.Text(this, &SAnimationOutlinerTreeNode::GetDisplayName )
+		.Font( NodeFont );
+
+	auto NodeHeight = [=]() -> FOptionalSize { return DisplayNode->GetNodeHeight(); };
+
+	TAttribute<FLinearColor> HoverTint(this, &SAnimationOutlinerTreeNode::GetHoverTint);
+	ForegroundColor.Bind(this, &SAnimationOutlinerTreeNode::GetForegroundBasedOnSelection);
+
+	TSharedRef<SWidget>	FinalWidget = 
+		SNew( SBorder )
+		.VAlign( VAlign_Center )
+		.BorderImage( this, &SAnimationOutlinerTreeNode::GetNodeBorderImage )
+		.Padding(FMargin(0, Node->GetNodePadding().Combined() / 2))
+		[
+			SNew(SBox)
+			.HeightOverride_Lambda(NodeHeight)
+			[
+				SNew( SHorizontalBox )
+
+				// Expand track lanes button
+				+ SHorizontalBox::Slot()
+				.Padding(FMargin(2.f, 0.f))
+				.VAlign( VAlign_Center )
+				.AutoWidth()
+				[
+					SNew(SExpanderArrow, InTableRow).IndentAmount(SequencerLayoutConstants::IndentAmount)
+				]
+
+				// Label Slot
+				+ SHorizontalBox::Slot()
+				.VAlign( VAlign_Center )
+				.FillWidth(3)
+				[
+					TextWidget
+				]
+
+				// Editor slot
+				+ SHorizontalBox::Slot()
+				.VAlign(VAlign_Center)
+				.AutoWidth()
+				[
+					// @todo Sequencer - Remove this box and width override.
+					SNew(SBox)
+					.WidthOverride(100)
+					[
+						DisplayNode->GenerateEditWidgetForOutliner()
+					]
+				]
+				// Previous key slot
+				+ SHorizontalBox::Slot()
+				.VAlign(VAlign_Center)
+				.AutoWidth()
+				.Padding(3, 0, 0, 0)
+				[
+					SNew(SBorder)
+					.Padding(0)
+					.BorderImage(NotSelectedBrush)
+					.ColorAndOpacity(HoverTint)
+					[
+						SNew(SButton)
+						.ButtonStyle(FEditorStyle::Get(), "FlatButton")
+						.ToolTipText(LOCTEXT("PreviousKeyButton", "Set the time to the previous key"))
+						.OnClicked(this, &SAnimationOutlinerTreeNode::OnPreviousKeyClicked)
+						.ForegroundColor( FSlateColor::UseForeground() )
+						.ContentPadding(0)
+						[
+							SNew(STextBlock)
+							.Font(FEditorStyle::Get().GetFontStyle("FontAwesome.7"))
+							.Text(FText::FromString(FString(TEXT("\xf060"))) /*fa-arrow-left*/)
+						]
+					]
+				]
+				// Add key slot
+				+ SHorizontalBox::Slot()
+				.VAlign(VAlign_Center)
+				.AutoWidth()
+				[
+					SNew(SBorder)
+					.Padding(0)
+					.BorderImage(NotSelectedBrush)
+					.ColorAndOpacity(HoverTint)
+					[
+						SNew(SButton)
+						.ButtonStyle(FEditorStyle::Get(), "FlatButton")
+						.ToolTipText(LOCTEXT("AddKeyButton", "Add a new key at the current time"))
+						.OnClicked(this, &SAnimationOutlinerTreeNode::OnAddKeyClicked)
+						.ForegroundColor( FSlateColor::UseForeground() )
+						.ContentPadding(0)
+						[
+							SNew(STextBlock)
+							.Font(FEditorStyle::Get().GetFontStyle("FontAwesome.7"))
+							.Text(FText::FromString(FString(TEXT("\xf055"))) /*fa-plus-circle*/)
+						]
+					]
+				]
+				// Next key slot
+				+ SHorizontalBox::Slot()
+				.VAlign(VAlign_Center)
+				.AutoWidth()
+				[
+					SNew(SBorder)
+					.Padding(0)
+					.BorderImage(NotSelectedBrush)
+					.ColorAndOpacity(HoverTint)
+					[
+						SNew(SButton)
+						.ButtonStyle(FEditorStyle::Get(), "FlatButton")
+						.ToolTipText(LOCTEXT("NextKeyButton", "Set the time to the next key"))
+						.OnClicked(this, &SAnimationOutlinerTreeNode::OnNextKeyClicked)
+						.ContentPadding(0)
+						.ForegroundColor( FSlateColor::UseForeground() )
+						[
+							SNew(STextBlock)
+							.Font(FEditorStyle::Get().GetFontStyle("FontAwesome.7"))
+							.Text(FText::FromString(FString(TEXT("\xf061"))) /*fa-arrow-right*/)
+						]
+					]
+				]
+			]
+		];
+
+	ChildSlot
+	[
+		FinalWidget
+	];
+}
+
+FLinearColor SAnimationOutlinerTreeNode::GetHoverTint() const
+{
+	return IsHovered() ? FLinearColor(1,1,1,0.9f) : FLinearColor(1,1,1,0.4f);
+}
+
+
+
+FReply SAnimationOutlinerTreeNode::OnPreviousKeyClicked()
+{
+	FSequencer& Sequencer = DisplayNode->GetSequencer();
+	float ClosestPreviousKeyDistance = MAX_FLT;
+	float CurrentTime = Sequencer.GetCurrentLocalTime(*Sequencer.GetFocusedMovieScene());
+	float PreviousTime = 0;
+	bool PreviousKeyFound = false;
+
+	TSet<TSharedPtr<IKeyArea>> KeyAreas;
+	SequencerHelpers::GetAllKeyAreas(DisplayNode, KeyAreas);
+	for (TSharedPtr<IKeyArea> KeyArea : KeyAreas)
+	{
+		for (FKeyHandle& KeyHandle : KeyArea->GetUnsortedKeyHandles())
+		{
+			float KeyTime = KeyArea->GetKeyTime(KeyHandle);
+			if (KeyTime < CurrentTime && CurrentTime - KeyTime < ClosestPreviousKeyDistance)
+			{
+				PreviousTime = KeyTime;
+				ClosestPreviousKeyDistance = CurrentTime - KeyTime;
+				PreviousKeyFound = true;
+			}
+		}
+	}
+
+	if (PreviousKeyFound)
+	{
+		Sequencer.SetGlobalTime(PreviousTime);
+	}
+	return FReply::Handled();
+}
+
+FReply SAnimationOutlinerTreeNode::OnNextKeyClicked()
+{
+	FSequencer& Sequencer = DisplayNode->GetSequencer();
+	float ClosestNextKeyDistance = MAX_FLT;
+	float CurrentTime = Sequencer.GetCurrentLocalTime(*Sequencer.GetFocusedMovieScene());
+	float NextTime = 0;
+	bool NextKeyFound = false;
+
+	TSet<TSharedPtr<IKeyArea>> KeyAreas;
+	SequencerHelpers::GetAllKeyAreas(DisplayNode, KeyAreas);
+	for (TSharedPtr<IKeyArea> KeyArea : KeyAreas)
+	{
+		for (FKeyHandle& KeyHandle : KeyArea->GetUnsortedKeyHandles())
+		{
+			float KeyTime = KeyArea->GetKeyTime(KeyHandle);
+			if (KeyTime > CurrentTime && KeyTime - CurrentTime < ClosestNextKeyDistance)
+			{
+				NextTime = KeyTime;
+				ClosestNextKeyDistance = KeyTime - CurrentTime;
+				NextKeyFound = true;
+			}
+		}
+	}
+
+	if (NextKeyFound)
+	{
+		Sequencer.SetGlobalTime(NextTime);
+	}
+	return FReply::Handled();
+}
+
+FReply SAnimationOutlinerTreeNode::OnAddKeyClicked()
+{
+	FSequencer& Sequencer = DisplayNode->GetSequencer();
+	float CurrentTime = Sequencer.GetCurrentLocalTime(*Sequencer.GetFocusedMovieScene());
+
+	TSet<TSharedPtr<IKeyArea>> KeyAreas;
+	SequencerHelpers::GetAllKeyAreas(DisplayNode, KeyAreas);
+
+	TArray<UMovieSceneSection*> KeyAreaSections;
+	for (TSharedPtr<IKeyArea> KeyArea : KeyAreas)
+	{
+		UMovieSceneSection* OwningSection = KeyArea->GetOwningSection();
+		KeyAreaSections.Add(OwningSection);
+	}
+
+	UMovieSceneSection* NearestSection = MovieSceneHelpers::FindNearestSectionAtTime(KeyAreaSections, CurrentTime);
+	if (!NearestSection)
+	{
+		return FReply::Unhandled();
+	}
+
+	FScopedTransaction Transaction(LOCTEXT("AddKeys", "Add keys at current time"));
+	for (TSharedPtr<IKeyArea> KeyArea : KeyAreas)
+	{
+		UMovieSceneSection* OwningSection = KeyArea->GetOwningSection();
+		if (OwningSection == NearestSection)
+		{
+			OwningSection->SetFlags(RF_Transactional);
+			OwningSection->Modify();
+			KeyArea->AddKeyUnique(CurrentTime);
+		}
+	}
+
+	return FReply::Handled();
+}
+
+void SAnimationOutlinerTreeNode::GetAllDescendantNodes(TSharedPtr<FSequencerDisplayNode> RootNode, TArray<TSharedRef<FSequencerDisplayNode> >& AllNodes)
+{
+	if (!RootNode.IsValid())
+	{
+		return;
+	}
+
+	AllNodes.Add(RootNode.ToSharedRef());
+
+	const FSequencerDisplayNode* RootNodeC = RootNode.Get();
+
+	for (TSharedRef<FSequencerDisplayNode> ChildNode : RootNodeC->GetChildNodes())
+	{
+		AllNodes.Add(ChildNode);
+
+		GetAllDescendantNodes(ChildNode, AllNodes);
+	}
+}
+
+FReply SAnimationOutlinerTreeNode::OnMouseButtonDown( const FGeometry& MyGeometry, const FPointerEvent& MouseEvent )
+{
+	if( MouseEvent.GetEffectingButton() == EKeys::LeftMouseButton && DisplayNode->IsSelectable() )
+	{
+		FSequencer& Sequencer = DisplayNode->GetSequencer();
+		bool bSelected = Sequencer.GetSelection().IsSelected(DisplayNode.ToSharedRef());
+
+		TArray<TSharedPtr<FSequencerDisplayNode> > AffectedNodes;
+		AffectedNodes.Add(DisplayNode.ToSharedRef());
+
+		if (MouseEvent.IsShiftDown())
+		{
+			FSequencerNodeTree& ParentTree = DisplayNode->GetParentTree();
+
+			const TArray< TSharedRef<FSequencerDisplayNode> >RootNodes = ParentTree.GetRootNodes();
+
+			// Get all nodes in order
+			TArray<TSharedRef<FSequencerDisplayNode> > AllNodes;
+			for (int32 i = 0; i < RootNodes.Num(); ++i)
+			{
+				GetAllDescendantNodes(RootNodes[i], AllNodes);
+			}
+
+			int32 FirstIndexToSelect = INT32_MAX;
+			int32 LastIndexToSelect = INT32_MIN;
+
+			for (int32 ChildIndex = 0; ChildIndex < AllNodes.Num(); ++ChildIndex)
+			{
+				TSharedRef<FSequencerDisplayNode> ChildNode = AllNodes[ChildIndex];
+
+				if (ChildNode == DisplayNode.ToSharedRef() || Sequencer.GetSelection().IsSelected(ChildNode))
+				{
+					if (ChildIndex < FirstIndexToSelect)
+					{
+						FirstIndexToSelect = ChildIndex;
+					}
+					if (ChildIndex > LastIndexToSelect)
+					{
+						LastIndexToSelect = ChildIndex;
+					}
+				}
+			}
+
+			if (FirstIndexToSelect != INT32_MAX && LastIndexToSelect != INT32_MIN)
+			{
+				for (int32 ChildIndex = FirstIndexToSelect; ChildIndex <= LastIndexToSelect; ++ChildIndex)
+				{
+					TSharedRef<FSequencerDisplayNode> ChildNode = AllNodes[ChildIndex];
+
+					if (!Sequencer.GetSelection().IsSelected(ChildNode))
+					{
+						Sequencer.GetSelection().AddToSelection(ChildNode);
+						AffectedNodes.Add(ChildNode);
+					}
+				}
+			}
+		}
+		else if( MouseEvent.IsControlDown() )
+		{
+			// Toggle selection when control is down
+			if (bSelected)
+			{
+				Sequencer.GetSelection().RemoveFromSelection(DisplayNode.ToSharedRef());
+			}
+			else
+			{
+				Sequencer.GetSelection().AddToSelection(DisplayNode.ToSharedRef());
+			}
+		}
+		else
+		{
+			// Deselect the other nodes and select this node.
+			Sequencer.GetSelection().EmptySelectedOutlinerNodes();
+			Sequencer.GetSelection().AddToSelection(DisplayNode.ToSharedRef());
+		}
+
+		OnSelectionChanged( AffectedNodes );
+		return FReply::Handled();
+	}
+
+	return FReply::Unhandled();
+}
+
+FReply SAnimationOutlinerTreeNode::OnMouseButtonUp( const FGeometry& MyGeometry, const FPointerEvent& MouseEvent )
+{
+	if( MouseEvent.GetEffectingButton() == EKeys::RightMouseButton )
+	{
+		TSharedPtr<SWidget> MenuContent = DisplayNode->OnSummonContextMenu(MyGeometry, MouseEvent);
+		if (MenuContent.IsValid())
+		{
+			FWidgetPath WidgetPath = MouseEvent.GetEventPath() != nullptr ? *MouseEvent.GetEventPath() : FWidgetPath();
+
+			FSlateApplication::Get().PushMenu(
+				AsShared(),
+				WidgetPath,
+				MenuContent.ToSharedRef(),
+				MouseEvent.GetScreenSpacePosition(),
+				FPopupTransitionEffect( FPopupTransitionEffect::ContextMenu )
+				);
+			
+			return FReply::Handled().SetUserFocus(MenuContent.ToSharedRef(), EFocusCause::SetDirectly);
+		}
+
+		return FReply::Handled();
+	}
+
+	return FReply::Unhandled();
+}
+
+const FSlateBrush* SAnimationOutlinerTreeNode::GetNodeBorderImage() const
+{
+	// Display a highlight when the node is selected
+	FSequencer& Sequencer = DisplayNode->GetSequencer();
+	const bool bIsSelected = Sequencer.GetSelection().IsSelected(DisplayNode.ToSharedRef());
+	if (bIsSelected)
+	{
+		if (Sequencer.GetSelection().GetActiveSelection() == FSequencerSelection::EActiveSelection::OutlinerNode)
+		{
+			return SelectedBrush;
+		}
+		else
+		{
+			return SelectedBrushInactive;
+		}
+	}
+	else
+	{
+		return  NotSelectedBrush;
+	}
+}
+
+FSlateColor SAnimationOutlinerTreeNode::GetForegroundBasedOnSelection() const
+{
+	FSequencer& Sequencer = DisplayNode->GetSequencer();
+	const bool bIsSelected = Sequencer.GetSelection().IsSelected(DisplayNode.ToSharedRef());
+
+	return bIsSelected ? TableRowStyle->SelectedTextColor : TableRowStyle->TextColor;
+}
+
+EVisibility SAnimationOutlinerTreeNode::GetExpanderVisibility() const
+{
+	return DisplayNode->GetNumChildren() > 0 ? EVisibility::Visible : EVisibility::Hidden;
+}
+
+FText SAnimationOutlinerTreeNode::GetDisplayName() const
+{
+	return DisplayNode->GetDisplayName();
+}
+
+void SAnimationOutlinerTreeNode::OnSelectionChanged( TArray<TSharedPtr<FSequencerDisplayNode> > AffectedNodes )
+{
+	TArray<TSharedPtr<FSequencerDisplayNode> > ObjectNodes;
+
+	for (int32 NodeIdx = 0; NodeIdx < AffectedNodes.Num(); ++NodeIdx)
+	{
+		if( AffectedNodes[NodeIdx]->GetType() == ESequencerNode::Object )
+		{
+			ObjectNodes.Add(AffectedNodes[NodeIdx]);
+		}
+	}
+
+	if (!ObjectNodes.Num())
+	{
+		return;
+	}
+
+	if (!Sequencer->IsLevelEditorSequencer())
+	{
+		return;
+	}
+
+	// Mark that the user is selecting so that the UI doesn't respond to the selection changes in the following block
+	TSharedRef<SSequencer> SequencerWidget = StaticCastSharedRef<SSequencer>(Sequencer->GetSequencerWidget());
+	SequencerWidget->SetUserIsSelecting(true);
+
+	const FScopedTransaction Transaction(NSLOCTEXT("UnrealEd", "ClickingOnActors", "Clicking on Actors"));
+	bool bActorSelected = false;
+
+	const FModifierKeysState ModifierKeys = FSlateApplication::Get().GetModifierKeys();
+	bool IsControlDown = ModifierKeys.IsControlDown();
+	bool IsShiftDown = ModifierKeys.IsShiftDown();
+
+	const bool bNotifySelectionChanged = false;
+	const bool bDeselectBSP = true;
+	const bool bWarnAboutTooManyActors = false;
+
+	GEditor->GetSelectedActors()->Modify();
+
+	if (!IsControlDown && !IsShiftDown)
+	{
+		GEditor->SelectNone(bNotifySelectionChanged, bDeselectBSP, bWarnAboutTooManyActors);
+	}
+
+	GEditor->GetSelectedActors()->BeginBatchSelectOperation();
+
+	// Select objects bound to the object node
+	for (int32 ObjectIdx = 0; ObjectIdx < ObjectNodes.Num(); ++ObjectIdx)
+	{
+		const TSharedPtr<const FObjectBindingNode> ObjectNode = StaticCastSharedPtr<const FObjectBindingNode>( ObjectNodes[ObjectIdx] );
+
+		// Get the bound objects
+		TArray<UObject*> RuntimeObjects;
+		Sequencer->GetRuntimeObjects( Sequencer->GetFocusedMovieSceneInstance(), ObjectNode->GetObjectBinding(), RuntimeObjects );
+		
+		if( RuntimeObjects.Num() > 0 )
+		{
+			// Select each actor			
+			for( int32 ActorIdx = 0; ActorIdx < RuntimeObjects.Num(); ++ActorIdx )
+			{
+				AActor* Actor = Cast<AActor>( RuntimeObjects[ActorIdx] );
+
+				if( Actor )
+				{
+					bool bSelectActor = true;
+
+					if (IsControlDown)
+					{
+						bSelectActor = Sequencer->GetSelection().IsSelected(ObjectNodes[ObjectIdx].ToSharedRef());
+					}
+
+					GEditor->SelectActor(Actor, bSelectActor, bNotifySelectionChanged );
+					bActorSelected = true;
+				}
+			}
+		}
+	}
+
+	GEditor->GetSelectedActors()->EndBatchSelectOperation();
+
+	if( bActorSelected )
+	{
+		GEditor->NoteSelectionChange();
+	}
+
+	// Unlock the selection so that the sequencer widget can now respond to selection changes in the level
+	SequencerWidget->SetUserIsSelecting(false);
+}
+
+#undef LOCTEXT_NAMESPACE
\ No newline at end of file
diff --git a/Engine/Source/Editor/Sequencer/Private/SAnimationOutlinerTreeNode.h b/Engine/Source/Editor/Sequencer/Private/SAnimationOutlinerTreeNode.h
new file mode 100644
--- /dev/null
+++ b/Engine/Source/Editor/Sequencer/Private/SAnimationOutlinerTreeNode.h
@@ -0,0 +1,78 @@
+// Copyright 1998-2015 Epic Games, Inc. All Rights Reserved.
+
+#pragma once
+
+class SSequencerTreeView;
+class SSequencerTreeViewRow;
+
+/**
+ * A widget for displaying a sequencer tree node in the animation outliner
+ */
+class SAnimationOutlinerTreeNode : public SCompoundWidget
+{
+public:
+	SLATE_BEGIN_ARGS(SAnimationOutlinerTreeNode){}
+	SLATE_END_ARGS()
+
+	void Construct( const FArguments& InArgs, TSharedRef<FSequencerDisplayNode> Node, const TSharedRef<SSequencerTreeViewRow>& InTableRow, FSequencer* InSequencer );
+
+	/**
+	 * @return The display node used by this widget                                                              
+	 */
+	const TSharedPtr< FSequencerDisplayNode> GetDisplayNode() const { return DisplayNode; }
+
+	/**
+	 * Gets a tint to apply buttons on hover
+	 */
+	FLinearColor GetHoverTint() const;
+
+private:
+	/** SWidget Interface */
+	virtual FReply OnMouseButtonDown( const FGeometry& MyGeometry, const FPointerEvent& MouseEvent ) override;
+	virtual FReply OnMouseButtonUp( const FGeometry& MyGeometry, const FPointerEvent& MouseEvent ) override;
+
+	FSlateColor GetForegroundBasedOnSelection() const;
+
+	/**
+	 * @return The border image to show in the tree node
+	 */
+	const FSlateBrush* GetNodeBorderImage() const;
+	
+	/**
+	 * @return The expander visibility of this node
+	 */
+	EVisibility GetExpanderVisibility() const;
+
+	/**
+	 * @return The display name for this node.
+	 */
+	FText GetDisplayName() const;
+
+	/** Handles the previous key button being clicked. */
+	FReply OnPreviousKeyClicked();
+
+	/** Handles the next key button being clicked. */
+	FReply OnNextKeyClicked();
+
+	/** Handles the add key button being clicked. */
+	FReply OnAddKeyClicked();
+
+	/** Get all descendant nodes from the given root node */
+	void GetAllDescendantNodes(TSharedPtr<FSequencerDisplayNode> RootNode, TArray<TSharedRef<FSequencerDisplayNode> >& AllNodes);
+
+	void OnSelectionChanged( TArray<TSharedPtr<FSequencerDisplayNode> > AffectedNodes );
+
+private:
+	/** Internal sequencer interface */
+	FSequencer* Sequencer;
+	/** Layout node the widget is visualizing */
+	TSharedPtr<FSequencerDisplayNode> DisplayNode;
+	/** Brush to display a border around the widget when it is selected */
+	const FSlateBrush* SelectedBrush;
+	/** Brush to display a border around the widget when it is selected but inactive */
+	const FSlateBrush* SelectedBrushInactive;
+	/** Brush to use if the node is not selected */
+	const FSlateBrush* NotSelectedBrush;
+	/** The table row style used for nodes in the tree. This is required as we don't actually use the tree for selection */
+	const FTableRowStyle* TableRowStyle;
+};
\ No newline at end of file
diff --git a/Engine/Source/Editor/Sequencer/Private/SAnimationOutlinerView.cpp b/Engine/Source/Editor/Sequencer/Private/SAnimationOutlinerView.cpp
deleted file mode 100644
diff --git a/Engine/Source/Editor/Sequencer/Private/SAnimationOutlinerView.h b/Engine/Source/Editor/Sequencer/Private/SAnimationOutlinerView.h
deleted file mode 100644
diff --git a/Engine/Source/Editor/Sequencer/Private/SSection.cpp b/Engine/Source/Editor/Sequencer/Private/SSection.cpp
--- a/Engine/Source/Editor/Sequencer/Private/SSection.cpp
+++ b/Engine/Source/Editor/Sequencer/Private/SSection.cpp
@@ -19,6 +19,7 @@ void SSection::Construct( const FArguments& InArgs, TSharedRef<FTrackNode> Secti
 	SectionIndex = InSectionIndex;
 	ParentSectionArea = SectionNode;
 	SectionInterface = SectionNode->GetSections()[InSectionIndex];
+	CachedKeyAreaHeight = 0.f;
 
 	ChildSlot
 	[
@@ -26,42 +27,50 @@ void SSection::Construct( const FArguments& InArgs, TSharedRef<FTrackNode> Secti
 	];
 }
 
-void SSection::GetKeyAreas( const TSharedPtr<FTrackNode>& SectionAreaNode, TArray<FKeyAreaElement>& OutKeyAreas ) const
+FVector2D SSection::ComputeDesiredSize(float) const
 {
-	float HeightOffset = 0;
-	TSharedPtr<FSectionKeyAreaNode> SectionKeyNode = SectionAreaNode->GetTopLevelKeyNode();
-	if( SectionKeyNode.IsValid() )
-	{
-		new( OutKeyAreas ) FKeyAreaElement( *SectionKeyNode, HeightOffset );
-	}
-
-	GetKeyAreas_Recursive( SectionAreaNode->GetChildNodes(), HeightOffset, OutKeyAreas );
+	return FVector2D(100, CachedKeyAreaHeight);
 }
 
-void SSection::GetKeyAreas_Recursive(  const TArray< TSharedRef<FSequencerDisplayNode> >& Nodes, float& HeightOffset, TArray<FKeyAreaElement>& OutKeyAreas ) const
+void SSection::GetKeyAreas( const TSharedPtr<FTrackNode>& SectionAreaNode, TArray<FKeyAreaElement>& OutKeyAreas, float* TotalHeight ) const
 {
-	const float Padding = SequencerLayoutConstants::NodePadding;
+	*TotalHeight = 0.f;
 
-	for( int32 NodeIndex = 0; NodeIndex < Nodes.Num(); ++NodeIndex )
+	float LastPadding = 0.f;
+
+	// First, layout the parent
 	{
-		FSequencerDisplayNode& Node = *Nodes[NodeIndex];
+		TSharedPtr<FSequencerDisplayNode> LayoutNode = SectionAreaNode;
 
-		if( Node.IsVisible() )
+		TSharedPtr<FSectionKeyAreaNode> SectionKeyNode = SectionAreaNode->GetTopLevelKeyNode();
+		if( SectionKeyNode.IsValid() )
 		{
-			// Compute the node height.  If this is not a key area it will contribute to the accumulated height offset
-			float NodeHeight = Node.GetNodeHeight();
-			HeightOffset += Padding + NodeHeight;
-			if( Node.GetType() == ESequencerNode::KeyArea )
-			{
-				// The node is a key area, we need to draw it
-				new( OutKeyAreas ) FKeyAreaElement( *StaticCastSharedRef<FSectionKeyAreaNode>( Nodes[NodeIndex] ), HeightOffset );
-			}
+			LayoutNode = SectionKeyNode;
+			OutKeyAreas.Add(FKeyAreaElement( *SectionKeyNode, *TotalHeight ));
+		}
 
-			// Get any child key areas
-			GetKeyAreas_Recursive( Node.GetChildNodes(), HeightOffset, OutKeyAreas );
+		LastPadding = LayoutNode->GetNodePadding().Bottom;
+		*TotalHeight += LayoutNode->GetNodeHeight() + LastPadding;
+	}
+
+	// Then any children
+	SectionAreaNode->TraverseVisible_ParentFirst([&](FSequencerDisplayNode& Node){
+		
+		*TotalHeight += Node.GetNodePadding().Top;
 
+		if( Node.GetType() == ESequencerNode::KeyArea )
+		{
+			OutKeyAreas.Add(FKeyAreaElement( static_cast<FSectionKeyAreaNode&>(Node), *TotalHeight ));
 		}
-	}
+
+		LastPadding = Node.GetNodePadding().Bottom;
+		*TotalHeight += Node.GetNodeHeight() + LastPadding;
+		return true;
+
+	}, false);
+
+	// Remove the padding from the bottom
+	*TotalHeight -= LastPadding;
 }
 
 
@@ -472,8 +481,10 @@ void SSection::Tick( const FGeometry& AllottedGeometry, const double InCurrentTi
 {	
 	if( GetVisibility() == EVisibility::Visible )
 	{
+		CachedKeyAreaHeight = 0.f;
+
 		KeyAreas.Reset();
-		GetKeyAreas( ParentSectionArea, KeyAreas );
+		GetKeyAreas( ParentSectionArea, KeyAreas, &CachedKeyAreaHeight );
 
 		SectionInterface->Tick(AllottedGeometry, ParentGeometry, InCurrentTime, InDeltaTime);
 	}
@@ -492,20 +503,24 @@ FReply SSection::OnMouseButtonDown( const FGeometry& MyGeometry, const FPointerE
 	{
 		// Check for clicking on a key and mark it as the pressed key for drag detection (if necessary) later
 		PressedKey = GetKeyUnderMouse( MouseEvent.GetScreenSpacePosition(), MyGeometry );
+	}
 
-		if( !PressedKey.IsValid() && MouseEvent.GetEffectingButton() == EKeys::LeftMouseButton )
+	if (MouseEvent.GetEffectingButton() == EKeys::LeftMouseButton)
+	{
+		if( !PressedKey.IsValid() )
 		{
 			CheckForEdgeInteraction( MouseEvent, MyGeometry );
 		}
 
 		return FReply::Handled().CaptureMouse( AsShared() );
 	}
-	else if( MouseEvent.GetEffectingButton() == EKeys::RightMouseButton )
+	else if (MouseEvent.GetEffectingButton() == EKeys::RightMouseButton && PressedKey.IsValid())
 	{
-		return FReply::Handled().CaptureMouse(AsShared());
+		// Don't return handled if we didn't click on a key so that right-click panning gets a look in
+		return FReply::Handled();
 	}
 
-	return FReply::Handled();
+	return FReply::Unhandled();
 }
 
 void SSection::ResetState()
@@ -535,12 +550,12 @@ FReply SSection::OnMouseButtonUp( const FGeometry& MyGeometry, const FPointerEve
 	}
 	else
 	{
-		if( ( MouseEvent.GetEffectingButton() == EKeys::LeftMouseButton || MouseEvent.GetEffectingButton() == EKeys::RightMouseButton ) && HasMouseCapture() && MyGeometry.IsUnderLocation( MouseEvent.GetScreenSpacePosition() ) )
+		if( ( MouseEvent.GetEffectingButton() == EKeys::LeftMouseButton || MouseEvent.GetEffectingButton() == EKeys::RightMouseButton ) && MyGeometry.IsUnderLocation( MouseEvent.GetScreenSpacePosition() ) )
 		{
 			HandleSelection( MyGeometry, MouseEvent );
 		}
 
-		if( MouseEvent.GetEffectingButton() == EKeys::RightMouseButton && HasMouseCapture() )
+		if( MouseEvent.GetEffectingButton() == EKeys::RightMouseButton )
 		{
 			TSharedPtr<SWidget> MenuContent = OnSummonContextMenu( MyGeometry, MouseEvent );
 			if (MenuContent.IsValid())
@@ -555,7 +570,7 @@ FReply SSection::OnMouseButtonUp( const FGeometry& MyGeometry, const FPointerEve
 					FPopupTransitionEffect( FPopupTransitionEffect::ContextMenu )
 					);
 
-				return FReply::Handled().ReleaseMouseCapture().SetUserFocus(MenuContent.ToSharedRef(), EFocusCause::SetDirectly);
+				return FReply::Handled().SetUserFocus(MenuContent.ToSharedRef(), EFocusCause::SetDirectly);
 			}
 		}
 	}
diff --git a/Engine/Source/Editor/Sequencer/Private/SSection.h b/Engine/Source/Editor/Sequencer/Private/SSection.h
--- a/Engine/Source/Editor/Sequencer/Private/SSection.h
+++ b/Engine/Source/Editor/Sequencer/Private/SSection.h
@@ -27,6 +27,8 @@ public:
 	/** Caches the parent geometry to be given to section interfaces that need it on tick */
 	void CacheParentGeometry(const FGeometry& InParentGeometry) {ParentGeometry = InParentGeometry;}
 
+	virtual FVector2D ComputeDesiredSize(float) const override;
+	
 private:
 	/** Information about how to draw a key area */
 	struct FKeyAreaElement
@@ -46,16 +48,7 @@ private:
 	 * @param SectionAreaNode		The section area node containing key areas
 	 * @param OutKeyAreaElements	Populated list of key elements to draw
 	 */
-	void GetKeyAreas( const TSharedPtr<FTrackNode>& SectionAreaNode, TArray<FKeyAreaElement>& OutKeyAreaElements ) const;
-
-	/**
-	 * Recursively gets keys areas from a list of nodes
-	 *
-	 * @param Nodes					List of nodes that could be key areas or contain key areas
-	 * @param HeightOffset			The current height offset incremented for each key area
-	 * @param OutKeyAreaElements	Populated list of key elements to draw
-	 */
-	void GetKeyAreas_Recursive(  const TArray< TSharedRef<FSequencerDisplayNode> >& Nodes, float& HeightOffset, TArray<FKeyAreaElement>& OutKeyAreaElements ) const;
+	void GetKeyAreas( const TSharedPtr<FTrackNode>& SectionAreaNode, TArray<FKeyAreaElement>& OutKeyAreaElements, float* HeightOffset = nullptr ) const;
 
 	/**
 	 * Computes the geometry for a key area
@@ -163,6 +156,8 @@ private:
 	TSharedPtr<FTrackNode> ParentSectionArea;
 	/** Key areas inside this section */
 	TArray<FKeyAreaElement> KeyAreas;
+	/** Cached height of all the key areas */
+	float CachedKeyAreaHeight;
 	/** Current drag operation if any */
 	TSharedPtr<class FSequencerDragOperation> DragOperation;
 	/** The currently pressed key if any */
diff --git a/Engine/Source/Editor/Sequencer/Private/SSequencer.cpp b/Engine/Source/Editor/Sequencer/Private/SSequencer.cpp
--- a/Engine/Source/Editor/Sequencer/Private/SSequencer.cpp
+++ b/Engine/Source/Editor/Sequencer/Private/SSequencer.cpp
@@ -26,7 +26,8 @@
 #include "SSearchBox.h"
 #include "SNumericDropDown.h"
 #include "EditorWidgetsModule.h"
-#include "SequencerTrackLaneFactory.h"
+#include "SSequencerTreeView.h"
+
 #include "SSequencerSplitterOverlay.h"
 #include "IKeyArea.h"
 
@@ -34,47 +35,6 @@
 #define LOCTEXT_NAMESPACE "Sequencer"
 
 
-class SSequencerScrollBox : public SScrollBox
-{
-public:
-	void Construct( const FArguments& InArgs, TSharedRef<FSequencer> InSequencer);
-
-	virtual FReply OnMouseButtonDown(const FGeometry& MyGeometry, const FPointerEvent& MouseEvent) override;
-
-	/** The sequencer which owns this widget. */
-	TWeakPtr<FSequencer> Sequencer;
-};
-
-void SSequencerScrollBox::Construct(const FArguments& InArgs, TSharedRef<FSequencer> InSequencer)
-{
-	Sequencer = InSequencer;
-	SScrollBox::Construct(InArgs);
-}
-
-FReply SSequencerScrollBox::OnMouseButtonDown(const FGeometry& MyGeometry, const FPointerEvent& MouseEvent)
-{
-	// Clear the selection if no sequencer display nodes were clicked on
-	if (MouseEvent.GetEffectingButton() == EKeys::LeftMouseButton)
-	{
-		if (Sequencer.Pin()->GetSelection().GetSelectedOutlinerNodes().Num() != 0)
-		{
-			Sequencer.Pin()->GetSelection().EmptySelectedOutlinerNodes();
-
-			if (Sequencer.Pin()->IsLevelEditorSequencer())
-			{
-				const bool bNotifySelectionChanged = false;
-				const bool bDeselectBSP = true;
-				const bool bWarnAboutTooManyActors = false;
-
-				const FScopedTransaction Transaction(NSLOCTEXT("UnrealEd", "ClickingOnActors", "Clicking on Actors"));
-				GEditor->SelectNone(bNotifySelectionChanged, bDeselectBSP, bWarnAboutTooManyActors);
-				return FReply::Handled();
-			}
-		}
-	}
-	return FReply::Unhandled();
-}
-
 /**
  * The shot filter overlay displays the overlay needed to filter out widgets based on
  * which shots are actively in use.
@@ -236,8 +196,11 @@ void SSequencer::Construct( const FArguments& InArgs, TSharedRef< class FSequenc
 
 	SAssignNew( TrackOutliner, SSequencerTrackOutliner );
 
-	SAssignNew( TrackArea, SSequencerTrackArea, TimeSliderController )
+	SAssignNew( TrackArea, SSequencerTrackArea, TimeSliderController, PinnedSequencer )
 		.Visibility( this, &SSequencer::GetTrackAreaVisibility );
+	
+	SAssignNew( TreeView, SSequencerTreeView, SequencerNodeTree.ToSharedRef(), TrackArea.ToSharedRef() )
+	.ExternalScrollbar( ScrollBar );
 
 	SAssignNew( CurveEditor, SSequencerCurveEditor, PinnedSequencer, TimeSliderController )
 		.Visibility( this, &SSequencer::GetCurveEditorVisibility )
@@ -245,6 +208,8 @@ void SSequencer::Construct( const FArguments& InArgs, TSharedRef< class FSequenc
 		.ViewRange( InArgs._ViewRange );
 	CurveEditor->SetAllowAutoFrame(Settings->GetShowCurveEditor());
 
+	TrackArea->SetTreeView(TreeView);
+
 	const int32 				Column0	= 0,	Column1	= 1;
 	const int32 Row0	= 0,
 				Row1	= 1,
@@ -304,29 +269,23 @@ void SSequencer::Construct( const FArguments& InArgs, TSharedRef< class FSequenc
 
 							+ SOverlay::Slot()
 							[
-								SNew( SSequencerScrollBox, Sequencer.Pin().ToSharedRef() )
-								.ExternalScrollbar(ScrollBar)
-
-								+ SScrollBox::Slot()
+								SNew(SHorizontalBox)
+								
+								+ SHorizontalBox::Slot()
+								.FillWidth( FillCoefficient_0 )
 								[
-									SNew( SHorizontalBox )
-
-									+ SHorizontalBox::Slot()
-									.FillWidth( FillCoefficient_0 )
+									SNew(SBox)
+									// Padding to allow space for the scroll bar
+									.Padding(FMargin(0,0,10.f,0))
 									[
-										SNew(SBox)
-										// Padding to allow space for the scroll bar
-										.Padding(FMargin(0,0,10.f,0))
-										[
-											TrackOutliner.ToSharedRef()
-										]
+										TreeView.ToSharedRef()
 									]
+								]
 
-									+ SHorizontalBox::Slot()
-									.FillWidth( FillCoefficient_1 )
-									[
-										TrackArea.ToSharedRef()
-									]
+								+ SHorizontalBox::Slot()
+								.FillWidth( FillCoefficient_1 )
+								[
+									TrackArea.ToSharedRef()
 								]
 							]
 
@@ -711,8 +670,7 @@ void SSequencer::UpdateLayoutTree()
 	// Update the node tree
 	SequencerNodeTree->Update();
 
-	FSequencerTrackLaneFactory Factory(TrackOutliner.ToSharedRef(), TrackArea.ToSharedRef(), Sequencer.Pin().ToSharedRef() );
-	Factory.Repopulate( *SequencerNodeTree );
+	TreeView->Refresh();
 
 	// Restore the selection state.
 	RestoreSelectionState(SequencerNodeTree->GetRootNodes(), SelectedPathNames, Sequencer.Pin()->GetSelection());	// Update to actor selection.
@@ -1126,6 +1084,11 @@ FText SSequencer::GetShotSectionTitle(UMovieSceneSection* ShotSection) const
 	return Cast<UMovieSceneShotSection>(ShotSection)->GetShotDisplayName();
 }
 
+TSharedPtr<SSequencerTreeView> SSequencer::GetTreeView() const
+{
+	return TreeView;
+}
+
 void SSequencer::DeleteSelectedNodes()
 {
 	TSet< TSharedRef<FSequencerDisplayNode> > SelectedNodesCopy = Sequencer.Pin()->GetSelection().GetSelectedOutlinerNodes();
@@ -1138,7 +1101,7 @@ void SSequencer::DeleteSelectedNodes()
 
 		for( const TSharedRef<FSequencerDisplayNode>& SelectedNode : SelectedNodesCopy )
 		{
-			if( SelectedNode->IsVisible() )
+			if( !SelectedNode->IsHidden() )
 			{
 				// Delete everything in the entire node
 				TSharedRef<const FSequencerDisplayNode> NodeToBeDeleted = StaticCastSharedRef<const FSequencerDisplayNode>(SelectedNode);
@@ -1148,22 +1111,6 @@ void SSequencer::DeleteSelectedNodes()
 	}
 }
 
-void SSequencer::ExpandCollapseNode(TSharedRef<FSequencerDisplayNode> Node, bool bDescendants, bool bExpand)
-{
-	if (Node->IsExpanded() != bExpand)
-	{
-		Node->ToggleExpansion();
-	}
-
-	if (bDescendants)
-	{
-		for (auto ChildNode : Node->GetChildNodes())
-		{
-			ExpandCollapseNode(ChildNode, bDescendants, bExpand);
-		}
-	}
-}
-
 FReply SSequencer::OnSaveMovieSceneClicked()
 {
 	Sequencer.Pin()->SaveCurrentMovieScene();
@@ -1295,30 +1242,6 @@ void SSequencer::StepToKey(bool bStepToNextKey, bool bCameraOnly)
 	}
 }
 
-void SSequencer::ToggleExpandCollapseSelectedNodes(bool bDescendants)
-{
-	const TSet< TSharedRef<FSequencerDisplayNode> >& SelectedNodes = Sequencer.Pin()->GetSelection().GetSelectedOutlinerNodes();
-
-	TSet< TSharedRef<FSequencerDisplayNode> > Nodes(SelectedNodes);
-
-	if (Nodes.Num() == 0)
-	{
-		TSet<TSharedRef<FSequencerDisplayNode>> RootNodes(SequencerNodeTree->GetRootNodes());
-		Nodes = RootNodes;
-	}
-
-	if (Nodes.Num() > 0)
-	{
-		auto It = Nodes.CreateConstIterator();
-		bool bExpand = !(*It).Get().IsExpanded();
-
-		for (auto Node : Nodes)
-		{
-			ExpandCollapseNode(Node, bDescendants, bExpand);
-		}
-	}
-}
-
 EVisibility SSequencer::GetBreadcrumbTrailVisibility() const
 {
 	return Sequencer.Pin()->IsLevelEditorSequencer() ? EVisibility::Visible : EVisibility::Collapsed;
diff --git a/Engine/Source/Editor/Sequencer/Private/SSequencer.h b/Engine/Source/Editor/Sequencer/Private/SSequencer.h
--- a/Engine/Source/Editor/Sequencer/Private/SSequencer.h
+++ b/Engine/Source/Editor/Sequencer/Private/SSequencer.h
@@ -2,14 +2,12 @@
 
 #pragma once
 
+class SSequencerTreeView;
+
 namespace SequencerLayoutConstants
 {
 	/** The amount to indent child nodes of the layout tree */
-	const float IndentAmount = 25.0f;
-	/** Padding between each node */
-	const float NodePadding = 3.0f;
-	/** Padding between each node */
-	const float ObjectNodePadding = 5.0f;
+	const float IndentAmount = 8.0f;
 	/** Height of each object node */
 	const float ObjectNodeHeight = 20.0f;
 	/** Height of each section area if there are no sections (note: section areas may be larger than this if they have children. This is the height of a section area with no children or all children hidden) */
@@ -106,16 +104,16 @@ public:
 	void StepToPreviousCameraKey();
 	void StepToKey(bool bStepToNextKey, bool bCameraOnly);
 
-	/** Expand or collapse selected nodes out of the sequencer node tree */
-	void ToggleExpandCollapseSelectedNodes(bool bDescendants = false);
-	void ExpandCollapseNode(TSharedRef<FSequencerDisplayNode> SelectedNode, bool bDescendants, bool bExpand);
-
 	/** Whether the user is selecting. Ignore selection changes from the level when the user is selecting. */
 	void SetUserIsSelecting(bool bUserIsSelectingIn) { bUserIsSelecting = bUserIsSelectingIn; }
 	bool UserIsSelecting() { return bUserIsSelecting; }
 
 	/** Called when the save button is clicked */
 	FReply OnSaveMovieSceneClicked();
+
+	/** Access the tree view for this sequencer */
+	TSharedPtr<SSequencerTreeView> GetTreeView() const;
+
 private:
 	/** Empty active timer to ensure Slate ticks during Sequencer playback */
 	EActiveTimerReturnType EnsureSlateTickDuringPlayback(double InCurrentTime, float InDeltaTime);	
@@ -250,6 +248,8 @@ private:
 	TSharedPtr<class FSequencerNodeTree> SequencerNodeTree;
 	/** The breadcrumb trail widget for this sequencer */
 	TSharedPtr< class SBreadcrumbTrail<FSequencerBreadcrumb> > BreadcrumbTrail;
+	/** The sequencer tree view responsible for the outliner and track areas */
+	TSharedPtr<SSequencerTreeView> TreeView;
 	/** The main sequencer interface */
 	TWeakPtr<FSequencer> Sequencer;
 	/** Cached settings provided to the sequencer itself on creation */
diff --git a/Engine/Source/Editor/Sequencer/Private/SSequencerSectionAreaView.cpp b/Engine/Source/Editor/Sequencer/Private/SSequencerSectionAreaView.cpp
--- a/Engine/Source/Editor/Sequencer/Private/SSequencerSectionAreaView.cpp
+++ b/Engine/Source/Editor/Sequencer/Private/SSequencerSectionAreaView.cpp
@@ -62,19 +62,28 @@ namespace SequencerSectionUtils
 
 void SSequencerSectionAreaView::Construct( const FArguments& InArgs, TSharedRef<FSequencerDisplayNode> Node )
 {
-	SSequencerSectionAreaViewBase::Construct( SSequencerSectionAreaViewBase::FArguments(), Node );
+	//SSequencerSectionAreaViewBase::Construct( SSequencerSectionAreaViewBase::FArguments(), Node );
 
 	ViewRange = InArgs._ViewRange;
 
 	check( Node->GetType() == ESequencerNode::Track );
 	SectionAreaNode = StaticCastSharedRef<FTrackNode>( Node );
 
-	SetVisibility( TAttribute<EVisibility>( this, &SSequencerSectionAreaView::GetNodeVisibility ) );
-
 	BackgroundBrush = FEditorStyle::GetBrush("Sequencer.SectionArea.Background");
 
 	// Generate widgets for sections in this view
 	GenerateSectionWidgets();
+};
+
+FVector2D SSequencerSectionAreaView::ComputeDesiredSize(float) const
+{
+	// Note: X Size is not used
+	FVector2D Size(100, 0.f);
+	for (int32 Index = 0; Index < Children.Num(); ++Index)
+	{
+		Size.Y = FMath::Max(Size.Y, Children[Index]->GetDesiredSize().Y);
+	}
+	return Size + SectionAreaNode->GetNodePadding().Bottom;
 }
 
 void SSequencerSectionAreaView::GenerateSectionWidgets()
@@ -100,17 +109,6 @@ EVisibility SSequencerSectionAreaView::GetSectionVisibility( UMovieSceneSection*
 	return GetSequencer().IsSectionVisible( SectionObject ) ? EVisibility::Visible : EVisibility::Collapsed;
 }
 
-float SSequencerSectionAreaView::GetSectionAreaHeight() const
-{
- 	return SectionAreaNode->GetNodeHeight() + ComputeHeightRecursive( SectionAreaNode->GetChildNodes() );
-}
-
-EVisibility SSequencerSectionAreaView::GetNodeVisibility() const
-{
-	// We are not visible if our node is hidden
-	return SectionAreaNode->IsVisible() ? EVisibility::Visible : EVisibility::Collapsed;
-}
-
 
 /** SWidget Interface */
 int32 SSequencerSectionAreaView::OnPaint( const FPaintArgs& Args, const FGeometry& AllottedGeometry, const FSlateRect& MyClippingRect, FSlateWindowElementList& OutDrawElements, int32 LayerId, const FWidgetStyle& InWidgetStyle, bool bParentEnabled ) const
@@ -118,20 +116,6 @@ int32 SSequencerSectionAreaView::OnPaint( const FPaintArgs& Args, const FGeometr
 	FArrangedChildren ArrangedChildren(EVisibility::Visible);
 	ArrangeChildren(AllottedGeometry, ArrangedChildren);
 
-	if( SectionAreaNode.IsValid() )
-	{
-		// Draw a region around the entire section area
-		FSlateDrawElement::MakeBox( 
-			OutDrawElements, 
-			LayerId,
-			AllottedGeometry.ToPaintGeometry(),
-			BackgroundBrush,
-			MyClippingRect,
-			ESlateDrawEffect::None,
-			SequencerSectionAreaConstants::BackgroundColor
-			);
-	}
-
 	for (int32 ChildIndex = 0; ChildIndex < ArrangedChildren.Num(); ++ChildIndex)
 	{
 		FArrangedWidget& CurWidget = ArrangedChildren[ChildIndex];
@@ -142,17 +126,6 @@ int32 SSequencerSectionAreaView::OnPaint( const FPaintArgs& Args, const FGeometr
 	return LayerId+1;
 }
 
-FReply SSequencerSectionAreaView::OnMouseButtonDown( const FGeometry& MyGeometry, const FPointerEvent& MouseEvent )
-{
-	// Clear selected sections
-	if( !MouseEvent.IsControlDown() )
-	{
-		GetSequencer().GetSelection().EmptySelectedSections();
-	}
-
-	return FReply::Handled();
-}
-
 void SSequencerSectionAreaView::OnArrangeChildren( const FGeometry& AllottedGeometry, FArrangedChildren& ArrangedChildren ) const
 {
 	int32 MaxRowIndex = 0;
@@ -167,7 +140,7 @@ void SSequencerSectionAreaView::OnArrangeChildren( const FGeometry& AllottedGeom
 	int32 MaxTracks = MaxRowIndex + 1;
 
 
-	float SectionHeight = GetSectionAreaHeight();
+	float SectionHeight = AllottedGeometry.GetLocalSize().Y - SectionAreaNode->GetNodePadding().Bottom;
 
 	FTimeToPixel TimeToPixelConverter = GetTimeToPixel( AllottedGeometry );
 
diff --git a/Engine/Source/Editor/Sequencer/Private/SSequencerSectionAreaView.h b/Engine/Source/Editor/Sequencer/Private/SSequencerSectionAreaView.h
--- a/Engine/Source/Editor/Sequencer/Private/SSequencerSectionAreaView.h
+++ b/Engine/Source/Editor/Sequencer/Private/SSequencerSectionAreaView.h
@@ -5,14 +5,10 @@
 class FSectionKeyAreaNode;
 class FTrackNode;
 
-
-typedef TSequencerTrackViewPanel<SSection> SSequencerSectionAreaViewBase;
-
-
 /**
  * Visualizes a section area and its children                                        
  */
-class SSequencerSectionAreaView : public SSequencerSectionAreaViewBase
+class SSequencerSectionAreaView : public SPanel
 {
 public:
 	SLATE_BEGIN_ARGS( SSequencerSectionAreaView )
@@ -23,24 +19,15 @@ public:
 
 	void Construct( const FArguments& InArgs, TSharedRef<FSequencerDisplayNode> Node );
 
-	/**
-	 * @return The current height of the section area
-	 */
-	float GetSectionAreaHeight() const;
-	
-	/**
-	 * @return The visibility of this node                                                              
-	 */
-	EVisibility GetNodeVisibility() const;
-
 private:
 	/** SWidget Interface */
 	virtual int32 OnPaint( const FPaintArgs& Args, const FGeometry& AllottedGeometry, const FSlateRect& MyClippingRect, FSlateWindowElementList& OutDrawElements, int32 LayerId, const FWidgetStyle& InWidgetStyle, bool bParentEnabled ) const override;
-	virtual FReply OnMouseButtonDown( const FGeometry& MyGeometry, const FPointerEvent& MouseEvent ) override;
 
 	/** SPanel Interface */
 	virtual void OnArrangeChildren( const FGeometry& AllottedGeometry, FArrangedChildren& ArrangedChildren ) const override;
-
+	virtual FVector2D ComputeDesiredSize(float) const override;
+	virtual FChildren* GetChildren() override { return &Children; }
+	
 	/**
 	 * Returns the visibility of a section
 	 *
@@ -79,4 +66,6 @@ private:
 	float DistanceDragged;
 	/** Whether or not we are dragging */
 	bool bDragging;
+	/** All the widgets in the panel */
+	TSlotlessChildren<SSection> Children;
 };
diff --git a/Engine/Source/Editor/Sequencer/Private/SSequencerTrackArea.cpp b/Engine/Source/Editor/Sequencer/Private/SSequencerTrackArea.cpp
--- a/Engine/Source/Editor/Sequencer/Private/SSequencerTrackArea.cpp
+++ b/Engine/Source/Editor/Sequencer/Private/SSequencerTrackArea.cpp
@@ -2,32 +2,233 @@
 
 #include "SequencerPrivatePCH.h"
 #include "SSequencerTrackArea.h"
+#include "SSequencerTrackLane.h"
 #include "TimeSliderController.h"
 
-void SSequencerTrackArea::Construct( const FArguments& InArgs, TSharedRef<FSequencerTimeSliderController> InTimeSliderController )
+
+struct FMarqueeSelectData
+{
+	FMarqueeSelectData(const FVector2D& InInitialPosition)
+		: InitialPosition(InInitialPosition)
+		, CurrentPosition(InInitialPosition)
+	{}
+
+	FVector2D TopLeft() const
+	{
+		return FVector2D(
+			FMath::Min(InitialPosition.X, CurrentPosition.X),
+			FMath::Min(InitialPosition.Y, CurrentPosition.Y)
+		);
+	}
+	
+	FVector2D BottomRight() const
+	{
+		return FVector2D(
+			FMath::Max(InitialPosition.X, CurrentPosition.X),
+			FMath::Max(InitialPosition.Y, CurrentPosition.Y)
+		);
+	}
+
+	FVector2D Size() const
+	{
+		return BottomRight() - TopLeft();
+	}
+
+	bool HasDragged() const
+	{
+		return !Size().IsNearlyZero();
+	}
+
+	FVector2D InitialPosition;
+	FVector2D CurrentPosition;
+};
+
+/** Structure used for converting X and Y positional values to Time and vertical offsets respectively */
+struct FVirtualConverter : FTimeToPixel
+{
+	FVirtualConverter(TSharedPtr<SSequencerTreeView> InTreeView, TSharedPtr<FSequencer> InSequencer, const FGeometry& InGeometry)
+		: FTimeToPixel(InGeometry, InSequencer->GetViewRange())
+		, TreeView(*InTreeView)
+	{}
+
+	float PixelToVerticalOffset(float InPixel) const
+	{
+		return TreeView.PhysicalToVirtual(InPixel);
+	}
+
+	float VerticalOffsetToPixel(float InOffset) const
+	{
+		return TreeView.VirtualToPhysical(InOffset);
+	}
+
+	FVector2D PhysicalToVirtual(FVector2D InPosition) const
+	{
+		InPosition.Y = PixelToVerticalOffset(InPosition.Y);
+		InPosition.X = PixelToTime(InPosition.X);
+
+		return InPosition;
+	}
+
+	FVector2D VirtualToPhysical(FVector2D InPosition) const
+	{
+		InPosition.Y = VerticalOffsetToPixel(InPosition.Y);
+		InPosition.X = TimeToPixel(InPosition.X);
+
+		return InPosition;
+	}
+
+private:
+	SSequencerTreeView& TreeView;
+};
+
+void SSequencerTrackArea::Construct( const FArguments& InArgs, TSharedRef<FSequencerTimeSliderController> InTimeSliderController, TSharedRef<FSequencer> InSequencer )
 {
+	Sequencer = InSequencer;
 	TimeSliderController = InTimeSliderController;
-	SVerticalBox::Construct( InArgs );
+
+	SOverlay::Construct(SOverlay::FArguments());
+}
+
+void SSequencerTrackArea::SetTreeView(const TSharedPtr<SSequencerTreeView>& InTreeView)
+{
+	TreeView = InTreeView;
+}
+
+void SSequencerTrackArea::AddTrackSlot(const TSharedRef<FSequencerDisplayNode>& InNode, const TSharedPtr<SSequencerTrackLane>& InSlot)
+{
+	// Capture a weak reference to ensure that the slot can get cleaned up ok
+	TWeakPtr<SSequencerTrackLane> WeakSlot = InSlot;
+	auto Padding = [WeakSlot]{
+		auto PinnedSlot = WeakSlot.Pin();
+		return PinnedSlot.IsValid() ? FMargin(0.f, PinnedSlot->GetPhysicalPosition(), 0, 0) : FMargin(0);
+	};
+
+	TrackSlots.Add(InNode, InSlot);
+
+	typedef TAttribute<FMargin> Attr;
+
+	AddSlot()
+	.VAlign(VAlign_Top)
+	.Padding(Attr::Create(Attr::FGetter::CreateLambda(Padding)))
+	[
+		SNew(SWeakWidget)
+		.PossiblyNullContent(InSlot)
+	];
+}
+
+TSharedPtr<SSequencerTrackLane> SSequencerTrackArea::FindTrackSlot(const TSharedRef<FSequencerDisplayNode>& InNode)
+{
+	return TrackSlots.FindRef(InNode).Pin();
+}
+
+int32 SSequencerTrackArea::OnPaint(const FPaintArgs& Args, const FGeometry& AllottedGeometry, const FSlateRect& MyClippingRect, FSlateWindowElementList& OutDrawElements, int32 LayerId, const FWidgetStyle& InWidgetStyle, bool bParentEnabled ) const
+{
+	LayerId = SOverlay::OnPaint(Args, AllottedGeometry, MyClippingRect, OutDrawElements, LayerId, InWidgetStyle, bParentEnabled);
+	return PaintMarquee(AllottedGeometry, MyClippingRect, OutDrawElements, LayerId);
+}
+
+int32 SSequencerTrackArea::PaintMarquee(const FGeometry& AllottedGeometry, const FSlateRect& MyClippingRect, FSlateWindowElementList& OutDrawElements, int32 LayerId) const
+{
+	if (MarqueeSelectData.IsValid() && MarqueeSelectData->HasDragged())
+	{
+		// Convert to physical space for rendering
+		const FVirtualConverter Converter(TreeView.Pin(), Sequencer.Pin(), AllottedGeometry);
+
+		FVector2D TopLeft = Converter.VirtualToPhysical(MarqueeSelectData->TopLeft());
+		FVector2D BottomRight = Converter.VirtualToPhysical(MarqueeSelectData->BottomRight());
+
+		FSlateDrawElement::MakeBox(
+			OutDrawElements,
+			LayerId,
+			AllottedGeometry.ToPaintGeometry(TopLeft, BottomRight - TopLeft),
+			FEditorStyle::GetBrush(TEXT("MarqueeSelection")),
+			MyClippingRect
+			);
+
+		return LayerId + 1;
+	}
+
+	return LayerId;
 }
 
 FReply SSequencerTrackArea::OnMouseButtonDown( const FGeometry& MyGeometry, const FPointerEvent& MouseEvent )
 {
-	return TimeSliderController->OnMouseButtonDown( SharedThis(this), MyGeometry, MouseEvent );
+	FReply Reply = FReply::Unhandled();
+
+	if (MouseEvent.GetEffectingButton() == EKeys::LeftMouseButton)
+	{
+		const FVirtualConverter Converter(TreeView.Pin(), Sequencer.Pin(), MyGeometry);
+
+		// Marquee data is stored in virtual space
+		FVector2D VirtualClickPosition = Converter.PhysicalToVirtual(MyGeometry.AbsoluteToLocal(MouseEvent.GetScreenSpacePosition()));
+		MarqueeSelectData.Reset(new FMarqueeSelectData(VirtualClickPosition));
+
+		Reply = FReply::Handled().CaptureMouse(AsShared());
+	}
+	else
+	{
+		Reply = TimeSliderController->OnMouseButtonDown( SharedThis(this), MyGeometry, MouseEvent );
+	}
+
+	// Clear selected sections
+	auto SequencerPin = Sequencer.Pin();
+	if( !MouseEvent.IsControlDown() && SequencerPin.IsValid() )
+	{
+		SequencerPin->GetSelection().EmptySelectedSections();
+	}
+
+	return Reply;
 }
 
 FReply SSequencerTrackArea::OnMouseButtonUp( const FGeometry& MyGeometry, const FPointerEvent& MouseEvent )
 {
-	return TimeSliderController->OnMouseButtonUp( SharedThis(this),  MyGeometry, MouseEvent );
+	FReply Reply = FReply::Unhandled();
+
+	if (MarqueeSelectData.IsValid())
+	{
+		if (MarqueeSelectData->HasDragged())
+		{
+			HandleMarqueeSelection(MyGeometry, MouseEvent);
+		}
+		MarqueeSelectData.Reset();
+		Sequencer.Pin()->StopAutoscroll();
+		Reply = FReply::Handled().ReleaseMouseCapture();
+	}
+
+	if (!Reply.IsEventHandled())
+	{
+		Reply = TimeSliderController->OnMouseButtonUp( SharedThis(this), MyGeometry, MouseEvent );
+	}
+
+	return Reply;
 }
 
 FReply SSequencerTrackArea::OnMouseMove( const FGeometry& MyGeometry, const FPointerEvent& MouseEvent )
 {
+	if (MarqueeSelectData.IsValid())
+	{
+		SetNewMarqueeBounds(MyGeometry.AbsoluteToLocal(MouseEvent.GetScreenSpacePosition()), MyGeometry);
+		return FReply::Handled();
+	}
+	else if (MouseEvent.IsMouseButtonDown(EKeys::RightMouseButton) && HasMouseCapture())
+	{
+		TreeView.Pin()->ScrollByDelta( -MouseEvent.GetCursorDelta().Y );
+	}
+
 	return TimeSliderController->OnMouseMove( SharedThis(this), MyGeometry, MouseEvent );
 }
 
 FReply SSequencerTrackArea::OnMouseWheel( const FGeometry& MyGeometry, const FPointerEvent& MouseEvent )
 {
-	return TimeSliderController->OnMouseWheel( SharedThis(this), MyGeometry, MouseEvent );
+	FReply Reply = TimeSliderController->OnMouseWheel( SharedThis(this), MyGeometry, MouseEvent );
+
+	if (!Reply.IsEventHandled())
+	{
+		TreeView.Pin()->ScrollByDelta(WheelScrollAmount * -MouseEvent.GetWheelDelta());
+		Reply = FReply::Handled();
+	}
+
+	return Reply;
 }
 
 void SSequencerTrackArea::Tick( const FGeometry& AllottedGeometry, const double InCurrentTime, const float InDeltaTime )
@@ -48,4 +249,140 @@ void SSequencerTrackArea::Tick( const FGeometry& AllottedGeometry, const double
 	}
 
 	SizeLastFrame = Size;
+
+	for (int32 Index = 0; Index < Children.Num(); )
+	{
+		if (!StaticCastSharedRef<SWeakWidget>(Children[Index].GetWidget())->ChildWidgetIsValid())
+		{
+			Children.RemoveAt(Index);
+		}
+		else
+		{
+			++Index;
+		}
+	}
+}
+
+void SSequencerTrackArea::SetNewMarqueeBounds(FVector2D NewVirtualPosition, const FGeometry& MyGeometry)
+{
+	const FVirtualConverter Converter(TreeView.Pin(), Sequencer.Pin(), MyGeometry);
+	TRange<float> ViewRange = Sequencer.Pin()->GetViewRange();
+
+	FVector2D Change = MarqueeSelectData->CurrentPosition;
+
+	MarqueeSelectData->CurrentPosition = Converter.PhysicalToVirtual(NewVirtualPosition);
+
+	// Handle virtual scrolling when at the extremes of the widget
+	Change = MarqueeSelectData->CurrentPosition - Change;
+
+	const float ScrollThresholdH = ViewRange.Size<float>() * 0.025f;
+	const float ScrollThresholdV = MyGeometry.Size.Y * 0.025f;
+
+	float Difference = MarqueeSelectData->CurrentPosition.X - (ViewRange.GetLowerBoundValue() + ScrollThresholdH);
+	if (Difference < 0 && Change.X < 0)
+	{
+		Sequencer.Pin()->StartAutoscroll(Difference);
+	}
+	else
+	{
+		Difference = MarqueeSelectData->CurrentPosition.X - (ViewRange.GetUpperBoundValue() - ScrollThresholdH);
+		if (Difference > 0 && Change.X > 0)
+		{
+			Sequencer.Pin()->StartAutoscroll(Difference);
+		}
+		else
+		{
+			Sequencer.Pin()->StopAutoscroll();
+		}
+	}
+
+	Difference = Converter.VerticalOffsetToPixel(MarqueeSelectData->CurrentPosition.Y) - ScrollThresholdV;
+	if (Difference < 0 && Change.Y < 0)
+	{
+		TreeView.Pin()->ScrollByDelta( Difference * 0.1f );
+	}
+
+	Difference = Converter.VerticalOffsetToPixel(MarqueeSelectData->CurrentPosition.Y) - (MyGeometry.Size.Y - ScrollThresholdV);
+	if (Difference > 0 && Change.Y > 0)
+	{
+		TreeView.Pin()->ScrollByDelta( Difference * 0.1f );
+	}
+}
+
+void SSequencerTrackArea::HandleMarqueeSelection(const FGeometry& Geometry, const FPointerEvent& MouseEvent)
+{
+	auto PinnedSequencer = Sequencer.Pin();
+	auto PinnedTreeView = TreeView.Pin();
+
+	if (!PinnedSequencer.IsValid() || !PinnedTreeView.IsValid())
+	{
+		return;
+	}
+
+	auto& Selection = PinnedSequencer->GetSelection();
+	if (!MouseEvent.IsLeftControlDown() && !MouseEvent.IsLeftShiftDown())
+	{
+		Selection.Empty();
+	}
+
+	auto TopLeft = MarqueeSelectData->TopLeft();
+	auto BottomRight = MarqueeSelectData->BottomRight();
+
+	for (auto& RootNode : PinnedTreeView->GetNodeTree()->GetRootNodes())
+	{
+		const TArray<TSharedRef<ISequencerSection>>* CurrentSections = nullptr;
+
+		RootNode->TraverseVisible_ParentFirst([&](FSequencerDisplayNode& InNode){
+
+			if (InNode.GetType() == ESequencerNode::KeyArea)
+			{
+				FSectionKeyAreaNode& KeyAreaNode = static_cast<FSectionKeyAreaNode&>(InNode);
+
+				if (CurrentSections && InNode.VirtualTop > TopLeft.Y && InNode.VirtualBottom < BottomRight.Y)
+				{
+					for( int32 SectionIndex = 0; SectionIndex < CurrentSections->Num(); ++SectionIndex )
+					{
+						FTimeToPixel TimeToPixelConverter(Geometry, PinnedSequencer->GetViewRange());
+
+						TSharedRef<IKeyArea> KeyArea = KeyAreaNode.GetKeyArea(SectionIndex);
+
+						for (FKeyHandle KeyHandle : KeyArea->GetUnsortedKeyHandles())
+						{
+							float KeyPosition = KeyArea->GetKeyTime(KeyHandle);
+							if (KeyPosition > TopLeft.X && KeyPosition < BottomRight.X)
+							{
+								FSelectedKey SelectedKey(*(*CurrentSections)[SectionIndex]->GetSectionObject(), KeyArea, KeyHandle);
+
+								if (MouseEvent.IsLeftShiftDown())
+								{
+									Selection.AddToSelection(SelectedKey);
+								}
+								else if (MouseEvent.IsLeftControlDown())
+								{
+									if (Selection.IsSelected(SelectedKey))
+									{
+										Selection.RemoveFromSelection(SelectedKey);
+									}
+									else
+									{
+										Selection.AddToSelection(SelectedKey);
+									}
+								}
+								else
+								{
+									Selection.AddToSelection(SelectedKey);
+								}
+							}
+						}
+					}
+				}
+			}
+			else if (InNode.GetType() == ESequencerNode::Track)
+			{
+				CurrentSections = &static_cast<FTrackNode&>(InNode).GetSections();
+			}
+
+			return true;
+		});
+	}
 }
\ No newline at end of file
diff --git a/Engine/Source/Editor/Sequencer/Private/SSequencerTrackArea.h b/Engine/Source/Editor/Sequencer/Private/SSequencerTrackArea.h
--- a/Engine/Source/Editor/Sequencer/Private/SSequencerTrackArea.h
+++ b/Engine/Source/Editor/Sequencer/Private/SSequencerTrackArea.h
@@ -2,32 +2,74 @@
 
 #pragma once
 
-class SSequencerOutlinerArea;
+#include "UniquePtr.h"
+
+class SSequencerTreeViewRow;
+class SSequencerTreeView;
+class SSequencerTrackLane;
+
 class FSequencerTimeSliderController;
+class FSequencer;
+
+struct FMarqueeSelectData;
 
 /**
  * The area where tracks( rows of sections ) are displayed
  */
-class SSequencerTrackArea : public SVerticalBox
+class SSequencerTrackArea : public SOverlay
 {
 public:
-	DECLARE_DELEGATE_OneParam(FOnValueChanged, float)
+	SLATE_BEGIN_ARGS( SSequencerTrackArea ){}
+	SLATE_END_ARGS()
 
+	/** Construct this widget */
+	void Construct( const FArguments& InArgs, TSharedRef<FSequencerTimeSliderController> InTimeSliderController, TSharedRef<FSequencer> InSequencer );
 
-	void Construct( const FArguments& InArgs, TSharedRef<FSequencerTimeSliderController> InTimeSliderController );
-	
-	virtual void Tick( const FGeometry& AllottedGeometry, const double InCurrentTime, const float InDeltaTime ) override;
+	/** Assign a tree view to this track area. */
+	void SetTreeView(const TSharedPtr<SSequencerTreeView>& InTreeView);
 
+	/** Add a new track slot to this area for the given node. The slot will be automatically cleaned up when all external references to the supplied slot are removed. */
+	void AddTrackSlot(const TSharedRef<FSequencerDisplayNode>& InNode, const TSharedPtr<SSequencerTrackLane>& InSlot);
+
+	/** Attempt to find an existing slot relating to the given node */
+	TSharedPtr<SSequencerTrackLane> FindTrackSlot(const TSharedRef<FSequencerDisplayNode>& InNode);
+
+	/** SWidget overrides */
 	virtual FReply OnMouseButtonDown( const FGeometry& MyGeometry, const FPointerEvent& MouseEvent ) override;
 	virtual FReply OnMouseButtonUp( const FGeometry& MyGeometry, const FPointerEvent& MouseEvent ) override;
 	virtual FReply OnMouseMove( const FGeometry& MyGeometry, const FPointerEvent& MouseEvent ) override;
 	virtual FReply OnMouseWheel( const FGeometry& MyGeometry, const FPointerEvent& MouseEvent ) override;
+	virtual int32 OnPaint(const FPaintArgs& Args, const FGeometry& AllottedGeometry, const FSlateRect& MyClippingRect, FSlateWindowElementList& OutDrawElements, int32 LayerId, const FWidgetStyle& InWidgetStyle, bool bParentEnabled ) const override;
+	virtual void Tick( const FGeometry& AllottedGeometry, const double InCurrentTime, const float InDeltaTime ) override;
+
+private:
+
+	/** Paint the marquee box, if it's valid */
+	int32 PaintMarquee(const FGeometry& AllottedGeometry, const FSlateRect& MyClippingRect, FSlateWindowElementList& OutDrawElements, int32 LayerId) const;
+
+	/** Set new bounds for the marquee selection */
+	void SetNewMarqueeBounds(FVector2D NewVirtualPosition, const FGeometry& Geometry);
+
+	/** Select all the keys contained in the current marquee */
+	void HandleMarqueeSelection(const FGeometry& Geometry, const FPointerEvent& MouseEvent);
 
 private:
 
+	/** A map of child slot content that exist in our view */
+	TMap<TSharedPtr<FSequencerDisplayNode>, TWeakPtr<SSequencerTrackLane>> TrackSlots;
+
+	/** Weak pointer to the sequencer */
+	TWeakPtr<FSequencer> Sequencer;
+
+	/** Weak pointer to the tree view (used for scrolling interactions) */
+	TWeakPtr<SSequencerTreeView> TreeView;
+
 	/** Time slider controller for controlling zoom/pan etc */
-	TSharedPtr<class FSequencerTimeSliderController> TimeSliderController;
+	TSharedPtr<FSequencerTimeSliderController> TimeSliderController;
 
 	/** Keep a hold of the size of the area so we can maintain zoom levels */
 	TOptional<FVector2D> SizeLastFrame;
+
+	/** Optional marquee selection data */
+	TUniquePtr<FMarqueeSelectData> MarqueeSelectData;
 };
diff --git a/Engine/Source/Editor/Sequencer/Private/SSequencerTrackLane.cpp b/Engine/Source/Editor/Sequencer/Private/SSequencerTrackLane.cpp
new file mode 100644
--- /dev/null
+++ b/Engine/Source/Editor/Sequencer/Private/SSequencerTrackLane.cpp
@@ -0,0 +1,79 @@
+// Copyright 1998-2015 Epic Games, Inc. All Rights Reserved.
+
+#pragma once
+
+#include "SequencerPrivatePCH.h"
+#include "SSequencerTrackLane.h"
+#include "SSequencerTreeView.h"
+
+void SSequencerTrackLane::Construct(const FArguments& InArgs, const TSharedRef<FSequencerDisplayNode>& InDisplayNode, const TSharedRef<SSequencerTreeView>& InTreeView)
+{
+	DisplayNode = InDisplayNode;
+	TreeView = InTreeView;
+
+	ChildSlot
+	.HAlign(HAlign_Fill)
+	.Padding(0, DisplayNode->GetNodePadding().Top, 0, 0)
+	[
+		InArgs._Content.Widget
+	];
+}
+
+int32 SSequencerTrackLane::OnPaint(const FPaintArgs& Args, const FGeometry& AllottedGeometry, const FSlateRect& MyClippingRect, FSlateWindowElementList& OutDrawElements, int32 LayerId, const FWidgetStyle& InWidgetStyle, bool bParentEnabled) const
+{
+	// Draw a separator below the track
+	auto PaintGeometry = AllottedGeometry.ToPaintGeometry();
+
+	TArray<FVector2D> Points;
+	Points.Emplace(0, PaintGeometry.GetLocalSize().Y);
+	Points.Emplace(PaintGeometry.GetLocalSize().X, PaintGeometry.GetLocalSize().Y);
+
+	static FLinearColor SeparatorColor(0.f, 0.f, 0.f, 0.15f);
+
+	FSlateDrawElement::MakeLines( 
+		OutDrawElements, 
+		LayerId,
+		PaintGeometry,
+		Points,
+		MyClippingRect,
+		ESlateDrawEffect::None,
+		SeparatorColor
+		);
+
+	return SCompoundWidget::OnPaint(Args, AllottedGeometry, MyClippingRect, OutDrawElements, LayerId + 1, InWidgetStyle, bParentEnabled);
+}
+
+float SSequencerTrackLane::GetPhysicalPosition() const
+{
+	// Positioning strategy:
+	// We know that there must be at least one tree view row around that references something in our subtree (otherwise this widget wouldn't exist)
+	// Knowing this, we just need to root it out by traversing the visible tree, and asking for the first rows's offset
+	float NegativeOffset = 0.f;
+	TOptional<float> Top;
+
+	auto TreeViewPinned = TreeView.Pin();
+	
+	// Iterate parent first until we find a tree view row we can use for the offset height
+	auto Iter = [&](FSequencerDisplayNode& InNode){
+		
+		auto ChildRow = StaticCastSharedPtr<SSequencerTreeViewRow>(TreeViewPinned->WidgetFromItem(InNode.AsShared()));
+		if (ChildRow.IsValid())
+		{
+			Top = ChildRow->GetCachedGeometry().Position.Y;
+			// Stop iterating
+			return false;
+		}
+
+		NegativeOffset -= InNode.GetNodeHeight() + InNode.GetNodePadding().Combined();
+		return true;
+	};
+
+	DisplayNode->TraverseVisible_ParentFirst(Iter);
+
+	if (!Top.IsSet())
+	{
+		Top = 0.f;
+	}
+
+	return NegativeOffset + Top.GetValue();
+}
\ No newline at end of file
diff --git a/Engine/Source/Editor/Sequencer/Private/SSequencerTrackLane.h b/Engine/Source/Editor/Sequencer/Private/SSequencerTrackLane.h
new file mode 100644
--- /dev/null
+++ b/Engine/Source/Editor/Sequencer/Private/SSequencerTrackLane.h
@@ -0,0 +1,32 @@
+// Copyright 1998-2015 Epic Games, Inc. All Rights Reserved.
+
+#pragma once
+
+class SSequencerTreeView;
+class FSequencerDisplayNode;
+
+/**
+ * A wrapper widget responsible for positioning a track lane within the section area
+ */
+class SSequencerTrackLane : public SCompoundWidget
+{
+public:
+	SLATE_BEGIN_ARGS(SSequencerTrackLane){}
+		SLATE_DEFAULT_SLOT(FArguments, Content)
+	SLATE_END_ARGS()
+
+	/** Construct this widget */
+	void Construct(const FArguments& InArgs, const TSharedRef<FSequencerDisplayNode>& InDisplayNode, const TSharedRef<SSequencerTreeView>& InTreeView);
+	
+	/** Paint this widget */
+	virtual int32 OnPaint(const FPaintArgs& Args, const FGeometry& AllottedGeometry, const FSlateRect& MyClippingRect, FSlateWindowElementList& OutDrawElements, int32 LayerId, const FWidgetStyle& InWidgetStyle, bool bParentEnabled) const override;
+
+	/** Get the desired physical vertical position of this track lane */
+	float GetPhysicalPosition() const;
+	
+private:
+	/** The authoritative display node that created us */
+	TSharedPtr<FSequencerDisplayNode> DisplayNode;
+	/** Pointer back to the tree view for virtual <-> physical space conversions. Important: weak ptr to avoid circular references */
+	TWeakPtr<SSequencerTreeView> TreeView;
+};
\ No newline at end of file
diff --git a/Engine/Source/Editor/Sequencer/Private/SSequencerTrackViewPanel.h b/Engine/Source/Editor/Sequencer/Private/SSequencerTrackViewPanel.h
deleted file mode 100644
diff --git a/Engine/Source/Editor/Sequencer/Private/SSequencerTreeView.cpp b/Engine/Source/Editor/Sequencer/Private/SSequencerTreeView.cpp
new file mode 100644
--- /dev/null
+++ b/Engine/Source/Editor/Sequencer/Private/SSequencerTreeView.cpp
@@ -0,0 +1,377 @@
+// Copyright 1998-2015 Epic Games, Inc. All Rights Reserved.
+
+#pragma once
+
+#include "SequencerPrivatePCH.h"
+#include "Sequencer.h"
+#include "SListPanel.h"
+
+static FName TrackAreaName = "TrackArea";
+
+/** Structure used to cache physical geometry for a particular node */
+struct FCachedGeometry
+{
+	FCachedGeometry(FDisplayNodeRef InNode, float InPhysicalTop, float InPhysicalHeight)
+		: Node(MoveTemp(InNode)), PhysicalTop(InPhysicalTop), PhysicalHeight(InPhysicalHeight)
+	{}
+
+	FDisplayNodeRef Node;
+	float PhysicalTop, PhysicalHeight;
+};
+
+/** Construct function for this widget */
+void SSequencerTreeViewRow::Construct(const FArguments& InArgs, const TSharedRef<STableViewBase>& OwnerTableView, const FDisplayNodeRef& InNode)
+{
+	Node = InNode;
+	OnGenerateWidgetForColumn = InArgs._OnGenerateWidgetForColumn;
+
+	SMultiColumnTableRow::Construct(SMultiColumnTableRow::FArguments(), OwnerTableView);
+}
+
+TSharedRef<SWidget> SSequencerTreeViewRow::GenerateWidgetForColumn(const FName& ColumnId)
+{
+	auto PinnedNode = Node.Pin();
+	if (PinnedNode.IsValid())
+	{
+		return OnGenerateWidgetForColumn.Execute(PinnedNode.ToSharedRef(), ColumnId, SharedThis(this));
+	}
+
+	return SNullWidget::NullWidget;
+}
+
+TSharedPtr<FSequencerDisplayNode> SSequencerTreeViewRow::GetDisplayNode() const
+{
+	return Node.Pin();
+}
+
+const FGeometry& SSequencerTreeViewRow::GetCachedGeometry() const
+{
+	return RowGeometry;
+}
+
+void SSequencerTreeViewRow::AddTrackAreaReference(const TSharedPtr<SSequencerTrackLane>& Lane)
+{
+	TrackLaneReference = Lane;
+}
+
+void SSequencerTreeViewRow::Tick(const FGeometry& AllottedGeometry, const double InCurrentTime, const float InDeltaTime)
+{
+	RowGeometry = AllottedGeometry;
+}
+
+void SSequencerTreeView::Construct(const FArguments& InArgs, const TSharedRef<FSequencerNodeTree>& InNodeTree, const TSharedRef<SSequencerTrackArea>& InTrackArea)
+{
+	SequencerNodeTree = InNodeTree;
+	TrackArea = InTrackArea;
+
+	// We 'leak' these delegates (they'll get cleaned up automatically when the invocation list changes)
+	// It's not safe to attempt their removal in ~SSequencerTreeView because SequencerNodeTree->GetSequencer() may not be valid
+	FSequencer& Sequencer = InNodeTree->GetSequencer();
+	Sequencer.GetSettings()->GetOnShowCurveEditorChanged().AddSP(this, &SSequencerTreeView::UpdateTrackArea);
+	Sequencer.GetSelection().GetOnOutlinerNodeSelectionChanged().AddSP(this, &SSequencerTreeView::OnSequencerSelectionChangedExternally);
+
+	HeaderRow = SNew(SHeaderRow).Visibility(EVisibility::Collapsed);
+
+	SetupColumns(InArgs);
+
+	STreeView::Construct
+	(
+		STreeView::FArguments()
+		.TreeItemsSource(&RootNodes)
+		.SelectionMode(ESelectionMode::None)
+		.OnGenerateRow(this, &SSequencerTreeView::OnGenerateRow)
+		.OnGetChildren(this, &SSequencerTreeView::OnGetChildren)
+		.HeaderRow(HeaderRow)
+		.ExternalScrollbar(InArgs._ExternalScrollbar)
+		.OnExpansionChanged(this, &SSequencerTreeView::OnExpansionChanged)
+		.AllowOverscroll(EAllowOverscroll::No)
+	);
+}
+
+void SSequencerTreeView::Tick(const FGeometry& AllottedGeometry, const double InCurrentTime, const float InDeltaTime)
+{
+	STreeView::Tick(AllottedGeometry, InCurrentTime, InDeltaTime);
+	UpdateCachedVerticalGeometry();
+}
+
+void SSequencerTreeView::UpdateCachedVerticalGeometry()
+{
+	FChildren* Children = ItemsPanel->GetChildren();
+
+	CachedGeometry.Empty(Children->Num());
+	for (int32 Index = 0; Index < Children->Num(); ++Index)
+	{
+		TSharedRef<SSequencerTreeViewRow> TreeViewRow = StaticCastSharedRef<SSequencerTreeViewRow>(Children->GetChildAt(Index));
+		TSharedPtr<FSequencerDisplayNode> Node = TreeViewRow->GetDisplayNode();
+		if (Node.IsValid())
+		{
+			CachedGeometry.Emplace(Node.ToSharedRef(), TreeViewRow->GetCachedGeometry().Position.Y, TreeViewRow->GetCachedGeometry().Size.Y);
+		}
+	}
+}
+
+float SSequencerTreeView::PhysicalToVirtual(float InPhysical) const
+{
+	for (auto& Geometry : CachedGeometry)
+	{
+		if (InPhysical >= Geometry.PhysicalTop && InPhysical <= Geometry.PhysicalTop + Geometry.PhysicalHeight)
+		{
+			const float FractionalHeight = (InPhysical - Geometry.PhysicalTop) / Geometry.PhysicalHeight;
+			return Geometry.Node->VirtualTop + (Geometry.Node->VirtualBottom - Geometry.Node->VirtualTop) * FractionalHeight;
+		}
+	}
+
+	if (CachedGeometry.Num())
+	{
+		auto Last = CachedGeometry.Last();
+		return Last.Node->VirtualTop + (InPhysical - Last.PhysicalTop);
+	}
+
+	return InPhysical;
+}
+
+float SSequencerTreeView::VirtualToPhysical(float InVirtual) const
+{
+	for (auto& Geometry : CachedGeometry)
+	{
+		if (InVirtual >= Geometry.Node->VirtualTop && InVirtual <= Geometry.Node->VirtualBottom)
+		{
+			const float FractionalHeight = (InVirtual - Geometry.Node->VirtualTop) / (Geometry.Node->VirtualBottom - Geometry.Node->VirtualTop);
+			return Geometry.PhysicalTop + Geometry.PhysicalHeight * FractionalHeight;
+		}
+	}
+	
+	if (CachedGeometry.Num())
+	{
+		auto Last = CachedGeometry.Last();
+		return Last.PhysicalTop + (InVirtual - Last.Node->VirtualTop);
+	}
+
+	return InVirtual;
+}
+
+void SSequencerTreeView::SetupColumns(const FArguments& InArgs)
+{
+	FSequencer& Sequencer = SequencerNodeTree->GetSequencer();
+
+	// Define a column for the Outliner
+	auto GenerateOutliner = [=](const FDisplayNodeRef& InNode, const TSharedRef<SSequencerTreeViewRow>& InRow)
+	{
+		return InNode->GenerateContainerWidgetForOutliner(InRow);
+	};
+
+	Columns.Add("Outliner", FSequencerTreeViewColumn(GenerateOutliner, 1.f));
+
+	// Now populate the header row with the columns
+	for (auto& Pair : Columns)
+	{
+		if (Pair.Key != TrackAreaName || !Sequencer.GetSettings()->GetShowCurveEditor())
+		{
+			HeaderRow->AddColumn(
+				SHeaderRow::Column(Pair.Key)
+				.FillWidth(Pair.Value.Width)
+			);
+		}
+	}
+}
+
+void SSequencerTreeView::UpdateTrackArea()
+{
+	FSequencer& Sequencer = SequencerNodeTree->GetSequencer();
+
+	// Add or remove the column
+	if (Sequencer.GetSettings()->GetShowCurveEditor())
+	{
+		HeaderRow->RemoveColumn(TrackAreaName);
+	}
+	else if (const auto* Column = Columns.Find(TrackAreaName))
+	{
+		HeaderRow->AddColumn(
+			SHeaderRow::Column(TrackAreaName)
+			.FillWidth(Column->Width)
+		);
+	}
+}
+
+void SSequencerTreeView::OnSequencerSelectionChangedExternally()
+{
+	Private_ClearSelection();
+
+	FSequencer& Sequencer = SequencerNodeTree->GetSequencer();
+	for (auto& Node : Sequencer.GetSelection().GetSelectedOutlinerNodes())
+	{
+		Private_SetItemSelection(Node, true, false);
+	}
+
+	Private_SignalSelectionChanged(ESelectInfo::Direct);
+}
+
+void SSequencerTreeView::Private_SignalSelectionChanged(ESelectInfo::Type SelectInfo)
+{
+	if (SelectInfo != ESelectInfo::Direct)
+	{
+		FSequencer& Sequencer = SequencerNodeTree->GetSequencer();
+		FSequencerSelection& Selection = Sequencer.GetSelection();
+		Selection.EmptySelectedOutlinerNodes();
+		for (auto& Item : GetSelectedItems())
+		{
+			Selection.AddToSelection(Item);
+		}
+	}
+
+	STreeView::Private_SignalSelectionChanged(SelectInfo);
+}
+
+void SSequencerTreeView::Refresh()
+{
+	RootNodes.Reset(SequencerNodeTree->GetRootNodes().Num());
+
+	for (const auto& RootNode : SequencerNodeTree->GetRootNodes())
+	{
+		if (RootNode->IsExpanded())
+		{
+			SetItemExpansion(RootNode, true);
+		}
+
+		if (!RootNode->IsHidden())
+		{
+			RootNodes.Add(RootNode);
+		}
+	}
+
+	RequestTreeRefresh();
+}
+
+void SSequencerTreeView::ScrollByDelta(float DeltaInSlateUnits)
+{
+	const float FractionalScrollOffsetInItems = (ScrollOffset + GetScrollRateInItems() * DeltaInSlateUnits);
+	ScrollTo( FractionalScrollOffsetInItems );
+}
+
+template<typename T>
+bool ShouldExpand(const T& InContainer, ETreeRecursion Recursion)
+{
+	bool bAllExpanded = true;
+	for (auto& Item : InContainer)
+	{
+		bAllExpanded &= Item->IsExpanded();
+		if (Recursion == ETreeRecursion::Recursive)
+		{
+			Item->TraverseVisible_ParentFirst([&](FSequencerDisplayNode& InNode){
+				bAllExpanded &= InNode.IsExpanded();
+				return true;
+			});
+		}
+	}
+	return !bAllExpanded;
+}
+
+void SSequencerTreeView::ToggleSelectedNodeExpansion(ETreeRecursion Recursion)
+{
+	FSequencer& Sequencer = SequencerNodeTree->GetSequencer();
+
+	const TSet< FDisplayNodeRef >& SelectedNodes = Sequencer.GetSelection().GetSelectedOutlinerNodes();
+
+	if (SelectedNodes.Num() != 0)
+	{
+		const bool bExpand = ShouldExpand(SelectedNodes, Recursion);
+		for (auto& Item : SelectedNodes)
+		{
+			ExpandCollapseNode(Item, bExpand, Recursion);
+		}
+	}
+	else
+	{
+		const bool bExpand = ShouldExpand(SequencerNodeTree->GetRootNodes(), Recursion);
+		for (auto& Item : SequencerNodeTree->GetRootNodes())
+		{
+			ExpandCollapseNode(Item, bExpand, Recursion);
+		}
+	}
+}
+
+void SSequencerTreeView::ExpandCollapseNode(const FDisplayNodeRef& InNode, bool bExpansionState, ETreeRecursion Recursion)
+{
+	SetItemExpansion(InNode, bExpansionState);
+
+	if (Recursion == ETreeRecursion::Recursive)
+	{
+		for (auto& Node : InNode->GetChildNodes())
+		{
+			ExpandCollapseNode(Node, bExpansionState, ETreeRecursion::Recursive);
+		}
+	}
+}
+
+void SSequencerTreeView::OnExpansionChanged(FDisplayNodeRef InItem, bool bIsExpanded)
+{
+	InItem->SetExpansionState(bIsExpanded);
+	
+	// Expand any children that are also expanded
+	for (auto& Child : InItem->GetChildNodes())
+	{
+		if (Child->IsExpanded())
+		{
+			SetItemExpansion(Child, true);
+		}
+	}
+
+	InItem->UpdateCachedShotFilteredVisibility();
+}
+
+void SSequencerTreeView::OnGetChildren(FDisplayNodeRef InParent, TArray<FDisplayNodeRef>& OutChildren) const
+{
+	for (const auto& Node : InParent->GetChildNodes())
+	{
+		if (!Node->IsHidden())
+		{
+			OutChildren.Add(Node);
+		}
+	}
+}
+
+TSharedRef<ITableRow> SSequencerTreeView::OnGenerateRow(FDisplayNodeRef InDisplayNode, const TSharedRef<STableViewBase>& OwnerTable)
+{
+	TSharedRef<SSequencerTreeViewRow> Row =
+		SNew(SSequencerTreeViewRow, OwnerTable, InDisplayNode)
+		.OnGenerateWidgetForColumn(this, &SSequencerTreeView::GenerateWidgetForColumn);
+
+	// Ensure the track area is kept up to date with the virtualized scroll of the tree view
+	TSharedPtr<FSequencerDisplayNode> SectionAuthority = InDisplayNode->GetSectionAreaAuthority();
+	if (SectionAuthority.IsValid())
+	{
+		TSharedPtr<SSequencerTrackLane> TrackLane = TrackArea->FindTrackSlot(SectionAuthority.ToSharedRef());
+
+		if (!TrackLane.IsValid())
+		{
+			// Add a track slot for the row
+			TAttribute<TRange<float>> ViewRange = FAnimatedRange::WrapAttribute( TAttribute<FAnimatedRange>::Create(TAttribute<FAnimatedRange>::FGetter::CreateSP(&SequencerNodeTree->GetSequencer(), &FSequencer::GetViewRange)) );	
+
+			TrackLane = SNew(SSequencerTrackLane, SectionAuthority.ToSharedRef(), SharedThis(this))
+			[
+				SectionAuthority->GenerateWidgetForSectionArea(ViewRange)
+			];
+
+			TrackArea->AddTrackSlot(SectionAuthority.ToSharedRef(), TrackLane);
+		}
+
+		if (ensure(TrackLane.IsValid()))
+		{
+			Row->AddTrackAreaReference(TrackLane);
+		}
+	}
+
+	return Row;
+}
+
+TSharedRef<SWidget> SSequencerTreeView::GenerateWidgetForColumn(const FDisplayNodeRef& InNode, const FName& ColumnId, const TSharedRef<SSequencerTreeViewRow>& Row) const
+{
+	const auto* Definition = Columns.Find(ColumnId);
+
+	if (ensureMsgf(Definition, TEXT("Invalid column name specified")))
+	{
+		return Definition->Generator(InNode, Row);
+	}
+
+	return SNullWidget::NullWidget;
+}
\ No newline at end of file
diff --git a/Engine/Source/Editor/Sequencer/Private/SSequencerTreeView.h b/Engine/Source/Editor/Sequencer/Private/SSequencerTreeView.h
new file mode 100644
--- /dev/null
+++ b/Engine/Source/Editor/Sequencer/Private/SSequencerTreeView.h
@@ -0,0 +1,168 @@
+// Copyright 1998-2015 Epic Games, Inc. All Rights Reserved.
+
+#pragma once
+
+class FSequencer;
+class FSequencerNodeTree;
+class FSequencerDisplayNode;
+
+class SSequencerTrackLane;
+class SSequencerTrackArea;
+class SSequencerTreeViewRow;
+
+typedef TSharedRef<FSequencerDisplayNode> FDisplayNodeRef;
+
+enum class ETreeRecursion
+{
+	Recursive, NonRecursive
+};
+
+/** Structure used to define a column in the tree view */
+struct FSequencerTreeViewColumn
+{
+	typedef TFunction<TSharedRef<SWidget>(const FDisplayNodeRef&, const TSharedRef<SSequencerTreeViewRow>&)> FOnGenerate;
+
+	FSequencerTreeViewColumn(const FOnGenerate& InOnGenerate, const TAttribute<float>& InWidth) : Generator(InOnGenerate), Width(InWidth) {}
+
+	/** Function used to generate a cell for this column */
+	FOnGenerate Generator;
+	/** Attribute specifying the width of this column */
+	TAttribute<float> Width;
+};
+
+/** The tree view used in the sequencer */
+class SSequencerTreeView : public STreeView<FDisplayNodeRef>
+{
+public:
+
+	SLATE_BEGIN_ARGS(SSequencerTreeView){}
+		/** Externally supplied scroll bar */
+		SLATE_ARGUMENT( TSharedPtr<SScrollBar>, ExternalScrollbar )
+	SLATE_END_ARGS()
+
+	/** Construct this widget */
+	void Construct(const FArguments& InArgs, const TSharedRef<FSequencerNodeTree>& InNodeTree, const TSharedRef<SSequencerTrackArea>& InTrackArea);
+	virtual void Tick( const FGeometry& AllottedGeometry, const double InCurrentTime, const float InDeltaTime );
+
+	/** Access the underlying tree data */
+	TSharedPtr<FSequencerNodeTree> GetNodeTree() { return SequencerNodeTree; }
+
+	/** Convert the specified physical vertical position into an absolute virtual position, ignoring expanded states */
+	float PhysicalToVirtual(float InPhysical) const;
+
+	/** Convert the specified absolute virtual position into a physical position in the tree.
+	 * @note: Will not work reliably for virtual positions that are outside of the physical space
+	 */
+	float VirtualToPhysical(float InVirtual) const;
+
+	/** Refresh this tree as a result of the underlying tree data changing */
+	void Refresh();
+
+	/** Expand or collapse selected nodes */
+	void ToggleSelectedNodeExpansion(ETreeRecursion Recursion);
+
+	/** Scroll this tree view by the specified number of slate units */
+	void ScrollByDelta(float DeltaInSlateUnits);
+
+protected:
+
+	/** Set the item's expansion state, including all of its children */
+	void ExpandCollapseNode(const FDisplayNodeRef& InNode, bool bExpansionState, ETreeRecursion Recursion);
+
+	/** Generate a row for a particular node */
+	TSharedRef<ITableRow> OnGenerateRow(FDisplayNodeRef InDisplayNode, const TSharedRef<STableViewBase>& OwnerTable);
+
+	/** Gather the children from the specified node */
+	void OnGetChildren(FDisplayNodeRef InParent, TArray<FDisplayNodeRef>& OutChildren) const;
+
+	/** Generate a widget for the specified Node and Column */
+	TSharedRef<SWidget> GenerateWidgetForColumn(const FDisplayNodeRef& Node, const FName& ColumnId, const TSharedRef<SSequencerTreeViewRow>& Row) const;
+
+	/** Called when a node has been expanded or collapsed */
+	void OnExpansionChanged(FDisplayNodeRef InItem, bool bIsExpanded);
+
+	/** Called when the sequencer outliner selection has been changed externally */
+	void OnSequencerSelectionChangedExternally();
+
+	/** Overridden to keep external selection states up to date */
+	virtual void Private_SignalSelectionChanged(ESelectInfo::Type SelectInfo) override;
+
+protected:
+
+	/** Populate the map of column definitions, and add relevant columns to the header row */
+	void SetupColumns(const FArguments& InArgs);
+
+	/** Ensure that the track area column is either show or hidden, depending on the visibility of the curve editor */
+	void UpdateTrackArea();
+
+	/** Update the cached visible geometry for each tree item in this view */
+	void UpdateCachedVerticalGeometry();
+
+	/** Linear array of cached geometry for visible nodes in the tree */
+	TArray<struct FCachedGeometry> CachedGeometry;
+
+private:
+
+	/** The tree view's header row (hidden) */
+	TSharedPtr<SHeaderRow> HeaderRow;
+
+	/** Pointer to the node tree data that is used to populate this tree */
+	TSharedPtr<FSequencerNodeTree> SequencerNodeTree;
+
+	/** Cached copy of the root nodes from the tree data */
+	TArray<FDisplayNodeRef> RootNodes;
+
+	/** Column definitions for each of the columns in the tree view */
+	TMap<FName, FSequencerTreeViewColumn> Columns;
+
+	/** Strong pointer to the track area so we can generate track lanes as we need them */
+	TSharedPtr<SSequencerTrackArea> TrackArea;
+};
+
+/** Widget that represents a row in the sequencer's tree control. */
+class SSequencerTreeViewRow	: public SMultiColumnTableRow<FDisplayNodeRef>
+{
+public:
+	DECLARE_DELEGATE_RetVal_ThreeParams(TSharedRef<SWidget>, FOnGenerateWidgetForColumn, const FDisplayNodeRef&, const FName&, const TSharedRef<SSequencerTreeViewRow>&);
+
+	SLATE_BEGIN_ARGS(SSequencerTreeViewRow){}
+
+		/** Delegate to invoke to create a new column for this row */
+		SLATE_EVENT(FOnGenerateWidgetForColumn, OnGenerateWidgetForColumn)
+
+	SLATE_END_ARGS()
+
+	/** Construct function for this widget */
+	void Construct(const FArguments& InArgs, const TSharedRef<STableViewBase>& OwnerTableView, const FDisplayNodeRef& InNode);
+
+	/** Get the dispolay node to which this row relates */
+	TSharedPtr<FSequencerDisplayNode> GetDisplayNode() const;
+
+	/** Access this row's cached geometry */
+	const FGeometry& GetCachedGeometry() const;
+
+	/** Add a reference to the specified track lane, keeping it alive until this row is destroyed */
+	void AddTrackAreaReference(const TSharedPtr<SSequencerTrackLane>& Lane);
+
+	virtual void Tick(const FGeometry& AllottedGeometry, const double InCurrentTime, const float InDeltaTime) override;
+
+	/** Access this row's owner tree view. */
+	TSharedPtr<SSequencerTreeView> GetOwnerTree() const;
+
+	/** Overridden from SMultiColumnTableRow.  Generates a widget for this column of the tree row. */
+	virtual TSharedRef<SWidget> GenerateWidgetForColumn(const FName& ColumnId) override;
+
+private:
+
+	/** Cached row geometry, used for virtual <-> physical space conversions */
+	FGeometry RowGeometry;
+
+	/** Cached reference to a track lane that we relate to. This keeps the track lane alive (it's a weak widget) as long as we are in view. */
+	TSharedPtr<SSequencerTrackLane> TrackLaneReference;
+
+	/** The item associated with this row of data */
+	mutable TWeakPtr<FSequencerDisplayNode> Node;
+
+	/** Delegate to call to create a new widget for a particular column. */
+	FOnGenerateWidgetForColumn OnGenerateWidgetForColumn;
+};
\ No newline at end of file
diff --git a/Engine/Source/Editor/Sequencer/Private/Sequencer.cpp b/Engine/Source/Editor/Sequencer/Private/Sequencer.cpp
--- a/Engine/Source/Editor/Sequencer/Private/Sequencer.cpp
+++ b/Engine/Source/Editor/Sequencer/Private/Sequencer.cpp
@@ -170,7 +170,9 @@ FSequencer::FSequencer()
 	, bPerspectiveViewportPossessionEnabled( true )
 	, bIsEditingWithinLevelEditor( false )
 	, bNeedTreeRefresh( false )
-{ }
+{
+	Selection.GetOnOutlinerNodeSelectionChanged().AddRaw(this, &FSequencer::OnSelectedOutlinerNodesChanged);
+}
 
 FSequencer::~FSequencer()
 {
@@ -218,12 +220,19 @@ void FSequencer::Tick(float InDeltaTime)
 		bNeedTreeRefresh = false;
 	}
 	
+	static const float AutoScrollFactor = 0.1f;
+
 	// Animate the autoscroll offset if it's set
 	if (AutoscrollOffset.IsSet())
 	{
-		static const float AutoScrollFactor = 0.1f;
 		float Offset = AutoscrollOffset.GetValue() * AutoScrollFactor;
 		OnViewRangeChanged(TRange<float>(TargetViewRange.GetLowerBoundValue() + Offset, TargetViewRange.GetUpperBoundValue() + Offset), EViewRangeInterpolation::Immediate);
+	}
+
+	// Animate the autoscrub offset if it's set
+	if (AutoscrubOffset.IsSet())
+	{
+		float Offset = AutoscrubOffset.GetValue() * AutoScrollFactor;
 		SetGlobalTime(GetGlobalTime() + Offset);
 	}
 
@@ -575,14 +584,14 @@ void FSequencer::SetGlobalTime( float NewTime )
 	GEditor->RedrawLevelEditingViewports();
 }
 
-TOptional<float> FSequencer::CalculateAutoscrollEncroachment(float NewTime) const
+TOptional<float> FSequencer::CalculateAutoscrollEncroachment(float NewTime, float ThresholdPercentage) const
 {
 	enum class EDirection { Positive, Negative };
 	const EDirection Movement = NewTime - ScrubPosition >= 0 ? EDirection::Positive : EDirection::Negative;
 
 	const TRange<float> CurrentRange = GetViewRange();
 	const float RangeMin = CurrentRange.GetLowerBoundValue(), RangeMax = CurrentRange.GetUpperBoundValue();
-	const float AutoScrollThreshold = (RangeMax - RangeMin) * 0.1f;
+	const float AutoScrollThreshold = (RangeMax - RangeMin) * ThresholdPercentage;
 
 	if (Movement == EDirection::Negative && NewTime < RangeMin + AutoScrollThreshold)
 	{
@@ -989,11 +998,16 @@ void FSequencer::OnScrubPositionChanged( float NewScrubPosition, bool bScrubbing
 		else if (bAutoScrollEnabled)
 		{
 			// When scrubbing, we animate auto-scrolled scrub position in Tick()
-			AutoscrollOffset = CalculateAutoscrollEncroachment(NewScrubPosition);
+			AutoscrollOffset = CalculateAutoscrollEncroachment(NewScrubPosition, 0.025f);
 			if (AutoscrollOffset.IsSet())
 			{
+				AutoscrubOffset = AutoscrollOffset;
 				return;
 			}
+			else
+			{
+				AutoscrubOffset.Reset();
+			}
 		}
 	}
 
@@ -1008,6 +1022,17 @@ void FSequencer::OnBeginScrubbing()
 void FSequencer::OnEndScrubbing()
 {
 	PlaybackState = EMovieScenePlayerStatus::Stopped;
+	AutoscrubOffset.Reset();
+	StopAutoscroll();
+}
+
+void FSequencer::StartAutoscroll(float UnitsPerS)
+{
+	AutoscrollOffset = UnitsPerS;
+}
+
+void FSequencer::StopAutoscroll()
+{
 	AutoscrollOffset.Reset();
 }
 
@@ -1354,6 +1379,19 @@ void FSequencer::OnSectionSelectionChanged()
 	}
 }
 
+void FSequencer::OnSelectedOutlinerNodesChanged()
+{
+	if (IsLevelEditorSequencer())
+	{
+		const bool bNotifySelectionChanged = false;
+		const bool bDeselectBSP = true;
+		const bool bWarnAboutTooManyActors = false;
+
+		const FScopedTransaction Transaction(NSLOCTEXT("UnrealEd", "ClickingOnActors", "Clicking on Actors"));
+		GEditor->SelectNone(bNotifySelectionChanged, bDeselectBSP, bWarnAboutTooManyActors);
+	}
+}
+
 void FSequencer::ZoomToSelectedSections()
 {
 	TArray< TRange<float> > Bounds;
@@ -1593,12 +1631,12 @@ void FSequencer::StepToPreviousCameraKey()
 
 void FSequencer::ToggleExpandCollapseNodes()
 {
-	SequencerWidget->ToggleExpandCollapseSelectedNodes();
+	SequencerWidget->GetTreeView()->ToggleSelectedNodeExpansion(ETreeRecursion::NonRecursive);
 }
 
 void FSequencer::ToggleExpandCollapseNodesAndDescendants()
 {
-	SequencerWidget->ToggleExpandCollapseSelectedNodes(true);
+	SequencerWidget->GetTreeView()->ToggleSelectedNodeExpansion(ETreeRecursion::Recursive);
 }
 
 void FSequencer::SetKey()
diff --git a/Engine/Source/Editor/Sequencer/Private/Sequencer.h b/Engine/Source/Editor/Sequencer/Private/Sequencer.h
--- a/Engine/Source/Editor/Sequencer/Private/Sequencer.h
+++ b/Engine/Source/Editor/Sequencer/Private/Sequencer.h
@@ -317,10 +317,7 @@ protected:
 	 */
 	void OnViewRangeChanged( TRange<float> NewViewRange, EViewRangeInterpolation Interpolation = EViewRangeInterpolation::Animated );
 
-	/**
-	 * Calculates the amount of encroachment the specified time has into the autoscroll region, if any
-	 */
-	TOptional<float> CalculateAutoscrollEncroachment(float NewTime) const;
+
 
 	/**
 	 * Called when the scrub position is changed by the user
@@ -336,6 +333,20 @@ protected:
 	/** Called when the user has finished scrubbing */
 	void OnEndScrubbing();
 
+public:
+
+	/** Put the sequencer in a horizontally autoscrolling state with the given rate */
+	void StartAutoscroll(float UnitsPerS);
+
+	/** Stop the sequencer from autoscrolling */
+	void StopAutoscroll();
+
+protected:
+	/**
+	 * Calculates the amount of encroachment the specified time has into the autoscroll region, if any
+	 */
+	TOptional<float> CalculateAutoscrollEncroachment(float NewTime, float ThresholdPercentage = 0.1f) const;
+
 	/** Called to toggle auto-scroll on and off */
 	void OnToggleAutoScroll();
 
@@ -349,6 +360,8 @@ protected:
 		return bAutoScrollEnabled;
 	}
 
+protected:
+
 	/** Called via UEditorEngine::GetActorRecordingStateEvent to check to see whether we need to record actor state */
 	void GetActorRecordingState( bool& bIsRecording /* In+Out */ ) const;
 	
@@ -386,6 +399,8 @@ protected:
 
 	void OnSectionSelectionChanged();
 
+	void OnSelectedOutlinerNodesChanged();
+
 	/** Called before the world is going to be saved. The sequencer puts everything back to its initial state. */
 	void OnPreSaveWorld(uint32 SaveFlags, class UWorld* World);
 
@@ -451,6 +466,8 @@ private:
 	/** The amount of autoscroll pan offset that is currently being applied */
 	TOptional<float> AutoscrollOffset;
 
+	/** The amount of autoscrub offset that is currently being applied */
+	TOptional<float> AutoscrubOffset;
 	/** Whether or not we are allowing autoscroll */
 	bool bAutoScrollEnabled;
 
diff --git a/Engine/Source/Editor/Sequencer/Private/SequencerDisplayNode.cpp b/Engine/Source/Editor/Sequencer/Private/SequencerDisplayNode.cpp
--- a/Engine/Source/Editor/Sequencer/Private/SequencerDisplayNode.cpp
+++ b/Engine/Source/Editor/Sequencer/Private/SequencerDisplayNode.cpp
@@ -17,6 +17,10 @@
 
 #define LOCTEXT_NAMESPACE "SequencerDisplayNode"
 
+namespace SequencerNodeConstants
+{
+	float CommonPadding = 4.f;
+}
 
 class SSequencerObjectTrack : public SLeafWidget
 {
@@ -55,17 +59,6 @@ private:
 
 int32 SSequencerObjectTrack::OnPaint(const FPaintArgs& Args, const FGeometry& AllottedGeometry, const FSlateRect& MyClippingRect, FSlateWindowElementList& OutDrawElements, int32 LayerId, const FWidgetStyle& InWidgetStyle, bool bParentEnabled) const
 {
-	// Draw a region around the entire section area
-	FSlateDrawElement::MakeBox( 
-		OutDrawElements, 
-		LayerId,
-		AllottedGeometry.ToPaintGeometry(),
-		FEditorStyle::GetBrush("Sequencer.SectionArea.Background"),
-		MyClippingRect,
-		ESlateDrawEffect::None,
-		FLinearColor( .1f, .1f, .1f, 0.5f )
-		);
-
 	TArray<float> OutKeyTimes;
 	CollectAllKeyTimes(OutKeyTimes);
 	
@@ -134,7 +127,6 @@ FSequencerDisplayNode::FSequencerDisplayNode( FName InNodeName, TSharedPtr<FSequ
 	: ParentNode( InParentNode )
 	, ParentTree( InParentTree )
 	, NodeName( InNodeName )
-	, TreeLevel( InParentNode.IsValid() ? InParentNode->GetTreeLevel() + 1 : 0 )
 	, bExpanded( false )
 	, bCachedShotFilteredVisibility( true )
 	, bNodeIsPinned( false )
@@ -147,6 +139,84 @@ void FSequencerDisplayNode::AddObjectBindingNode(TSharedRef<FObjectBindingNode>
 	ChildNodes.Add(ObjectBindingNode);
 }
 
+
+bool FSequencerDisplayNode::Traverse_ChildFirst(const TFunctionRef<bool(FSequencerDisplayNode&)>& InPredicate, bool bIncludeThisNode)
+{
+	for (auto& Child : GetChildNodes())
+	{
+		if (!Child->Traverse_ChildFirst(InPredicate, true))
+		{
+			return false;
+		}
+	}
+
+	return bIncludeThisNode ? InPredicate(*this) : true;
+}
+
+bool FSequencerDisplayNode::Traverse_ParentFirst(const TFunctionRef<bool(FSequencerDisplayNode&)>& InPredicate, bool bIncludeThisNode)
+{
+	if (bIncludeThisNode && !InPredicate(*this))
+	{
+		return false;
+	}
+
+	for (auto& Child : GetChildNodes())
+	{
+		if (!Child->Traverse_ParentFirst(InPredicate, true))
+		{
+			return false;
+		}
+	}
+
+	return true;
+}
+
+bool FSequencerDisplayNode::TraverseVisible_ChildFirst(const TFunctionRef<bool(FSequencerDisplayNode&)>& InPredicate, bool bIncludeThisNode)
+{
+	// If the item is not expanded, its children ain't visible
+	if (IsExpanded())
+	{
+		for (auto& Child : GetChildNodes())
+		{
+			if (!Child->IsHidden() && !Child->TraverseVisible_ChildFirst(InPredicate, true))
+			{
+				return false;
+			}
+		}
+	}
+
+	if (bIncludeThisNode && !IsHidden())
+	{
+		return InPredicate(*this);
+	}
+
+	// Continue iterating regardless of visibility
+	return true;
+}
+
+
+bool FSequencerDisplayNode::TraverseVisible_ParentFirst(const TFunctionRef<bool(FSequencerDisplayNode&)>& InPredicate, bool bIncludeThisNode)
+{
+	if (bIncludeThisNode && !IsHidden() && !InPredicate(*this))
+	{
+		return false;
+	}
+
+	// If the item is not expanded, its children ain't visible
+	if (IsExpanded())
+	{
+		for (auto& Child : GetChildNodes())
+		{
+			if (!Child->IsHidden() && !Child->TraverseVisible_ParentFirst(InPredicate, true))
+			{
+				return false;
+			}
+		}
+	}
+
+	return true;
+}
+
 TSharedRef<FSectionCategoryNode> FSequencerDisplayNode::AddCategoryNode( FName CategoryName, const FText& DisplayLabel )
 {
 	TSharedPtr<FSectionCategoryNode> CategoryNode;
@@ -222,9 +292,9 @@ void FSequencerDisplayNode::AddKeyAreaNode( FName KeyAreaName, const FText& Disp
 	KeyAreaNode->AddKeyArea( KeyArea );
 }
 
-TSharedRef<SWidget> FSequencerDisplayNode::GenerateContainerWidgetForOutliner()
+TSharedRef<SWidget> FSequencerDisplayNode::GenerateContainerWidgetForOutliner(const TSharedRef<SSequencerTreeViewRow>& InRow)
 {
-	return SNew( SAnimationOutlinerView, SharedThis( this ), &GetSequencer() );
+	return SNew( SAnimationOutlinerTreeNode, SharedThis( this ), InRow, &GetSequencer() );
 }
 
 TSharedRef<SWidget> FSequencerDisplayNode::GenerateEditWidgetForOutliner()
@@ -252,6 +322,24 @@ TSharedRef<SWidget> FSequencerDisplayNode::GenerateWidgetForSectionArea( const T
 	}
 }
 
+TSharedPtr<FSequencerDisplayNode> FSequencerDisplayNode::GetSectionAreaAuthority() const
+{
+	TSharedPtr<FSequencerDisplayNode> Authority = SharedThis(const_cast<FSequencerDisplayNode*>(this));
+
+	while (Authority.IsValid())
+	{
+		if (Authority->GetType() == ESequencerNode::Object || Authority->GetType() == ESequencerNode::Track)
+		{
+			return Authority;
+		}
+		else
+		{
+			Authority = Authority->GetParent();
+		}
+	}
+	return Authority;
+}
+
 FString FSequencerDisplayNode::GetPathName() const
 {
 	// First get our parent's path
@@ -279,9 +367,9 @@ void FSequencerDisplayNode::GetChildKeyAreaNodesRecursively(TArray< TSharedRef<F
 	}
 }
 
-void FSequencerDisplayNode::ToggleExpansion()
+void FSequencerDisplayNode::SetExpansionState(bool bInExpanded)
 {
-	bExpanded = !bExpanded;
+	bExpanded = bInExpanded;
 
 	// Expansion state has changed, save it to the movie scene now
 	ParentTree.SaveExpansionState( *this, bExpanded );
@@ -294,14 +382,15 @@ bool FSequencerDisplayNode::IsExpanded() const
 	return ParentTree.HasActiveFilter() ? ParentTree.IsNodeFiltered( AsShared() ) : bExpanded;
 }
 
-bool FSequencerDisplayNode::IsVisible() const
+bool FSequencerDisplayNode::IsHidden() const
 {
-	// Must be visible after shot filtering AND
-	// If there is a search filter, must be filtered, otherwise, it's parent must be expanded AND
-	// If shot filtering is off on clean view is on, node must be pinned
-	return bCachedShotFilteredVisibility &&
-		(ParentTree.HasActiveFilter() ? ParentTree.IsNodeFiltered(AsShared()) : IsParentExpandedOrIsARootNode()) &&
-		(GetSequencer().IsShotFilteringOn() || !GetSequencer().GetSettings()->GetIsUsingCleanView() || bNodeIsPinned);
+	// Not visible if:
+	// Its cached shot visibility is false OR
+	// There is a search filter, and it's not filtered, OR
+	// If shot filtering is off and clean view is on, but the node isn't pinned
+	return !bCachedShotFilteredVisibility ||
+		(ParentTree.HasActiveFilter() && !ParentTree.IsNodeFiltered(AsShared())) ||
+		(!GetSequencer().IsShotFilteringOn() && GetSequencer().GetSettings()->GetIsUsingCleanView() && !bNodeIsPinned);
 }
 
 bool FSequencerDisplayNode::HasVisibleChildren() const
@@ -313,11 +402,6 @@ bool FSequencerDisplayNode::HasVisibleChildren() const
 	return false;
 }
 
-bool FSequencerDisplayNode::IsParentExpandedOrIsARootNode() const
-{
-	return !ParentNode.IsValid() || ParentNode.Pin()->bExpanded;
-}
-
 void FSequencerDisplayNode::UpdateCachedShotFilteredVisibility()
 {
 	// Tell our children to update their visibility first
@@ -344,6 +428,11 @@ float FSectionKeyAreaNode::GetNodeHeight() const
 	return SequencerLayoutConstants::KeyAreaHeight;
 }
 
+FNodePadding FSectionKeyAreaNode::GetNodePadding() const
+{
+	return FNodePadding(0.f, 1.f);
+}
+
 bool FSectionKeyAreaNode::GetShotFilteredVisibilityToCache() const
 {
 	return true;
@@ -379,6 +468,11 @@ float FTrackNode::GetNodeHeight() const
 	return Sections.Num() > 0 ? Sections[0]->GetSectionHeight() * (GetMaxRowIndex()+1) : SequencerLayoutConstants::SectionAreaDefaultHeight;
 }
 
+FNodePadding FTrackNode::GetNodePadding() const
+{
+	return FNodePadding(SequencerNodeConstants::CommonPadding);
+}
+
 FText FTrackNode::GetDisplayName() const
 {
 	// @todo Sequencer - IS there a better way to get the section interface name for the animation outliner?
@@ -488,10 +582,15 @@ FText FObjectBindingNode::GetDisplayName() const
 }
 
 float FObjectBindingNode::GetNodeHeight() const
-{ 
+{
 	return SequencerLayoutConstants::ObjectNodeHeight;
 }
 
+FNodePadding FObjectBindingNode::GetNodePadding() const
+{
+	return FNodePadding(SequencerNodeConstants::CommonPadding * 2, 0.f);
+}
+
 bool FObjectBindingNode::GetShotFilteredVisibilityToCache() const
 {
 	// if shot filtering is off and we are not unfilterable
@@ -666,6 +765,11 @@ float FSectionCategoryNode::GetNodeHeight() const
 	return SequencerLayoutConstants::CategoryNodeHeight;
 }
 
+FNodePadding FSectionCategoryNode::GetNodePadding() const
+{
+	return FNodePadding(SequencerNodeConstants::CommonPadding/2);
+}
+
 bool FSectionCategoryNode::GetShotFilteredVisibilityToCache() const
 {
 	// this node is only visible if at least one child node is visible
diff --git a/Engine/Source/Editor/Sequencer/Private/SequencerDisplayNode.h b/Engine/Source/Editor/Sequencer/Private/SequencerDisplayNode.h
--- a/Engine/Source/Editor/Sequencer/Private/SequencerDisplayNode.h
+++ b/Engine/Source/Editor/Sequencer/Private/SequencerDisplayNode.h
@@ -4,6 +4,22 @@
 
 class IKeyArea;
 class ISequencerSection;
+class SSequencerTreeViewRow;
+
+/** Structure used to define padding for a particular node */
+struct FNodePadding
+{
+	FNodePadding(float InUniform) : Top(InUniform), Bottom(InUniform) {}
+	FNodePadding(float InTop, float InBottom) : Top(InTop), Bottom(InBottom) {}
+	
+	/** @return The sum total of the separate padding values */
+	float Combined() const { return Top + Bottom; }
+
+	/** Padding to be applied to the top of the node */
+	float Top;
+	/** Padding to be applied to the bottom of the node */
+	float Bottom;
+};
 
 namespace ESequencerNode
 {
@@ -81,6 +97,11 @@ public:
 	 */
 	virtual float GetNodeHeight() const = 0;
 	
+	/**
+	 * @return The desired padding of the node when displayed
+	 */
+	virtual FNodePadding GetNodePadding() const = 0;
+
 	/**
 	 * @return The localized display name of this node
 	 */
@@ -91,7 +112,7 @@ public:
 	 * 
 	 * @return Generated outliner container widget
 	 */
-	virtual TSharedRef<SWidget> GenerateContainerWidgetForOutliner();
+	virtual TSharedRef<SWidget> GenerateContainerWidgetForOutliner(const TSharedRef<SSequencerTreeViewRow>& InRow);
 
 	/**
 	* Generates a widget editing a row in the animation outliner portion of the track area
@@ -108,6 +129,12 @@ public:
 	 */
 	virtual TSharedRef<SWidget> GenerateWidgetForSectionArea( const TAttribute< TRange<float> >& ViewRange );
 
+	/**
+	 * Get the display node that is ultimately responsible for constructing a section area widget for this node.
+	 * Could return this node itself, or a parent node
+	 */
+	TSharedPtr<FSequencerDisplayNode> GetSectionAreaAuthority() const;
+
 	/**
 	 * @return the path to this node starting with the outermost parent
 	 */
@@ -127,14 +154,41 @@ public:
 	uint32 GetNumChildren() const { return ChildNodes.Num(); }
 
 	/**
-	 * @return How deep we are in the master tree
+	 * @return A List of all Child nodes belonging to this node
+	 */
+	const TArray< TSharedRef<FSequencerDisplayNode> >& GetChildNodes() const { return ChildNodes; }
+
+	/**
+	 * Iterate this entire node tree, child first.
+	 * @param 	InPredicate			Predicate to call for each node, returning whether to continue iteration or not
+	 * @param 	bIncludeThisNode	Whether to include this node in the iteration, or just children
+	 * @return  true where the client prematurely exited the iteration, false otherwise
 	 */
-	uint32 GetTreeLevel() const { return TreeLevel; }
+	bool Traverse_ChildFirst(const TFunctionRef<bool(FSequencerDisplayNode&)>& InPredicate, bool bIncludeThisNode = true);
 
 	/**
-	 * @return A List of all Child nodes belonging to this node
+	 * Iterate this entire node tree, parent first.
+	 * @param 	InPredicate			Predicate to call for each node, returning whether to continue iteration or not
+	 * @param 	bIncludeThisNode	Whether to include this node in the iteration, or just children
+	 * @return  true where the client prematurely exited the iteration, false otherwise
 	 */
-	const TArray< TSharedRef<FSequencerDisplayNode> >& GetChildNodes() const { return ChildNodes; }
+	bool Traverse_ParentFirst(const TFunctionRef<bool(FSequencerDisplayNode&)>& InPredicate, bool bIncludeThisNode = true);
+
+	/**
+	 * Iterate any visible portions of this node's sub-tree, child first.
+	 * @param 	InPredicate			Predicate to call for each node, returning whether to continue iteration or not
+	 * @param 	bIncludeThisNode	Whether to include this node in the iteration, or just children
+	 * @return  true where the client prematurely exited the iteration, false otherwise
+	 */
+	bool TraverseVisible_ChildFirst(const TFunctionRef<bool(FSequencerDisplayNode&)>& InPredicate, bool bIncludeThisNode = true);
+
+	/**
+	 * Iterate any visible portions of this node's sub-tree, parent first.
+	 * @param 	InPredicate			Predicate to call for each node, returning whether to continue iteration or not
+	 * @param 	bIncludeThisNode	Whether to include this node in the iteration, or just children
+	 * @return  true where the client prematurely exited the iteration, false otherwise
+	 */
+	bool TraverseVisible_ParentFirst(const TFunctionRef<bool(FSequencerDisplayNode&)>& InPredicate, bool bIncludeThisNode = true);
 
 	/**
 	 * Sorts the child nodes with the supplied predicate
@@ -164,19 +218,19 @@ public:
 	virtual void GetChildKeyAreaNodesRecursively(TArray< TSharedRef<class FSectionKeyAreaNode> >& OutNodes) const;
 
 	/**
-	 * Toggles the expansion state of this node
+	 * Set whether this node is expanded or not
 	 */
-	void ToggleExpansion();
+	void SetExpansionState(bool bInExpanded);
 
 	/**
-	 * @return Whether or not this node is expanded                                                              
+	 * @return Whether or not this node is expanded
 	 */
 	bool IsExpanded() const;
 
 	/**
-	 * @return Whether or not a node is visible
+	 * @return Whether this node is explicitly hidden from the view or not
 	 */
-	bool IsVisible() const;
+	bool IsHidden() const;
 
 	/** Updates the cached shot filtered visibility flag */
 	void UpdateCachedShotFilteredVisibility();
@@ -218,9 +272,9 @@ protected:
 	
 	/** Whether this node has visible children, based on cached shot filtering visibility only */
 	bool HasVisibleChildren() const;
-	/** Whether this node is a root node, or it's parent is expanded */
-	bool IsParentExpandedOrIsARootNode() const;
-	
+public:
+	float VirtualTop;
+	float VirtualBottom;
 protected:
 	/** The parent of this node*/
 	TWeakPtr< FSequencerDisplayNode > ParentNode;
@@ -230,8 +284,6 @@ protected:
 	FSequencerNodeTree& ParentTree;
 	/** The name identifier of this node */
 	FName NodeName;
-	/** How far deep in the tree the node is */
-	uint32 TreeLevel;
 	/** Whether or not the node is expanded */
 	bool bExpanded;
 	/** A cached state of what this node's visibility is based only on shot filtering */
@@ -265,6 +317,7 @@ public:
 	/** FSequencerDisplayNodeInterface */
 	virtual ESequencerNode::Type GetType() const override { return ESequencerNode::KeyArea; }
 	virtual float GetNodeHeight() const override;
+	virtual FNodePadding GetNodePadding() const override;
 	virtual FText GetDisplayName() const override { return DisplayName; }
 	virtual bool GetShotFilteredVisibilityToCache() const override;
 	virtual TSharedRef<SWidget> GenerateEditWidgetForOutliner() override;
@@ -322,6 +375,7 @@ public:
 	/** FSequencerDisplayNodeInterface */
 	virtual ESequencerNode::Type GetType() const override { return ESequencerNode::Track; }
 	virtual float GetNodeHeight() const override;
+	virtual FNodePadding GetNodePadding() const override;
 	virtual FText GetDisplayName() const override;
 	virtual bool GetShotFilteredVisibilityToCache() const override;
 	virtual void GetChildKeyAreaNodesRecursively(TArray< TSharedRef<class FSectionKeyAreaNode> >& OutNodes) const override;
@@ -396,6 +450,7 @@ public:
 	virtual ESequencerNode::Type GetType() const override { return ESequencerNode::Object; }
 	virtual FText GetDisplayName() const override;
 	virtual float GetNodeHeight() const override;
+	virtual FNodePadding GetNodePadding() const override;
 	virtual bool GetShotFilteredVisibilityToCache() const override;
 	TSharedRef<SWidget> GenerateEditWidgetForOutliner() override;
 	
@@ -441,6 +496,7 @@ public:
 	/** FSequencerDisplayNodeInterface */
 	virtual ESequencerNode::Type GetType() const override { return ESequencerNode::Category; }
 	virtual float GetNodeHeight() const override;
+	virtual FNodePadding GetNodePadding() const override;
 	virtual FText GetDisplayName() const override { return DisplayName; }
 	virtual bool GetShotFilteredVisibilityToCache() const override;
 private:
diff --git a/Engine/Source/Editor/Sequencer/Private/SequencerNodeTree.cpp b/Engine/Source/Editor/Sequencer/Private/SequencerNodeTree.cpp
--- a/Engine/Source/Editor/Sequencer/Private/SequencerNodeTree.cpp
+++ b/Engine/Source/Editor/Sequencer/Private/SequencerNodeTree.cpp
@@ -115,6 +115,18 @@ void FSequencerNodeTree::Update()
 	// Add all other nodes after the shot track
 	RootNodes.Append( NewRootNodes );
 
+	float VerticalOffset = 0.f;
+	for (auto& Node : RootNodes)
+	{
+		Node->Traverse_ParentFirst([&](FSequencerDisplayNode& InNode){
+			InNode.VirtualTop = VerticalOffset;
+			VerticalOffset += InNode.GetNodeHeight() + InNode.GetNodePadding().Combined();
+			InNode.VirtualBottom = VerticalOffset;
+			return true;
+		});
+
+	}
+
 	// Re-filter the tree after updating 
 	// @todo Sequencer - Newly added sections may need to be visible even when there is a filter
 	FilterNodes( FilterString );
diff --git a/Engine/Source/Editor/Sequencer/Private/SequencerPrivatePCH.h b/Engine/Source/Editor/Sequencer/Private/SequencerPrivatePCH.h
--- a/Engine/Source/Editor/Sequencer/Private/SequencerPrivatePCH.h
+++ b/Engine/Source/Editor/Sequencer/Private/SequencerPrivatePCH.h
@@ -19,8 +19,7 @@
 #include "ITimeSlider.h"
 #include "SSequencer.h"
 #include "SSection.h"
-#include "SSequencerTrackViewPanel.h"
-#include "SAnimationOutlinerView.h"
+#include "SAnimationOutlinerTreeNode.h"
 #include "MovieSceneInstance.h"
 #include "SequencerSettings.h"
 #include "SequencerCommonHelpers.h"
diff --git a/Engine/Source/Editor/Sequencer/Private/SequencerTrackLaneFactory.cpp b/Engine/Source/Editor/Sequencer/Private/SequencerTrackLaneFactory.cpp
deleted file mode 100644
diff --git a/Engine/Source/Editor/Sequencer/Private/SequencerTrackLaneFactory.h b/Engine/Source/Editor/Sequencer/Private/SequencerTrackLaneFactory.h
deleted file mode 100644
diff --git a/Engine/Source/Editor/Sequencer/Private/TimeSliderController.cpp b/Engine/Source/Editor/Sequencer/Private/TimeSliderController.cpp
--- a/Engine/Source/Editor/Sequencer/Private/TimeSliderController.cpp
+++ b/Engine/Source/Editor/Sequencer/Private/TimeSliderController.cpp
@@ -466,12 +466,14 @@ int32 FSequencerTimeSliderController::OnPaintSectionView( const FGeometry& Allot
 
 	if( bDisplayTickLines )
 	{
+		static FLinearColor TickColor(0.f, 0.f, 0.f, 0.3f);
+		
 		// Draw major tick lines in the section area
 		FDrawTickArgs Args;
 		Args.AllottedGeometry = AllottedGeometry;
 		Args.bMirrorLabels = false;
 		Args.bOnlyDrawMajorTicks = true;
-		Args.TickColor = FLinearColor( 0.3f, 0.3f, 0.3f, 0.3f );
+		Args.TickColor = TickColor;
 		Args.ClippingRect = MyClippingRect;
 		Args.DrawEffects = DrawEffects;
 		// Draw major ticks under sections
