Fixed memory leak which ultimately resulted in a crash due to a delegate binding outliving its assumption

diff --git a/Engine/Source/Editor/Sequencer/Private/SSequencerCurveEditor.cpp b/Engine/Source/Editor/Sequencer/Private/SSequencerCurveEditor.cpp
--- a/Engine/Source/Editor/Sequencer/Private/SSequencerCurveEditor.cpp
+++ b/Engine/Source/Editor/Sequencer/Private/SSequencerCurveEditor.cpp
@@ -110,46 +110,38 @@ void SSequencerCurveEditor::SetSequencerNodeTree( TSharedPtr<FSequencerNodeTree>
 
 void SSequencerCurveEditor::UpdateCurveOwner()
 {
-	FSequencerCurveOwner* NewCurveOwner = new FSequencerCurveOwner( SequencerNodeTree, GetSettings()->GetCurveVisibility() );
+	TSharedRef<FSequencerCurveOwner> NewCurveOwner = MakeShareable( new FSequencerCurveOwner( SequencerNodeTree, GetSettings()->GetCurveVisibility() ) );
 
 	bool bAllFound = false;
-	if (NewCurveOwner != NULL && CurveOwner.IsValid())
+	if (CurveOwner.IsValid())
 	{
 		TSet<FName> NewCurveOwnerCurveNames;
-		if (NewCurveOwner != NULL)
+		for (auto Curve : NewCurveOwner->GetCurves())
 		{
-			for (auto Curve : NewCurveOwner->GetCurves())
-			{
-				NewCurveOwnerCurveNames.Add(Curve.CurveName);
-			}
+			NewCurveOwnerCurveNames.Add(Curve.CurveName);
 		}
 
 		TSet<FName> CurveOwnerCurveNames;
-		if (CurveOwner.Get() != NULL)
+		if (CurveOwner->GetCurves().Num() == NewCurveOwner->GetCurves().Num())
 		{
-			if (CurveOwner.Get()->GetCurves().Num() == NewCurveOwner->GetCurves().Num())
+			bAllFound = true;
+			for (auto Curve : CurveOwner->GetCurves())
 			{
-				bAllFound = true;
-				for (auto Curve : CurveOwner.Get()->GetCurves())
+				if (!NewCurveOwnerCurveNames.Contains(Curve.CurveName))
 				{
-					if (!NewCurveOwnerCurveNames.Contains(Curve.CurveName))
-					{
-						bAllFound = false;
-					}
+					bAllFound = false;
+					break;
 				}
 			}
 		}
 	}
 
-	if (bAllFound)
+	if (!bAllFound)
 	{
-		UpdateCurveViewModelSelection();
-		return;
+		CurveOwner = NewCurveOwner;
+		SetCurveOwner( CurveOwner.Get() );
 	}
-
-	CurveOwner = MakeShareable( NewCurveOwner );
-	SetCurveOwner( CurveOwner.Get() );
-
+	
 	UpdateCurveViewModelSelection();
 }
 
