Fixed not being able to delete things in the sequencer

diff --git a/Engine/Source/Editor/Sequencer/Private/Sequencer.cpp b/Engine/Source/Editor/Sequencer/Private/Sequencer.cpp
--- a/Engine/Source/Editor/Sequencer/Private/Sequencer.cpp
+++ b/Engine/Source/Editor/Sequencer/Private/Sequencer.cpp
@@ -1581,7 +1581,7 @@ void FSequencer::DeleteSelectedItems()
 bool FSequencer::CanDeleteSelectedItems()
 {
 	// Only delete sequencer items if the delete key was pressed in the sequencer widget and not in, for example, the viewport.
-	return SequencerWidget->HasKeyboardFocus();
+	return SequencerWidget->HasFocusedDescendants();
 }
 
 void FSequencer::TogglePlay()
