Fixed sequencer tree not updating on some UI interaction (clean view, search filter)

diff --git a/Engine/Source/Editor/Sequencer/Private/SSequencer.cpp b/Engine/Source/Editor/Sequencer/Private/SSequencer.cpp
--- a/Engine/Source/Editor/Sequencer/Private/SSequencer.cpp
+++ b/Engine/Source/Editor/Sequencer/Private/SSequencer.cpp
@@ -727,6 +727,7 @@ void SSequencer::ResetBreadcrumbs()
 void SSequencer::OnOutlinerSearchChanged( const FText& Filter )
 {
 	SequencerNodeTree->FilterNodes( Filter.ToString() );
+	TreeView->Refresh();
 }
 
 float SSequencer::OnGetTimeSnapInterval() const
diff --git a/Engine/Source/Editor/Sequencer/Private/Sequencer.cpp b/Engine/Source/Editor/Sequencer/Private/Sequencer.cpp
--- a/Engine/Source/Editor/Sequencer/Private/Sequencer.cpp
+++ b/Engine/Source/Editor/Sequencer/Private/Sequencer.cpp
@@ -1732,7 +1732,13 @@ void FSequencer::BindSequencerCommands()
 
 	SequencerCommandBindings->MapAction(
 		Commands.ToggleCleanView,
-		FExecuteAction::CreateLambda( [this]{ Settings->SetIsUsingCleanView( !Settings->GetIsUsingCleanView() ); } ),
+		FExecuteAction::CreateLambda( [this]{
+			Settings->SetIsUsingCleanView( !Settings->GetIsUsingCleanView() );
+			if (SequencerWidget.IsValid())
+			{
+				SequencerWidget->GetTreeView()->Refresh();
+			}
+		} ),
 		FCanExecuteAction::CreateLambda( []{ return true; } ),
 		FIsActionChecked::CreateLambda( [this]{ return Settings->GetIsUsingCleanView(); } ) );
 
