Changed sequencer tree scrolling to better suit variable height items

diff --git a/Engine/Source/Editor/Sequencer/Private/SSequencerTreeView.cpp b/Engine/Source/Editor/Sequencer/Private/SSequencerTreeView.cpp
--- a/Engine/Source/Editor/Sequencer/Private/SSequencerTreeView.cpp
+++ b/Engine/Source/Editor/Sequencer/Private/SSequencerTreeView.cpp
@@ -91,6 +91,7 @@ void SSequencerTreeView::Construct(const FArguments& InArgs, const TSharedRef<FS
 void SSequencerTreeView::Tick(const FGeometry& AllottedGeometry, const double InCurrentTime, const float InDeltaTime)
 {
 	STreeView::Tick(AllottedGeometry, InCurrentTime, InDeltaTime);
+	CachedTreeGeometry = AllottedGeometry;
 	UpdateCachedVerticalGeometry();
 }
 
@@ -244,8 +245,7 @@ void SSequencerTreeView::Refresh()
 
 void SSequencerTreeView::ScrollByDelta(float DeltaInSlateUnits)
 {
-	const float FractionalScrollOffsetInItems = (ScrollOffset + GetScrollRateInItems() * DeltaInSlateUnits);
-	ScrollTo( FractionalScrollOffsetInItems );
+	ScrollBy( CachedTreeGeometry, DeltaInSlateUnits, EAllowOverscroll::No );
 }
 
 template<typename T>
diff --git a/Engine/Source/Editor/Sequencer/Private/SSequencerTreeView.h b/Engine/Source/Editor/Sequencer/Private/SSequencerTreeView.h
--- a/Engine/Source/Editor/Sequencer/Private/SSequencerTreeView.h
+++ b/Engine/Source/Editor/Sequencer/Private/SSequencerTreeView.h
@@ -103,6 +103,9 @@ protected:
 
 private:
 
+	/** The cached geometry for this tree view*/
+	FGeometry CachedTreeGeometry;
+
 	/** The tree view's header row (hidden) */
 	TSharedPtr<SHeaderRow> HeaderRow;
 
