Fixed not being able to marquee select keys on top level key areas

diff --git a/Engine/Source/Editor/Sequencer/Private/SSequencerTrackArea.cpp b/Engine/Source/Editor/Sequencer/Private/SSequencerTrackArea.cpp
--- a/Engine/Source/Editor/Sequencer/Private/SSequencerTrackArea.cpp
+++ b/Engine/Source/Editor/Sequencer/Private/SSequencerTrackArea.cpp
@@ -328,10 +328,47 @@ void SSequencerTrackArea::HandleMarqueeSelection(const FGeometry& Geometry, cons
 	auto TopLeft = MarqueeSelectData->TopLeft();
 	auto BottomRight = MarqueeSelectData->BottomRight();
 
+	// Pointer to the current traversal parent's sections.
+	const TArray<TSharedRef<ISequencerSection>>* CurrentSections = nullptr;
+
+	auto PerformSelection = [&](FSectionKeyAreaNode& KeyAreaNode){
+		for( int32 SectionIndex = 0; SectionIndex < CurrentSections->Num(); ++SectionIndex )
+		{
+			TSharedRef<IKeyArea> KeyArea = KeyAreaNode.GetKeyArea(SectionIndex);
+
+			for (FKeyHandle KeyHandle : KeyArea->GetUnsortedKeyHandles())
+			{
+				float KeyPosition = KeyArea->GetKeyTime(KeyHandle);
+				if (KeyPosition > TopLeft.X && KeyPosition < BottomRight.X)
+				{
+					FSelectedKey SelectedKey(*(*CurrentSections)[SectionIndex]->GetSectionObject(), KeyArea, KeyHandle);
+
+					if (MouseEvent.IsLeftShiftDown())
+					{
+						Selection.AddToSelection(SelectedKey);
+					}
+					else if (MouseEvent.IsLeftControlDown())
+					{
+						if (Selection.IsSelected(SelectedKey))
+						{
+							Selection.RemoveFromSelection(SelectedKey);
+						}
+						else
+						{
+							Selection.AddToSelection(SelectedKey);
+						}
+					}
+					else
+					{
+						Selection.AddToSelection(SelectedKey);
+					}
+				}
+			}
+		}
+	};
+
 	for (auto& RootNode : PinnedTreeView->GetNodeTree()->GetRootNodes())
 	{
-		const TArray<TSharedRef<ISequencerSection>>* CurrentSections = nullptr;
-
 		RootNode->TraverseVisible_ParentFirst([&](FSequencerDisplayNode& InNode){
 
 			if (InNode.GetType() == ESequencerNode::KeyArea)
@@ -340,46 +377,20 @@ void SSequencerTrackArea::HandleMarqueeSelection(const FGeometry& Geometry, cons
 
 				if (CurrentSections && InNode.VirtualTop > TopLeft.Y && InNode.VirtualBottom < BottomRight.Y)
 				{
-					for( int32 SectionIndex = 0; SectionIndex < CurrentSections->Num(); ++SectionIndex )
-					{
-						FTimeToPixel TimeToPixelConverter(Geometry, PinnedSequencer->GetViewRange());
-
-						TSharedRef<IKeyArea> KeyArea = KeyAreaNode.GetKeyArea(SectionIndex);
-
-						for (FKeyHandle KeyHandle : KeyArea->GetUnsortedKeyHandles())
-						{
-							float KeyPosition = KeyArea->GetKeyTime(KeyHandle);
-							if (KeyPosition > TopLeft.X && KeyPosition < BottomRight.X)
-							{
-								FSelectedKey SelectedKey(*(*CurrentSections)[SectionIndex]->GetSectionObject(), KeyArea, KeyHandle);
-
-								if (MouseEvent.IsLeftShiftDown())
-								{
-									Selection.AddToSelection(SelectedKey);
-								}
-								else if (MouseEvent.IsLeftControlDown())
-								{
-									if (Selection.IsSelected(SelectedKey))
-									{
-										Selection.RemoveFromSelection(SelectedKey);
-									}
-									else
-									{
-										Selection.AddToSelection(SelectedKey);
-									}
-								}
-								else
-								{
-									Selection.AddToSelection(SelectedKey);
-								}
-							}
-						}
-					}
+					PerformSelection(KeyAreaNode);
 				}
 			}
 			else if (InNode.GetType() == ESequencerNode::Track)
 			{
-				CurrentSections = &static_cast<FTrackNode&>(InNode).GetSections();
+				FTrackNode& TrackNode = static_cast<FTrackNode&>(InNode);
+
+				CurrentSections = &TrackNode.GetSections();
+
+				TSharedPtr<FSectionKeyAreaNode> SectionKeyNode = TrackNode.GetTopLevelKeyNode();
+				if( SectionKeyNode.IsValid() && TrackNode.VirtualTop > TopLeft.Y && TrackNode.VirtualBottom < BottomRight.Y )
+				{
+					PerformSelection(*SectionKeyNode);
+				}
 			}
 
 			return true;
