Non unity compile fix

diff --git a/Engine/Source/Editor/Sequencer/Private/SSequencerTrackArea.cpp b/Engine/Source/Editor/Sequencer/Private/SSequencerTrackArea.cpp
--- a/Engine/Source/Editor/Sequencer/Private/SSequencerTrackArea.cpp
+++ b/Engine/Source/Editor/Sequencer/Private/SSequencerTrackArea.cpp
@@ -4,6 +4,8 @@
 #include "SSequencerTrackArea.h"
 #include "SSequencerTrackLane.h"
 #include "TimeSliderController.h"
+#include "CommonMovieSceneTools.h"
+#include "Sequencer.h"
 
 
 struct FMarqueeSelectData
