Some nonunity fixes

diff --git a/Engine/Source/Editor/Sequencer/Private/SSequencerTrackArea.cpp b/Engine/Source/Editor/Sequencer/Private/SSequencerTrackArea.cpp
--- a/Engine/Source/Editor/Sequencer/Private/SSequencerTrackArea.cpp
+++ b/Engine/Source/Editor/Sequencer/Private/SSequencerTrackArea.cpp
@@ -6,6 +6,9 @@
 #include "TimeSliderController.h"
 #include "CommonMovieSceneTools.h"
 #include "Sequencer.h"
+#include "SSequencerTreeView.h"
+#include "IKeyArea.h"
+#include "ISequencerSection.h"
 
 
 struct FMarqueeSelectData
diff --git a/Engine/Source/Editor/Sequencer/Private/SSequencerTreeView.cpp b/Engine/Source/Editor/Sequencer/Private/SSequencerTreeView.cpp
--- a/Engine/Source/Editor/Sequencer/Private/SSequencerTreeView.cpp
+++ b/Engine/Source/Editor/Sequencer/Private/SSequencerTreeView.cpp
@@ -5,6 +5,9 @@
 #include "SequencerPrivatePCH.h"
 #include "Sequencer.h"
 #include "SListPanel.h"
+#include "SSequencerTreeView.h"
+#include "SSequencerTrackLane.h"
+#include "SSequencerTrackArea.h"
 
 static FName TrackAreaName = "TrackArea";
 
diff --git a/Engine/Source/Editor/Sequencer/Private/Sequencer.cpp b/Engine/Source/Editor/Sequencer/Private/Sequencer.cpp
--- a/Engine/Source/Editor/Sequencer/Private/Sequencer.cpp
+++ b/Engine/Source/Editor/Sequencer/Private/Sequencer.cpp
@@ -14,6 +14,7 @@
 #include "Editor/LevelEditor/Public/ILevelViewport.h"
 #include "Editor/WorkspaceMenuStructure/Public/WorkspaceMenuStructureModule.h"
 #include "SSequencer.h"
+#include "SSequencerTreeView.h"
 #include "ScopedTransaction.h"
 #include "Kismet2/BlueprintEditorUtils.h"
 #include "BlueprintEditorModule.h"
