Non-unity compile fixes

diff --git a/Engine/Source/Editor/Sequencer/Private/SAnimationOutlinerTreeNode.cpp b/Engine/Source/Editor/Sequencer/Private/SAnimationOutlinerTreeNode.cpp
--- a/Engine/Source/Editor/Sequencer/Private/SAnimationOutlinerTreeNode.cpp
+++ b/Engine/Source/Editor/Sequencer/Private/SAnimationOutlinerTreeNode.cpp
@@ -15,9 +15,8 @@
 #define LOCTEXT_NAMESPACE "AnimationOutliner"
 
 
-void SAnimationOutlinerTreeNode::Construct( const FArguments& InArgs, TSharedRef<FSequencerDisplayNode> Node, const TSharedRef<SSequencerTreeViewRow>& InTableRow, FSequencer* InSequencer  )
+void SAnimationOutlinerTreeNode::Construct( const FArguments& InArgs, TSharedRef<FSequencerDisplayNode> Node, const TSharedRef<SSequencerTreeViewRow>& InTableRow )
 {
-	Sequencer = InSequencer;
 	DisplayNode = Node;
 
 	SelectedBrush = FEditorStyle::GetBrush( "Sequencer.AnimationOutliner.SelectionBorder" );
@@ -449,13 +448,14 @@ void SAnimationOutlinerTreeNode::OnSelectionChanged( TArray<TSharedPtr<FSequence
 		return;
 	}
 
-	if (!Sequencer->IsLevelEditorSequencer())
+	FSequencer& Sequencer = DisplayNode->GetSequencer();
+	if (!Sequencer.IsLevelEditorSequencer())
 	{
 		return;
 	}
 
 	// Mark that the user is selecting so that the UI doesn't respond to the selection changes in the following block
-	TSharedRef<SSequencer> SequencerWidget = StaticCastSharedRef<SSequencer>(Sequencer->GetSequencerWidget());
+	TSharedRef<SSequencer> SequencerWidget = StaticCastSharedRef<SSequencer>(Sequencer.GetSequencerWidget());
 	SequencerWidget->SetUserIsSelecting(true);
 
 	const FScopedTransaction Transaction(NSLOCTEXT("UnrealEd", "ClickingOnActors", "Clicking on Actors"));
@@ -485,7 +485,7 @@ void SAnimationOutlinerTreeNode::OnSelectionChanged( TArray<TSharedPtr<FSequence
 
 		// Get the bound objects
 		TArray<UObject*> RuntimeObjects;
-		Sequencer->GetRuntimeObjects( Sequencer->GetFocusedMovieSceneInstance(), ObjectNode->GetObjectBinding(), RuntimeObjects );
+		Sequencer.GetRuntimeObjects( Sequencer.GetFocusedMovieSceneInstance(), ObjectNode->GetObjectBinding(), RuntimeObjects );
 		
 		if( RuntimeObjects.Num() > 0 )
 		{
@@ -500,7 +500,7 @@ void SAnimationOutlinerTreeNode::OnSelectionChanged( TArray<TSharedPtr<FSequence
 
 					if (IsControlDown)
 					{
-						bSelectActor = Sequencer->GetSelection().IsSelected(ObjectNodes[ObjectIdx].ToSharedRef());
+						bSelectActor = Sequencer.GetSelection().IsSelected(ObjectNodes[ObjectIdx].ToSharedRef());
 					}
 
 					GEditor->SelectActor(Actor, bSelectActor, bNotifySelectionChanged );
diff --git a/Engine/Source/Editor/Sequencer/Private/SAnimationOutlinerTreeNode.h b/Engine/Source/Editor/Sequencer/Private/SAnimationOutlinerTreeNode.h
--- a/Engine/Source/Editor/Sequencer/Private/SAnimationOutlinerTreeNode.h
+++ b/Engine/Source/Editor/Sequencer/Private/SAnimationOutlinerTreeNode.h
@@ -14,7 +14,7 @@ public:
 	SLATE_BEGIN_ARGS(SAnimationOutlinerTreeNode){}
 	SLATE_END_ARGS()
 
-	void Construct( const FArguments& InArgs, TSharedRef<FSequencerDisplayNode> Node, const TSharedRef<SSequencerTreeViewRow>& InTableRow, FSequencer* InSequencer );
+	void Construct( const FArguments& InArgs, TSharedRef<FSequencerDisplayNode> Node, const TSharedRef<SSequencerTreeViewRow>& InTableRow );
 
 	/**
 	 * @return The display node used by this widget                                                              
@@ -63,8 +63,6 @@ private:
 	void OnSelectionChanged( TArray<TSharedPtr<FSequencerDisplayNode> > AffectedNodes );
 
 private:
-	/** Internal sequencer interface */
-	FSequencer* Sequencer;
 	/** Layout node the widget is visualizing */
 	TSharedPtr<FSequencerDisplayNode> DisplayNode;
 	/** Brush to display a border around the widget when it is selected */
diff --git a/Engine/Source/Editor/Sequencer/Private/SSequencerTrackLane.cpp b/Engine/Source/Editor/Sequencer/Private/SSequencerTrackLane.cpp
--- a/Engine/Source/Editor/Sequencer/Private/SSequencerTrackLane.cpp
+++ b/Engine/Source/Editor/Sequencer/Private/SSequencerTrackLane.cpp
@@ -1,7 +1,5 @@
 // Copyright 1998-2015 Epic Games, Inc. All Rights Reserved.
 
-#pragma once
-
 #include "SequencerPrivatePCH.h"
 #include "SSequencerTrackLane.h"
 #include "SSequencerTreeView.h"
diff --git a/Engine/Source/Editor/Sequencer/Private/SSequencerTreeView.cpp b/Engine/Source/Editor/Sequencer/Private/SSequencerTreeView.cpp
--- a/Engine/Source/Editor/Sequencer/Private/SSequencerTreeView.cpp
+++ b/Engine/Source/Editor/Sequencer/Private/SSequencerTreeView.cpp
@@ -1,7 +1,5 @@
 // Copyright 1998-2015 Epic Games, Inc. All Rights Reserved.
 
-#pragma once
-
 #include "SequencerPrivatePCH.h"
 #include "Sequencer.h"
 #include "SListPanel.h"
diff --git a/Engine/Source/Editor/Sequencer/Private/SequencerDisplayNode.cpp b/Engine/Source/Editor/Sequencer/Private/SequencerDisplayNode.cpp
--- a/Engine/Source/Editor/Sequencer/Private/SequencerDisplayNode.cpp
+++ b/Engine/Source/Editor/Sequencer/Private/SequencerDisplayNode.cpp
@@ -294,7 +294,7 @@ void FSequencerDisplayNode::AddKeyAreaNode( FName KeyAreaName, const FText& Disp
 
 TSharedRef<SWidget> FSequencerDisplayNode::GenerateContainerWidgetForOutliner(const TSharedRef<SSequencerTreeViewRow>& InRow)
 {
-	return SNew( SAnimationOutlinerTreeNode, SharedThis( this ), InRow, &GetSequencer() );
+	return SNew( SAnimationOutlinerTreeNode, SharedThis( this ), InRow );
 }
 
 TSharedRef<SWidget> FSequencerDisplayNode::GenerateEditWidgetForOutliner()
