Speculative fix for compile error relating to destruction of an incomplete type

diff --git a/Engine/Source/Editor/Sequencer/Private/SSequencerTrackArea.cpp b/Engine/Source/Editor/Sequencer/Private/SSequencerTrackArea.cpp
--- a/Engine/Source/Editor/Sequencer/Private/SSequencerTrackArea.cpp
+++ b/Engine/Source/Editor/Sequencer/Private/SSequencerTrackArea.cpp
@@ -86,6 +86,11 @@ private:
 	SSequencerTreeView& TreeView;
 };
 
+SSequencerTrackArea::~SSequencerTrackArea()
+{
+	// User-defined destructor to ensure TUniquePtr destruction
+}
+
 void SSequencerTrackArea::Construct( const FArguments& InArgs, TSharedRef<FSequencerTimeSliderController> InTimeSliderController, TSharedRef<FSequencer> InSequencer )
 {
 	Sequencer = InSequencer;
diff --git a/Engine/Source/Editor/Sequencer/Private/SSequencerTrackArea.h b/Engine/Source/Editor/Sequencer/Private/SSequencerTrackArea.h
--- a/Engine/Source/Editor/Sequencer/Private/SSequencerTrackArea.h
+++ b/Engine/Source/Editor/Sequencer/Private/SSequencerTrackArea.h
@@ -22,6 +22,9 @@ public:
 	SLATE_BEGIN_ARGS( SSequencerTrackArea ){}
 	SLATE_END_ARGS()
 
+	/** User-defined destructor to ensure TUniquePtr destruction */
+	~SSequencerTrackArea();
+
 	/** Construct this widget */
 	void Construct( const FArguments& InArgs, TSharedRef<FSequencerTimeSliderController> InTimeSliderController, TSharedRef<FSequencer> InSequencer );
 
