Do not call deprecated functions with Sequencer

diff --git a/Engine/Source/Editor/Sequencer/Private/SequencerObjectChangeListener.cpp b/Engine/Source/Editor/Sequencer/Private/SequencerObjectChangeListener.cpp
--- a/Engine/Source/Editor/Sequencer/Private/SequencerObjectChangeListener.cpp
+++ b/Engine/Source/Editor/Sequencer/Private/SequencerObjectChangeListener.cpp
@@ -105,7 +105,12 @@ bool FSequencerObjectChangeListener::FindPropertySetter( const UClass& ObjectCla
 
 		FName FunctionName = FName(*FunctionString);
 
-		bFound = ObjectClass.FindFunctionByName(FunctionName) != nullptr;
+		static const FName DeprecatedFunctionName(TEXT("DeprecatedFunction"));
+		UFunction* Function = ObjectClass.FindFunctionByName(FunctionName);
+		if( Function && !Function->HasMetaData(DeprecatedFunctionName) )
+		{
+			bFound = true;
+		}
 	}
 
 	return bFound;
