More non-unity compile error fixes

diff --git a/Engine/Source/Editor/Sequencer/Private/SSequencerTrackArea.cpp b/Engine/Source/Editor/Sequencer/Private/SSequencerTrackArea.cpp
--- a/Engine/Source/Editor/Sequencer/Private/SSequencerTrackArea.cpp
+++ b/Engine/Source/Editor/Sequencer/Private/SSequencerTrackArea.cpp
@@ -86,9 +86,12 @@ private:
 	SSequencerTreeView& TreeView;
 };
 
+// User-defined constructor/destructor to ensure TUniquePtr destruction
+SSequencerTrackArea::SSequencerTrackArea()
+{
+}
 SSequencerTrackArea::~SSequencerTrackArea()
 {
-	// User-defined destructor to ensure TUniquePtr destruction
 }
 
 void SSequencerTrackArea::Construct( const FArguments& InArgs, TSharedRef<FSequencerTimeSliderController> InTimeSliderController, TSharedRef<FSequencer> InSequencer )
diff --git a/Engine/Source/Editor/Sequencer/Private/SSequencerTrackArea.h b/Engine/Source/Editor/Sequencer/Private/SSequencerTrackArea.h
--- a/Engine/Source/Editor/Sequencer/Private/SSequencerTrackArea.h
+++ b/Engine/Source/Editor/Sequencer/Private/SSequencerTrackArea.h
@@ -22,7 +22,8 @@ public:
 	SLATE_BEGIN_ARGS( SSequencerTrackArea ){}
 	SLATE_END_ARGS()
 
-	/** User-defined destructor to ensure TUniquePtr destruction */
+	/** User-defined constructor/destructor to ensure TUniquePtr destruction */
+	SSequencerTrackArea();
 	~SSequencerTrackArea();
 
 	/** Construct this widget */
diff --git a/Engine/Source/Editor/Sequencer/Private/SSequencerTreeView.cpp b/Engine/Source/Editor/Sequencer/Private/SSequencerTreeView.cpp
--- a/Engine/Source/Editor/Sequencer/Private/SSequencerTreeView.cpp
+++ b/Engine/Source/Editor/Sequencer/Private/SSequencerTreeView.cpp
@@ -9,17 +9,6 @@
 
 static FName TrackAreaName = "TrackArea";
 
-/** Structure used to cache physical geometry for a particular node */
-struct FCachedGeometry
-{
-	FCachedGeometry(FDisplayNodeRef InNode, float InPhysicalTop, float InPhysicalHeight)
-		: Node(MoveTemp(InNode)), PhysicalTop(InPhysicalTop), PhysicalHeight(InPhysicalHeight)
-	{}
-
-	FDisplayNodeRef Node;
-	float PhysicalTop, PhysicalHeight;
-};
-
 /** Construct function for this widget */
 void SSequencerTreeViewRow::Construct(const FArguments& InArgs, const TSharedRef<STableViewBase>& OwnerTableView, const FDisplayNodeRef& InNode)
 {
diff --git a/Engine/Source/Editor/Sequencer/Private/SSequencerTreeView.h b/Engine/Source/Editor/Sequencer/Private/SSequencerTreeView.h
--- a/Engine/Source/Editor/Sequencer/Private/SSequencerTreeView.h
+++ b/Engine/Source/Editor/Sequencer/Private/SSequencerTreeView.h
@@ -98,6 +98,17 @@ protected:
 	/** Update the cached visible geometry for each tree item in this view */
 	void UpdateCachedVerticalGeometry();
 
+	/** Structure used to cache physical geometry for a particular node */
+	struct FCachedGeometry
+	{
+		FCachedGeometry(FDisplayNodeRef InNode, float InPhysicalTop, float InPhysicalHeight)
+			: Node(MoveTemp(InNode)), PhysicalTop(InPhysicalTop), PhysicalHeight(InPhysicalHeight)
+		{}
+
+		FDisplayNodeRef Node;
+		float PhysicalTop, PhysicalHeight;
+	};
+	
 	/** Linear array of cached geometry for visible nodes in the tree */
 	TArray<struct FCachedGeometry> CachedGeometry;
 
