Sequencer: Made Close() method virtual

diff --git a/Engine/Source/Editor/Sequencer/Private/Sequencer.cpp b/Engine/Source/Editor/Sequencer/Private/Sequencer.cpp
--- a/Engine/Source/Editor/Sequencer/Private/Sequencer.cpp
+++ b/Engine/Source/Editor/Sequencer/Private/Sequencer.cpp
@@ -68,7 +68,7 @@ void FSequencer::InitSequencer( const FSequencerInitParams& InitParams, const TA
 		{
 			if (ActiveSequencer.IsValid())
 			{
-				ActiveSequencer->OnClose();
+				ActiveSequencer->Close();
 			}
 			ActiveSequencer = SharedThis(this);
 
@@ -187,8 +187,7 @@ FSequencer::~FSequencer()
 	SequencerWidget.Reset();
 }
 
-void
-FSequencer::OnClose()
+void FSequencer::Close()
 {
 	if (ActiveSequencer.IsValid() && ActiveSequencer.Get() == this)
 	{
diff --git a/Engine/Source/Editor/Sequencer/Private/Sequencer.h b/Engine/Source/Editor/Sequencer/Private/Sequencer.h
--- a/Engine/Source/Editor/Sequencer/Private/Sequencer.h
+++ b/Engine/Source/Editor/Sequencer/Private/Sequencer.h
@@ -58,9 +58,6 @@ public:
 	 */
 	void InitSequencer( const FSequencerInitParams& InitParams, const TArray<FOnCreateTrackEditor>& TrackEditorDelegates );
 
-	/** Close the sequencer. */
-	void OnClose();
-
 	/** Constructor */
 	FSequencer();
 
@@ -85,6 +82,7 @@ public:
 
 	// ISequencer interface
 
+	virtual void Close() override;
 	virtual TSharedRef<SWidget> GetSequencerWidget() const override { return SequencerWidget.ToSharedRef(); }
 	virtual UMovieScene* GetRootMovieScene() const override;
 	virtual UMovieScene* GetFocusedMovieScene() const override;
diff --git a/Engine/Source/Editor/Sequencer/Private/SequencerAssetEditor.cpp b/Engine/Source/Editor/Sequencer/Private/SequencerAssetEditor.cpp
--- a/Engine/Source/Editor/Sequencer/Private/SequencerAssetEditor.cpp
+++ b/Engine/Source/Editor/Sequencer/Private/SequencerAssetEditor.cpp
@@ -119,7 +119,7 @@ FSequencerAssetEditor::FSequencerAssetEditor()
 
 FSequencerAssetEditor::~FSequencerAssetEditor()
 {
-	Sequencer->OnClose();
+	Sequencer->Close();
 
 	// Unregister delegates
 	if( FModuleManager::Get().IsModuleLoaded( TEXT( "LevelEditor" ) ) )
diff --git a/Engine/Source/Editor/Sequencer/Public/ISequencer.h b/Engine/Source/Editor/Sequencer/Public/ISequencer.h
--- a/Engine/Source/Editor/Sequencer/Public/ISequencer.h
+++ b/Engine/Source/Editor/Sequencer/Public/ISequencer.h
@@ -24,6 +24,9 @@ class ISequencer
 {
 public:
 
+	/** Close the sequencer. */
+	virtual void Close() = 0;
+
 	/** @return Widget used to display the sequencer */
 	virtual TSharedRef<SWidget> GetSequencerWidget() const = 0;
 
