Sequencer: Don't use foreign functions for event callbacks

diff --git a/Engine/Source/Editor/Sequencer/Private/Sequencer.cpp b/Engine/Source/Editor/Sequencer/Private/Sequencer.cpp
--- a/Engine/Source/Editor/Sequencer/Private/Sequencer.cpp
+++ b/Engine/Source/Editor/Sequencer/Private/Sequencer.cpp
@@ -458,7 +458,7 @@ void FSequencer::SpawnOrDestroyPuppetObjects( TSharedRef<FMovieSceneInstance> Mo
 	}
 }
 
-void FSequencer::OnMapChanged( class UWorld* NewWorld, EMapChangeType::Type MapChangeType )
+void FSequencer::NotifyMapChanged( class UWorld* NewWorld, EMapChangeType::Type MapChangeType )
 {
 	// @todo Sequencer Sub-MovieScenes Needs more investigation of what to spawn
 	// Destroy our puppets because the world is going away.  We probably don't have to do this (the actors will
diff --git a/Engine/Source/Editor/Sequencer/Private/Sequencer.h b/Engine/Source/Editor/Sequencer/Private/Sequencer.h
--- a/Engine/Source/Editor/Sequencer/Private/Sequencer.h
+++ b/Engine/Source/Editor/Sequencer/Private/Sequencer.h
@@ -111,6 +111,7 @@ public:
 	virtual void KeyProperty(FKeyPropertyParams KeyPropertyParams) override;
 	virtual TSharedRef<ISequencerObjectBindingManager> GetObjectBindingManager() const override;
 	virtual FSequencerSelection& GetSelection() override;
+	virtual void NotifyMapChanged(class UWorld* NewWorld, EMapChangeType::Type MapChangeType) override;
 
 	/** @return The current view range */
 	virtual FAnimatedRange GetViewRange() const override;
@@ -252,9 +253,6 @@ public:
 
 	virtual TArray< UMovieScene* > GetMovieScenesBeingEdited();
 
-	/** Called by LevelEditor when the map changes */
-	void OnMapChanged(class UWorld* NewWorld, EMapChangeType::Type MapChangeType);
-
 	/** Called when an actor is dropped into Sequencer */
 	void OnActorsDropped( const TArray<TWeakObjectPtr<AActor> >& Actors );
 
diff --git a/Engine/Source/Editor/Sequencer/Private/SequencerAssetEditor.cpp b/Engine/Source/Editor/Sequencer/Private/SequencerAssetEditor.cpp
--- a/Engine/Source/Editor/Sequencer/Private/SequencerAssetEditor.cpp
+++ b/Engine/Source/Editor/Sequencer/Private/SequencerAssetEditor.cpp
@@ -96,7 +96,7 @@ void FSequencerAssetEditor::InitSequencerAssetEditor( const EToolkitMode::Type M
 			
 		// We need to find out when the user loads a new map, because we might need to re-create puppet actors
 		// when previewing a MovieScene
-		LevelEditorModule.OnMapChanged().AddSP(Sequencer.ToSharedRef(), &FSequencer::OnMapChanged);
+		LevelEditorModule.OnMapChanged().AddSP(Sequencer.ToSharedRef(), &FSequencerAssetEditor::HandleMapChanged);
 	}
 }
 
@@ -214,4 +214,10 @@ void FSequencerAssetEditor::AddComponentTrack( UActorComponent* Component )
 }
 
 
+void FSequencerAssetEditor::HandleMapChanged(class UWorld* NewWorld, EMapChangeType::Type MapChangeType)
+{
+	Sequencer->NotifyMapChanged(NewWorld, MapChangeType);
+}
+
+
 #undef LOCTEXT_NAMESPACE
diff --git a/Engine/Source/Editor/Sequencer/Private/SequencerAssetEditor.h b/Engine/Source/Editor/Sequencer/Private/SequencerAssetEditor.h
--- a/Engine/Source/Editor/Sequencer/Private/SequencerAssetEditor.h
+++ b/Engine/Source/Editor/Sequencer/Private/SequencerAssetEditor.h
@@ -39,14 +39,17 @@ public:
 	void UpdateViewports(AActor* ActorToViewThrough) const;
 
 private:
+
 	EVisibility GetTransportControlVisibility(TWeakPtr<ILevelViewport> LevelViewport) const;
 	TSharedRef<SDockTab> SpawnTab_SequencerMain(const FSpawnTabArgs& Args);	
 
 	TSharedRef<FExtender> GetContextSensitiveSequencerExtender( const TSharedRef<FUICommandList> CommandList, const TArray<UObject*> ContextSensitiveObjects );
 	void ExtendSequencerAddTrackMenu( FMenuBuilder& AddTrackMenuBuilder, TArray<UObject*> ContextObjects );
 	void AddComponentTrack( UActorComponent* Component );
+	void HandleMapChanged(class UWorld* NewWorld, EMapChangeType::Type MapChangeType);
 
 private:
+
 	TSharedPtr<FSequencer> Sequencer;
 
 	/** A map of all the transport controls to viewports that this sequencer has made */
diff --git a/Engine/Source/Editor/Sequencer/Public/ISequencer.h b/Engine/Source/Editor/Sequencer/Public/ISequencer.h
--- a/Engine/Source/Editor/Sequencer/Public/ISequencer.h
+++ b/Engine/Source/Editor/Sequencer/Public/ISequencer.h
@@ -158,4 +158,6 @@ public:
 	virtual TSharedRef<ISequencerObjectBindingManager> GetObjectBindingManager() const = 0;
 
 	virtual FSequencerSelection& GetSelection() = 0;
+
+	virtual void NotifyMapChanged(class UWorld* NewWorld, EMapChangeType::Type MapChangeType) = 0;
 };
