Sequencer: Fix outliner selection when clicking on nothing.

diff --git a/Engine/Source/Editor/Sequencer/Private/SSequencer.cpp b/Engine/Source/Editor/Sequencer/Private/SSequencer.cpp
--- a/Engine/Source/Editor/Sequencer/Private/SSequencer.cpp
+++ b/Engine/Source/Editor/Sequencer/Private/SSequencer.cpp
@@ -145,6 +145,37 @@ private:
 	TArray< TRange<float> > CachedFilteredRanges;
 };
 
+/**
+* Wrapper box to encompass clicks that don't fall on a tree view row.
+*/
+class SSequencerTreeViewBox : public SBox
+{
+public:
+	void Construct( const FArguments& InArgs, TSharedRef<FSequencer> InSequencer ) 
+	{
+		Sequencer = InSequencer;
+		SBox::Construct(InArgs);
+	}
+	
+	virtual FReply OnMouseButtonDown( const FGeometry& MyGeometry, const FPointerEvent& MouseEvent ) override 
+	{
+		if (Sequencer.Pin().IsValid())
+		{
+			FSequencerSelection& Selection = Sequencer.Pin()->GetSelection();
+			if (Selection.GetSelectedOutlinerNodes().Num())
+			{
+				Selection.EmptySelectedOutlinerNodes();
+				return FReply::Handled();
+			}
+		}
+		return FReply::Unhandled();
+	}
+	
+private:
+	/** Weak pointer to the sequencer */
+	TWeakPtr<FSequencer> Sequencer;
+};
+
 void SSequencer::Construct( const FArguments& InArgs, TSharedRef< class FSequencer > InSequencer )
 {
 	Sequencer = InSequencer;
@@ -274,7 +305,7 @@ void SSequencer::Construct( const FArguments& InArgs, TSharedRef< class FSequenc
 								+ SHorizontalBox::Slot()
 								.FillWidth( FillCoefficient_0 )
 								[
-									SNew(SBox)
+									SNew(SSequencerTreeViewBox, PinnedSequencer)
 									// Padding to allow space for the scroll bar
 									.Padding(FMargin(0,0,10.f,0))
 									[
diff --git a/Engine/Source/Editor/Sequencer/Private/Sequencer.cpp b/Engine/Source/Editor/Sequencer/Private/Sequencer.cpp
--- a/Engine/Source/Editor/Sequencer/Private/Sequencer.cpp
+++ b/Engine/Source/Editor/Sequencer/Private/Sequencer.cpp
@@ -1383,12 +1383,16 @@ void FSequencer::OnSelectedOutlinerNodesChanged()
 {
 	if (IsLevelEditorSequencer())
 	{
-		const bool bNotifySelectionChanged = false;
-		const bool bDeselectBSP = true;
-		const bool bWarnAboutTooManyActors = false;
+		// Clear the selection if no sequencer display nodes were clicked on and there are selected actors
+		if (Selection.GetSelectedOutlinerNodes().Num() == 0 && GEditor->GetSelectedActors()->Num() != 0)
+		{
+			const bool bNotifySelectionChanged = true;
+			const bool bDeselectBSP = true;
+			const bool bWarnAboutTooManyActors = false;
 
-		const FScopedTransaction Transaction(NSLOCTEXT("UnrealEd", "ClickingOnActors", "Clicking on Actors"));
-		GEditor->SelectNone(bNotifySelectionChanged, bDeselectBSP, bWarnAboutTooManyActors);
+			const FScopedTransaction Transaction(NSLOCTEXT("UnrealEd", "ClickingOnActors", "Clicking on Actors"));
+			GEditor->SelectNone(bNotifySelectionChanged, bDeselectBSP, bWarnAboutTooManyActors);
+		}
 	}
 }
 
