Sequencer: Fix some key selection stuff:

- Fix clicking on keys. Shift should add to selection, control should toggle.
- Fix section right click menu so that the menu for keys is shown if the mouse is over a key or keys are selected.
- Fix left clicking on empty areas should deselect keys just as it deselects sections.

diff --git a/Engine/Source/Editor/Sequencer/Private/SSection.cpp b/Engine/Source/Editor/Sequencer/Private/SSection.cpp
--- a/Engine/Source/Editor/Sequencer/Private/SSection.cpp
+++ b/Engine/Source/Editor/Sequencer/Private/SSection.cpp
@@ -450,7 +450,7 @@ TSharedPtr<SWidget> SSection::OnSummonContextMenu( const FGeometry& MyGeometry,
 	const bool bShouldCloseWindowAfterMenuSelection = true;
 	FMenuBuilder MenuBuilder(bShouldCloseWindowAfterMenuSelection, NULL);
 
-	if (Key.IsValid())
+	if (Key.IsValid() || Sequencer.GetSelection().GetSelectedKeys().Num())
 	{
 		MenuBuilder.AddMenuEntry(
 			NSLOCTEXT("Sequencer", "DeleteKey", "Delete"),
@@ -499,6 +499,8 @@ FReply SSection::OnMouseButtonDown( const FGeometry& MyGeometry, const FPointerE
 
 	bDragging = false;
 
+	FSequencer& Sequencer = GetSequencer();
+
 	if( MouseEvent.GetEffectingButton() == EKeys::LeftMouseButton || MouseEvent.GetEffectingButton() == EKeys::RightMouseButton )
 	{
 		// Check for clicking on a key and mark it as the pressed key for drag detection (if necessary) later
@@ -514,7 +516,7 @@ FReply SSection::OnMouseButtonDown( const FGeometry& MyGeometry, const FPointerE
 
 		return FReply::Handled().CaptureMouse( AsShared() );
 	}
-	else if (MouseEvent.GetEffectingButton() == EKeys::RightMouseButton && PressedKey.IsValid())
+	else if (MouseEvent.GetEffectingButton() == EKeys::RightMouseButton && (PressedKey.IsValid() || Sequencer.GetSelection().GetSelectedKeys().Num()))
 	{
 		// Don't return handled if we didn't click on a key so that right-click panning gets a look in
 		return FReply::Handled();
@@ -637,7 +639,8 @@ FReply SSection::OnMouseMove( const FGeometry& MyGeometry, const FPointerEvent&
 						// Clear selected keys when beginning to drag a section
 						GetSequencer().GetSelection().EmptySelectedKeys();
 
-						HandleSectionSelection( MouseEvent );
+						bool bSelectDueToDrag = true;
+						HandleSectionSelection( MouseEvent, bSelectDueToDrag );
 
 						bool bKeysUnderMouse = false;
 						CreateDragOperation( MyGeometry, MouseEvent, bKeysUnderMouse );
@@ -722,7 +725,8 @@ void SSection::HandleSelection( const FGeometry& MyGeometry, const FPointerEvent
 	}
 	else if( MouseEvent.GetEffectingButton() == EKeys::LeftMouseButton ) // Only select anything but keys on left mouse button
 	{
-		HandleSectionSelection( MouseEvent );
+		bool bSelectDueToDrag = false;
+		HandleSectionSelection( MouseEvent, bSelectDueToDrag );
 	}
 }
 
@@ -731,30 +735,55 @@ void SSection::HandleKeySelection( const FSelectedKey& Key, const FPointerEvent&
 {
 	if( Key.IsValid() )
 	{
+		bool bKeyIsSelected = GetSequencer().GetSelection().IsSelected( Key );
+
 		// Clear previous key selection if:
-		// we are selecting due to drag and the key being dragged is not selected or control is not down
-		bool bShouldClearSelectionDueToDrag =  bSelectDueToDrag ? !GetSequencer().GetSelection().IsSelected( Key ) : true;
+		// we are selecting due to drag and the key being dragged is not selected
+		bool bShouldClearSelectionDueToDrag =  bSelectDueToDrag ? !bKeyIsSelected : true;
 		
 		// Keep key selection if right clicking to bring up a menu and the current key is selected
-		bool bKeepKeySelection = MouseEvent.GetEffectingButton() == EKeys::RightMouseButton && GetSequencer().GetSelection().IsSelected( Key );
+		bool bKeepKeySelection = MouseEvent.GetEffectingButton() == EKeys::RightMouseButton && bKeyIsSelected;
 
-		if( (!MouseEvent.IsControlDown() && bShouldClearSelectionDueToDrag) && !bKeepKeySelection )
+		// Toggle selection if control is down and key is selected
+		if (MouseEvent.IsControlDown() && bKeyIsSelected)
+		{
+			GetSequencer().GetSelection().RemoveFromSelection( Key );
+		}
+		else
 		{
-			GetSequencer().GetSelection().EmptySelectedKeys();
+			if( (!MouseEvent.IsShiftDown() && !MouseEvent.IsControlDown() && bShouldClearSelectionDueToDrag) && !bKeepKeySelection )
+			{
+				GetSequencer().GetSelection().EmptySelectedKeys();
+			}
+			GetSequencer().GetSelection().AddToSelection( Key );
 		}
-		GetSequencer().GetSelection().AddToSelection( Key );
 	}
 }
 
-void SSection::HandleSectionSelection( const FPointerEvent& MouseEvent )
+void SSection::HandleSectionSelection( const FPointerEvent& MouseEvent, bool bSelectDueToDrag )
 {
-	if( !MouseEvent.IsControlDown() )
-	{
-		GetSequencer().GetSelection().EmptySelectedSections();
-	}
-
 	// handle selecting sections 
 	UMovieSceneSection* Section = SectionInterface->GetSectionObject();
-	GetSequencer().GetSelection().AddToSelection(Section);
+	bool bSectionIsSelected = GetSequencer().GetSelection().IsSelected(Section);
+
+	// Clear previous section selection if:
+	// we are selecting due to drag and the section being dragged is not selected
+	bool bShouldClearSelectionDueToDrag =  bSelectDueToDrag ? !bSectionIsSelected : true;
+
+	// Keep key selection if right clicking to bring up a menu and the current key is selected
+	bool bKeepSectionSelection = MouseEvent.GetEffectingButton() == EKeys::RightMouseButton && bSectionIsSelected;
+	
+	if (MouseEvent.IsControlDown() && bSectionIsSelected)
+	{
+		GetSequencer().GetSelection().RemoveFromSelection(Section);
+	}
+	else
+	{
+		if ( (!MouseEvent.IsShiftDown() && !MouseEvent.IsControlDown() && bShouldClearSelectionDueToDrag) && !bKeepSectionSelection)
+		{
+			GetSequencer().GetSelection().EmptySelectedSections();
+		}
+		GetSequencer().GetSelection().AddToSelection(Section);
+	}
 }
 
diff --git a/Engine/Source/Editor/Sequencer/Private/SSection.h b/Engine/Source/Editor/Sequencer/Private/SSection.h
--- a/Engine/Source/Editor/Sequencer/Private/SSection.h
+++ b/Engine/Source/Editor/Sequencer/Private/SSection.h
@@ -125,8 +125,9 @@ private:
 	 * Selects the section if needed
 	 *
 	 * @param MouseEvent		The mouse event that may have caused selection
+	 * @param bSelectDueToDrag	If we are selecting the key due to a drag operation rather than it just being clicked
 	 */
-	void HandleSectionSelection( const FPointerEvent& MouseEvent );
+	void HandleSectionSelection( const FPointerEvent& MouseEvent, bool bSelectDueToDrag );
 
 	/** @return the sequencer interface */
 	FSequencer& GetSequencer() const;
diff --git a/Engine/Source/Editor/Sequencer/Private/SSequencerTrackArea.cpp b/Engine/Source/Editor/Sequencer/Private/SSequencerTrackArea.cpp
--- a/Engine/Source/Editor/Sequencer/Private/SSequencerTrackArea.cpp
+++ b/Engine/Source/Editor/Sequencer/Private/SSequencerTrackArea.cpp
@@ -176,6 +176,14 @@ FReply SSequencerTrackArea::OnMouseButtonDown( const FGeometry& MyGeometry, cons
 		FVector2D VirtualClickPosition = Converter.PhysicalToVirtual(MyGeometry.AbsoluteToLocal(MouseEvent.GetScreenSpacePosition()));
 		MarqueeSelectData.Reset(new FMarqueeSelectData(VirtualClickPosition));
 
+		// Clear selected sections and keys
+		auto SequencerPin = Sequencer.Pin();
+		if( !MouseEvent.IsControlDown() && !MouseEvent.IsShiftDown() && SequencerPin.IsValid() )
+		{
+			SequencerPin->GetSelection().EmptySelectedSections();
+			SequencerPin->GetSelection().EmptySelectedKeys();
+		}
+
 		Reply = FReply::Handled().CaptureMouse(AsShared());
 	}
 	else
@@ -183,13 +191,6 @@ FReply SSequencerTrackArea::OnMouseButtonDown( const FGeometry& MyGeometry, cons
 		Reply = TimeSliderController->OnMouseButtonDown( SharedThis(this), MyGeometry, MouseEvent );
 	}
 
-	// Clear selected sections
-	auto SequencerPin = Sequencer.Pin();
-	if( !MouseEvent.IsControlDown() && SequencerPin.IsValid() )
-	{
-		SequencerPin->GetSelection().EmptySelectedSections();
-	}
-
 	return Reply;
 }
 
