EMapChangeType is strongly types, so it can be forward-declared.

NonUnityTestCompile_OnMac Cis fix: "ISequencer.h:162:56: error: unknown type name 'EMapChangeType'"

diff --git a/Engine/Source/Editor/Sequencer/Public/ISequencer.h b/Engine/Source/Editor/Sequencer/Public/ISequencer.h
--- a/Engine/Source/Editor/Sequencer/Public/ISequencer.h
+++ b/Engine/Source/Editor/Sequencer/Public/ISequencer.h
@@ -13,7 +13,7 @@ class ISequencerObjectBindingManager;
 class UAnimSequence;
 class UMovieScene;
 class UMovieSceneSection;
-
+enum class EMapChangeType : uint8;
 
 /**
  * Interface for sequencers.
diff --git a/Engine/Source/Editor/UnrealEd/Public/UnrealEdMisc.h b/Engine/Source/Editor/UnrealEd/Public/UnrealEdMisc.h
--- a/Engine/Source/Editor/UnrealEd/Public/UnrealEdMisc.h
+++ b/Engine/Source/Editor/UnrealEd/Public/UnrealEdMisc.h
@@ -8,7 +8,7 @@ class FPerformanceAnalyticsStats;
 class FTickableEditorObject;
 
 
-enum class EMapChangeType
+enum class EMapChangeType : uint8
 {
 	/** Map has just been loaded*/
 	LoadMap,
