Sequencer tree view no longer relies on access to its underlying list panel

diff --git a/Engine/Source/Editor/Sequencer/Private/SSequencerTrackLane.cpp b/Engine/Source/Editor/Sequencer/Private/SSequencerTrackLane.cpp
--- a/Engine/Source/Editor/Sequencer/Private/SSequencerTrackLane.cpp
+++ b/Engine/Source/Editor/Sequencer/Private/SSequencerTrackLane.cpp
@@ -54,10 +54,10 @@ float SSequencerTrackLane::GetPhysicalPosition() const
 	// Iterate parent first until we find a tree view row we can use for the offset height
 	auto Iter = [&](FSequencerDisplayNode& InNode){
 		
-		auto ChildRow = StaticCastSharedPtr<SSequencerTreeViewRow>(TreeViewPinned->WidgetFromItem(InNode.AsShared()));
-		if (ChildRow.IsValid())
+		auto ChildRowGeometry = TreeViewPinned->GetPhysicalGeometryForNode(InNode.AsShared());
+		if (ChildRowGeometry.IsSet())
 		{
-			Top = ChildRow->GetCachedGeometry().Position.Y;
+			Top = ChildRowGeometry->PhysicalTop;
 			// Stop iterating
 			return false;
 		}
diff --git a/Engine/Source/Editor/Sequencer/Private/SSequencerTreeView.cpp b/Engine/Source/Editor/Sequencer/Private/SSequencerTreeView.cpp
--- a/Engine/Source/Editor/Sequencer/Private/SSequencerTreeView.cpp
+++ b/Engine/Source/Editor/Sequencer/Private/SSequencerTreeView.cpp
@@ -2,13 +2,52 @@
 
 #include "SequencerPrivatePCH.h"
 #include "Sequencer.h"
-#include "SListPanel.h"
 #include "SSequencerTreeView.h"
 #include "SSequencerTrackLane.h"
 #include "SSequencerTrackArea.h"
 
 static FName TrackAreaName = "TrackArea";
 
+namespace Utils
+{
+	enum class ESearchState { Before, After, Found };
+
+	template<typename T, typename F>
+	const T* BinarySearch(const TArray<T>& InContainer, const F& InPredicate)
+	{
+		int32 Min = 0;
+		int32 Max = InContainer.Num();
+
+		ESearchState State = ESearchState::Before;
+
+		for ( ; Min != Max ; )
+		{
+			int32 SearchIndex = Min + (Max - Min) / 2;
+
+			auto& Item = InContainer[SearchIndex];
+			State = InPredicate(Item);
+			switch (State)
+			{
+				case ESearchState::Before:	Max = SearchIndex; break;
+				case ESearchState::After:		Min = SearchIndex + 1; break;
+				case ESearchState::Found: 	return &Item;
+			}
+		}
+
+		return nullptr;
+	}
+}
+
+SSequencerTreeViewRow::~SSequencerTreeViewRow()
+{
+	auto TreeView = StaticCastSharedPtr<SSequencerTreeView>(OwnerTablePtr.Pin());
+	auto PinnedNode = Node.Pin();
+	if (TreeView.IsValid() && PinnedNode.IsValid())
+	{
+		TreeView->OnChildRowRemoved(PinnedNode.ToSharedRef());
+	}
+}
+
 /** Construct function for this widget */
 void SSequencerTreeViewRow::Construct(const FArguments& InArgs, const TSharedRef<STableViewBase>& OwnerTableView, const FDisplayNodeRef& InNode)
 {
@@ -34,11 +73,6 @@ TSharedPtr<FSequencerDisplayNode> SSequencerTreeViewRow::GetDisplayNode() const
 	return Node.Pin();
 }
 
-const FGeometry& SSequencerTreeViewRow::GetCachedGeometry() const
-{
-	return RowGeometry;
-}
-
 void SSequencerTreeViewRow::AddTrackAreaReference(const TSharedPtr<SSequencerTrackLane>& Lane)
 {
 	TrackLaneReference = Lane;
@@ -46,7 +80,7 @@ void SSequencerTreeViewRow::AddTrackAreaReference(const TSharedPtr<SSequencerTra
 
 void SSequencerTreeViewRow::Tick(const FGeometry& AllottedGeometry, const double InCurrentTime, const float InDeltaTime)
 {
-	RowGeometry = AllottedGeometry;
+	StaticCastSharedPtr<SSequencerTreeView>(OwnerTablePtr.Pin())->ReportChildRowGeometry(Node.Pin().ToSharedRef(), AllottedGeometry);
 }
 
 void SSequencerTreeView::Construct(const FArguments& InArgs, const TSharedRef<FSequencerNodeTree>& InNodeTree, const TSharedRef<SSequencerTrackArea>& InTrackArea)
@@ -81,40 +115,72 @@ void SSequencerTreeView::Construct(const FArguments& InArgs, const TSharedRef<FS
 void SSequencerTreeView::Tick(const FGeometry& AllottedGeometry, const double InCurrentTime, const float InDeltaTime)
 {
 	STreeView::Tick(AllottedGeometry, InCurrentTime, InDeltaTime);
-	CachedTreeGeometry = AllottedGeometry;
-	UpdateCachedVerticalGeometry();
+	
+	PhysicalNodes.Reset();
+	CachedRowGeometry.GenerateValueArray(PhysicalNodes);
+
+	PhysicalNodes.Sort([](const FCachedGeometry& A, const FCachedGeometry& B){
+		return A.PhysicalTop < B.PhysicalTop;
+	});
 }
 
-void SSequencerTreeView::UpdateCachedVerticalGeometry()
+TOptional<SSequencerTreeView::FCachedGeometry> SSequencerTreeView::GetPhysicalGeometryForNode(const FDisplayNodeRef& InNode) const
 {
-	FChildren* Children = ItemsPanel->GetChildren();
-
-	CachedGeometry.Empty(Children->Num());
-	for (int32 Index = 0; Index < Children->Num(); ++Index)
+	if (const FCachedGeometry* FoundGeometry = CachedRowGeometry.Find(InNode))
 	{
-		TSharedRef<SSequencerTreeViewRow> TreeViewRow = StaticCastSharedRef<SSequencerTreeViewRow>(Children->GetChildAt(Index));
-		TSharedPtr<FSequencerDisplayNode> Node = TreeViewRow->GetDisplayNode();
-		if (Node.IsValid())
-		{
-			CachedGeometry.Emplace(Node.ToSharedRef(), TreeViewRow->GetCachedGeometry().Position.Y, TreeViewRow->GetCachedGeometry().Size.Y);
-		}
+		return *FoundGeometry;
 	}
+
+	return TOptional<FCachedGeometry>();
+}
+
+void SSequencerTreeView::ReportChildRowGeometry(const FDisplayNodeRef& InNode, const FGeometry& InGeometry)
+{
+	float ChildOffset = TransformPoint(
+		Concatenate(
+			InGeometry.GetAccumulatedLayoutTransform(),
+			CachedGeometry.GetAccumulatedLayoutTransform().Inverse()
+		),
+		FVector2D(0,0)
+	).Y;
+
+	CachedRowGeometry.Add(InNode, FCachedGeometry(InNode, ChildOffset, InGeometry.Size.Y));
+}
+
+void SSequencerTreeView::OnChildRowRemoved(const FDisplayNodeRef& InNode)
+{
+	CachedRowGeometry.Remove(InNode);
 }
 
 float SSequencerTreeView::PhysicalToVirtual(float InPhysical) const
 {
-	for (auto& Geometry : CachedGeometry)
-	{
-		if (InPhysical >= Geometry.PhysicalTop && InPhysical <= Geometry.PhysicalTop + Geometry.PhysicalHeight)
+	int32 SearchIndex = PhysicalNodes.Num() / 2;
+
+	auto* Found = Utils::BinarySearch<FCachedGeometry>(PhysicalNodes, [&](const FCachedGeometry& In){
+
+		if (InPhysical < In.PhysicalTop)
 		{
-			const float FractionalHeight = (InPhysical - Geometry.PhysicalTop) / Geometry.PhysicalHeight;
-			return Geometry.Node->VirtualTop + (Geometry.Node->VirtualBottom - Geometry.Node->VirtualTop) * FractionalHeight;
+			return Utils::ESearchState::Before;
 		}
+		else if (InPhysical > In.PhysicalTop + In.PhysicalHeight)
+		{
+			return Utils::ESearchState::After;
+		}
+
+		return Utils::ESearchState::Found;
+
+	});
+
+
+	if (Found)
+	{
+		const float FractionalHeight = (InPhysical - Found->PhysicalTop) / Found->PhysicalHeight;
+		return Found->Node->VirtualTop + (Found->Node->VirtualBottom - Found->Node->VirtualTop) * FractionalHeight;
 	}
 
-	if (CachedGeometry.Num())
+	if (PhysicalNodes.Num())
 	{
-		auto Last = CachedGeometry.Last();
+		auto& Last = PhysicalNodes.Last();
 		return Last.Node->VirtualTop + (InPhysical - Last.PhysicalTop);
 	}
 
@@ -123,18 +189,30 @@ float SSequencerTreeView::PhysicalToVirtual(float InPhysical) const
 
 float SSequencerTreeView::VirtualToPhysical(float InVirtual) const
 {
-	for (auto& Geometry : CachedGeometry)
-	{
-		if (InVirtual >= Geometry.Node->VirtualTop && InVirtual <= Geometry.Node->VirtualBottom)
+	auto* Found = Utils::BinarySearch(PhysicalNodes, [&](const FCachedGeometry& In){
+
+		if (InVirtual < In.Node->VirtualTop)
 		{
-			const float FractionalHeight = (InVirtual - Geometry.Node->VirtualTop) / (Geometry.Node->VirtualBottom - Geometry.Node->VirtualTop);
-			return Geometry.PhysicalTop + Geometry.PhysicalHeight * FractionalHeight;
+			return Utils::ESearchState::Before;
 		}
+		else if (InVirtual > In.Node->VirtualBottom)
+		{
+			return Utils::ESearchState::After;
+		}
+
+		return Utils::ESearchState::Found;
+
+	});
+
+	if (Found)
+	{
+		const float FractionalHeight = (InVirtual - Found->Node->VirtualTop) / (Found->Node->VirtualBottom - Found->Node->VirtualTop);
+		return Found->PhysicalTop + Found->PhysicalHeight * FractionalHeight;
 	}
 	
-	if (CachedGeometry.Num())
+	if (PhysicalNodes.Num())
 	{
-		auto Last = CachedGeometry.Last();
+		auto Last = PhysicalNodes.Last();
 		return Last.PhysicalTop + (InVirtual - Last.Node->VirtualTop);
 	}
 
@@ -235,7 +313,7 @@ void SSequencerTreeView::Refresh()
 
 void SSequencerTreeView::ScrollByDelta(float DeltaInSlateUnits)
 {
-	ScrollBy( CachedTreeGeometry, DeltaInSlateUnits, EAllowOverscroll::No );
+	ScrollBy( CachedGeometry, DeltaInSlateUnits, EAllowOverscroll::No );
 }
 
 template<typename T>
diff --git a/Engine/Source/Editor/Sequencer/Private/SSequencerTreeView.h b/Engine/Source/Editor/Sequencer/Private/SSequencerTreeView.h
--- a/Engine/Source/Editor/Sequencer/Private/SSequencerTreeView.h
+++ b/Engine/Source/Editor/Sequencer/Private/SSequencerTreeView.h
@@ -87,17 +87,8 @@ protected:
 	/** Overridden to keep external selection states up to date */
 	virtual void Private_SignalSelectionChanged(ESelectInfo::Type SelectInfo) override;
 
-protected:
-
-	/** Populate the map of column definitions, and add relevant columns to the header row */
-	void SetupColumns(const FArguments& InArgs);
-
-	/** Ensure that the track area column is either show or hidden, depending on the visibility of the curve editor */
-	void UpdateTrackArea();
-
-	/** Update the cached visible geometry for each tree item in this view */
-	void UpdateCachedVerticalGeometry();
-
+public:
+	
 	/** Structure used to cache physical geometry for a particular node */
 	struct FCachedGeometry
 	{
@@ -108,14 +99,33 @@ protected:
 		FDisplayNodeRef Node;
 		float PhysicalTop, PhysicalHeight;
 	};
+
+	/** Retrieve the last reported physical geometry for the specified node, if available */
+	TOptional<FCachedGeometry> GetPhysicalGeometryForNode(const FDisplayNodeRef& InNode) const;
+
+	/** Report geometry for a child row */
+	void ReportChildRowGeometry(const FDisplayNodeRef& InNode, const FGeometry& InGeometry);
+
+	/** Called when a child row widget has been added/removed */
+	void OnChildRowRemoved(const FDisplayNodeRef& InNode);
+
+protected:
 	
-	/** Linear array of cached geometry for visible nodes in the tree */
-	TArray<struct FCachedGeometry> CachedGeometry;
+	/** Linear, sorted array of nodes that we currently have generated widgets for */
+	TArray<FCachedGeometry> PhysicalNodes;
 
-private:
+	/** Map of cached geometries for visible nodes */
+	TMap<FDisplayNodeRef, FCachedGeometry> CachedRowGeometry;
+
+protected:
 
-	/** The cached geometry for this tree view*/
-	FGeometry CachedTreeGeometry;
+	/** Populate the map of column definitions, and add relevant columns to the header row */
+	void SetupColumns(const FArguments& InArgs);
+
+	/** Ensure that the track area column is either show or hidden, depending on the visibility of the curve editor */
+	void UpdateTrackArea();
+
+private:
 
 	/** The tree view's header row (hidden) */
 	TSharedPtr<SHeaderRow> HeaderRow;
@@ -146,31 +156,24 @@ public:
 
 	SLATE_END_ARGS()
 
+	~SSequencerTreeViewRow();
+
 	/** Construct function for this widget */
 	void Construct(const FArguments& InArgs, const TSharedRef<STableViewBase>& OwnerTableView, const FDisplayNodeRef& InNode);
 
 	/** Get the dispolay node to which this row relates */
 	TSharedPtr<FSequencerDisplayNode> GetDisplayNode() const;
 
-	/** Access this row's cached geometry */
-	const FGeometry& GetCachedGeometry() const;
-
 	/** Add a reference to the specified track lane, keeping it alive until this row is destroyed */
 	void AddTrackAreaReference(const TSharedPtr<SSequencerTrackLane>& Lane);
 
 	virtual void Tick(const FGeometry& AllottedGeometry, const double InCurrentTime, const float InDeltaTime) override;
 
-	/** Access this row's owner tree view. */
-	TSharedPtr<SSequencerTreeView> GetOwnerTree() const;
-
 	/** Overridden from SMultiColumnTableRow.  Generates a widget for this column of the tree row. */
 	virtual TSharedRef<SWidget> GenerateWidgetForColumn(const FName& ColumnId) override;
 
 private:
 
-	/** Cached row geometry, used for virtual <-> physical space conversions */
-	FGeometry RowGeometry;
-
 	/** Cached reference to a track lane that we relate to. This keeps the track lane alive (it's a weak widget) as long as we are in view. */
 	TSharedPtr<SSequencerTrackLane> TrackLaneReference;
 
