Fix sequencer track dropdowns showing every property ever.

diff --git a/Engine/Source/Editor/Sequencer/Private/SequencerObjectChangeListener.cpp b/Engine/Source/Editor/Sequencer/Private/SequencerObjectChangeListener.cpp
--- a/Engine/Source/Editor/Sequencer/Private/SequencerObjectChangeListener.cpp
+++ b/Engine/Source/Editor/Sequencer/Private/SequencerObjectChangeListener.cpp
@@ -111,6 +111,10 @@ bool FSequencerObjectChangeListener::FindPropertySetter( const UClass& ObjectCla
 		{
 			bFound = true;
 		}
+		else
+		{
+			bFound = false;
+		}
 	}
 
 	return bFound;
