Fixed sequencer marquee selection sometimes popping when negative mouse positions were converted to virtual space

diff --git a/Engine/Source/Editor/Sequencer/Private/SSequencerTrackArea.cpp b/Engine/Source/Editor/Sequencer/Private/SSequencerTrackArea.cpp
--- a/Engine/Source/Editor/Sequencer/Private/SSequencerTrackArea.cpp
+++ b/Engine/Source/Editor/Sequencer/Private/SSequencerTrackArea.cpp
@@ -221,7 +221,7 @@ FReply SSequencerTrackArea::OnMouseMove( const FGeometry& MyGeometry, const FPoi
 {
 	if (MarqueeSelectData.IsValid())
 	{
-		SetNewMarqueeBounds(MyGeometry.AbsoluteToLocal(MouseEvent.GetScreenSpacePosition()), MyGeometry);
+		SetNewMarqueeBounds(MouseEvent, MyGeometry);
 		return FReply::Handled();
 	}
 	else if (MouseEvent.IsMouseButtonDown(EKeys::RightMouseButton) && HasMouseCapture())
@@ -277,50 +277,58 @@ void SSequencerTrackArea::Tick( const FGeometry& AllottedGeometry, const double
 	}
 }
 
-void SSequencerTrackArea::SetNewMarqueeBounds(FVector2D NewVirtualPosition, const FGeometry& MyGeometry)
+void SSequencerTrackArea::SetNewMarqueeBounds(const FPointerEvent& MouseEvent, const FGeometry& MyGeometry)
 {
-	const FVirtualConverter Converter(TreeView.Pin(), Sequencer.Pin(), MyGeometry);
-	TRange<float> ViewRange = Sequencer.Pin()->GetViewRange();
-
-	FVector2D Change = MarqueeSelectData->CurrentPosition;
+	const FVector2D MouseDelta = MouseEvent.GetCursorDelta();
 
-	MarqueeSelectData->CurrentPosition = Converter.PhysicalToVirtual(NewVirtualPosition);
+	FVector2D LocalPosition = MyGeometry.AbsoluteToLocal(MouseEvent.GetScreenSpacePosition());
 
-	// Handle virtual scrolling when at the extremes of the widget
-	Change = MarqueeSelectData->CurrentPosition - Change;
+	// Handle virtual scrolling when at the vertical extremes of the widget (performed before we clamp the mouse pos)
+	{
+		const float ScrollThresholdV = MyGeometry.Size.Y * 0.025f;
 
-	const float ScrollThresholdH = ViewRange.Size<float>() * 0.025f;
-	const float ScrollThresholdV = MyGeometry.Size.Y * 0.025f;
+		float Difference = LocalPosition.Y - ScrollThresholdV;
+		if (Difference < 0 && MouseDelta.Y < 0)
+		{
+			TreeView.Pin()->ScrollByDelta( Difference * 0.1f );
+		}
 
-	float Difference = MarqueeSelectData->CurrentPosition.X - (ViewRange.GetLowerBoundValue() + ScrollThresholdH);
-	if (Difference < 0 && Change.X < 0)
-	{
-		Sequencer.Pin()->StartAutoscroll(Difference);
+		Difference = LocalPosition.Y - (MyGeometry.Size.Y - ScrollThresholdV);
+		if (Difference > 0 && MouseDelta.Y > 0)
+		{
+			TreeView.Pin()->ScrollByDelta( Difference * 0.1f );
+		}
 	}
-	else
+
+	const FVirtualConverter Converter(TreeView.Pin(), Sequencer.Pin(), MyGeometry);
+
+	// Clamp the vertical position to the actual bounds of the track area
+	LocalPosition.Y = FMath::Clamp(LocalPosition.Y, 0.f, MyGeometry.GetLocalSize().Y);
+	MarqueeSelectData->CurrentPosition = Converter.PhysicalToVirtual(LocalPosition);
+
+	// Handle virtual scrolling when at the horizontal extremes of the widget
 	{
-		Difference = MarqueeSelectData->CurrentPosition.X - (ViewRange.GetUpperBoundValue() - ScrollThresholdH);
-		if (Difference > 0 && Change.X > 0)
+		TRange<float> ViewRange = Sequencer.Pin()->GetViewRange();
+		const float ScrollThresholdH = ViewRange.Size<float>() * 0.025f;
+
+		float Difference = MarqueeSelectData->CurrentPosition.X - (ViewRange.GetLowerBoundValue() + ScrollThresholdH);
+		if (Difference < 0 && MouseDelta.X < 0)
 		{
 			Sequencer.Pin()->StartAutoscroll(Difference);
 		}
 		else
 		{
-			Sequencer.Pin()->StopAutoscroll();
+			Difference = MarqueeSelectData->CurrentPosition.X - (ViewRange.GetUpperBoundValue() - ScrollThresholdH);
+			if (Difference > 0 && MouseDelta.X > 0)
+			{
+				Sequencer.Pin()->StartAutoscroll(Difference);
+			}
+			else
+			{
+				Sequencer.Pin()->StopAutoscroll();
+			}
 		}
 	}
-
-	Difference = Converter.VerticalOffsetToPixel(MarqueeSelectData->CurrentPosition.Y) - ScrollThresholdV;
-	if (Difference < 0 && Change.Y < 0)
-	{
-		TreeView.Pin()->ScrollByDelta( Difference * 0.1f );
-	}
-
-	Difference = Converter.VerticalOffsetToPixel(MarqueeSelectData->CurrentPosition.Y) - (MyGeometry.Size.Y - ScrollThresholdV);
-	if (Difference > 0 && Change.Y > 0)
-	{
-		TreeView.Pin()->ScrollByDelta( Difference * 0.1f );
-	}
 }
 
 void SSequencerTrackArea::HandleMarqueeSelection(const FGeometry& Geometry, const FPointerEvent& MouseEvent)
diff --git a/Engine/Source/Editor/Sequencer/Private/SSequencerTrackArea.h b/Engine/Source/Editor/Sequencer/Private/SSequencerTrackArea.h
--- a/Engine/Source/Editor/Sequencer/Private/SSequencerTrackArea.h
+++ b/Engine/Source/Editor/Sequencer/Private/SSequencerTrackArea.h
@@ -52,7 +52,7 @@ private:
 	int32 PaintMarquee(const FGeometry& AllottedGeometry, const FSlateRect& MyClippingRect, FSlateWindowElementList& OutDrawElements, int32 LayerId) const;
 
 	/** Set new bounds for the marquee selection */
-	void SetNewMarqueeBounds(FVector2D NewVirtualPosition, const FGeometry& Geometry);
+	void SetNewMarqueeBounds(const FPointerEvent& MouseEvent, const FGeometry& Geometry);
 
 	/** Select all the keys contained in the current marquee */
 	void HandleMarqueeSelection(const FGeometry& Geometry, const FPointerEvent& MouseEvent);
