Sequencer: Fix non-unity build error.

diff --git a/Engine/Source/Editor/Sequencer/Private/SequencerDisplayNode.cpp b/Engine/Source/Editor/Sequencer/Private/SequencerDisplayNode.cpp
--- a/Engine/Source/Editor/Sequencer/Private/SequencerDisplayNode.cpp
+++ b/Engine/Source/Editor/Sequencer/Private/SequencerDisplayNode.cpp
@@ -8,7 +8,7 @@
 #include "SSequencer.h"
 #include "SSequencerSectionAreaView.h"
 #include "MovieSceneSection.h"
-#include "MovieScene.h"
+#include "MovieSceneAnimation.h"
 #include "MovieSceneTrack.h"
 #include "CommonMovieSceneTools.h"
 #include "IKeyArea.h"
