Sequencer: Add support for multiple sections in drag operations.

diff --git a/Engine/Source/Editor/Sequencer/Private/SSection.cpp b/Engine/Source/Editor/Sequencer/Private/SSection.cpp
--- a/Engine/Source/Editor/Sequencer/Private/SSection.cpp
+++ b/Engine/Source/Editor/Sequencer/Private/SSection.cpp
@@ -204,17 +204,17 @@ void SSection::CreateDragOperation( const FGeometry& MyGeometry, const FPointerE
 		if( bLeftEdgePressed || bLeftEdgeHovered )
 		{
 			// Selected the start of a section
-			DragOperation = MakeShareable( new FResizeSection( GetSequencer(), *SectionObject, false ) );
+			DragOperation = MakeShareable( new FResizeSection( GetSequencer(), GetSequencer().GetSelection().GetSelectedSections(), false ) );
 		}
 		else if( bRightEdgePressed || bRightEdgeHovered )
 		{
 			// Selected the end of a section
-			DragOperation = MakeShareable( new FResizeSection( GetSequencer(), *SectionObject, true ) );
+			DragOperation = MakeShareable( new FResizeSection( GetSequencer(), GetSequencer().GetSelection().GetSelectedSections(), true ) );
 		}
 		else
 		{
 			// Entire selection moved
-			DragOperation = MakeShareable( new FMoveSection( GetSequencer(), *SectionObject ) );
+			DragOperation = MakeShareable( new FMoveSection( GetSequencer(), GetSequencer().GetSelection().GetSelectedSections() ) );
 		}
 	}
 	
@@ -516,7 +516,7 @@ FReply SSection::OnMouseButtonDown( const FGeometry& MyGeometry, const FPointerE
 
 		return FReply::Handled().CaptureMouse( AsShared() );
 	}
-	else if (MouseEvent.GetEffectingButton() == EKeys::RightMouseButton && (PressedKey.IsValid() || Sequencer.GetSelection().GetSelectedKeys().Num()))
+	else if (MouseEvent.GetEffectingButton() == EKeys::RightMouseButton && (PressedKey.IsValid() || Sequencer.GetSelection().GetSelectedKeys().Num() || Sequencer.GetSelection().GetSelectedSections().Num()))
 	{
 		// Don't return handled if we didn't click on a key so that right-click panning gets a look in
 		return FReply::Handled();
@@ -616,6 +616,10 @@ FReply SSection::OnMouseMove( const FGeometry& MyGeometry, const FPointerEvent&
 			
 			FVector2D LocalMousePos = MyGeometry.AbsoluteToLocal( MouseEvent.GetScreenSpacePosition() );
 
+			FTimeToPixel TimeToPixelConverter = SectionInterface->GetSectionObject()->IsInfinite() ? 			
+				FTimeToPixel( ParentGeometry, GetSequencer().GetViewRange()) : 
+				FTimeToPixel( MyGeometry, TRange<float>( SectionInterface->GetSectionObject()->GetStartTime(), SectionInterface->GetSectionObject()->GetEndTime() ) );
+
 			if( !bDragging )
 			{
 				// If we are not dragging determine if the mouse has moved far enough to start a drag
@@ -648,7 +652,7 @@ FReply SSection::OnMouseMove( const FGeometry& MyGeometry, const FPointerEvent&
 				
 					if( DragOperation.IsValid() )
 					{
-						DragOperation->OnBeginDrag(LocalMousePos, ParentSectionArea);
+						DragOperation->OnBeginDrag(LocalMousePos, TimeToPixelConverter, ParentSectionArea);
 					}
 
 				}
@@ -656,9 +660,6 @@ FReply SSection::OnMouseMove( const FGeometry& MyGeometry, const FPointerEvent&
 			else if( DragOperation.IsValid() )
 			{
 				// Already in a drag, tell all operations to perform their drag implementations
-				FTimeToPixel TimeToPixelConverter = SectionInterface->GetSectionObject()->IsInfinite() ? 			
-					FTimeToPixel( ParentGeometry, GetSequencer().GetViewRange()) : 
-					FTimeToPixel( MyGeometry, TRange<float>( SectionInterface->GetSectionObject()->GetStartTime(), SectionInterface->GetSectionObject()->GetEndTime() ) );
 
 				DragOperation->OnDrag( MouseEvent, LocalMousePos, TimeToPixelConverter, ParentSectionArea );
 			}
@@ -773,7 +774,10 @@ void SSection::HandleSectionSelection( const FPointerEvent& MouseEvent, bool bSe
 	// Keep key selection if right clicking to bring up a menu and the current key is selected
 	bool bKeepSectionSelection = MouseEvent.GetEffectingButton() == EKeys::RightMouseButton && bSectionIsSelected;
 	
-	if (MouseEvent.IsControlDown() && bSectionIsSelected)
+	// Don't clear selection edge is being dragged
+	bool bEdgePressed = bLeftEdgePressed || bRightEdgePressed;
+
+	if (MouseEvent.IsControlDown() && bSectionIsSelected && !bEdgePressed)
 	{
 		GetSequencer().GetSelection().RemoveFromSelection(Section);
 	}
diff --git a/Engine/Source/Editor/Sequencer/Private/SectionDragOperations.cpp b/Engine/Source/Editor/Sequencer/Private/SectionDragOperations.cpp
--- a/Engine/Source/Editor/Sequencer/Private/SectionDragOperations.cpp
+++ b/Engine/Source/Editor/Sequencer/Private/SectionDragOperations.cpp
@@ -20,13 +20,16 @@ FSequencerDragOperation::FSequencerDragOperation( FSequencer& InSequencer )
 	Settings = Sequencer.GetSettings();
 }
 
-void FSequencerDragOperation::BeginTransaction( UMovieSceneSection& Section, const FText& TransactionDesc )
+void FSequencerDragOperation::BeginTransaction( const TSet<TWeakObjectPtr<UMovieSceneSection> >& Sections, const FText& TransactionDesc )
 {
 	// Begin an editor transaction and mark the section as transactional so it's state will be saved
 	GEditor->BeginTransaction( TransactionDesc );
-	Section.SetFlags( RF_Transactional );
-	// Save the current state of the section
-	Section.Modify();
+	for (auto Section : Sections)
+	{
+		Section->SetFlags( RF_Transactional );
+		// Save the current state of the section
+		Section->Modify();
+	}
 }
 
 void FSequencerDragOperation::EndTransaction()
@@ -141,22 +144,27 @@ TOptional<float> FSequencerDragOperation::SnapToTimes(float InitialTime, const T
 }
 
 
-FResizeSection::FResizeSection( FSequencer& Sequencer, UMovieSceneSection& InSection, bool bInDraggingByEnd )
+FResizeSection::FResizeSection( FSequencer& Sequencer, const TSet<TWeakObjectPtr<UMovieSceneSection> >& InSections, bool bInDraggingByEnd )
 	: FSequencerDragOperation( Sequencer )
-	, Section( &InSection )
+	, Sections( InSections )
+	, MouseDownTime(ForceInit)
 	, bDraggingByEnd( bInDraggingByEnd )
 {
 }
 
-void FResizeSection::OnBeginDrag(const FVector2D& LocalMousePos, TSharedPtr<FTrackNode> SequencerNode)
+void FResizeSection::OnBeginDrag(const FVector2D& LocalMousePos, const FTimeToPixel& TimeToPixelConverter, TSharedPtr<FTrackNode> SequencerNode)
 {
-	if( Section.IsValid() )
-	{
-		BeginTransaction( *Section, NSLOCTEXT("Sequencer", "DragSectionEdgeTransaction", "Resize section") );
+	BeginTransaction( Sections, NSLOCTEXT("Sequencer", "DragSectionEdgeTransaction", "Resize section") );
 
-		DraggedKeyHandles.Empty();
+	MouseDownTime = TimeToPixelConverter.PixelToTime(LocalMousePos.X);
 
+	DraggedKeyHandles.Empty();
+	SectionInitTimes.Empty();
+
+	for (auto Section : Sections)
+	{
 		Section->GetKeyHandles(DraggedKeyHandles);
+		SectionInitTimes.Add(Section, bDraggingByEnd ? Section->GetEndTime() : Section->GetStartTime());
 	}
 }
 
@@ -171,10 +179,11 @@ void FResizeSection::OnDrag( const FPointerEvent& MouseEvent, const FVector2D& L
 {
 	bool bIsDilating = MouseEvent.IsControlDown();
 
-	if( Section.IsValid() )
+	for (auto Section : Sections)
 	{
 		// Convert the current mouse position to a time
-		float NewTime = TimeToPixelConverter.PixelToTime( LocalMousePos.X );
+		float DeltaTime = TimeToPixelConverter.PixelToTime(LocalMousePos.X) - MouseDownTime;
+		float NewTime = SectionInitTimes[Section] + DeltaTime;
 
 		// Find the borders of where you can drag to
 		TRange<float> SectionBoundaries = GetSectionBoundaries(Section.Get(), SequencerNode);
@@ -184,15 +193,15 @@ void FResizeSection::OnDrag( const FPointerEvent& MouseEvent, const FVector2D& L
 		{
 			bool bSnappedToSection = false;
 			if ( Settings->GetSnapSectionTimesToSections() )
-		{
-			TArray<float> TimesToSnapTo;
-			GetSectionSnapTimes(TimesToSnapTo, Section.Get(), SequencerNode, bIsDilating);
+			{
+				TArray<float> TimesToSnapTo;
+				GetSectionSnapTimes(TimesToSnapTo, Section.Get(), SequencerNode, bIsDilating);
 
-			TOptional<float> NewSnappedTime = SnapToTimes(NewTime, TimesToSnapTo, TimeToPixelConverter);
+				TOptional<float> NewSnappedTime = SnapToTimes(NewTime, TimesToSnapTo, TimeToPixelConverter);
 
-			if (NewSnappedTime.IsSet())
-			{
-				NewTime = NewSnappedTime.GetValue();
+				if (NewSnappedTime.IsSet())
+				{
+					NewTime = NewSnappedTime.GetValue();
 					bSnappedToSection = true;
 				}
 			}
@@ -240,23 +249,22 @@ void FResizeSection::OnDrag( const FPointerEvent& MouseEvent, const FVector2D& L
 	}
 }
 
-FMoveSection::FMoveSection( FSequencer& Sequencer, UMovieSceneSection& InSection )
+FMoveSection::FMoveSection( FSequencer& Sequencer, const TSet<TWeakObjectPtr<UMovieSceneSection> >& InSections )
 	: FSequencerDragOperation( Sequencer )
-	, Section( &InSection )
-	, DragOffset(ForceInit)
+	, Sections( InSections )
 {
 }
 
-void FMoveSection::OnBeginDrag(const FVector2D& LocalMousePos, TSharedPtr<FTrackNode> SequencerNode)
+void FMoveSection::OnBeginDrag(const FVector2D& LocalMousePos, const FTimeToPixel& TimeToPixelConverter, TSharedPtr<FTrackNode> SequencerNode)
 {
-	if( Section.IsValid() )
-	{
-		BeginTransaction( *Section, NSLOCTEXT("Sequencer", "MoveSectionTransaction", "Move Section") );
+	BeginTransaction( Sections, NSLOCTEXT("Sequencer", "MoveSectionTransaction", "Move Section") );
 		
-		DragOffset = LocalMousePos;
+	MouseDownTime = TimeToPixelConverter.PixelToTime(LocalMousePos.X);
 
-		DraggedKeyHandles.Empty();
+	DraggedKeyHandles.Empty();
 
+	for (auto Section : Sections)
+	{
 		Section->GetKeyHandles(DraggedKeyHandles);
 	}
 }
@@ -265,7 +273,7 @@ void FMoveSection::OnEndDrag(TSharedPtr<FTrackNode> SequencerNode)
 {
 	DraggedKeyHandles.Empty();
 
-	if (Section.IsValid())
+	for (auto Section : Sections)
 	{
 		SequencerNode->FixRowIndices();
 
@@ -284,29 +292,19 @@ void FMoveSection::OnEndDrag(TSharedPtr<FTrackNode> SequencerNode)
 
 void FMoveSection::OnDrag( const FPointerEvent& MouseEvent, const FVector2D& LocalMousePos, const FTimeToPixel& TimeToPixelConverter, TSharedPtr<FTrackNode> SequencerNode )
 {
-	if( Section.IsValid() )
+	float NewMouseDownTime = TimeToPixelConverter.PixelToTime(LocalMousePos.X);
+	float DistanceMoved =  NewMouseDownTime - MouseDownTime;
+	float DeltaTime = DistanceMoved;
+
+	for (auto Section : Sections)
 	{
-		auto Sections = SequencerNode->GetSections();
+		auto SequencerNodeSections = SequencerNode->GetSections();
 		TArray<UMovieSceneSection*> MovieSceneSections;
-		for (int32 i = 0; i < Sections.Num(); ++i)
+		for (int32 i = 0; i < SequencerNodeSections.Num(); ++i)
 		{
-			MovieSceneSections.Add(Sections[i]->GetSectionObject());
+			MovieSceneSections.Add(SequencerNodeSections[i]->GetSectionObject());
 		}
-		
-		FVector2D TotalDelta = LocalMousePos - DragOffset;
-		
-		// When the section is not infinite, the local mouse pos is relative to the start time of the section which continually updates as the section is dragged. 
-		// When the section is infinite, the local mouse pos is relative to the track extents, which do not change. 
-		// Reset the drag offset so that the motion is relative to the last moved position.
-		if (Section.Get()->IsInfinite())
-		{
-			DragOffset = LocalMousePos;
-		}
-
-		float DistanceMoved = TotalDelta.X / TimeToPixelConverter.GetPixelsPerInput();
-		
-		float DeltaTime = DistanceMoved;
-		
+						
 		if ( Settings->GetIsSnapEnabled() )
 		{
 			bool bSnappedToSection = false;
@@ -338,18 +336,18 @@ void FMoveSection::OnDrag( const FPointerEvent& MouseEvent, const FVector2D& Loc
 		int32 TargetRowIndex = Section->GetRowIndex();
 
 		// vertical dragging - master tracks only
-		if (SequencerNode->GetTrack()->SupportsMultipleRows() && Sections.Num() > 1)
+		if (SequencerNode->GetTrack()->SupportsMultipleRows() && SequencerNodeSections.Num() > 1)
 		{
-			float TrackHeight = Sections[0]->GetSectionHeight();
+			float TrackHeight = SequencerNodeSections[0]->GetSectionHeight();
 
 			if (LocalMousePos.Y < 0.f || LocalMousePos.Y > TrackHeight)
 			{
 				int32 MaxRowIndex = 0;
-				for (int32 i = 0; i < Sections.Num(); ++i)
+				for (int32 i = 0; i < SequencerNodeSections.Num(); ++i)
 				{
-					if (Sections[i]->GetSectionObject() != Section.Get())
+					if (SequencerNodeSections[i]->GetSectionObject() != Section.Get())
 					{
-						MaxRowIndex = FMath::Max(MaxRowIndex, Sections[i]->GetSectionObject()->GetRowIndex());
+						MaxRowIndex = FMath::Max(MaxRowIndex, SequencerNodeSections[i]->GetSectionObject()->GetRowIndex());
 					}
 				}
 
@@ -358,7 +356,7 @@ void FMoveSection::OnDrag( const FPointerEvent& MouseEvent, const FVector2D& Loc
 			}
 		}
 		
-		bool bDeltaX = !FMath::IsNearlyZero(TotalDelta.X);
+		bool bDeltaX = !FMath::IsNearlyZero(DeltaTime);
 		bool bDeltaY = TargetRowIndex != Section->GetRowIndex();
 
 		if (bDeltaX && bDeltaY &&
@@ -395,6 +393,8 @@ void FMoveSection::OnDrag( const FPointerEvent& MouseEvent, const FVector2D& Loc
 			}
 		}
 	}
+
+	MouseDownTime = NewMouseDownTime;
 }
 
 bool FMoveSection::IsSectionAtNewLocationValid(UMovieSceneSection* InSection, int32 RowIndex, float DeltaTime, TSharedPtr<FTrackNode> SequencerNode) const
@@ -419,7 +419,7 @@ bool FMoveSection::IsSectionAtNewLocationValid(UMovieSceneSection* InSection, in
 }
 
 
-void FMoveKeys::OnBeginDrag(const FVector2D& LocalMousePos, TSharedPtr<FTrackNode> SequencerNode)
+void FMoveKeys::OnBeginDrag(const FVector2D& LocalMousePos, const FTimeToPixel& TimeToPixelConverter, TSharedPtr<FTrackNode> SequencerNode)
 {
 	check( SelectedKeys.Num() > 0 )
 
diff --git a/Engine/Source/Editor/Sequencer/Private/SectionDragOperations.h b/Engine/Source/Editor/Sequencer/Private/SectionDragOperations.h
--- a/Engine/Source/Editor/Sequencer/Private/SectionDragOperations.h
+++ b/Engine/Source/Editor/Sequencer/Private/SectionDragOperations.h
@@ -13,7 +13,7 @@ public:
 	FSequencerDragOperation( FSequencer& InSequencer );
 
 	virtual ~FSequencerDragOperation(){}
-	void BeginTransaction( UMovieSceneSection& Section, const FText& TransactionDesc );
+	void BeginTransaction( const TSet<TWeakObjectPtr<UMovieSceneSection> >& Sections, const FText& TransactionDesc );
 	void EndTransaction();
 
 	/**
@@ -22,7 +22,7 @@ public:
 	 * @param LocalMousePos	The current relative (to sequencer) mouse location
 	 * @param SequencerNode	The sequencer node that this operation is being performed on
 	 */
-	virtual void OnBeginDrag(const FVector2D& LocalMousePos, TSharedPtr<FTrackNode> SequencerNode) = 0;
+	virtual void OnBeginDrag(const FVector2D& LocalMousePos, const struct FTimeToPixel& TimeToPixelConverter, TSharedPtr<FTrackNode> SequencerNode) = 0;
 
 	/**
 	 * Notification that drag is ending
@@ -88,18 +88,22 @@ protected:
 class FResizeSection : public FSequencerDragOperation
 {
 public:
-	FResizeSection( FSequencer& Sequencer, UMovieSceneSection& InSection, bool bInDraggingByEnd );
+	FResizeSection( FSequencer& Sequencer, const TSet<TWeakObjectPtr<UMovieSceneSection> >& InSections, bool bInDraggingByEnd );
 
 	/** FSequencerDragOperation interface */
-	virtual void OnBeginDrag(const FVector2D& LocalMousePos, TSharedPtr<FTrackNode> SequencerNode) override;
+	virtual void OnBeginDrag(const FVector2D& LocalMousePos, const FTimeToPixel& TimeToPixelConverter, TSharedPtr<FTrackNode> SequencerNode) override;
 	virtual void OnEndDrag(TSharedPtr<FTrackNode> SequencerNode) override;
 	virtual void OnDrag( const FPointerEvent& MouseEvent, const FVector2D& LocalMousePos, const FTimeToPixel& TimeToPixelConverter, TSharedPtr<FTrackNode> SequencerNode ) override;
 	virtual FCursorReply GetCursor() const override { return FCursorReply::Cursor( EMouseCursor::ResizeLeftRight ); }
 private:
-	/** The section we are interacting with */
-	TWeakObjectPtr<UMovieSceneSection> Section;
+	/** The sections we are interacting with */
+	TSet<TWeakObjectPtr<UMovieSceneSection> > Sections;
 	/** true if dragging  the end of the section, false if dragging the start */
 	bool bDraggingByEnd;
+	/** Time where the mouse is pressed */
+	float MouseDownTime;
+	/** The section start or end times when the mouse is pressed */
+	TMap<TWeakObjectPtr<UMovieSceneSection>, float> SectionInitTimes;
 	/** The exact key handles that we're dragging */
 	TSet<FKeyHandle> DraggedKeyHandles;
 };
@@ -110,10 +114,10 @@ private:
 class FMoveSection : public FSequencerDragOperation
 {
 public:
-	FMoveSection( FSequencer& Sequencer, UMovieSceneSection& InSection );
+	FMoveSection( FSequencer& Sequencer, const TSet<TWeakObjectPtr<UMovieSceneSection> >& Sections );
 
 	/** FSequencerDragOperation interface */
-	virtual void OnBeginDrag(const FVector2D& LocalMousePos, TSharedPtr<FTrackNode> SequencerNode) override;
+	virtual void OnBeginDrag(const FVector2D& LocalMousePos, const FTimeToPixel& TimeToPixelConverter, TSharedPtr<FTrackNode> SequencerNode) override;
 	virtual void OnEndDrag(TSharedPtr<FTrackNode> SequencerNode) override;
 	virtual void OnDrag( const FPointerEvent& MouseEvent, const FVector2D& LocalMousePos, const FTimeToPixel& TimeToPixelConverter, TSharedPtr<FTrackNode> SequencerNode ) override;
 	virtual FCursorReply GetCursor() const override { return FCursorReply::Cursor( EMouseCursor::CardinalCross ); }
@@ -123,10 +127,10 @@ private:
 	bool IsSectionAtNewLocationValid(UMovieSceneSection* Section, int32 TrackIndex, float DeltaTime, TSharedPtr<FTrackNode> SequencerNode) const;
 
 private:
-	/** The section we are interacting with */
-	TWeakObjectPtr<UMovieSceneSection> Section;
-	/** Local mouse position when dragging the section */
-	FVector2D DragOffset;
+	/** The sections we are interacting with */
+	TSet<TWeakObjectPtr<UMovieSceneSection> > Sections;
+	/** Time where the mouse is pressed */
+	float MouseDownTime;
 	/** The exact key handles that we're dragging */
 	TSet<FKeyHandle> DraggedKeyHandles;
 };
@@ -144,7 +148,7 @@ public:
 	{}
 
 	/** FSequencerDragOperation interface */
-	virtual void OnBeginDrag(const FVector2D& LocalMousePos, TSharedPtr<FTrackNode> SequencerNode) override;
+	virtual void OnBeginDrag(const FVector2D& LocalMousePos, const FTimeToPixel& TimeToPixelConverter, TSharedPtr<FTrackNode> SequencerNode) override;
 	virtual void OnEndDrag(TSharedPtr<FTrackNode> SequencerNode) override;
 	virtual void OnDrag( const FPointerEvent& MouseEvent, const FVector2D& LocalMousePos, const FTimeToPixel& TimeToPixelConverter, TSharedPtr<FTrackNode> SequencerNode ) override;
 
