Sequencer: Fix shadowed variable.

diff --git a/Engine/Source/Editor/Sequencer/Private/SectionDragOperations.cpp b/Engine/Source/Editor/Sequencer/Private/SectionDragOperations.cpp
--- a/Engine/Source/Editor/Sequencer/Private/SectionDragOperations.cpp
+++ b/Engine/Source/Editor/Sequencer/Private/SectionDragOperations.cpp
@@ -400,10 +400,10 @@ void FMoveSection::OnDrag( const FPointerEvent& MouseEvent, const FVector2D& Loc
 bool FMoveSection::IsSectionAtNewLocationValid(UMovieSceneSection* InSection, int32 RowIndex, float DeltaTime, TSharedPtr<FTrackNode> SequencerNode) const
 {
 	// Also get the closest borders on either side
-	const TArray< TSharedRef<ISequencerSection> >& Sections = SequencerNode->GetSections();
-	for (int32 SectionIndex = 0; SectionIndex < Sections.Num(); ++SectionIndex)
+	const TArray< TSharedRef<ISequencerSection> >& SequencerNodeSections = SequencerNode->GetSections();
+	for (int32 SectionIndex = 0; SectionIndex < SequencerNodeSections.Num(); ++SectionIndex)
 	{
-		const UMovieSceneSection* MovieSection = Sections[SectionIndex]->GetSectionObject();
+		const UMovieSceneSection* MovieSection = SequencerNodeSections[SectionIndex]->GetSectionObject();
 		if (InSection != MovieSection && MovieSection->GetRowIndex() == RowIndex)
 		{
 			TRange<float> OffsetSectionRange = TRange<float>(
