Sequencer: Time slider options.
- Added frame rates to time snap intervals. The default for UMG remains at 0.05f. The default for the level editor sequencer is 24 fps.
- Show current time/frame on playhead.
- Fix step forwards and backwards so that they now work with the time snap interval.
- Clamp mouse position so that you can't scrub outside the current view range this doesn't conflict with autoscroll because the view range will grow)
- Ctrl-drag to set the time range.
- Fixed bug where you can't snap to a time interval that's closer than a key.

diff --git a/Engine/Source/Editor/Sequencer/Private/SSequencer.cpp b/Engine/Source/Editor/Sequencer/Private/SSequencer.cpp
--- a/Engine/Source/Editor/Sequencer/Private/SSequencer.cpp
+++ b/Engine/Source/Editor/Sequencer/Private/SSequencer.cpp
@@ -531,13 +531,6 @@ TSharedRef<SWidget> SSequencer::MakeToolBar()
 
 	ToolBarBuilder.BeginSection("Snapping");
 	{
-		TArray<SNumericDropDown<float>::FNamedValue> SnapValues;
-		SnapValues.Add( SNumericDropDown<float>::FNamedValue( 0.001f, LOCTEXT( "Snap_OneThousandth", "0.001" ), LOCTEXT( "SnapDescription_OneThousandth", "Set snap to 1/1000th" ) ) );
-		SnapValues.Add( SNumericDropDown<float>::FNamedValue( 0.01f, LOCTEXT( "Snap_OneHundredth", "0.01" ), LOCTEXT( "SnapDescription_OneHundredth", "Set snap to 1/100th" ) ) );
-		SnapValues.Add( SNumericDropDown<float>::FNamedValue( 0.1f, LOCTEXT( "Snap_OneTenth", "0.1" ), LOCTEXT( "SnapDescription_OneTenth", "Set snap to 1/10th" ) ) );
-		SnapValues.Add( SNumericDropDown<float>::FNamedValue( 1.0f, LOCTEXT( "Snap_One", "1" ), LOCTEXT( "SnapDescription_One", "Set snap to 1" ) ) );
-		SnapValues.Add( SNumericDropDown<float>::FNamedValue( 10.0f, LOCTEXT( "Snap_Ten", "10" ), LOCTEXT( "SnapDescription_Ten", "Set snap to 10" ) ) );
-		SnapValues.Add( SNumericDropDown<float>::FNamedValue( 100.0f, LOCTEXT( "Snap_OneHundred", "100" ), LOCTEXT( "SnapDescription_OneHundred", "Set snap to 100" ) ) );
 
 		ToolBarBuilder.SetLabelVisibility( EVisibility::Collapsed );
 		ToolBarBuilder.AddToolBarButton( FSequencerCommands::Get().ToggleIsSnapEnabled, NAME_None, TAttribute<FText>( FText::GetEmpty() ) );
@@ -553,8 +546,9 @@ TSharedRef<SWidget> SSequencer::MakeToolBar()
 			.VAlign( VAlign_Center )
 			[
 				SNew( SNumericDropDown<float> )
-				.DropDownValues( SnapValues )
+				.DropDownValues( SequencerSnapValues::GetTimeSnapValues() )
 				.LabelText( LOCTEXT( "TimeSnapLabel", "Time" ) )
+				.bShowNamedValue(true)
 				.ToolTipText( LOCTEXT( "TimeSnappingIntervalToolTip", "Time snapping interval" ) )
 				.Value( this, &SSequencer::OnGetTimeSnapInterval )
 				.OnValueChanged( this, &SSequencer::OnTimeSnapIntervalChanged )
@@ -564,7 +558,7 @@ TSharedRef<SWidget> SSequencer::MakeToolBar()
 			.VAlign( VAlign_Center )
 			[
 				SNew( SNumericDropDown<float> )
-				.DropDownValues( SnapValues )
+				.DropDownValues( SequencerSnapValues::GetSnapValues() )
 				.LabelText( LOCTEXT( "ValueSnapLabel", "Value" ) )
 				.ToolTipText( LOCTEXT( "ValueSnappingIntervalToolTip", "Curve value snapping interval" ) )
 				.Value( this, &SSequencer::OnGetValueSnapInterval )
@@ -592,6 +586,12 @@ TSharedRef<SWidget> SSequencer::MakeSnapMenu()
 {
 	FMenuBuilder MenuBuilder( false, Sequencer.Pin()->GetCommandBindings() );
 
+	MenuBuilder.BeginSection("Frames", LOCTEXT("SnappingMenuFramesHeader", "Frames") );
+	{
+		MenuBuilder.AddMenuEntry( FSequencerCommands::Get().ToggleShowFrameNumbers );
+	}
+	MenuBuilder.EndSection();
+
 	MenuBuilder.BeginSection( "KeySnapping", LOCTEXT( "SnappingMenuKeyHeader", "Key Snapping" ) );
 	{
 		MenuBuilder.AddMenuEntry( FSequencerCommands::Get().ToggleSnapKeyTimesToInterval );
diff --git a/Engine/Source/Editor/Sequencer/Private/SectionDragOperations.cpp b/Engine/Source/Editor/Sequencer/Private/SectionDragOperations.cpp
--- a/Engine/Source/Editor/Sequencer/Private/SectionDragOperations.cpp
+++ b/Engine/Source/Editor/Sequencer/Private/SectionDragOperations.cpp
@@ -483,9 +483,22 @@ void FMoveKeys::OnDrag( const FPointerEvent& MouseEvent, const FVector2D& LocalM
 				}
 			}
 
-			if ( bSnappedToKeyTime == false && Settings->GetSnapKeyTimesToInterval() )
+			if (Settings->GetSnapKeyTimesToInterval())
 			{
-				TimeDelta = Settings->SnapTimeToInterval( MouseTime ) - SelectedKeyTime;
+				float SnapTimeIntervalDelta = Settings->SnapTimeToInterval(MouseTime) - SelectedKeyTime;
+
+				// Snap to time interval only if we haven't snapped to a keyframe or if the time difference is smaller than the difference to the keyframe.
+				if (bSnappedToKeyTime)
+				{					
+					if (FMath::Abs(SnapTimeIntervalDelta) < FMath::Abs(TimeDelta))
+					{
+						TimeDelta = SnapTimeIntervalDelta;
+					}
+				}
+				else
+				{
+					TimeDelta = SnapTimeIntervalDelta;
+				}
 			}
 		}
 
diff --git a/Engine/Source/Editor/Sequencer/Private/Sequencer.cpp b/Engine/Source/Editor/Sequencer/Private/Sequencer.cpp
--- a/Engine/Source/Editor/Sequencer/Private/Sequencer.cpp
+++ b/Engine/Source/Editor/Sequencer/Private/Sequencer.cpp
@@ -60,11 +60,11 @@ void FSequencer::InitSequencer(const FSequencerInitParams& InitParams, const TSh
 	{
 		bIsEditingWithinLevelEditor = InitParams.bEditWithinLevelEditor;
 		
-		Settings = USequencerSettingsContainer::GetOrCreate(*InitParams.ViewParams.UniqueName);
-
 		// If this sequencer edits the level, close out the existing active sequencer and mark this as the active sequencer.
 		if (bIsEditingWithinLevelEditor)
 		{
+			Settings = USequencerSettingsContainer::GetOrCreate<ULevelEditorSequencerSettings>(*InitParams.ViewParams.UniqueName);
+
 			if (ActiveSequencer.IsValid())
 			{
 				ActiveSequencer->Close();
@@ -75,6 +75,10 @@ void FSequencer::InitSequencer(const FSequencerInitParams& InitParams, const TSh
 			FEditorDelegates::PreSaveWorld.AddSP(this, &FSequencer::OnPreSaveWorld);
 			FEditorDelegates::PostSaveWorld.AddSP(this, &FSequencer::OnPostSaveWorld);
 		}
+		else
+		{
+			Settings = USequencerSettingsContainer::GetOrCreate<USequencerSettings>(*InitParams.ViewParams.UniqueName);
+		}
 
 		ToolkitHost = InitParams.ToolkitHost;
 
@@ -843,16 +847,24 @@ FReply FSequencer::OnRecord()
 FReply FSequencer::OnStepForward()
 {
 	PlaybackState = EMovieScenePlayerStatus::Stopped;
-	// @todo Sequencer Add proper step sizes
-	SetGlobalTime(ScrubPosition + 0.1f);
+	ScrubPosition += Settings->GetTimeSnapInterval();
+	if ( Settings->GetIsSnapEnabled() )
+	{
+		ScrubPosition = Settings->SnapTimeToInterval(ScrubPosition);
+	}
+	SetGlobalTime(ScrubPosition);
 	return FReply::Handled();
 }
 
 FReply FSequencer::OnStepBackward()
 {
 	PlaybackState = EMovieScenePlayerStatus::Stopped;
-	// @todo Sequencer Add proper step sizes
-	SetGlobalTime(ScrubPosition - 0.1f);
+	ScrubPosition -= Settings->GetTimeSnapInterval();
+	if ( Settings->GetIsSnapEnabled() )
+	{
+		ScrubPosition = Settings->SnapTimeToInterval(ScrubPosition);
+	}
+	SetGlobalTime(ScrubPosition);
 	return FReply::Handled();
 }
 
@@ -889,6 +901,11 @@ bool FSequencer::IsLooping() const
 	return bLoopingEnabled;
 }
 
+bool FSequencer::CanShowFrameNumbers() const
+{
+	return SequencerSnapValues::IsTimeSnapIntervalFrameRate(Settings->GetTimeSnapInterval());
+}
+
 EPlaybackMode::Type FSequencer::GetPlaybackMode() const
 {
 	return PlaybackState == EMovieScenePlayerStatus::Playing ? EPlaybackMode::PlayingForward :
@@ -1726,6 +1743,12 @@ void FSequencer::BindSequencerCommands()
 		FCanExecuteAction::CreateLambda( []{ return true; } ),
 		FIsActionChecked::CreateLambda( [this]{ return Settings->GetIsUsingCleanView(); } ) );
 
+	SequencerCommandBindings->MapAction(
+		Commands.ToggleShowFrameNumbers,
+		FExecuteAction::CreateLambda( [this]{ Settings->SetShowFrameNumbers( !Settings->GetShowFrameNumbers() ); } ),
+		FCanExecuteAction::CreateSP(this, &FSequencer::CanShowFrameNumbers ),
+		FIsActionChecked::CreateLambda( [this]{ return Settings->GetShowFrameNumbers(); } ) );
+
 	SequencerCommandBindings->MapAction(
 		Commands.ToggleIsSnapEnabled,
 		FExecuteAction::CreateLambda( [this]{ Settings->SetIsSnapEnabled( !Settings->GetIsSnapEnabled() ); } ),
diff --git a/Engine/Source/Editor/Sequencer/Private/Sequencer.h b/Engine/Source/Editor/Sequencer/Private/Sequencer.h
--- a/Engine/Source/Editor/Sequencer/Private/Sequencer.h
+++ b/Engine/Source/Editor/Sequencer/Private/Sequencer.h
@@ -245,6 +245,9 @@ public:
 	bool IsLooping() const;
 	EPlaybackMode::Type GetPlaybackMode() const;
 
+	/** Called to determine whether a frame number is set so that frame numbers can be shown */
+	bool CanShowFrameNumbers() const;
+
 	/** @return The toolkit that this sequencer is hosted in (if any) */
 	TSharedPtr<IToolkitHost> GetToolkitHost() const { return ToolkitHost.Pin(); }
 
diff --git a/Engine/Source/Editor/Sequencer/Private/SequencerCommands.cpp b/Engine/Source/Editor/Sequencer/Private/SequencerCommands.cpp
--- a/Engine/Source/Editor/Sequencer/Private/SequencerCommands.cpp
+++ b/Engine/Source/Editor/Sequencer/Private/SequencerCommands.cpp
@@ -26,6 +26,7 @@ void FSequencerCommands::RegisterCommands()
 
 	UI_COMMAND( ToggleAutoScroll, "Auto Scroll", "Toggle auto-scroll: When enabled, automatically scrolls the sequencer view to keep the current time visible.", EUserInterfaceActionType::ToggleButton, FInputChord(EModifierKey::Shift, EKeys::S) );
 
+	UI_COMMAND( ToggleShowFrameNumbers, "Show frame numbers", "Enables and disables showing frame numbers.", EUserInterfaceActionType::ToggleButton, FInputChord() );
 	UI_COMMAND( ToggleIsSnapEnabled, "Enable Snapping", "Enables and disables snapping.", EUserInterfaceActionType::ToggleButton, FInputChord() );
 
 	UI_COMMAND( ToggleSnapKeyTimesToInterval, "Snap to the interval", "Snap keys to the interval.", EUserInterfaceActionType::ToggleButton, FInputChord() );
diff --git a/Engine/Source/Editor/Sequencer/Private/SequencerCommands.h b/Engine/Source/Editor/Sequencer/Private/SequencerCommands.h
--- a/Engine/Source/Editor/Sequencer/Private/SequencerCommands.h
+++ b/Engine/Source/Editor/Sequencer/Private/SequencerCommands.h
@@ -54,6 +54,9 @@ public:
 	/** Turns auto keying on and off. */
 	TSharedPtr< FUICommandInfo > ToggleAutoKeyEnabled;
 
+	/** Turns show frame numbers on and off. */
+	TSharedPtr< FUICommandInfo > ToggleShowFrameNumbers;
+
 	/** Turns snapping on and off. */
 	TSharedPtr< FUICommandInfo > ToggleIsSnapEnabled;
 
diff --git a/Engine/Source/Editor/Sequencer/Private/SequencerCommonHelpers.h b/Engine/Source/Editor/Sequencer/Private/SequencerCommonHelpers.h
--- a/Engine/Source/Editor/Sequencer/Private/SequencerCommonHelpers.h
+++ b/Engine/Source/Editor/Sequencer/Private/SequencerCommonHelpers.h
@@ -2,6 +2,9 @@
 
 #pragma once
 
+#include "SNumericDropDown.h"
+
+#define LOCTEXT_NAMESPACE "SequencerHelpers"
 
 class SequencerHelpers
 {
@@ -15,4 +18,75 @@ public:
 	* Get descendant nodes
 	*/
 	static void GetDescendantNodes(TSharedRef<FSequencerDisplayNode> DisplayNode, TSet<TSharedRef<FSequencerDisplayNode> >& Nodes);
-};
\ No newline at end of file
+};
+
+class SequencerSnapValues
+{
+public:
+	SequencerSnapValues() {}
+
+	static const TArray<SNumericDropDown<float>::FNamedValue>& GetSnapValues()
+	{
+		static TArray<SNumericDropDown<float>::FNamedValue> SnapValues;
+		static bool bInitSnapValues = false;
+		if (!bInitSnapValues)
+		{
+			bInitSnapValues = true;
+			SnapValues.Add( SNumericDropDown<float>::FNamedValue( 0.001f, LOCTEXT( "Snap_OneThousandth", "0.001" ), LOCTEXT( "SnapDescription_OneThousandth", "Set snap to 1/1000th" ) ) );
+			SnapValues.Add( SNumericDropDown<float>::FNamedValue( 0.01f, LOCTEXT( "Snap_OneHundredth", "0.01" ), LOCTEXT( "SnapDescription_OneHundredth", "Set snap to 1/100th" ) ) );
+			SnapValues.Add( SNumericDropDown<float>::FNamedValue( 0.1f, LOCTEXT( "Snap_OneTenth", "0.1" ), LOCTEXT( "SnapDescription_OneTenth", "Set snap to 1/10th" ) ) );
+			SnapValues.Add( SNumericDropDown<float>::FNamedValue( 1.0f, LOCTEXT( "Snap_One", "1" ), LOCTEXT( "SnapDescription_One", "Set snap to 1" ) ) );
+			SnapValues.Add( SNumericDropDown<float>::FNamedValue( 10.0f, LOCTEXT( "Snap_Ten", "10" ), LOCTEXT( "SnapDescription_Ten", "Set snap to 10" ) ) );
+			SnapValues.Add( SNumericDropDown<float>::FNamedValue( 100.0f, LOCTEXT( "Snap_OneHundred", "100" ), LOCTEXT( "SnapDescription_OneHundred", "Set snap to 100" ) ) );
+		}
+		return SnapValues;
+	}
+
+	static const TArray<SNumericDropDown<float>::FNamedValue>& GetFrameRateSnapValues()
+	{
+		static TArray<SNumericDropDown<float>::FNamedValue> FrameRateSnapValues;
+		static bool bInitFrameRateSnapValues = false;
+		if (!bInitFrameRateSnapValues)
+		{
+			bInitFrameRateSnapValues = true;
+			FrameRateSnapValues.Add( SNumericDropDown<float>::FNamedValue( 0.066667f, LOCTEXT( "Snap_15Fps", "15 fps" ), LOCTEXT( "SnapDescription_15Fps", "Set snap to 15 fps" ) ) );
+			FrameRateSnapValues.Add( SNumericDropDown<float>::FNamedValue( 0.041667f, LOCTEXT( "Snap_24Fps", "24 fps (film)" ), LOCTEXT( "SnapDescription_24Fps", "Set snap to 24 fps" ) ) );
+			FrameRateSnapValues.Add( SNumericDropDown<float>::FNamedValue( 0.04f, LOCTEXT( "Snap_25Fps", "25 fps (PAL/25)" ), LOCTEXT( "SnapDescription_25Fps", "Set snap to 25 fps" ) ) );
+			FrameRateSnapValues.Add( SNumericDropDown<float>::FNamedValue( 0.033367, LOCTEXT( "Snap_29.97Fps", "29.97 fps (NTSC/30)" ), LOCTEXT( "SnapDescription_29.97Fps", "Set snap to 29.97 fps" ) ) );
+			FrameRateSnapValues.Add( SNumericDropDown<float>::FNamedValue( 0.033334f, LOCTEXT( "Snap_30Fps", "30 fps" ), LOCTEXT( "SnapDescription_30Fps", "Set snap to 30 fps" ) ) );
+			FrameRateSnapValues.Add( SNumericDropDown<float>::FNamedValue( 0.020834f, LOCTEXT( "Snap_48Fps", "48 fps" ), LOCTEXT( "SnapDescription_48Fps", "Set snap to 48 fps" ) ) );
+			FrameRateSnapValues.Add( SNumericDropDown<float>::FNamedValue( 0.02f, LOCTEXT( "Snap_50Fps", "50 fps (PAL/50)" ), LOCTEXT( "SnapDescription_50Fps", "Set snap to 50 fps" ) ) );
+			FrameRateSnapValues.Add( SNumericDropDown<float>::FNamedValue( 0.016683f, LOCTEXT( "Snap_50.94Fps", "50.94 fps (NTSC/60)" ), LOCTEXT( "SnapDescription_50.94Fps", "Set snap to 50.94 fps" ) ) );
+			FrameRateSnapValues.Add( SNumericDropDown<float>::FNamedValue( 0.016667f, LOCTEXT( "Snap_60Fps", "60 fps" ), LOCTEXT( "SnapDescription_60Fps", "Set snap to 60 fps" ) ) );
+			FrameRateSnapValues.Add( SNumericDropDown<float>::FNamedValue( 0.008334f, LOCTEXT( "Snap_120Fps", "120 fps" ), LOCTEXT( "SnapDescription_120Fps", "Set snap to 120 fps" ) ) );
+		}
+		return FrameRateSnapValues;
+	}
+
+	static const TArray<SNumericDropDown<float>::FNamedValue>& GetTimeSnapValues()
+	{
+		static TArray<SNumericDropDown<float>::FNamedValue> TimeSnapValues;
+		static bool bInitTimeSnapValues = false;
+		if (!bInitTimeSnapValues)
+		{
+			bInitTimeSnapValues = true;
+			TimeSnapValues.Append(GetSnapValues());
+			TimeSnapValues.Append(GetFrameRateSnapValues());
+		}
+		return TimeSnapValues;
+	}
+
+	static bool IsTimeSnapIntervalFrameRate(float InFrameRate)
+	{
+		for (auto FrameRateSnapValue : GetFrameRateSnapValues())
+		{
+			if (FMath::IsNearlyEqual(InFrameRate, FrameRateSnapValue.GetValue()))
+			{
+				return true;
+			}
+		}
+		return false;
+	}
+};
+
+#undef LOCTEXT_NAMESPACE
\ No newline at end of file
diff --git a/Engine/Source/Editor/Sequencer/Private/SequencerSettings.cpp b/Engine/Source/Editor/Sequencer/Private/SequencerSettings.cpp
--- a/Engine/Source/Editor/Sequencer/Private/SequencerSettings.cpp
+++ b/Engine/Source/Editor/Sequencer/Private/SequencerSettings.cpp
@@ -2,31 +2,11 @@
 
 #include "SequencerPrivatePCH.h"
 
-USequencerSettings* USequencerSettingsContainer::GetOrCreate(const TCHAR* InName)
-{
-	static const TCHAR* SettingsContainerName = TEXT("SequencerSettingsContainer");
-
-	auto* Outer = FindObject<USequencerSettingsContainer>(GetTransientPackage(), SettingsContainerName);
-	if (!Outer)
-	{
-		Outer = NewObject<USequencerSettingsContainer>(GetTransientPackage(), USequencerSettingsContainer::StaticClass(), SettingsContainerName);
-		Outer->AddToRoot();
-	}
-	
-	USequencerSettings* Inst = FindObject<USequencerSettings>( Outer, InName );
-	if (!Inst)
-	{
-		Inst = NewObject<USequencerSettings>( Outer, USequencerSettings::StaticClass(), InName );
-		Inst->LoadConfig();
-	}
-
-	return Inst;
-}
-
 USequencerSettings::USequencerSettings( const FObjectInitializer& ObjectInitializer )
 	: Super( ObjectInitializer )
 {
 	bAutoKeyEnabled = false;
+	bShowFrameNumbers = true;
 	bIsSnapEnabled = true;
 	TimeSnapInterval = .05f;
 	bSnapKeyTimesToInterval = true;
@@ -57,6 +37,20 @@ void USequencerSettings::SetAutoKeyEnabled(bool InbAutoKeyEnabled)
 	}
 }
 
+bool USequencerSettings::GetShowFrameNumbers() const
+{
+	return bShowFrameNumbers;
+}
+
+void USequencerSettings::SetShowFrameNumbers(bool InbShowFrameNumbers)
+{
+	if ( bShowFrameNumbers != InbShowFrameNumbers )
+	{
+		bShowFrameNumbers = InbShowFrameNumbers;
+		SaveConfig();
+	}
+}
+
 bool USequencerSettings::GetIsSnapEnabled() const
 {
 	return bIsSnapEnabled;
@@ -267,3 +261,9 @@ USequencerSettings::FOnShowCurveEditorChanged& USequencerSettings::GetOnShowCurv
 	return OnShowCurveEditorChanged;
 }
 
+/** Level editor specific sequencer settings */
+ULevelEditorSequencerSettings::ULevelEditorSequencerSettings( const FObjectInitializer& ObjectInitializer )
+	: Super( ObjectInitializer )
+{
+	TimeSnapInterval = 0.041667f;
+}
\ No newline at end of file
diff --git a/Engine/Source/Editor/Sequencer/Private/SequencerSettings.h b/Engine/Source/Editor/Sequencer/Private/SequencerSettings.h
--- a/Engine/Source/Editor/Sequencer/Private/SequencerSettings.h
+++ b/Engine/Source/Editor/Sequencer/Private/SequencerSettings.h
@@ -12,7 +12,27 @@ public:
 	GENERATED_BODY()
 
 	/** Get or create a settings object for the specified name */
-	static USequencerSettings* GetOrCreate(const TCHAR* InName);
+	template<class T> 
+	static T* GetOrCreate(const TCHAR* InName)
+	{
+		static const TCHAR* SettingsContainerName = TEXT("SequencerSettingsContainer");
+
+		auto* Outer = FindObject<USequencerSettingsContainer>(GetTransientPackage(), SettingsContainerName);
+		if (!Outer)
+		{
+			Outer = NewObject<USequencerSettingsContainer>(GetTransientPackage(), USequencerSettingsContainer::StaticClass(), SettingsContainerName);
+			Outer->AddToRoot();
+		}
+	
+		T* Inst = FindObject<T>( Outer, InName );
+		if (!Inst)
+		{
+			Inst = NewObject<T>( Outer, T::StaticClass(), InName );
+			Inst->LoadConfig();
+		}
+
+		return Inst;
+	}
 };
 
 /** Serializable options for sequencer. */
@@ -29,6 +49,11 @@ public:
 	/** Sets whether or not auto key is enabled. */
 	void SetAutoKeyEnabled(bool InbAutoKeyEnabled);
 
+	/** Gets whether or not to show frame numbers. */
+	bool GetShowFrameNumbers() const;
+	/** Sets whether or not to show frame numbers. */
+	void SetShowFrameNumbers(bool InbShowFrameNumbers);
+
 	/** Gets whether or not snapping is enabled. */
 	bool GetIsSnapEnabled() const;
 	/** Sets whether or not snapping is enabled. */
@@ -109,6 +134,9 @@ protected:
 	UPROPERTY( config )
 	bool bAutoKeyEnabled;
 
+	UPROPERTY( config )
+	bool bShowFrameNumbers;
+
 	UPROPERTY( config )
 	bool bIsSnapEnabled;
 
@@ -152,4 +180,12 @@ protected:
 	bool bShowCurveEditorCurveToolTips;
 
 	FOnShowCurveEditorChanged OnShowCurveEditorChanged;
-};
\ No newline at end of file
+};
+
+/** Level editor specific sequencer settings */
+UCLASS(config=EditorPerProjectUserSettings, PerObjectConfig)
+class ULevelEditorSequencerSettings : public USequencerSettings
+{
+public:
+	GENERATED_UCLASS_BODY()
+};
diff --git a/Engine/Source/Editor/Sequencer/Private/TimeSliderController.cpp b/Engine/Source/Editor/Sequencer/Private/TimeSliderController.cpp
--- a/Engine/Source/Editor/Sequencer/Private/TimeSliderController.cpp
+++ b/Engine/Source/Editor/Sequencer/Private/TimeSliderController.cpp
@@ -8,7 +8,7 @@
 namespace ScrubConstants
 {
 	/** The minimum amount of pixels between each major ticks on the widget */
-	const int32 MinPixelsPerDisplayTick = 5;
+	const int32 MinPixelsPerDisplayTick = 12;
 
 	/**The smallest number of units between between major tick marks */
 	const float MinDisplayTickSpacing = 0.001f;
@@ -17,6 +17,13 @@ namespace ScrubConstants
 	const float ScrollPanFraction = 0.1f;
 }
 
+namespace TimeConversion
+{
+	int32 TimeToFrame(float Time, float FrameRate)
+	{
+		return FMath::RoundToInt(Time * FrameRate);
+	}
+}
 
 /** Utility struct for converting between scrub range space and local/absolute screen space */
 struct FScrubRangeToScreen
@@ -102,7 +109,7 @@ float DetermineOptimalSpacing( float InPixelsPerInput, uint32 MinTick, float Min
 FSequencerTimeSliderController::FSequencerTimeSliderController( const FTimeSliderArgs& InArgs )
 	: TimeSliderArgs( InArgs )
 	, DistanceDragged( 0.0f )
-	, bDraggingScrubber( false )
+	, MouseDragType( DRAG_NONE )
 	, bPanning( false )
 {
 	ScrubHandleUp = FEditorStyle::GetBrush( TEXT( "Sequencer.Timeline.ScrubHandleUp" ) ); 
@@ -134,7 +141,13 @@ struct FDrawTickArgs
 
 void FSequencerTimeSliderController::DrawTicks( FSlateWindowElementList& OutDrawElements, const struct FScrubRangeToScreen& RangeToScreen, FDrawTickArgs& InArgs ) const
 {
-	const float Spacing = DetermineOptimalSpacing( RangeToScreen.PixelsPerInput, ScrubConstants::MinPixelsPerDisplayTick, ScrubConstants::MinDisplayTickSpacing );
+	float MinDisplayTickSpacing = ScrubConstants::MinDisplayTickSpacing;
+	if (SequencerSnapValues::IsTimeSnapIntervalFrameRate(TimeSliderArgs.Settings->GetTimeSnapInterval()) && TimeSliderArgs.Settings->GetShowFrameNumbers())
+	{
+		MinDisplayTickSpacing = TimeSliderArgs.Settings->GetTimeSnapInterval();
+	}
+
+	const float Spacing = DetermineOptimalSpacing( RangeToScreen.PixelsPerInput, ScrubConstants::MinPixelsPerDisplayTick, MinDisplayTickSpacing );
 
 	// Sub divisions
 	// @todo Sequencer may need more robust calculation
@@ -149,6 +162,8 @@ void FSequencerTimeSliderController::DrawTicks( FSlateWindowElementList& OutDraw
 	TArray<FVector2D> LinePoints;
 	LinePoints.AddUninitialized(2);
 
+	float SnapSize = TimeSliderArgs.Settings->GetTimeSnapInterval();
+
 	float Seconds = 0;
 	while( (Seconds = OffsetNum*Spacing) < RangeToScreen.ViewInput.GetUpperBoundValue() )
 	{
@@ -181,13 +196,20 @@ void FSequencerTimeSliderController::DrawTicks( FSlateWindowElementList& OutDraw
 
 			if( !InArgs.bOnlyDrawMajorTicks )
 			{
-				FString FrameString = Spacing == ScrubConstants::MinDisplayTickSpacing ? FString::Printf( TEXT("%.3f"), Seconds ) : FString::Printf( TEXT("%.2f"), Seconds );
+				FString FrameString;
+				if (SequencerSnapValues::IsTimeSnapIntervalFrameRate(TimeSliderArgs.Settings->GetTimeSnapInterval()) && TimeSliderArgs.Settings->GetShowFrameNumbers())
+				{
+					FrameString = FString::Printf( TEXT("%d"), TimeConversion::TimeToFrame(Seconds, 1.0f/SnapSize));
+				}
+				else
+				{
+					FrameString = Spacing == ScrubConstants::MinDisplayTickSpacing ? FString::Printf( TEXT("%.3f"), Seconds ) : FString::Printf( TEXT("%.2f"), Seconds );
+				}
 
 				// Space the text between the tick mark but slightly above
 				const TSharedRef< FSlateFontMeasure > FontMeasureService = FSlateApplication::Get().GetRenderer()->GetFontMeasureService();
 				FVector2D TextSize = FontMeasureService->Measure(FrameString, SmallLayoutFont);
-				FVector2D TextOffset( XPos-(TextSize.X*0.5f), InArgs.bMirrorLabels ? TextSize.Y :  FMath::Abs( InArgs.AllottedGeometry.Size.Y - (InArgs.MajorTickHeight+TextSize.Y) ) );
-
+				FVector2D TextOffset( XPos + 5.f, InArgs.bMirrorLabels ? 3.f : FMath::Abs( InArgs.AllottedGeometry.Size.Y - (InArgs.MajorTickHeight+3.f) ) );
 				FSlateDrawElement::MakeText(
 					OutDrawElements,
 					InArgs.StartLayer+1, 
@@ -196,14 +218,14 @@ void FSequencerTimeSliderController::DrawTicks( FSlateWindowElementList& OutDraw
 					SmallLayoutFont, 
 					InArgs.ClippingRect, 
 					InArgs.DrawEffects,
-					InArgs.TickColor 
+					InArgs.TickColor*0.65f 
 				);
 			}
 		}
 		else if( !InArgs.bOnlyDrawMajorTicks )
 		{
 			// Compute the size of each tick mark.  If we are half way between to visible values display a slightly larger tick mark
-			const float MinorTickHeight = AbsOffsetNum % HalfDivider == 0 ? 7.0f : 4.0f;
+			const float MinorTickHeight = AbsOffsetNum % HalfDivider == 0 ? 6.0f : 2.0f;
 
 			FVector2D Offset(XPos, InArgs.bMirrorLabels ? 0.0f : FMath::Abs( InArgs.AllottedGeometry.Size.Y - MinorTickHeight ) );
 			FVector2D TickSize(0.0f, MinorTickHeight);
@@ -264,7 +286,7 @@ int32 FSequencerTimeSliderController::OnPaintTimeSlider( bool bMirrorLabels, con
 		float HalfSize = FMath::CeilToFloat(HandleSize/2.0f);
 
 		// Draw the scrub handle
-		const float XPos = RangeToScreen.InputToLocalX( TimeSliderArgs.ScrubPosition.Get() );
+		float XPos = RangeToScreen.InputToLocalX( TimeSliderArgs.ScrubPosition.Get() );
 
 		// Should draw above the text
 		const int32 ArrowLayer = LayerId + 2;
@@ -285,6 +307,67 @@ int32 FSequencerTimeSliderController::OnPaintTimeSlider( bool bMirrorLabels, con
 			ScrubColor
 			);
 
+		// Draw the current time next to the scrub handle
+		float Time = TimeSliderArgs.ScrubPosition.Get();
+
+		FString FrameString;
+		if (SequencerSnapValues::IsTimeSnapIntervalFrameRate(TimeSliderArgs.Settings->GetTimeSnapInterval()) && TimeSliderArgs.Settings->GetShowFrameNumbers())
+		{
+			float SnapSize = TimeSliderArgs.Settings->GetTimeSnapInterval();
+			FrameString = FString::Printf( TEXT("%d"), TimeConversion::TimeToFrame(Time, 1.0f/SnapSize));
+		}
+		else
+		{
+			FrameString = FString::Printf( TEXT("%.2f"), Time );
+		}
+
+		FSlateFontInfo SmallLayoutFont( FPaths::EngineContentDir() / TEXT("Slate/Fonts/Roboto-Regular.ttf"), 10 );
+
+		const TSharedRef< FSlateFontMeasure > FontMeasureService = FSlateApplication::Get().GetRenderer()->GetFontMeasureService();
+		FVector2D TextSize = FontMeasureService->Measure(FrameString, SmallLayoutFont);
+
+		// Flip the text position if getting near the end of the view range
+		if ((AllottedGeometry.Size.X - XPos) < (TextSize.X + 14.f))
+		{
+			XPos = XPos - TextSize.X - 12.f;
+		}
+		else
+		{
+			XPos = XPos + 10.f;
+		}
+
+		FVector2D TextOffset( XPos,  Args.bMirrorLabels ? TextSize.Y-6.f : Args.AllottedGeometry.Size.Y - (Args.MajorTickHeight+TextSize.Y) );
+
+		FSlateDrawElement::MakeText(
+			OutDrawElements,
+			Args.StartLayer+1, 
+			Args.AllottedGeometry.ToPaintGeometry( TextOffset, TextSize ), 
+			FrameString, 
+			SmallLayoutFont, 
+			Args.ClippingRect, 
+			Args.DrawEffects,
+			Args.TickColor 
+		);
+		
+		if (MouseDragType == DRAG_SETTING_RANGE)
+		{
+			float MouseStartPosX = RangeToScreen.InputToLocalX(MouseDownRange[0]);
+			float MouseEndPosX = RangeToScreen.InputToLocalX(MouseDownRange[1]);
+
+			float RangePosX = MouseStartPosX < MouseEndPosX ? MouseStartPosX : MouseEndPosX;
+			float RangeSizeX = FMath::Abs(MouseStartPosX - MouseEndPosX);
+
+			FSlateDrawElement::MakeBox(
+				OutDrawElements,
+				LayerId+1,
+				AllottedGeometry.ToPaintGeometry( FVector2D(RangePosX, 0.f), FVector2D(RangeSizeX, AllottedGeometry.Size.Y) ),
+				bMirrorLabels ? ScrubHandleDown : ScrubHandleUp,
+				MyClippingRect,
+				DrawEffects,
+				MouseStartPosX < MouseEndPosX ? FLinearColor(0.5f, 0.5f, 0.5f) : FLinearColor(0.25f, 0.3f, 0.3f)
+			);
+		}
+
 		return ArrowLayer;
 	}
 
@@ -298,6 +381,10 @@ FReply FSequencerTimeSliderController::OnMouseButtonDown( TSharedRef<SWidget> Wi
 	
 	DistanceDragged = 0;
 
+	FScrubRangeToScreen RangeToScreen( TimeSliderArgs.ViewRange.Get(), MyGeometry.Size );
+	MouseDownRange[0] = RangeToScreen.LocalXToInput(MyGeometry.AbsoluteToLocal(MouseEvent.GetLastScreenSpacePosition()).X);
+	MouseDownRange[1] = MouseDownRange[0];
+
 	if ( bHandleLeftMouseButton )
 	{
 		return FReply::Handled().CaptureMouse( WidgetOwner ).PreventThrottling();
@@ -330,9 +417,51 @@ FReply FSequencerTimeSliderController::OnMouseButtonUp( TSharedRef<SWidget> Widg
 	}
 	else if ( bHandleLeftMouseButton )
 	{
-		if( bDraggingScrubber )
+		if( MouseDragType != DRAG_NONE )
 		{
 			TimeSliderArgs.OnEndScrubberMovement.ExecuteIfBound();
+
+			if (MouseDragType == DRAG_SETTING_RANGE)
+			{
+				FScrubRangeToScreen RangeToScreen( TimeSliderArgs.ViewRange.Get(), MyGeometry.Size );
+				FVector2D CursorPos = MyGeometry.AbsoluteToLocal(MouseEvent.GetLastScreenSpacePosition());
+				float NewValue = RangeToScreen.LocalXToInput(CursorPos.X);
+
+				if ( TimeSliderArgs.Settings->GetIsSnapEnabled() )
+				{
+					NewValue = TimeSliderArgs.Settings->SnapTimeToInterval(NewValue);
+				}
+
+				float DownValue = MouseDownRange[0];
+					
+				if ( TimeSliderArgs.Settings->GetIsSnapEnabled() )
+				{
+					DownValue = TimeSliderArgs.Settings->SnapTimeToInterval(DownValue);
+				}
+
+				// Zoom in
+				if (NewValue > DownValue)
+				{
+					// push the current value onto the stack
+					RangeStack.Add(FVector2D(TimeSliderArgs.ViewRange.Get().GetLowerBoundValue(), TimeSliderArgs.ViewRange.Get().GetUpperBoundValue()));
+				}
+				// Zoom out
+				else if (RangeStack.Num())
+				{
+					// pop the stack
+					FVector2D LastRange = RangeStack.Pop();
+					DownValue = LastRange[0];
+					NewValue = LastRange[1];
+				}
+
+				TimeSliderArgs.OnViewRangeChanged.ExecuteIfBound(TRange<float>(DownValue, NewValue), EViewRangeInterpolation::Immediate);
+						
+				if( !TimeSliderArgs.ViewRange.IsBound() )
+				{	
+					// The output is not bound to a delegate so we'll manage the value ourselves
+					TimeSliderArgs.ViewRange.Set( TRange<float>( DownValue, NewValue ) );
+				}
+			}
 		}
 		else
 		{
@@ -340,15 +469,15 @@ FReply FSequencerTimeSliderController::OnMouseButtonUp( TSharedRef<SWidget> Widg
 			FVector2D CursorPos = MyGeometry.AbsoluteToLocal(MouseEvent.GetLastScreenSpacePosition());
 			float NewValue = RangeToScreen.LocalXToInput(CursorPos.X);
 
-			if ( TimeSliderArgs.Settings->GetIsSnapEnabled() && TimeSliderArgs.Settings->GetSnapPlayTimeToInterval() )
+			if ( TimeSliderArgs.Settings->GetIsSnapEnabled() )
 			{
-				NewValue = TimeSliderArgs.Settings->SnapTimeToInterval( NewValue );
+				NewValue = TimeSliderArgs.Settings->SnapTimeToInterval(NewValue);
 			}
-
+			
 			CommitScrubPosition( NewValue, /*bIsScrubbing=*/false );
 		}
 
-		bDraggingScrubber = false;
+		MouseDragType = DRAG_NONE;
 		return FReply::Handled().ReleaseMouseCapture();
 
 	}
@@ -389,12 +518,16 @@ FReply FSequencerTimeSliderController::OnMouseMove( TSharedRef<SWidget> WidgetOw
 		}
 		else if (MouseEvent.IsMouseButtonDown( EKeys::LeftMouseButton ))
 		{
-			if ( !bDraggingScrubber )
+			DistanceDragged += FMath::Abs( MouseEvent.GetCursorDelta().X );
+
+			if ( MouseDragType == DRAG_NONE )
 			{
-				DistanceDragged += FMath::Abs( MouseEvent.GetCursorDelta().X );
 				if ( DistanceDragged > FSlateApplication::Get().GetDragTriggerDistance() )
 				{
-					bDraggingScrubber = true;
+					if (FSlateApplication::Get().GetModifierKeys().AreModifersDown(EModifierKey::Control))
+						MouseDragType = DRAG_SETTING_RANGE;
+					else
+						MouseDragType = DRAG_SCRUBBING_TIME;
 					TimeSliderArgs.OnBeginScrubberMovement.ExecuteIfBound();
 				}
 			}
@@ -404,12 +537,22 @@ FReply FSequencerTimeSliderController::OnMouseMove( TSharedRef<SWidget> WidgetOw
 				FVector2D CursorPos = MyGeometry.AbsoluteToLocal( MouseEvent.GetLastScreenSpacePosition() );
 				float NewValue = RangeToScreen.LocalXToInput( CursorPos.X );
 
-				if ( TimeSliderArgs.Settings->GetIsSnapEnabled() && TimeSliderArgs.Settings->GetSnapPlayTimeToInterval() )
+				if (MouseDragType == DRAG_SCRUBBING_TIME)
 				{
-					NewValue = TimeSliderArgs.Settings->SnapTimeToInterval(NewValue);
-				}
+					if ( TimeSliderArgs.Settings->GetIsSnapEnabled() )
+					{
+						NewValue = TimeSliderArgs.Settings->SnapTimeToInterval(NewValue);
+					}
 
-				CommitScrubPosition( NewValue, /*bIsScrubbing=*/true );
+					// Clamp to the view range so that you can't scroll beyond the current view range. This doesn't conflict with auto-scroll because the view range will grow before you get to the bounds.
+					NewValue = FMath::Clamp(NewValue, TimeSliderArgs.ViewRange.Get().GetLowerBoundValue(), TimeSliderArgs.ViewRange.Get().GetUpperBoundValue());
+
+					CommitScrubPosition( NewValue, /*bIsScrubbing=*/true );
+				}
+				else if (MouseDragType == DRAG_SETTING_RANGE)
+				{
+					MouseDownRange[1] = NewValue;
+				}
 			}
 		}
 		return FReply::Handled();
diff --git a/Engine/Source/Editor/Sequencer/Private/TimeSliderController.h b/Engine/Source/Editor/Sequencer/Private/TimeSliderController.h
--- a/Engine/Source/Editor/Sequencer/Private/TimeSliderController.h
+++ b/Engine/Source/Editor/Sequencer/Private/TimeSliderController.h
@@ -79,8 +79,18 @@ private:
 	const FSlateBrush* ScrubHandleDown;
 	/** Total mouse delta during dragging **/
 	float DistanceDragged;
-	/** If we are dragging the scrubber */
-	bool bDraggingScrubber;
+	/** If we are dragging the scrubber or dragging to set the time range */
+	enum DragType
+	{
+		DRAG_SCRUBBING_TIME,
+		DRAG_SETTING_RANGE,
+		DRAG_NONE
+	};
+	DragType MouseDragType;
 	/** If we are currently panning the panel */
 	bool bPanning;
+	/** Mouse down time range */
+	FVector2D MouseDownRange;
+	/** Range stack */
+	TArray<FVector2D> RangeStack;
 };
