Sequencer: Shift click to snap time to selected key.

diff --git a/Engine/Source/Editor/Sequencer/Private/SSection.cpp b/Engine/Source/Editor/Sequencer/Private/SSection.cpp
--- a/Engine/Source/Editor/Sequencer/Private/SSection.cpp
+++ b/Engine/Source/Editor/Sequencer/Private/SSection.cpp
@@ -554,6 +554,17 @@ FReply SSection::OnMouseButtonUp( const FGeometry& MyGeometry, const FPointerEve
 	{
 		if( ( MouseEvent.GetEffectingButton() == EKeys::LeftMouseButton || MouseEvent.GetEffectingButton() == EKeys::RightMouseButton ) && MyGeometry.IsUnderLocation( MouseEvent.GetScreenSpacePosition() ) )
 		{
+			// Snap time to the key under the mouse if shift is down
+			if (MouseEvent.IsShiftDown() && MouseEvent.GetEffectingButton() == EKeys::LeftMouseButton)
+			{
+				FSelectedKey Key = GetKeyUnderMouse( MouseEvent.GetScreenSpacePosition(), MyGeometry );
+				if (Key.IsValid())
+				{
+					float KeyTime = Key.KeyArea->GetKeyTime(Key.KeyHandle.GetValue());
+					GetSequencer().SetGlobalTime(KeyTime);
+				}
+			}
+
 			HandleSelection( MyGeometry, MouseEvent );
 		}
 
