Sequencer category nodes are now collapsed by default

diff --git a/Engine/Source/Editor/Sequencer/Private/SSequencerTrackArea.cpp b/Engine/Source/Editor/Sequencer/Private/SSequencerTrackArea.cpp
--- a/Engine/Source/Editor/Sequencer/Private/SSequencerTrackArea.cpp
+++ b/Engine/Source/Editor/Sequencer/Private/SSequencerTrackArea.cpp
@@ -397,7 +397,7 @@ void SSequencerTrackArea::HandleMarqueeSelection(const FGeometry& Geometry, cons
 			{
 				FSectionKeyAreaNode& KeyAreaNode = static_cast<FSectionKeyAreaNode&>(InNode);
 
-				if (CurrentSections && InNode.VirtualTop > TopLeft.Y && InNode.VirtualBottom < BottomRight.Y)
+				if (CurrentSections && InNode.GetVirtualTop() > TopLeft.Y && InNode.GetVirtualBottom() < BottomRight.Y)
 				{
 					PerformSelection(KeyAreaNode);
 				}
@@ -409,7 +409,7 @@ void SSequencerTrackArea::HandleMarqueeSelection(const FGeometry& Geometry, cons
 				CurrentSections = &TrackNode.GetSections();
 
 				TSharedPtr<FSectionKeyAreaNode> SectionKeyNode = TrackNode.GetTopLevelKeyNode();
-				if( SectionKeyNode.IsValid() && TrackNode.VirtualTop > TopLeft.Y && TrackNode.VirtualBottom < BottomRight.Y )
+				if( SectionKeyNode.IsValid() && TrackNode.GetVirtualTop() > TopLeft.Y && TrackNode.GetVirtualBottom() < BottomRight.Y )
 				{
 					PerformSelection(*SectionKeyNode);
 				}
diff --git a/Engine/Source/Editor/Sequencer/Private/SSequencerTreeView.cpp b/Engine/Source/Editor/Sequencer/Private/SSequencerTreeView.cpp
--- a/Engine/Source/Editor/Sequencer/Private/SSequencerTreeView.cpp
+++ b/Engine/Source/Editor/Sequencer/Private/SSequencerTreeView.cpp
@@ -175,13 +175,13 @@ float SSequencerTreeView::PhysicalToVirtual(float InPhysical) const
 	if (Found)
 	{
 		const float FractionalHeight = (InPhysical - Found->PhysicalTop) / Found->PhysicalHeight;
-		return Found->Node->VirtualTop + (Found->Node->VirtualBottom - Found->Node->VirtualTop) * FractionalHeight;
+		return Found->Node->GetVirtualTop() + (Found->Node->GetVirtualBottom() - Found->Node->GetVirtualTop()) * FractionalHeight;
 	}
 
 	if (PhysicalNodes.Num())
 	{
 		auto& Last = PhysicalNodes.Last();
-		return Last.Node->VirtualTop + (InPhysical - Last.PhysicalTop);
+		return Last.Node->GetVirtualTop() + (InPhysical - Last.PhysicalTop);
 	}
 
 	return InPhysical;
@@ -191,11 +191,11 @@ float SSequencerTreeView::VirtualToPhysical(float InVirtual) const
 {
 	auto* Found = Utils::BinarySearch(PhysicalNodes, [&](const FCachedGeometry& In){
 
-		if (InVirtual < In.Node->VirtualTop)
+		if (InVirtual < In.Node->GetVirtualTop())
 		{
 			return Utils::ESearchState::Before;
 		}
-		else if (InVirtual > In.Node->VirtualBottom)
+		else if (InVirtual > In.Node->GetVirtualBottom())
 		{
 			return Utils::ESearchState::After;
 		}
@@ -206,14 +206,14 @@ float SSequencerTreeView::VirtualToPhysical(float InVirtual) const
 
 	if (Found)
 	{
-		const float FractionalHeight = (InVirtual - Found->Node->VirtualTop) / (Found->Node->VirtualBottom - Found->Node->VirtualTop);
+		const float FractionalHeight = (InVirtual - Found->Node->GetVirtualTop()) / (Found->Node->GetVirtualBottom() - Found->Node->GetVirtualTop());
 		return Found->PhysicalTop + Found->PhysicalHeight * FractionalHeight;
 	}
 	
 	if (PhysicalNodes.Num())
 	{
 		auto Last = PhysicalNodes.Last();
-		return Last.PhysicalTop + (InVirtual - Last.Node->VirtualTop);
+		return Last.PhysicalTop + (InVirtual - Last.Node->GetVirtualTop());
 	}
 
 	return InVirtual;
diff --git a/Engine/Source/Editor/Sequencer/Private/SequencerDisplayNode.cpp b/Engine/Source/Editor/Sequencer/Private/SequencerDisplayNode.cpp
--- a/Engine/Source/Editor/Sequencer/Private/SequencerDisplayNode.cpp
+++ b/Engine/Source/Editor/Sequencer/Private/SequencerDisplayNode.cpp
@@ -129,8 +129,16 @@ FSequencerDisplayNode::FSequencerDisplayNode( FName InNodeName, TSharedPtr<FSequ
 	, bExpanded( false )
 	, bCachedShotFilteredVisibility( true )
 	, bNodeIsPinned( false )
+{
+	
+}
+
+void FSequencerDisplayNode::Initialize(float InVirtualTop, float InVirtualBottom)
 {
 	bExpanded = ParentTree.GetSavedExpansionState( *this );
+
+	VirtualTop = InVirtualTop;
+	VirtualBottom = InVirtualBottom;
 }
 
 void FSequencerDisplayNode::AddObjectBindingNode(TSharedRef<FObjectBindingNode> ObjectBindingNode)
diff --git a/Engine/Source/Editor/Sequencer/Private/SequencerDisplayNode.h b/Engine/Source/Editor/Sequencer/Private/SequencerDisplayNode.h
--- a/Engine/Source/Editor/Sequencer/Private/SequencerDisplayNode.h
+++ b/Engine/Source/Editor/Sequencer/Private/SequencerDisplayNode.h
@@ -237,6 +237,15 @@ public:
 
 	/** Pins this node, forcing it to the top of the sequencer. Affects visibility */
 	void PinNode();
+
+	/** Initialize this node with expansion states and virtual offsets */
+	void Initialize(float InVirtualTop, float InVirtualBottom);
+
+	/** @return this node's virtual offset from the top of the tree, irrespective of expansion states */
+	float GetVirtualTop() const { return VirtualTop; }
+	
+	/** @return this node's virtual offset plus its virtual height, irrespective of expansion states */
+	float GetVirtualBottom() const { return VirtualBottom; }
 	
 protected:
 	/**
@@ -272,9 +281,13 @@ protected:
 	
 	/** Whether this node has visible children, based on cached shot filtering visibility only */
 	bool HasVisibleChildren() const;
-public:
+
+protected:
+	/** The virtual offset of this item from the top of the tree, irrespective of expansion states */
 	float VirtualTop;
+	/** The virtual offset + virtual heightof this item, irrespective of expansion states */
 	float VirtualBottom;
+
 protected:
 	/** The parent of this node*/
 	TWeakPtr< FSequencerDisplayNode > ParentNode;
diff --git a/Engine/Source/Editor/Sequencer/Private/SequencerNodeTree.cpp b/Engine/Source/Editor/Sequencer/Private/SequencerNodeTree.cpp
--- a/Engine/Source/Editor/Sequencer/Private/SequencerNodeTree.cpp
+++ b/Engine/Source/Editor/Sequencer/Private/SequencerNodeTree.cpp
@@ -116,13 +116,15 @@ void FSequencerNodeTree::Update()
 	// Add all other nodes after the shot track
 	RootNodes.Append( NewRootNodes );
 
+	// Set up virtual offsets, and expansion states
 	float VerticalOffset = 0.f;
 	for (auto& Node : RootNodes)
 	{
 		Node->Traverse_ParentFirst([&](FSequencerDisplayNode& InNode){
-			InNode.VirtualTop = VerticalOffset;
+
+			float VerticalTop = VerticalOffset;
 			VerticalOffset += InNode.GetNodeHeight() + InNode.GetNodePadding().Combined();
-			InNode.VirtualBottom = VerticalOffset;
+			InNode.Initialize(VerticalTop, VerticalOffset);
 			return true;
 		});
 
@@ -249,16 +251,7 @@ void FSequencerNodeTree::SaveExpansionState( const FSequencerDisplayNode& Node,
 
 	FMovieSceneEditorData& EditorData = MovieScene->GetEditorData();
 
-	if( bExpanded )
-	{
-		EditorData.CollapsedSequencerNodes.Remove( Node.GetPathName() );
-	}
-	else
-	{
-		// Collapsed nodes are stored instead of expanded so that new nodes are expanded by default
-		EditorData.CollapsedSequencerNodes.AddUnique( Node.GetPathName() ) ;
-	}
-
+	EditorData.ExpansionStates.Add( Node.GetPathName(), FMovieSceneExpansionState(bExpanded) );
 }
 
 bool FSequencerNodeTree::GetSavedExpansionState( const FSequencerDisplayNode& Node ) const
@@ -266,11 +259,15 @@ bool FSequencerNodeTree::GetSavedExpansionState( const FSequencerDisplayNode& No
 	UMovieScene* MovieScene = Sequencer.GetFocusedMovieScene();
 
 	FMovieSceneEditorData& EditorData = MovieScene->GetEditorData();
+	FMovieSceneExpansionState* ExpansionState = EditorData.ExpansionStates.Find( Node.GetPathName() );
 
-	// Collapsed nodes are stored instead of expanded so that new nodes are expanded by default
-	bool bCollapsed = EditorData.CollapsedSequencerNodes.Contains( Node.GetPathName() );
+	return ExpansionState ? ExpansionState->bExpanded : GetDefaultExpansionState(Node);
+}
 
-	return !bCollapsed;
+bool FSequencerNodeTree::GetDefaultExpansionState( const FSequencerDisplayNode& Node ) const
+{
+	// For now, all types except categories and key areas are expanded by default
+	return Node.GetType() != ESequencerNode::Category && Node.GetType() != ESequencerNode::KeyArea;
 }
 
 bool FSequencerNodeTree::IsNodeFiltered( const TSharedRef<const FSequencerDisplayNode> Node ) const
diff --git a/Engine/Source/Editor/Sequencer/Private/SequencerNodeTree.h b/Engine/Source/Editor/Sequencer/Private/SequencerNodeTree.h
--- a/Engine/Source/Editor/Sequencer/Private/SequencerNodeTree.h
+++ b/Engine/Source/Editor/Sequencer/Private/SequencerNodeTree.h
@@ -68,6 +68,13 @@ public:
 	 * @return true if the node should be expanded, false otherwise	
 	 */
 	bool GetSavedExpansionState( const FSequencerDisplayNode& Node ) const;
+
+	/**
+	 * Get the default expansion state for the specified node, where its state has not yet been saved
+	 *
+	 * @return true if the node is to be expanded, false otherwise
+	 */
+	bool GetDefaultExpansionState( const FSequencerDisplayNode& Node ) const;
 	
 	/** Updates nodes visibility for use when shot filters change at all */
 	void UpdateCachedVisibilityBasedOnShotFiltersChanged();
diff --git a/Engine/Source/Runtime/MovieScene/Public/MovieScene.h b/Engine/Source/Runtime/MovieScene/Public/MovieScene.h
--- a/Engine/Source/Runtime/MovieScene/Public/MovieScene.h
+++ b/Engine/Source/Runtime/MovieScene/Public/MovieScene.h
@@ -16,6 +16,18 @@ class UMovieSceneTrack;
 MOVIESCENE_API DECLARE_LOG_CATEGORY_EXTERN(LogSequencerRuntime, Log, All);
 
 
+/** @todo: remove this type when support for intrinsics on TMap values is added? */
+USTRUCT()
+struct FMovieSceneExpansionState
+{
+	GENERATED_BODY()
+
+	FMovieSceneExpansionState(bool bInExpanded = true) : bExpanded(bInExpanded) {}
+
+	UPROPERTY()
+	bool bExpanded;
+};
+
 /**
  * Editor only data that needs to be saved between sessions for editing but has no runtime purpose
  */
@@ -24,9 +36,9 @@ struct FMovieSceneEditorData
 {
 	GENERATED_USTRUCT_BODY()
 
-	/** List of collapsed sequencer nodes.  We store collapsed instead of expanded so that new nodes with no saved state are expanded by default */
+	/** Map of node path -> expansion state. */
 	UPROPERTY()
-	TArray<FString> CollapsedSequencerNodes;
+	TMap<FString, FMovieSceneExpansionState> ExpansionStates;
 };
 
 
