Sequencer: Partial/complete keys are now represented on collapsed parent nodes

diff --git a/Engine/Source/Editor/EditorStyle/Private/SlateEditorStyle.cpp b/Engine/Source/Editor/EditorStyle/Private/SlateEditorStyle.cpp
--- a/Engine/Source/Editor/EditorStyle/Private/SlateEditorStyle.cpp
+++ b/Engine/Source/Editor/EditorStyle/Private/SlateEditorStyle.cpp
@@ -1582,6 +1582,7 @@ void FSlateEditorStyle::FStyle::SetupGeneralStyles()
 		Set( "Sequencer.Timeline.NotifyAlignmentMarker", new IMAGE_BRUSH( "Sequencer/NotifyAlignmentMarker", FVector2D(10,19) ) );
 		Set( "Sequencer.NotificationImage_AddedPlayMovieSceneEvent", new IMAGE_BRUSH( "Old/Checkbox_checked", Icon16x16 ) );
 		Set( "Sequencer.Key", new IMAGE_BRUSH( "Common/Key", FVector2D(7.0f,12.0f) ) );
+		Set( "Sequencer.PartialKey", new IMAGE_BRUSH( "Sequencer/PartialKey", FVector2D(11.f, 11.f) ) );
 		Set( "Sequencer.GenericSection.Background", new BOX_BRUSH( "Sequencer/GenericSectionBackground", FMargin(5.f/16.f) ) );
 		Set( "Sequencer.SectionGripLeft", new BOX_BRUSH( "Sequencer/SectionGripLeft", FMargin(5.f/16.f) ) );
 		Set( "Sequencer.SectionGripRight", new BOX_BRUSH( "Sequencer/SectionGripRight", FMargin(5.f/16.f) ) );
diff --git a/Engine/Source/Editor/Sequencer/Private/GroupedKeyArea.cpp b/Engine/Source/Editor/Sequencer/Private/GroupedKeyArea.cpp
new file mode 100644
--- /dev/null
+++ b/Engine/Source/Editor/Sequencer/Private/GroupedKeyArea.cpp
@@ -0,0 +1,274 @@
+// Copyright 1998-2015 Epic Games, Inc. All Rights Reserved.
+
+#include "SequencerPrivatePCH.h"
+#include "GroupedKeyArea.h"
+#include "MovieSceneSection.h"
+
+FIndexKey::FIndexKey(FString InNodePath, UMovieSceneSection* InSection)
+	: NodePath(MoveTemp(InNodePath))
+	, Section(InSection)
+	, CachedHash(HashCombine(GetTypeHash(NodePath), GetTypeHash(InSection)))
+{}
+
+namespace 
+{
+	/** Structure that defines the index for a particular section */
+	struct FIndexEntry
+	{
+		TMap<FKeyHandle, int32> HandleToGroup;
+		TArray<FKeyHandle> GroupHandles;
+		TArray<float> RepresentativeTimes;
+	};
+
+	/** A persistent index is required to ensure that generated key handles are maintained for the lifetime of specific display nodes */
+	TMap<FIndexKey, FIndexEntry> GlobalIndex;
+}
+
+
+FGroupedKeyArea::FGroupedKeyArea(FSequencerDisplayNode& InNode, int32 InSectionIndex)
+	: Section(nullptr)
+	, IndexKey(FString(), nullptr)
+{
+	TArray<TSharedRef<FSectionKeyAreaNode>> AllKeyAreaNodes;
+	AllKeyAreaNodes.Reserve(36);
+	InNode.GetChildKeyAreaNodesRecursively(AllKeyAreaNodes);
+
+	for (const auto& Node : AllKeyAreaNodes)
+	{
+		TSharedRef<IKeyArea> KeyArea = Node->GetKeyArea(InSectionIndex);
+
+		const int32 KeyAreaIndex = KeyAreas.Num();
+		KeyAreas.Add(KeyArea);
+
+		{
+			auto* OwningSection = KeyArea->GetOwningSection();
+			// Ensure they all belong to the same section
+			ensure(!Section || Section == OwningSection);
+			Section = OwningSection;
+		}
+
+		for (const FKeyHandle& KeyHandle : KeyArea->GetUnsortedKeyHandles())
+		{
+			Groups.Emplace(KeyArea->GetKeyTime(KeyHandle), KeyAreaIndex, KeyHandle);
+		}
+	}
+
+	Groups.Sort([](const FKeyGrouping& A, const FKeyGrouping& B){
+		return A.RepresentativeTime < B.RepresentativeTime;
+	});
+
+	// Remove duplicates
+	for (int32 Index = 0; Index < Groups.Num(); ++Index)
+	{
+		const float CurrentTime = Groups[Index].RepresentativeTime;
+
+		int32 NumToMerge = 0;
+		for (int32 DuplIndex = Index + 1; DuplIndex < Groups.Num(); ++DuplIndex)
+		{
+			if (!FMath::IsNearlyEqual(CurrentTime, Groups[DuplIndex].RepresentativeTime))
+			{
+				break;
+			}
+			++NumToMerge;
+		}
+
+		if (NumToMerge)
+		{
+			const int32 Start = Index + 1, End = Start + NumToMerge;
+			for (int32 MergeIndex = Start; MergeIndex < End; ++MergeIndex)
+			{
+				Groups[Index].Keys.Append(MoveTemp(Groups[MergeIndex].Keys));
+			}
+
+			Groups.RemoveAt(Start, NumToMerge, false);
+		}
+	}
+
+	if (Section)
+	{
+		IndexKey = FIndexKey(InNode.GetPathName(), Section);
+		UpdateIndex();
+	}
+}
+
+void FGroupedKeyArea::UpdateIndex() const
+{
+	auto& IndexEntry = GlobalIndex.FindOrAdd(IndexKey);
+
+	TArray<FKeyHandle> NewKeyHandles;
+	TArray<float> NewRepresentativeTimes;
+
+	IndexEntry.HandleToGroup.Reset();
+
+	for (int32 GroupIndex = 0; GroupIndex < Groups.Num(); ++GroupIndex)
+	{
+		float RepresentativeTime = Groups[GroupIndex].RepresentativeTime;
+
+		// Find a key handle we can recycle
+		int32 RecycledIndex = IndexEntry.RepresentativeTimes.IndexOfByPredicate([&](float Time){
+			// Must be an *exact* match to recycle
+			return Time == RepresentativeTime;
+		});
+		
+		if (RecycledIndex != INDEX_NONE)
+		{
+			NewKeyHandles.Add(IndexEntry.GroupHandles[RecycledIndex]);
+			NewRepresentativeTimes.Add(IndexEntry.RepresentativeTimes[RecycledIndex]);
+
+			IndexEntry.GroupHandles.RemoveAt(RecycledIndex, 1, false);
+			IndexEntry.RepresentativeTimes.RemoveAt(RecycledIndex, 1, false);
+		}
+		else
+		{
+			NewKeyHandles.Add(FKeyHandle());
+			NewRepresentativeTimes.Add(RepresentativeTime);
+		}
+
+		IndexEntry.HandleToGroup.Add(NewKeyHandles.Last(), GroupIndex);
+	}
+
+	IndexEntry.GroupHandles = MoveTemp(NewKeyHandles);
+	IndexEntry.RepresentativeTimes = MoveTemp(NewRepresentativeTimes);
+}
+
+const FKeyGrouping* FGroupedKeyArea::FindGroup(FKeyHandle InHandle) const
+{
+	if (auto* IndexEntry = GlobalIndex.Find(IndexKey))
+	{
+		if (int32* GroupIndex = IndexEntry->HandleToGroup.Find(InHandle))
+		{
+			return Groups.IsValidIndex(*GroupIndex) ? &Groups[*GroupIndex] : nullptr;
+		}
+	}
+	return nullptr;
+}
+
+FKeyGrouping* FGroupedKeyArea::FindGroup(FKeyHandle InHandle)
+{
+	if (auto* IndexEntry = GlobalIndex.Find(IndexKey))
+	{
+		if (int32* GroupIndex = IndexEntry->HandleToGroup.Find(InHandle))
+		{
+			return Groups.IsValidIndex(*GroupIndex) ? &Groups[*GroupIndex] : nullptr;
+		}
+	}
+	return nullptr;
+}
+
+FLinearColor FGroupedKeyArea::GetKeyTint(FKeyHandle InHandle) const
+{
+	// Everything is untinted for now
+	return FLinearColor::White;
+}
+
+const FSlateBrush* FGroupedKeyArea::GetBrush(FKeyHandle InHandle) const
+{
+	static const FSlateBrush* KeyBrush = FEditorStyle::GetBrush("Sequencer.Key");
+	static const FSlateBrush* PartialKeyBrush = FEditorStyle::GetBrush("Sequencer.PartialKey");
+
+	const auto* Group = FindGroup(InHandle);
+
+	bool bPartialKey = false;
+
+	// Ensure that each key area is represented at least once for it to be considered a 'complete key'
+	for (int32 AreaIndex = 0; Group && AreaIndex < KeyAreas.Num(); ++AreaIndex)
+	{
+		if (!Group->Keys.ContainsByPredicate([&](const FKeyGrouping::FKeyIndex& Key){ return Key.AreaIndex == AreaIndex; }))
+		{
+			bPartialKey = true;
+			break;
+		}
+	}
+
+	return bPartialKey ? PartialKeyBrush : KeyBrush;
+}
+
+TArray<FKeyHandle> FGroupedKeyArea::GetUnsortedKeyHandles() const
+{
+	TArray<FKeyHandle> Array;
+	if (auto* IndexEntry = GlobalIndex.Find(IndexKey))
+	{
+		IndexEntry->HandleToGroup.GenerateKeyArray(Array);
+	}
+	return Array;
+}
+
+float FGroupedKeyArea::GetKeyTime( FKeyHandle KeyHandle ) const
+{
+	auto* Group = FindGroup(KeyHandle);
+	return Group ? Group->RepresentativeTime : 0.f;
+}
+
+FKeyHandle FGroupedKeyArea::MoveKey( FKeyHandle KeyHandle, float DeltaPosition )
+{
+	int32* GroupIndex = nullptr;
+
+	auto* IndexEntry = GlobalIndex.Find(IndexKey);
+	if (IndexEntry)
+	{
+		GroupIndex = IndexEntry->HandleToGroup.Find(KeyHandle);
+	}
+
+	if (!GroupIndex || !Groups.IsValidIndex(*GroupIndex))
+	{
+		return KeyHandle;
+	}
+
+	bool bIsTimeUpdated = false;
+
+	FKeyGrouping& Group = Groups[*GroupIndex];
+
+	// Move all the keys
+	for (auto& Key : Group.Keys)
+	{
+		Key.KeyHandle = KeyAreas[Key.AreaIndex]->MoveKey(Key.KeyHandle, DeltaPosition);
+
+		const float KeyTime = KeyAreas[Key.AreaIndex]->GetKeyTime(Key.KeyHandle);
+		Group.RepresentativeTime = bIsTimeUpdated ? FMath::Min(Group.RepresentativeTime, KeyTime) : KeyTime;
+		bIsTimeUpdated = true;
+	}
+
+	// Update the representative time to the smallest of all the keys (so it will deterministically get the same key handle on regeneration)
+	IndexEntry->RepresentativeTimes[*GroupIndex] = Group.RepresentativeTime;
+
+	return KeyHandle;
+}
+
+void FGroupedKeyArea::DeleteKey(FKeyHandle KeyHandle)
+{
+	if (auto* Group = FindGroup(KeyHandle))
+	{
+		for (auto& Key : Group->Keys)
+		{
+			KeyAreas[Key.AreaIndex]->DeleteKey(Key.KeyHandle);
+		}
+	}
+}
+
+void FGroupedKeyArea::AddKeyUnique(float Time)
+{
+	for (auto& Area : KeyAreas)
+	{
+		Area->AddKeyUnique(Time);
+	}
+}
+
+FRichCurve* FGroupedKeyArea::GetRichCurve()
+{
+	return nullptr;
+}
+
+UMovieSceneSection* FGroupedKeyArea::GetOwningSection()
+{
+	return Section;
+}
+
+bool FGroupedKeyArea::CanCreateKeyEditor()
+{
+	return false;
+}
+
+TSharedRef<SWidget> FGroupedKeyArea::CreateKeyEditor(ISequencer* Sequencer)
+{
+	return SNullWidget::NullWidget;
+}
\ No newline at end of file
diff --git a/Engine/Source/Editor/Sequencer/Private/GroupedKeyArea.h b/Engine/Source/Editor/Sequencer/Private/GroupedKeyArea.h
new file mode 100644
--- /dev/null
+++ b/Engine/Source/Editor/Sequencer/Private/GroupedKeyArea.h
@@ -0,0 +1,111 @@
+// Copyright 1998-2015 Epic Games, Inc. All Rights Reserved.
+
+#pragma once
+
+#include "IKeyArea.h"
+
+class FSequencerDisplayNode;
+class UMovieSceneSection;
+
+/** Keyable structure used to identify a particular FGroupedKeyArea */
+struct FIndexKey
+{
+	FIndexKey(FString InNodePath, UMovieSceneSection* InSection);
+
+	friend bool operator==(const FIndexKey& A, const FIndexKey& B)
+	{
+		return A.CachedHash == B.CachedHash && A.NodePath == B.NodePath && A.Section == B.Section;
+	}
+
+	friend uint32 GetTypeHash( const FIndexKey& IndexKey )
+	{
+		return IndexKey.CachedHash;
+	}
+
+private:
+	/** Cached path of the node that we relate to */
+	FString NodePath;
+	/** The movie scene section that we relate to */
+	TWeakObjectPtr<UMovieSceneSection> Section;
+	/** Cached hash for efficiency (to avoid having to re-hash strings all the time) */
+	uint32 CachedHash;
+};
+
+/** Structure that represents a number of keys */
+struct FKeyGrouping
+{
+	/** Structure to define a single key in this group - persisting of a key handle, and an index for its parent key area */
+	struct FKeyIndex
+	{
+		FKeyIndex(int32 InAreaIndex, FKeyHandle InKeyHandle) : AreaIndex(InAreaIndex), KeyHandle(InKeyHandle) {}
+
+		/** Index into an externally owned array of IKeyArea ptrs. Kept this way to avoid having a shared ptr per key */
+		int32 AreaIndex;
+		/** The actual key handle for this key, obtained from its owner IKeyArea */
+		FKeyHandle KeyHandle;
+	};
+
+	/** Construct this group with a single key handle */
+	FKeyGrouping(float InRepresentativeTime, int32 AreaIndex, FKeyHandle KeyHandle)
+		: RepresentativeTime(InRepresentativeTime)
+	{
+		Keys.Emplace(AreaIndex, KeyHandle);
+	}
+
+	/** Move construction/assignment */
+	FKeyGrouping(FKeyGrouping&& In) : RepresentativeTime(In.RepresentativeTime), Keys(MoveTemp(In.Keys)) {}
+	FKeyGrouping& operator=(FKeyGrouping&& In) { RepresentativeTime = In.RepresentativeTime; Keys = MoveTemp(In.Keys); return *this; }
+
+	/** The representative time at which all keys within this group are located */
+	float RepresentativeTime;
+	/** Array of keys contained within this group */
+	TArray<FKeyIndex> Keys;
+};
+
+/** A key area that represents multiple, distinct key areas */
+class FGroupedKeyArea : public IKeyArea
+{
+public:
+	/** Construct the area from a display node */
+	FGroupedKeyArea(FSequencerDisplayNode& InNode, int32 InSectionIndex);
+	
+	/* IKeyArea interface */
+	virtual TArray<FKeyHandle>		GetUnsortedKeyHandles() const override;
+	virtual float					GetKeyTime( FKeyHandle KeyHandle ) const override;
+	virtual FKeyHandle				MoveKey( FKeyHandle KeyHandle, float DeltaPosition ) override;
+	virtual void					DeleteKey(FKeyHandle KeyHandle) override;
+	virtual void					AddKeyUnique(float Time) override;
+	virtual FRichCurve*				GetRichCurve() override;
+	virtual UMovieSceneSection*		GetOwningSection() override;
+	virtual bool					CanCreateKeyEditor() override;
+	virtual TSharedRef<SWidget>		CreateKeyEditor(ISequencer* Sequencer) override;
+	/* End IKeyArea interface */
+
+	/** Get the desired tint for the specified key handle */
+	FLinearColor GetKeyTint( FKeyHandle InHandle ) const;
+
+	/** Get the brush name for the specified key handle */
+	const FSlateBrush* GetBrush(FKeyHandle InHandle) const;
+
+protected:
+
+	/** Update the persistent index with our new key handles */
+	void UpdateIndex() const;
+
+	/** Helper function to find the group of keys relating to the specified handle */
+	const FKeyGrouping* FindGroup(FKeyHandle InHandle) const;
+
+	/** Helper function to find the group of keys relating to the specified handle */
+	FKeyGrouping* FindGroup(FKeyHandle InHandle);
+
+protected:
+	
+	/** Pointer to the section to which this key area relates */
+	UMovieSceneSection* Section;
+	/** Array of (child) key areas that we are reflecting */
+	TArray<TSharedRef<IKeyArea>> KeyAreas;
+	/** Generated array of groups of keys harvested from the above array */
+	TArray<FKeyGrouping> Groups;
+	/** Key into our index */
+	FIndexKey IndexKey;
+};
\ No newline at end of file
diff --git a/Engine/Source/Editor/Sequencer/Private/KeyAreaLayout.cpp b/Engine/Source/Editor/Sequencer/Private/KeyAreaLayout.cpp
new file mode 100644
--- /dev/null
+++ b/Engine/Source/Editor/Sequencer/Private/KeyAreaLayout.cpp
@@ -0,0 +1,112 @@
+
+#include "SequencerPrivatePCH.h"
+#include "KeyAreaLayout.h"
+
+FKeyAreaLayoutElement FKeyAreaLayoutElement::FromGroup(const TSharedRef<FGroupedKeyArea>& InKeyAreaGroup, float InOffset, TOptional<float> InHeight)
+{
+	FKeyAreaLayoutElement Tmp;
+	Tmp.Type = Group;
+	Tmp.KeyArea = InKeyAreaGroup;
+	Tmp.LocalOffset = InOffset;
+	Tmp.Height = InHeight;
+	return Tmp;
+}
+
+FKeyAreaLayoutElement FKeyAreaLayoutElement::FromKeyAreaNode(const TSharedRef<FSectionKeyAreaNode>& InKeyAreaNode, int32 SectionIndex, float InOffset)
+{
+	FKeyAreaLayoutElement Tmp;
+	Tmp.Type = Single;
+	Tmp.KeyArea = InKeyAreaNode->GetKeyArea(SectionIndex);
+	Tmp.LocalOffset = InOffset;
+
+	if (!InKeyAreaNode->IsTopLevel())
+	{
+		Tmp.Height = InKeyAreaNode->GetNodeHeight();
+	}
+	return Tmp;
+}
+
+FKeyAreaLayoutElement::EType FKeyAreaLayoutElement::GetType() const
+{
+	return Type;
+}
+
+float FKeyAreaLayoutElement::GetOffset() const
+{
+	return LocalOffset;
+}
+
+float FKeyAreaLayoutElement::GetHeight(const FGeometry& InParentGeometry) const
+{
+	return Height.IsSet() ? Height.GetValue() : InParentGeometry.GetDrawSize().Y;
+}
+
+TSharedRef<IKeyArea> FKeyAreaLayoutElement::GetKeyArea() const
+{
+	return KeyArea.ToSharedRef();
+}
+
+FKeyAreaLayout::FKeyAreaLayout(FSequencerDisplayNode& InNode, int32 InSectionIndex)
+{
+	TotalHeight = 0.f;
+
+	float LastPadding = 0.f;
+
+	auto SetupKeyArea = [&](FSequencerDisplayNode& Node){
+
+		if (Node.GetType() == ESequencerNode::KeyArea)
+		{
+			Elements.Add(FKeyAreaLayoutElement::FromKeyAreaNode(StaticCastSharedRef<FSectionKeyAreaNode>(Node.AsShared()), InSectionIndex, TotalHeight));
+		}
+		else if (!Node.IsExpanded())
+		{
+			Elements.Add(FKeyAreaLayoutElement::FromGroup(Node.UpdateKeyGrouping(InSectionIndex), TotalHeight, Node.GetNodeHeight()));
+		}
+
+	};
+
+
+	// First, layout the parent
+	{
+		FSequencerDisplayNode* LayoutNode = &InNode;
+		if (InNode.GetType() == ESequencerNode::Track)
+		{
+			TSharedPtr<FSequencerDisplayNode> KeyAreaNode = static_cast<FTrackNode&>(InNode).GetTopLevelKeyNode();
+			if (KeyAreaNode.IsValid())
+			{
+				LayoutNode = KeyAreaNode.Get();
+			}
+		}
+
+		SetupKeyArea(*LayoutNode);
+
+		LastPadding = LayoutNode->GetNodePadding().Bottom;
+		TotalHeight += LayoutNode->GetNodeHeight() + LastPadding;
+	}
+
+	// Then any children
+	InNode.TraverseVisible_ParentFirst([&](FSequencerDisplayNode& Node){
+		
+		TotalHeight += Node.GetNodePadding().Top;
+
+		SetupKeyArea(Node);
+
+		LastPadding = Node.GetNodePadding().Bottom;
+		TotalHeight += Node.GetNodeHeight() + LastPadding;
+		return true;
+
+	}, false);
+
+	// Remove the padding from the bottom
+	TotalHeight -= LastPadding;
+}
+
+const TArray<FKeyAreaLayoutElement>& FKeyAreaLayout::GetElements() const
+{
+	return Elements;
+}
+
+float FKeyAreaLayout::GetTotalHeight() const
+{
+	return TotalHeight;
+}
\ No newline at end of file
diff --git a/Engine/Source/Editor/Sequencer/Private/KeyAreaLayout.h b/Engine/Source/Editor/Sequencer/Private/KeyAreaLayout.h
new file mode 100644
--- /dev/null
+++ b/Engine/Source/Editor/Sequencer/Private/KeyAreaLayout.h
@@ -0,0 +1,65 @@
+
+#pragma once
+
+#include "GroupedKeyArea.h"
+
+/** A layout element specifying the geometry required to render a key area */
+struct FKeyAreaLayoutElement
+{
+	/** Whether this layout element pertains to one or multiple key areas */
+	enum EType { Single, Group };
+
+	/** Construct this element from a grouped key area */
+	static FKeyAreaLayoutElement FromGroup(const TSharedRef<FGroupedKeyArea>& InKeyAreaGroup, float InOffset, TOptional<float> InHeight = TOptional<float>());
+
+	/** Construct this element from a single Key area node */
+	static FKeyAreaLayoutElement FromKeyAreaNode(const TSharedRef<FSectionKeyAreaNode>& InKeyAreaNode, int32 SectionIndex, float InOffset);
+
+	/** Retrieve the type of this layout element */
+	EType GetType() const;
+
+	/** Retrieve vertical offset from the top of this element's parent */
+	float GetOffset() const;
+
+	/** Retrieve the desired height of this element based on the specified parent geometry */
+	float GetHeight(const FGeometry& InParentGeometry) const;
+
+	/** Access the key area that this layout element was generated for */
+	TSharedRef<IKeyArea> GetKeyArea() const;
+
+private:
+
+	/** Pointer to the key area that we were generated from */
+	TSharedPtr<IKeyArea> KeyArea;
+
+	/** The type of this layout element */
+	EType Type;
+
+	/** The vertical offset from the top of the element's parent */
+	float LocalOffset;
+
+	/** Explicit height, or unset to fill the parent geometry */
+	TOptional<float> Height;
+};
+
+/** Class used for generating, and caching the layout geometry for a given display node's key areas */
+class FKeyAreaLayout
+{
+public:
+
+	/** Constructor that takes a display node, and the index of the section to layout */
+	FKeyAreaLayout(FSequencerDisplayNode& InNode, int32 InSectionIndex);
+
+	/** Get all layout elements that we generated */
+	const TArray<FKeyAreaLayoutElement>& GetElements() const;
+	
+	/** Get the desired total height of this layout */
+	float GetTotalHeight() const;
+
+private:
+	/** Array of layout elements that we generated */
+	TArray<FKeyAreaLayoutElement> Elements;
+
+	/** The desired total height of this layout */
+	float TotalHeight;
+};
\ No newline at end of file
diff --git a/Engine/Source/Editor/Sequencer/Private/SSection.cpp b/Engine/Source/Editor/Sequencer/Private/SSection.cpp
--- a/Engine/Source/Editor/Sequencer/Private/SSection.cpp
+++ b/Engine/Source/Editor/Sequencer/Private/SSection.cpp
@@ -10,7 +10,7 @@
 #include "MovieSceneShotSection.h"
 #include "CommonMovieSceneTools.h"
 
-
+FThreadSafeCounter SSection::LayoutRegenerationLock;
 
 void SSection::Construct( const FArguments& InArgs, TSharedRef<FTrackNode> SectionNode, int32 InSectionIndex )
 {
@@ -19,7 +19,7 @@ void SSection::Construct( const FArguments& InArgs, TSharedRef<FTrackNode> Secti
 	SectionIndex = InSectionIndex;
 	ParentSectionArea = SectionNode;
 	SectionInterface = SectionNode->GetSections()[InSectionIndex];
-	CachedKeyAreaHeight = 0.f;
+	Layout = FKeyAreaLayout(*SectionNode, InSectionIndex);
 
 	ChildSlot
 	[
@@ -29,70 +29,23 @@ void SSection::Construct( const FArguments& InArgs, TSharedRef<FTrackNode> Secti
 
 FVector2D SSection::ComputeDesiredSize(float) const
 {
-	return FVector2D(100, CachedKeyAreaHeight);
+	return FVector2D(100, Layout->GetTotalHeight());
 }
 
-void SSection::GetKeyAreas( const TSharedPtr<FTrackNode>& SectionAreaNode, TArray<FKeyAreaElement>& OutKeyAreas, float* TotalHeight ) const
+FGeometry SSection::GetKeyAreaGeometry( const FKeyAreaLayoutElement& KeyArea, const FGeometry& SectionGeometry ) const
 {
-	*TotalHeight = 0.f;
-
-	float LastPadding = 0.f;
-
-	// First, layout the parent
-	{
-		TSharedPtr<FSequencerDisplayNode> LayoutNode = SectionAreaNode;
-
-		TSharedPtr<FSectionKeyAreaNode> SectionKeyNode = SectionAreaNode->GetTopLevelKeyNode();
-		if( SectionKeyNode.IsValid() )
-		{
-			LayoutNode = SectionKeyNode;
-			OutKeyAreas.Add(FKeyAreaElement( *SectionKeyNode, *TotalHeight ));
-		}
-
-		LastPadding = LayoutNode->GetNodePadding().Bottom;
-		*TotalHeight += LayoutNode->GetNodeHeight() + LastPadding;
-	}
-
-	// Then any children
-	SectionAreaNode->TraverseVisible_ParentFirst([&](FSequencerDisplayNode& Node){
-		
-		*TotalHeight += Node.GetNodePadding().Top;
-
-		if( Node.GetType() == ESequencerNode::KeyArea )
-		{
-			OutKeyAreas.Add(FKeyAreaElement( static_cast<FSectionKeyAreaNode&>(Node), *TotalHeight ));
-		}
-
-		LastPadding = Node.GetNodePadding().Bottom;
-		*TotalHeight += Node.GetNodeHeight() + LastPadding;
-		return true;
-
-	}, false);
-
-	// Remove the padding from the bottom
-	*TotalHeight -= LastPadding;
-}
-
-
-FGeometry SSection::GetKeyAreaGeometry( const struct FKeyAreaElement& KeyArea, const FGeometry& SectionGeometry ) const
-{
-	// Get the height of the key area node.  If the key area is top level then it is part of the section (and the same height ) and doesn't take up extra space
-	float KeyAreaHeight = KeyArea.KeyAreaNode.IsTopLevel() ? SectionGeometry.GetDrawSize().Y : KeyArea.KeyAreaNode.GetNodeHeight();
-
 	// Compute the geometry for the key area
-	return SectionGeometry.MakeChild( FVector2D( 0, KeyArea.HeightOffset ), FVector2D( SectionGeometry.Size.X, KeyAreaHeight ) );
+	return SectionGeometry.MakeChild( FVector2D( 0, KeyArea.GetOffset() ), FVector2D( SectionGeometry.Size.X, KeyArea.GetHeight(SectionGeometry) ) );
 }
 
-
 FSelectedKey SSection::GetKeyUnderMouse( const FVector2D& MousePosition, const FGeometry& AllottedGeometry ) const
 {
 	UMovieSceneSection& Section = *SectionInterface->GetSectionObject();
 
 	// Search every key area until we find the one under the mouse
-	for( int32 KeyAreaIndex = 0; KeyAreaIndex < KeyAreas.Num(); ++KeyAreaIndex )
+	for (const FKeyAreaLayoutElement& Element : Layout->GetElements())
 	{
-		const FKeyAreaElement& Element = KeyAreas[KeyAreaIndex];
-		TSharedRef<IKeyArea> KeyArea = Element.KeyAreaNode.GetKeyArea( SectionIndex ); 
+		TSharedPtr<IKeyArea> KeyArea = Element.GetKeyArea();
 
 		// Compute the current key area geometry
 		FGeometry KeyAreaGeometryPadded = GetKeyAreaGeometry( Element, AllottedGeometry );
@@ -195,6 +148,8 @@ void SSection::CreateDragOperation( const FGeometry& MyGeometry, const FPointerE
 
 	if( bKeysUnderMouse )
 	{
+		// Disable layout regeneration of dynamic key areas
+		LayoutRegenerationLock.Increment();
 		DragOperation = MakeShareable( new FMoveKeys( GetSequencer(), GetSequencer().GetSelection().GetSelectedKeys(), PressedKey ) );
 	}
 	else
@@ -217,7 +172,6 @@ void SSection::CreateDragOperation( const FGeometry& MyGeometry, const FPointerE
 			DragOperation = MakeShareable( new FMoveSection( GetSequencer(), GetSequencer().GetSelection().GetSelectedSections() ) );
 		}
 	}
-	
 }
 
 int32 SSection::OnPaint( const FPaintArgs& Args, const FGeometry& AllottedGeometry, const FSlateRect& MyClippingRect, FSlateWindowElementList& OutDrawElements, int32 LayerId, const FWidgetStyle& InWidgetStyle, bool bParentEnabled ) const
@@ -262,9 +216,6 @@ int32 SSection::OnPaint( const FPaintArgs& Args, const FGeometry& AllottedGeomet
 			);
 	}
 
-	
-
-
 	return LayerId;
 }
 
@@ -294,12 +245,10 @@ void SSection::PaintKeys( const FGeometry& AllottedGeometry, const FSlateRect& M
 	FLinearColor HoveredKeyColor = SelectedKeyColor * FLinearColor(1.5,1.5,1.5,1.0f);
 
 	// Draw all keys in each key area
-	for( int32 KeyAreaIndex = 0; KeyAreaIndex < KeyAreas.Num(); ++KeyAreaIndex )
+	for (const FKeyAreaLayoutElement& Element : Layout->GetElements())
 	{
-		const FKeyAreaElement& Element = KeyAreas[KeyAreaIndex];
-
 		// Get the key area at the same index of the section.  Each section in this widget has the same layout and the same number of key areas
-		const TSharedRef<IKeyArea>& KeyArea = Element.KeyAreaNode.GetKeyArea( SectionIndex );
+		TSharedPtr<IKeyArea> KeyArea = Element.GetKeyArea();
 
 		FGeometry KeyAreaGeometry = GetKeyAreaGeometry( Element, AllottedGeometry );
 
@@ -330,7 +279,16 @@ void SSection::PaintKeys( const FGeometry& AllottedGeometry, const FSlateRect& M
 			// Omit keys which would not be visible
 			if( SectionObject.IsTimeWithinSection( KeyTime ) )
 			{
+				const FSlateBrush* BrushToUse = KeyBrush;
 				FLinearColor KeyColor( 1.0f, 1.0f, 1.0f, 1.0f );
+				FLinearColor KeyTint(1.f, 1.f, 1.f, 1.f);
+
+				if (Element.GetType() == FKeyAreaLayoutElement::Group)
+				{
+					auto Group = StaticCastSharedPtr<FGroupedKeyArea>(KeyArea);
+					KeyTint = Group->GetKeyTint(KeyHandle);
+					BrushToUse = Group->GetBrush(KeyHandle);
+				}
 
 				// Where to start drawing the key (relative to the section)
 				float KeyPosition =  TimeToPixelConverter.TimeToPixel( KeyTime );
@@ -360,6 +318,8 @@ void SSection::PaintKeys( const FGeometry& AllottedGeometry, const FSlateRect& M
 					}
 				}
 
+				KeyColor *= KeyTint;
+
 				// Draw the key
 				FSlateDrawElement::MakeBox(
 					OutDrawElements,
@@ -367,7 +327,7 @@ void SSection::PaintKeys( const FGeometry& AllottedGeometry, const FSlateRect& M
 					bSelected ? KeyLayer+1 : KeyLayer,
 					// Center the key along Y.  Ensure the middle of the key is at the actual key time
 					KeyAreaGeometry.ToPaintGeometry( FVector2D( KeyPosition - FMath::CeilToFloat(SequencerSectionConstants::KeySize.X/2.0f), ((KeyAreaGeometry.Size.Y*.5f)-(SequencerSectionConstants::KeySize.Y*.5f)) ), SequencerSectionConstants::KeySize ),
-					KeyBrush,
+					BrushToUse,
 					MyClippingRect,
 					ESlateDrawEffect::None,
 					KeyColor
@@ -481,11 +441,10 @@ void SSection::Tick( const FGeometry& AllottedGeometry, const double InCurrentTi
 {	
 	if( GetVisibility() == EVisibility::Visible )
 	{
-		CachedKeyAreaHeight = 0.f;
-
-		KeyAreas.Reset();
-		GetKeyAreas( ParentSectionArea, KeyAreas, &CachedKeyAreaHeight );
-
+		if (LayoutRegenerationLock.GetValue() == 0)
+		{
+			Layout = FKeyAreaLayout(*ParentSectionArea, SectionIndex);
+		}
 		SectionInterface->Tick(AllottedGeometry, ParentGeometry, InCurrentTime, InDeltaTime);
 	}
 }
@@ -532,6 +491,7 @@ void SSection::ResetState()
 	DragOperation.Reset();
 	ResetHoveredState();
 	PressedKey = FSelectedKey();
+	LayoutRegenerationLock.Reset();
 }
 
 void SSection::ResetHoveredState()
@@ -546,7 +506,7 @@ void SSection::ResetHoveredState()
 FReply SSection::OnMouseButtonUp( const FGeometry& MyGeometry, const FPointerEvent& MouseEvent )
 {
 	if( bDragging && DragOperation.IsValid() )
-	{
+	{	
 		// If dragging tell the operation we are no longer dragging
 		DragOperation->OnEndDrag(ParentSectionArea);
 	}
diff --git a/Engine/Source/Editor/Sequencer/Private/SSection.h b/Engine/Source/Editor/Sequencer/Private/SSection.h
--- a/Engine/Source/Editor/Sequencer/Private/SSection.h
+++ b/Engine/Source/Editor/Sequencer/Private/SSection.h
@@ -2,6 +2,8 @@
 
 #pragma once
 
+#include "KeyAreaLayout.h"
+
 namespace SequencerSectionConstants
 {
 	/** How far the user has to drag the mouse before we consider the action dragging rather than a click */
@@ -30,25 +32,6 @@ public:
 	virtual FVector2D ComputeDesiredSize(float) const override;
 	
 private:
-	/** Information about how to draw a key area */
-	struct FKeyAreaElement
-	{
-		FKeyAreaElement( FSectionKeyAreaNode& InKeyAreaNode, float InHeightOffset )
-			: KeyAreaNode( InKeyAreaNode )
-			, HeightOffset( InHeightOffset )
-		{}
-
-		FSectionKeyAreaNode& KeyAreaNode;
-		float HeightOffset;
-	};
-
-	/**
-	 * Gets all the visible key areas in this section
-	 *
-	 * @param SectionAreaNode		The section area node containing key areas
-	 * @param OutKeyAreaElements	Populated list of key elements to draw
-	 */
-	void GetKeyAreas( const TSharedPtr<FTrackNode>& SectionAreaNode, TArray<FKeyAreaElement>& OutKeyAreaElements, float* HeightOffset = nullptr ) const;
 
 	/**
 	 * Computes the geometry for a key area
@@ -57,7 +40,7 @@ private:
 	 * @param SectionGeometry	The geometry of the section
 	 * @return The geometry of the key area
 	 */
-	FGeometry GetKeyAreaGeometry( const FKeyAreaElement& KeyArea, const FGeometry& SectionGeometry ) const;
+	FGeometry GetKeyAreaGeometry( const FKeyAreaLayoutElement& KeyArea, const FGeometry& SectionGeometry ) const;
 
 	/**
 	 * Determines the key that is under the mouse
@@ -155,10 +138,8 @@ private:
 	TSharedPtr<ISequencerSection> SectionInterface;
 	/** Section area where this section resides */
 	TSharedPtr<FTrackNode> ParentSectionArea;
-	/** Key areas inside this section */
-	TArray<FKeyAreaElement> KeyAreas;
-	/** Cached height of all the key areas */
-	float CachedKeyAreaHeight;
+	/** Cached layout generated each tick */
+	TOptional<FKeyAreaLayout> Layout;
 	/** Current drag operation if any */
 	TSharedPtr<class FSequencerDragOperation> DragOperation;
 	/** The currently pressed key if any */
@@ -179,6 +160,8 @@ private:
 	bool bRightEdgeHovered;
 	/** Whether or not the right edge of the section is pressed */
 	bool bRightEdgePressed;
+	/** When 0, regeneration of dynamic key layouts is enabled, when non-zero, such behaviour is disabled */
+	static FThreadSafeCounter LayoutRegenerationLock;
 	/** Cached parent geometry to pass down to any section interfaces that need it during tick */
 	FGeometry ParentGeometry;
 };
\ No newline at end of file
diff --git a/Engine/Source/Editor/Sequencer/Private/SSequencerTrackArea.cpp b/Engine/Source/Editor/Sequencer/Private/SSequencerTrackArea.cpp
--- a/Engine/Source/Editor/Sequencer/Private/SSequencerTrackArea.cpp
+++ b/Engine/Source/Editor/Sequencer/Private/SSequencerTrackArea.cpp
@@ -9,6 +9,7 @@
 #include "SSequencerTreeView.h"
 #include "IKeyArea.h"
 #include "ISequencerSection.h"
+#include "SSection.h"
 
 
 struct FMarqueeSelectData
@@ -48,6 +49,141 @@ struct FMarqueeSelectData
 	FVector2D CurrentPosition;
 };
 
+
+struct IMarqueeVisitor
+{
+	virtual void VisitKey(FKeyHandle KeyHandle, float KeyTime, const TSharedPtr<IKeyArea>& KeyArea, UMovieSceneSection* Section) const { }
+	virtual void VisitSection(FKeyHandle KeyHandle, float KeyTime, const TSharedPtr<IKeyArea>& KeyArea, UMovieSceneSection* Section) const { }
+};
+
+/** Struct used to iterate a marquee range with a user-supplied visitor */
+/* @todo: This could probably live in a header so it can be used elsewhere */
+struct FMarqueeRange
+{
+	/** Construction from the range itself */
+	FMarqueeRange(FVector2D InTopLeft, FVector2D InBottomRight, FVector2D InVirtualKeySize)
+		: TopLeft(InTopLeft), BottomRight(InBottomRight), VirtualKeySize(InVirtualKeySize)
+	{}
+
+	/** Visit the specified nodes (recursively) with this range and a user-supplied visitor */
+	void Visit(const IMarqueeVisitor& Visitor, const TArray< TSharedRef<FSequencerDisplayNode> >& Nodes);
+
+private:
+
+	/** Handle visitation of a particular node */
+	void HandleNode(const IMarqueeVisitor& Visitor, FSequencerDisplayNode& InNode);
+	/** Handle visitation of a particular node, along with a set of sections */
+	void HandleNode(const IMarqueeVisitor& Visitor, FSequencerDisplayNode& InNode, const TArray<TSharedRef<ISequencerSection>>& InSections);
+	/** Handle visitation of a key area node */
+	void HandleKeyAreaNode(const IMarqueeVisitor& Visitor, FSectionKeyAreaNode& InNode, const TArray<TSharedRef<ISequencerSection>>& InSections);
+	/** Handle visitation of a key area */
+	void HandleKeyArea(const IMarqueeVisitor& Visitor, const TSharedPtr<IKeyArea>& KeyArea, UMovieSceneSection* Section);
+	
+	/** The bounds of the range */
+	FVector2D TopLeft, BottomRight;
+
+	/** Key size in virtual space */
+	FVector2D VirtualKeySize;
+};
+
+void FMarqueeRange::Visit(const IMarqueeVisitor& Visitor, const TArray< TSharedRef<FSequencerDisplayNode> >& Nodes)
+{
+	for (auto& Child : Nodes)
+	{
+		if (!Child->IsHidden())
+		{
+			HandleNode(Visitor, *Child);
+		}
+	}
+}
+
+void FMarqueeRange::HandleNode(const IMarqueeVisitor& Visitor, FSequencerDisplayNode& InNode)
+{
+	if (InNode.GetType() == ESequencerNode::Track)
+	{
+		HandleNode(Visitor, InNode, static_cast<FTrackNode&>(InNode).GetSections());
+	}
+
+	if (InNode.IsExpanded())
+	{
+		for (auto& Child : InNode.GetChildNodes())
+		{
+			if (!Child->IsHidden())
+			{
+				HandleNode(Visitor, *Child);
+			}
+		}
+	}
+}
+
+void FMarqueeRange::HandleNode(const IMarqueeVisitor& Visitor, FSequencerDisplayNode& InNode, const TArray<TSharedRef<ISequencerSection>>& InSections)
+{
+	const float NodeCenter = InNode.GetVirtualTop() + (InNode.GetVirtualBottom() - InNode.GetVirtualTop())/2;
+	if (NodeCenter + VirtualKeySize.Y/2 > TopLeft.Y &&
+		NodeCenter - VirtualKeySize.Y/2 < BottomRight.Y)
+	{
+		bool bNodeHasKeyArea = false;
+		if (InNode.GetType() == ESequencerNode::KeyArea)
+		{
+			HandleKeyAreaNode(Visitor, static_cast<FSectionKeyAreaNode&>(InNode), InSections);
+			bNodeHasKeyArea = true;
+		}
+		else if (InNode.GetType() == ESequencerNode::Track)
+		{
+			TSharedPtr<FSectionKeyAreaNode> SectionKeyNode = static_cast<FTrackNode&>(InNode).GetTopLevelKeyNode();
+			if (SectionKeyNode.IsValid())
+			{
+				HandleKeyAreaNode(Visitor, static_cast<FSectionKeyAreaNode&>(InNode), InSections);
+				bNodeHasKeyArea = true;
+			}
+		}
+
+		if (!InNode.IsExpanded() && !bNodeHasKeyArea)
+		{
+			// As a fallback, handle key groupings on collapsed parents
+			for (int32 SectionIndex = 0; SectionIndex < InSections.Num(); ++SectionIndex)
+			{
+				TSharedRef<IKeyArea> KeyArea = InNode.UpdateKeyGrouping(SectionIndex);
+				HandleKeyArea(Visitor, KeyArea, InSections[SectionIndex]->GetSectionObject());
+			}
+		}
+	}
+
+	if (InNode.IsExpanded())
+	{
+		// Handle Children
+		for (auto& Child : InNode.GetChildNodes())
+		{
+			if (!Child->IsHidden())
+			{
+				HandleNode(Visitor, *Child, InSections);
+			}
+		}
+	}
+}
+
+void FMarqueeRange::HandleKeyAreaNode(const IMarqueeVisitor& Visitor, FSectionKeyAreaNode& KeyAreaNode, const TArray<TSharedRef<ISequencerSection>>& InSections)
+{
+	for( int32 SectionIndex = 0; SectionIndex < InSections.Num(); ++SectionIndex )
+	{
+		TSharedRef<IKeyArea> KeyArea = KeyAreaNode.GetKeyArea(SectionIndex);
+		HandleKeyArea(Visitor, KeyArea, InSections[SectionIndex]->GetSectionObject());
+	}
+}
+
+void FMarqueeRange::HandleKeyArea(const IMarqueeVisitor& Visitor, const TSharedPtr<IKeyArea>& KeyArea, UMovieSceneSection* Section)
+{
+	for (FKeyHandle KeyHandle : KeyArea->GetUnsortedKeyHandles())
+	{
+		float KeyPosition = KeyArea->GetKeyTime(KeyHandle);
+		if (KeyPosition + VirtualKeySize.X/2 > TopLeft.X &&
+			KeyPosition - VirtualKeySize.X/2 < BottomRight.X)
+		{
+			Visitor.VisitKey(KeyHandle, KeyPosition, KeyArea, Section);
+		}
+	}
+}
+
 /** Structure used for converting X and Y positional values to Time and vertical offsets respectively */
 struct FVirtualConverter : FTimeToPixel
 {
@@ -331,6 +467,48 @@ void SSequencerTrackArea::SetNewMarqueeBounds(const FPointerEvent& MouseEvent, c
 	}
 }
 
+struct FSelectionVisitor : IMarqueeVisitor
+{
+	FSelectionVisitor(FSequencerSelection& InSelection, const FPointerEvent& InMouseEvent)
+		: Selection(InSelection)
+		, MouseEvent(InMouseEvent)
+	{}
+
+	virtual void VisitKey(FKeyHandle KeyHandle, float KeyTime, const TSharedPtr<IKeyArea>& KeyArea, UMovieSceneSection* Section) const override
+	{
+		FSelectedKey SelectedKey(*Section, KeyArea, KeyHandle);
+
+		if (MouseEvent.IsLeftShiftDown())
+		{
+			Selection.AddToSelection(SelectedKey);
+		}
+		else if (MouseEvent.IsLeftControlDown())
+		{
+			if (Selection.IsSelected(SelectedKey))
+			{
+				Selection.RemoveFromSelection(SelectedKey);
+			}
+			else
+			{
+				Selection.AddToSelection(SelectedKey);
+			}
+		}
+		else
+		{
+			Selection.AddToSelection(SelectedKey);
+		}
+	}
+
+	virtual void VisitSection(FKeyHandle KeyHandle, float KeyTime, const TSharedPtr<IKeyArea>& KeyArea, UMovieSceneSection* Section) const
+	{
+
+	}
+
+private:
+	mutable FSequencerSelection& Selection;
+	const FPointerEvent& MouseEvent;
+};
+
 void SSequencerTrackArea::HandleMarqueeSelection(const FGeometry& Geometry, const FPointerEvent& MouseEvent)
 {
 	auto PinnedSequencer = Sequencer.Pin();
@@ -347,75 +525,16 @@ void SSequencerTrackArea::HandleMarqueeSelection(const FGeometry& Geometry, cons
 		Selection.Empty();
 	}
 
-	auto TopLeft = MarqueeSelectData->TopLeft();
-	auto BottomRight = MarqueeSelectData->BottomRight();
+	Selection.SuspendBroadcast();
 
-	// Pointer to the current traversal parent's sections.
-	const TArray<TSharedRef<ISequencerSection>>* CurrentSections = nullptr;
+	FVector2D VirtualKeySize;
+	VirtualKeySize.X = SequencerSectionConstants::KeySize.X / Geometry.Size.X * (PinnedSequencer->GetViewRange().GetUpperBoundValue() - PinnedSequencer->GetViewRange().GetLowerBoundValue());
+	// Vertically, virtual units == physical units
+	VirtualKeySize.Y = SequencerSectionConstants::KeySize.Y;
 
-	auto PerformSelection = [&](FSectionKeyAreaNode& KeyAreaNode){
-		for( int32 SectionIndex = 0; SectionIndex < CurrentSections->Num(); ++SectionIndex )
-		{
-			TSharedRef<IKeyArea> KeyArea = KeyAreaNode.GetKeyArea(SectionIndex);
+	FMarqueeRange Range(MarqueeSelectData->TopLeft(), MarqueeSelectData->BottomRight(), VirtualKeySize);
+	Range.Visit(FSelectionVisitor(Selection, MouseEvent), PinnedTreeView->GetNodeTree()->GetRootNodes());
 
-			for (FKeyHandle KeyHandle : KeyArea->GetUnsortedKeyHandles())
-			{
-				float KeyPosition = KeyArea->GetKeyTime(KeyHandle);
-				if (KeyPosition > TopLeft.X && KeyPosition < BottomRight.X)
-				{
-					FSelectedKey SelectedKey(*(*CurrentSections)[SectionIndex]->GetSectionObject(), KeyArea, KeyHandle);
-
-					if (MouseEvent.IsLeftShiftDown())
-					{
-						Selection.AddToSelection(SelectedKey);
-					}
-					else if (MouseEvent.IsLeftControlDown())
-					{
-						if (Selection.IsSelected(SelectedKey))
-						{
-							Selection.RemoveFromSelection(SelectedKey);
-						}
-						else
-						{
-							Selection.AddToSelection(SelectedKey);
-						}
-					}
-					else
-					{
-						Selection.AddToSelection(SelectedKey);
-					}
-				}
-			}
-		}
-	};
-
-	for (auto& RootNode : PinnedTreeView->GetNodeTree()->GetRootNodes())
-	{
-		RootNode->TraverseVisible_ParentFirst([&](FSequencerDisplayNode& InNode){
-
-			if (InNode.GetType() == ESequencerNode::KeyArea)
-			{
-				FSectionKeyAreaNode& KeyAreaNode = static_cast<FSectionKeyAreaNode&>(InNode);
-
-				if (CurrentSections && InNode.GetVirtualTop() > TopLeft.Y && InNode.GetVirtualBottom() < BottomRight.Y)
-				{
-					PerformSelection(KeyAreaNode);
-				}
-			}
-			else if (InNode.GetType() == ESequencerNode::Track)
-			{
-				FTrackNode& TrackNode = static_cast<FTrackNode&>(InNode);
-
-				CurrentSections = &TrackNode.GetSections();
-
-				TSharedPtr<FSectionKeyAreaNode> SectionKeyNode = TrackNode.GetTopLevelKeyNode();
-				if( SectionKeyNode.IsValid() && TrackNode.GetVirtualTop() > TopLeft.Y && TrackNode.GetVirtualBottom() < BottomRight.Y )
-				{
-					PerformSelection(*SectionKeyNode);
-				}
-			}
-
-			return true;
-		});
-	}
+	Selection.ResumeBroadcast();
+	Selection.GetOnKeySelectionChanged().Broadcast();
 }
\ No newline at end of file
diff --git a/Engine/Source/Editor/Sequencer/Private/SequencerDisplayNode.cpp b/Engine/Source/Editor/Sequencer/Private/SequencerDisplayNode.cpp
--- a/Engine/Source/Editor/Sequencer/Private/SequencerDisplayNode.cpp
+++ b/Engine/Source/Editor/Sequencer/Private/SequencerDisplayNode.cpp
@@ -12,7 +12,7 @@
 #include "MovieSceneTrack.h"
 #include "CommonMovieSceneTools.h"
 #include "IKeyArea.h"
-
+#include "GroupedKeyArea.h"
 
 #define LOCTEXT_NAMESPACE "SequencerDisplayNode"
 
@@ -427,6 +427,39 @@ void FSequencerDisplayNode::PinNode()
 	bNodeIsPinned = true;
 }
 
+TSharedRef<FGroupedKeyArea> FSequencerDisplayNode::GetKeyGrouping(int32 InSectionIndex)
+{
+	if (!KeyGroupings.IsValidIndex(InSectionIndex))
+	{
+		KeyGroupings.SetNum(InSectionIndex + 1);
+	}
+
+	if (!KeyGroupings[InSectionIndex].IsValid())
+	{
+		KeyGroupings[InSectionIndex] = MakeShareable(new FGroupedKeyArea(*this, InSectionIndex));
+	}
+
+	return KeyGroupings[InSectionIndex].ToSharedRef();
+}
+
+TSharedRef<FGroupedKeyArea> FSequencerDisplayNode::UpdateKeyGrouping(int32 InSectionIndex)
+{
+	if (!KeyGroupings.IsValidIndex(InSectionIndex))
+	{
+		KeyGroupings.SetNum(InSectionIndex + 1);
+	}
+
+	if (!KeyGroupings[InSectionIndex].IsValid())
+	{
+		KeyGroupings[InSectionIndex] = MakeShareable(new FGroupedKeyArea(*this, InSectionIndex));
+	}
+	else
+	{
+		*KeyGroupings[InSectionIndex] = FGroupedKeyArea(*this, InSectionIndex);
+	}
+
+	return KeyGroupings[InSectionIndex].ToSharedRef();
+}
 
 
 float FSectionKeyAreaNode::GetNodeHeight() const
diff --git a/Engine/Source/Editor/Sequencer/Private/SequencerDisplayNode.h b/Engine/Source/Editor/Sequencer/Private/SequencerDisplayNode.h
--- a/Engine/Source/Editor/Sequencer/Private/SequencerDisplayNode.h
+++ b/Engine/Source/Editor/Sequencer/Private/SequencerDisplayNode.h
@@ -5,6 +5,7 @@
 class IKeyArea;
 class ISequencerSection;
 class SSequencerTreeViewRow;
+class FGroupedKeyArea;
 
 /** Structure used to define padding for a particular node */
 struct FNodePadding
@@ -246,7 +247,13 @@ public:
 	
 	/** @return this node's virtual offset plus its virtual height, irrespective of expansion states */
 	float GetVirtualBottom() const { return VirtualBottom; }
-	
+
+	/** Get the key grouping for the specified section index, ensuring it is fully up to date */
+	TSharedRef<FGroupedKeyArea> UpdateKeyGrouping(int32 InSectionIndex);
+
+	/** Get the key grouping for the specified section index */
+	TSharedRef<FGroupedKeyArea> GetKeyGrouping(int32 InSectionIndex);
+
 protected:
 	/**
 	 * Visibility of nodes is a complex situation. There are 7 different factors to account for:
@@ -303,6 +310,8 @@ protected:
 	bool bCachedShotFilteredVisibility;
 	/** Whether this node is pinned to the top of the sequencer */
 	bool bNodeIsPinned;
+	/** Transient grouped keys for this node */
+	TArray< TSharedPtr<FGroupedKeyArea> > KeyGroupings;
 };
 
 /**
