Removed mutable keyword

diff --git a/Engine/Source/Editor/Sequencer/Private/SSequencerTrackArea.cpp b/Engine/Source/Editor/Sequencer/Private/SSequencerTrackArea.cpp
--- a/Engine/Source/Editor/Sequencer/Private/SSequencerTrackArea.cpp
+++ b/Engine/Source/Editor/Sequencer/Private/SSequencerTrackArea.cpp
@@ -505,7 +505,7 @@ struct FSelectionVisitor : IMarqueeVisitor
 	}
 
 private:
-	mutable FSequencerSelection& Selection;
+	FSequencerSelection& Selection;
 	const FPointerEvent& MouseEvent;
 };
 
