Sequencer: Add sequencer editor settings.

diff --git a/Engine/Source/Editor/Sequencer/Private/SequencerModule.cpp b/Engine/Source/Editor/Sequencer/Private/SequencerModule.cpp
--- a/Engine/Source/Editor/Sequencer/Private/SequencerModule.cpp
+++ b/Engine/Source/Editor/Sequencer/Private/SequencerModule.cpp
@@ -1,6 +1,7 @@
 // Copyright 1998-2015 Epic Games, Inc. All Rights Reserved.
 
 #include "SequencerPrivatePCH.h"
+#include "ISettingsModule.h"
 #include "ModuleManager.h"
 #include "Sequencer.h"
 #include "Toolkits/ToolkitManager.h"
@@ -16,6 +17,8 @@
 // These macros should be removed when those functions are removed.
 PRAGMA_DISABLE_DEPRECATION_WARNINGS
 
+#define LOCTEXT_NAMESPACE "SequencerEditor"
+
 
 /**
  * SequencerModule implementation (private)
@@ -63,6 +66,14 @@ public:
 
 		MenuExtensibilityManager = MakeShareable( new FExtensibilityManager );
 		ToolBarExtensibilityManager = MakeShareable( new FExtensibilityManager );
+
+		if (ISettingsModule* SettingsModule = FModuleManager::GetModulePtr<ISettingsModule>("Settings"))
+		{
+			SettingsModule->RegisterSettings("Project", "Editor", "Sequencer",
+				LOCTEXT("RuntimeSettingsName", "Sequencer"),
+				LOCTEXT("RuntimeSettingsDescription", "Configure the Sequencer plugin"),
+				GetMutableDefault<USequencerProjectSettings>());
+		}
 	}
 
 	virtual void ShutdownModule() override
@@ -73,6 +84,11 @@ public:
 
 			FEditorModeRegistry::Get().UnregisterMode(FSequencerEdMode::EM_SequencerMode);
 		}
+
+		if (ISettingsModule* SettingsModule = FModuleManager::GetModulePtr<ISettingsModule>("Settings"))
+		{
+			SettingsModule->UnregisterSettings("Project", "Editor", "Sequencer");
+		}
 	}
 
 	virtual TSharedPtr<FExtensibilityManager> GetMenuExtensibilityManager() const override { return MenuExtensibilityManager; }
@@ -91,3 +107,5 @@ private:
 PRAGMA_ENABLE_DEPRECATION_WARNINGS
 
 IMPLEMENT_MODULE(FSequencerModule, Sequencer);
+
+#undef LOCTEXT_NAMESPACE
\ No newline at end of file
diff --git a/Engine/Source/Editor/Sequencer/Private/SequencerPrivatePCH.h b/Engine/Source/Editor/Sequencer/Private/SequencerPrivatePCH.h
--- a/Engine/Source/Editor/Sequencer/Private/SequencerPrivatePCH.h
+++ b/Engine/Source/Editor/Sequencer/Private/SequencerPrivatePCH.h
@@ -22,6 +22,7 @@
 #include "SAnimationOutlinerTreeNode.h"
 #include "MovieSceneInstance.h"
 #include "SequencerSettings.h"
+#include "SequencerProjectSettings.h"
 #include "SequencerCommonHelpers.h"
 #include "SequencerCurveOwner.h"
 #include "SSequencerCurveEditor.h"
diff --git a/Engine/Source/Editor/Sequencer/Private/SequencerProjectSettings.cpp b/Engine/Source/Editor/Sequencer/Private/SequencerProjectSettings.cpp
new file mode 100644
--- /dev/null
+++ b/Engine/Source/Editor/Sequencer/Private/SequencerProjectSettings.cpp
@@ -0,0 +1,13 @@
+// Copyright 1998-2015 Epic Games, Inc. All Rights Reserved.
+
+#include "SequencerPrivatePCH.h"
+#include "SequencerProjectSettings.h"
+
+//////////////////////////////////////////////////////////////////////////
+// USequencerProjectSettings
+
+USequencerProjectSettings::USequencerProjectSettings()
+	: InFrame(101)
+	, OutFrame(200)
+{
+}
diff --git a/Engine/Source/Editor/Sequencer/Private/SequencerProjectSettings.h b/Engine/Source/Editor/Sequencer/Private/SequencerProjectSettings.h
new file mode 100644
--- /dev/null
+++ b/Engine/Source/Editor/Sequencer/Private/SequencerProjectSettings.h
@@ -0,0 +1,24 @@
+// Copyright 1998-2015 Epic Games, Inc. All Rights Reserved.
+
+#pragma once
+
+#include "SequencerProjectSettings.generated.h"
+
+// Settings for the Sequencer editor
+UCLASS(config=EditorPerProjectUserSettings)
+class USequencerProjectSettings : public UObject
+{
+	GENERATED_BODY()
+
+public:
+	USequencerProjectSettings();
+
+	/** The in frame. */
+	UPROPERTY(config, EditAnywhere, Category=Timeline)
+	int32 InFrame;
+
+	/** The out frame. */
+	UPROPERTY(config, EditAnywhere, Category=Timeline)
+	int32 OutFrame;
+
+};
