- Rename UMovieSceneAnimation to UMovieSceneSequence (and all related classes to use the term Sequence)
- Fixed shots and sub-movie scene sequences being broken by the removal of UMovieScene as an asset class

diff --git a/Engine/Plugins/MovieScene/ActorAnimationEditor/Source/ActorAnimationEditor/Private/ActorAnimationEditorToolkit.cpp b/Engine/Plugins/MovieScene/ActorAnimationEditor/Source/ActorAnimationEditor/Private/ActorAnimationEditorToolkit.cpp
--- a/Engine/Plugins/MovieScene/ActorAnimationEditor/Source/ActorAnimationEditor/Private/ActorAnimationEditorToolkit.cpp
+++ b/Engine/Plugins/MovieScene/ActorAnimationEditor/Source/ActorAnimationEditor/Private/ActorAnimationEditorToolkit.cpp
@@ -85,7 +85,7 @@ void FActorAnimationEditorToolkit::Initialize( const EToolkitMode::Type Mode, co
 	FSequencerInitParams SequencerInitParams;
 	{
 		SequencerInitParams.ViewParams = InViewParams;
-		SequencerInitParams.Animation = ActorAnimation;
+		SequencerInitParams.RootSequence = ActorAnimation;
 		SequencerInitParams.bEditWithinLevelEditor = bEditWithinLevelEditor;
 		SequencerInitParams.ToolkitHost = InitToolkitHost;
 	}
diff --git a/Engine/Plugins/MovieScene/ActorAnimationEditor/Source/ActorAnimationEditor/Private/ActorAnimationPlayer.cpp b/Engine/Plugins/MovieScene/ActorAnimationEditor/Source/ActorAnimationEditor/Private/ActorAnimationPlayer.cpp
--- a/Engine/Plugins/MovieScene/ActorAnimationEditor/Source/ActorAnimationEditor/Private/ActorAnimationPlayer.cpp
+++ b/Engine/Plugins/MovieScene/ActorAnimationEditor/Source/ActorAnimationEditor/Private/ActorAnimationPlayer.cpp
@@ -3,8 +3,8 @@
 #include "ActorAnimationEditorPrivatePCH.h"
 #include "ActorAnimationPlayer.h"
 #include "MovieScene.h"
-#include "MovieSceneAnimation.h"
-#include "MovieSceneInstance.h"
+#include "MovieSceneSequence.h"
+#include "MovieSceneSequenceInstance.h"
 
 
 /* UActorAnimationPlayer structors
@@ -62,7 +62,7 @@ void UActorAnimationPlayer::Play()
 	}
 
 	// @todo Sequencer playback: Should we recreate the instance every time?
-	RootMovieSceneInstance = MakeShareable(new FMovieSceneInstance(*ActorAnimation->GetMovieScene()));
+	RootMovieSceneInstance = MakeShareable(new FMovieSceneSequenceInstance(*ActorAnimation));
 
 	// @odo Sequencer Should we spawn actors here?
 	SpawnActorsForMovie(RootMovieSceneInstance.ToSharedRef());
@@ -85,7 +85,7 @@ void UActorAnimationPlayer::Initialize(UActorAnimation* InActorAnimation, UWorld
 /* IMovieScenePlayer interface
  *****************************************************************************/
 
-void UActorAnimationPlayer::SpawnActorsForMovie(TSharedRef<FMovieSceneInstance> MovieSceneInstance)
+void UActorAnimationPlayer::SpawnActorsForMovie(TSharedRef<FMovieSceneSequenceInstance> MovieSceneInstance)
 {
 	UWorld* WorldPtr = World.Get();
 
@@ -94,7 +94,7 @@ void UActorAnimationPlayer::SpawnActorsForMovie(TSharedRef<FMovieSceneInstance>
 		return;
 	}
 
-	UMovieScene* MovieScene = MovieSceneInstance->GetMovieScene();
+	UMovieScene* MovieScene = MovieSceneInstance->GetSequence()->GetMovieScene();
 
 	if (MovieScene == nullptr)
 	{
@@ -147,7 +147,7 @@ void UActorAnimationPlayer::SpawnActorsForMovie(TSharedRef<FMovieSceneInstance>
 }
 
 
-void UActorAnimationPlayer::DestroyActorsForMovie(TSharedRef<FMovieSceneInstance> MovieSceneInstance)
+void UActorAnimationPlayer::DestroyActorsForMovie(TSharedRef<FMovieSceneSequenceInstance> MovieSceneInstance)
 {
 	UWorld* WorldPtr = World.Get();
 
@@ -184,7 +184,7 @@ void UActorAnimationPlayer::DestroyActorsForMovie(TSharedRef<FMovieSceneInstance
 }
 
 
-void UActorAnimationPlayer::GetRuntimeObjects(TSharedRef<FMovieSceneInstance> MovieSceneInstance, const FGuid& ObjectId, TArray<UObject*>& OutObjects) const
+void UActorAnimationPlayer::GetRuntimeObjects(TSharedRef<FMovieSceneSequenceInstance> MovieSceneInstance, const FGuid& ObjectId, TArray<UObject*>& OutObjects) const
 {
 	// @todo sequencer runtime: Add support for individually spawning actors on demand when first requested?
 	//    This may be important to reduce the up-front hitch when spawning actors for the entire movie, but
@@ -231,20 +231,20 @@ EMovieScenePlayerStatus::Type UActorAnimationPlayer::GetPlaybackStatus() const
 }
 
 
-void UActorAnimationPlayer::AddOrUpdateMovieSceneInstance(UMovieSceneSection& MovieSceneSection, TSharedRef<FMovieSceneInstance> InstanceToAdd)
+void UActorAnimationPlayer::AddOrUpdateMovieSceneInstance(UMovieSceneSection& MovieSceneSection, TSharedRef<FMovieSceneSequenceInstance> InstanceToAdd)
 {
 	SpawnActorsForMovie( InstanceToAdd );
 }
 
 
-void UActorAnimationPlayer::RemoveMovieSceneInstance(UMovieSceneSection& MovieSceneSection, TSharedRef<FMovieSceneInstance> InstanceToRemove)
+void UActorAnimationPlayer::RemoveMovieSceneInstance(UMovieSceneSection& MovieSceneSection, TSharedRef<FMovieSceneSequenceInstance> InstanceToRemove)
 {
 	const bool bDestroyAll = true;
 	DestroyActorsForMovie( InstanceToRemove );
 }
 
 
-TSharedRef<FMovieSceneInstance> UActorAnimationPlayer::GetRootMovieSceneInstance() const
+TSharedRef<FMovieSceneSequenceInstance> UActorAnimationPlayer::GetRootMovieSceneSequenceInstance() const
 {
 	return RootMovieSceneInstance.ToSharedRef();
 }
diff --git a/Engine/Plugins/MovieScene/ActorAnimationEditor/Source/ActorAnimationEditor/Private/ActorAnimationPlayer.h b/Engine/Plugins/MovieScene/ActorAnimationEditor/Source/ActorAnimationEditor/Private/ActorAnimationPlayer.h
--- a/Engine/Plugins/MovieScene/ActorAnimationEditor/Source/ActorAnimationEditor/Private/ActorAnimationPlayer.h
+++ b/Engine/Plugins/MovieScene/ActorAnimationEditor/Source/ActorAnimationEditor/Private/ActorAnimationPlayer.h
@@ -9,7 +9,7 @@
 #include "ActorAnimationPlayer.generated.h"
 
 
-class FMovieSceneInstance;
+class FMovieSceneSequenceInstance;
 class ULevel;
 class UMovieSceneBindings;
 
@@ -63,14 +63,14 @@ protected:
 
 	// IMovieScenePlayer interface
 
-	virtual void SpawnActorsForMovie(TSharedRef<FMovieSceneInstance> MovieSceneInstance);
-	virtual void DestroyActorsForMovie(TSharedRef<FMovieSceneInstance> MovieSceneInstance);
-	virtual void GetRuntimeObjects(TSharedRef<FMovieSceneInstance> MovieSceneInstance, const FGuid& ObjectHandle, TArray<UObject*>& OutObjects) const override;
+	virtual void SpawnActorsForMovie(TSharedRef<FMovieSceneSequenceInstance> MovieSceneInstance);
+	virtual void DestroyActorsForMovie(TSharedRef<FMovieSceneSequenceInstance> MovieSceneInstance);
+	virtual void GetRuntimeObjects(TSharedRef<FMovieSceneSequenceInstance> MovieSceneInstance, const FGuid& ObjectHandle, TArray<UObject*>& OutObjects) const override;
 	virtual void UpdateCameraCut(UObject* ObjectToViewThrough, bool bNewCameraCut) const override;
 	virtual EMovieScenePlayerStatus::Type GetPlaybackStatus() const override;
-	virtual void AddOrUpdateMovieSceneInstance(UMovieSceneSection& MovieSceneSection, TSharedRef<FMovieSceneInstance> InstanceToAdd) override;
-	virtual void RemoveMovieSceneInstance(UMovieSceneSection& MovieSceneSection, TSharedRef<FMovieSceneInstance> InstanceToRemove) override;
-	virtual TSharedRef<FMovieSceneInstance> GetRootMovieSceneInstance() const override;
+	virtual void AddOrUpdateMovieSceneInstance(UMovieSceneSection& MovieSceneSection, TSharedRef<FMovieSceneSequenceInstance> InstanceToAdd) override;
+	virtual void RemoveMovieSceneInstance(UMovieSceneSection& MovieSceneSection, TSharedRef<FMovieSceneSequenceInstance> InstanceToRemove) override;
+	virtual TSharedRef<FMovieSceneSequenceInstance> GetRootMovieSceneSequenceInstance() const override;
 
 protected:
 
@@ -104,10 +104,10 @@ private:
 	};
 
 	/** Maps spawnable GUIDs to their spawned actor in the world */
-	TMap<TWeakPtr<FMovieSceneInstance>, TArray<FSpawnedActorInfo>> InstanceToSpawnedActorMap;
+	TMap<TWeakPtr<FMovieSceneSequenceInstance>, TArray<FSpawnedActorInfo>> InstanceToSpawnedActorMap;
 
 	/** The root movie scene instance to update when playing. */
-	TSharedPtr<FMovieSceneInstance> RootMovieSceneInstance;
+	TSharedPtr<FMovieSceneSequenceInstance> RootMovieSceneInstance;
 
 	/** The world this player will spawn actors in, if needed */
 	TWeakObjectPtr<UWorld> World;
diff --git a/Engine/Plugins/MovieScene/ActorAnimationEditor/Source/ActorAnimationEditor/Private/Factories/ActorAnimationFactoryNew.cpp b/Engine/Plugins/MovieScene/ActorAnimationEditor/Source/ActorAnimationEditor/Private/Factories/ActorAnimationFactoryNew.cpp
--- a/Engine/Plugins/MovieScene/ActorAnimationEditor/Source/ActorAnimationEditor/Private/Factories/ActorAnimationFactoryNew.cpp
+++ b/Engine/Plugins/MovieScene/ActorAnimationEditor/Source/ActorAnimationEditor/Private/Factories/ActorAnimationFactoryNew.cpp
@@ -30,7 +30,7 @@ UActorAnimationFactoryNew::UActorAnimationFactoryNew(const FObjectInitializer& O
 
 UObject* UActorAnimationFactoryNew::FactoryCreateNew(UClass* Class, UObject* InParent, FName Name, EObjectFlags Flags, UObject* Context, FFeedbackContext* Warn)
 {
-	auto NewActorAnimation = NewObject<UActorAnimation>(InParent, Name, Flags);
+	auto NewActorAnimation = NewObject<UActorAnimation>(InParent, Name, Flags|RF_Transactional);
 	NewActorAnimation->Initialize();
 
 	return NewActorAnimation;
diff --git a/Engine/Source/Editor/MovieSceneTools/Private/CurveKeyEditors/SBoolCurveKeyEditor.cpp b/Engine/Source/Editor/MovieSceneTools/Private/CurveKeyEditors/SBoolCurveKeyEditor.cpp
--- a/Engine/Source/Editor/MovieSceneTools/Private/CurveKeyEditors/SBoolCurveKeyEditor.cpp
+++ b/Engine/Source/Editor/MovieSceneTools/Private/CurveKeyEditors/SBoolCurveKeyEditor.cpp
@@ -10,7 +10,7 @@ void SBoolCurveKeyEditor::Construct(const FArguments& InArgs)
 	Sequencer = InArgs._Sequencer;
 	OwningSection = InArgs._OwningSection;
 	Curve = InArgs._Curve;
-	float CurrentTime = Sequencer->GetCurrentLocalTime(*Sequencer->GetFocusedMovieScene());
+	float CurrentTime = Sequencer->GetCurrentLocalTime(*Sequencer->GetFocusedMovieSceneSequence());
 	ChildSlot
 	[
 		SNew(SCheckBox)
@@ -21,7 +21,7 @@ void SBoolCurveKeyEditor::Construct(const FArguments& InArgs)
 
 ECheckBoxState SBoolCurveKeyEditor::IsChecked() const
 {
-	float CurrentTime = Sequencer->GetCurrentLocalTime(*Sequencer->GetFocusedMovieScene());
+	float CurrentTime = Sequencer->GetCurrentLocalTime(*Sequencer->GetFocusedMovieSceneSequence());
 	return !!Curve->Evaluate(CurrentTime) ? ECheckBoxState::Checked : ECheckBoxState::Unchecked;
 }
 
@@ -31,7 +31,7 @@ void SBoolCurveKeyEditor::OnCheckStateChanged(ECheckBoxState NewCheckboxState)
 	OwningSection->SetFlags(RF_Transactional);
 	OwningSection->Modify();
 
-	float CurrentTime = Sequencer->GetCurrentLocalTime(*Sequencer->GetFocusedMovieScene());
+	float CurrentTime = Sequencer->GetCurrentLocalTime(*Sequencer->GetFocusedMovieSceneSequence());
 
 	bool bKeyWillBeAdded = Curve->IsKeyHandleValid(Curve->FindKey(CurrentTime)) == false;
 	if (bKeyWillBeAdded)
diff --git a/Engine/Source/Editor/MovieSceneTools/Private/CurveKeyEditors/SEnumCurveKeyEditor.cpp b/Engine/Source/Editor/MovieSceneTools/Private/CurveKeyEditors/SEnumCurveKeyEditor.cpp
--- a/Engine/Source/Editor/MovieSceneTools/Private/CurveKeyEditors/SEnumCurveKeyEditor.cpp
+++ b/Engine/Source/Editor/MovieSceneTools/Private/CurveKeyEditors/SEnumCurveKeyEditor.cpp
@@ -19,7 +19,7 @@ void SEnumCurveKeyEditor::Construct(const FArguments& InArgs)
 		}
 		EnumValues.Add(MakeShareable(new FString(EnumDisplayValue)));
 	}
-	float CurrentTime = Sequencer->GetCurrentLocalTime(*Sequencer->GetFocusedMovieScene());
+	float CurrentTime = Sequencer->GetCurrentLocalTime(*Sequencer->GetFocusedMovieSceneSequence());
 	ChildSlot
 	[
 		SNew(SComboBox<TSharedPtr<FString>>)
@@ -39,7 +39,7 @@ void SEnumCurveKeyEditor::Construct(const FArguments& InArgs)
 
 FText SEnumCurveKeyEditor::GetCurrentValue() const
 {
-	float CurrentTime = Sequencer->GetCurrentLocalTime(*Sequencer->GetFocusedMovieScene());
+	float CurrentTime = Sequencer->GetCurrentLocalTime(*Sequencer->GetFocusedMovieSceneSequence());
 	return FText::FromString(*EnumValues[Curve->Evaluate(CurrentTime)]);
 }
 
@@ -55,7 +55,7 @@ void SEnumCurveKeyEditor::OnComboSelectionChanged(TSharedPtr<FString> InSelected
 	OwningSection->SetFlags(RF_Transactional);
 	OwningSection->Modify();
 
-	float CurrentTime = Sequencer->GetCurrentLocalTime(*Sequencer->GetFocusedMovieScene());
+	float CurrentTime = Sequencer->GetCurrentLocalTime(*Sequencer->GetFocusedMovieSceneSequence());
 
 	bool bKeyWillBeAdded = Curve->IsKeyHandleValid(Curve->FindKey(CurrentTime)) == false;
 	if (bKeyWillBeAdded)
diff --git a/Engine/Source/Editor/MovieSceneTools/Private/CurveKeyEditors/SFloatCurveKeyEditor.cpp b/Engine/Source/Editor/MovieSceneTools/Private/CurveKeyEditors/SFloatCurveKeyEditor.cpp
--- a/Engine/Source/Editor/MovieSceneTools/Private/CurveKeyEditors/SFloatCurveKeyEditor.cpp
+++ b/Engine/Source/Editor/MovieSceneTools/Private/CurveKeyEditors/SFloatCurveKeyEditor.cpp
@@ -45,7 +45,7 @@ void SFloatCurveKeyEditor::OnEndSliderMovement(float Value)
 
 float SFloatCurveKeyEditor::OnGetKeyValue() const
 {
-	float CurrentTime = Sequencer->GetCurrentLocalTime(*Sequencer->GetFocusedMovieScene());
+	float CurrentTime = Sequencer->GetCurrentLocalTime(*Sequencer->GetFocusedMovieSceneSequence());
 	return Curve->Eval(CurrentTime);
 }
 
@@ -53,7 +53,7 @@ void SFloatCurveKeyEditor::OnValueChanged(float Value)
 {
 	OwningSection->Modify();
 
-	float CurrentTime = Sequencer->GetCurrentLocalTime(*Sequencer->GetFocusedMovieScene());
+	float CurrentTime = Sequencer->GetCurrentLocalTime(*Sequencer->GetFocusedMovieSceneSequence());
 	FKeyHandle CurrentKeyHandle = Curve->FindKey(CurrentTime);
 	if (Curve->IsKeyHandleValid(CurrentKeyHandle))
 	{
diff --git a/Engine/Source/Editor/MovieSceneTools/Private/CurveKeyEditors/SIntegralCurveKeyEditor.cpp b/Engine/Source/Editor/MovieSceneTools/Private/CurveKeyEditors/SIntegralCurveKeyEditor.cpp
--- a/Engine/Source/Editor/MovieSceneTools/Private/CurveKeyEditors/SIntegralCurveKeyEditor.cpp
+++ b/Engine/Source/Editor/MovieSceneTools/Private/CurveKeyEditors/SIntegralCurveKeyEditor.cpp
@@ -45,7 +45,7 @@ void SIntegralCurveKeyEditor::OnEndSliderMovement(int32 Value)
 
 int32 SIntegralCurveKeyEditor::OnGetKeyValue() const
 {
-	float CurrentTime = Sequencer->GetCurrentLocalTime(*Sequencer->GetFocusedMovieScene());
+	float CurrentTime = Sequencer->GetCurrentLocalTime(*Sequencer->GetFocusedMovieSceneSequence());
 	return Curve->Evaluate(CurrentTime);
 }
 
@@ -53,7 +53,7 @@ void SIntegralCurveKeyEditor::OnValueChanged(int32 Value)
 {
 	OwningSection->Modify();
 
-	float CurrentTime = Sequencer->GetCurrentLocalTime(*Sequencer->GetFocusedMovieScene());
+	float CurrentTime = Sequencer->GetCurrentLocalTime(*Sequencer->GetFocusedMovieSceneSequence());
 
 	bool bKeyWillBeAdded = Curve->IsKeyHandleValid(Curve->FindKey(CurrentTime)) == false;
 	if (bKeyWillBeAdded)
diff --git a/Engine/Source/Editor/MovieSceneTools/Private/MovieSceneToolsModule.cpp b/Engine/Source/Editor/MovieSceneTools/Private/MovieSceneToolsModule.cpp
--- a/Engine/Source/Editor/MovieSceneTools/Private/MovieSceneToolsModule.cpp
+++ b/Engine/Source/Editor/MovieSceneTools/Private/MovieSceneToolsModule.cpp
@@ -14,7 +14,7 @@
 #include "ShotTrackEditor.h"
 #include "SubMovieSceneTrackEditor.h"
 #include "AudioTrackEditor.h"
-#include "AnimationTrackEditor.h"
+#include "SkeletalAnimationTrackEditor.h"
 #include "ParticleTrackEditor.h"
 #include "PathTrackEditor.h"
 
@@ -33,7 +33,7 @@ class FMovieSceneToolsModule : public IMovieSceneTools
 		EditorCreateTrackEditorDelegateHandle					= SequencerModule.RegisterTrackEditor_Handle( FOnCreateTrackEditor::CreateStatic( &FShotTrackEditor::CreateTrackEditor ) );
 		SubMovieSceneTrackEditorCreateTrackEditorDelegateHandle = SequencerModule.RegisterTrackEditor_Handle( FOnCreateTrackEditor::CreateStatic( &FSubMovieSceneTrackEditor::CreateTrackEditor ) );
 		AudioTrackEditorCreateTrackEditorDelegateHandle         = SequencerModule.RegisterTrackEditor_Handle( FOnCreateTrackEditor::CreateStatic( &FAudioTrackEditor::CreateTrackEditor ) );
-		AnimationTrackEditorCreateTrackEditorDelegateHandle     = SequencerModule.RegisterTrackEditor_Handle( FOnCreateTrackEditor::CreateStatic( &FAnimationTrackEditor::CreateTrackEditor ) );
+		AnimationTrackEditorCreateTrackEditorDelegateHandle     = SequencerModule.RegisterTrackEditor_Handle( FOnCreateTrackEditor::CreateStatic( &FSkeletalAnimationTrackEditor::CreateTrackEditor ) );
 		ParticleTrackEditorCreateTrackEditorDelegateHandle      = SequencerModule.RegisterTrackEditor_Handle( FOnCreateTrackEditor::CreateStatic( &FParticleTrackEditor::CreateTrackEditor ) );
 		PathTrackEditorCreateTrackEditorDelegateHandle          = SequencerModule.RegisterTrackEditor_Handle( FOnCreateTrackEditor::CreateStatic( &F3DPathTrackEditor::CreateTrackEditor ) );
 	}
diff --git a/Engine/Source/Editor/MovieSceneTools/Private/MovieSceneToolsPrivatePCH.h b/Engine/Source/Editor/MovieSceneTools/Private/MovieSceneToolsPrivatePCH.h
--- a/Engine/Source/Editor/MovieSceneTools/Private/MovieSceneToolsPrivatePCH.h
+++ b/Engine/Source/Editor/MovieSceneTools/Private/MovieSceneToolsPrivatePCH.h
@@ -7,7 +7,7 @@
 #include "MovieScene.h"
 #include "ScopedTransaction.h"
 #include "ISequencerSection.h"
-#include "MovieSceneInstance.h"
+#include "MovieSceneSequenceInstance.h"
 #include "MovieSceneTrackEditor.h"
 #include "IMovieSceneTools.h"
 #include "MovieSceneSection.h"
diff --git a/Engine/Source/Editor/MovieSceneTools/Private/TrackEditors/AnimationTrackEditor.cpp b/Engine/Source/Editor/MovieSceneTools/Private/TrackEditors/AnimationTrackEditor.cpp
deleted file mode 100644
diff --git a/Engine/Source/Editor/MovieSceneTools/Private/TrackEditors/AnimationTrackEditor.h b/Engine/Source/Editor/MovieSceneTools/Private/TrackEditors/AnimationTrackEditor.h
deleted file mode 100644
diff --git a/Engine/Source/Editor/MovieSceneTools/Private/TrackEditors/AudioTrackEditor.cpp b/Engine/Source/Editor/MovieSceneTools/Private/TrackEditors/AudioTrackEditor.cpp
--- a/Engine/Source/Editor/MovieSceneTools/Private/TrackEditors/AudioTrackEditor.cpp
+++ b/Engine/Source/Editor/MovieSceneTools/Private/TrackEditors/AudioTrackEditor.cpp
@@ -474,7 +474,7 @@ bool FAudioTrackEditor::HandleAssetAdded(UObject* Asset, const FGuid& TargetObje
 		if (TargetObjectGuid.IsValid())
 		{
 			TArray<UObject*> OutObjects;
-			GetSequencer()->GetRuntimeObjects(GetSequencer()->GetFocusedMovieSceneInstance(), TargetObjectGuid, OutObjects);
+			GetSequencer()->GetRuntimeObjects(GetSequencer()->GetFocusedMovieSceneSequenceInstance(), TargetObjectGuid, OutObjects);
 
 			AnimatablePropertyChanged(UMovieSceneAudioTrack::StaticClass(), false,
 				FOnKeyProperty::CreateRaw(this, &FAudioTrackEditor::AddNewAttachedSound, Sound, OutObjects));
diff --git a/Engine/Source/Editor/MovieSceneTools/Private/TrackEditors/ParticleTrackEditor.cpp b/Engine/Source/Editor/MovieSceneTools/Private/TrackEditors/ParticleTrackEditor.cpp
--- a/Engine/Source/Editor/MovieSceneTools/Private/TrackEditors/ParticleTrackEditor.cpp
+++ b/Engine/Source/Editor/MovieSceneTools/Private/TrackEditors/ParticleTrackEditor.cpp
@@ -157,7 +157,7 @@ void FParticleTrackEditor::BuildObjectBindingContextMenu(FMenuBuilder& MenuBuild
 void FParticleTrackEditor::AddParticleKey(const FGuid ObjectGuid, bool bTrigger)
 {
 	TArray<UObject*> OutObjects;
-	GetSequencer()->GetRuntimeObjects( GetSequencer()->GetFocusedMovieSceneInstance(), ObjectGuid, OutObjects );
+	GetSequencer()->GetRuntimeObjects( GetSequencer()->GetFocusedMovieSceneSequenceInstance(), ObjectGuid, OutObjects );
 
 	AnimatablePropertyChanged( UMovieSceneParticleTrack::StaticClass(), false,
 		FOnKeyProperty::CreateRaw( this, &FParticleTrackEditor::AddKeyInternal, OutObjects, bTrigger) );
diff --git a/Engine/Source/Editor/MovieSceneTools/Private/TrackEditors/PathTrackEditor.cpp b/Engine/Source/Editor/MovieSceneTools/Private/TrackEditors/PathTrackEditor.cpp
--- a/Engine/Source/Editor/MovieSceneTools/Private/TrackEditors/PathTrackEditor.cpp
+++ b/Engine/Source/Editor/MovieSceneTools/Private/TrackEditors/PathTrackEditor.cpp
@@ -268,7 +268,7 @@ void F3DPathTrackEditor::AddPathSubMenu(FMenuBuilder& MenuBuilder, FGuid ObjectB
 void F3DPathTrackEditor::AddPath(FGuid ObjectGuid, UObject* AdditionalAsset)
 {
 	TArray<UObject*> OutObjects;
-	GetSequencer()->GetRuntimeObjects( GetSequencer()->GetFocusedMovieSceneInstance(), ObjectGuid, OutObjects);
+	GetSequencer()->GetRuntimeObjects( GetSequencer()->GetFocusedMovieSceneSequenceInstance(), ObjectGuid, OutObjects);
 
 	AnimatablePropertyChanged( UMovieScene3DPathTrack::StaticClass(), false,
 		FOnKeyProperty::CreateRaw( this, &F3DPathTrackEditor::AddKeyInternal, OutObjects, AdditionalAsset) );
diff --git a/Engine/Source/Editor/MovieSceneTools/Private/TrackEditors/PropertyTrackEditor.cpp b/Engine/Source/Editor/MovieSceneTools/Private/TrackEditors/PropertyTrackEditor.cpp
--- a/Engine/Source/Editor/MovieSceneTools/Private/TrackEditors/PropertyTrackEditor.cpp
+++ b/Engine/Source/Editor/MovieSceneTools/Private/TrackEditors/PropertyTrackEditor.cpp
@@ -3,7 +3,7 @@
 #include "MovieSceneToolsPrivatePCH.h"
 #include "PropertyEditing.h"
 #include "MovieScene.h"
-#include "MovieSceneAnimation.h"
+#include "MovieSceneSequence.h"
 #include "MovieSceneTrack.h"
 #include "MovieSceneFloatTrack.h"
 #include "MovieSceneVectorTrack.h"
@@ -108,7 +108,9 @@ private:
 
 		static const FName SlateColorName("SlateColor");
 
-		const TArray<FMovieSceneBinding>& MovieSceneBindings = Sequencer->GetFocusedMovieScene()->GetBindings();
+		UMovieSceneSequence* Sequence = Sequencer->GetFocusedMovieSceneSequence();
+
+		const TArray<FMovieSceneBinding>& MovieSceneBindings = Sequence->GetMovieScene()->GetBindings();
 
 		bool bFoundColor = false;
 		for (int32 BindingIndex = 0; BindingIndex < MovieSceneBindings.Num() && !bFoundColor; ++BindingIndex)
@@ -119,7 +121,7 @@ private:
 			{
 				if (MovieSceneBinding.GetTracks()[TrackIndex] == Track)
 				{
-					UObject* RuntimeObject = Sequencer->GetAnimation()->FindObject(MovieSceneBinding.GetObjectGuid());
+					UObject* RuntimeObject = Sequence->FindObject(MovieSceneBinding.GetObjectGuid());
 
 					if (RuntimeObject != nullptr)
 					{
@@ -456,7 +458,7 @@ void FPropertyTrackEditor::AddKey(const FGuid& ObjectGuid, UObject* AdditionalAs
 	ISequencerObjectChangeListener& ObjectChangeListener = GetSequencer()->GetObjectChangeListener();
 
 	TArray<UObject*> OutObjects;
-	GetSequencer()->GetRuntimeObjects( GetSequencer()->GetFocusedMovieSceneInstance(), ObjectGuid, OutObjects);
+	GetSequencer()->GetRuntimeObjects( GetSequencer()->GetFocusedMovieSceneSequenceInstance(), ObjectGuid, OutObjects);
 	for (int32 i = 0; i < OutObjects.Num(); ++i)
 	{
 		ObjectChangeListener.TriggerAllPropertiesChanged(OutObjects[i]);
@@ -465,7 +467,7 @@ void FPropertyTrackEditor::AddKey(const FGuid& ObjectGuid, UObject* AdditionalAs
 
 UEnum* GetEnumForByteTrack(TSharedPtr<ISequencer> Sequencer, const FGuid& OwnerObjectHandle, FName PropertyName, UMovieSceneByteTrack* ByteTrack)
 {
-	UObject* RuntimeObject = Sequencer->GetAnimation()->FindObject(OwnerObjectHandle);
+	UObject* RuntimeObject = Sequencer->GetFocusedMovieSceneSequence()->FindObject(OwnerObjectHandle);
 
 	TSet<UEnum*> PropertyEnums;
 	if (RuntimeObject != nullptr)
diff --git a/Engine/Source/Editor/MovieSceneTools/Private/TrackEditors/ShotTrackEditor.cpp b/Engine/Source/Editor/MovieSceneTools/Private/TrackEditors/ShotTrackEditor.cpp
--- a/Engine/Source/Editor/MovieSceneTools/Private/TrackEditors/ShotTrackEditor.cpp
+++ b/Engine/Source/Editor/MovieSceneTools/Private/TrackEditors/ShotTrackEditor.cpp
@@ -248,7 +248,7 @@ FReply FShotSection::OnSectionDoubleClicked( const FGeometry& SectionGeometry, c
 {
 	if( MouseEvent.GetEffectingButton() == EKeys::LeftMouseButton )
 	{
-		TSharedRef<FMovieSceneInstance> MovieSceneInstance = Sequencer.Pin()->GetInstanceForSubMovieSceneSection( *Section );
+		TSharedRef<FMovieSceneSequenceInstance> MovieSceneInstance = Sequencer.Pin()->GetInstanceForSubMovieSceneSection( *Section );
 
 		Sequencer.Pin()->FocusSubMovieScene( MovieSceneInstance );
 	}
@@ -448,7 +448,7 @@ ACameraActor* FShotSection::UpdateCameraObject() const
 {
 	TArray<UObject*> OutObjects;
 	// @todo Sequencer - Sub-MovieScenes The director track may be able to get cameras from sub-movie scenes
-	Sequencer.Pin()->GetRuntimeObjects( Sequencer.Pin()->GetRootMovieSceneInstance(),  Section->GetCameraGuid(), OutObjects);
+	Sequencer.Pin()->GetRuntimeObjects( Sequencer.Pin()->GetRootMovieSceneSequenceInstance(),  Section->GetCameraGuid(), OutObjects);
 
 	ACameraActor* ReturnCam = nullptr;
 	if( OutObjects.Num() > 0 )
@@ -471,20 +471,20 @@ void FShotSection::OnRenameShot( const FText& NewShotName, ETextCommit::Type Com
 		Section->Modify();
 		Section->SetShotNameAndNumber( NewShotName, Section->GetShotNumber() );
 
-		UMovieScene* ShotMovieScene = Section->GetMovieScene();
-		if( ShotMovieScene )
+		UMovieSceneSequence* ShotMovieSceneAnim = Section->GetMovieSceneAnimation();
+		if( ShotMovieSceneAnim )
 		{
 			FAssetToolsModule& AssetToolsModule = FModuleManager::GetModuleChecked<FAssetToolsModule>("AssetTools");
 
 			TArray<FAssetRenameData> AssetsAndNames;
-			const FString PackagePath = FPackageName::GetLongPackagePath(ShotMovieScene->GetOutermost()->GetName());
+			const FString PackagePath = FPackageName::GetLongPackagePath(ShotMovieSceneAnim->GetOutermost()->GetName());
 
 			// We want to prepend the owning movie scene name.  This is not the movie scene inside referenced by the shot but the movie scene containing 
 			// the shot track
 			UMovieScene* OwningMovieScene = Section->GetTypedOuter<UMovieScene>();
 			const FString NewAssetName = OwningMovieScene->GetName() + TEXT("_") + NewShotName.ToString();
 
-			new (AssetsAndNames) FAssetRenameData(ShotMovieScene, PackagePath, NewAssetName );
+			new (AssetsAndNames) FAssetRenameData(ShotMovieSceneAnim, PackagePath, NewAssetName );
 
 			AssetToolsModule.Get().RenameAssets(AssetsAndNames);
 		}
@@ -527,7 +527,7 @@ void FShotTrackEditor::AddKey(const FGuid& ObjectGuid, UObject* AdditionalAsset)
 {
 	TArray<UObject*> OutObjects;
 	// @todo Sequencer - Sub-MovieScenes The director track may be able to get cameras from sub-movie scenes
-	GetSequencer()->GetRuntimeObjects(  GetSequencer()->GetRootMovieSceneInstance(), ObjectGuid, OutObjects);
+	GetSequencer()->GetRuntimeObjects(  GetSequencer()->GetRootMovieSceneSequenceInstance(), ObjectGuid, OutObjects);
 	bool bValidKey = OutObjects.Num() == 1 && OutObjects[0]->IsA<ACameraActor>();
 	if (bValidKey)
 	{
@@ -560,7 +560,8 @@ void FShotTrackEditor::BuildObjectBindingContextMenu(FMenuBuilder& MenuBuilder,
 
 void FShotTrackEditor::AddKeyInternal( float KeyTime, const FGuid ObjectGuid )
 {
-	UMovieScene* MovieScene = GetSequencer()->GetFocusedMovieScene();
+	UMovieSceneSequence* FocusedSequence = GetSequencer()->GetFocusedMovieSceneSequence();
+	UMovieScene* MovieScene = FocusedSequence->GetMovieScene();
 
 	UMovieSceneTrack* Track = MovieScene->GetShotTrack();
 	if( !Track )
@@ -594,7 +595,12 @@ void FShotTrackEditor::AddKeyInternal( float KeyTime, const FGuid ObjectGuid )
 	AssetToolsModule.Get().CreateUniqueAssetName(NewShotPackagePath, EmptySuffix, PackageName, AssetName);
 	const FString PackagePath = FPackageName::GetLongPackagePath(PackageName);
 
-	UObject* NewAsset = AssetToolsModule.Get().CreateAsset(AssetName, PackagePath, UMovieScene::StaticClass(), nullptr );
+	// Create the same type as the sequence type we have
+	UClass* AssetClass = FocusedSequence->GetClass();
+
+	UFactory* Factory = GetAssetFactoryForNewShot( AssetClass );
+
+	UObject* NewAsset = AssetToolsModule.Get().CreateAsset( AssetName, PackagePath, AssetClass, Factory );
 	
 	float EndTime = KeyTime;
 	if( AllSections.IsValidIndex( SectionIndex + 1 ) )
@@ -622,7 +628,7 @@ void FShotTrackEditor::AddKeyInternal( float KeyTime, const FGuid ObjectGuid )
 	}
 
 	TRange<float> TimeRange( KeyTime, EndTime );
-	ShotTrack->AddNewShot( ObjectGuid, *CastChecked<UMovieScene>( NewAsset ), TimeRange, FText::FromString( ShotName ), ShotNumber );
+	ShotTrack->AddNewShot( ObjectGuid, *CastChecked<UMovieSceneSequence>( NewAsset ), TimeRange, FText::FromString( ShotName ), ShotNumber );
 }
 
 int32 FShotTrackEditor::GenerateShotNumber( UMovieSceneShotTrack& ShotTrack, int32 SectionIndex ) const
@@ -686,4 +692,32 @@ int32 FShotTrackEditor::FindIndexForNewShot( const TArray<UMovieSceneSection*>&
 	for (SectionIndex = 0; SectionIndex < ShotSections.Num() && ShotSections[SectionIndex]->GetStartTime() < NewShotTime; ++SectionIndex);
 
 	return SectionIndex;
+}
+
+UFactory* FShotTrackEditor::GetAssetFactoryForNewShot( UClass* SequenceClass )
+{
+	static TWeakObjectPtr<UFactory> ShotFactory;
+
+	if( !ShotFactory.IsValid() || ShotFactory->SupportedClass != SequenceClass )
+	{
+		TArray<UFactory*> Factories;
+		for (TObjectIterator<UClass> It; It; ++It)
+		{
+			UClass* Class = *It;
+			if (Class->IsChildOf(UFactory::StaticClass()) && !Class->HasAnyClassFlags(CLASS_Abstract))
+			{
+				UFactory* Factory = Class->GetDefaultObject<UFactory>();
+				if (Factory->CanCreateNew())
+				{
+					UClass* SupportedClass = Factory->GetSupportedClass();
+					if ( SupportedClass == SequenceClass )
+					{
+						ShotFactory = Factory;
+					}
+				}
+			}
+		}
+	}
+
+	return ShotFactory.Get();
 }
\ No newline at end of file
diff --git a/Engine/Source/Editor/MovieSceneTools/Private/TrackEditors/ShotTrackEditor.h b/Engine/Source/Editor/MovieSceneTools/Private/TrackEditors/ShotTrackEditor.h
--- a/Engine/Source/Editor/MovieSceneTools/Private/TrackEditors/ShotTrackEditor.h
+++ b/Engine/Source/Editor/MovieSceneTools/Private/TrackEditors/ShotTrackEditor.h
@@ -41,6 +41,9 @@ private:
 
 	/** Finds the index in the ShotSections array where a new shot should be inserted */
 	int32 FindIndexForNewShot( const TArray<UMovieSceneSection*>& ShotSections, float NewShotTime ) const;
+
+	UFactory* GetAssetFactoryForNewShot( UClass* SequenceClass );
+
 private:
 	/** The Thumbnail pool which draws all the viewport thumbnails for the director track */
 	TSharedPtr<class FShotThumbnailPool> ThumbnailPool;
diff --git a/Engine/Source/Editor/MovieSceneTools/Private/TrackEditors/SkeletalAnimationTrackEditor.cpp b/Engine/Source/Editor/MovieSceneTools/Private/TrackEditors/SkeletalAnimationTrackEditor.cpp
new file mode 100644
--- /dev/null
+++ b/Engine/Source/Editor/MovieSceneTools/Private/TrackEditors/SkeletalAnimationTrackEditor.cpp
@@ -0,0 +1,278 @@
+// Copyright 1998-2015 Epic Games, Inc. All Rights Reserved.
+
+#include "MovieSceneToolsPrivatePCH.h"
+#include "MovieScene.h"
+#include "MovieSceneSection.h"
+#include "ISequencerSection.h"
+#include "PropertyEditorModule.h"
+#include "PropertyHandle.h"
+#include "MovieSceneTrack.h"
+#include "MovieSceneSkeletalAnimationTrack.h"
+#include "ScopedTransaction.h"
+#include "ISequencerObjectChangeListener.h"
+#include "ISectionLayoutBuilder.h"
+#include "IKeyArea.h"
+#include "MovieSceneToolHelpers.h"
+#include "MovieSceneTrackEditor.h"
+#include "SkeletalAnimationTrackEditor.h"
+#include "MovieSceneSkeletalAnimationSection.h"
+#include "CommonMovieSceneTools.h"
+#include "AssetRegistryModule.h"
+#include "Animation/SkeletalMeshActor.h"
+
+
+
+namespace SkeletalAnimationEditorConstants
+{
+	// @todo Sequencer Allow this to be customizable
+	const uint32 AnimationTrackHeight = 20;
+}
+
+
+
+FSkeletalAnimationSection::FSkeletalAnimationSection( UMovieSceneSection& InSection )
+	: Section( InSection )
+{
+}
+
+FSkeletalAnimationSection::~FSkeletalAnimationSection()
+{
+}
+
+UMovieSceneSection* FSkeletalAnimationSection::GetSectionObject()
+{ 
+	return &Section;
+}
+
+FText FSkeletalAnimationSection::GetDisplayName() const
+{
+	return NSLOCTEXT("FAnimationSection", "AnimationSection", "Animation");
+}
+
+FText FSkeletalAnimationSection::GetSectionTitle() const
+{
+	return FText::FromString( Cast<UMovieSceneSkeletalAnimationSection>(&Section)->GetAnimSequence()->GetName() );
+}
+
+float FSkeletalAnimationSection::GetSectionHeight() const
+{
+	return (float)SkeletalAnimationEditorConstants::AnimationTrackHeight;
+ }
+
+int32 FSkeletalAnimationSection::OnPaintSection( const FGeometry& AllottedGeometry, const FSlateRect& SectionClippingRect, FSlateWindowElementList& OutDrawElements, int32 LayerId, bool bParentEnabled ) const
+{
+	UMovieSceneSkeletalAnimationSection* AnimSection = Cast<UMovieSceneSkeletalAnimationSection>(&Section);
+	
+	FTimeToPixel TimeToPixelConverter( AllottedGeometry, TRange<float>( Section.GetStartTime(), Section.GetEndTime() ) );
+
+	// Add a box for the section
+	FSlateDrawElement::MakeBox(
+		OutDrawElements,
+		LayerId,
+		AllottedGeometry.ToPaintGeometry(),
+		FEditorStyle::GetBrush("Sequencer.GenericSection.Background"),
+		SectionClippingRect,
+		ESlateDrawEffect::None,
+		FLinearColor(0.7f, 0.4f, 0.7f, 1.f)
+	);
+
+	// Darken the part that doesn't have animation
+	if (AnimSection->GetAnimationStartTime() > AnimSection->GetStartTime())
+	{
+		float StartDarkening = AnimSection->GetStartTime();
+		float EndDarkening = FMath::Min(AnimSection->GetAnimationStartTime(), AnimSection->GetEndTime());
+		
+		float StartPixels = TimeToPixelConverter.TimeToPixel(StartDarkening);
+		float EndPixels = TimeToPixelConverter.TimeToPixel(EndDarkening);
+
+		FSlateDrawElement::MakeBox(
+			OutDrawElements,
+			LayerId + 1,
+			AllottedGeometry.ToPaintGeometry(FVector2D(StartPixels, 0), FVector2D(EndPixels - StartPixels, AllottedGeometry.Size.Y)),
+			FEditorStyle::GetBrush("WhiteTexture"),
+			SectionClippingRect,
+			ESlateDrawEffect::None,
+			FLinearColor(0.f, 0.f, 0.f, 0.3f)
+		);
+	}
+
+	// Add lines where the animation starts and ends/loops
+	float CurrentTime = AnimSection->GetAnimationStartTime();
+	while (CurrentTime < AnimSection->GetEndTime())
+	{
+		if (CurrentTime > AnimSection->GetStartTime())
+		{
+			float CurrentPixels = TimeToPixelConverter.TimeToPixel(CurrentTime);
+
+			TArray<FVector2D> Points;
+			Points.Add(FVector2D(CurrentPixels, 0));
+			Points.Add(FVector2D(CurrentPixels, AllottedGeometry.Size.Y));
+
+			FSlateDrawElement::MakeLines(
+				OutDrawElements,
+				LayerId + 2,
+				AllottedGeometry.ToPaintGeometry(),
+				Points,
+				SectionClippingRect
+			);
+		}
+		CurrentTime += AnimSection->GetAnimationDuration();
+	}
+
+	return LayerId+3;
+}
+
+FSkeletalAnimationTrackEditor::FSkeletalAnimationTrackEditor( TSharedRef<ISequencer> InSequencer )
+	: FMovieSceneTrackEditor( InSequencer ) 
+{
+}
+
+FSkeletalAnimationTrackEditor::~FSkeletalAnimationTrackEditor()
+{
+}
+
+TSharedRef<FMovieSceneTrackEditor> FSkeletalAnimationTrackEditor::CreateTrackEditor( TSharedRef<ISequencer> InSequencer )
+{
+	return MakeShareable( new FSkeletalAnimationTrackEditor( InSequencer ) );
+}
+
+bool FSkeletalAnimationTrackEditor::SupportsType( TSubclassOf<UMovieSceneTrack> Type ) const
+{
+	return Type == UMovieSceneSkeletalAnimationTrack::StaticClass();
+}
+
+TSharedRef<ISequencerSection> FSkeletalAnimationTrackEditor::MakeSectionInterface( UMovieSceneSection& SectionObject, UMovieSceneTrack* Track )
+{
+	check( SupportsType( SectionObject.GetOuter()->GetClass() ) );
+	
+	TSharedRef<ISequencerSection> NewSection( new FSkeletalAnimationSection(SectionObject) );
+
+	return NewSection;
+}
+
+void FSkeletalAnimationTrackEditor::AddKey(const FGuid& ObjectGuid, UObject* AdditionalAsset)
+{
+	UAnimSequence* AnimSequence = Cast<UAnimSequence>(AdditionalAsset);
+
+	if (AnimSequence)
+	{
+		TArray<UObject*> OutObjects;
+		GetSequencer()->GetRuntimeObjects( GetSequencer()->GetFocusedMovieSceneSequenceInstance(), ObjectGuid, OutObjects);
+
+		AnimatablePropertyChanged( UMovieSceneSkeletalAnimationTrack::StaticClass(), false,
+			FOnKeyProperty::CreateRaw( this, &FSkeletalAnimationTrackEditor::AddKeyInternal, OutObjects, AnimSequence) );
+	}
+}
+
+bool FSkeletalAnimationTrackEditor::HandleAssetAdded(UObject* Asset, const FGuid& TargetObjectGuid)
+{
+	if (Asset->IsA<UAnimSequence>())
+	{
+		UAnimSequence* AnimSequence = Cast<UAnimSequence>(Asset);
+		
+		if (TargetObjectGuid.IsValid())
+		{
+			USkeleton* Skeleton = AcquireSkeletonFromObjectGuid(TargetObjectGuid);
+
+			if (Skeleton && Skeleton == AnimSequence->GetSkeleton())
+			{
+				TArray<UObject*> OutObjects;
+				GetSequencer()->GetRuntimeObjects(GetSequencer()->GetFocusedMovieSceneSequenceInstance(), TargetObjectGuid, OutObjects);
+
+				AnimatablePropertyChanged(UMovieSceneSkeletalAnimationTrack::StaticClass(), false,
+					FOnKeyProperty::CreateRaw(this, &FSkeletalAnimationTrackEditor::AddKeyInternal, OutObjects, AnimSequence));
+
+				return true;
+			}
+		}
+	}
+	return false;
+}
+
+void FSkeletalAnimationTrackEditor::BuildObjectBindingContextMenu(FMenuBuilder& MenuBuilder, const FGuid& ObjectBinding, const UClass* ObjectClass)
+{
+	if (ObjectClass->IsChildOf(ASkeletalMeshActor::StaticClass()))
+	{
+		const TSharedPtr<ISequencer> ParentSequencer = GetSequencer();
+
+		USkeleton* Skeleton = AcquireSkeletonFromObjectGuid(ObjectBinding);
+
+		if (Skeleton)
+		{
+			// Load the asset registry module
+			FAssetRegistryModule& AssetRegistryModule = FModuleManager::LoadModuleChecked<FAssetRegistryModule>(TEXT("AssetRegistry"));
+
+			// Collect a full list of assets with the specified class
+			TArray<FAssetData> AssetDataList;
+			AssetRegistryModule.Get().GetAssetsByClass(UAnimSequence::StaticClass()->GetFName(), AssetDataList);
+
+			for( int32 AssetIndex = 0; AssetIndex < AssetDataList.Num(); ++AssetIndex )
+			{
+				const FAssetData& PossibleAnimSequence = AssetDataList[AssetIndex];
+				if( FAssetData(Skeleton).GetExportTextName() == PossibleAnimSequence.TagsAndValues.FindRef("Skeleton") )
+				{
+					FFormatNamedArguments Args;
+					Args.Add( TEXT("AnimationName"), FText::FromName( PossibleAnimSequence.AssetName ) );
+
+					UAnimSequence* AnimSequence = CastChecked<UAnimSequence>(AssetDataList[AssetIndex].GetAsset());
+					MenuBuilder.AddMenuEntry(
+						FText::Format( NSLOCTEXT("Sequencer", "AddAnimSequence", "Add {AnimationName}"), Args ),
+						FText::Format( NSLOCTEXT("Sequencer", "AddAnimSequenceTooltip", "Adds a {AnimationName} animation to this skeletal mesh."), Args ),
+						FSlateIcon(),
+						FUIAction(FExecuteAction::CreateSP(ParentSequencer.Get(), &ISequencer::AddAnimation, ObjectBinding, AnimSequence))
+						);
+				}
+			}
+		}
+	}
+}
+
+void FSkeletalAnimationTrackEditor::AddKeyInternal( float KeyTime, const TArray<UObject*> Objects, class UAnimSequence* AnimSequence )
+{
+	for( int32 ObjectIndex = 0; ObjectIndex < Objects.Num(); ++ObjectIndex )
+	{
+		UObject* Object = Objects[ObjectIndex];
+
+		FGuid ObjectHandle = FindOrCreateHandleToObject( Object );
+		if (ObjectHandle.IsValid())
+		{
+			UMovieSceneTrack* Track = GetTrackForObject( ObjectHandle, UMovieSceneSkeletalAnimationTrack::StaticClass(), FName("Animation"));
+
+			if (ensure(Track))
+			{
+				Cast<UMovieSceneSkeletalAnimationTrack>(Track)->AddNewAnimation( KeyTime, AnimSequence );
+			}
+		}
+	}
+}
+
+USkeleton* FSkeletalAnimationTrackEditor::AcquireSkeletonFromObjectGuid(const FGuid& Guid)
+{
+	TArray<UObject*> OutObjects;
+	GetSequencer()->GetRuntimeObjects(GetSequencer()->GetFocusedMovieSceneSequenceInstance(), Guid, OutObjects);
+
+	USkeleton* Skeleton = NULL;
+	for (int32 i = 0; i < OutObjects.Num(); ++i)
+	{
+		AActor* Actor = Cast<AActor>(OutObjects[i]);
+
+		if (Actor != NULL)
+		{
+			TInlineComponentArray<USkeletalMeshComponent*> SkeletalMeshComponents;
+			Actor->GetComponents(SkeletalMeshComponents);
+
+			for (int32 j = 0; j <SkeletalMeshComponents.Num(); ++j)
+			{
+				USkeletalMeshComponent* SkeletalMeshComp = SkeletalMeshComponents[j];
+				if (SkeletalMeshComp->SkeletalMesh && SkeletalMeshComp->SkeletalMesh->Skeleton)
+				{
+					// @todo Multiple actors, multiple components
+					check(!Skeleton);
+					Skeleton = SkeletalMeshComp->SkeletalMesh->Skeleton;
+				}
+			}
+		}
+	}
+
+	return Skeleton;
+}
diff --git a/Engine/Source/Editor/MovieSceneTools/Private/TrackEditors/SkeletalAnimationTrackEditor.h b/Engine/Source/Editor/MovieSceneTools/Private/TrackEditors/SkeletalAnimationTrackEditor.h
new file mode 100644
--- /dev/null
+++ b/Engine/Source/Editor/MovieSceneTools/Private/TrackEditors/SkeletalAnimationTrackEditor.h
@@ -0,0 +1,63 @@
+// Copyright 1998-2015 Epic Games, Inc. All Rights Reserved.
+
+#pragma once
+
+
+/**
+ * Tools for animation tracks
+ */
+class FSkeletalAnimationTrackEditor : public FMovieSceneTrackEditor
+{
+public:
+	/**
+	 * Constructor
+	 *
+	 * @param InSequencer	The sequencer instance to be used by this tool
+	 */
+	FSkeletalAnimationTrackEditor( TSharedRef<ISequencer> InSequencer );
+	~FSkeletalAnimationTrackEditor();
+
+	/**
+	 * Creates an instance of this class.  Called by a sequencer 
+	 *
+	 * @param OwningSequencer The sequencer instance to be used by this tool
+	 * @return The new instance of this class
+	 */
+	static TSharedRef<FMovieSceneTrackEditor> CreateTrackEditor( TSharedRef<ISequencer> OwningSequencer );
+
+	/** FMovieSceneTrackEditor Interface */
+	virtual bool SupportsType( TSubclassOf<UMovieSceneTrack> Type ) const override;
+	virtual TSharedRef<ISequencerSection> MakeSectionInterface( UMovieSceneSection& SectionObject, UMovieSceneTrack* Track ) override;
+	virtual void AddKey(const FGuid& ObjectGuid, UObject* AdditionalAsset) override;
+	virtual bool HandleAssetAdded(UObject* Asset, const FGuid& TargetObjectGuid) override;
+	virtual void BuildObjectBindingContextMenu(FMenuBuilder& MenuBuilder, const FGuid& ObjectBinding, const UClass* ObjectClass) override;
+
+private:
+	/** Delegate for AnimatablePropertyChanged in AddKey */
+	void AddKeyInternal(float KeyTime, const TArray<UObject*> Objects, class UAnimSequence* AnimSequence);
+
+	/** Gets a skeleton from an object guid in the movie scene */
+	class USkeleton* AcquireSkeletonFromObjectGuid(const FGuid& Guid);
+};
+
+
+
+/** Class for animation sections, handles drawing of all waveform previews */
+class FSkeletalAnimationSection : public ISequencerSection, public TSharedFromThis<FSkeletalAnimationSection>
+{
+public:
+	FSkeletalAnimationSection( UMovieSceneSection& InSection );
+	~FSkeletalAnimationSection();
+
+	/** ISequencerSection interface */
+	virtual UMovieSceneSection* GetSectionObject() override;
+	virtual FText GetDisplayName() const override;
+	virtual FText GetSectionTitle() const override;
+	virtual float GetSectionHeight() const override;
+	virtual void GenerateSectionLayout( class ISectionLayoutBuilder& LayoutBuilder ) const override {}
+	virtual int32 OnPaintSection( const FGeometry& AllottedGeometry, const FSlateRect& SectionClippingRect, FSlateWindowElementList& OutDrawElements, int32 LayerId, bool bParentEnabled ) const override;
+
+private:
+	/** The section we are visualizing */
+	UMovieSceneSection& Section;
+};
diff --git a/Engine/Source/Editor/MovieSceneTools/Private/TrackEditors/SubMovieSceneTrackEditor.cpp b/Engine/Source/Editor/MovieSceneTools/Private/TrackEditors/SubMovieSceneTrackEditor.cpp
--- a/Engine/Source/Editor/MovieSceneTools/Private/TrackEditors/SubMovieSceneTrackEditor.cpp
+++ b/Engine/Source/Editor/MovieSceneTools/Private/TrackEditors/SubMovieSceneTrackEditor.cpp
@@ -4,6 +4,7 @@
 #include "SubMovieSceneTrackEditor.h"
 #include "SubMovieSceneTrack.h"
 #include "SubMovieSceneSection.h"
+#include "MovieSceneSequence.h"
 
 /**
  * A generic implementation for displaying simple property sections
@@ -29,7 +30,7 @@ public:
 
 	virtual FText GetSectionTitle() const override
 	{
-		return FText::FromString( SectionObject.GetMovieScene()->GetName() );
+		return FText::FromString( SectionObject.GetMovieSceneAnimation()->GetName() );
 	}
 
 	virtual void GenerateSectionLayout( class ISectionLayoutBuilder& LayoutBuilder ) const override
@@ -69,7 +70,7 @@ private:
 	/** The section we are visualizing */
 	USubMovieSceneSection& SectionObject;
 	/** The instance that this section is part of */
-	TWeakPtr<FMovieSceneInstance> MovieSceneInstance;
+	TWeakPtr<FMovieSceneSequenceInstance> MovieSceneInstance;
 	/** Sequencer interface */
 	TWeakPtr<ISequencer> Sequencer;
 };
@@ -99,9 +100,11 @@ TSharedRef<ISequencerSection> FSubMovieSceneTrackEditor::MakeSectionInterface( U
 
 bool FSubMovieSceneTrackEditor::HandleAssetAdded(UObject* Asset, const FGuid& TargetObjectGuid)
 {
-	if (Asset->IsA<UMovieScene>())
+	UMovieSceneSequence* Animation = Cast<UMovieSceneSequence>( Asset );
+
+	if( Animation )
 	{
-		GetSequencer()->AddSubMovieScene( CastChecked<UMovieScene>( Asset ) );
+		GetSequencer()->AddSubMovieScene( Animation );
 
 		return true;
 	}
diff --git a/Engine/Source/Editor/MovieSceneTools/Private/TrackEditors/TransformTrackEditor.cpp b/Engine/Source/Editor/MovieSceneTools/Private/TrackEditors/TransformTrackEditor.cpp
--- a/Engine/Source/Editor/MovieSceneTools/Private/TrackEditors/TransformTrackEditor.cpp
+++ b/Engine/Source/Editor/MovieSceneTools/Private/TrackEditors/TransformTrackEditor.cpp
@@ -174,8 +174,8 @@ TSharedRef<ISequencerSection> F3DTransformTrackEditor::MakeSectionInterface( UMo
 
 void F3DTransformTrackEditor::OnPreTransformChanged( UObject& InObject )
 {
-	UMovieScene* MovieScene = GetMovieScene();
-	float AutoKeyTime = GetTimeForKey( MovieScene );
+	UMovieSceneSequence* MovieSceneSequence = GetMovieSceneSequence();
+	float AutoKeyTime = GetTimeForKey( MovieSceneSequence );
 
 	if( IsAllowedToAutoKey() )
 	{
@@ -273,7 +273,7 @@ void F3DTransformTrackEditor::AddKey(const FGuid& ObjectGuid, UObject* Additiona
 void F3DTransformTrackEditor::AddKeyInternal(const FGuid& ObjectGuid, UObject* AdditionalAsset, bool bForceKey, F3DTransformTrackKey::Type KeyType)
 {
 	TArray<UObject*> OutObjects;
-	GetSequencer()->GetRuntimeObjects( GetSequencer()->GetFocusedMovieSceneInstance(), ObjectGuid, OutObjects);
+	GetSequencer()->GetRuntimeObjects( GetSequencer()->GetFocusedMovieSceneSequenceInstance(), ObjectGuid, OutObjects);
 
 	for ( UObject* Object : OutObjects )
 	{
@@ -376,7 +376,7 @@ void F3DTransformTrackEditor::BindCommands(TSharedRef<FUICommandList> SequencerC
 void F3DTransformTrackEditor::BuildObjectBindingEditButtons(TSharedPtr<SHorizontalBox> EditBox, const FGuid& ObjectGuid, const UClass* ObjectClass)
 {
 	TArray<UObject*> OutObjects;
-	GetSequencer()->GetRuntimeObjects( GetSequencer()->GetFocusedMovieSceneInstance(), ObjectGuid, OutObjects);
+	GetSequencer()->GetRuntimeObjects( GetSequencer()->GetFocusedMovieSceneSequenceInstance(), ObjectGuid, OutObjects);
 
 	TWeakObjectPtr<ACameraActor> CameraActor;
 
diff --git a/Engine/Source/Editor/NiagaraEditor/Private/NiagaraEffectEditor.cpp b/Engine/Source/Editor/NiagaraEditor/Private/NiagaraEffectEditor.cpp
--- a/Engine/Source/Editor/NiagaraEditor/Private/NiagaraEffectEditor.cpp
+++ b/Engine/Source/Editor/NiagaraEditor/Private/NiagaraEffectEditor.cpp
@@ -101,7 +101,7 @@ void FNiagaraEffectEditor::InitNiagaraEffectEditor(const EToolkitMode::Type Mode
 		FSequencerInitParams SequencerInitParams;
 		{
 			SequencerInitParams.ViewParams = ViewParams;
-			SequencerInitParams.Animation = NewAnimation;
+			SequencerInitParams.RootSequence = NewAnimation;
 			SequencerInitParams.bEditWithinLevelEditor = false;
 			SequencerInitParams.ToolkitHost = nullptr;
 		}
diff --git a/Engine/Source/Editor/Sequencer/Private/SAnimationOutlinerTreeNode.cpp b/Engine/Source/Editor/Sequencer/Private/SAnimationOutlinerTreeNode.cpp
--- a/Engine/Source/Editor/Sequencer/Private/SAnimationOutlinerTreeNode.cpp
+++ b/Engine/Source/Editor/Sequencer/Private/SAnimationOutlinerTreeNode.cpp
@@ -169,7 +169,7 @@ FReply SAnimationOutlinerTreeNode::OnPreviousKeyClicked()
 {
 	FSequencer& Sequencer = DisplayNode->GetSequencer();
 	float ClosestPreviousKeyDistance = MAX_FLT;
-	float CurrentTime = Sequencer.GetCurrentLocalTime(*Sequencer.GetFocusedMovieScene());
+	float CurrentTime = Sequencer.GetCurrentLocalTime(*Sequencer.GetFocusedMovieSceneSequence());
 	float PreviousTime = 0;
 	bool PreviousKeyFound = false;
 
@@ -200,7 +200,7 @@ FReply SAnimationOutlinerTreeNode::OnNextKeyClicked()
 {
 	FSequencer& Sequencer = DisplayNode->GetSequencer();
 	float ClosestNextKeyDistance = MAX_FLT;
-	float CurrentTime = Sequencer.GetCurrentLocalTime(*Sequencer.GetFocusedMovieScene());
+	float CurrentTime = Sequencer.GetCurrentLocalTime(*Sequencer.GetFocusedMovieSceneSequence());
 	float NextTime = 0;
 	bool NextKeyFound = false;
 
@@ -230,7 +230,7 @@ FReply SAnimationOutlinerTreeNode::OnNextKeyClicked()
 FReply SAnimationOutlinerTreeNode::OnAddKeyClicked()
 {
 	FSequencer& Sequencer = DisplayNode->GetSequencer();
-	float CurrentTime = Sequencer.GetCurrentLocalTime(*Sequencer.GetFocusedMovieScene());
+	float CurrentTime = Sequencer.GetCurrentLocalTime(*Sequencer.GetFocusedMovieSceneSequence());
 
 	TSet<TSharedPtr<IKeyArea>> KeyAreas;
 	SequencerHelpers::GetAllKeyAreas(DisplayNode, KeyAreas);
@@ -485,7 +485,7 @@ void SAnimationOutlinerTreeNode::OnSelectionChanged( TArray<TSharedPtr<FSequence
 
 		// Get the bound objects
 		TArray<UObject*> RuntimeObjects;
-		Sequencer.GetRuntimeObjects( Sequencer.GetFocusedMovieSceneInstance(), ObjectNode->GetObjectBinding(), RuntimeObjects );
+		Sequencer.GetRuntimeObjects( Sequencer.GetFocusedMovieSceneSequenceInstance(), ObjectNode->GetObjectBinding(), RuntimeObjects );
 		
 		if( RuntimeObjects.Num() > 0 )
 		{
diff --git a/Engine/Source/Editor/Sequencer/Private/SSequencer.cpp b/Engine/Source/Editor/Sequencer/Private/SSequencer.cpp
--- a/Engine/Source/Editor/Sequencer/Private/SSequencer.cpp
+++ b/Engine/Source/Editor/Sequencer/Private/SSequencer.cpp
@@ -26,7 +26,7 @@
 #include "SNumericDropDown.h"
 #include "EditorWidgetsModule.h"
 #include "SSequencerTreeView.h"
-
+#include "MovieSceneSequence.h"
 #include "SSequencerSplitterOverlay.h"
 #include "IKeyArea.h"
 
@@ -733,7 +733,7 @@ void SSequencer::UpdateLayoutTree()
 
 void SSequencer::UpdateBreadcrumbs(const TArray< TWeakObjectPtr<class UMovieSceneSection> >& FilteringShots)
 {
-	TSharedRef<FMovieSceneInstance> FocusedMovieSceneInstance = Sequencer.Pin()->GetFocusedMovieSceneInstance();
+	TSharedRef<FMovieSceneSequenceInstance> FocusedMovieSceneInstance = Sequencer.Pin()->GetFocusedMovieSceneSequenceInstance();
 
 	if (BreadcrumbTrail->PeekCrumb().BreadcrumbType == FSequencerBreadcrumb::ShotType)
 	{
@@ -742,8 +742,9 @@ void SSequencer::UpdateBreadcrumbs(const TArray< TWeakObjectPtr<class UMovieScen
 
 	if( BreadcrumbTrail->PeekCrumb().BreadcrumbType == FSequencerBreadcrumb::MovieSceneType && BreadcrumbTrail->PeekCrumb().MovieSceneInstance.Pin() != FocusedMovieSceneInstance )
 	{
+		FText CrumbName = FText::FromString( FocusedMovieSceneInstance->GetSequence()->GetName() );
 		// The current breadcrumb is not a moviescene so we need to make a new breadcrumb in order return to the parent moviescene later
-		BreadcrumbTrail->PushCrumb( FText::FromString(FocusedMovieSceneInstance->GetMovieScene()->GetName()), FSequencerBreadcrumb( FocusedMovieSceneInstance ) );
+		BreadcrumbTrail->PushCrumb( CrumbName, FSequencerBreadcrumb( FocusedMovieSceneInstance ) );
 	}
 
 	if (Sequencer.Pin()->IsShotFilteringOn())
@@ -767,7 +768,7 @@ void SSequencer::UpdateBreadcrumbs(const TArray< TWeakObjectPtr<class UMovieScen
 void SSequencer::ResetBreadcrumbs()
 {
 	BreadcrumbTrail->ClearCrumbs();
-	BreadcrumbTrail->PushCrumb(TAttribute<FText>::Create(TAttribute<FText>::FGetter::CreateSP(this, &SSequencer::GetRootMovieSceneName)), FSequencerBreadcrumb(Sequencer.Pin()->GetRootMovieSceneInstance()));
+	BreadcrumbTrail->PushCrumb(TAttribute<FText>::Create(TAttribute<FText>::FGetter::CreateSP(this, &SSequencer::GetRootAnimationName)), FSequencerBreadcrumb(Sequencer.Pin()->GetRootMovieSceneSequenceInstance()));
 }
 
 void SSequencer::OnOutlinerSearchChanged( const FText& Filter )
@@ -949,7 +950,7 @@ void SSequencer::OnAssetsDropped( const FAssetDragDropOp& DragDropOp )
 
 	if( bSpawnableAdded )
 	{
-		SequencerRef.SpawnOrDestroyPuppetObjects( SequencerRef.GetFocusedMovieSceneInstance() );
+		SequencerRef.SpawnOrDestroyPuppetObjects( SequencerRef.GetFocusedMovieSceneSequenceInstance() );
 	}
 
 	if( bObjectAdded )
@@ -981,7 +982,7 @@ void SSequencer::OnClassesDropped( const FClassDragDropOp& DragDropOp )
 
 	if( bSpawnableAdded )
 	{
-		SequencerRef.SpawnOrDestroyPuppetObjects( SequencerRef.GetFocusedMovieSceneInstance() );
+		SequencerRef.SpawnOrDestroyPuppetObjects( SequencerRef.GetFocusedMovieSceneSequenceInstance() );
 
 		// Update the sequencers view of the movie scene data
 		SequencerRef.NotifyMovieSceneDataChanged();
@@ -1036,7 +1037,7 @@ void SSequencer::OnUnloadedClassesDropped( const FUnloadedClassDragDropOp& DragD
 
 	if( bSpawnableAdded )
 	{
-		SequencerRef.SpawnOrDestroyPuppetObjects( SequencerRef.GetFocusedMovieSceneInstance() );
+		SequencerRef.SpawnOrDestroyPuppetObjects( SequencerRef.GetFocusedMovieSceneSequenceInstance() );
 
 		SequencerRef.NotifyMovieSceneDataChanged();
 	}
@@ -1073,7 +1074,7 @@ void SSequencer::OnActorSelectionChanged(UObject*)
 	{
 		TSharedRef<FObjectBindingNode> ObjectBindingNode = StaticCastSharedRef<FObjectBindingNode>(Node);
 		TArray<UObject*> RuntimeObjects;
-		Sequencer.Pin()->GetRuntimeObjects( Sequencer.Pin()->GetFocusedMovieSceneInstance(), ObjectBindingNode->GetObjectBinding(), RuntimeObjects );
+		Sequencer.Pin()->GetRuntimeObjects( Sequencer.Pin()->GetFocusedMovieSceneSequenceInstance(), ObjectBindingNode->GetObjectBinding(), RuntimeObjects );
 		
 		bool bDoSelect = false;
 		for (int32 RuntimeIndex = 0; RuntimeIndex < RuntimeObjects.Num(); ++RuntimeIndex )
@@ -1103,7 +1104,7 @@ void SSequencer::OnCrumbClicked(const FSequencerBreadcrumb& Item)
 {
 	if (Item.BreadcrumbType != FSequencerBreadcrumb::ShotType)
 	{
-		if( Sequencer.Pin()->GetFocusedMovieSceneInstance() == Item.MovieSceneInstance.Pin() ) 
+		if( Sequencer.Pin()->GetFocusedMovieSceneSequenceInstance() == Item.MovieSceneInstance.Pin() ) 
 		{
 			// then do zooming
 		}
@@ -1121,9 +1122,9 @@ void SSequencer::OnCrumbClicked(const FSequencerBreadcrumb& Item)
 	}
 }
 
-FText SSequencer::GetRootMovieSceneName() const
+FText SSequencer::GetRootAnimationName() const
 {
-	return FText::FromString(Sequencer.Pin()->GetRootMovieScene()->GetName());
+	return FText::FromName(Sequencer.Pin()->GetRootMovieSceneSequence()->GetFName());
 }
 
 FText SSequencer::GetShotSectionTitle(UMovieSceneSection* ShotSection) const
@@ -1211,7 +1212,7 @@ void SSequencer::StepToKey(bool bStepToNextKey, bool bCameraOnly)
 		{
 			TSharedRef<FObjectBindingNode> ObjectBindingNode = StaticCastSharedRef<FObjectBindingNode>(RootNode);
 			TArray<UObject*> RuntimeObjects;
-			Sequencer.Pin()->GetRuntimeObjects( Sequencer.Pin()->GetFocusedMovieSceneInstance(), ObjectBindingNode->GetObjectBinding(), RuntimeObjects );
+			Sequencer.Pin()->GetRuntimeObjects( Sequencer.Pin()->GetFocusedMovieSceneSequenceInstance(), ObjectBindingNode->GetObjectBinding(), RuntimeObjects );
 		
 			for (int32 RuntimeIndex = 0; RuntimeIndex < RuntimeObjects.Num(); ++RuntimeIndex )
 			{
@@ -1243,7 +1244,7 @@ void SSequencer::StepToKey(bool bStepToNextKey, bool bCameraOnly)
 	if (Nodes.Num() > 0)
 	{
 		float ClosestKeyDistance = MAX_FLT;
-		float CurrentTime = Sequencer.Pin()->GetCurrentLocalTime(*Sequencer.Pin()->GetFocusedMovieScene());
+		float CurrentTime = Sequencer.Pin()->GetCurrentLocalTime(*Sequencer.Pin()->GetFocusedMovieSceneSequence());
 		float StepToTime = 0;
 		bool StepToKeyFound = false;
 
diff --git a/Engine/Source/Editor/Sequencer/Private/SSequencer.h b/Engine/Source/Editor/Sequencer/Private/SSequencer.h
--- a/Engine/Source/Editor/Sequencer/Private/SSequencer.h
+++ b/Engine/Source/Editor/Sequencer/Private/SSequencer.h
@@ -32,9 +32,9 @@ struct FSequencerBreadcrumb
 	/** The type of breadcrumb this is */
 	FSequencerBreadcrumb::Type BreadcrumbType;
 	/** The movie scene this may point to */
-	class TWeakPtr<FMovieSceneInstance> MovieSceneInstance;
+	class TWeakPtr<FMovieSceneSequenceInstance> MovieSceneInstance;
 
-	FSequencerBreadcrumb( TSharedRef<class FMovieSceneInstance> InMovieSceneInstance )
+	FSequencerBreadcrumb( TSharedRef<class FMovieSceneSequenceInstance> InMovieSceneInstance )
 		: BreadcrumbType(FSequencerBreadcrumb::MovieSceneType)
 		, MovieSceneInstance(InMovieSceneInstance) {}
 
@@ -223,7 +223,7 @@ private:
 	void OnCrumbClicked(const FSequencerBreadcrumb& Item);
 
 	/** Gets the root movie scene name */
-	FText GetRootMovieSceneName() const;
+	FText GetRootAnimationName() const;
 
 	/** Gets the title of the passed in shot section */
 	FText GetShotSectionTitle(UMovieSceneSection* ShotSection) const;
diff --git a/Engine/Source/Editor/Sequencer/Private/Sequencer.cpp b/Engine/Source/Editor/Sequencer/Private/Sequencer.cpp
--- a/Engine/Source/Editor/Sequencer/Private/Sequencer.cpp
+++ b/Engine/Source/Editor/Sequencer/Private/Sequencer.cpp
@@ -3,7 +3,7 @@
 #include "SequencerPrivatePCH.h"
 #include "Sequencer.h"
 #include "SequencerEdMode.h"
-#include "MovieSceneAnimation.h"
+#include "MovieSceneSequence.h"
 #include "MovieScene.h"
 #include "Engine/LevelScriptBlueprint.h"
 #include "Editor/PropertyEditor/Public/PropertyEditorModule.h"
@@ -21,7 +21,7 @@
 #include "MovieSceneTrack.h"
 #include "MovieSceneShotTrack.h"
 #include "MovieSceneAudioTrack.h"
-#include "MovieSceneAnimationTrack.h"
+#include "MovieSceneSkeletalAnimationTrack.h"
 #include "MovieSceneTrackEditor.h"
 #include "Kismet2/KismetEditorUtilities.h"
 #include "Kismet2/BlueprintEditorUtils.h"
@@ -30,7 +30,7 @@
 #include "MovieSceneShotSection.h"
 #include "SubMovieSceneTrack.h"
 #include "ISequencerSection.h"
-#include "MovieSceneInstance.h"
+#include "MovieSceneSequenceInstance.h"
 #include "IKeyArea.h"
 #include "SnappingUtils.h"
 #include "GenericCommands.h"
@@ -85,17 +85,14 @@ void FSequencer::InitSequencer(const FSequencerInitParams& InitParams, const TSh
 		LastViewRange = TargetViewRange = InitParams.ViewParams.InitalViewRange;
 		ScrubPosition = InitParams.ViewParams.InitialScrubPosition;
 		ObjectChangeListener = InObjectChangeListener;
-		Animation = InitParams.Animation;
 
 		check( ObjectChangeListener.IsValid() );
-		check( Animation != nullptr );
 		
-		UMovieScene& RootMovieScene = *InitParams.Animation->GetMovieScene();
-		RootMovieScene.SetFlags( RF_Transactional );
+		UMovieSceneSequence& RootSequence = *InitParams.RootSequence;
 		
 		// Focusing the initial movie scene needs to be done before the first time GetFocusedMovieSceneInstane or GetRootMovieSceneInstance is used
-		RootMovieSceneInstance = MakeShareable( new FMovieSceneInstance( RootMovieScene ) );
-		MovieSceneStack.Add( RootMovieSceneInstance.ToSharedRef() );
+		RootMovieSceneSequenceInstance = MakeShareable( new FMovieSceneSequenceInstance( RootSequence ) );
+		MovieSceneStack.Add( RootMovieSceneSequenceInstance.ToSharedRef() );
 
 		// Make internal widgets
 		SequencerWidget = SNew( SSequencer, SharedThis( this ) )
@@ -131,9 +128,9 @@ void FSequencer::InitSequencer(const FSequencerInitParams& InitParams, const TSh
 			{
 				ShotTrackEditor = TrackEditor;
 			}
-			else if ( TrackEditor->SupportsType( UMovieSceneAnimationTrack::StaticClass() ) )
+			else if ( TrackEditor->SupportsType( UMovieSceneSkeletalAnimationTrack::StaticClass() ) )
 			{
-				AnimationTrackEditor = TrackEditor;
+				SkeletalAnimationTrackEditor = TrackEditor;
 			}
 
 			TrackEditors.Add( TrackEditor );
@@ -187,9 +184,8 @@ FSequencer::~FSequencer()
 	GEditor->GetActorRecordingState().RemoveAll( this );
 	GEditor->UnregisterForUndo( this );
 
-	DestroySpawnablesForAllMovieScenes();
+	DestroySpawnablesForAllSequences();
 
-	Animation = nullptr;
 	TrackEditors.Empty();
 	SequencerWidget.Reset();
 }
@@ -198,7 +194,7 @@ void FSequencer::Close()
 {
 	if (ActiveSequencer.IsValid() && ActiveSequencer.Get() == this)
 	{
-		RootMovieSceneInstance->RestoreState();
+		RootMovieSceneSequenceInstance->RestoreState();
 
 		FEditorDelegates::PreSaveWorld.RemoveAll(this);
 		FEditorDelegates::PostSaveWorld.RemoveAll(this);
@@ -284,36 +280,21 @@ void FSequencer::Tick(float InDeltaTime)
 	}
 }
 
-TArray< UMovieScene* > FSequencer::GetMovieScenesBeingEdited()
-{
-	TArray<UMovieScene* > OutMovieScenes;
-
-	// Get the root movie scene
-	OutMovieScenes.Add( RootMovieSceneInstance->GetMovieScene() );
-
-	// Get Sub-MovieScenes
-	for( auto It = MovieSceneSectionToInstanceMap.CreateConstIterator(); It; ++It )
-	{
-		OutMovieScenes.AddUnique( It.Value()->GetMovieScene() );
-	}
-
-	return OutMovieScenes;
-}
 
-UMovieScene* FSequencer::GetRootMovieScene() const
+UMovieSceneSequence* FSequencer::GetRootMovieSceneSequence() const
 {
-	return MovieSceneStack[0]->GetMovieScene();
+	return MovieSceneStack[0]->GetSequence();
 }
 
-UMovieScene* FSequencer::GetFocusedMovieScene() const
+UMovieSceneSequence* FSequencer::GetFocusedMovieSceneSequence() const
 {
 	// the last item is the focused movie scene
-	return MovieSceneStack.Top()->GetMovieScene();
+	return MovieSceneStack.Top()->GetSequence();
 }
 
-void FSequencer::ResetToNewAnimation(UMovieSceneAnimation& NewAnimation)
+void FSequencer::ResetToNewRootSequence(UMovieSceneSequence& NewSequence)
 {
-	DestroySpawnablesForAllMovieScenes();
+	DestroySpawnablesForAllSequences();
 
 	//@todo Sequencer - Encapsulate this better
 	MovieSceneStack.Empty();
@@ -323,14 +304,9 @@ void FSequencer::ResetToNewAnimation(UMovieSceneAnimation& NewAnimation)
 	UnfilterableObjects.Empty();
 	MovieSceneSectionToInstanceMap.Empty();
 
-	NewAnimation.GetMovieScene()->SetFlags(RF_Transactional);
-
-	Animation = &NewAnimation;
-	check(Animation != nullptr);
-
-	// Focusing the initial movie scene needs to be done before the first time GetFocusedMovieSceneInstane or GetRootMovieSceneInstance is used
-	RootMovieSceneInstance = MakeShareable(new FMovieSceneInstance(*Animation->GetMovieScene()));
-	MovieSceneStack.Add(RootMovieSceneInstance.ToSharedRef());
+	// Focusing the initial movie scene needs to be done before the first time NewSequence or GetRootMovieSceneInstance is used
+	RootMovieSceneSequenceInstance = MakeShareable(new FMovieSceneSequenceInstance(NewSequence));
+	MovieSceneStack.Add(RootMovieSceneSequenceInstance.ToSharedRef());
 
 	SequencerWidget->ResetBreadcrumbs();
 
@@ -339,20 +315,20 @@ void FSequencer::ResetToNewAnimation(UMovieSceneAnimation& NewAnimation)
 	UpdateTimeBoundsToFocusedMovieScene();
 }
 
-TSharedRef<FMovieSceneInstance> FSequencer::GetRootMovieSceneInstance() const
+TSharedRef<FMovieSceneSequenceInstance> FSequencer::GetRootMovieSceneSequenceInstance() const
 {
 	return MovieSceneStack[0];
 }
 
 
-TSharedRef<FMovieSceneInstance> FSequencer::GetFocusedMovieSceneInstance() const
+TSharedRef<FMovieSceneSequenceInstance> FSequencer::GetFocusedMovieSceneSequenceInstance() const
 {
 	// the last item is the focused movie scene
 	return MovieSceneStack.Top();
 }
 
 
-void FSequencer::FocusSubMovieScene( TSharedRef<FMovieSceneInstance> SubMovieSceneInstance )
+void FSequencer::FocusSubMovieScene( TSharedRef<FMovieSceneSequenceInstance> SubMovieSceneInstance )
 {
 	// Check for infinite recursion
 	check( SubMovieSceneInstance != MovieSceneStack.Top() );
@@ -369,12 +345,12 @@ void FSequencer::FocusSubMovieScene( TSharedRef<FMovieSceneInstance> SubMovieSce
 	UpdateTimeBoundsToFocusedMovieScene();
 }
 
-TSharedRef<FMovieSceneInstance> FSequencer::GetInstanceForSubMovieSceneSection( UMovieSceneSection& SubMovieSceneSection ) const
+TSharedRef<FMovieSceneSequenceInstance> FSequencer::GetInstanceForSubMovieSceneSection( UMovieSceneSection& SubMovieSceneSection ) const
 {
 	return MovieSceneSectionToInstanceMap.FindChecked( &SubMovieSceneSection );
 }
 
-void FSequencer::PopToMovieScene( TSharedRef<FMovieSceneInstance> SubMovieSceneInstance )
+void FSequencer::PopToMovieScene( TSharedRef<FMovieSceneSequenceInstance> SubMovieSceneInstance )
 {
 	if( MovieSceneStack.Num() > 1 )
 	{
@@ -396,6 +372,7 @@ void FSequencer::PopToMovieScene( TSharedRef<FMovieSceneInstance> SubMovieSceneI
 
 void FSequencer::AddNewShot(FGuid CameraGuid)
 {
+	// @todo Sequencer: This should not be needed here.  the track editor could be the entry point
 	if (ShotTrackEditor.IsValid())
 	{
 		ShotTrackEditor.Pin()->AddKey(CameraGuid);
@@ -404,15 +381,16 @@ void FSequencer::AddNewShot(FGuid CameraGuid)
 
 void FSequencer::AddAnimation(FGuid ObjectGuid, class UAnimSequence* AnimSequence)
 {
-	if (AnimationTrackEditor.IsValid())
+	// @todo Sequencer: This should not be needed here.  the track editor could be the entry point
+	if (SkeletalAnimationTrackEditor.IsValid())
 	{
-		AnimationTrackEditor.Pin()->AddKey(ObjectGuid, AnimSequence);
+		SkeletalAnimationTrackEditor.Pin()->AddKey(ObjectGuid, AnimSequence);
 	}
 }
 
 void FSequencer::DeleteSection(class UMovieSceneSection* Section)
 {
-	UMovieScene* MovieScene = GetFocusedMovieScene();
+	UMovieScene* MovieScene = GetFocusedMovieSceneSequence()->GetMovieScene();
 	
 	bool bAnythingRemoved = false;
 
@@ -464,11 +442,12 @@ void FSequencer::DeleteSelectedKeys()
 	Selection.EmptySelectedKeys();
 }
 
-void FSequencer::SpawnOrDestroyPuppetObjects( TSharedRef<FMovieSceneInstance> MovieSceneInstance )
+void FSequencer::SpawnOrDestroyPuppetObjects( TSharedRef<FMovieSceneSequenceInstance> MovieSceneInstance )
 {
-	if( Animation->AllowsSpawnableObjects() )
+	UMovieSceneSequence* Sequence = MovieSceneInstance->GetSequence();
+	if( Sequence->AllowsSpawnableObjects() )
 	{
-		Animation->SpawnOrDestroyObjects(MovieSceneInstance->GetMovieScene(), false /*DestroyAll*/);
+		Sequence->SpawnOrDestroyObjects(false /*DestroyAll*/);
 	}
 }
 
@@ -477,13 +456,13 @@ void FSequencer::NotifyMapChanged( class UWorld* NewWorld, EMapChangeType MapCha
 	// @todo Sequencer Sub-MovieScenes Needs more investigation of what to spawn
 	// Destroy our puppets because the world is going away.  We probably don't have to do this (the actors will
 	// be naturally destroyed with the level, but we might as well.)
-	DestroySpawnablesForAllMovieScenes();
+	DestroySpawnablesForAllSequences();
 
 
 	// @todo sequencer: We should only wipe/respawn puppets that are affected by the world that is being changed! (multi-UWorld support)
 	if( ( MapChangeType == EMapChangeType::LoadMap || MapChangeType == EMapChangeType::NewMap ) )
 	{
-		SpawnOrDestroyPuppetObjects( GetFocusedMovieSceneInstance() );
+		SpawnOrDestroyPuppetObjects( GetFocusedMovieSceneSequenceInstance() );
 
 		NotifyMovieSceneDataChanged();
 	}
@@ -492,8 +471,8 @@ void FSequencer::NotifyMapChanged( class UWorld* NewWorld, EMapChangeType MapCha
 void FSequencer::OnActorsDropped( const TArray<TWeakObjectPtr<AActor> >& Actors )
 {
 	bool bPossessableAdded = false;
-	UMovieScene* OwnerMovieScene = GetFocusedMovieScene();
-
+	UMovieSceneSequence* OwnerSequence = GetFocusedMovieSceneSequence();
+	UMovieScene* OwnerMovieScene = OwnerSequence->GetMovieScene();
 	for (TWeakObjectPtr<AActor> WeakActor : Actors)
 	{
 		AActor* Actor = WeakActor.Get();
@@ -504,7 +483,7 @@ void FSequencer::OnActorsDropped( const TArray<TWeakObjectPtr<AActor> >& Actors
 			const FScopedTransaction Transaction(LOCTEXT("UndoPossessingObject", "Possess Object with MovieScene"));
 			
 			// possess the object!
-			OwnerMovieScene->Modify();
+			OwnerSequence->Modify();
 				
 			const FGuid PossessableGuid = OwnerMovieScene->AddPossessable(Actor->GetActorLabel(), Actor->GetClass());
 
@@ -513,14 +492,14 @@ void FSequencer::OnActorsDropped( const TArray<TWeakObjectPtr<AActor> >& Actors
 				AddUnfilterableObject(PossessableGuid);
 			}
 
-			Animation->BindPossessableObject(PossessableGuid, *Actor);
+			OwnerSequence->BindPossessableObject(PossessableGuid, *Actor);
 			bPossessableAdded = true;
 		}
 	}
 	
 	if (bPossessableAdded)
 	{
-		SpawnOrDestroyPuppetObjects(GetFocusedMovieSceneInstance());
+		SpawnOrDestroyPuppetObjects(GetFocusedMovieSceneSequenceInstance());
 		NotifyMovieSceneDataChanged();
 	}
 }
@@ -547,7 +526,7 @@ FAnimatedRange FSequencer::GetViewRange() const
 
 FAnimatedRange FSequencer::GetClampRange() const
 {
-	return TRange<float>(GetFocusedMovieScene()->StartTime, GetFocusedMovieScene()->EndTime);
+	return TRange<float>(GetFocusedMovieSceneSequence()->GetMovieScene()->GetTimeRange() );
 }
 
 bool FSequencer::GetAutoKeyEnabled() const 
@@ -565,7 +544,7 @@ bool FSequencer::IsRecordingLive() const
 	return PlaybackState == EMovieScenePlayerStatus::Recording && GIsPlayInEditorWorld;
 }
 
-float FSequencer::GetCurrentLocalTime( UMovieScene& MovieScene )
+float FSequencer::GetCurrentLocalTime( UMovieSceneSequence& InMovieSceneSequence )
 {
 	//@todo Sequencer - Nested movie scenes:  Figure out the parent of the passed in movie scene and 
 	// calculate local time
@@ -597,13 +576,14 @@ void FSequencer::SetGlobalTime( float NewTime )
 
 	if (Settings->GetLockInOutToStartEndRange())
 	{
-		NewTime = FMath::Clamp(NewTime, GetFocusedMovieScene()->StartTime, GetFocusedMovieScene()->EndTime);
+		UMovieScene* FocusedMovieScene = GetFocusedMovieSceneSequence()->GetMovieScene();
+		NewTime = FMath::Clamp(NewTime, FocusedMovieScene->StartTime, FocusedMovieScene->EndTime);
 	}
 
 	// Update the position
 	ScrubPosition = NewTime;
 
-	RootMovieSceneInstance->Update( ScrubPosition, LastTime, *this );
+	RootMovieSceneSequenceInstance->Update( ScrubPosition, LastTime, *this );
 
 	GEditor->RedrawLevelEditingViewports();
 }
@@ -645,9 +625,11 @@ FGuid FSequencer::GetHandleToObject( UObject* Object )
 		return FGuid();
 	}
 
-	TSharedRef<FMovieSceneInstance> FocusedMovieSceneInstance = GetFocusedMovieSceneInstance();
-	UMovieScene* FocusedMovieScene = FocusedMovieSceneInstance->GetMovieScene();
-	FGuid ObjectGuid = Animation->FindObjectId(*Object);
+	TSharedRef<FMovieSceneSequenceInstance> FocusedMovieSceneSequenceInstance = GetFocusedMovieSceneSequenceInstance();
+	UMovieSceneSequence* FocusedMovieSceneSequence = FocusedMovieSceneSequenceInstance->GetSequence();
+
+	UMovieScene* FocusedMovieScene = FocusedMovieSceneSequence->GetMovieScene();
+	FGuid ObjectGuid = FocusedMovieSceneSequence->FindObjectId(*Object);
 
 	// Check here for spawnable otherwise spawnables get recreated as possessables, which doesn't make sens
 	FMovieSceneSpawnable* Spawnable = FocusedMovieScene->FindSpawnable(ObjectGuid);
@@ -664,7 +646,7 @@ FGuid FSequencer::GetHandleToObject( UObject* Object )
 		FMovieScenePossessable* Possessable = FocusedMovieScene->FindPossessable(ObjectGuid);
 		if(Possessable == nullptr)
 		{
-			Animation->UnbindPossessableObjects(ObjectGuid);
+			FocusedMovieSceneSequence->UnbindPossessableObjects(ObjectGuid);
 			ObjectGuid.Invalidate();
 		}
 	}
@@ -673,7 +655,7 @@ FGuid FSequencer::GetHandleToObject( UObject* Object )
 
 	// If the object guid was not found attempt to add it
 	// Note: Only possessed actors can be added like this
-	if (!ObjectGuid.IsValid() && Animation->CanPossessObject(*Object))
+	if (!ObjectGuid.IsValid() && FocusedMovieSceneSequence->CanPossessObject(*Object))
 	{
 		// @todo sequencer: Undo doesn't seem to be working at all
 		const FScopedTransaction Transaction(LOCTEXT("UndoPossessingObject", "Possess Object with MovieScene"));
@@ -690,12 +672,12 @@ FGuid FSequencer::GetHandleToObject( UObject* Object )
 				AddUnfilterableObject(ObjectGuid);
 			}
 
-			Animation->BindPossessableObject(ObjectGuid, *Object);
+			FocusedMovieSceneSequence->BindPossessableObject(ObjectGuid, *Object);
 			bPossessableAdded = true;
 		}
 
 		// If we're adding a possessable, generate handles to it's parent objects too.
-		UObject* ParentObject = Animation->GetParentObject(Object);
+		UObject* ParentObject = FocusedMovieSceneSequence->GetParentObject(Object);
 		if ( ParentObject != nullptr )
 		{
 			GetHandleToObject(ParentObject);
@@ -704,7 +686,7 @@ FGuid FSequencer::GetHandleToObject( UObject* Object )
 
 	if (bPossessableAdded)
 	{
-		SpawnOrDestroyPuppetObjects(GetFocusedMovieSceneInstance());
+		SpawnOrDestroyPuppetObjects(GetFocusedMovieSceneSequenceInstance());
 		NotifyMovieSceneDataChanged();
 	}
 	
@@ -717,19 +699,15 @@ ISequencerObjectChangeListener& FSequencer::GetObjectChangeListener()
 }
 
 
-void FSequencer::SpawnActorsForMovie( TSharedRef<FMovieSceneInstance> MovieSceneInstance )
+void FSequencer::SpawnActorsForMovie( TSharedRef<FMovieSceneSequenceInstance> MovieSceneInstance )
 {
 	SpawnOrDestroyPuppetObjects( MovieSceneInstance );
 }
 
-void FSequencer::GetRuntimeObjects( TSharedRef<FMovieSceneInstance> MovieSceneInstance, const FGuid& ObjectHandle, TArray< UObject*>& OutObjects ) const
+void FSequencer::GetRuntimeObjects( TSharedRef<FMovieSceneSequenceInstance> MovieSceneInstance, const FGuid& ObjectHandle, TArray< UObject*>& OutObjects ) const
 {
-	if (Animation == nullptr)
-	{
-		return;
-	}
 
-	UObject* FoundObject = Animation->FindObject(ObjectHandle);
+	UObject* FoundObject = MovieSceneInstance->GetSequence()->FindObject(ObjectHandle);
 
 	if (FoundObject != nullptr)
 	{
@@ -758,7 +736,7 @@ EMovieScenePlayerStatus::Type FSequencer::GetPlaybackStatus() const
 	return PlaybackState;
 }
 
-void FSequencer::AddOrUpdateMovieSceneInstance( UMovieSceneSection& MovieSceneSection, TSharedRef<FMovieSceneInstance> InstanceToAdd )
+void FSequencer::AddOrUpdateMovieSceneInstance( UMovieSceneSection& MovieSceneSection, TSharedRef<FMovieSceneSequenceInstance> InstanceToAdd )
 {
 	if( !MovieSceneSectionToInstanceMap.Contains( &MovieSceneSection ) )
 	{
@@ -768,10 +746,10 @@ void FSequencer::AddOrUpdateMovieSceneInstance( UMovieSceneSection& MovieSceneSe
 	SpawnOrDestroyPuppetObjects( InstanceToAdd );
 }
 
-void FSequencer::RemoveMovieSceneInstance( UMovieSceneSection& MovieSceneSection, TSharedRef<FMovieSceneInstance> InstanceToRemove )
+void FSequencer::RemoveMovieSceneInstance( UMovieSceneSection& MovieSceneSection, TSharedRef<FMovieSceneSequenceInstance> InstanceToRemove )
 {
-	UMovieScene* MovieScene = InstanceToRemove->GetMovieScene();
-	Animation->DestroyAllSpawnedObjects(*MovieScene);
+	UMovieSceneSequence* Sequence = InstanceToRemove->GetSequence();
+	Sequence->DestroyAllSpawnedObjects();
 
 	MovieSceneSectionToInstanceMap.Remove( &MovieSceneSection );
 }
@@ -782,8 +760,8 @@ void FSequencer::AddReferencedObjects( FReferenceCollector& Collector )
 
 	for( int32 MovieSceneIndex = 0; MovieSceneIndex < MovieSceneStack.Num(); ++MovieSceneIndex )
 	{
-		UMovieScene* Scene = MovieSceneStack[MovieSceneIndex]->GetMovieScene();
-		Collector.AddReferencedObject( Scene );
+		UMovieSceneSequence* Sequence = MovieSceneStack[MovieSceneIndex]->GetSequence();
+		Collector.AddReferencedObject( Sequence );
 	}
 }
 
@@ -801,17 +779,17 @@ void FSequencer::ResetPerMovieSceneData()
 void FSequencer::UpdateRuntimeInstances()
 {
 	// Spawn runtime objects for the root movie scene
-	SpawnOrDestroyPuppetObjects( RootMovieSceneInstance.ToSharedRef() );
+	SpawnOrDestroyPuppetObjects( RootMovieSceneSequenceInstance.ToSharedRef() );
 
 	// Refresh the current root instance
-	RootMovieSceneInstance->RefreshInstance( *this );
-	RootMovieSceneInstance->Update(ScrubPosition, ScrubPosition, *this);
+	RootMovieSceneSequenceInstance->RefreshInstance( *this );
+	RootMovieSceneSequenceInstance->Update(ScrubPosition, ScrubPosition, *this);
 }
 
-void FSequencer::DestroySpawnablesForAllMovieScenes()
+void FSequencer::DestroySpawnablesForAllSequences()
 {
-	UMovieScene* MovieScene = RootMovieSceneInstance->GetMovieScene();
-	Animation->DestroyAllSpawnedObjects(*MovieScene);
+	// @todo Sequencer - This seems incorrect and should destroy for all sequences
+	RootMovieSceneSequenceInstance->GetSequence()->DestroyAllSpawnedObjects();
 }
 
 FReply FSequencer::OnPlay(bool bTogglePlay)
@@ -821,7 +799,7 @@ FReply FSequencer::OnPlay(bool bTogglePlay)
 	{
 		PlaybackState = EMovieScenePlayerStatus::Stopped;
 		// Update on stop (cleans up things like sounds that are playing)
-		RootMovieSceneInstance->Update( ScrubPosition, ScrubPosition, *this );
+		RootMovieSceneSequenceInstance->Update( ScrubPosition, ScrubPosition, *this );
 	}
 	else
 	{
@@ -940,25 +918,27 @@ void FSequencer::UpdateTimeBoundsToFocusedMovieScene()
 	const int32 InFrame = ProjectSettings->InFrame;
 	const int32 OutFrame = ProjectSettings->OutFrame;
 
-	float InTime = GetFocusedMovieScene()->InTime;
+	UMovieScene* FocusedMovieScene = GetFocusedMovieSceneSequence()->GetMovieScene();
+
+	float InTime = FocusedMovieScene->InTime;
 	if (InTime >= FLT_MAX)
 	{
 		InTime = SequencerHelpers::FrameToTime(InFrame, 1.0f/Settings->GetTimeSnapInterval());
 	}
 
-	float OutTime = GetFocusedMovieScene()->OutTime;
+	float OutTime = FocusedMovieScene->OutTime;
 	if (OutTime <= -FLT_MAX)
 	{
 		OutTime = SequencerHelpers::FrameToTime(OutFrame, 1.0f/Settings->GetTimeSnapInterval());
 	}
 
-	float StartTime = GetFocusedMovieScene()->StartTime;
+	float StartTime = FocusedMovieScene->StartTime;
 	if (StartTime >= FLT_MAX)
 	{
 		StartTime = SequencerHelpers::FrameToTime(InFrame, 1.0f/Settings->GetTimeSnapInterval());
 	}
 
-	float EndTime = GetFocusedMovieScene()->EndTime;
+	float EndTime = FocusedMovieScene->EndTime;
 	if (EndTime <= -FLT_MAX)
 	{
 		EndTime = SequencerHelpers::FrameToTime(OutFrame, 1.0f/Settings->GetTimeSnapInterval());
@@ -986,7 +966,7 @@ TRange<float> FSequencer::GetTimeBounds() const
 		return TRange<float>( -100000.0f, 100000.0f );
 	}
 
-	const UMovieScene* MovieScene = GetFocusedMovieScene();
+	const UMovieScene* MovieScene = GetFocusedMovieSceneSequence()->GetMovieScene();
 	const UMovieSceneTrack* AnimatableShot = MovieScene->FindMasterTrack( UMovieSceneShotTrack::StaticClass() );
 	if (AnimatableShot)
 	{
@@ -1024,30 +1004,31 @@ void FSequencer::OnViewRangeChanged( TRange<float> NewViewRange, EViewRangeInter
 
 	if (!NewViewRange.IsEmpty())
 	{
+		UMovieScene* FocusedMovieScene = GetFocusedMovieSceneSequence()->GetMovieScene();
 		// If locked, clamp the new range to clamp range
 		if (Settings->GetLockInOutToStartEndRange())
 		{
-			if ( NewViewRange.GetLowerBoundValue() < GetFocusedMovieScene()->StartTime)
+			if ( NewViewRange.GetLowerBoundValue() < FocusedMovieScene->StartTime)
 			{
-				NewViewRange = TRange<float>(GetFocusedMovieScene()->StartTime, GetFocusedMovieScene()->StartTime + (NewViewRange.GetUpperBoundValue() - NewViewRange.GetLowerBoundValue()));
+				NewViewRange = TRange<float>(FocusedMovieScene->StartTime, FocusedMovieScene->StartTime + (NewViewRange.GetUpperBoundValue() - NewViewRange.GetLowerBoundValue()));
 			}
-			else if (NewViewRange.GetUpperBoundValue() > GetFocusedMovieScene()->EndTime)
+			else if (NewViewRange.GetUpperBoundValue() > FocusedMovieScene->EndTime)
 			{
-				NewViewRange = TRange<float>(GetFocusedMovieScene()->EndTime - (NewViewRange.GetUpperBoundValue() - NewViewRange.GetLowerBoundValue()), GetFocusedMovieScene()->EndTime);
+				NewViewRange = TRange<float>(FocusedMovieScene->EndTime - (NewViewRange.GetUpperBoundValue() - NewViewRange.GetLowerBoundValue()), FocusedMovieScene->EndTime);
 			}
 		}
 		// Otherwise, grow the clamp range to the new range
 		else
 		{
 			bool bNeedsClampSet = false;
-			float NewClampRangeMin = GetFocusedMovieScene()->StartTime;
+			float NewClampRangeMin = FocusedMovieScene->StartTime;
 			if ( NewViewRange.GetLowerBoundValue() < NewClampRangeMin)
 			{
 				bNeedsClampSet = true;
 				NewClampRangeMin = NewViewRange.GetLowerBoundValue();
 			}
 	
-			float NewClampRangeMax = GetFocusedMovieScene()->EndTime;
+			float NewClampRangeMax = FocusedMovieScene->EndTime;
 			if ( NewViewRange.GetUpperBoundValue() > NewClampRangeMax)
 			{
 				bNeedsClampSet = true;
@@ -1082,8 +1063,8 @@ void FSequencer::OnViewRangeChanged( TRange<float> NewViewRange, EViewRangeInter
 		}
 
 		// Synchronize the in and out time of the focused movie scene
-		GetFocusedMovieScene()->InTime = NewViewRange.GetLowerBoundValue();
-		GetFocusedMovieScene()->OutTime = NewViewRange.GetUpperBoundValue();
+		FocusedMovieScene->InTime = NewViewRange.GetLowerBoundValue();
+		FocusedMovieScene->OutTime = NewViewRange.GetUpperBoundValue();
 	}
 }
 
@@ -1091,9 +1072,11 @@ void FSequencer::OnClampRangeChanged( TRange<float> NewClampRange )
 {
 	if (!NewClampRange.IsEmpty())
 	{
+		UMovieScene* FocusedMovieScene = GetFocusedMovieSceneSequence()->GetMovieScene();
+
 		// Synchronize the in and out time of the focused movie scene
-		GetFocusedMovieScene()->StartTime = NewClampRange.GetLowerBoundValue();
-		GetFocusedMovieScene()->EndTime = NewClampRange.GetUpperBoundValue();
+		FocusedMovieScene->StartTime = NewClampRange.GetLowerBoundValue();
+		FocusedMovieScene->EndTime = NewClampRange.GetUpperBoundValue();
 	}
 }
 
@@ -1156,11 +1139,13 @@ FGuid FSequencer::AddSpawnableForAssetOrClass( UObject* Object, UObject* Counter
 {
 	FGuid NewSpawnableGuid;
 
-	if (Animation->AllowsSpawnableObjects())
+
+	UMovieSceneSequence* Sequence = GetFocusedMovieSceneSequence();
+	if (Sequence->AllowsSpawnableObjects())
 	{
 		// Grab the MovieScene that is currently focused.  We'll add our Blueprint as an inner of the
 		// MovieScene asset.
-		UMovieScene* OwnerMovieScene = GetFocusedMovieScene();
+		UMovieScene* OwnerMovieScene = Sequence->GetMovieScene();
 
 		// @todo sequencer: Undo doesn't seem to be working at all
 		const FScopedTransaction Transaction( LOCTEXT("UndoAddingObject", "Add Object to MovieScene") );
@@ -1239,13 +1224,12 @@ FGuid FSequencer::AddSpawnableForAssetOrClass( UObject* Object, UObject* Counter
 	return NewSpawnableGuid;
 }
 
-void FSequencer::AddSubMovieScene( UMovieScene* SubMovieScene )
+void FSequencer::AddSubMovieScene( UMovieSceneSequence* SubMovieSceneSequence )
 {
 	// @todo Sequencer - sub-moviescenes This should be moved to the sub-moviescene editor
-	SubMovieScene->SetFlags( RF_Transactional );
-	
+
 	// Grab the MovieScene that is currently focused.  THis is the movie scene that will contain the sub-moviescene
-	UMovieScene* OwnerMovieScene = GetFocusedMovieScene();
+	UMovieScene* OwnerMovieScene = GetFocusedMovieSceneSequence()->GetMovieScene();
 
 	// @todo sequencer: Undo doesn't seem to be working at all
 	const FScopedTransaction Transaction( LOCTEXT("UndoAddingObject", "Add Object to MovieScene") );
@@ -1256,7 +1240,7 @@ void FSequencer::AddSubMovieScene( UMovieScene* SubMovieScene )
 
 	USubMovieSceneTrack* SubMovieSceneType = CastChecked<USubMovieSceneTrack>( NewTrack );
 
-	SubMovieSceneType->AddMovieSceneSection( SubMovieScene, ScrubPosition );
+	SubMovieSceneType->AddMovieSceneSection( SubMovieSceneSequence, ScrubPosition );
 }
 
 bool FSequencer::OnHandleAssetDropped(UObject* DroppedAsset, const FGuid& TargetObjectGuid)
@@ -1281,13 +1265,13 @@ void FSequencer::OnRequestNodeDeleted( TSharedRef<const FSequencerDisplayNode>&
 	bool bAnySpawnablesRemoved = false;
 	bool bAnythingRemoved = false;
 	
-	TSharedRef<FMovieSceneInstance> MovieSceneInstance = GetFocusedMovieSceneInstance();
-	UMovieScene* OwnerMovieScene = MovieSceneInstance->GetMovieScene();
+	TSharedRef<FMovieSceneSequenceInstance> MovieSceneInstance = GetFocusedMovieSceneSequenceInstance();
+	UMovieSceneSequence* Sequence = MovieSceneInstance->GetSequence();
+	UMovieScene* OwnerMovieScene = Sequence->GetMovieScene();
 
 	// Only object nodes or section areas can be deleted
 	if( NodeToBeDeleted->GetType() == ESequencerNode::Object  )
 	{
-		OwnerMovieScene->SetFlags( RF_Transactional );
 		const FGuid& BindingToRemove = StaticCastSharedRef<const FObjectBindingNode>( NodeToBeDeleted )->GetObjectBinding();
 
 		//@todo Sequencer - add transaction
@@ -1305,7 +1289,7 @@ void FSequencer::OnRequestNodeDeleted( TSharedRef<const FSequencerDisplayNode>&
 			bRemoved = OwnerMovieScene->RemovePossessable( BindingToRemove );
 			
 			// @todo sequencer: undo needs to work here
-			Animation->UnbindPossessableObjects( BindingToRemove );
+			Sequence->UnbindPossessableObjects( BindingToRemove );
 
 			// If this check fails the guid was not associated with a spawnable or possessable so there was an invalid guid being stored on a node
 			check( bRemoved );
@@ -1318,21 +1302,17 @@ void FSequencer::OnRequestNodeDeleted( TSharedRef<const FSequencerDisplayNode>&
 		TSharedRef<const FTrackNode> SectionAreaNode = StaticCastSharedRef<const FTrackNode>( NodeToBeDeleted );
 		UMovieSceneTrack* Track = SectionAreaNode->GetTrack();
 
-		UMovieScene* FocusedMovieScene = GetFocusedMovieScene();
-
-		FocusedMovieScene->SetFlags( RF_Transactional );
-		
-		if( FocusedMovieScene->IsAMasterTrack( Track ) )
+		if( OwnerMovieScene->IsAMasterTrack( Track ) )
 		{
-			FocusedMovieScene->RemoveMasterTrack( Track );
+			OwnerMovieScene->RemoveMasterTrack( Track );
 		}
-		else if( FocusedMovieScene->GetShotTrack() == Track )
+		else if( OwnerMovieScene->GetShotTrack() == Track )
 		{
-			FocusedMovieScene->RemoveShotTrack();
+			OwnerMovieScene->RemoveShotTrack();
 		}
 		else
 		{
-			FocusedMovieScene->RemoveTrack( Track );
+			OwnerMovieScene->RemoveTrack( Track );
 		}
 		
 		bAnythingRemoved = true;
@@ -1403,7 +1383,7 @@ void FSequencer::PostUndo(bool bSuccess)
 void FSequencer::OnPreSaveWorld(uint32 SaveFlags, class UWorld* World)
 {
 	// Restore the saved state so that the level save can save that instead of the animated state.
-	RootMovieSceneInstance->RestoreState();
+	RootMovieSceneSequenceInstance->RestoreState();
 }
 
 void FSequencer::OnPostSaveWorld(uint32 SaveFlags, class UWorld* World, bool bSuccess)
@@ -1471,7 +1451,7 @@ void FSequencer::UpdatePreviewLevelViewportClientFromCameraCut( FLevelEditorView
 
 void FSequencer::SaveCurrentMovieScene()
 {
-	UPackage* MovieScenePackage = GetFocusedMovieScene()->GetOutermost();
+	UPackage* MovieScenePackage = GetFocusedMovieSceneSequence()->GetOutermost();
 
 	TArray<UPackage*> PackagesToSave;
 
@@ -1565,7 +1545,7 @@ void FSequencer::FilterToShotSections(const TArray< TWeakObjectPtr<class UMovieS
 
 		// populate unfilterable shots - shots that started not filtered
 		TArray< TWeakObjectPtr<class UMovieSceneSection> > TempUnfilterableSections;
-		const TArray<UMovieSceneSection*>& AllSections = GetFocusedMovieScene()->GetAllSections();
+		const TArray<UMovieSceneSection*>& AllSections = GetFocusedMovieSceneSequence()->GetMovieScene()->GetAllSections();
 		for (int32 SectionIndex = 0; SectionIndex < AllSections.Num(); ++SectionIndex)
 		{
 			UMovieSceneSection* Section = AllSections[SectionIndex];
@@ -1620,11 +1600,6 @@ void FSequencer::KeyProperty(FKeyPropertyParams KeyPropertyParams)
 	ObjectChangeListener->KeyProperty(KeyPropertyParams);
 }
 
-UMovieSceneAnimation* FSequencer::GetAnimation()
-{
-	return Animation;
-}
-
 FSequencerSelection& FSequencer::GetSelection()
 {
 	return Selection;
diff --git a/Engine/Source/Editor/Sequencer/Private/Sequencer.h b/Engine/Source/Editor/Sequencer/Private/Sequencer.h
--- a/Engine/Source/Editor/Sequencer/Private/Sequencer.h
+++ b/Engine/Source/Editor/Sequencer/Private/Sequencer.h
@@ -9,7 +9,7 @@
 
 class UMovieScene;
 class IToolkitHost;
-class UMovieSceneAnimation;
+class UMovieSceneSequence;
 class ISequencerObjectChangeListener;
 
 
@@ -61,19 +61,19 @@ public:
 
 	virtual void Close() override;
 	virtual TSharedRef<SWidget> GetSequencerWidget() const override { return SequencerWidget.ToSharedRef(); }
-	virtual UMovieScene* GetRootMovieScene() const override;
-	virtual UMovieScene* GetFocusedMovieScene() const override;
-	virtual void ResetToNewAnimation(UMovieSceneAnimation& NewAnimation) override;
-	virtual TSharedRef<FMovieSceneInstance> GetRootMovieSceneInstance() const override;
-	virtual TSharedRef<FMovieSceneInstance> GetFocusedMovieSceneInstance() const override;
-	virtual void FocusSubMovieScene( TSharedRef<FMovieSceneInstance> SubMovieSceneInstance ) override;
-	TSharedRef<FMovieSceneInstance> GetInstanceForSubMovieSceneSection( UMovieSceneSection& SubMovieSceneSection ) const override;
+	virtual UMovieSceneSequence* GetRootMovieSceneSequence() const override;
+	virtual UMovieSceneSequence* GetFocusedMovieSceneSequence() const override;
+	virtual void ResetToNewRootSequence(UMovieSceneSequence& NewSequence) override;
+	virtual TSharedRef<FMovieSceneSequenceInstance> GetRootMovieSceneSequenceInstance() const override;
+	virtual TSharedRef<FMovieSceneSequenceInstance> GetFocusedMovieSceneSequenceInstance() const override;
+	virtual void FocusSubMovieScene( TSharedRef<FMovieSceneSequenceInstance> SubMovieSceneInstance ) override;
+	TSharedRef<FMovieSceneSequenceInstance> GetInstanceForSubMovieSceneSection( UMovieSceneSection& SubMovieSceneSection ) const override;
 	virtual void AddNewShot(FGuid CameraGuid) override;
 	virtual void AddAnimation(FGuid ObjectGuid, class UAnimSequence* AnimSequence) override;
 	virtual bool GetAutoKeyEnabled() const override;
 	virtual void SetAutoKeyEnabled(bool bAutoKeyEnabled) override;
 	virtual bool IsRecordingLive() const override;
-	virtual float GetCurrentLocalTime(UMovieScene& MovieScene) override;
+	virtual float GetCurrentLocalTime(UMovieSceneSequence& InMovieSceneSequence) override;
 	virtual float GetGlobalTime() override;
 	virtual void SetGlobalTime(float Time) override;
 	virtual void SetPerspectiveViewportPossessionEnabled(bool bEnabled) override;
@@ -81,12 +81,11 @@ public:
 	virtual ISequencerObjectChangeListener& GetObjectChangeListener() override;
 	virtual void NotifyMovieSceneDataChanged() override;
 	virtual void UpdateRuntimeInstances() override;
-	virtual void AddSubMovieScene(UMovieScene* SubMovieScene) override;
+	virtual void AddSubMovieScene(UMovieSceneSequence* SubMovieSceneSequence) override;
 	virtual void FilterToShotSections(const TArray< TWeakObjectPtr<class UMovieSceneSection> >& ShotSections, bool bZoomToShotBounds = false) override;
 	virtual void FilterToSelectedShotSections(bool bZoomToShotBounds = true) override;
 	virtual bool CanKeyProperty(FCanKeyPropertyParams CanKeyPropertyParams) const override;
 	virtual void KeyProperty(FKeyPropertyParams KeyPropertyParams) override;
-	virtual UMovieSceneAnimation* GetAnimation() override;
 	virtual FSequencerSelection& GetSelection() override;
 	virtual void NotifyMapChanged(class UWorld* NewWorld, EMapChangeType MapChangeType) override;
 
@@ -106,14 +105,14 @@ public:
 	/**
 	 * Pops the current focused movie scene from the stack.  The parent of this movie scene will be come the focused one
 	 */
-	void PopToMovieScene( TSharedRef<FMovieSceneInstance> SubMovieSceneInstance );
+	void PopToMovieScene( TSharedRef<FMovieSceneSequenceInstance> SubMovieSceneInstance );
 
 	/**
 	 * Spawn (or destroy) puppet objects as needed to match the spawnables array in the MovieScene we're editing
 	 *
 	 * @param MovieSceneInstance	The movie scene instance to spawn or destroy puppet objects for
 	 */
-	void SpawnOrDestroyPuppetObjects( TSharedRef<FMovieSceneInstance> MovieSceneInstance );
+	void SpawnOrDestroyPuppetObjects( TSharedRef<FMovieSceneSequenceInstance> MovieSceneInstance );
 
 	/** 
 	 * Deletes the passed in section
@@ -223,15 +222,13 @@ public:
 	void BuildObjectBindingEditButtons(TSharedPtr<SHorizontalBox> EditBox, const FGuid& ObjectBinding, const class UClass* ObjectClass);
 
 	/** IMovieScenePlayer interface */
-	virtual void GetRuntimeObjects( TSharedRef<FMovieSceneInstance> MovieSceneInstance, const FGuid& ObjectHandle, TArray< UObject* >& OutObjects ) const override;
+	virtual void GetRuntimeObjects( TSharedRef<FMovieSceneSequenceInstance> MovieSceneInstance, const FGuid& ObjectHandle, TArray< UObject* >& OutObjects ) const override;
 	virtual void UpdateCameraCut(UObject* ObjectToViewThrough, bool bNewCameraCut) const override;
 	virtual EMovieScenePlayerStatus::Type GetPlaybackStatus() const override;
-	virtual void AddOrUpdateMovieSceneInstance( class UMovieSceneSection& MovieSceneSection, TSharedRef<FMovieSceneInstance> InstanceToAdd ) override;
-	virtual void RemoveMovieSceneInstance( class UMovieSceneSection& MovieSceneSection, TSharedRef<FMovieSceneInstance> InstanceToRemove ) override;
+	virtual void AddOrUpdateMovieSceneInstance( class UMovieSceneSection& MovieSceneSection, TSharedRef<FMovieSceneSequenceInstance> InstanceToAdd ) override;
+	virtual void RemoveMovieSceneInstance( class UMovieSceneSection& MovieSceneSection, TSharedRef<FMovieSceneSequenceInstance> InstanceToRemove ) override;
 
-	virtual void SpawnActorsForMovie( TSharedRef<FMovieSceneInstance> MovieSceneInstance );
-
-	virtual TArray< UMovieScene* > GetMovieScenesBeingEdited();
+	virtual void SpawnActorsForMovie( TSharedRef<FMovieSceneSequenceInstance> MovieSceneInstance );
 
 	/** Called when an actor is dropped into Sequencer */
 	void OnActorsDropped( const TArray<TWeakObjectPtr<AActor> >& Actors );
@@ -268,7 +265,7 @@ protected:
 	/**
 	 * Destroys spawnables for all movie scenes in the stack
 	 */
-	void DestroySpawnablesForAllMovieScenes();
+	void DestroySpawnablesForAllSequences();
 
 	/** Sets the actor CDO such that it is placed in front of the active perspective viewport camera, if we have one */
 	static void PlaceActorInFrontOfCamera( AActor* ActorCDO );
@@ -402,7 +399,7 @@ private:
 	/** User-supplied settings object for this sequencer */
 	USequencerSettings* Settings;
 
-	TMap< TWeakObjectPtr<UMovieSceneSection>, TSharedRef<FMovieSceneInstance> > MovieSceneSectionToInstanceMap;
+	TMap< TWeakObjectPtr<UMovieSceneSection>, TSharedRef<FMovieSceneSequenceInstance> > MovieSceneSectionToInstanceMap;
 
 	/** Command list for sequencer commands */
 	TSharedRef<FUICommandList> SequencerCommandBindings;
@@ -412,16 +409,13 @@ private:
 
 	/** The editors we keep track of for special behaviors */
 	TWeakPtr<FMovieSceneTrackEditor> ShotTrackEditor;
-	TWeakPtr<FMovieSceneTrackEditor> AnimationTrackEditor;
-
-	/** The current animation being edited. */
-	UMovieSceneAnimation* Animation;
+	TWeakPtr<FMovieSceneTrackEditor> SkeletalAnimationTrackEditor;
 
 	/** Listener for object changes being made while this sequencer is open*/
 	TSharedPtr< class ISequencerObjectChangeListener > ObjectChangeListener;
 
 	/** The runtime instance for the root movie scene */
-	TSharedPtr< class FMovieSceneInstance > RootMovieSceneInstance;
+	TSharedPtr< class FMovieSceneSequenceInstance > RootMovieSceneSequenceInstance;
 
 	/** Main sequencer widget */
 	TSharedPtr< class SSequencer > SequencerWidget;
@@ -439,7 +433,7 @@ private:
 	TArray<FGuid> UnfilterableObjects;
 
 	/** Stack of movie scenes.  The first element is always the root movie scene.  The last element is the focused movie scene */
-	TArray< TSharedRef<FMovieSceneInstance> > MovieSceneStack;
+	TArray< TSharedRef<FMovieSceneSequenceInstance> > MovieSceneStack;
 
 	/** The time range target to be viewed */
 	TRange<float> TargetViewRange;
diff --git a/Engine/Source/Editor/Sequencer/Private/SequencerDisplayNode.cpp b/Engine/Source/Editor/Sequencer/Private/SequencerDisplayNode.cpp
--- a/Engine/Source/Editor/Sequencer/Private/SequencerDisplayNode.cpp
+++ b/Engine/Source/Editor/Sequencer/Private/SequencerDisplayNode.cpp
@@ -8,7 +8,7 @@
 #include "SSequencer.h"
 #include "SSequencerSectionAreaView.h"
 #include "MovieSceneSection.h"
-#include "MovieSceneAnimation.h"
+#include "MovieSceneSequence.h"
 #include "MovieSceneTrack.h"
 #include "CommonMovieSceneTools.h"
 #include "IKeyArea.h"
@@ -617,7 +617,7 @@ void FTrackNode::FixRowIndices()
 FText FObjectBindingNode::GetDisplayName() const
 {
 	FText DisplayName;
-	return GetSequencer().GetAnimation()->TryGetObjectDisplayName(ObjectBinding, DisplayName) ?
+	return GetSequencer().GetFocusedMovieSceneSequence()->TryGetObjectDisplayName(ObjectBinding, DisplayName) ?
 		DisplayName : DefaultDisplayName;
 }
 
@@ -642,7 +642,7 @@ const UClass* FObjectBindingNode::GetClassForObjectBinding()
 {
 	FSequencer& ParentSequencer = GetSequencer();
 
-	UMovieScene* MovieScene = GetSequencer().GetFocusedMovieScene();
+	UMovieScene* MovieScene = GetSequencer().GetFocusedMovieSceneSequence()->GetMovieScene();
 
 	FMovieSceneSpawnable* Spawnable = MovieScene->FindSpawnable(ObjectBinding);
 	FMovieScenePossessable* Possessable = MovieScene->FindPossessable(ObjectBinding);
@@ -738,7 +738,7 @@ TSharedRef<SWidget> FObjectBindingNode::OnGetAddPropertyTrackMenuContent()
 {
 	TArray<TArray<UProperty*>> KeyablePropertyPaths;
 	FSequencer& Sequencer = GetSequencer();
-	UObject* BoundObject = Sequencer.GetAnimation()->FindObject(ObjectBinding);
+	UObject* BoundObject = Sequencer.GetFocusedMovieSceneSequence()->FindObject(ObjectBinding);
 
 	if (BoundObject != nullptr)
 	{
@@ -785,7 +785,7 @@ TSharedRef<SWidget> FObjectBindingNode::OnGetAddPropertyTrackMenuContent()
 void FObjectBindingNode::AddTrackForProperty(TArray<UProperty*> PropertyPath)
 {
 	FSequencer& Sequencer = GetSequencer();
-	UObject* BoundObject = Sequencer.GetAnimation()->FindObject(ObjectBinding);
+	UObject* BoundObject = Sequencer.GetFocusedMovieSceneSequence()->FindObject(ObjectBinding);
 
 	TArray<UObject*> KeyableBoundObjects;
 	if (BoundObject != nullptr)
diff --git a/Engine/Source/Editor/Sequencer/Private/SequencerNodeTree.cpp b/Engine/Source/Editor/Sequencer/Private/SequencerNodeTree.cpp
--- a/Engine/Source/Editor/Sequencer/Private/SequencerNodeTree.cpp
+++ b/Engine/Source/Editor/Sequencer/Private/SequencerNodeTree.cpp
@@ -1,7 +1,7 @@
 // Copyright 1998-2015 Epic Games, Inc. All Rights Reserved.
 
 #include "SequencerPrivatePCH.h"
-#include "MovieSceneAnimation.h"
+#include "MovieSceneSequence.h"
 #include "MovieSceneSection.h"
 #include "SequencerNodeTree.h"
 #include "Sequencer.h"
@@ -27,7 +27,7 @@ void FSequencerNodeTree::Update()
 	// @todo Sequencer - This update pass is too aggressive.  Some nodes may still be valid
 	Empty();
 
-	UMovieScene* MovieScene = Sequencer.GetFocusedMovieScene();
+	UMovieScene* MovieScene = Sequencer.GetFocusedMovieSceneSequence()->GetMovieScene();
 
 	TArray< TSharedRef<FSequencerDisplayNode> > NewRootNodes;
 
@@ -203,7 +203,7 @@ TSharedRef<FObjectBindingNode> FSequencerNodeTree::AddObjectBinding(const FStrin
 		// Try to get the parent object node if there is one.
 		TSharedPtr<FObjectBindingNode> ParentNode;
 		TArray<UObject*> RuntimeObjects;
-		UMovieSceneAnimation* Animation = Sequencer.GetAnimation();
+		UMovieSceneSequence* Animation = Sequencer.GetFocusedMovieSceneSequence();
 		UObject* RuntimeObject = Animation->FindObject(ObjectBinding);
 		if ( RuntimeObject != nullptr)
 		{
@@ -247,7 +247,8 @@ TSharedRef<FObjectBindingNode> FSequencerNodeTree::AddObjectBinding(const FStrin
 
 void FSequencerNodeTree::SaveExpansionState( const FSequencerDisplayNode& Node, bool bExpanded )
 {	
-	UMovieScene* MovieScene = Sequencer.GetFocusedMovieScene();
+	// @todo Sequencer - This should be moved to the sequence level
+	UMovieScene* MovieScene = Sequencer.GetFocusedMovieSceneSequence()->GetMovieScene();
 
 	FMovieSceneEditorData& EditorData = MovieScene->GetEditorData();
 
@@ -256,7 +257,8 @@ void FSequencerNodeTree::SaveExpansionState( const FSequencerDisplayNode& Node,
 
 bool FSequencerNodeTree::GetSavedExpansionState( const FSequencerDisplayNode& Node ) const
 {
-	UMovieScene* MovieScene = Sequencer.GetFocusedMovieScene();
+	// @todo Sequencer - This should be moved to the sequence level
+	UMovieScene* MovieScene = Sequencer.GetFocusedMovieSceneSequence()->GetMovieScene();
 
 	FMovieSceneEditorData& EditorData = MovieScene->GetEditorData();
 	FMovieSceneExpansionState* ExpansionState = EditorData.ExpansionStates.Find( Node.GetPathName() );
diff --git a/Engine/Source/Editor/Sequencer/Private/SequencerPrivatePCH.h b/Engine/Source/Editor/Sequencer/Private/SequencerPrivatePCH.h
--- a/Engine/Source/Editor/Sequencer/Private/SequencerPrivatePCH.h
+++ b/Engine/Source/Editor/Sequencer/Private/SequencerPrivatePCH.h
@@ -20,7 +20,7 @@
 #include "SSequencer.h"
 #include "SSection.h"
 #include "SAnimationOutlinerTreeNode.h"
-#include "MovieSceneInstance.h"
+#include "MovieSceneSequenceInstance.h"
 #include "SequencerSettings.h"
 #include "SequencerProjectSettings.h"
 #include "SequencerCommonHelpers.h"
diff --git a/Engine/Source/Editor/Sequencer/Public/ISequencer.h b/Engine/Source/Editor/Sequencer/Public/ISequencer.h
--- a/Engine/Source/Editor/Sequencer/Public/ISequencer.h
+++ b/Engine/Source/Editor/Sequencer/Public/ISequencer.h
@@ -9,7 +9,7 @@
 
 
 class FSequencerSelection;
-class UMovieSceneAnimation;
+class UMovieSceneSequence;
 class UAnimSequence;
 class UMovieScene;
 class UMovieSceneSection;
@@ -32,25 +32,25 @@ public:
 	virtual TSharedRef<SWidget> GetSequencerWidget() const = 0;
 
 	/** @return The root movie scene being used */
-	virtual UMovieScene* GetRootMovieScene() const = 0;
+	virtual UMovieSceneSequence* GetRootMovieSceneSequence() const = 0;
 
 	/** @return Returns the MovieScene that is currently focused for editing by the sequencer.  This can change at any time. */
-	virtual UMovieScene* GetFocusedMovieScene() const = 0;
+	virtual UMovieSceneSequence* GetFocusedMovieSceneSequence() const = 0;
 
 	/**
 	 * @return the currently focused movie scene instance
 	 */
-	virtual TSharedRef<FMovieSceneInstance> GetFocusedMovieSceneInstance() const = 0;
+	virtual TSharedRef<FMovieSceneSequenceInstance> GetFocusedMovieSceneSequenceInstance() const = 0;
 
 	/** Resets sequencer with a new animation */
-	virtual void ResetToNewAnimation(UMovieSceneAnimation& NewAnimation) = 0;
+	virtual void ResetToNewRootSequence(UMovieSceneSequence& NewAnimation) = 0;
 
 	/**
 	 * Focuses a sub-movie scene (MovieScene within a MovieScene) in the sequencer.
 	 * 
 	 * @param SubMovieSceneInstance	Sub-MovieScene instance to focus.
 	 */
-	virtual void FocusSubMovieScene(TSharedRef<FMovieSceneInstance> SubMovieSceneInstance) = 0;
+	virtual void FocusSubMovieScene(TSharedRef<FMovieSceneSequenceInstance> SubMovieSceneInstance) = 0;
 	
 	/**
 	 * Given a sub-movie scene section, returns the instance of the movie scene for that section.
@@ -58,14 +58,14 @@ public:
 	 * @param SubMovieSceneSection	The sub-movie scene section containing a movie scene to get the instance for.
 	 * @return The instance for the sub-movie scene
 	 */
-	virtual TSharedRef<FMovieSceneInstance> GetInstanceForSubMovieSceneSection(UMovieSceneSection& SubMovieSceneSection) const = 0;
+	virtual TSharedRef<FMovieSceneSequenceInstance> GetInstanceForSubMovieSceneSection(UMovieSceneSection& SubMovieSceneSection) const = 0;
 
 	/**
 	 * Adds a movie scene as a section inside the current movie scene
 	 * 
 	 * @param SubMovieScene The movie scene to add.
 	 */
-	virtual void AddSubMovieScene(UMovieScene* SubMovieScene) = 0;
+	virtual void AddSubMovieScene(UMovieSceneSequence* SubMovieScene) = 0;
 
 	/** @return Returns whether auto-key is enabled in this sequencer */
 	virtual bool GetAutoKeyEnabled() const = 0;
@@ -104,7 +104,7 @@ public:
 	 *
 	 * @param MovieScene The movie scene to get the local time for (must exist in this sequencer).
 	 */
-	virtual float GetCurrentLocalTime(UMovieScene& MovieScene) = 0;
+	virtual float GetCurrentLocalTime(UMovieSceneSequence& InMovieSceneSequence) = 0;
 	
 	/**
 	 * Gets the global time.
@@ -156,13 +156,6 @@ public:
 
 	virtual void UpdateRuntimeInstances() = 0;
 
-	/**
-	 * Get the animation that is being played by this sequencer.
-	 *
-	 * @return The current animation, or nullptr if no animation is playing.
-	 */
-	virtual UMovieSceneAnimation* GetAnimation() = 0;
-
 	virtual FSequencerSelection& GetSelection() = 0;
 
 	virtual void NotifyMapChanged(class UWorld* NewWorld, EMapChangeType MapChangeType) = 0;
diff --git a/Engine/Source/Editor/Sequencer/Public/ISequencerModule.h b/Engine/Source/Editor/Sequencer/Public/ISequencerModule.h
--- a/Engine/Source/Editor/Sequencer/Public/ISequencerModule.h
+++ b/Engine/Source/Editor/Sequencer/Public/ISequencerModule.h
@@ -55,8 +55,8 @@ struct FSequencerViewParams
  */
 struct FSequencerInitParams
 {
-	/** The animation being edited. */
-	UMovieSceneAnimation* Animation;
+	/** The root movie scene sequence being edited. */
+	UMovieSceneSequence* RootSequence;
 
 	/** The asset editor created for this (if any) */
 	TSharedPtr<IToolkitHost> ToolkitHost;
diff --git a/Engine/Source/Editor/Sequencer/Public/MovieSceneTrackEditor.h b/Engine/Source/Editor/Sequencer/Public/MovieSceneTrackEditor.h
--- a/Engine/Source/Editor/Sequencer/Public/MovieSceneTrackEditor.h
+++ b/Engine/Source/Editor/Sequencer/Public/MovieSceneTrackEditor.h
@@ -6,6 +6,7 @@
 #include "ISequencer.h"
 #include "MovieScene.h"
 #include "ScopedTransaction.h"
+#include "MovieSceneSequence.h"
 
 class ISequencerSection;
 
@@ -30,9 +31,9 @@ public:
 	virtual void OnRelease(){};
 
 	/** @return The current movie scene */
-	UMovieScene* GetMovieScene() const
+	UMovieSceneSequence* GetMovieSceneSequence() const
 	{
-		return Sequencer.Pin()->GetFocusedMovieScene();
+		return Sequencer.Pin()->GetFocusedMovieSceneSequence();
 	}
 
 	/**
@@ -42,9 +43,9 @@ public:
 	 * @param OutTime		The current time of the sequencer which should be used for adding keys during auto-key
 	 * @return true	if we can auto-key
 	 */
-	float GetTimeForKey( UMovieScene* InMovieScene )
+	float GetTimeForKey( UMovieSceneSequence* InMovieSceneSequence )
 	{ 
-		return Sequencer.Pin()->GetCurrentLocalTime( *InMovieScene );
+		return Sequencer.Pin()->GetCurrentLocalTime( *InMovieSceneSequence );
 	}
 
 	/** Gets whether the tool can legally autokey */
@@ -69,16 +70,16 @@ public:
 		check(OnKeyProperty.IsBound());
 
 		// Get the movie scene we want to autokey
-		UMovieScene* MovieScene = GetMovieScene();
-		float KeyTime = GetTimeForKey( MovieScene );
+		UMovieSceneSequence* MovieSceneSequence = GetMovieSceneSequence();
+		float KeyTime = GetTimeForKey( MovieSceneSequence );
 		bool bIsKeyingValid = !bMustBeAutokeying || IsAllowedToAutoKey();
 
 		if( bIsKeyingValid )
 		{
-			check( MovieScene );
+			check( MovieSceneSequence );
 
 			// @todo Sequencer - The sequencer probably should have taken care of this
-			MovieScene->SetFlags(RF_Transactional);
+			MovieSceneSequence->SetFlags(RF_Transactional);
 		
 			// Create a transaction record because we are about to add keys
 			const bool bShouldActuallyTransact = !Sequencer.Pin()->IsRecordingLive();		// Don't transact if we're recording in a PIE world.  That type of keyframe capture cannot be undone.
@@ -106,7 +107,7 @@ public:
 	{
 		check( UniqueTypeName != NAME_None );
 
-		UMovieScene* MovieScene = GetSequencer()->GetFocusedMovieScene();
+		UMovieScene* MovieScene = GetSequencer()->GetFocusedMovieSceneSequence()->GetMovieScene();
 
 		UMovieSceneTrack* Type = MovieScene->FindTrack( TrackClass, ObjectHandle, UniqueTypeName );
 
@@ -125,7 +126,7 @@ public:
 
 	UMovieSceneTrack* GetMasterTrack( TSubclassOf<UMovieSceneTrack> TrackClass )
 	{
-		UMovieScene* MovieScene = GetSequencer()->GetFocusedMovieScene();
+		UMovieScene* MovieScene = GetSequencer()->GetFocusedMovieSceneSequence()->GetMovieScene();
 		UMovieSceneTrack* Type = MovieScene->FindMasterTrack( TrackClass );
 
 		if( !Type )
diff --git a/Engine/Source/Editor/UMGEditor/Private/Animation/UMGDetailKeyframeHandler.cpp b/Engine/Source/Editor/UMGEditor/Private/Animation/UMGDetailKeyframeHandler.cpp
--- a/Engine/Source/Editor/UMGEditor/Private/Animation/UMGDetailKeyframeHandler.cpp
+++ b/Engine/Source/Editor/UMGEditor/Private/Animation/UMGDetailKeyframeHandler.cpp
@@ -17,8 +17,8 @@ bool FUMGDetailKeyframeHandler::IsPropertyKeyable(UClass* InObjectClass, const I
 
 bool FUMGDetailKeyframeHandler::IsPropertyKeyingEnabled() const
 {
-	UMovieScene* MovieScene = BlueprintEditor.Pin()->GetSequencer()->GetRootMovieScene();
-	return MovieScene != nullptr && MovieScene != UWidgetAnimation::GetNullAnimation()->MovieScene;
+	UMovieSceneSequence* Sequence = BlueprintEditor.Pin()->GetSequencer()->GetRootMovieSceneSequence();
+	return Sequence != nullptr && Sequence != UWidgetAnimation::GetNullAnimation();
 }
 
 void FUMGDetailKeyframeHandler::OnKeyPropertyClicked(const IPropertyHandle& KeyedPropertyHandle)
diff --git a/Engine/Source/Editor/UMGEditor/Private/TabFactory/AnimationTabSummoner.cpp b/Engine/Source/Editor/UMGEditor/Private/TabFactory/AnimationTabSummoner.cpp
--- a/Engine/Source/Editor/UMGEditor/Private/TabFactory/AnimationTabSummoner.cpp
+++ b/Engine/Source/Editor/UMGEditor/Private/TabFactory/AnimationTabSummoner.cpp
@@ -138,10 +138,10 @@ private:
 			{
 				const FScopedTransaction Transaction(TransactionName);
 				WidgetAnimation->Modify();
-				WidgetAnimation->MovieScene->Modify();
+				WidgetAnimation->GetMovieScene()->Modify();
 
 				WidgetAnimation->Rename(*InText.ToString());
-				WidgetAnimation->MovieScene->Rename(*InText.ToString());
+				WidgetAnimation->GetMovieScene()->Rename(*InText.ToString());
 
 				if(bNewAnimation)
 				{
diff --git a/Engine/Source/Editor/UMGEditor/Private/WidgetBlueprintEditor.cpp b/Engine/Source/Editor/UMGEditor/Private/WidgetBlueprintEditor.cpp
--- a/Engine/Source/Editor/UMGEditor/Private/WidgetBlueprintEditor.cpp
+++ b/Engine/Source/Editor/UMGEditor/Private/WidgetBlueprintEditor.cpp
@@ -567,7 +567,7 @@ TSharedPtr<ISequencer>& FWidgetBlueprintEditor::GetSequencer()
 			NullAnimation->MovieScene->EndTime = OutTime;
 
 			SequencerInitParams.ViewParams = ViewParams;
-			SequencerInitParams.Animation = NullAnimation;
+			SequencerInitParams.RootSequence = NullAnimation;
 			SequencerInitParams.bEditWithinLevelEditor = false;
 			SequencerInitParams.ToolkitHost = nullptr;
 		};
@@ -589,7 +589,7 @@ void FWidgetBlueprintEditor::ChangeViewedAnimation( UWidgetAnimation& InAnimatio
 	CurrentAnimation = &InAnimationToView;
 	CurrentAnimation->Initialize(PreviewWidgetPtr.Get());
 
-	Sequencer->ResetToNewAnimation(InAnimationToView);
+	Sequencer->ResetToNewRootSequence(InAnimationToView);
 
 	TSharedPtr<SOverlay> SequencerOverlayPin = SequencerOverlay.Pin();
 	if (SequencerOverlayPin.IsValid())
diff --git a/Engine/Source/Runtime/ActorAnimation/Private/ActorAnimation.cpp b/Engine/Source/Runtime/ActorAnimation/Private/ActorAnimation.cpp
--- a/Engine/Source/Runtime/ActorAnimation/Private/ActorAnimation.cpp
+++ b/Engine/Source/Runtime/ActorAnimation/Private/ActorAnimation.cpp
@@ -21,7 +21,7 @@ UActorAnimation::UActorAnimation(const FObjectInitializer& ObjectInitializer)
 void UActorAnimation::Initialize()
 {
 	// @todo sequencer: gmp: fix me
-	MovieScene = NewObject<UMovieScene>(this);
+	MovieScene = NewObject<UMovieScene>(this, NAME_None, RF_Transactional);
 }
 
 
@@ -60,10 +60,10 @@ bool UActorAnimation::CanPossessObject(UObject& Object) const
 }
 
 
-void UActorAnimation::DestroyAllSpawnedObjects(UMovieScene& SubMovieScene)
+void UActorAnimation::DestroyAllSpawnedObjects()
 {
 	const bool DestroyAll = true;
-	SpawnOrDestroyObjects(&SubMovieScene, DestroyAll);
+	SpawnOrDestroyObjects(DestroyAll);
 	SpawnedActors.Empty();
 }
 
@@ -141,12 +141,12 @@ UObject* UActorAnimation::GetParentObject(UObject* Object) const
 }
 
 
-void UActorAnimation::SpawnOrDestroyObjects(UMovieScene* SubMovieScene, bool DestroyAll)
+void UActorAnimation::SpawnOrDestroyObjects(bool DestroyAll)
 {
 	bool bAnyLevelActorsChanged = false;
 
 	// remove any proxy actors that we no longer need
-	if (SubMovieScene != nullptr)
+	if (MovieScene != nullptr)
 	{
 		for (auto SpawnedActorPair : SpawnedActors)
 		{
@@ -155,9 +155,9 @@ void UActorAnimation::SpawnOrDestroyObjects(UMovieScene* SubMovieScene, bool Des
 			if (!DestroyAll)
 			{
 				// figure out if we still need this proxy actor
-				for (auto SpawnableIndex = 0; SpawnableIndex < SubMovieScene->GetSpawnableCount(); ++SpawnableIndex)
+				for (auto SpawnableIndex = 0; SpawnableIndex < MovieScene->GetSpawnableCount(); ++SpawnableIndex)
 				{
-					auto& Spawnable = SubMovieScene->GetSpawnable(SpawnableIndex);
+					auto& Spawnable = MovieScene->GetSpawnable(SpawnableIndex);
 
 					if (Spawnable.GetGuid() == SpawnedActorPair.Key)
 					{
@@ -210,11 +210,11 @@ void UActorAnimation::SpawnOrDestroyObjects(UMovieScene* SubMovieScene, bool Des
 		}
 	}
 
-	if (!DestroyAll && (SubMovieScene != nullptr))
+	if (!DestroyAll && (MovieScene != nullptr))
 	{
-		for (auto SpawnableIndex = 0; SpawnableIndex < SubMovieScene->GetSpawnableCount(); ++SpawnableIndex)
+		for (auto SpawnableIndex = 0; SpawnableIndex < MovieScene->GetSpawnableCount(); ++SpawnableIndex)
 		{
-			FMovieSceneSpawnable& Spawnable = SubMovieScene->GetSpawnable(SpawnableIndex);
+			FMovieSceneSpawnable& Spawnable = MovieScene->GetSpawnable(SpawnableIndex);
 
 			// Do we already have a proxy for this spawnable?
 			if (SpawnedActors.Contains(Spawnable.GetGuid()))
diff --git a/Engine/Source/Runtime/ActorAnimation/Public/ActorAnimation.h b/Engine/Source/Runtime/ActorAnimation/Public/ActorAnimation.h
--- a/Engine/Source/Runtime/ActorAnimation/Public/ActorAnimation.h
+++ b/Engine/Source/Runtime/ActorAnimation/Public/ActorAnimation.h
@@ -2,7 +2,7 @@
 
 #pragma once
 
-#include "MovieSceneAnimation.h"
+#include "MovieSceneSequence.h"
 #include "ActorAnimationObject.h"
 #include "ActorAnimation.generated.h"
 
@@ -12,7 +12,7 @@
  */
 UCLASS(BlueprintType, MinimalAPI)
 class UActorAnimation
-	: public UMovieSceneAnimation
+	: public UMovieSceneSequence
 {
 	GENERATED_UCLASS_BODY()
 
@@ -26,12 +26,12 @@ public:
 	virtual bool AllowsSpawnableObjects() const override;
 	virtual void BindPossessableObject(const FGuid& ObjectId, UObject& PossessedObject) override;
 	virtual bool CanPossessObject(UObject& Object) const override;
-	virtual void DestroyAllSpawnedObjects(UMovieScene& SubMovieScene) override;
+	virtual void DestroyAllSpawnedObjects() override;
 	virtual UObject* FindObject(const FGuid& ObjectId) const override;
 	virtual FGuid FindObjectId(UObject& Object) const override;
 	virtual UMovieScene* GetMovieScene() const override;
 	virtual UObject* GetParentObject(UObject* Object) const override;
-	virtual void SpawnOrDestroyObjects(UMovieScene* SubMovieScene, bool DestroyAll) override;
+	virtual void SpawnOrDestroyObjects(bool DestroyAll) override;
 	virtual void UnbindPossessableObjects(const FGuid& ObjectId) override;
 
 #if WITH_EDITOR
diff --git a/Engine/Source/Runtime/MovieScene/Private/MovieSceneInstance.cpp b/Engine/Source/Runtime/MovieScene/Private/MovieSceneInstance.cpp
deleted file mode 100644
diff --git a/Engine/Source/Runtime/MovieScene/Private/MovieScenePrivatePCH.h b/Engine/Source/Runtime/MovieScene/Private/MovieScenePrivatePCH.h
--- a/Engine/Source/Runtime/MovieScene/Private/MovieScenePrivatePCH.h
+++ b/Engine/Source/Runtime/MovieScene/Private/MovieScenePrivatePCH.h
@@ -11,7 +11,8 @@
 #include "MovieScene.h"
 #include "MovieSceneSection.h"
 #include "MovieSceneTrack.h"
-#include "MovieSceneInstance.h"
+#include "MovieSceneSequenceInstance.h"
 #include "IMovieSceneTrackInstance.h"
+#include "MovieSceneSequence.h"
 
 #endif
\ No newline at end of file
diff --git a/Engine/Source/Runtime/MovieScene/Private/MovieSceneSequenceInstance.cpp b/Engine/Source/Runtime/MovieScene/Private/MovieSceneSequenceInstance.cpp
new file mode 100644
--- /dev/null
+++ b/Engine/Source/Runtime/MovieScene/Private/MovieSceneSequenceInstance.cpp
@@ -0,0 +1,202 @@
+// Copyright 1998-2015 Epic Games, Inc. All Rights Reserved.
+
+#include "MovieScenePrivatePCH.h"
+#include "MovieSceneSequenceInstance.h"
+#include "MovieSceneSequence.h"
+
+FMovieSceneSequenceInstance::FMovieSceneSequenceInstance(const UMovieSceneSequence& InMovieSceneSequence)
+	: MovieSceneSequence( &InMovieSceneSequence )
+{
+	TimeRange = MovieSceneSequence->GetMovieScene()->GetTimeRange();
+}
+
+
+void FMovieSceneSequenceInstance::SaveState()
+{
+	TMap<FGuid, FMovieSceneObjectBindingInstance>::TIterator ObjectIt = ObjectBindingInstances.CreateIterator();
+	for (; ObjectIt; ++ObjectIt)
+	{
+		FMovieSceneObjectBindingInstance& ObjectBindingInstance = ObjectIt.Value();
+
+		for (FMovieSceneInstanceMap::TIterator It = ObjectBindingInstance.TrackInstances.CreateIterator(); It; ++It)
+		{
+			It.Value()->SaveState(ObjectBindingInstance.RuntimeObjects);
+		}
+	}
+}
+
+
+void FMovieSceneSequenceInstance::RestoreState()
+{
+	TMap<FGuid, FMovieSceneObjectBindingInstance>::TIterator ObjectIt = ObjectBindingInstances.CreateIterator();
+	for (; ObjectIt; ++ObjectIt)
+	{
+		FMovieSceneObjectBindingInstance& ObjectBindingInstance = ObjectIt.Value();
+
+		for (FMovieSceneInstanceMap::TIterator It = ObjectBindingInstance.TrackInstances.CreateIterator(); It; ++It)
+		{
+			It.Value()->RestoreState(ObjectBindingInstance.RuntimeObjects);
+		}
+	}
+}
+
+
+void FMovieSceneSequenceInstance::Update( float Position, float LastPosition, class IMovieScenePlayer& Player )
+{
+	// Update  shot track
+	TArray<UObject*> NoObjects;
+	if( ShotTrackInstance.IsValid() )
+	{
+		ShotTrackInstance->Update( Position, LastPosition, NoObjects, Player );
+	}
+
+	// Update each master track
+	for( FMovieSceneInstanceMap::TIterator It( MasterTrackInstances ); It; ++It )
+	{
+		It.Value()->Update( Position, LastPosition, NoObjects, Player );
+	}
+
+	// Update tracks bound to objects
+	TMap<FGuid, FMovieSceneObjectBindingInstance>::TIterator ObjectIt = ObjectBindingInstances.CreateIterator();
+	for(; ObjectIt; ++ObjectIt )
+	{
+		FMovieSceneObjectBindingInstance& ObjectBindingInstance = ObjectIt.Value();
+		
+		for( FMovieSceneInstanceMap::TIterator It = ObjectBindingInstance.TrackInstances.CreateIterator(); It; ++It )
+		{
+			It.Value()->Update( Position, LastPosition, ObjectBindingInstance.RuntimeObjects, Player );
+		}
+	}
+}
+
+
+void FMovieSceneSequenceInstance::RefreshInstance( IMovieScenePlayer& Player )
+{
+	UMovieScene* MovieScene = MovieSceneSequence->GetMovieScene();
+	TimeRange = MovieScene->GetTimeRange();
+
+	UMovieSceneTrack* ShotTrack = MovieScene->GetShotTrack();
+
+	TSharedRef<FMovieSceneSequenceInstance> ThisInstance = AsShared();
+
+	FMovieSceneInstanceMap ShotTrackInstanceMap;
+	// Only if root movie scene. Any sub-movie scene that has a shot track is ignored
+	if( ShotTrack && Player.GetRootMovieSceneSequenceInstance() == ThisInstance )
+	{
+		if( ShotTrackInstance.IsValid() )
+		{
+			ShotTrackInstanceMap.Add(  ShotTrack, ShotTrackInstance );
+		}
+
+		TArray<UObject*> Objects;
+		TArray<UMovieSceneTrack*> Tracks;
+		Tracks.Add(ShotTrack);
+		RefreshInstanceMap(Tracks, Objects, ShotTrackInstanceMap, Player);
+
+		ShotTrackInstance = ShotTrackInstanceMap.FindRef( ShotTrack );
+	}
+	else if( !ShotTrack && ShotTrackInstance.IsValid() )
+	{
+		ShotTrackInstance->ClearInstance( Player );
+		ShotTrackInstance.Reset();
+	}
+
+	// Get all the master tracks and create instances for them if needed
+	const TArray<UMovieSceneTrack*>& MasterTracks = MovieScene->GetMasterTracks();
+	TArray<UObject*> Objects;
+	RefreshInstanceMap( MasterTracks, Objects, MasterTrackInstances, Player );
+
+	TSet< FGuid > FoundObjectBindings;
+	// Get all tracks for each object binding and create instances for them if needed
+	const TArray<FMovieSceneBinding>& ObjectBindings = MovieScene->GetBindings();
+	for( int32 BindingIndex = 0; BindingIndex < ObjectBindings.Num(); ++BindingIndex )
+	{
+		const FMovieSceneBinding& ObjectBinding = ObjectBindings[BindingIndex];
+
+		// Create an instance for this object binding
+		FMovieSceneObjectBindingInstance& BindingInstance = ObjectBindingInstances.FindOrAdd( ObjectBinding.GetObjectGuid() );
+		BindingInstance.ObjectGuid = ObjectBinding.GetObjectGuid();
+
+		FoundObjectBindings.Add( ObjectBinding.GetObjectGuid() );
+
+		// Spawn the runtime objects
+		// @todo Sequencer SubMovieScenes: We need to know which actors were removed and which actors were added so we know which saved actor state to restore/create
+		BindingInstance.RuntimeObjects.Empty();
+		Player.GetRuntimeObjects( SharedThis( this ), BindingInstance.ObjectGuid, BindingInstance.RuntimeObjects );
+		
+		// Refresh the instance's tracks
+		const TArray<UMovieSceneTrack*>& Tracks = ObjectBinding.GetTracks();
+		RefreshInstanceMap( Tracks, BindingInstance.RuntimeObjects, BindingInstance.TrackInstances, Player );
+	}
+
+
+	// Remove object binding instances which are no longer bound
+	TMap<FGuid, FMovieSceneObjectBindingInstance>::TIterator It = ObjectBindingInstances.CreateIterator();
+	for( ; It; ++It )
+	{
+		if( !FoundObjectBindings.Contains( It.Key() ) )
+		{
+			// The instance no longer is bound to an existing guid
+			It.RemoveCurrent();
+		}
+	}
+}
+
+
+struct FTrackInstanceEvalSorter
+{
+	bool operator()( const TSharedPtr<IMovieSceneTrackInstance> A, const TSharedPtr<IMovieSceneTrackInstance> B ) const
+	{
+		return A->EvalOrder() > B->EvalOrder();
+	}
+};
+
+
+void FMovieSceneSequenceInstance::RefreshInstanceMap( const TArray<UMovieSceneTrack*>& Tracks, const TArray<UObject*>& RuntimeObjects, FMovieSceneInstanceMap& TrackInstances, IMovieScenePlayer& Player  )
+{
+	// All the tracks we found during this pass
+	TSet< UMovieSceneTrack* > FoundTracks;
+
+	// For every track, check if it has an instance, if not create one, otherwise refresh that instance
+	for( int32 TrackIndex = 0; TrackIndex < Tracks.Num(); ++TrackIndex )
+	{
+		UMovieSceneTrack* Track = Tracks[TrackIndex];
+
+		// A new track has been encountered
+		FoundTracks.Add( Track );
+
+		// See if the track has an instance
+		TSharedPtr<IMovieSceneTrackInstance> Instance = TrackInstances.FindRef( Track );
+		if ( !Instance.IsValid() )
+		{
+			// The track does not have an instance, create one
+			Instance = Track->CreateInstance();
+			Instance->RefreshInstance( RuntimeObjects, Player );
+			Instance->SaveState(RuntimeObjects);
+
+			TrackInstances.Add( Track, Instance );
+		}
+		else
+		{
+			// The track has an instance, refresh it
+			Instance->RefreshInstance( RuntimeObjects, Player );
+		}
+
+	}
+
+	// Remove entries which no longer have a track associated with them
+	FMovieSceneInstanceMap::TIterator It = TrackInstances.CreateIterator();
+	for( ; It; ++It )
+	{
+		if( !FoundTracks.Contains( It.Key().Get() ) )
+		{
+			It.Value()->ClearInstance( Player );
+
+			// This track was not found in the moviescene's track list so it was removed.
+			It.RemoveCurrent();
+		}
+	}
+
+	// Sort based on evaluation order
+	TrackInstances.ValueSort(FTrackInstanceEvalSorter());
+}
diff --git a/Engine/Source/Runtime/MovieScene/Public/IMovieScenePlayer.h b/Engine/Source/Runtime/MovieScene/Public/IMovieScenePlayer.h
--- a/Engine/Source/Runtime/MovieScene/Public/IMovieScenePlayer.h
+++ b/Engine/Source/Runtime/MovieScene/Public/IMovieScenePlayer.h
@@ -14,7 +14,7 @@ namespace EMovieScenePlayerStatus
 	};
 }
 
-class FMovieSceneInstance;
+class FMovieSceneSequenceInstance;
 
 /**
  * Interface for movie scene players
@@ -30,7 +30,7 @@ public:
 	 * @param ObjectHandle The handle to runtime objects
 	 * @Param The list of runtime objects that will be populated
 	 */
-	virtual void GetRuntimeObjects( TSharedRef<FMovieSceneInstance> MovieSceneInstance, const FGuid& ObjectHandle, TArray< UObject* >& OutObjects ) const = 0;
+	virtual void GetRuntimeObjects( TSharedRef<FMovieSceneSequenceInstance> MovieSceneInstance, const FGuid& ObjectHandle, TArray< UObject* >& OutObjects ) const = 0;
 
 
 	/**
@@ -47,7 +47,7 @@ public:
 	 * @param MovieSceneSection	The section owning the MovieScene being instanced. 
 	 * @param InstanceToAdd		The instance being added
 	 */
-	virtual void AddOrUpdateMovieSceneInstance( class UMovieSceneSection& MovieSceneSection, TSharedRef<FMovieSceneInstance> InstanceToAdd ) = 0;
+	virtual void AddOrUpdateMovieSceneInstance( class UMovieSceneSection& MovieSceneSection, TSharedRef<FMovieSceneSequenceInstance> InstanceToAdd ) = 0;
 
 	/**
 	 * Removes a MovieScene instance from the player. 
@@ -55,12 +55,12 @@ public:
 	 * @param MovieSceneSection	The section owning the MovieScene that was instanced. 
 	 * @param InstanceToAdd		The instance being removed
 	 */
-	virtual void RemoveMovieSceneInstance( class UMovieSceneSection& MovieSceneSection, TSharedRef<FMovieSceneInstance> InstanceToRemove ) = 0;
+	virtual void RemoveMovieSceneInstance( class UMovieSceneSection& MovieSceneSection, TSharedRef<FMovieSceneSequenceInstance> InstanceToRemove ) = 0;
 
 	/**
 	 * @return The root movie scene instance being played
 	 */
-	virtual TSharedRef<FMovieSceneInstance> GetRootMovieSceneInstance() const = 0;
+	virtual TSharedRef<FMovieSceneSequenceInstance> GetRootMovieSceneSequenceInstance() const = 0;
 
 	/** @return whether the player is currently playing, scrubbing, etc. */
 	virtual EMovieScenePlayerStatus::Type GetPlaybackStatus() const = 0;
diff --git a/Engine/Source/Runtime/MovieScene/Public/IMovieSceneTrackInstance.h b/Engine/Source/Runtime/MovieScene/Public/IMovieSceneTrackInstance.h
--- a/Engine/Source/Runtime/MovieScene/Public/IMovieSceneTrackInstance.h
+++ b/Engine/Source/Runtime/MovieScene/Public/IMovieSceneTrackInstance.h
@@ -7,7 +7,7 @@ class IMovieScenePlayer;
 
 
 /**
- * Base class for an animated type instance in a Movie Scene
+ * A track instance holds the live objects for a track.  
  */
 class IMovieSceneTrackInstance
 {
@@ -18,11 +18,13 @@ public:
 
 	/**
 	 * Save state of objects that this instance will be editing.
+	 * @todo Sequencer: This is likely editor only
 	 */
 	virtual void SaveState(const TArray<UObject*>& RuntimeObjects) = 0;
 
 	/**
 	 * Restore state of objects that this instance edited.
+	 * @todo Sequencer: This is likely editor only
 	 */
 	virtual void RestoreState(const TArray<UObject*>& RuntimeObjects) = 0;
 
diff --git a/Engine/Source/Runtime/MovieScene/Public/MovieSceneAnimation.h b/Engine/Source/Runtime/MovieScene/Public/MovieSceneAnimation.h
deleted file mode 100644
diff --git a/Engine/Source/Runtime/MovieScene/Public/MovieSceneInstance.h b/Engine/Source/Runtime/MovieScene/Public/MovieSceneInstance.h
deleted file mode 100644
diff --git a/Engine/Source/Runtime/MovieScene/Public/MovieSceneSequence.h b/Engine/Source/Runtime/MovieScene/Public/MovieSceneSequence.h
new file mode 100644
--- /dev/null
+++ b/Engine/Source/Runtime/MovieScene/Public/MovieSceneSequence.h
@@ -0,0 +1,118 @@
+// Copyright 1998-2015 Epic Games, Inc. All Rights Reserved.
+
+#pragma once
+
+#include "MovieSceneSequence.generated.h"
+
+
+class FMovieSceneSequenceInstance;
+class UMovieScene;
+class UObject;
+
+
+/**
+ * Abstract base class for movie scene animations (C++ version).
+ */
+UCLASS(MinimalAPI)
+class UMovieSceneSequence
+	: public UObject
+{
+	GENERATED_BODY()
+
+public:
+
+	/**
+	 * Whether objects can be spawned at run-time.
+	 *
+	 * @return true if objects can be spawned by sequencer, false if only existing objects can be possessed.
+	 */
+	virtual bool AllowsSpawnableObjects() const PURE_VIRTUAL(UMovieSceneSequence::AllowsSpawnableObjects, return false;);
+
+	/**
+	 * Called when Sequencer has created an object binding for a possessable object
+	 * 
+	 * @param ObjectId The guid used to map to the possessable object.  Note the guid can be bound to multiple objects at once
+	 * @param PossessedObject The runtime object which was possessed.
+	 * @see UnbindPossessableObjects
+	 */
+	virtual void BindPossessableObject(const FGuid& ObjectId, UObject& PossessedObject) PURE_VIRTUAL(UMovieSceneSequence::BindPossessableObject,);
+	
+	/**
+	 * Check whether the given object can be possessed by this animation.
+	 *
+	 * @param Object The object to check.
+	 * @return true if the object can be possessed, false otherwise.
+	 */
+	virtual bool CanPossessObject(UObject& Object) const PURE_VIRTUAL(UMovieSceneSequence::CanPossessObject, return false;);
+
+	/**
+	 * Destroy all objects that were spawned by this animation.
+	 *
+	 * @param SubMovieScene The (sub) movie scene to destroy objects for.
+	 */
+	virtual void DestroyAllSpawnedObjects() PURE_VIRTUAL(UMovieSceneSequence::DestroyAllSpawnedObjects,);
+
+	/**
+	 * Gets the possessed or spawned object for the specified identifier.
+	 *
+	 * @param ObjectId The unique identifier of the object.
+	 * @return The object, or nullptr if not found.
+	 */
+	virtual UObject* FindObject(const FGuid& ObjectId) const PURE_VIRTUAL(UMovieSceneSequence::FindObject, return nullptr;);
+
+	/**
+	 * Find the identifier for a possessed or spawned object.
+	 *
+	 * @param The object to get the identifier for.
+	 * @return The object identifier, or an invalid GUID if not found.
+	 */
+	virtual FGuid FindObjectId( UObject& Object) const PURE_VIRTUAL(UMovieSceneSequence::FindGuidForObject, return FGuid(););
+
+	/**
+	 * Get the movie scene that controls this animation.
+	 *
+	 * The returned movie scene represents the root movie scene.
+	 * It may contain additional child movie scenes.
+	 *
+	 * @return The movie scene.
+	 */
+	virtual UMovieScene* GetMovieScene() const PURE_VIRTUAL(UMovieSceneSequence::GetMovieScene(), return nullptr;);
+
+	/**
+	 * Get the logical parent object for the supplied object (not necessarily its outer).
+	 *
+	 * @param Object The object whose parent to get.
+	 * @return The parent object, or nullptr if the object has no logical parent.
+	 */
+	virtual UObject* GetParentObject(UObject* Object) const PURE_VIRTUAL(UMovieSceneSequence::GetParentObject(), return nullptr;);
+
+	/**
+	 * Called to spawn or destroy objects for a movie instance.
+	 *
+	 * @param SubMovieScene The (sub) movie scene to spawn or destroy objects for.
+	 * @param DestroyAll If true, destroy all spawned objects for the instance, if false only destroy unused objects.
+	 */
+	virtual void SpawnOrDestroyObjects(bool DestroyAll) PURE_VIRTUAL(UMovieSceneSequence::SpawnOrDestroyObjects,);
+
+	/**
+	 * Unbinds all possessable objects from the provided GUID.
+	 *
+	 * @param ObjectId The guid bound to possessable objects that should be removed.
+	 * @see BindPossessableObject
+	 */
+	virtual void UnbindPossessableObjects(const FGuid& ObjectId) PURE_VIRTUAL(UMovieSceneSequence::UnbindPossessableObjects,);
+
+public:
+
+#if WITH_EDITOR
+	/**
+	 * 
+	 * Tries to get a display name for the binding represented by the guid.
+	 *
+	 * @param ObjectGuid The guid for the object binding.
+	 * @param DisplayName The display name for the object binding.
+	 * @returns true if DisplayName has been set to a valid display name, otherwise false.
+	 */
+	virtual bool TryGetObjectDisplayName(const FGuid& ObjectId, FText& OutDisplayName) const PURE_VIRTUAL(UMovieSceneSequence::TryGetObjectDisplayName, return false;);
+#endif
+};
diff --git a/Engine/Source/Runtime/MovieScene/Public/MovieSceneSequenceInstance.h b/Engine/Source/Runtime/MovieScene/Public/MovieSceneSequenceInstance.h
new file mode 100644
--- /dev/null
+++ b/Engine/Source/Runtime/MovieScene/Public/MovieSceneSequenceInstance.h
@@ -0,0 +1,108 @@
+// Copyright 1998-2015 Epic Games, Inc. All Rights Reserved.
+
+#pragma once
+
+
+class IMovieSceneTrackInstance;
+class UMovieSceneTrack;
+class UMovieSceneSequence;
+
+typedef TMap<TWeakObjectPtr<UMovieSceneTrack>, TSharedPtr<class IMovieSceneTrackInstance>> FMovieSceneInstanceMap;
+
+/**
+ * A sequence instance holds the live objects bound to the tracks in a sequence.  This is completely transient
+ * There is one instance per sequence.  If a sequence holds multiple sub-sequences each sub-sequence will have it's own instance even if they are the same sequence asset
+ * A sequence instance also creates and manages all track instances for the tracks in a sequence
+ */
+class FMovieSceneSequenceInstance
+	: public TSharedFromThis<FMovieSceneSequenceInstance>
+{
+public:
+
+	/**
+	 * Creates and initializes a new instance from a Sequence
+	 *
+	 * @param InMovieSceneSequence The sequence that this instance represents.
+	 */
+	MOVIESCENE_API FMovieSceneSequenceInstance(const UMovieSceneSequence& InMovieSceneSequence);
+
+	/** Save state of the objects that this movie scene controls. */
+	MOVIESCENE_API void SaveState();
+
+	/** Restore state of the objects that this movie scene controls. */
+	MOVIESCENE_API void RestoreState();
+
+	/**
+	 * Updates this movie scene.
+	 *
+	 * @param Position The local playback position.
+	 * @param Player Movie scene player interface for interaction with runtime data.
+	 */
+	MOVIESCENE_API void Update(float Position, float LastPosition, class IMovieScenePlayer& Player);
+
+	/**
+	 * Refreshes the existing instance.
+	 *
+	 * Called when something significant about movie scene data changes (like adding or removing a track).
+	 * Instantiates all new tracks found and removes instances for tracks that no longer exist.
+	 */
+	MOVIESCENE_API void RefreshInstance(class IMovieScenePlayer& Player);
+
+	/**
+	 * Get the sequence associated with this instance.
+	 *
+	 * @return The movie scene sequence object.
+	 */
+	UMovieSceneSequence* GetSequence()
+	{
+		return MovieSceneSequence.Get();
+	}
+
+	/** 
+	 * Get the time range for the movie scene.
+	 *
+	 * Note: This is not necessarily how long the instance will be playing.
+	 * The parent could have changed the length.
+	 *
+	 * @return The time range.
+	 */
+	TRange<float> GetMovieSceneTimeRange() const
+	{
+		return TimeRange;
+	}
+
+protected:
+
+	void RefreshInstanceMap(const TArray<UMovieSceneTrack*>& Tracks, const TArray<UObject*>& RuntimeObjects, FMovieSceneInstanceMap& TrackInstances, class IMovieScenePlayer& Player);
+
+private:
+
+	/** A paring of a runtime object guid to a runtime data and track instances animating the runtime objects */
+	struct FMovieSceneObjectBindingInstance
+	{
+		/** Runtime object guid for looking up runtime UObjects */
+		FGuid ObjectGuid;
+
+		/** Actual runtime objects */	
+		TArray<UObject*> RuntimeObjects;
+
+		/** Instances that animate the runtime objects */
+		FMovieSceneInstanceMap TrackInstances;
+	};
+
+
+	/** MovieScene that is instanced */
+	const TWeakObjectPtr<UMovieSceneSequence> MovieSceneSequence;
+
+	/** The shot track instance map */
+	TSharedPtr<IMovieSceneTrackInstance> ShotTrackInstance;
+
+	/** All Master track instances */
+	FMovieSceneInstanceMap MasterTrackInstances;
+
+	/** All object binding instances */
+	TMap<FGuid, FMovieSceneObjectBindingInstance> ObjectBindingInstances;
+
+	/** Cached time range for the movie scene */
+	TRange<float> TimeRange;
+};
diff --git a/Engine/Source/Runtime/MovieSceneTracks/Private/Sections/MovieSceneAnimationSection.cpp b/Engine/Source/Runtime/MovieSceneTracks/Private/Sections/MovieSceneAnimationSection.cpp
deleted file mode 100644
diff --git a/Engine/Source/Runtime/MovieSceneTracks/Private/Sections/MovieSceneSkeletalAnimationSection.cpp b/Engine/Source/Runtime/MovieSceneTracks/Private/Sections/MovieSceneSkeletalAnimationSection.cpp
new file mode 100644
--- /dev/null
+++ b/Engine/Source/Runtime/MovieSceneTracks/Private/Sections/MovieSceneSkeletalAnimationSection.cpp
@@ -0,0 +1,45 @@
+// Copyright 1998-2015 Epic Games, Inc. All Rights Reserved.
+
+#include "MovieSceneTracksPrivatePCH.h"
+#include "MovieSceneSkeletalAnimationSection.h"
+
+UMovieSceneSkeletalAnimationSection::UMovieSceneSkeletalAnimationSection( const FObjectInitializer& ObjectInitializer )
+	: Super( ObjectInitializer )
+{
+	AnimSequence = NULL;
+	AnimationStartTime = 0.f;
+	AnimationDilationFactor = 1.f;
+}
+
+void UMovieSceneSkeletalAnimationSection::MoveSection( float DeltaTime, TSet<FKeyHandle>& KeyHandles )
+{
+	Super::MoveSection(DeltaTime, KeyHandles);
+
+	AnimationStartTime += DeltaTime;
+}
+
+void UMovieSceneSkeletalAnimationSection::DilateSection( float DilationFactor, float Origin, TSet<FKeyHandle>& KeyHandles )
+{
+	AnimationStartTime = (AnimationStartTime - Origin) * DilationFactor + Origin;
+	AnimationDilationFactor *= DilationFactor;
+
+	Super::DilateSection(DilationFactor, Origin, KeyHandles);
+}
+
+void UMovieSceneSkeletalAnimationSection::GetKeyHandles(TSet<FKeyHandle>& KeyHandles) const
+{
+}
+
+void UMovieSceneSkeletalAnimationSection::GetSnapTimes(TArray<float>& OutSnapTimes, bool bGetSectionBorders) const
+{
+	Super::GetSnapTimes(OutSnapTimes, bGetSectionBorders);
+	float CurrentTime = GetAnimationStartTime();
+	while (CurrentTime <= GetEndTime())
+	{
+		if (CurrentTime >= GetStartTime())
+		{
+			OutSnapTimes.Add(CurrentTime);
+		}
+		CurrentTime += GetAnimationDuration();
+	}
+}
diff --git a/Engine/Source/Runtime/MovieSceneTracks/Private/Sections/SubMovieSceneSection.cpp b/Engine/Source/Runtime/MovieSceneTracks/Private/Sections/SubMovieSceneSection.cpp
--- a/Engine/Source/Runtime/MovieSceneTracks/Private/Sections/SubMovieSceneSection.cpp
+++ b/Engine/Source/Runtime/MovieSceneTracks/Private/Sections/SubMovieSceneSection.cpp
@@ -3,8 +3,19 @@
 #include "MovieSceneTracksPrivatePCH.h"
 #include "SubMovieSceneSection.h"
 #include "IMovieScenePlayer.h"
+#include "MovieSceneSequence.h"
 
 USubMovieSceneSection::USubMovieSceneSection( const FObjectInitializer& ObjectInitializer )
 	: Super( ObjectInitializer )
 {
 }
+
+void USubMovieSceneSection::SetMovieSceneAnimation( UMovieSceneSequence* InSequence )
+{
+	Sequence = InSequence;
+}
+
+UMovieSceneSequence* USubMovieSceneSection::GetMovieSceneAnimation() const
+{
+	return Sequence;
+}
\ No newline at end of file
diff --git a/Engine/Source/Runtime/MovieSceneTracks/Private/TrackInstances/MovieScene3DPathTrackInstance.cpp b/Engine/Source/Runtime/MovieSceneTracks/Private/TrackInstances/MovieScene3DPathTrackInstance.cpp
--- a/Engine/Source/Runtime/MovieSceneTracks/Private/TrackInstances/MovieScene3DPathTrackInstance.cpp
+++ b/Engine/Source/Runtime/MovieSceneTracks/Private/TrackInstances/MovieScene3DPathTrackInstance.cpp
@@ -67,7 +67,7 @@ void FMovieScene3DPathTrackInstance::Update( float Position, float LastPosition,
 
 			if (PathId.IsValid())
 			{
-				Player.GetRuntimeObjects( Player.GetRootMovieSceneInstance(), PathId, PathObjects);
+				Player.GetRuntimeObjects( Player.GetRootMovieSceneSequenceInstance(), PathId, PathObjects);
 
 				for (int32 PathObjectIndex = 0; PathObjectIndex < PathObjects.Num(); ++PathObjectIndex)
 				{
diff --git a/Engine/Source/Runtime/MovieSceneTracks/Private/TrackInstances/MovieSceneAnimationTrackInstance.cpp b/Engine/Source/Runtime/MovieSceneTracks/Private/TrackInstances/MovieSceneAnimationTrackInstance.cpp
deleted file mode 100644
diff --git a/Engine/Source/Runtime/MovieSceneTracks/Private/TrackInstances/MovieSceneShotTrackInstance.cpp b/Engine/Source/Runtime/MovieSceneTracks/Private/TrackInstances/MovieSceneShotTrackInstance.cpp
--- a/Engine/Source/Runtime/MovieSceneTracks/Private/TrackInstances/MovieSceneShotTrackInstance.cpp
+++ b/Engine/Source/Runtime/MovieSceneTracks/Private/TrackInstances/MovieSceneShotTrackInstance.cpp
@@ -22,7 +22,7 @@ void FMovieSceneShotTrackInstance::RefreshInstance( const TArray<UObject*>& Runt
 	{
 		// @todo Sequencer - Sub-moviescenes: Get the cameras from the root movie scene instance.  We should support adding cameras for sub-moviescenes as shots
 		TArray<UObject*> CameraObjects;
-		Player.GetRuntimeObjects(Player.GetRootMovieSceneInstance(), CastChecked<UMovieSceneShotSection>( Section )->GetCameraGuid(), CameraObjects );
+		Player.GetRuntimeObjects(Player.GetRootMovieSceneSequenceInstance(), CastChecked<UMovieSceneShotSection>( Section )->GetCameraGuid(), CameraObjects );
 		if( CameraObjects.Num() == 1 )
 		{
 			RuntimeCameraObjects.Add( CameraObjects[0] );
diff --git a/Engine/Source/Runtime/MovieSceneTracks/Private/TrackInstances/MovieSceneSkeletalAnimationTrackInstance.cpp b/Engine/Source/Runtime/MovieSceneTracks/Private/TrackInstances/MovieSceneSkeletalAnimationTrackInstance.cpp
new file mode 100644
--- /dev/null
+++ b/Engine/Source/Runtime/MovieSceneTracks/Private/TrackInstances/MovieSceneSkeletalAnimationTrackInstance.cpp
@@ -0,0 +1,42 @@
+// Copyright 1998-2015 Epic Games, Inc. All Rights Reserved.
+
+#include "MovieSceneTracksPrivatePCH.h"
+#include "MovieSceneSkeletalAnimationTrack.h"
+#include "MovieSceneSkeletalAnimationTrackInstance.h"
+#include "MovieSceneSkeletalAnimationSection.h"
+#include "IMovieScenePlayer.h"
+#include "Matinee/MatineeAnimInterface.h"
+
+FMovieSceneSkeletalAnimationTrackInstance::FMovieSceneSkeletalAnimationTrackInstance( UMovieSceneSkeletalAnimationTrack& InAnimationTrack )
+{
+	AnimationTrack = &InAnimationTrack;
+}
+
+FMovieSceneSkeletalAnimationTrackInstance::~FMovieSceneSkeletalAnimationTrackInstance()
+{
+	// @todo Sequencer Need to find some way to call PreviewFinishAnimControl (needs the runtime objects)
+}
+
+void FMovieSceneSkeletalAnimationTrackInstance::Update( float Position, float LastPosition, const TArray<UObject*>& RuntimeObjects, class IMovieScenePlayer& Player ) 
+{
+	// @todo Sequencer gameplay update has a different code path than editor update for animation
+
+	for (int32 i = 0; i < RuntimeObjects.Num(); ++i)
+	{
+		IMatineeAnimInterface* AnimInterface = Cast<IMatineeAnimInterface>(RuntimeObjects[i]);
+		if (AnimInterface)
+		{
+			UMovieSceneSkeletalAnimationSection* AnimSection = Cast<UMovieSceneSkeletalAnimationSection>(AnimationTrack->GetAnimSectionAtTime(Position));
+			if (AnimSection)
+			{
+				int32 ChannelIndex = 0;
+				FName SlotName = FName("AnimationSlot");
+				UAnimSequence* AnimSequence = AnimSection->GetAnimSequence();
+				float AnimPosition = FMath::Fmod(Position - AnimSection->GetAnimationStartTime(), AnimSection->GetAnimationDuration()) / AnimSection->GetAnimationDilationFactor();
+
+				AnimInterface->PreviewBeginAnimControl(NULL);
+				AnimInterface->PreviewSetAnimPosition(SlotName, ChannelIndex, AnimSequence, AnimPosition, true, false, 0.f);
+			}
+		}
+	}
+}
diff --git a/Engine/Source/Runtime/MovieSceneTracks/Private/TrackInstances/SubMovieSceneTrackInstance.cpp b/Engine/Source/Runtime/MovieSceneTracks/Private/TrackInstances/SubMovieSceneTrackInstance.cpp
--- a/Engine/Source/Runtime/MovieSceneTracks/Private/TrackInstances/SubMovieSceneTrackInstance.cpp
+++ b/Engine/Source/Runtime/MovieSceneTracks/Private/TrackInstances/SubMovieSceneTrackInstance.cpp
@@ -3,7 +3,7 @@
 #include "MovieSceneTracksPrivatePCH.h"
 #include "SubMovieSceneTrack.h"
 #include "SubMovieSceneTrackInstance.h"
-#include "MovieSceneInstance.h"
+#include "MovieSceneSequenceInstance.h"
 #include "SubMovieSceneSection.h"
 #include "IMovieScenePlayer.h"
 
@@ -28,14 +28,14 @@ void FSubMovieSceneTrackInstance::RefreshInstance( const TArray<UObject*>& Runti
 
 		// If the section doesn't have a valid movie scene or no longer has one 
 		// (e.g user deleted it) then skip adding an instance for it
-		if( Section->GetMovieScene() )
+		if( Section->GetMovieSceneAnimation() )
 		{
 			FoundSections.Add(Section);
 
-			TSharedPtr<FMovieSceneInstance> Instance = SubMovieSceneInstances.FindRef(Section);
+			TSharedPtr<FMovieSceneSequenceInstance> Instance = SubMovieSceneInstances.FindRef(Section);
 			if (!Instance.IsValid())
 			{
-				Instance = MakeShareable(new FMovieSceneInstance(*Section->GetMovieScene()));
+				Instance = MakeShareable(new FMovieSceneSequenceInstance(*Section->GetMovieSceneAnimation()));
 
 				SubMovieSceneInstances.Add(Section, Instance.ToSharedRef());
 			}
@@ -47,7 +47,7 @@ void FSubMovieSceneTrackInstance::RefreshInstance( const TArray<UObject*>& Runti
 		}
 	}
 
-	TMap< TWeakObjectPtr<USubMovieSceneSection>, TSharedPtr<FMovieSceneInstance> >::TIterator It =  SubMovieSceneInstances.CreateIterator();
+	TMap< TWeakObjectPtr<USubMovieSceneSection>, TSharedPtr<FMovieSceneSequenceInstance> >::TIterator It =  SubMovieSceneInstances.CreateIterator();
 	for( ; It; ++It )
 	{
 		// Remove any sections that no longer exist
@@ -70,7 +70,7 @@ void FSubMovieSceneTrackInstance::Update( float Position, float LastPosition, co
 	{
 		USubMovieSceneSection* Section = CastChecked<USubMovieSceneSection>( TraversedSections[SectionIndex] );
 
-		TSharedPtr<FMovieSceneInstance>& Instance = SubMovieSceneInstances.FindChecked( Section );
+		TSharedPtr<FMovieSceneSequenceInstance>& Instance = SubMovieSceneInstances.FindChecked( Section );
 
 		TRange<float> TimeRange = Instance->GetMovieSceneTimeRange();
 
@@ -93,7 +93,7 @@ void FSubMovieSceneTrackInstance::SaveState(const TArray<UObject*>& RuntimeObjec
 	{
 		USubMovieSceneSection* Section = CastChecked<USubMovieSceneSection>( AllSections[SectionIndex] );
 
-		TSharedPtr<FMovieSceneInstance> Instance = SubMovieSceneInstances.FindRef( Section );
+		TSharedPtr<FMovieSceneSequenceInstance> Instance = SubMovieSceneInstances.FindRef( Section );
 		if( Instance.IsValid() )
 		{
 			Instance->SaveState();
@@ -109,7 +109,7 @@ void FSubMovieSceneTrackInstance::RestoreState(const TArray<UObject*>& RuntimeOb
 	{
 		USubMovieSceneSection* Section = CastChecked<USubMovieSceneSection>( AllSections[SectionIndex] );
 
-		TSharedPtr<FMovieSceneInstance> Instance = SubMovieSceneInstances.FindRef( Section );
+		TSharedPtr<FMovieSceneSequenceInstance> Instance = SubMovieSceneInstances.FindRef( Section );
 		if( Instance.IsValid() )
 		{
 			Instance->RestoreState();
@@ -119,7 +119,7 @@ void FSubMovieSceneTrackInstance::RestoreState(const TArray<UObject*>& RuntimeOb
 
 void FSubMovieSceneTrackInstance::ClearInstance( IMovieScenePlayer& Player )
 {
-	TMap< TWeakObjectPtr<USubMovieSceneSection>, TSharedPtr<FMovieSceneInstance> >::TIterator It =  SubMovieSceneInstances.CreateIterator();
+	TMap< TWeakObjectPtr<USubMovieSceneSection>, TSharedPtr<FMovieSceneSequenceInstance> >::TIterator It =  SubMovieSceneInstances.CreateIterator();
 	for( ; It; ++It )
 	{
 		Player.RemoveMovieSceneInstance( *It.Key().Get(), It.Value().ToSharedRef() );
diff --git a/Engine/Source/Runtime/MovieSceneTracks/Private/Tracks/MovieSceneAnimationTrack.cpp b/Engine/Source/Runtime/MovieSceneTracks/Private/Tracks/MovieSceneAnimationTrack.cpp
deleted file mode 100644
diff --git a/Engine/Source/Runtime/MovieSceneTracks/Private/Tracks/MovieSceneShotTrack.cpp b/Engine/Source/Runtime/MovieSceneTracks/Private/Tracks/MovieSceneShotTrack.cpp
--- a/Engine/Source/Runtime/MovieSceneTracks/Private/Tracks/MovieSceneShotTrack.cpp
+++ b/Engine/Source/Runtime/MovieSceneTracks/Private/Tracks/MovieSceneShotTrack.cpp
@@ -33,14 +33,14 @@ void UMovieSceneShotTrack::RemoveSection( UMovieSceneSection* Section )
 	// @todo Sequencer: The movie scene owned by the section is now abandoned.  Should we offer to delete it?  
 }
 
-void UMovieSceneShotTrack::AddNewShot(FGuid CameraHandle, UMovieScene& ShotMovieScene, const TRange<float>& TimeRange, const FText& ShotName, int32 ShotNumber )
+void UMovieSceneShotTrack::AddNewShot(FGuid CameraHandle, UMovieSceneSequence& ShotMovieSceneAnimation, const TRange<float>& TimeRange, const FText& ShotName, int32 ShotNumber )
 {
 	Modify();
 
 	FName UniqueShotName = MakeUniqueObjectName( this, UMovieSceneShotSection::StaticClass(), *ShotName.ToString() );
 
 	UMovieSceneShotSection* NewSection = NewObject<UMovieSceneShotSection>( this, UniqueShotName, RF_Transactional );
-	NewSection->SetMovieScene( &ShotMovieScene );
+	NewSection->SetMovieSceneAnimation( &ShotMovieSceneAnimation );
 	NewSection->SetStartTime( TimeRange.GetLowerBoundValue() );
 	NewSection->SetEndTime( TimeRange.GetUpperBoundValue() );
 	NewSection->SetCameraGuid( CameraHandle );
diff --git a/Engine/Source/Runtime/MovieSceneTracks/Private/Tracks/MovieSceneSkeletalAnimationTrack.cpp b/Engine/Source/Runtime/MovieSceneTracks/Private/Tracks/MovieSceneSkeletalAnimationTrack.cpp
new file mode 100644
--- /dev/null
+++ b/Engine/Source/Runtime/MovieSceneTracks/Private/Tracks/MovieSceneSkeletalAnimationTrack.cpp
@@ -0,0 +1,87 @@
+// Copyright 1998-2015 Epic Games, Inc. All Rights Reserved.
+
+#include "MovieSceneTracksPrivatePCH.h"
+#include "MovieSceneSkeletalAnimationSection.h"
+#include "MovieSceneSkeletalAnimationTrack.h"
+#include "IMovieScenePlayer.h"
+#include "MovieSceneSkeletalAnimationTrackInstance.h"
+
+#define LOCTEXT_NAMESPACE "MovieSceneSkeletalAnimationTrack"
+
+UMovieSceneSkeletalAnimationTrack::UMovieSceneSkeletalAnimationTrack( const FObjectInitializer& ObjectInitializer )
+	: Super( ObjectInitializer )
+{
+}
+
+FName UMovieSceneSkeletalAnimationTrack::GetTrackName() const
+{
+	return FName("Animation");
+}
+
+
+TSharedPtr<IMovieSceneTrackInstance> UMovieSceneSkeletalAnimationTrack::CreateInstance()
+{
+	return MakeShareable( new FMovieSceneSkeletalAnimationTrackInstance( *this ) ); 
+}
+
+const TArray<UMovieSceneSection*>& UMovieSceneSkeletalAnimationTrack::GetAllSections() const
+{
+	return AnimationSections;
+}
+
+void UMovieSceneSkeletalAnimationTrack::RemoveAllAnimationData()
+{
+}
+
+bool UMovieSceneSkeletalAnimationTrack::HasSection( UMovieSceneSection* Section ) const
+{
+	return AnimationSections.Find( Section ) != INDEX_NONE;
+}
+
+void UMovieSceneSkeletalAnimationTrack::RemoveSection( UMovieSceneSection* Section )
+{
+	AnimationSections.Remove( Section );
+}
+
+bool UMovieSceneSkeletalAnimationTrack::IsEmpty() const
+{
+	return AnimationSections.Num() == 0;
+}
+
+TRange<float> UMovieSceneSkeletalAnimationTrack::GetSectionBoundaries() const
+{
+	TArray< TRange<float> > Bounds;
+	for (int32 i = 0; i < AnimationSections.Num(); ++i)
+	{
+		Bounds.Add(AnimationSections[i]->GetRange());
+	}
+	return TRange<float>::Hull(Bounds);
+}
+
+void UMovieSceneSkeletalAnimationTrack::AddNewAnimation(float KeyTime, UAnimSequence* AnimSequence)
+{
+	// add the section
+	UMovieSceneSkeletalAnimationSection* NewSection = NewObject<UMovieSceneSkeletalAnimationSection>(this);
+	NewSection->InitialPlacement( AnimationSections, KeyTime, KeyTime + AnimSequence->SequenceLength, SupportsMultipleRows() );
+	NewSection->SetAnimationStartTime( KeyTime );
+	NewSection->SetAnimSequence(AnimSequence);
+
+	AnimationSections.Add(NewSection);
+}
+
+UMovieSceneSection* UMovieSceneSkeletalAnimationTrack::GetAnimSectionAtTime(float Time)
+{
+	for (int32 i = 0; i < AnimationSections.Num(); ++i)
+	{
+		UMovieSceneSection* Section = AnimationSections[i];
+		if( Section->IsTimeWithinSection( Time ) )
+		{
+			return Section;
+		}
+	}
+
+	return NULL;
+}
+
+
+#undef LOCTEXT_NAMESPACE
diff --git a/Engine/Source/Runtime/MovieSceneTracks/Private/Tracks/SubMovieSceneTrack.cpp b/Engine/Source/Runtime/MovieSceneTracks/Private/Tracks/SubMovieSceneTrack.cpp
--- a/Engine/Source/Runtime/MovieSceneTracks/Private/Tracks/SubMovieSceneTrack.cpp
+++ b/Engine/Source/Runtime/MovieSceneTracks/Private/Tracks/SubMovieSceneTrack.cpp
@@ -18,12 +18,12 @@ UMovieSceneSection* USubMovieSceneTrack::CreateNewSection()
 
 FName USubMovieSceneTrack::GetTrackName() const
 {
-	static const FName DefaultName( "MovieScene" );
+	static const FName DefaultName( "Sequence" );
 
 	FName Name = DefaultName;
 	if( SubMovieSceneSections.Num() > 0 )
 	{
-		Name = CastChecked<USubMovieSceneSection>(SubMovieSceneSections[0])->GetMovieScene()->GetFName();
+		Name = CastChecked<USubMovieSceneSection>(SubMovieSceneSections[0])->GetMovieSceneAnimation()->GetFName();
 	}
 
 	return Name;
@@ -73,15 +73,15 @@ bool USubMovieSceneTrack::HasSection( UMovieSceneSection* Section ) const
 	return SubMovieSceneSections.Contains( Section );
 }
 
-void USubMovieSceneTrack::AddMovieSceneSection( UMovieScene* SubMovieScene, float Time )
+void USubMovieSceneTrack::AddMovieSceneSection( UMovieSceneSequence* SubMovieSceneSequence, float Time )
 {
 	Modify();
 
 	USubMovieSceneSection* NewSection = CastChecked<USubMovieSceneSection>( CreateNewSection() );
 
-	NewSection->SetMovieScene( SubMovieScene );
+	NewSection->SetMovieSceneAnimation( SubMovieSceneSequence );
 	NewSection->SetStartTime( Time );
-	NewSection->SetEndTime( Time + SubMovieScene->GetTimeRange().Size<float>() );
+	NewSection->SetEndTime( Time + SubMovieSceneSequence->GetMovieScene()->GetTimeRange().Size<float>() );
 
 	SubMovieSceneSections.Add( NewSection );
 
diff --git a/Engine/Source/Runtime/MovieSceneTracks/Public/Sections/MovieSceneAnimationSection.h b/Engine/Source/Runtime/MovieSceneTracks/Public/Sections/MovieSceneAnimationSection.h
deleted file mode 100644
diff --git a/Engine/Source/Runtime/MovieSceneTracks/Public/Sections/MovieSceneSkeletalAnimationSection.h b/Engine/Source/Runtime/MovieSceneTracks/Public/Sections/MovieSceneSkeletalAnimationSection.h
new file mode 100644
--- /dev/null
+++ b/Engine/Source/Runtime/MovieSceneTracks/Public/Sections/MovieSceneSkeletalAnimationSection.h
@@ -0,0 +1,54 @@
+// Copyright 1998-2015 Epic Games, Inc. All Rights Reserved.
+
+#pragma once
+
+#include "MovieSceneSection.h"
+#include "MovieSceneSkeletalAnimationSection.generated.h"
+
+/**
+ * Audio section, for use in the master audio, or by attached audio objects
+ */
+UCLASS( MinimalAPI )
+class UMovieSceneSkeletalAnimationSection : public UMovieSceneSection
+{
+	GENERATED_UCLASS_BODY()
+public:
+	/** Sets the animation sequence for this section */
+	void SetAnimSequence(class UAnimSequence* InAnimSequence) {AnimSequence = InAnimSequence;}
+	
+	/** Gets the animation sequence for this section */
+	class UAnimSequence* GetAnimSequence() {return AnimSequence;}
+	
+	/** Sets the time that the animation is supposed to be played at */
+	void SetAnimationStartTime(float InAnimationStartTime) {AnimationStartTime = InAnimationStartTime;}
+	
+	/** Gets the (absolute) time that the animation is supposed to be played at */
+	float GetAnimationStartTime() const {return AnimationStartTime;}
+	
+	/** Gets the animation duration, modified by dilation */
+	float GetAnimationDuration() const {return AnimSequence->SequenceLength * AnimationDilationFactor;}
+
+	/**
+	 * @return The dilation multiplier of the animation
+	 */
+	float GetAnimationDilationFactor() const {return AnimationDilationFactor;}
+
+	/** MovieSceneSection interface */
+	virtual void MoveSection( float DeltaPosition, TSet<FKeyHandle>& KeyHandles ) override;
+	virtual void DilateSection( float DilationFactor, float Origin, TSet<FKeyHandle>& KeyHandles  ) override;
+	virtual void GetKeyHandles(TSet<FKeyHandle>& KeyHandles) const override;
+	virtual void GetSnapTimes(TArray<float>& OutSnapTimes, bool bGetSectionBorders) const override;
+
+private:
+	/** The animation sequence this section has */
+	UPROPERTY()
+	class UAnimSequence* AnimSequence;
+
+	/** The absolute time that the animation starts playing at */
+	UPROPERTY()
+	float AnimationStartTime;
+	
+	/** The amount which this animation is time dilated by */
+	UPROPERTY()
+	float AnimationDilationFactor;
+};
diff --git a/Engine/Source/Runtime/MovieSceneTracks/Public/Sections/SubMovieSceneSection.h b/Engine/Source/Runtime/MovieSceneTracks/Public/Sections/SubMovieSceneSection.h
--- a/Engine/Source/Runtime/MovieSceneTracks/Public/Sections/SubMovieSceneSection.h
+++ b/Engine/Source/Runtime/MovieSceneTracks/Public/Sections/SubMovieSceneSection.h
@@ -4,11 +4,13 @@
 
 #include "SubMovieSceneSection.generated.h"
 
+class UMovieSceneSequence;
+
 /**
  * A container for a movie scene within a movie scene
  */
-UCLASS( MinimalAPI )
-class USubMovieSceneSection : public UMovieSceneSection
+UCLASS()
+class MOVIESCENETRACKS_API USubMovieSceneSection : public UMovieSceneSection
 {
 	GENERATED_UCLASS_BODY()
 public:
@@ -17,13 +19,13 @@ public:
 	 *
 	 * @param InMovieScene	The movie scene to play
 	 */
-	void SetMovieScene( UMovieScene* InMovieScene ) { MovieScene = InMovieScene; }
+	void SetMovieSceneAnimation( UMovieSceneSequence* InAnimation );
 
 	/** @return The movie scene played by this section */
-	UMovieScene* GetMovieScene() const { return MovieScene; }
+	UMovieSceneSequence* GetMovieSceneAnimation() const;
 private:
 	/** Movie scene being played by this section */
 	/** @todo Sequencer: Should this be lazy loaded? */
 	UPROPERTY()
-	UMovieScene* MovieScene;
+	UMovieSceneSequence* Sequence;
 };
diff --git a/Engine/Source/Runtime/MovieSceneTracks/Public/TrackInstances/MovieSceneAnimationTrackInstance.h b/Engine/Source/Runtime/MovieSceneTracks/Public/TrackInstances/MovieSceneAnimationTrackInstance.h
deleted file mode 100644
diff --git a/Engine/Source/Runtime/MovieSceneTracks/Public/TrackInstances/MovieSceneColorTrackInstance.h b/Engine/Source/Runtime/MovieSceneTracks/Public/TrackInstances/MovieSceneColorTrackInstance.h
--- a/Engine/Source/Runtime/MovieSceneTracks/Public/TrackInstances/MovieSceneColorTrackInstance.h
+++ b/Engine/Source/Runtime/MovieSceneTracks/Public/TrackInstances/MovieSceneColorTrackInstance.h
@@ -43,6 +43,7 @@ private:
 	TSharedPtr<FTrackInstancePropertyBindings> PropertyBindings;
 
 	/** Map from object to initial state */
+	/** @todo Sequencer: This may be editor only data */
 	TMap< TWeakObjectPtr<UObject>, FSlateColor > InitSlateColorMap;
 	TMap< TWeakObjectPtr<UObject>, FLinearColor > InitLinearColorMap;
 
diff --git a/Engine/Source/Runtime/MovieSceneTracks/Public/TrackInstances/MovieSceneSkeletalAnimationTrackInstance.h b/Engine/Source/Runtime/MovieSceneTracks/Public/TrackInstances/MovieSceneSkeletalAnimationTrackInstance.h
new file mode 100644
--- /dev/null
+++ b/Engine/Source/Runtime/MovieSceneTracks/Public/TrackInstances/MovieSceneSkeletalAnimationTrackInstance.h
@@ -0,0 +1,33 @@
+// Copyright 1998-2015 Epic Games, Inc. All Rights Reserved.
+
+#pragma once
+
+#include "IMovieSceneTrackInstance.h"
+
+
+class UMovieSceneSkeletalAnimationTrack;
+
+
+/**
+ * Instance of a UMovieSceneSkeletalAnimationTrack
+ */
+class FMovieSceneSkeletalAnimationTrackInstance
+	: public IMovieSceneTrackInstance
+{
+public:
+
+	FMovieSceneSkeletalAnimationTrackInstance( UMovieSceneSkeletalAnimationTrack& InAnimationTrack );
+	virtual ~FMovieSceneSkeletalAnimationTrackInstance();
+
+	/** IMovieSceneTrackInstance interface */
+	virtual void SaveState(const TArray<UObject*>& RuntimeObjects) override {}
+	virtual void RestoreState(const TArray<UObject*>& RuntimeObjects) override {}
+	virtual void Update( float Position, float LastPosition, const TArray<UObject*>& RuntimeObjects, IMovieScenePlayer& Player ) override;
+	virtual void RefreshInstance( const TArray<UObject*>& RuntimeObjects, IMovieScenePlayer& Player ) override {}
+	virtual void ClearInstance( IMovieScenePlayer& Player ) override {}
+
+private:
+
+	/** Track that is being instanced */
+	UMovieSceneSkeletalAnimationTrack* AnimationTrack;
+};
diff --git a/Engine/Source/Runtime/MovieSceneTracks/Public/TrackInstances/SubMovieSceneTrackInstance.h b/Engine/Source/Runtime/MovieSceneTracks/Public/TrackInstances/SubMovieSceneTrackInstance.h
--- a/Engine/Source/Runtime/MovieSceneTracks/Public/TrackInstances/SubMovieSceneTrackInstance.h
+++ b/Engine/Source/Runtime/MovieSceneTracks/Public/TrackInstances/SubMovieSceneTrackInstance.h
@@ -5,7 +5,7 @@
 #include "IMovieSceneTrackInstance.h"
 
 
-class FMovieSceneInstance;
+class FMovieSceneSequenceInstance;
 class USubMovieSceneSection;
 class USubMovieSceneTrack;
 
@@ -32,5 +32,5 @@ protected:
 	TWeakObjectPtr<USubMovieSceneTrack> SubMovieSceneTrack;
 
 	/** Mapping of section lookups to instances.  Each section has a movie scene which must be instanced */
-	TMap< TWeakObjectPtr<USubMovieSceneSection>, TSharedPtr<FMovieSceneInstance> > SubMovieSceneInstances; 
+	TMap< TWeakObjectPtr<USubMovieSceneSection>, TSharedPtr<FMovieSceneSequenceInstance> > SubMovieSceneInstances; 
 };
diff --git a/Engine/Source/Runtime/MovieSceneTracks/Public/Tracks/MovieSceneAnimationTrack.h b/Engine/Source/Runtime/MovieSceneTracks/Public/Tracks/MovieSceneAnimationTrack.h
deleted file mode 100644
diff --git a/Engine/Source/Runtime/MovieSceneTracks/Public/Tracks/MovieSceneShotTrack.h b/Engine/Source/Runtime/MovieSceneTracks/Public/Tracks/MovieSceneShotTrack.h
--- a/Engine/Source/Runtime/MovieSceneTracks/Public/Tracks/MovieSceneShotTrack.h
+++ b/Engine/Source/Runtime/MovieSceneTracks/Public/Tracks/MovieSceneShotTrack.h
@@ -30,7 +30,7 @@ public:
 	 * @param ShotName		The display name of the shot
 	 * @param ShotNumber		The number of the shot.  This is used to assist with auto-generated shot names when new shots are added
 	 */
-	MOVIESCENETRACKS_API void AddNewShot(FGuid CameraHandle, UMovieScene& ShotMovieScene, const TRange<float>& TimeRange, const FText& ShotName, int32 ShotNumber );
+	MOVIESCENETRACKS_API void AddNewShot(FGuid CameraHandle, UMovieSceneSequence& ShotMovieSceneAnimation, const TRange<float>& TimeRange, const FText& ShotName, int32 ShotNumber );
 
 #if WITH_EDITOR
 	virtual void OnSectionMoved( UMovieSceneSection& Section ) override;
diff --git a/Engine/Source/Runtime/MovieSceneTracks/Public/Tracks/MovieSceneSkeletalAnimationTrack.h b/Engine/Source/Runtime/MovieSceneTracks/Public/Tracks/MovieSceneSkeletalAnimationTrack.h
new file mode 100644
--- /dev/null
+++ b/Engine/Source/Runtime/MovieSceneTracks/Public/Tracks/MovieSceneSkeletalAnimationTrack.h
@@ -0,0 +1,40 @@
+// Copyright 1998-2015 Epic Games, Inc. All Rights Reserved.
+
+#pragma once
+
+#include "MovieScene.h"
+#include "MovieSceneTrack.h"
+#include "MovieSceneSkeletalAnimationTrack.generated.h"
+
+
+/**
+ * Handles animation of skeletal mesh actors
+ */
+UCLASS( MinimalAPI )
+class UMovieSceneSkeletalAnimationTrack : public UMovieSceneTrack
+{
+	GENERATED_UCLASS_BODY()
+
+public:
+	/** UMovieSceneTrack interface */
+	virtual FName GetTrackName() const override;
+	virtual TSharedPtr<IMovieSceneTrackInstance> CreateInstance() override;
+	virtual void RemoveAllAnimationData() override;
+	virtual bool HasSection( UMovieSceneSection* Section ) const override;
+	virtual void RemoveSection( UMovieSceneSection* Section ) override;
+	virtual bool IsEmpty() const override;
+	virtual TRange<float> GetSectionBoundaries() const override;
+	virtual bool SupportsMultipleRows() const override { return true; }
+	virtual const TArray<UMovieSceneSection*>& GetAllSections() const override;
+
+	/** Adds a new animation to this track */
+	virtual void AddNewAnimation(float KeyTime, class UAnimSequence* AnimSequence);
+
+	/** Gets the animation section at a certain time, or NULL if there is none */
+	class UMovieSceneSection* GetAnimSectionAtTime(float Time);
+
+private:
+	/** List of all animation sections */
+	UPROPERTY()
+	TArray<UMovieSceneSection*> AnimationSections;
+};
diff --git a/Engine/Source/Runtime/MovieSceneTracks/Public/Tracks/SubMovieSceneTrack.h b/Engine/Source/Runtime/MovieSceneTracks/Public/Tracks/SubMovieSceneTrack.h
--- a/Engine/Source/Runtime/MovieSceneTracks/Public/Tracks/SubMovieSceneTrack.h
+++ b/Engine/Source/Runtime/MovieSceneTracks/Public/Tracks/SubMovieSceneTrack.h
@@ -6,7 +6,7 @@
 #include "SubMovieSceneTrack.generated.h"
 
 /**
- * Handles manipulation of float properties in a movie scene
+ * A track that holds a sub-sequence within a larger sequence
  */
 UCLASS( MinimalAPI )
 class USubMovieSceneTrack : public UMovieSceneTrack
@@ -30,7 +30,7 @@ public:
 	 * @param SubMovieScene	The movie scene to add
 	 * @param Time	The time to add the section at
 	 */
-	virtual void AddMovieSceneSection( UMovieScene* SubMovieScene, float Time );
+	virtual void AddMovieSceneSection( class UMovieSceneSequence* SubMovieSceneSequence, float Time );
 protected:
 	/** All movie scene sections.  */
 	UPROPERTY()
diff --git a/Engine/Source/Runtime/Niagara/Private/NiagaraAnimation.cpp b/Engine/Source/Runtime/Niagara/Private/NiagaraAnimation.cpp
--- a/Engine/Source/Runtime/Niagara/Private/NiagaraAnimation.cpp
+++ b/Engine/Source/Runtime/Niagara/Private/NiagaraAnimation.cpp
@@ -33,7 +33,7 @@ bool UNiagaraAnimation::CanPossessObject(UObject& Object) const
 }
 
 
-void UNiagaraAnimation::DestroyAllSpawnedObjects(UMovieScene& SubMovieScene)
+void UNiagaraAnimation::DestroyAllSpawnedObjects()
 {
 }
 
@@ -62,7 +62,7 @@ UObject* UNiagaraAnimation::GetParentObject(UObject* Object) const
 }
 
 
-void UNiagaraAnimation::SpawnOrDestroyObjects(UMovieScene* SubMovieScene, bool DestroyAll)
+void UNiagaraAnimation::SpawnOrDestroyObjects(bool DestroyAll)
 {
 }
 
diff --git a/Engine/Source/Runtime/Niagara/Public/NiagaraAnimation.h b/Engine/Source/Runtime/Niagara/Public/NiagaraAnimation.h
--- a/Engine/Source/Runtime/Niagara/Public/NiagaraAnimation.h
+++ b/Engine/Source/Runtime/Niagara/Public/NiagaraAnimation.h
@@ -2,7 +2,7 @@
 
 #pragma once
 
-#include "MovieSceneAnimation.h"
+#include "MovieSceneSequence.h"
 #include "NiagaraAnimation.generated.h"
 
 
@@ -11,7 +11,7 @@
  */
 UCLASS(BlueprintType, MinimalAPI)
 class UNiagaraAnimation
-	: public UMovieSceneAnimation
+	: public UMovieSceneSequence
 {
 	GENERATED_UCLASS_BODY()
 
@@ -22,12 +22,12 @@ public:
 	virtual bool AllowsSpawnableObjects() const override;
 	virtual void BindPossessableObject(const FGuid& ObjectId, UObject& PossessedObject) override;
 	virtual bool CanPossessObject(UObject& Object) const override;
-	virtual void DestroyAllSpawnedObjects(UMovieScene& SubMovieScene) override;
+	virtual void DestroyAllSpawnedObjects() override;
 	virtual UObject* FindObject(const FGuid& ObjectId) const override;
 	virtual FGuid FindObjectId(UObject& Object) const override;
 	virtual UMovieScene* GetMovieScene() const override;
 	virtual UObject* GetParentObject(UObject* Object) const override;
-	virtual void SpawnOrDestroyObjects(UMovieScene* SubMovieScene, bool DestroyAll) override;
+	virtual void SpawnOrDestroyObjects(bool DestroyAll) override;
 	virtual void UnbindPossessableObjects(const FGuid& ObjectId) override;
 
 #if WITH_EDITOR
diff --git a/Engine/Source/Runtime/UMG/Private/Animation/UMGSequencePlayer.cpp b/Engine/Source/Runtime/UMG/Private/Animation/UMGSequencePlayer.cpp
--- a/Engine/Source/Runtime/UMG/Private/Animation/UMGSequencePlayer.cpp
+++ b/Engine/Source/Runtime/UMG/Private/Animation/UMGSequencePlayer.cpp
@@ -3,7 +3,7 @@
 #include "UMGPrivatePCH.h"
 #include "UMGSequencePlayer.h"
 #include "MovieScene.h"
-#include "MovieSceneInstance.h"
+#include "MovieSceneSequenceInstance.h"
 #include "MovieScene.h"
 #include "WidgetAnimation.h"
 
@@ -20,7 +20,7 @@ void UUMGSequencePlayer::InitSequencePlayer( const UWidgetAnimation& InAnimation
 {
 	Animation = &InAnimation;
 
-	UMovieScene* MovieScene = Animation->MovieScene;
+	UMovieScene* MovieScene = Animation->GetMovieScene();
 
 	// Cache the time range of the sequence to determine when we stop
 	TimeRange = MovieScene->GetTimeRange();
@@ -28,7 +28,7 @@ void UUMGSequencePlayer::InitSequencePlayer( const UWidgetAnimation& InAnimation
 	UWidgetTree* WidgetTree = UserWidget.WidgetTree;
 
 	// Bind to Runtime Objects
-	for (const FWidgetAnimationBinding& Binding : InAnimation.AnimationBindings)
+	for (const FWidgetAnimationBinding& Binding : InAnimation.GetBindings())
 	{
 		UObject* FoundObject = Binding.FindRuntimeObject( *WidgetTree );
 
@@ -103,8 +103,7 @@ void UUMGSequencePlayer::Tick(float DeltaTime)
 
 void UUMGSequencePlayer::Play(float StartAtTime, int32 InNumLoopsToPlay, EUMGSequencePlayMode::Type InPlayMode)
 {
-	UMovieScene* MovieScene = Animation->MovieScene;
-	RootMovieSceneInstance = MakeShareable( new FMovieSceneInstance( *MovieScene ) );
+	RootMovieSceneInstance = MakeShareable( new FMovieSceneSequenceInstance( *Animation ) );
 	RootMovieSceneInstance->RefreshInstance( *this );
 
 	PlayMode = InPlayMode;
@@ -148,7 +147,7 @@ void UUMGSequencePlayer::Stop()
 	TimeCursorPosition = 0;
 }
 
-void UUMGSequencePlayer::GetRuntimeObjects( TSharedRef<FMovieSceneInstance> MovieSceneInstance, const FGuid& ObjectHandle, TArray< UObject* >& OutObjects ) const
+void UUMGSequencePlayer::GetRuntimeObjects( TSharedRef<FMovieSceneSequenceInstance> MovieSceneInstance, const FGuid& ObjectHandle, TArray< UObject* >& OutObjects ) const
 {
 	const TArray<UObject*>* FoundObjects = GuidToRuntimeObjectMap.Find( ObjectHandle );
 
diff --git a/Engine/Source/Runtime/UMG/Private/WidgetBlueprintGeneratedClass.cpp b/Engine/Source/Runtime/UMG/Private/WidgetBlueprintGeneratedClass.cpp
--- a/Engine/Source/Runtime/UMG/Private/WidgetBlueprintGeneratedClass.cpp
+++ b/Engine/Source/Runtime/UMG/Private/WidgetBlueprintGeneratedClass.cpp
@@ -39,10 +39,10 @@ void UWidgetBlueprintGeneratedClass::InitializeWidget(UUserWidget* UserWidget) c
 		{
 			UWidgetAnimation* Anim = DuplicateObject<UWidgetAnimation>( Animation, UserWidget );
 
-			if( Anim->MovieScene )
+			if( Anim->GetMovieScene() )
 			{
 				// Find property with the same name as the template and assign the new widget to it.
-				UObjectPropertyBase* Prop = FindField<UObjectPropertyBase>(WidgetBlueprintClass, Anim->MovieScene->GetFName());
+				UObjectPropertyBase* Prop = FindField<UObjectPropertyBase>(WidgetBlueprintClass, Anim->GetMovieScene()->GetFName());
 				if(Prop)
 				{
 					Prop->SetObjectPropertyValue_InContainer(UserWidget, Anim);
diff --git a/Engine/Source/Runtime/UMG/Public/Animation/UMGSequencePlayer.h b/Engine/Source/Runtime/UMG/Public/Animation/UMGSequencePlayer.h
--- a/Engine/Source/Runtime/UMG/Public/Animation/UMGSequencePlayer.h
+++ b/Engine/Source/Runtime/UMG/Public/Animation/UMGSequencePlayer.h
@@ -35,11 +35,11 @@ public:
 	const UWidgetAnimation* GetAnimation() const { return Animation; }
 
 	/** IMovieScenePlayer interface */
-	virtual void GetRuntimeObjects( TSharedRef<FMovieSceneInstance> MovieSceneInstance, const FGuid& ObjectHandle, TArray< UObject* >& OutObjects ) const override;
+	virtual void GetRuntimeObjects( TSharedRef<FMovieSceneSequenceInstance> MovieSceneInstance, const FGuid& ObjectHandle, TArray< UObject* >& OutObjects ) const override;
 	virtual void UpdateCameraCut(UObject* ObjectToViewThrough, bool bNewCameraCut) const override {}
-	virtual void AddOrUpdateMovieSceneInstance( class UMovieSceneSection& MovieSceneSection, TSharedRef<FMovieSceneInstance> InstanceToAdd ) override {}
-	virtual void RemoveMovieSceneInstance( class UMovieSceneSection& MovieSceneSection, TSharedRef<FMovieSceneInstance> InstanceToRemove ) override {}
-	virtual TSharedRef<FMovieSceneInstance> GetRootMovieSceneInstance() const override { return RootMovieSceneInstance.ToSharedRef(); }
+	virtual void AddOrUpdateMovieSceneInstance( class UMovieSceneSection& MovieSceneSection, TSharedRef<FMovieSceneSequenceInstance> InstanceToAdd ) override {}
+	virtual void RemoveMovieSceneInstance( class UMovieSceneSection& MovieSceneSection, TSharedRef<FMovieSceneSequenceInstance> InstanceToRemove ) override {}
+	virtual TSharedRef<FMovieSceneSequenceInstance> GetRootMovieSceneSequenceInstance() const override { return RootMovieSceneInstance.ToSharedRef(); }
 	virtual EMovieScenePlayerStatus::Type GetPlaybackStatus() const override;
 
 	DECLARE_EVENT_OneParam(UUMGSequencePlayer, FOnSequenceFinishedPlaying, UUMGSequencePlayer&);
@@ -53,7 +53,7 @@ private:
 	TMap<FGuid, TArray<UObject*> > GuidToRuntimeObjectMap;
 
 	/** The root movie scene instance to update when playing. */
-	TSharedPtr<class FMovieSceneInstance> RootMovieSceneInstance;
+	TSharedPtr<class FMovieSceneSequenceInstance> RootMovieSceneInstance;
 
 	/** Time range of the animation */
 	TRange<float> TimeRange;
diff --git a/Engine/Source/Runtime/UMG/Public/Animation/WidgetAnimation.h b/Engine/Source/Runtime/UMG/Public/Animation/WidgetAnimation.h
--- a/Engine/Source/Runtime/UMG/Public/Animation/WidgetAnimation.h
+++ b/Engine/Source/Runtime/UMG/Public/Animation/WidgetAnimation.h
@@ -2,7 +2,7 @@
 
 #pragma once
 
-#include "MovieSceneAnimation.h"
+#include "MovieSceneSequence.h"
 #include "WidgetAnimationBinding.h"
 #include "WidgetAnimation.generated.h"
 
@@ -13,7 +13,7 @@ class UUserWidget;
 
 UCLASS(BlueprintType, MinimalAPI)
 class UWidgetAnimation
-	: public UMovieSceneAnimation
+	: public UMovieSceneSequence
 {
 	GENERATED_UCLASS_BODY()
 
@@ -47,7 +47,7 @@ public:
 	UMG_API float GetEndTime() const;
 
 	/**
-	 * Initialize the widge editor with a new user widget.
+	 * Initialize the animation with a new user widget.
 	 *
 	 * @param InPreviewWidget The user widget to preview.
 	 */
@@ -60,18 +60,19 @@ public:
 	virtual bool AllowsSpawnableObjects() const override;
 	virtual void BindPossessableObject(const FGuid& ObjectId, UObject& PossessedObject) override;
 	virtual bool CanPossessObject(UObject& Object) const override;
-	virtual void DestroyAllSpawnedObjects(UMovieScene& SubMovieScene) override { }
+	virtual void DestroyAllSpawnedObjects() override { }
 	virtual UObject* FindObject(const FGuid& ObjectId) const override;
 	virtual FGuid FindObjectId(UObject& Object) const override;
 	virtual UMovieScene* GetMovieScene() const override;
 	virtual UObject* GetParentObject(UObject* Object) const override;
-	virtual void SpawnOrDestroyObjects(UMovieScene* SubMovieScene, bool DestroyAll) override { }
+	virtual void SpawnOrDestroyObjects(bool DestroyAll) override { }
 	virtual void UnbindPossessableObjects(const FGuid& ObjectId) override;
 
 #if WITH_EDITOR
 	virtual bool TryGetObjectDisplayName(const FGuid& ObjectId, FText& OutDisplayName) const override;
 #endif
 
+	const TArray<FWidgetAnimationBinding>& GetBindings() const { return AnimationBindings; }
 public:
 
 	/** Pointer to the movie scene that controls this animation. */
