Sequencer: Fix rounding error.

diff --git a/Engine/Source/Editor/Sequencer/Private/SequencerCommonHelpers.cpp b/Engine/Source/Editor/Sequencer/Private/SequencerCommonHelpers.cpp
--- a/Engine/Source/Editor/Sequencer/Private/SequencerCommonHelpers.cpp
+++ b/Engine/Source/Editor/Sequencer/Private/SequencerCommonHelpers.cpp
@@ -55,7 +55,7 @@ void SequencerHelpers::GetDescendantNodes(TSharedRef<FSequencerDisplayNode> Disp
 
 int32 SequencerHelpers::TimeToFrame(float Time, float FrameRate)
 {
-	int32 Frame = Time * FrameRate;
+	float Frame = Time * FrameRate;
 	return FMath::RoundToInt(Frame);
 }
 
