Sequencer: Fix in/out time initialization when snap setting is not a frame rate. Also, fixed clamp ranges not returning the value from the movie scene. And removed commented out code.

diff --git a/Engine/Source/Editor/Sequencer/Private/Sequencer.cpp b/Engine/Source/Editor/Sequencer/Private/Sequencer.cpp
--- a/Engine/Source/Editor/Sequencer/Private/Sequencer.cpp
+++ b/Engine/Source/Editor/Sequencer/Private/Sequencer.cpp
@@ -526,7 +526,7 @@ FAnimatedRange FSequencer::GetViewRange() const
 
 FAnimatedRange FSequencer::GetClampRange() const
 {
-	return TRange<float>(GetFocusedMovieSceneSequence()->GetMovieScene()->GetTimeRange() );
+	return TRange<float>(GetFocusedMovieSceneSequence()->GetMovieScene()->StartTime, GetFocusedMovieSceneSequence()->GetMovieScene()->EndTime);
 }
 
 bool FSequencer::GetAutoKeyEnabled() const 
@@ -923,25 +923,53 @@ void FSequencer::UpdateTimeBoundsToFocusedMovieScene()
 	float InTime = FocusedMovieScene->InTime;
 	if (InTime >= FLT_MAX)
 	{
-		InTime = SequencerHelpers::FrameToTime(InFrame, 1.0f/Settings->GetTimeSnapInterval());
+		if (SequencerSnapValues::IsTimeSnapIntervalFrameRate(Settings->GetTimeSnapInterval()))
+		{
+			InTime = SequencerHelpers::FrameToTime(InFrame, 1.0f/Settings->GetTimeSnapInterval());
+		}
+		else
+		{
+			InTime = 0.0f;
+		}
 	}
 
 	float OutTime = FocusedMovieScene->OutTime;
 	if (OutTime <= -FLT_MAX)
 	{
-		OutTime = SequencerHelpers::FrameToTime(OutFrame, 1.0f/Settings->GetTimeSnapInterval());
+		if (SequencerSnapValues::IsTimeSnapIntervalFrameRate(Settings->GetTimeSnapInterval()))
+		{
+			OutTime = SequencerHelpers::FrameToTime(OutFrame, 1.0f/Settings->GetTimeSnapInterval());
+		}
+		else
+		{
+			OutTime = 5.0f;
+		}
 	}
 
 	float StartTime = FocusedMovieScene->StartTime;
 	if (StartTime >= FLT_MAX)
 	{
-		StartTime = SequencerHelpers::FrameToTime(InFrame, 1.0f/Settings->GetTimeSnapInterval());
+		if (SequencerSnapValues::IsTimeSnapIntervalFrameRate(Settings->GetTimeSnapInterval()))
+		{
+			StartTime = SequencerHelpers::FrameToTime(InFrame, 1.0f/Settings->GetTimeSnapInterval());
+		}
+		else
+		{
+			StartTime = 0.0f;
+		}
 	}
 
 	float EndTime = FocusedMovieScene->EndTime;
 	if (EndTime <= -FLT_MAX)
 	{
-		EndTime = SequencerHelpers::FrameToTime(OutFrame, 1.0f/Settings->GetTimeSnapInterval());
+		if (SequencerSnapValues::IsTimeSnapIntervalFrameRate(Settings->GetTimeSnapInterval()))
+		{
+			EndTime = SequencerHelpers::FrameToTime(OutFrame, 1.0f/Settings->GetTimeSnapInterval());
+		}
+		else
+		{
+			EndTime = 5.0f;
+		}
 	}
 
 	// Set the clamp range to the movie scene's clamp range.
diff --git a/Engine/Source/Editor/Sequencer/Private/TimeSliderController.h b/Engine/Source/Editor/Sequencer/Private/TimeSliderController.h
--- a/Engine/Source/Editor/Sequencer/Private/TimeSliderController.h
+++ b/Engine/Source/Editor/Sequencer/Private/TimeSliderController.h
@@ -25,7 +25,6 @@ public:
 	virtual FAnimatedRange GetViewRange() const override { return TimeSliderArgs.ViewRange.Get(); }
 
 	/** Get the current clamp range for this controller */
-	//virtual FAnimatedRange GetClampRange() const override { return TRange<float>(TimeSliderArgs.ClampMin.Get().GetValue(), TimeSliderArgs.ClampMax.Get().GetValue()); }
 	virtual FAnimatedRange GetClampRange() const override { return TimeSliderArgs.ClampRange.Get(); }
 
 	/** Convert time to frame */
diff --git a/Engine/Source/Editor/SequencerWidgets/Private/STimeRange.cpp b/Engine/Source/Editor/SequencerWidgets/Private/STimeRange.cpp
--- a/Engine/Source/Editor/SequencerWidgets/Private/STimeRange.cpp
+++ b/Engine/Source/Editor/SequencerWidgets/Private/STimeRange.cpp
@@ -188,11 +188,6 @@ void STimeRange::OnStartTimeCommitted(const FText& NewText, ETextCommit::Type In
 			NewStart = FCString::Atof(*NewTextStr);
 		}
 
-		if (NewStart > TimeSliderController.Get()->GetViewRange().GetLowerBoundValue())
-		{
-			TimeSliderController.Get()->SetViewRange(NewStart, TimeSliderController.Get()->GetViewRange().GetUpperBoundValue(), EViewRangeInterpolation::Immediate);
-		}
-
 		TimeSliderController.Get()->SetClampRange(NewStart, TimeSliderController.Get()->GetClampRange().GetUpperBoundValue());
 	}
 }
@@ -216,11 +211,6 @@ void STimeRange::OnEndTimeCommitted(const FText& NewText, ETextCommit::Type InTe
 			NewEnd = FCString::Atof(*NewTextStr);
 		}
 
-		if (NewEnd < TimeSliderController.Get()->GetViewRange().GetUpperBoundValue())
-		{
-			TimeSliderController.Get()->SetViewRange(TimeSliderController.Get()->GetViewRange().GetLowerBoundValue(), NewEnd, EViewRangeInterpolation::Immediate);
-		}
-
 		TimeSliderController.Get()->SetClampRange(TimeSliderController.Get()->GetClampRange().GetLowerBoundValue(), NewEnd);
 	}
 }
diff --git a/Engine/Source/Editor/SequencerWidgets/Private/STimeRangeSlider.cpp b/Engine/Source/Editor/SequencerWidgets/Private/STimeRangeSlider.cpp
--- a/Engine/Source/Editor/SequencerWidgets/Private/STimeRangeSlider.cpp
+++ b/Engine/Source/Editor/SequencerWidgets/Private/STimeRangeSlider.cpp
@@ -15,7 +15,7 @@ namespace TimeRangeSliderConstants
 void STimeRangeSlider::Construct( const FArguments& InArgs, TSharedRef<ITimeSliderController> InTimeSliderController)
 {
 	TimeSliderController = InTimeSliderController;
-	LastViewRange.Empty();
+	LastViewRange = TimeSliderController.Get()->GetViewRange();
 
 	ResetState();
 	ResetHoveredState();
