Sequencer: Fix undo of track deletion and duplicate binding ocurring after trying to re-key an object. The actor animation wasn't being marked for Modify().

diff --git a/Engine/Source/Editor/Sequencer/Private/Sequencer.cpp b/Engine/Source/Editor/Sequencer/Private/Sequencer.cpp
--- a/Engine/Source/Editor/Sequencer/Private/Sequencer.cpp
+++ b/Engine/Source/Editor/Sequencer/Private/Sequencer.cpp
@@ -1316,7 +1316,8 @@ void FSequencer::OnRequestNodeDeleted( TSharedRef<const FSequencerDisplayNode>&
 			// The guid should be associated with a possessable if it wasnt a spawnable
 			bRemoved = OwnerMovieScene->RemovePossessable( BindingToRemove );
 			
-			// @todo sequencer: undo needs to work here
+			Sequence->Modify();
+
 			Sequence->UnbindPossessableObjects( BindingToRemove );
 
 			// If this check fails the guid was not associated with a spawnable or possessable so there was an invalid guid being stored on a node
