Sequencer: Fix for expanding clamp range when zooming to a section.

diff --git a/Engine/Source/Developer/LogVisualizer/Private/TimeSliderController.cpp b/Engine/Source/Developer/LogVisualizer/Private/TimeSliderController.cpp
--- a/Engine/Source/Developer/LogVisualizer/Private/TimeSliderController.cpp
+++ b/Engine/Source/Developer/LogVisualizer/Private/TimeSliderController.cpp
@@ -422,7 +422,7 @@ FReply FVisualLoggerTimeSliderController::OnMouseMove( TSharedRef<SWidget> Widge
 					NewViewOutputMax = LocalClampMax;
 				}
 
-				TimeSliderArgs.OnViewRangeChanged.ExecuteIfBound(TRange<float>(NewViewOutputMin, NewViewOutputMax), EViewRangeInterpolation::Immediate);
+				TimeSliderArgs.OnViewRangeChanged.ExecuteIfBound(TRange<float>(NewViewOutputMin, NewViewOutputMax), EViewRangeInterpolation::Immediate, false);
 				if (Scrollbar.IsValid())
 				{
 					float InOffsetFraction = (NewViewOutputMin - LocalClampMin) / (LocalClampMax - LocalClampMin);
@@ -600,7 +600,7 @@ FReply FVisualLoggerTimeSliderController::OnMouseWheel( TSharedRef<SWidget> Widg
 					NewViewOutputMax = LocalClampMax;
 				}
 
-				TimeSliderArgs.OnViewRangeChanged.ExecuteIfBound(TRange<float>(NewViewOutputMin, NewViewOutputMax), EViewRangeInterpolation::Immediate);
+				TimeSliderArgs.OnViewRangeChanged.ExecuteIfBound(TRange<float>(NewViewOutputMin, NewViewOutputMax), EViewRangeInterpolation::Immediate, false);
 				if (Scrollbar.IsValid())
 				{
 					float InOffsetFraction = (NewViewOutputMin - LocalClampMin) / (LocalClampMax - LocalClampMin);
diff --git a/Engine/Source/Editor/Sequencer/Private/SSequencerCurveEditor.cpp b/Engine/Source/Editor/Sequencer/Private/SSequencerCurveEditor.cpp
--- a/Engine/Source/Editor/Sequencer/Private/SSequencerCurveEditor.cpp
+++ b/Engine/Source/Editor/Sequencer/Private/SSequencerCurveEditor.cpp
@@ -23,7 +23,7 @@ void SSequencerCurveEditor::Construct( const FArguments& InArgs, TSharedRef<FSeq
 		SCurveEditor::FArguments()
 		.ViewMinInput_Lambda( [=]{ return ViewRange.Get().GetLowerBoundValue(); } )
 		.ViewMaxInput_Lambda( [=]{ return ViewRange.Get().GetUpperBoundValue(); } )
-		.OnSetInputViewRange_Lambda( [=](float InLowerBound, float InUpperBound){ OnViewRangeChanged.ExecuteIfBound(TRange<float>(InLowerBound, InUpperBound), EViewRangeInterpolation::Immediate); } )
+		.OnSetInputViewRange_Lambda( [=](float InLowerBound, float InUpperBound){ OnViewRangeChanged.ExecuteIfBound(TRange<float>(InLowerBound, InUpperBound), EViewRangeInterpolation::Immediate, false); } )
 		.HideUI( false )
 		.ZoomToFitHorizontal( true )
 		.ShowCurveSelector( false )
diff --git a/Engine/Source/Editor/Sequencer/Private/Sequencer.cpp b/Engine/Source/Editor/Sequencer/Private/Sequencer.cpp
--- a/Engine/Source/Editor/Sequencer/Private/Sequencer.cpp
+++ b/Engine/Source/Editor/Sequencer/Private/Sequencer.cpp
@@ -1026,7 +1026,7 @@ TRange<float> FSequencer::GetFilteringShotsTimeBounds() const
 	return TRange<float>::Empty();
 }
 
-void FSequencer::OnViewRangeChanged( TRange<float> NewViewRange, EViewRangeInterpolation Interpolation )
+void FSequencer::OnViewRangeChanged( TRange<float> NewViewRange, EViewRangeInterpolation Interpolation, bool bExpandClampRange )
 {
 	const float AnimationLengthSeconds = Interpolation == EViewRangeInterpolation::Immediate ? 0.f : 0.1f;
 
@@ -1034,7 +1034,7 @@ void FSequencer::OnViewRangeChanged( TRange<float> NewViewRange, EViewRangeInter
 	{
 		UMovieScene* FocusedMovieScene = GetFocusedMovieSceneSequence()->GetMovieScene();
 		// If locked, clamp the new range to clamp range
-		if (Settings->GetLockInOutToStartEndRange())
+		if (Settings->GetLockInOutToStartEndRange() && !bExpandClampRange)
 		{
 			if ( NewViewRange.GetLowerBoundValue() < FocusedMovieScene->StartTime)
 			{
@@ -1530,19 +1530,21 @@ void FSequencer::ZoomToSelectedSections()
 	}
 
 	if (!BoundsHull.IsEmpty() && !BoundsHull.IsDegenerate())
-	{
+	{	
+		bool bExpandClampRange = true;
+
 		// Zoom back to last view range if already expanded
 		if (!ViewRangeBeforeZoom.IsEmpty() &&
 			FMath::IsNearlyEqual(BoundsHull.GetLowerBoundValue(), GetViewRange().GetLowerBoundValue(), KINDA_SMALL_NUMBER) &&
 			FMath::IsNearlyEqual(BoundsHull.GetUpperBoundValue(), GetViewRange().GetUpperBoundValue(), KINDA_SMALL_NUMBER))
 		{
-			OnViewRangeChanged(ViewRangeBeforeZoom);
+			OnViewRangeChanged(ViewRangeBeforeZoom, EViewRangeInterpolation::Animated, bExpandClampRange);
 		}
 		else
 		{
 			ViewRangeBeforeZoom = GetViewRange();
 
-			OnViewRangeChanged(BoundsHull);
+			OnViewRangeChanged(BoundsHull, EViewRangeInterpolation::Animated, bExpandClampRange);
 		}
 	}
 }
@@ -1600,7 +1602,8 @@ void FSequencer::FilterToShotSections(const TArray< TWeakObjectPtr<class UMovieS
 	if( bZoomToShotBounds )
 	{
 		// zoom in
-		OnViewRangeChanged(GetTimeBounds());
+		bool bExpandClampRange = true;
+		OnViewRangeChanged(GetTimeBounds(), EViewRangeInterpolation::Animated, bExpandClampRange);
 	}
 
 	SequencerWidget->UpdateBreadcrumbs(ActualShotSections);
diff --git a/Engine/Source/Editor/Sequencer/Private/Sequencer.h b/Engine/Source/Editor/Sequencer/Private/Sequencer.h
--- a/Engine/Source/Editor/Sequencer/Private/Sequencer.h
+++ b/Engine/Source/Editor/Sequencer/Private/Sequencer.h
@@ -294,7 +294,7 @@ protected:
 	 *
 	 * @param	NewViewRange The new view range
 	 */
-	void OnViewRangeChanged( TRange<float> NewViewRange, EViewRangeInterpolation Interpolation = EViewRangeInterpolation::Animated );
+	void OnViewRangeChanged( TRange<float> NewViewRange, EViewRangeInterpolation Interpolation = EViewRangeInterpolation::Animated, bool bExpandClampRange = false );
 
 	/**
 	 * Called when the clamp range is changed by the user
diff --git a/Engine/Source/Editor/Sequencer/Private/TimeSliderController.cpp b/Engine/Source/Editor/Sequencer/Private/TimeSliderController.cpp
--- a/Engine/Source/Editor/Sequencer/Private/TimeSliderController.cpp
+++ b/Engine/Source/Editor/Sequencer/Private/TimeSliderController.cpp
@@ -443,7 +443,7 @@ FReply FSequencerTimeSliderController::OnMouseButtonUp( TSharedRef<SWidget> Widg
 					NewValue = LastRange[1];
 				}
 
-				TimeSliderArgs.OnViewRangeChanged.ExecuteIfBound(TRange<float>(DownValue, NewValue), EViewRangeInterpolation::Immediate);
+				TimeSliderArgs.OnViewRangeChanged.ExecuteIfBound(TRange<float>(DownValue, NewValue), EViewRangeInterpolation::Immediate, false);
 						
 				if( !TimeSliderArgs.ViewRange.IsBound() )
 				{	
@@ -708,7 +708,7 @@ void FSequencerTimeSliderController::SetViewRange( float NewRangeMin, float NewR
 {
 	const TRange<float> NewRange(NewRangeMin, NewRangeMax);
 
-	TimeSliderArgs.OnViewRangeChanged.ExecuteIfBound( NewRange, Interpolation );
+	TimeSliderArgs.OnViewRangeChanged.ExecuteIfBound( NewRange, Interpolation, false );
 
 	if( !TimeSliderArgs.ViewRange.IsBound() )
 	{	
diff --git a/Engine/Source/Editor/SequencerWidgets/Public/ITimeSlider.h b/Engine/Source/Editor/SequencerWidgets/Public/ITimeSlider.h
--- a/Engine/Source/Editor/SequencerWidgets/Public/ITimeSlider.h
+++ b/Engine/Source/Editor/SequencerWidgets/Public/ITimeSlider.h
@@ -14,7 +14,7 @@ enum class EViewRangeInterpolation
 };
 
 DECLARE_DELEGATE_TwoParams( FOnScrubPositionChanged, float, bool )
-DECLARE_DELEGATE_TwoParams( FOnViewRangeChanged, TRange<float>, EViewRangeInterpolation )
+DECLARE_DELEGATE_ThreeParams( FOnViewRangeChanged, TRange<float>, EViewRangeInterpolation, bool )
 DECLARE_DELEGATE_OneParam( FOnClampRangeChanged, TRange<float> )
 
 /** Structure used to wrap up a range, and an optional animation target */
