Sequencer: Property change fixes.

- Enforce that property changes be broadcasted only for properties that can be keyed by the given object. This fixes the FieldOfView getting a track for the CameraActor and the CameraComponent.
- Remove 'b' from the bool property names in order to look for the setter function. (ie. SetHidden instead of SetbHidden)

diff --git a/Engine/Source/Editor/Sequencer/Private/SequencerObjectChangeListener.cpp b/Engine/Source/Editor/Sequencer/Private/SequencerObjectChangeListener.cpp
--- a/Engine/Source/Editor/Sequencer/Private/SequencerObjectChangeListener.cpp
+++ b/Engine/Source/Editor/Sequencer/Private/SequencerObjectChangeListener.cpp
@@ -31,6 +31,24 @@ void FSequencerObjectChangeListener::OnPropertyChanged(const TArray<UObject*>& C
 
 void FSequencerObjectChangeListener::BroadcastPropertyChanged( FKeyPropertyParams KeyPropertyParams, bool bRequireAutoKey ) const
 {
+	// Filter out objects that actually have the property path that will be keyable. 
+	// Otherwise, this might try to key objects that don't have the requested property.
+	// For example, a property changed for the FieldOfView property will be sent for 
+	// both the CameraActor and the CameraComponent.
+	TArray<UObject*> KeyableObjects;
+	for (auto ObjectToKey : KeyPropertyParams.ObjectsToKey)
+	{
+		if (CanKeyProperty(FCanKeyPropertyParams(ObjectToKey->GetClass(), KeyPropertyParams.PropertyPath)))
+		{
+			KeyableObjects.Add(ObjectToKey);
+		}
+	}
+
+	if (!KeyableObjects.Num())
+	{
+		return;
+	}
+
 	const UStructProperty* StructProperty = Cast<const UStructProperty>(KeyPropertyParams.PropertyPath.Last());
 	const UStructProperty* ParentStructProperty = nullptr;
 	if (KeyPropertyParams.PropertyPath.Num() > 1)
@@ -40,7 +58,7 @@ void FSequencerObjectChangeListener::BroadcastPropertyChanged( FKeyPropertyParam
 
 	FPropertyChangedParams Params;
 	Params.bRequireAutoKey = bRequireAutoKey;
-	Params.ObjectsThatChanged = KeyPropertyParams.ObjectsToKey;
+	Params.ObjectsThatChanged = KeyableObjects;
 	Params.StructPropertyNameToKey = NAME_None;
 
 	bool bFoundAndBroadcastedDelegate = false;
@@ -93,12 +111,21 @@ FOnPropagateObjectChanges& FSequencerObjectChangeListener::GetOnPropagateObjectC
 	return OnPropagateObjectChanges;
 }
 
-bool FSequencerObjectChangeListener::FindPropertySetter( const UClass& ObjectClass, const FName PropertyTypeName, const FString& PropertyVarName ) const
+bool FSequencerObjectChangeListener::FindPropertySetter( const UClass& ObjectClass, const FName PropertyTypeName, const FString& InPropertyVarName ) const
 {
 	bool bFound = ClassToPropertyChangedMap.Contains( PropertyTypeName );
 
 	if( bFound )
 	{
+		FString PropertyVarName = InPropertyVarName;
+
+		// If this is a bool property, strip off the 'b' so that the "Set" functions to be 
+		// found are, for example, "SetHidden" instead of "SetbHidden"
+		if (PropertyTypeName == "BoolProperty")
+		{
+			PropertyVarName.RemoveFromStart("b", ESearchCase::CaseSensitive);
+		}
+
 		static const FString Set(TEXT("Set"));
 
 		const FString FunctionString = Set + PropertyVarName;
