Sequencer: Expand the clamp range when panning through the curve editor.

diff --git a/Engine/Source/Editor/Sequencer/Private/SSequencerCurveEditor.cpp b/Engine/Source/Editor/Sequencer/Private/SSequencerCurveEditor.cpp
--- a/Engine/Source/Editor/Sequencer/Private/SSequencerCurveEditor.cpp
+++ b/Engine/Source/Editor/Sequencer/Private/SSequencerCurveEditor.cpp
@@ -23,7 +23,7 @@ void SSequencerCurveEditor::Construct( const FArguments& InArgs, TSharedRef<FSeq
 		SCurveEditor::FArguments()
 		.ViewMinInput_Lambda( [=]{ return ViewRange.Get().GetLowerBoundValue(); } )
 		.ViewMaxInput_Lambda( [=]{ return ViewRange.Get().GetUpperBoundValue(); } )
-		.OnSetInputViewRange_Lambda( [=](float InLowerBound, float InUpperBound){ OnViewRangeChanged.ExecuteIfBound(TRange<float>(InLowerBound, InUpperBound), EViewRangeInterpolation::Immediate, false); } )
+		.OnSetInputViewRange_Lambda( [=](float InLowerBound, float InUpperBound){ OnViewRangeChanged.ExecuteIfBound(TRange<float>(InLowerBound, InUpperBound), EViewRangeInterpolation::Immediate, true); } )
 		.HideUI( false )
 		.ZoomToFitHorizontal( true )
 		.ShowCurveSelector( false )
