Sequencer: Toolbar reorganization.

- Added curve editor icon instead of "Curve Editor" text.
- Added auto key icon instead of "Auto Key" text.
- Moved clean view and toggle auto scroll to menu.
- Added "Add Object" button to add selected objects to the animation.
- Added "Key All" button to toggle between keying all channels and not.

diff --git a/Engine/Source/Editor/EditorStyle/Private/SlateEditorStyle.cpp b/Engine/Source/Editor/EditorStyle/Private/SlateEditorStyle.cpp
--- a/Engine/Source/Editor/EditorStyle/Private/SlateEditorStyle.cpp
+++ b/Engine/Source/Editor/EditorStyle/Private/SlateEditorStyle.cpp
@@ -1584,6 +1584,8 @@ void FSlateEditorStyle::FStyle::SetupGeneralStyles()
 		Set( "Sequencer.Timeline.RangeHandle", new BOX_BRUSH( "Sequencer/GenericSectionBackground", FMargin(5.f/16.f) ) );
 		Set( "Sequencer.Timeline.NotifyAlignmentMarker", new IMAGE_BRUSH( "Sequencer/NotifyAlignmentMarker", FVector2D(10,19) ) );
 		Set( "Sequencer.NotificationImage_AddedPlayMovieSceneEvent", new IMAGE_BRUSH( "Old/Checkbox_checked", Icon16x16 ) );
+		Set( "Sequencer.Save", new IMAGE_BRUSH( "Icons/icon_file_save_40x", Icon40x40 ) );
+		Set( "Sequencer.Save.Small", new IMAGE_BRUSH( "Icons/icon_file_save_16px", Icon20x20 ) );
 		Set( "Sequencer.Key", new IMAGE_BRUSH( "Common/Key", FVector2D(7.0f,12.0f) ) );
 		Set( "Sequencer.PartialKey", new IMAGE_BRUSH( "Sequencer/PartialKey", FVector2D(11.f, 11.f) ) );
 		Set( "Sequencer.GenericSection.Background", new BOX_BRUSH( "Sequencer/GenericSectionBackground", FMargin(5.f/16.f) ) );
@@ -1596,8 +1598,16 @@ void FSlateEditorStyle::FStyle::SetupGeneralStyles()
 		Set( "Sequencer.AnimationOutliner.RegularFont", TTF_CORE_FONT( "Fonts/Roboto-Regular", 9 ) );
 		Set( "Sequencer.ShotFilter", new IMAGE_BRUSH( "Sequencer/FilteredArea", FVector2D(74,74), FLinearColor::White, ESlateBrushTileType::Both ) );
 		Set( "Sequencer.KeyMark", new IMAGE_BRUSH("Sequencer/KeyMark", FVector2D(3,21), FLinearColor::White, ESlateBrushTileType::NoTile ) );
+		Set( "Sequencer.AddObject", new IMAGE_BRUSH( "Icons/icon_Sequencer_AddObject_40x", Icon40x40 ) );
+		Set( "Sequencer.AddObject.Small", new IMAGE_BRUSH( "Icons/icon_Sequencer_AddObject_40x", Icon20x20 ) );
+		Set( "Sequencer.ToggleAutoKeyEnabled", new IMAGE_BRUSH( "Icons/icon_Sequencer_AutoKey_40x", Icon40x40 ) );
+		Set( "Sequencer.ToggleAutoKeyEnabled.Small", new IMAGE_BRUSH( "Icons/icon_Sequencer_AutoKey_40x", Icon20x20 ) );
+		Set( "Sequencer.ToggleKeyAllEnabled", new IMAGE_BRUSH( "Icons/icon_Sequencer_KeyAll_40x", Icon40x40 ) );
+		Set( "Sequencer.ToggleKeyAllEnabled.Small", new IMAGE_BRUSH( "Icons/icon_Sequencer_KeyAll_40x", Icon20x20 ) );
 		Set( "Sequencer.ToggleIsSnapEnabled", new IMAGE_BRUSH( "Icons/icon_Sequencer_ToggleSnap_40x", Icon40x40 ) );
-		Set( "Sequencer.ToggleIsSnapEnabled.Small", new IMAGE_BRUSH( "Icons/icon_Sequencer_ToggleSnap_40x", Icon16x16 ) );
+		Set( "Sequencer.ToggleIsSnapEnabled.Small", new IMAGE_BRUSH( "Icons/icon_Sequencer_ToggleSnap_40x", Icon20x20 ) );
+		Set( "Sequencer.ToggleShowCurveEditor", new IMAGE_BRUSH("Icons/icon_MatineeCurveView_40px", Icon40x40) );
+		Set( "Sequencer.ToggleShowCurveEditor.Small", new IMAGE_BRUSH("Icons/icon_MatineeCurveView_40px", Icon20x20) );
 		Set( "Sequencer.ToggleAutoScroll", new IMAGE_BRUSH( "Icons/icon_Sequencer_ToggleAutoScroll_40x", Icon40x40 ) );
 		Set( "Sequencer.ToggleAutoScroll.Small", new IMAGE_BRUSH( "Icons/icon_Sequencer_ToggleAutoScroll_16x", Icon16x16 ) );
 		Set( "Sequencer.LockCamera", new IMAGE_BRUSH( "Icons/icon_ToggleActorPilotCameraView_16x", Icon16x16 ) );
diff --git a/Engine/Source/Editor/MovieSceneTools/Private/TrackEditors/TransformTrackEditor.cpp b/Engine/Source/Editor/MovieSceneTools/Private/TrackEditors/TransformTrackEditor.cpp
--- a/Engine/Source/Editor/MovieSceneTools/Private/TrackEditors/TransformTrackEditor.cpp
+++ b/Engine/Source/Editor/MovieSceneTools/Private/TrackEditors/TransformTrackEditor.cpp
@@ -479,7 +479,7 @@ void F3DTransformTrackEditor::OnTransformChangedInternals(float KeyTime, UObject
 		}
 
 		// If forcing the keyframe, set the last data invalid so that the transform keyframe will be set regardless of the current time values being equal.
-		if ( bForceKey )
+		if ( bForceKey || IsAllowedKeyAll() )
 		{
 			TransformPair.LastTransformData.bValid = false;
 		}
diff --git a/Engine/Source/Editor/Sequencer/Private/SSequencer.cpp b/Engine/Source/Editor/Sequencer/Private/SSequencer.cpp
--- a/Engine/Source/Editor/Sequencer/Private/SSequencer.cpp
+++ b/Engine/Source/Editor/Sequencer/Private/SSequencer.cpp
@@ -448,8 +448,6 @@ TSharedRef<SWidget> SSequencer::MakeToolBar()
 	
 	ToolBarBuilder.BeginSection("Base Commands");
 	{
-		ToolBarBuilder.SetLabelVisibility(EVisibility::Visible);
-
 		if (OnGetAddMenuContent.IsBound())
 		{
 			ToolBarBuilder.AddWidget(
@@ -494,59 +492,45 @@ TSharedRef<SWidget> SSequencer::MakeToolBar()
 				]);
 		}
 
+		// General 
 		if( Sequencer.Pin()->IsLevelEditorSequencer() )
 		{
-			ToolBarBuilder.AddWidget
-			(
-				SNew( SButton )
-				.ButtonStyle(FEditorStyle::Get(), "FlatButton")
-				.ToolTipText( LOCTEXT( "SaveDirtyPackagesTooltip", "Saves the current Movie Scene" ) )
-				.ContentPadding(FMargin(6, 2))
-				.ForegroundColor( FSlateColor::UseForeground() )
-				.OnClicked( this, &SSequencer::OnSaveMovieSceneClicked )
-				[
-					SNew( SHorizontalBox )
-					+ SHorizontalBox::Slot()
-					.VAlign(VAlign_Center)
-					.AutoWidth()
-					[
-						SNew(STextBlock)
-						.Font(FEditorStyle::Get().GetFontStyle("FontAwesome.11"))
-						.Text(FText::FromString(FString(TEXT("\xf0c7"))) /*fa-floppy-o*/)
-						.ShadowOffset( FVector2D( 1,1) )
-					]
-					+ SHorizontalBox::Slot()
-					.AutoWidth()
-					.VAlign(VAlign_Center)
-					.Padding(4, 1, 0, 0)
-					[
-						SNew( STextBlock )
-						.Text( LOCTEXT( "SaveMovieScene", "Save" ) )
-						.ShadowOffset( FVector2D( 1,1) )
-					]
-				]
-			);
-
-			ToolBarBuilder.AddSeparator();
+			ToolBarBuilder.AddToolBarButton(
+				FUIAction(FExecuteAction::CreateSP(this, &SSequencer::OnSaveMovieSceneClicked))
+				, NAME_None
+				, LOCTEXT("SaveDirtyPackages", "Save")
+				, LOCTEXT("SaveDirtyPackagesTooltip", "Saves the current movie scene")
+				//, FSlateIcon(FEditorStyle::GetStyleSetName(), "LevelEditor.Save"));
+				, FSlateIcon(FEditorStyle::GetStyleSetName(), "Sequencer.Save"));
 		}
+		ToolBarBuilder.AddComboButton(
+			FUIAction(),
+			FOnGetContent::CreateSP( this, &SSequencer::MakeGeneralMenu ),
+			LOCTEXT( "GeneralOptions", "General Options" ),
+			LOCTEXT( "GeneralOptionsToolTip", "General Options" ),
+			TAttribute<FSlateIcon>(),
+			true );
 
-		ToolBarBuilder.AddToolBarButton( FSequencerCommands::Get().ToggleAutoKeyEnabled );
+		ToolBarBuilder.AddSeparator();
 
-		if ( Sequencer.Pin()->IsLevelEditorSequencer() )
+		if( Sequencer.Pin()->IsLevelEditorSequencer() )
 		{
-			ToolBarBuilder.AddToolBarButton( FSequencerCommands::Get().ToggleCleanView );
+			ToolBarBuilder.AddToolBarButton(
+				FUIAction(FExecuteAction::CreateSP(this, &SSequencer::OnAddObjectClicked))
+				, NAME_None
+				, LOCTEXT("AddObject", "Add Object")
+				, LOCTEXT("AddObjectTooltip", "Add selected objects")
+				, FSlateIcon(FEditorStyle::GetStyleSetName(), "Sequencer.AddObject"));
+
+			ToolBarBuilder.AddToolBarButton( FSequencerCommands::Get().ToggleKeyAllEnabled );
 		}
 
-		ToolBarBuilder.SetLabelVisibility(EVisibility::Collapsed);
-		ToolBarBuilder.AddToolBarButton( FSequencerCommands::Get().ToggleAutoScroll, NAME_None, TAttribute<FText>( FText::GetEmpty() ) );
-		ToolBarBuilder.SetLabelVisibility(EVisibility::Visible);
+		ToolBarBuilder.AddToolBarButton( FSequencerCommands::Get().ToggleAutoKeyEnabled );
 	}
 	ToolBarBuilder.EndSection();
 
 	ToolBarBuilder.BeginSection("Snapping");
 	{
-
-		ToolBarBuilder.SetLabelVisibility( EVisibility::Collapsed );
 		ToolBarBuilder.AddToolBarButton( FSequencerCommands::Get().ToggleIsSnapEnabled, NAME_None, TAttribute<FText>( FText::GetEmpty() ) );
 		ToolBarBuilder.AddComboButton(
 			FUIAction(),
@@ -583,9 +567,7 @@ TSharedRef<SWidget> SSequencer::MakeToolBar()
 
 	ToolBarBuilder.BeginSection("Curve Editor");
 	{
-		ToolBarBuilder.SetLabelVisibility( EVisibility::Visible );
 		ToolBarBuilder.AddToolBarButton( FSequencerCommands::Get().ToggleShowCurveEditor );
-		ToolBarBuilder.SetLabelVisibility( EVisibility::Collapsed );
 	}
 
 	return ToolBarBuilder.MakeWidget();
@@ -596,6 +578,23 @@ TSharedRef<SWidget> SSequencer::MakeAddMenu()
 	return OnGetAddMenuContent.Execute(Sequencer.Pin().ToSharedRef());
 }
 
+TSharedRef<SWidget> SSequencer::MakeGeneralMenu()
+{
+	FMenuBuilder MenuBuilder( false, Sequencer.Pin()->GetCommandBindings() );
+
+	MenuBuilder.BeginSection( "ViewOptions", LOCTEXT( "ViewMenuHeader", "View" ) );
+	{
+		if (Sequencer.Pin()->IsLevelEditorSequencer())
+		{
+			MenuBuilder.AddMenuEntry( FSequencerCommands::Get().ToggleCleanView );
+		}
+		MenuBuilder.AddMenuEntry( FSequencerCommands::Get().ToggleAutoScroll );
+	}
+	MenuBuilder.EndSection();
+
+	return MenuBuilder.MakeWidget();
+}
+
 TSharedRef<SWidget> SSequencer::MakeSnapMenu()
 {
 	FMenuBuilder MenuBuilder( false, Sequencer.Pin()->GetCommandBindings() );
@@ -1159,11 +1158,14 @@ void SSequencer::DeleteSelectedNodes()
 	}
 }
 
-FReply SSequencer::OnSaveMovieSceneClicked()
+void SSequencer::OnSaveMovieSceneClicked()
 {
 	Sequencer.Pin()->SaveCurrentMovieScene();
+}
 
-	return FReply::Unhandled();
+void SSequencer::OnAddObjectClicked()
+{
+	Sequencer.Pin()->AddSelectedObjects();
 }
 
 void SSequencer::StepToNextKey()
diff --git a/Engine/Source/Editor/Sequencer/Private/SSequencer.h b/Engine/Source/Editor/Sequencer/Private/SSequencer.h
--- a/Engine/Source/Editor/Sequencer/Private/SSequencer.h
+++ b/Engine/Source/Editor/Sequencer/Private/SSequencer.h
@@ -113,7 +113,10 @@ public:
 	bool UserIsSelecting() { return bUserIsSelecting; }
 
 	/** Called when the save button is clicked */
-	FReply OnSaveMovieSceneClicked();
+	void OnSaveMovieSceneClicked();
+
+	/** Called when the add object button is clicked */
+	void OnAddObjectClicked();
 
 	/** Access the tree view for this sequencer */
 	TSharedPtr<SSequencerTreeView> GetTreeView() const;
@@ -128,6 +131,9 @@ private:
 	/** Makes the add menu for the toolbar. */
 	TSharedRef<SWidget> MakeAddMenu();
 
+	/** Makes the general menu for the toolbar. */
+	TSharedRef<SWidget> MakeGeneralMenu();
+
 	/** Makes the snapping menu for the toolbar. */
 	TSharedRef<SWidget> MakeSnapMenu();
 
diff --git a/Engine/Source/Editor/Sequencer/Private/Sequencer.cpp b/Engine/Source/Editor/Sequencer/Private/Sequencer.cpp
--- a/Engine/Source/Editor/Sequencer/Private/Sequencer.cpp
+++ b/Engine/Source/Editor/Sequencer/Private/Sequencer.cpp
@@ -539,6 +539,16 @@ void FSequencer::SetAutoKeyEnabled(bool bAutoKeyEnabled)
 	Settings->SetAutoKeyEnabled(bAutoKeyEnabled);
 }
 
+bool FSequencer::GetKeyAllEnabled() const 
+{
+	return Settings->GetKeyAllEnabled();
+}
+
+void FSequencer::SetKeyAllEnabled(bool bKeyAllEnabled) 
+{
+	Settings->SetKeyAllEnabled(bKeyAllEnabled);
+}
+
 bool FSequencer::IsRecordingLive() const 
 {
 	return PlaybackState == EMovieScenePlayerStatus::Recording && GIsPlayInEditorWorld;
@@ -1489,6 +1499,68 @@ void FSequencer::SaveCurrentMovieScene()
 	FEditorFileUtils::PromptForCheckoutAndSave( PackagesToSave, false, false );
 }
 
+void FSequencer::AddSelectedObjects()
+{
+	UMovieSceneSequence* OwnerSequence = GetFocusedMovieSceneSequence();
+	UMovieScene* OwnerMovieScene = OwnerSequence->GetMovieScene();
+
+	TArray<AActor*> SelectedActors;
+
+	USelection* CurrentSelection = GEditor->GetSelectedActors();
+	TArray<UObject*> SelectedObjects;
+	CurrentSelection->GetSelectedObjects( AActor::StaticClass(), SelectedObjects );
+	for (TArray<UObject*>::TIterator It(SelectedObjects); It; ++It)
+	{
+		AActor* Actor = CastChecked<AActor>(*It);
+
+		if (Actor != nullptr)
+		{			
+			FGuid ObjectGuid = OwnerSequence->FindObjectId(*Actor);
+
+			// Add this object if it hasn't already been possessed
+			if (ObjectGuid.IsValid())
+			{
+				FMovieScenePossessable* Possessable = OwnerMovieScene->FindPossessable(ObjectGuid);
+				if (Possessable == nullptr)
+				{
+					SelectedActors.Add(Actor);
+				}
+			}
+			else
+			{
+				SelectedActors.Add(Actor);
+			}
+		}
+	}
+
+	bool bPossessableAdded = false;
+	if (SelectedActors.Num() != 0)
+	{
+		const FScopedTransaction Transaction(LOCTEXT("UndoPossessingObject", "Possess Object with MovieScene"));
+			
+		// possess the object!
+		OwnerSequence->Modify();
+				
+		for (auto Actor : SelectedActors)
+		{
+			const FGuid PossessableGuid = OwnerMovieScene->AddPossessable(Actor->GetActorLabel(), Actor->GetClass());
+
+			if (IsShotFilteringOn())
+			{
+				AddUnfilterableObject(PossessableGuid);
+			}
+
+			OwnerSequence->BindPossessableObject(PossessableGuid, *Actor);
+			bPossessableAdded = true;
+		}
+	}
+
+	if (bPossessableAdded)
+	{
+		NotifyMovieSceneDataChanged();
+	}
+}
+
 void FSequencer::OnSectionSelectionChanged()
 {
 	for (TWeakObjectPtr<UMovieSceneSection> SelectedSection : Selection.GetSelectedSections())
@@ -1844,6 +1916,12 @@ void FSequencer::BindSequencerCommands()
 		FCanExecuteAction::CreateLambda( []{ return true; } ),
 		FIsActionChecked::CreateLambda( [this]{ return Settings->GetAutoKeyEnabled(); } ) );
 
+	SequencerCommandBindings->MapAction(
+		Commands.ToggleKeyAllEnabled,
+		FExecuteAction::CreateLambda( [this]{ Settings->SetKeyAllEnabled( !Settings->GetKeyAllEnabled() ); } ),
+		FCanExecuteAction::CreateLambda( []{ return true; } ),
+		FIsActionChecked::CreateLambda( [this]{ return Settings->GetKeyAllEnabled(); } ) );
+
 	bAutoScrollEnabled = Settings->GetAutoScrollEnabled();
 	SequencerCommandBindings->MapAction(
 		Commands.ToggleAutoScroll,
diff --git a/Engine/Source/Editor/Sequencer/Private/Sequencer.h b/Engine/Source/Editor/Sequencer/Private/Sequencer.h
--- a/Engine/Source/Editor/Sequencer/Private/Sequencer.h
+++ b/Engine/Source/Editor/Sequencer/Private/Sequencer.h
@@ -72,6 +72,8 @@ public:
 	virtual void AddAnimation(FGuid ObjectGuid, class UAnimSequence* AnimSequence) override;
 	virtual bool GetAutoKeyEnabled() const override;
 	virtual void SetAutoKeyEnabled(bool bAutoKeyEnabled) override;
+	virtual bool GetKeyAllEnabled() const override;
+	virtual void SetKeyAllEnabled(bool bKeyAllEnabled) override;
 	virtual bool IsRecordingLive() const override;
 	virtual float GetCurrentLocalTime(UMovieSceneSequence& InMovieSceneSequence) override;
 	virtual float GetGlobalTime() override;
@@ -266,6 +268,9 @@ public:
 	/** Called to save the current movie scene */
 	void SaveCurrentMovieScene();
 
+	/** Called to add selected objects to the movie scene */
+	void AddSelectedObjects();
+
 protected:
 
 	/** Reset data about a movie scene when pushing or popping a movie scene. */
diff --git a/Engine/Source/Editor/Sequencer/Private/SequencerCommands.cpp b/Engine/Source/Editor/Sequencer/Private/SequencerCommands.cpp
--- a/Engine/Source/Editor/Sequencer/Private/SequencerCommands.cpp
+++ b/Engine/Source/Editor/Sequencer/Private/SequencerCommands.cpp
@@ -23,24 +23,25 @@ void FSequencerCommands::RegisterCommands()
 	UI_COMMAND( SetKey, "Set Key", "Sets a key at the current time for the selected actor.", EUserInterfaceActionType::Button, FInputChord(EKeys::K) );
 
 	UI_COMMAND( ToggleAutoKeyEnabled, "Auto Key", "Enables and disables auto keying.", EUserInterfaceActionType::ToggleButton, FInputChord() );
+	UI_COMMAND( ToggleKeyAllEnabled, "Key All", "Enables and disables key all.", EUserInterfaceActionType::ToggleButton, FInputChord() );
 
 	UI_COMMAND( ToggleAutoScroll, "Auto Scroll", "Toggle auto-scroll: When enabled, automatically scrolls the sequencer view to keep the current time visible.", EUserInterfaceActionType::ToggleButton, FInputChord(EModifierKey::Shift, EKeys::S) );
 
-	UI_COMMAND( ToggleShowFrameNumbers, "Show frame numbers", "Enables and disables showing frame numbers.", EUserInterfaceActionType::ToggleButton, FInputChord() );
-	UI_COMMAND( ToggleShowRangeSlider, "Show range slider", "Enables and disables showing the time range slider.", EUserInterfaceActionType::ToggleButton, FInputChord() );
+	UI_COMMAND( ToggleShowFrameNumbers, "Show Frame Numbers", "Enables and disables showing frame numbers.", EUserInterfaceActionType::ToggleButton, FInputChord() );
+	UI_COMMAND( ToggleShowRangeSlider, "Show Range Slider", "Enables and disables showing the time range slider.", EUserInterfaceActionType::ToggleButton, FInputChord() );
 	UI_COMMAND( ToggleLockInOutToStartEndRange, "Lock In/Out to Start/End", "Enables and disables locking the in/out to the start/end range while scrolling the view.", EUserInterfaceActionType::ToggleButton, FInputChord() );
 	UI_COMMAND( ToggleIsSnapEnabled, "Enable Snapping", "Enables and disables snapping.", EUserInterfaceActionType::ToggleButton, FInputChord() );
 
-	UI_COMMAND( ToggleSnapKeyTimesToInterval, "Snap to the interval", "Snap keys to the interval.", EUserInterfaceActionType::ToggleButton, FInputChord() );
-	UI_COMMAND( ToggleSnapKeyTimesToKeys, "Snap to other keys", "Snap keys to other keys in the section.", EUserInterfaceActionType::ToggleButton, FInputChord() );
+	UI_COMMAND( ToggleSnapKeyTimesToInterval, "Snap to the Interval", "Snap keys to the interval.", EUserInterfaceActionType::ToggleButton, FInputChord() );
+	UI_COMMAND( ToggleSnapKeyTimesToKeys, "Snap to Other Keys", "Snap keys to other keys in the section.", EUserInterfaceActionType::ToggleButton, FInputChord() );
 
-	UI_COMMAND( ToggleSnapSectionTimesToInterval, "Snap to the interval", "Snap sections to the interval.", EUserInterfaceActionType::ToggleButton, FInputChord() );
-	UI_COMMAND( ToggleSnapSectionTimesToSections, "Snap to other sections", "Snap sections to other sections.", EUserInterfaceActionType::ToggleButton, FInputChord() );
+	UI_COMMAND( ToggleSnapSectionTimesToInterval, "Snap to the Interval", "Snap sections to the interval.", EUserInterfaceActionType::ToggleButton, FInputChord() );
+	UI_COMMAND( ToggleSnapSectionTimesToSections, "Snap to Other Sections", "Snap sections to other sections.", EUserInterfaceActionType::ToggleButton, FInputChord() );
 
-	UI_COMMAND( ToggleSnapPlayTimeToInterval, "Snap to the interval while scrubbing", "Snap the play time to the interval while scrubbing.", EUserInterfaceActionType::ToggleButton, FInputChord() );
-	UI_COMMAND( ToggleSnapPlayTimeToDraggedKey, "Snap to the dragged key", "Snap the play time to the dragged key.", EUserInterfaceActionType::ToggleButton, FInputChord() );
+	UI_COMMAND( ToggleSnapPlayTimeToInterval, "Snap to the Interval While Scrubbing", "Snap the play time to the interval while scrubbing.", EUserInterfaceActionType::ToggleButton, FInputChord() );
+	UI_COMMAND( ToggleSnapPlayTimeToDraggedKey, "Snap to the Dragged Key", "Snap the play time to the dragged key.", EUserInterfaceActionType::ToggleButton, FInputChord() );
 
-	UI_COMMAND( ToggleSnapCurveValueToInterval, "Snap curve key values", "Snap curve keys to the value snapping interval.", EUserInterfaceActionType::ToggleButton, FInputChord() );
+	UI_COMMAND( ToggleSnapCurveValueToInterval, "Snap Curve Key Values", "Snap curve keys to the value snapping interval.", EUserInterfaceActionType::ToggleButton, FInputChord() );
 
 	UI_COMMAND( ToggleCleanView, "Clean View", "Enable 'Clean View' mode which displays only global tracks when no filter is applied.", EUserInterfaceActionType::ToggleButton, FInputChord() );
 	
diff --git a/Engine/Source/Editor/Sequencer/Private/SequencerCommands.h b/Engine/Source/Editor/Sequencer/Private/SequencerCommands.h
--- a/Engine/Source/Editor/Sequencer/Private/SequencerCommands.h
+++ b/Engine/Source/Editor/Sequencer/Private/SequencerCommands.h
@@ -54,6 +54,9 @@ public:
 	/** Turns auto keying on and off. */
 	TSharedPtr< FUICommandInfo > ToggleAutoKeyEnabled;
 
+	/** Turns key all on and off. */
+	TSharedPtr< FUICommandInfo > ToggleKeyAllEnabled;
+
 	/** Turns show frame numbers on and off. */
 	TSharedPtr< FUICommandInfo > ToggleShowFrameNumbers;
 
diff --git a/Engine/Source/Editor/Sequencer/Private/SequencerSettings.cpp b/Engine/Source/Editor/Sequencer/Private/SequencerSettings.cpp
--- a/Engine/Source/Editor/Sequencer/Private/SequencerSettings.cpp
+++ b/Engine/Source/Editor/Sequencer/Private/SequencerSettings.cpp
@@ -6,6 +6,7 @@ USequencerSettings::USequencerSettings( const FObjectInitializer& ObjectInitiali
 	: Super( ObjectInitializer )
 {
 	bAutoKeyEnabled = false;
+	bKeyAllEnabled = false;
 	bShowFrameNumbers = true;
 	bShowRangeSlider = false;
 	bLockInOutToStartEndRange = false;
@@ -39,6 +40,20 @@ void USequencerSettings::SetAutoKeyEnabled(bool InbAutoKeyEnabled)
 	}
 }
 
+bool USequencerSettings::GetKeyAllEnabled() const
+{
+	return bKeyAllEnabled;
+}
+
+void USequencerSettings::SetKeyAllEnabled(bool InbKeyAllEnabled)
+{
+	if ( bKeyAllEnabled != InbKeyAllEnabled )
+	{
+		bKeyAllEnabled = InbKeyAllEnabled;
+		SaveConfig();
+	}
+}
+
 bool USequencerSettings::GetShowFrameNumbers() const
 {
 	return bShowFrameNumbers;
diff --git a/Engine/Source/Editor/Sequencer/Private/SequencerSettings.h b/Engine/Source/Editor/Sequencer/Private/SequencerSettings.h
--- a/Engine/Source/Editor/Sequencer/Private/SequencerSettings.h
+++ b/Engine/Source/Editor/Sequencer/Private/SequencerSettings.h
@@ -49,6 +49,11 @@ public:
 	/** Sets whether or not auto key is enabled. */
 	void SetAutoKeyEnabled(bool InbAutoKeyEnabled);
 
+	/** Gets whether or not key all is enabled. */
+	bool GetKeyAllEnabled() const;
+	/** Sets whether or not key all is enabled. */
+	void SetKeyAllEnabled(bool InbKeyAllEnabled);
+
 	/** Gets whether or not to show frame numbers. */
 	bool GetShowFrameNumbers() const;
 	/** Sets whether or not to show frame numbers. */
@@ -144,6 +149,9 @@ protected:
 	UPROPERTY( config )
 	bool bAutoKeyEnabled;
 
+	UPROPERTY( config )
+	bool bKeyAllEnabled;
+
 	UPROPERTY( config )
 	bool bShowFrameNumbers;
 
diff --git a/Engine/Source/Editor/Sequencer/Public/MovieSceneTrackEditor.h b/Engine/Source/Editor/Sequencer/Public/MovieSceneTrackEditor.h
--- a/Engine/Source/Editor/Sequencer/Public/MovieSceneTrackEditor.h
+++ b/Engine/Source/Editor/Sequencer/Public/MovieSceneTrackEditor.h
@@ -56,6 +56,12 @@ public:
 		return Sequencer.Pin()->IsRecordingLive() || Sequencer.Pin()->GetAutoKeyEnabled();
 	}
 
+	/** Gets whether the tool can key all*/
+	virtual bool IsAllowedKeyAll() const
+	{
+		return Sequencer.Pin()->GetKeyAllEnabled();
+	}
+
 	void NotifyMovieSceneDataChanged()
 	{
 		TSharedPtr<ISequencer> SequencerPin = Sequencer.Pin();
