Sequencer: Add missing changes.

diff --git a/Engine/Source/Editor/Sequencer/Public/ISequencer.h b/Engine/Source/Editor/Sequencer/Public/ISequencer.h
--- a/Engine/Source/Editor/Sequencer/Public/ISequencer.h
+++ b/Engine/Source/Editor/Sequencer/Public/ISequencer.h
@@ -73,6 +73,12 @@ public:
 	/** Sets whether autokey is enabled in this sequencer. */
 	virtual void SetAutoKeyEnabled(bool bAutoKeyEnabled) = 0;
 
+	/** @return Returns whether key all is enabled in this sequencer */
+	virtual bool GetKeyAllEnabled() const = 0;
+
+	/** Sets whether key all is enabled in this sequencer. */
+	virtual void SetKeyAllEnabled(bool bKeyAllEnabled) = 0;
+
 	/** @return Returns whether sequencer is currently recording live data from simulated actors */
 	virtual bool IsRecordingLive() const = 0;
 
