Sequencer: Object should stay selected if a child node is selected. Also, fixed some situations where multiple selection transactions were being generated even though the resulting selection state was correct.

diff --git a/Engine/Source/Editor/Sequencer/Private/SAnimationOutlinerTreeNode.cpp b/Engine/Source/Editor/Sequencer/Private/SAnimationOutlinerTreeNode.cpp
--- a/Engine/Source/Editor/Sequencer/Private/SAnimationOutlinerTreeNode.cpp
+++ b/Engine/Source/Editor/Sequencer/Private/SAnimationOutlinerTreeNode.cpp
@@ -263,6 +263,16 @@ FReply SAnimationOutlinerTreeNode::OnAddKeyClicked()
 	return FReply::Handled();
 }
 
+TSharedPtr<FSequencerDisplayNode> SAnimationOutlinerTreeNode::GetRootNode(TSharedPtr<FSequencerDisplayNode> ObjectNode)
+{
+	if (!ObjectNode->GetParent().IsValid())
+	{
+		return ObjectNode;
+	}
+
+	return GetRootNode(ObjectNode->GetParent());
+}
+
 void SAnimationOutlinerTreeNode::GetAllDescendantNodes(TSharedPtr<FSequencerDisplayNode> RootNode, TArray<TSharedRef<FSequencerDisplayNode> >& AllNodes)
 {
 	if (!RootNode.IsValid())
@@ -292,6 +302,8 @@ FReply SAnimationOutlinerTreeNode::OnMouseButtonDown( const FGeometry& MyGeometr
 		TArray<TSharedPtr<FSequencerDisplayNode> > AffectedNodes;
 		AffectedNodes.Add(DisplayNode.ToSharedRef());
 
+		Sequencer.GetSelection().SuspendBroadcast();
+
 		if (MouseEvent.IsShiftDown())
 		{
 			FSequencerNodeTree& ParentTree = DisplayNode->GetParentTree();
@@ -359,6 +371,9 @@ FReply SAnimationOutlinerTreeNode::OnMouseButtonDown( const FGeometry& MyGeometr
 		}
 
 		OnSelectionChanged( AffectedNodes );
+
+		Sequencer.GetSelection().ResumeBroadcast();
+
 		return FReply::Handled();
 	}
 
@@ -433,55 +448,40 @@ FText SAnimationOutlinerTreeNode::GetDisplayName() const
 
 void SAnimationOutlinerTreeNode::OnSelectionChanged( TArray<TSharedPtr<FSequencerDisplayNode> > AffectedNodes )
 {
-	TArray<TSharedPtr<FSequencerDisplayNode> > ObjectNodes;
+	FSequencer& Sequencer = DisplayNode->GetSequencer();
+	if (!Sequencer.IsLevelEditorSequencer())
+	{
+		return;
+	}
+
+	TArray<TSharedPtr<FSequencerDisplayNode> > AffectedRootNodes;
 
 	for (int32 NodeIdx = 0; NodeIdx < AffectedNodes.Num(); ++NodeIdx)
 	{
-		if( AffectedNodes[NodeIdx]->GetType() == ESequencerNode::Object )
+		TSharedPtr<FSequencerDisplayNode> RootNode = GetRootNode(AffectedNodes[NodeIdx]);
+		if (RootNode.IsValid() && RootNode->GetType() == ESequencerNode::Object)
 		{
-			ObjectNodes.Add(AffectedNodes[NodeIdx]);
+			AffectedRootNodes.Add(RootNode);
 		}
 	}
 
-	if (!ObjectNodes.Num())
-	{
-		return;
-	}
-
-	FSequencer& Sequencer = DisplayNode->GetSequencer();
-	if (!Sequencer.IsLevelEditorSequencer())
+	if (!AffectedRootNodes.Num())
 	{
 		return;
 	}
 
-	// Mark that the user is selecting so that the UI doesn't respond to the selection changes in the following block
-	TSharedRef<SSequencer> SequencerWidget = StaticCastSharedRef<SSequencer>(Sequencer.GetSequencerWidget());
-	SequencerWidget->SetUserIsSelecting(true);
-
-	const FScopedTransaction Transaction(NSLOCTEXT("UnrealEd", "ClickingOnActors", "Clicking on Actors"));
-	bool bActorSelected = false;
-
 	const FModifierKeysState ModifierKeys = FSlateApplication::Get().GetModifierKeys();
 	bool IsControlDown = ModifierKeys.IsControlDown();
 	bool IsShiftDown = ModifierKeys.IsShiftDown();
 
-	const bool bNotifySelectionChanged = false;
-	const bool bDeselectBSP = true;
-	const bool bWarnAboutTooManyActors = false;
-
-	GEditor->GetSelectedActors()->Modify();
-
-	if (!IsControlDown && !IsShiftDown)
-	{
-		GEditor->SelectNone(bNotifySelectionChanged, bDeselectBSP, bWarnAboutTooManyActors);
-	}
-
-	GEditor->GetSelectedActors()->BeginBatchSelectOperation();
+	TArray<AActor*> ActorsToSelect;
+	TArray<AActor*> ActorsToRemainSelected;
+	TArray<AActor*> ActorsToDeselect;
 
-	// Select objects bound to the object node
-	for (int32 ObjectIdx = 0; ObjectIdx < ObjectNodes.Num(); ++ObjectIdx)
+	// Find objects bound to the object node and determine what actors need to be selected in the editor
+	for (int32 ObjectIdx = 0; ObjectIdx < AffectedRootNodes.Num(); ++ObjectIdx)
 	{
-		const TSharedPtr<const FObjectBindingNode> ObjectNode = StaticCastSharedPtr<const FObjectBindingNode>( ObjectNodes[ObjectIdx] );
+		const TSharedPtr<const FObjectBindingNode> ObjectNode = StaticCastSharedPtr<const FObjectBindingNode>( AffectedRootNodes[ObjectIdx] );
 
 		// Get the bound objects
 		TArray<UObject*> RuntimeObjects;
@@ -489,34 +489,97 @@ void SAnimationOutlinerTreeNode::OnSelectionChanged( TArray<TSharedPtr<FSequence
 		
 		if( RuntimeObjects.Num() > 0 )
 		{
-			// Select each actor			
 			for( int32 ActorIdx = 0; ActorIdx < RuntimeObjects.Num(); ++ActorIdx )
 			{
 				AActor* Actor = Cast<AActor>( RuntimeObjects[ActorIdx] );
 
 				if( Actor )
 				{
-					bool bSelectActor = true;
-
+					bool bIsActorSelected = GEditor->GetSelectedActors()->IsSelected(Actor);
 					if (IsControlDown)
 					{
-						bSelectActor = Sequencer.GetSelection().IsSelected(ObjectNodes[ObjectIdx].ToSharedRef());
+						if (!bIsActorSelected)
+						{
+							ActorsToSelect.Add(Actor);
+						}
+						else
+						{
+							ActorsToDeselect.Add(Actor);
+						}
 					}
+					else
+					{
+						if (!bIsActorSelected)
+						{
+							ActorsToSelect.Add(Actor);
+						}
+						else
+						{
+							ActorsToRemainSelected.Add(Actor);
+						}
+					}
+				}
+			}
+		}
+	}
 
-					GEditor->SelectActor(Actor, bSelectActor, bNotifySelectionChanged );
-					bActorSelected = true;
+	if (!IsControlDown && !IsShiftDown)
+	{
+		for (FSelectionIterator SelectionIt(*GEditor->GetSelectedActors()); SelectionIt; ++SelectionIt)
+		{
+			AActor* Actor = CastChecked< AActor >( *SelectionIt );
+			if (Actor)
+			{
+				if (ActorsToSelect.Find(Actor) == INDEX_NONE && 
+					ActorsToDeselect.Find(Actor) == INDEX_NONE &&
+					ActorsToRemainSelected.Find(Actor) == INDEX_NONE)
+				{
+					ActorsToDeselect.Add(Actor);
 				}
 			}
 		}
 	}
 
-	GEditor->GetSelectedActors()->EndBatchSelectOperation();
+	// If there's no selection to change, return early
+	if (ActorsToSelect.Num() + ActorsToDeselect.Num() == 0)
+	{
+		return;
+	}
+
+	// Mark that the user is selecting so that the UI doesn't respond to the selection changes in the following block
+	TSharedRef<SSequencer> SequencerWidget = StaticCastSharedRef<SSequencer>(Sequencer.GetSequencerWidget());
+	SequencerWidget->SetUserIsSelecting(true);
+
+	const FScopedTransaction Transaction(NSLOCTEXT("UnrealEd", "ClickingOnActors", "Clicking on Actors"));
+
+	const bool bNotifySelectionChanged = false;
+	const bool bDeselectBSP = true;
+	const bool bWarnAboutTooManyActors = false;
+	const bool bSelectEvenIfHidden = true;
+
+	GEditor->GetSelectedActors()->Modify();
+
+	if (!IsControlDown && !IsShiftDown)
+	{
+		GEditor->SelectNone(bNotifySelectionChanged, bDeselectBSP, bWarnAboutTooManyActors);
+	}
+
+	GEditor->GetSelectedActors()->BeginBatchSelectOperation();
+
+	for (auto ActorToSelect : ActorsToSelect)
+	{
+		GEditor->SelectActor(ActorToSelect, true, bNotifySelectionChanged, bSelectEvenIfHidden);
+	}
 
-	if( bActorSelected )
+	for (auto ActorToDeselect : ActorsToDeselect)
 	{
-		GEditor->NoteSelectionChange();
+		GEditor->SelectActor(ActorToDeselect, false, bNotifySelectionChanged, bSelectEvenIfHidden);
 	}
 
+	GEditor->GetSelectedActors()->EndBatchSelectOperation();
+
+	GEditor->NoteSelectionChange();
+
 	// Unlock the selection so that the sequencer widget can now respond to selection changes in the level
 	SequencerWidget->SetUserIsSelecting(false);
 }
diff --git a/Engine/Source/Editor/Sequencer/Private/SAnimationOutlinerTreeNode.h b/Engine/Source/Editor/Sequencer/Private/SAnimationOutlinerTreeNode.h
--- a/Engine/Source/Editor/Sequencer/Private/SAnimationOutlinerTreeNode.h
+++ b/Engine/Source/Editor/Sequencer/Private/SAnimationOutlinerTreeNode.h
@@ -60,6 +60,10 @@ private:
 	/** Get all descendant nodes from the given root node */
 	void GetAllDescendantNodes(TSharedPtr<FSequencerDisplayNode> RootNode, TArray<TSharedRef<FSequencerDisplayNode> >& AllNodes);
 
+	/** Return the root node given an object node */
+	TSharedPtr<FSequencerDisplayNode> GetRootNode(TSharedPtr<FSequencerDisplayNode> ObjectNode);
+
+	/** Called when nodes are selected */
 	void OnSelectionChanged( TArray<TSharedPtr<FSequencerDisplayNode> > AffectedNodes );
 
 private:
