Added the concept of 'Tools' to the sequencer

Conceptual changes:
  - SSection is no longer responsible for drag operations. This alleviates mouse capture fighting and event ordering issues to do with both SSequencerTrackArea and SSection wanting to potentially handle drag-like events.
  - ISequencerEditTool can be implemented to define interaction logic for the sequencer as a whole.
  - There is only ever one tool active at a time. This is set by SSequencer::SetEditTool.
    - Edit tools are responsible for handling mouse input based on the current environment
  - Sequencer hotspots define particular points of interest on the sequencer UI such as keys, sections etc. Hotspots can optionally define a drag protocol. This may or may not be honored by the current edit tool.
  - FVirtualConvertor has been formalised as FVirtualTrackArea which provides logic for converting between virtual and physical track area space, as well as hit-testing logic.

We currently have two tools: move (V) and marquee (M) selection. Marquee selection has two modes (key and section selection) which can be toggled with the toolbar button, or by pressing the hotkey multiple times.

diff --git a/Engine/Source/Editor/EditorStyle/Private/SlateEditorStyle.cpp b/Engine/Source/Editor/EditorStyle/Private/SlateEditorStyle.cpp
--- a/Engine/Source/Editor/EditorStyle/Private/SlateEditorStyle.cpp
+++ b/Engine/Source/Editor/EditorStyle/Private/SlateEditorStyle.cpp
@@ -1610,6 +1610,9 @@ void FSlateEditorStyle::FStyle::SetupGeneralStyles()
 		Set( "Sequencer.ToggleShowCurveEditor.Small", new IMAGE_BRUSH("Icons/icon_MatineeCurveView_40px", Icon20x20) );
 		Set( "Sequencer.ToggleAutoScroll", new IMAGE_BRUSH( "Icons/icon_Sequencer_ToggleAutoScroll_40x", Icon40x40 ) );
 		Set( "Sequencer.ToggleAutoScroll.Small", new IMAGE_BRUSH( "Icons/icon_Sequencer_ToggleAutoScroll_16x", Icon16x16 ) );
+		Set( "Sequencer.MoveTool.Small", new IMAGE_BRUSH( "Sequencer/MoveTool_16px", Icon16x16 ) );
+		Set( "Sequencer.MarqueeTool_Keys.Small", new IMAGE_BRUSH( "Sequencer/MarqueeTool_Keys_16px", Icon16x16 ) );
+		Set( "Sequencer.MarqueeTool_Sections.Small", new IMAGE_BRUSH( "Sequencer/MarqueeTool_Sections_16px", Icon16x16 ) );
 		Set( "Sequencer.LockCamera", new IMAGE_BRUSH( "Icons/icon_ToggleActorPilotCameraView_16x", Icon16x16 ) );
 		Set( "Sequencer.UnlockCamera", new IMAGE_BRUSH( "Icons/icon_ToggleActorPilotCameraView_16x", Icon16x16, FLinearColor(1.f, 1.f, 1.f, 0.5f) ) );
 		Set( "Sequencer.Section.SelectionBorder", new BORDER_BRUSH( TEXT("Sequencer/SectionHighlight"), FMargin(7.f/16.f) ) );
diff --git a/Engine/Source/Editor/Sequencer/Private/SSection.cpp b/Engine/Source/Editor/Sequencer/Private/SSection.cpp
--- a/Engine/Source/Editor/Sequencer/Private/SSection.cpp
+++ b/Engine/Source/Editor/Sequencer/Private/SSection.cpp
@@ -6,21 +6,21 @@
 #include "IKeyArea.h"
 #include "ISequencerSection.h"
 #include "MovieSceneSection.h"
-#include "SectionDragOperations.h"
 #include "MovieSceneShotSection.h"
 #include "CommonMovieSceneTools.h"
+#include "SequencerHotspots.h"
 
 FThreadSafeCounter SSection::LayoutRegenerationLock;
 
 void SSection::Construct( const FArguments& InArgs, TSharedRef<FTrackNode> SectionNode, int32 InSectionIndex )
 {
-	ResetState();
-
 	SectionIndex = InSectionIndex;
 	ParentSectionArea = SectionNode;
 	SectionInterface = SectionNode->GetSections()[InSectionIndex];
 	Layout = FKeyAreaLayout(*SectionNode, InSectionIndex);
 
+	ResetState();
+
 	ChildSlot
 	[
 		SectionInterface->GenerateSectionWidget()
@@ -123,6 +123,8 @@ void SSection::CheckForEdgeInteraction( const FPointerEvent& MouseEvent, const F
 		{
 			bLeftEdgeHovered = true;
 		}
+
+		GetSequencer().GetEditTool().SetHotspot(MakeShareable( new FSectionResizeHotspot(FSectionResizeHotspot::Left, FSectionHandle(ParentSectionArea, SectionIndex))) );
 	}
 	else if( SectionRectRight.IsUnderLocation( MouseEvent.GetScreenSpacePosition() ) )
 	{
@@ -134,6 +136,8 @@ void SSection::CheckForEdgeInteraction( const FPointerEvent& MouseEvent, const F
 		{
 			bRightEdgeHovered = true;
 		}
+
+		GetSequencer().GetEditTool().SetHotspot(MakeShareable( new FSectionResizeHotspot(FSectionResizeHotspot::Right, FSectionHandle(ParentSectionArea, SectionIndex))) );
 	}
 }
 
@@ -142,38 +146,6 @@ FSequencer& SSection::GetSequencer() const
 	return ParentSectionArea->GetSequencer();
 }
 
-void SSection::CreateDragOperation( const FGeometry& MyGeometry, const FPointerEvent& MouseEvent, bool bKeysUnderMouse )
-{	
-	check( !DragOperation.IsValid() );
-
-	if( bKeysUnderMouse )
-	{
-		// Disable layout regeneration of dynamic key areas
-		LayoutRegenerationLock.Increment();
-		DragOperation = MakeShareable( new FMoveKeys( GetSequencer(), GetSequencer().GetSelection().GetSelectedKeys(), PressedKey ) );
-	}
-	else
-	{
-		UMovieSceneSection* SectionObject = SectionInterface->GetSectionObject();
-
-		if( bLeftEdgePressed || bLeftEdgeHovered )
-		{
-			// Selected the start of a section
-			DragOperation = MakeShareable( new FResizeSection( GetSequencer(), GetSequencer().GetSelection().GetSelectedSections(), false ) );
-		}
-		else if( bRightEdgePressed || bRightEdgeHovered )
-		{
-			// Selected the end of a section
-			DragOperation = MakeShareable( new FResizeSection( GetSequencer(), GetSequencer().GetSelection().GetSelectedSections(), true ) );
-		}
-		else
-		{
-			// Entire selection moved
-			DragOperation = MakeShareable( new FMoveSection( GetSequencer(), GetSequencer().GetSelection().GetSelectedSections() ) );
-		}
-	}
-}
-
 int32 SSection::OnPaint( const FPaintArgs& Args, const FGeometry& AllottedGeometry, const FSlateRect& MyClippingRect, FSlateWindowElementList& OutDrawElements, int32 LayerId, const FWidgetStyle& InWidgetStyle, bool bParentEnabled ) const
 {
 	int32 StartLayer = SCompoundWidget::OnPaint( Args, AllottedGeometry, MyClippingRect, OutDrawElements, LayerId, InWidgetStyle, bParentEnabled );
@@ -452,12 +424,6 @@ void SSection::Tick( const FGeometry& AllottedGeometry, const double InCurrentTi
 
 FReply SSection::OnMouseButtonDown( const FGeometry& MyGeometry, const FPointerEvent& MouseEvent )
 {
-	DistanceDragged = 0;
-
-	DragOperation.Reset();
-
-	bDragging = false;
-
 	FSequencer& Sequencer = GetSequencer();
 
 	if( MouseEvent.GetEffectingButton() == EKeys::LeftMouseButton || MouseEvent.GetEffectingButton() == EKeys::RightMouseButton )
@@ -472,8 +438,6 @@ FReply SSection::OnMouseButtonDown( const FGeometry& MyGeometry, const FPointerE
 		{
 			CheckForEdgeInteraction( MouseEvent, MyGeometry );
 		}
-
-		return FReply::Handled().CaptureMouse( AsShared() );
 	}
 	else if (MouseEvent.GetEffectingButton() == EKeys::RightMouseButton && (PressedKey.IsValid() || Sequencer.GetSelection().GetSelectedKeys().Num() || Sequencer.GetSelection().GetSelectedSections().Num()))
 	{
@@ -486,9 +450,6 @@ FReply SSection::OnMouseButtonDown( const FGeometry& MyGeometry, const FPointerE
 
 void SSection::ResetState()
 {
-	DistanceDragged = 0;
-	bDragging = false;
-	DragOperation.Reset();
 	ResetHoveredState();
 	PressedKey = FSelectedKey();
 	LayoutRegenerationLock.Reset();
@@ -501,56 +462,50 @@ void SSection::ResetHoveredState()
 	bLeftEdgePressed = false;
 	bRightEdgePressed = false;
 	HoveredKey = FSelectedKey();
+
+	GetSequencer().GetEditTool().SetHotspot(nullptr);
 }
 
 FReply SSection::OnMouseButtonUp( const FGeometry& MyGeometry, const FPointerEvent& MouseEvent )
 {
-	if( bDragging && DragOperation.IsValid() )
-	{	
-		// If dragging tell the operation we are no longer dragging
-		DragOperation->OnEndDrag(ParentSectionArea);
-	}
-	else
+	if( ( MouseEvent.GetEffectingButton() == EKeys::LeftMouseButton || MouseEvent.GetEffectingButton() == EKeys::RightMouseButton ) && MyGeometry.IsUnderLocation( MouseEvent.GetScreenSpacePosition() ) )
 	{
-		if( ( MouseEvent.GetEffectingButton() == EKeys::LeftMouseButton || MouseEvent.GetEffectingButton() == EKeys::RightMouseButton ) && MyGeometry.IsUnderLocation( MouseEvent.GetScreenSpacePosition() ) )
+		// Snap time to the key under the mouse if shift is down
+		if (MouseEvent.IsShiftDown() && MouseEvent.GetEffectingButton() == EKeys::LeftMouseButton)
 		{
-			// Snap time to the key under the mouse if shift is down
-			if (MouseEvent.IsShiftDown() && MouseEvent.GetEffectingButton() == EKeys::LeftMouseButton)
+			FSelectedKey Key = GetKeyUnderMouse( MouseEvent.GetScreenSpacePosition(), MyGeometry );
+			if (Key.IsValid())
 			{
-				FSelectedKey Key = GetKeyUnderMouse( MouseEvent.GetScreenSpacePosition(), MyGeometry );
-				if (Key.IsValid())
-				{
-					float KeyTime = Key.KeyArea->GetKeyTime(Key.KeyHandle.GetValue());
-					GetSequencer().SetGlobalTime(KeyTime);
-				}
+				float KeyTime = Key.KeyArea->GetKeyTime(Key.KeyHandle.GetValue());
+				GetSequencer().SetGlobalTime(KeyTime);
 			}
-
-			HandleSelection( MyGeometry, MouseEvent );
 		}
 
-		if( MouseEvent.GetEffectingButton() == EKeys::RightMouseButton )
+		HandleSelection( MyGeometry, MouseEvent );
+	}
+
+	if( MouseEvent.GetEffectingButton() == EKeys::RightMouseButton )
+	{
+		TSharedPtr<SWidget> MenuContent = OnSummonContextMenu( MyGeometry, MouseEvent );
+		if (MenuContent.IsValid())
 		{
-			TSharedPtr<SWidget> MenuContent = OnSummonContextMenu( MyGeometry, MouseEvent );
-			if (MenuContent.IsValid())
-			{
-				FWidgetPath WidgetPath = MouseEvent.GetEventPath() != nullptr ? *MouseEvent.GetEventPath() : FWidgetPath();
-
-				FSlateApplication::Get().PushMenu(
-					AsShared(),
-					WidgetPath,
-					MenuContent.ToSharedRef(),
-					MouseEvent.GetScreenSpacePosition(),
-					FPopupTransitionEffect( FPopupTransitionEffect::ContextMenu )
-					);
+			FWidgetPath WidgetPath = MouseEvent.GetEventPath() != nullptr ? *MouseEvent.GetEventPath() : FWidgetPath();
 
-				return FReply::Handled().SetUserFocus(MenuContent.ToSharedRef(), EFocusCause::SetDirectly);
-			}
+			FSlateApplication::Get().PushMenu(
+				AsShared(),
+				WidgetPath,
+				MenuContent.ToSharedRef(),
+				MouseEvent.GetScreenSpacePosition(),
+				FPopupTransitionEffect( FPopupTransitionEffect::ContextMenu )
+				);
+
+			return FReply::Handled().SetUserFocus(MenuContent.ToSharedRef(), EFocusCause::SetDirectly);
 		}
 	}
 
 	ResetState();
 
-	return FReply::Handled().ReleaseMouseCapture();
+	return FReply::Handled();
 }
 
 FReply SSection::OnMouseButtonDoubleClick( const FGeometry& MyGeometry, const FPointerEvent& MouseEvent )
@@ -578,83 +533,24 @@ FReply SSection::OnMouseButtonDoubleClick( const FGeometry& MyGeometry, const FP
 
 FReply SSection::OnMouseMove( const FGeometry& MyGeometry, const FPointerEvent& MouseEvent )
 {
-	if( HasMouseCapture() )
-	{
-		// Have mouse capture and there are drag operations that need to be performed
-		if( MouseEvent.IsMouseButtonDown( EKeys::LeftMouseButton ) )
-		{
-			DistanceDragged += FMath::Abs( MouseEvent.GetCursorDelta().X );
-			
-			FVector2D LocalMousePos = MyGeometry.AbsoluteToLocal( MouseEvent.GetScreenSpacePosition() );
-
-			FTimeToPixel TimeToPixelConverter = SectionInterface->GetSectionObject()->IsInfinite() ? 			
-				FTimeToPixel( ParentGeometry, GetSequencer().GetViewRange()) : 
-				FTimeToPixel( MyGeometry, TRange<float>( SectionInterface->GetSectionObject()->GetStartTime(), SectionInterface->GetSectionObject()->GetEndTime() ) );
-
-			if( !bDragging )
-			{
-				// If we are not dragging determine if the mouse has moved far enough to start a drag
-				if( DistanceDragged >= SequencerSectionConstants::SectionDragStartDistance )
-				{
-					bDragging = true;
-
-					if( PressedKey.IsValid() )
-					{
-						// Clear selected sections when beginning to drag keys
-						GetSequencer().GetSelection().EmptySelectedSections();
-
-						bool bSelectDueToDrag = true;
-						HandleKeySelection( PressedKey, MouseEvent, bSelectDueToDrag );
-
-						bool bKeysUnderMouse = true;
-						CreateDragOperation( MyGeometry, MouseEvent, bKeysUnderMouse );
-					}
-					else
-					{
-						// Clear selected keys when beginning to drag a section
-						GetSequencer().GetSelection().EmptySelectedKeys();
-
-						bool bSelectDueToDrag = true;
-						HandleSectionSelection( MouseEvent, bSelectDueToDrag );
+	ResetHoveredState();
 
-						bool bKeysUnderMouse = false;
-						CreateDragOperation( MyGeometry, MouseEvent, bKeysUnderMouse );
-					}
-				
-					if( DragOperation.IsValid() )
-					{
-						DragOperation->OnBeginDrag(LocalMousePos, TimeToPixelConverter, ParentSectionArea);
-					}
+	// Checked for hovered key
+	// @todo Sequencer - Needs visual cue
+	HoveredKey = GetKeyUnderMouse( MouseEvent.GetScreenSpacePosition(), MyGeometry );
 
-				}
-			}
-			else if( DragOperation.IsValid() )
-			{
-				// Already in a drag, tell all operations to perform their drag implementations
+	auto& EditTool = GetSequencer().GetEditTool();
 
-				DragOperation->OnDrag( MouseEvent, LocalMousePos, TimeToPixelConverter, ParentSectionArea );
-			}
-		}
+	EditTool.SetHotspot( MakeShareable( new FSectionHotspot(FSectionHandle(ParentSectionArea, SectionIndex))) );
 
-		return FReply::Handled();
+	if ( HoveredKey.IsValid() )
+	{
+		EditTool.SetHotspot( MakeShareable( new FKeyHotspot(HoveredKey, ParentSectionArea) ) );
 	}
 	else
 	{
-		// Not dragging
-
-
-		ResetHoveredState();
-
-		// Checked for hovered key
-		// @todo Sequencer - Needs visual cue
-		HoveredKey = GetKeyUnderMouse( MouseEvent.GetScreenSpacePosition(), MyGeometry );
-
 		// Only check for edge interaction if not hovering over a key
-		if( !HoveredKey.IsValid() )
-		{
-			CheckForEdgeInteraction( MouseEvent, MyGeometry );
-		}
-		
+		CheckForEdgeInteraction( MouseEvent, MyGeometry );
 	}
 
 	return FReply::Unhandled();
@@ -663,20 +559,12 @@ FReply SSection::OnMouseMove( const FGeometry& MyGeometry, const FPointerEvent&
 void SSection::OnMouseLeave( const FPointerEvent& MouseEvent )
 {
 	SCompoundWidget::OnMouseLeave( MouseEvent );
-
-	if( !HasMouseCapture() )
-	{
-		ResetHoveredState();
-	}
+	ResetHoveredState();
 }
 
 FCursorReply SSection::OnCursorQuery( const FGeometry& MyGeometry, const FPointerEvent& CursorEvent ) const
 {
-	if( DragOperation.IsValid() )
-	{
-		return DragOperation->GetCursor();
-	}
-	else if( bLeftEdgeHovered || bRightEdgeHovered )
+	if( bLeftEdgeHovered || bRightEdgeHovered )
 	{
 		return FCursorReply::Cursor( EMouseCursor::ResizeLeftRight );
 	}
diff --git a/Engine/Source/Editor/Sequencer/Private/SSection.h b/Engine/Source/Editor/Sequencer/Private/SSection.h
--- a/Engine/Source/Editor/Sequencer/Private/SSection.h
+++ b/Engine/Source/Editor/Sequencer/Private/SSection.h
@@ -140,18 +140,12 @@ private:
 	TSharedPtr<FTrackNode> ParentSectionArea;
 	/** Cached layout generated each tick */
 	TOptional<FKeyAreaLayout> Layout;
-	/** Current drag operation if any */
-	TSharedPtr<class FSequencerDragOperation> DragOperation;
 	/** The currently pressed key if any */
 	FSelectedKey PressedKey;
 	/** The current hovered key if any */
 	FSelectedKey HoveredKey;
-	/** The distance we have dragged since the mouse was pressed down */
-	float DistanceDragged;
 	/** The index of this section in the parent section area */
 	int32 SectionIndex;
-	/** Whether or not we are dragging */
-	bool bDragging;
 	/** Whether or not the left edge of the section is hovered */
 	bool bLeftEdgeHovered;
 	/** Whether or not the left edge of the section is pressed */
diff --git a/Engine/Source/Editor/Sequencer/Private/SSequencer.cpp b/Engine/Source/Editor/Sequencer/Private/SSequencer.cpp
--- a/Engine/Source/Editor/Sequencer/Private/SSequencer.cpp
+++ b/Engine/Source/Editor/Sequencer/Private/SSequencer.cpp
@@ -19,6 +19,8 @@
 #include "DragAndDrop/AssetDragDropOp.h"
 #include "DragAndDrop/ActorDragDropGraphEdOp.h"
 #include "DragAndDrop/ClassDragDropOp.h"
+#include "Tools/SequencerEditTool_Selection.h"
+#include "Tools/SequencerEditTool_Movement.h"
 #include "AssetSelection.h"
 #include "MovieSceneShotSection.h"
 #include "CommonMovieSceneTools.h"
@@ -29,7 +31,8 @@
 #include "MovieSceneSequence.h"
 #include "SSequencerSplitterOverlay.h"
 #include "IKeyArea.h"
-
+#include "VirtualTrackArea.h"
+#include "SequencerHotspots.h"
 
 #define LOCTEXT_NAMESPACE "Sequencer"
 
@@ -230,7 +233,7 @@ void SSequencer::Construct( const FArguments& InArgs, TSharedRef< class FSequenc
 
 	SAssignNew( TrackOutliner, SSequencerTrackOutliner );
 
-	SAssignNew( TrackArea, SSequencerTrackArea, TimeSliderController, PinnedSequencer )
+	SAssignNew( TrackArea, SSequencerTrackArea, TimeSliderController, SharedThis(this) )
 		.Visibility( this, &SSequencer::GetTrackAreaVisibility );
 	
 	SAssignNew( TreeView, SSequencerTreeView, SequencerNodeTree.ToSharedRef(), TrackArea.ToSharedRef() )
@@ -244,6 +247,8 @@ void SSequencer::Construct( const FArguments& InArgs, TSharedRef< class FSequenc
 
 	TrackArea->SetTreeView(TreeView);
 
+	EditTool.Reset(new FSequencerEditTool_Movement(PinnedSequencer, SharedThis(this)));
+
 	const int32 				Column0	= 0,	Column1	= 1;
 	const int32 Row0	= 0,
 				Row1	= 1,
@@ -442,6 +447,46 @@ void SSequencer::Construct( const FArguments& InArgs, TSharedRef< class FSequenc
 	ResetBreadcrumbs();
 }
 
+void SSequencer::BindCommands(TSharedRef<FUICommandList> SequencerCommandBindings)
+{
+	SequencerCommandBindings->MapAction(
+		FSequencerCommands::Get().MoveTool,
+		FExecuteAction::CreateLambda( [&]{
+			EditTool.Reset( new FSequencerEditTool_Movement(Sequencer.Pin(), SharedThis(this)) );
+		} ),
+		FCanExecuteAction::CreateLambda( []{ return true; } ),
+		FIsActionChecked::CreateSP( this, &SSequencer::IsEditToolEnabled, FName("Movement") )
+		);
+
+	SequencerCommandBindings->MapAction(
+		FSequencerCommands::Get().MarqueeTool,
+		FExecuteAction::CreateLambda( [&]{
+
+			if (IsEditToolEnabled("Selection"))
+			{
+				// If the selection tool is already enabled, we toggle the selection mode
+				auto* SelectionTool = static_cast<FSequencerEditTool_Selection*>(EditTool.Get());
+
+				ESequencerSelectionMode CurrentMode = SelectionTool->GetSelectionMode();
+				SelectionTool->SetSelectionMode(CurrentMode == ESequencerSelectionMode::Keys ? ESequencerSelectionMode::Sections : ESequencerSelectionMode::Keys);
+			}
+			else
+			{
+				// Otherwise, switch the tool
+				EditTool.Reset( new FSequencerEditTool_Selection(Sequencer.Pin(), SharedThis(this)) );
+			}
+
+		} ),
+		FCanExecuteAction::CreateLambda( []{ return true; } ),
+		FIsActionChecked::CreateSP( this, &SSequencer::IsEditToolEnabled, FName("Selection") )
+		);
+}
+
+bool SSequencer::IsEditToolEnabled(FName InIdentifier)
+{
+	return GetEditTool().GetIdentifier() == InIdentifier;
+}
+
 TSharedRef<SWidget> SSequencer::MakeToolBar()
 {
 	FToolBarBuilder ToolBarBuilder( Sequencer.Pin()->GetCommandBindings(), FMultiBoxCustomization::None, TSharedPtr<FExtender>(), Orient_Horizontal, true);
@@ -529,6 +574,51 @@ TSharedRef<SWidget> SSequencer::MakeToolBar()
 	}
 	ToolBarBuilder.EndSection();
 
+	ToolBarBuilder.BeginSection("Tools");
+	{
+		ToolBarBuilder.SetLabelVisibility( EVisibility::Collapsed );
+
+		ToolBarBuilder.AddToolBarButton( FSequencerCommands::Get().MoveTool );
+
+		// The marquee selection button is a little more complicated as toggling it when enabled causes it to toggle between key/section selection
+		TAttribute<FSlateIcon> MarqueeIcon;
+		MarqueeIcon.Bind(TAttribute<FSlateIcon>::FGetter::CreateLambda([]{
+			static FSlateIcon MarqueeTool_Keys(FEditorStyle::GetStyleSetName(), "Sequencer.MarqueeTool_Keys");
+			static FSlateIcon MarqueeTool_Sections(FEditorStyle::GetStyleSetName(), "Sequencer.MarqueeTool_Sections");
+
+			return FSequencerEditTool_Selection::GetGlobalSelectionMode() == ESequencerSelectionMode::Keys ? MarqueeTool_Keys : MarqueeTool_Sections;
+		}));
+
+		TAttribute<FText> MarqueeTooltip;
+		MarqueeTooltip.Bind(TAttribute<FText>::FGetter::CreateLambda([&]{
+			if (FSequencerEditTool_Selection::GetGlobalSelectionMode() == ESequencerSelectionMode::Keys)
+			{
+				if (IsEditToolEnabled("Selection"))
+				{
+					return LOCTEXT("MarqueeTooltip_Enabled_Keys", "Change marquee selection to operate on keys rather than sections");
+				}
+				else
+				{
+					return LOCTEXT("MarqueeTooltip_Key", "Activate marquee key selection tool");
+				}
+			}
+			else
+			{
+				if (IsEditToolEnabled("Selection"))
+				{
+					return LOCTEXT("MarqueeTooltip_Enabled_Sections", "Change marquee selection to operate on sections rather than keys");
+				}
+				else
+				{
+					return LOCTEXT("MarqueeTooltip_Section", "Activate marquee section selection tool");
+				}
+			}
+		}));
+
+		ToolBarBuilder.AddToolBarButton( FSequencerCommands::Get().MarqueeTool, NAME_None, TAttribute<FText>(), MarqueeTooltip, MarqueeIcon );
+	}
+	ToolBarBuilder.EndSection();
+
 	ToolBarBuilder.BeginSection("Snapping");
 	{
 		ToolBarBuilder.AddToolBarButton( FSequencerCommands::Get().ToggleIsSnapEnabled, NAME_None, TAttribute<FText>( FText::GetEmpty() ) );
@@ -675,7 +765,7 @@ EActiveTimerReturnType SSequencer::EnsureSlateTickDuringPlayback(double InCurren
 	if (Sequencer.IsValid())
 	{
 		auto PlaybackStatus = Sequencer.Pin()->GetPlaybackStatus();
-		if (PlaybackStatus == EMovieScenePlayerStatus::Playing || PlaybackStatus == EMovieScenePlayerStatus::Recording)
+		if (PlaybackStatus == EMovieScenePlayerStatus::Playing || PlaybackStatus == EMovieScenePlayerStatus::Recording || PlaybackStatus == EMovieScenePlayerStatus::Scrubbing)
 		{
 			return EActiveTimerReturnType::Continue;
 		}
@@ -1158,6 +1248,33 @@ void SSequencer::DeleteSelectedNodes()
 	}
 }
 
+TArray<FSectionHandle> SSequencer::GetSectionHandles(const TSet<TWeakObjectPtr<UMovieSceneSection>>& DesiredSections) const
+{
+	TArray<FSectionHandle> SectionHandles;
+
+	for (auto& Node : SequencerNodeTree->GetRootNodes())
+	{
+		Node->Traverse_ParentFirst([&](FSequencerDisplayNode& InNode){
+			if (InNode.GetType() == ESequencerNode::Track)
+			{
+				FTrackNode& TrackNode = static_cast<FTrackNode&>(InNode);
+
+				const auto& AllSections = TrackNode.GetTrack()->GetAllSections();
+				for (int32 Index = 0; Index < AllSections.Num(); ++Index)
+				{
+					if (DesiredSections.Contains(TWeakObjectPtr<UMovieSceneSection>(AllSections[Index])))
+					{
+						SectionHandles.Emplace(StaticCastSharedRef<FTrackNode>(TrackNode.AsShared()), Index);
+					}
+				}
+			}
+			return true;
+		});
+	}
+
+	return SectionHandles;
+}
+
 void SSequencer::OnSaveMovieSceneClicked()
 {
 	Sequencer.Pin()->SaveCurrentMovieScene();
@@ -1347,4 +1464,9 @@ void SSequencer::OnCurveEditorVisibilityChanged()
 	}
 }
 
+FVirtualTrackArea SSequencer::GetVirtualTrackArea(const FGeometry& InTrackAreaGeometry) const
+{
+	return FVirtualTrackArea(*Sequencer.Pin(), *TreeView.Get(), InTrackAreaGeometry);
+}
+
 #undef LOCTEXT_NAMESPACE
diff --git a/Engine/Source/Editor/Sequencer/Private/SSequencer.h b/Engine/Source/Editor/Sequencer/Private/SSequencer.h
--- a/Engine/Source/Editor/Sequencer/Private/SSequencer.h
+++ b/Engine/Source/Editor/Sequencer/Private/SSequencer.h
@@ -2,7 +2,12 @@
 
 #pragma once
 
+#include "UniquePtr.h"
+#include "ISequencerEditTool.h"
+#include "SequencerCommonHelpers.h"
+
 class SSequencerTreeView;
+struct FSectionHandle;
 
 namespace SequencerLayoutConstants
 {
@@ -75,6 +80,8 @@ public:
 
 	void Construct( const FArguments& InArgs, TSharedRef< class FSequencer > InSequencer );
 
+	void BindCommands(TSharedRef<FUICommandList> SequencerCommandBindings);
+	
 	~SSequencer();
 	
 	virtual void AddReferencedObjects( FReferenceCollector& Collector )
@@ -121,7 +128,20 @@ public:
 	/** Access the tree view for this sequencer */
 	TSharedPtr<SSequencerTreeView> GetTreeView() const;
 
+	/** Access the currently active edit tool */
+	ISequencerEditTool& GetEditTool() const { return *EditTool; }
+
+	/** Generate a helper structure that can be used to transform between phsyical space and virtual space in the track area */
+	FVirtualTrackArea GetVirtualTrackArea(const FGeometry& InTrackAreaGeometry) const;
+
+	/** Get an array of section handles for the given set of movie scene sections */
+	TArray<FSectionHandle> GetSectionHandles(const TSet<TWeakObjectPtr<UMovieSceneSection>>& DesiredSections) const;
+
 private:
+	
+	/** Check whether the specified edit tool is enabled */
+	bool IsEditToolEnabled(FName InIdentifier);
+
 	/** Empty active timer to ensure Slate ticks during Sequencer playback */
 	EActiveTimerReturnType EnsureSlateTickDuringPlayback(double InCurrentTime, float InDeltaTime);	
 
@@ -279,6 +299,9 @@ private:
 	bool bIsActiveTimerRegistered;
 	/** Whether the user is selecting. Ignore selection changes from the level when the user is selecting. */
 	bool bUserIsSelecting;
+	/** The current edit tool */
+	TUniquePtr<ISequencerEditTool> EditTool;
 
 	FOnGetAddMenuContent OnGetAddMenuContent;
+
 };
diff --git a/Engine/Source/Editor/Sequencer/Private/SSequencerSectionAreaView.h b/Engine/Source/Editor/Sequencer/Private/SSequencerSectionAreaView.h
--- a/Engine/Source/Editor/Sequencer/Private/SSequencerSectionAreaView.h
+++ b/Engine/Source/Editor/Sequencer/Private/SSequencerSectionAreaView.h
@@ -52,20 +52,12 @@ private:
 	/** @return The sequencer interface */
 	FSequencer& GetSequencer() const;
 private:
-
-	/** List of all sequencer drag operations going on  in this area */
-	/** @todo Sequencer - This list may need to be on the section view itself for multi-select */
-	TArray< TSharedPtr< FSequencerDragOperation > > DragOperations;
 	/** The node containing the sections we are viewing/manipulating */
 	TSharedPtr<FTrackNode> SectionAreaNode;
 	/** The current view range */
 	TAttribute< TRange<float> > ViewRange;
 	/** Background brush of the section area */
 	const FSlateBrush* BackgroundBrush;
-	/** The distance we have dragged since the mouse was pressed down */
-	float DistanceDragged;
-	/** Whether or not we are dragging */
-	bool bDragging;
 	/** All the widgets in the panel */
 	TSlotlessChildren<SSection> Children;
 };
diff --git a/Engine/Source/Editor/Sequencer/Private/SSequencerTrackArea.cpp b/Engine/Source/Editor/Sequencer/Private/SSequencerTrackArea.cpp
--- a/Engine/Source/Editor/Sequencer/Private/SSequencerTrackArea.cpp
+++ b/Engine/Source/Editor/Sequencer/Private/SSequencerTrackArea.cpp
@@ -11,228 +11,9 @@
 #include "ISequencerSection.h"
 #include "SSection.h"
 
-
-struct FMarqueeSelectData
-{
-	FMarqueeSelectData(const FVector2D& InInitialPosition)
-		: InitialPosition(InInitialPosition)
-		, CurrentPosition(InInitialPosition)
-	{}
-
-	FVector2D TopLeft() const
-	{
-		return FVector2D(
-			FMath::Min(InitialPosition.X, CurrentPosition.X),
-			FMath::Min(InitialPosition.Y, CurrentPosition.Y)
-		);
-	}
-	
-	FVector2D BottomRight() const
-	{
-		return FVector2D(
-			FMath::Max(InitialPosition.X, CurrentPosition.X),
-			FMath::Max(InitialPosition.Y, CurrentPosition.Y)
-		);
-	}
-
-	FVector2D Size() const
-	{
-		return BottomRight() - TopLeft();
-	}
-
-	bool HasDragged() const
-	{
-		return !Size().IsNearlyZero();
-	}
-
-	FVector2D InitialPosition;
-	FVector2D CurrentPosition;
-};
-
-
-struct IMarqueeVisitor
-{
-	virtual void VisitKey(FKeyHandle KeyHandle, float KeyTime, const TSharedPtr<IKeyArea>& KeyArea, UMovieSceneSection* Section) const { }
-	virtual void VisitSection(FKeyHandle KeyHandle, float KeyTime, const TSharedPtr<IKeyArea>& KeyArea, UMovieSceneSection* Section) const { }
-};
-
-/** Struct used to iterate a marquee range with a user-supplied visitor */
-/* @todo: This could probably live in a header so it can be used elsewhere */
-struct FMarqueeRange
-{
-	/** Construction from the range itself */
-	FMarqueeRange(FVector2D InTopLeft, FVector2D InBottomRight, FVector2D InVirtualKeySize)
-		: TopLeft(InTopLeft), BottomRight(InBottomRight), VirtualKeySize(InVirtualKeySize)
-	{}
-
-	/** Visit the specified nodes (recursively) with this range and a user-supplied visitor */
-	void Visit(const IMarqueeVisitor& Visitor, const TArray< TSharedRef<FSequencerDisplayNode> >& Nodes);
-
-private:
-
-	/** Handle visitation of a particular node */
-	void HandleNode(const IMarqueeVisitor& Visitor, FSequencerDisplayNode& InNode);
-	/** Handle visitation of a particular node, along with a set of sections */
-	void HandleNode(const IMarqueeVisitor& Visitor, FSequencerDisplayNode& InNode, const TArray<TSharedRef<ISequencerSection>>& InSections);
-	/** Handle visitation of a key area node */
-	void HandleKeyAreaNode(const IMarqueeVisitor& Visitor, FSectionKeyAreaNode& InNode, const TArray<TSharedRef<ISequencerSection>>& InSections);
-	/** Handle visitation of a key area */
-	void HandleKeyArea(const IMarqueeVisitor& Visitor, const TSharedPtr<IKeyArea>& KeyArea, UMovieSceneSection* Section);
-	
-	/** The bounds of the range */
-	FVector2D TopLeft, BottomRight;
-
-	/** Key size in virtual space */
-	FVector2D VirtualKeySize;
-};
-
-void FMarqueeRange::Visit(const IMarqueeVisitor& Visitor, const TArray< TSharedRef<FSequencerDisplayNode> >& Nodes)
-{
-	for (auto& Child : Nodes)
-	{
-		if (!Child->IsHidden())
-		{
-			HandleNode(Visitor, *Child);
-		}
-	}
-}
-
-void FMarqueeRange::HandleNode(const IMarqueeVisitor& Visitor, FSequencerDisplayNode& InNode)
-{
-	if (InNode.GetType() == ESequencerNode::Track)
-	{
-		HandleNode(Visitor, InNode, static_cast<FTrackNode&>(InNode).GetSections());
-	}
-
-	if (InNode.IsExpanded())
-	{
-		for (auto& Child : InNode.GetChildNodes())
-		{
-			if (!Child->IsHidden())
-			{
-				HandleNode(Visitor, *Child);
-			}
-		}
-	}
-}
-
-void FMarqueeRange::HandleNode(const IMarqueeVisitor& Visitor, FSequencerDisplayNode& InNode, const TArray<TSharedRef<ISequencerSection>>& InSections)
-{
-	const float NodeCenter = InNode.GetVirtualTop() + (InNode.GetVirtualBottom() - InNode.GetVirtualTop())/2;
-	if (NodeCenter + VirtualKeySize.Y/2 > TopLeft.Y &&
-		NodeCenter - VirtualKeySize.Y/2 < BottomRight.Y)
-	{
-		bool bNodeHasKeyArea = false;
-		if (InNode.GetType() == ESequencerNode::KeyArea)
-		{
-			HandleKeyAreaNode(Visitor, static_cast<FSectionKeyAreaNode&>(InNode), InSections);
-			bNodeHasKeyArea = true;
-		}
-		else if (InNode.GetType() == ESequencerNode::Track)
-		{
-			TSharedPtr<FSectionKeyAreaNode> SectionKeyNode = static_cast<FTrackNode&>(InNode).GetTopLevelKeyNode();
-			if (SectionKeyNode.IsValid())
-			{
-				HandleKeyAreaNode(Visitor, static_cast<FSectionKeyAreaNode&>(InNode), InSections);
-				bNodeHasKeyArea = true;
-			}
-		}
-
-		if (!InNode.IsExpanded() && !bNodeHasKeyArea)
-		{
-			// As a fallback, handle key groupings on collapsed parents
-			for (int32 SectionIndex = 0; SectionIndex < InSections.Num(); ++SectionIndex)
-			{
-				TSharedRef<IKeyArea> KeyArea = InNode.UpdateKeyGrouping(SectionIndex);
-				HandleKeyArea(Visitor, KeyArea, InSections[SectionIndex]->GetSectionObject());
-			}
-		}
-	}
-
-	if (InNode.IsExpanded())
-	{
-		// Handle Children
-		for (auto& Child : InNode.GetChildNodes())
-		{
-			if (!Child->IsHidden())
-			{
-				HandleNode(Visitor, *Child, InSections);
-			}
-		}
-	}
-}
-
-void FMarqueeRange::HandleKeyAreaNode(const IMarqueeVisitor& Visitor, FSectionKeyAreaNode& KeyAreaNode, const TArray<TSharedRef<ISequencerSection>>& InSections)
-{
-	for( int32 SectionIndex = 0; SectionIndex < InSections.Num(); ++SectionIndex )
-	{
-		TSharedRef<IKeyArea> KeyArea = KeyAreaNode.GetKeyArea(SectionIndex);
-		HandleKeyArea(Visitor, KeyArea, InSections[SectionIndex]->GetSectionObject());
-	}
-}
-
-void FMarqueeRange::HandleKeyArea(const IMarqueeVisitor& Visitor, const TSharedPtr<IKeyArea>& KeyArea, UMovieSceneSection* Section)
-{
-	for (FKeyHandle KeyHandle : KeyArea->GetUnsortedKeyHandles())
-	{
-		float KeyPosition = KeyArea->GetKeyTime(KeyHandle);
-		if (KeyPosition + VirtualKeySize.X/2 > TopLeft.X &&
-			KeyPosition - VirtualKeySize.X/2 < BottomRight.X)
-		{
-			Visitor.VisitKey(KeyHandle, KeyPosition, KeyArea, Section);
-		}
-	}
-}
-
-/** Structure used for converting X and Y positional values to Time and vertical offsets respectively */
-struct FVirtualConverter : FTimeToPixel
+void SSequencerTrackArea::Construct( const FArguments& InArgs, TSharedRef<FSequencerTimeSliderController> InTimeSliderController, TSharedRef<SSequencer> InSequencerWidget )
 {
-	FVirtualConverter(TSharedPtr<SSequencerTreeView> InTreeView, TSharedPtr<FSequencer> InSequencer, const FGeometry& InGeometry)
-		: FTimeToPixel(InGeometry, InSequencer->GetViewRange())
-		, TreeView(*InTreeView)
-	{}
-
-	float PixelToVerticalOffset(float InPixel) const
-	{
-		return TreeView.PhysicalToVirtual(InPixel);
-	}
-
-	float VerticalOffsetToPixel(float InOffset) const
-	{
-		return TreeView.VirtualToPhysical(InOffset);
-	}
-
-	FVector2D PhysicalToVirtual(FVector2D InPosition) const
-	{
-		InPosition.Y = PixelToVerticalOffset(InPosition.Y);
-		InPosition.X = PixelToTime(InPosition.X);
-
-		return InPosition;
-	}
-
-	FVector2D VirtualToPhysical(FVector2D InPosition) const
-	{
-		InPosition.Y = VerticalOffsetToPixel(InPosition.Y);
-		InPosition.X = TimeToPixel(InPosition.X);
-
-		return InPosition;
-	}
-
-private:
-	SSequencerTreeView& TreeView;
-};
-
-// User-defined constructor/destructor to ensure TUniquePtr destruction
-SSequencerTrackArea::SSequencerTrackArea()
-{
-}
-SSequencerTrackArea::~SSequencerTrackArea()
-{
-}
-
-void SSequencerTrackArea::Construct( const FArguments& InArgs, TSharedRef<FSequencerTimeSliderController> InTimeSliderController, TSharedRef<FSequencer> InSequencer )
-{
-	Sequencer = InSequencer;
+	SequencerWidget = InSequencerWidget;
 	TimeSliderController = InTimeSliderController;
 
 	SOverlay::Construct(SOverlay::FArguments());
@@ -273,94 +54,58 @@ TSharedPtr<SSequencerTrackLane> SSequencerTrackArea::FindTrackSlot(const TShared
 int32 SSequencerTrackArea::OnPaint(const FPaintArgs& Args, const FGeometry& AllottedGeometry, const FSlateRect& MyClippingRect, FSlateWindowElementList& OutDrawElements, int32 LayerId, const FWidgetStyle& InWidgetStyle, bool bParentEnabled ) const
 {
 	LayerId = SOverlay::OnPaint(Args, AllottedGeometry, MyClippingRect, OutDrawElements, LayerId, InWidgetStyle, bParentEnabled);
-	return PaintMarquee(AllottedGeometry, MyClippingRect, OutDrawElements, LayerId);
-}
 
-int32 SSequencerTrackArea::PaintMarquee(const FGeometry& AllottedGeometry, const FSlateRect& MyClippingRect, FSlateWindowElementList& OutDrawElements, int32 LayerId) const
-{
-	if (MarqueeSelectData.IsValid() && MarqueeSelectData->HasDragged())
+	auto SequencerPin = SequencerWidget.Pin();
+	if (SequencerPin.IsValid())
 	{
-		// Convert to physical space for rendering
-		const FVirtualConverter Converter(TreeView.Pin(), Sequencer.Pin(), AllottedGeometry);
-
-		FVector2D TopLeft = Converter.VirtualToPhysical(MarqueeSelectData->TopLeft());
-		FVector2D BottomRight = Converter.VirtualToPhysical(MarqueeSelectData->BottomRight());
-
-		FSlateDrawElement::MakeBox(
-			OutDrawElements,
-			LayerId,
-			AllottedGeometry.ToPaintGeometry(TopLeft, BottomRight - TopLeft),
-			FEditorStyle::GetBrush(TEXT("MarqueeSelection")),
-			MyClippingRect
-			);
-
-		return LayerId + 1;
+		return SequencerPin->GetEditTool().OnPaint(AllottedGeometry, MyClippingRect, OutDrawElements, LayerId);
 	}
-
 	return LayerId;
 }
 
 FReply SSequencerTrackArea::OnMouseButtonDown( const FGeometry& MyGeometry, const FPointerEvent& MouseEvent )
 {
-	FReply Reply = FReply::Unhandled();
-
-	if (MouseEvent.GetEffectingButton() == EKeys::LeftMouseButton)
+	auto SequencerPin = SequencerWidget.Pin();
+	if (SequencerPin.IsValid())
 	{
-		const FVirtualConverter Converter(TreeView.Pin(), Sequencer.Pin(), MyGeometry);
-
-		// Marquee data is stored in virtual space
-		FVector2D VirtualClickPosition = Converter.PhysicalToVirtual(MyGeometry.AbsoluteToLocal(MouseEvent.GetScreenSpacePosition()));
-		MarqueeSelectData.Reset(new FMarqueeSelectData(VirtualClickPosition));
-
-		// Clear selected sections and keys
-		auto SequencerPin = Sequencer.Pin();
-		if( !MouseEvent.IsControlDown() && !MouseEvent.IsShiftDown() && SequencerPin.IsValid() )
+		FReply Reply = SequencerPin->GetEditTool().OnMouseButtonDown(*this, MyGeometry, MouseEvent);
+		if (Reply.IsEventHandled())
 		{
-			SequencerPin->GetSelection().EmptySelectedSections();
-			SequencerPin->GetSelection().EmptySelectedKeys();
+			return Reply;
 		}
-
-		Reply = FReply::Handled().CaptureMouse(AsShared());
-	}
-	else
-	{
-		Reply = TimeSliderController->OnMouseButtonDown( SharedThis(this), MyGeometry, MouseEvent );
 	}
 
-	return Reply;
+	return TimeSliderController->OnMouseButtonDown( SharedThis(this), MyGeometry, MouseEvent );
 }
 
 FReply SSequencerTrackArea::OnMouseButtonUp( const FGeometry& MyGeometry, const FPointerEvent& MouseEvent )
 {
-	FReply Reply = FReply::Unhandled();
-
-	if (MarqueeSelectData.IsValid())
+	auto SequencerPin = SequencerWidget.Pin();
+	if (SequencerPin.IsValid())
 	{
-		if (MarqueeSelectData->HasDragged())
+		FReply Reply = SequencerPin->GetEditTool().OnMouseButtonUp(*this, MyGeometry, MouseEvent);
+		if (Reply.IsEventHandled())
 		{
-			HandleMarqueeSelection(MyGeometry, MouseEvent);
+			return Reply;
 		}
-		MarqueeSelectData.Reset();
-		Sequencer.Pin()->StopAutoscroll();
-		Reply = FReply::Handled().ReleaseMouseCapture();
 	}
 
-	if (!Reply.IsEventHandled())
-	{
-		Reply = TimeSliderController->OnMouseButtonUp( SharedThis(this), MyGeometry, MouseEvent );
-	}
-
-	return Reply;
+	return TimeSliderController->OnMouseButtonUp( SharedThis(this), MyGeometry, MouseEvent );
 }
 
 FReply SSequencerTrackArea::OnMouseMove( const FGeometry& MyGeometry, const FPointerEvent& MouseEvent )
 {
-	if (MarqueeSelectData.IsValid())
+	auto SequencerPin = SequencerWidget.Pin();
+	if (SequencerPin.IsValid())
 	{
-		SetNewMarqueeBounds(MouseEvent, MyGeometry);
-		return FReply::Handled();
+		FReply Reply = SequencerPin->GetEditTool().OnMouseMove(*this, MyGeometry, MouseEvent);
+		if (Reply.IsEventHandled())
+		{
+			return Reply;
+		}
 	}
-	else if (MouseEvent.IsMouseButtonDown(EKeys::RightMouseButton) && HasMouseCapture())
+
+	if (MouseEvent.IsMouseButtonDown(EKeys::RightMouseButton) && HasMouseCapture())
 	{
 		TreeView.Pin()->ScrollByDelta( -MouseEvent.GetCursorDelta().Y );
 	}
@@ -370,19 +115,62 @@ FReply SSequencerTrackArea::OnMouseMove( const FGeometry& MyGeometry, const FPoi
 
 FReply SSequencerTrackArea::OnMouseWheel( const FGeometry& MyGeometry, const FPointerEvent& MouseEvent )
 {
+	// First try the edit tool
+	auto SequencerPin = SequencerWidget.Pin();
+	if (SequencerPin.IsValid())
+	{
+		FReply Reply = SequencerPin->GetEditTool().OnMouseWheel(*this, MyGeometry, MouseEvent);
+		if (Reply.IsEventHandled())
+		{
+			return Reply;
+		}
+	}
+
+	// Then the time slider
 	FReply Reply = TimeSliderController->OnMouseWheel( SharedThis(this), MyGeometry, MouseEvent );
+	if (Reply.IsEventHandled())
+	{
+		return Reply;
+	}
+
+	// Failing that, we'll just scroll vertically
+	TreeView.Pin()->ScrollByDelta(WheelScrollAmount * -MouseEvent.GetWheelDelta());
+	return FReply::Handled();
+}
+
+void SSequencerTrackArea::OnMouseCaptureLost()
+{
+	auto SequencerPin = SequencerWidget.Pin();
+	if (SequencerPin.IsValid())
+	{
+		SequencerPin->GetEditTool().OnMouseCaptureLost();
+	}
+}
+
+FCursorReply SSequencerTrackArea::OnCursorQuery( const FGeometry& MyGeometry, const FPointerEvent& CursorEvent ) const
+{
+	if (CursorEvent.IsMouseButtonDown(EKeys::RightMouseButton) && HasMouseCapture())
+	{
+		return FCursorReply::Cursor(EMouseCursor::GrabHandClosed);
+	}
 
-	if (!Reply.IsEventHandled())
+	auto SequencerPin = SequencerWidget.Pin();
+	if (SequencerPin.IsValid())
 	{
-		TreeView.Pin()->ScrollByDelta(WheelScrollAmount * -MouseEvent.GetWheelDelta());
-		Reply = FReply::Handled();
+		return SequencerPin->GetEditTool().OnCursorQuery(MyGeometry, CursorEvent);
 	}
 
-	return Reply;
+	return FCursorReply::Unhandled();
 }
 
 void SSequencerTrackArea::Tick( const FGeometry& AllottedGeometry, const double InCurrentTime, const float InDeltaTime )
 {
+	auto SequencerPin = SequencerWidget.Pin();
+	if (SequencerPin.IsValid())
+	{
+		SequencerPin->GetEditTool().Tick(AllottedGeometry, InCurrentTime, InDeltaTime);
+	}
+
 	FVector2D Size = AllottedGeometry.GetLocalSize();
 
 	if (SizeLastFrame.IsSet() && Size.X != SizeLastFrame->X)
@@ -411,130 +199,4 @@ void SSequencerTrackArea::Tick( const FGeometry& AllottedGeometry, const double
 			++Index;
 		}
 	}
-}
-
-void SSequencerTrackArea::SetNewMarqueeBounds(const FPointerEvent& MouseEvent, const FGeometry& MyGeometry)
-{
-	const FVector2D MouseDelta = MouseEvent.GetCursorDelta();
-
-	FVector2D LocalPosition = MyGeometry.AbsoluteToLocal(MouseEvent.GetScreenSpacePosition());
-
-	// Handle virtual scrolling when at the vertical extremes of the widget (performed before we clamp the mouse pos)
-	{
-		const float ScrollThresholdV = MyGeometry.Size.Y * 0.025f;
-
-		float Difference = LocalPosition.Y - ScrollThresholdV;
-		if (Difference < 0 && MouseDelta.Y < 0)
-		{
-			TreeView.Pin()->ScrollByDelta( Difference * 0.1f );
-		}
-
-		Difference = LocalPosition.Y - (MyGeometry.Size.Y - ScrollThresholdV);
-		if (Difference > 0 && MouseDelta.Y > 0)
-		{
-			TreeView.Pin()->ScrollByDelta( Difference * 0.1f );
-		}
-	}
-
-	const FVirtualConverter Converter(TreeView.Pin(), Sequencer.Pin(), MyGeometry);
-
-	// Clamp the vertical position to the actual bounds of the track area
-	LocalPosition.Y = FMath::Clamp(LocalPosition.Y, 0.f, MyGeometry.GetLocalSize().Y);
-	MarqueeSelectData->CurrentPosition = Converter.PhysicalToVirtual(LocalPosition);
-
-	// Handle virtual scrolling when at the horizontal extremes of the widget
-	{
-		TRange<float> ViewRange = Sequencer.Pin()->GetViewRange();
-		const float ScrollThresholdH = ViewRange.Size<float>() * 0.025f;
-
-		float Difference = MarqueeSelectData->CurrentPosition.X - (ViewRange.GetLowerBoundValue() + ScrollThresholdH);
-		if (Difference < 0 && MouseDelta.X < 0)
-		{
-			Sequencer.Pin()->StartAutoscroll(Difference);
-		}
-		else
-		{
-			Difference = MarqueeSelectData->CurrentPosition.X - (ViewRange.GetUpperBoundValue() - ScrollThresholdH);
-			if (Difference > 0 && MouseDelta.X > 0)
-			{
-				Sequencer.Pin()->StartAutoscroll(Difference);
-			}
-			else
-			{
-				Sequencer.Pin()->StopAutoscroll();
-			}
-		}
-	}
-}
-
-struct FSelectionVisitor : IMarqueeVisitor
-{
-	FSelectionVisitor(FSequencerSelection& InSelection, const FPointerEvent& InMouseEvent)
-		: Selection(InSelection)
-		, MouseEvent(InMouseEvent)
-	{}
-
-	virtual void VisitKey(FKeyHandle KeyHandle, float KeyTime, const TSharedPtr<IKeyArea>& KeyArea, UMovieSceneSection* Section) const override
-	{
-		FSelectedKey SelectedKey(*Section, KeyArea, KeyHandle);
-
-		if (MouseEvent.IsLeftShiftDown())
-		{
-			Selection.AddToSelection(SelectedKey);
-		}
-		else if (MouseEvent.IsLeftControlDown())
-		{
-			if (Selection.IsSelected(SelectedKey))
-			{
-				Selection.RemoveFromSelection(SelectedKey);
-			}
-			else
-			{
-				Selection.AddToSelection(SelectedKey);
-			}
-		}
-		else
-		{
-			Selection.AddToSelection(SelectedKey);
-		}
-	}
-
-	virtual void VisitSection(FKeyHandle KeyHandle, float KeyTime, const TSharedPtr<IKeyArea>& KeyArea, UMovieSceneSection* Section) const
-	{
-
-	}
-
-private:
-	FSequencerSelection& Selection;
-	const FPointerEvent& MouseEvent;
-};
-
-void SSequencerTrackArea::HandleMarqueeSelection(const FGeometry& Geometry, const FPointerEvent& MouseEvent)
-{
-	auto PinnedSequencer = Sequencer.Pin();
-	auto PinnedTreeView = TreeView.Pin();
-
-	if (!PinnedSequencer.IsValid() || !PinnedTreeView.IsValid())
-	{
-		return;
-	}
-
-	auto& Selection = PinnedSequencer->GetSelection();
-	if (!MouseEvent.IsLeftControlDown() && !MouseEvent.IsLeftShiftDown())
-	{
-		Selection.Empty();
-	}
-
-	Selection.SuspendBroadcast();
-
-	FVector2D VirtualKeySize;
-	VirtualKeySize.X = SequencerSectionConstants::KeySize.X / Geometry.Size.X * (PinnedSequencer->GetViewRange().GetUpperBoundValue() - PinnedSequencer->GetViewRange().GetLowerBoundValue());
-	// Vertically, virtual units == physical units
-	VirtualKeySize.Y = SequencerSectionConstants::KeySize.Y;
-
-	FMarqueeRange Range(MarqueeSelectData->TopLeft(), MarqueeSelectData->BottomRight(), VirtualKeySize);
-	Range.Visit(FSelectionVisitor(Selection, MouseEvent), PinnedTreeView->GetNodeTree()->GetRootNodes());
-
-	Selection.ResumeBroadcast();
-	Selection.GetOnKeySelectionChanged().Broadcast();
 }
\ No newline at end of file
diff --git a/Engine/Source/Editor/Sequencer/Private/SSequencerTrackArea.h b/Engine/Source/Editor/Sequencer/Private/SSequencerTrackArea.h
--- a/Engine/Source/Editor/Sequencer/Private/SSequencerTrackArea.h
+++ b/Engine/Source/Editor/Sequencer/Private/SSequencerTrackArea.h
@@ -2,8 +2,6 @@
 
 #pragma once
 
-#include "UniquePtr.h"
-
 class SSequencerTreeViewRow;
 class SSequencerTreeView;
 class SSequencerTrackLane;
@@ -11,8 +9,6 @@ class SSequencerTrackLane;
 class FSequencerTimeSliderController;
 class FSequencer;
 
-struct FMarqueeSelectData;
-
 /**
  * The area where tracks( rows of sections ) are displayed
  */
@@ -22,12 +18,8 @@ public:
 	SLATE_BEGIN_ARGS( SSequencerTrackArea ){}
 	SLATE_END_ARGS()
 
-	/** User-defined constructor/destructor to ensure TUniquePtr destruction */
-	SSequencerTrackArea();
-	~SSequencerTrackArea();
-
 	/** Construct this widget */
-	void Construct( const FArguments& InArgs, TSharedRef<FSequencerTimeSliderController> InTimeSliderController, TSharedRef<FSequencer> InSequencer );
+	void Construct( const FArguments& InArgs, TSharedRef<FSequencerTimeSliderController> InTimeSliderController, TSharedRef<SSequencer> InSequencerWidget );
 
 	/** Assign a tree view to this track area. */
 	void SetTreeView(const TSharedPtr<SSequencerTreeView>& InTreeView);
@@ -44,26 +36,17 @@ public:
 	virtual FReply OnMouseMove( const FGeometry& MyGeometry, const FPointerEvent& MouseEvent ) override;
 	virtual FReply OnMouseWheel( const FGeometry& MyGeometry, const FPointerEvent& MouseEvent ) override;
 	virtual int32 OnPaint(const FPaintArgs& Args, const FGeometry& AllottedGeometry, const FSlateRect& MyClippingRect, FSlateWindowElementList& OutDrawElements, int32 LayerId, const FWidgetStyle& InWidgetStyle, bool bParentEnabled ) const override;
+	virtual void OnMouseCaptureLost() override;
+	virtual FCursorReply OnCursorQuery( const FGeometry& MyGeometry, const FPointerEvent& CursorEvent ) const override;
 	virtual void Tick( const FGeometry& AllottedGeometry, const double InCurrentTime, const float InDeltaTime ) override;
 
-private:
-
-	/** Paint the marquee box, if it's valid */
-	int32 PaintMarquee(const FGeometry& AllottedGeometry, const FSlateRect& MyClippingRect, FSlateWindowElementList& OutDrawElements, int32 LayerId) const;
-
-	/** Set new bounds for the marquee selection */
-	void SetNewMarqueeBounds(const FPointerEvent& MouseEvent, const FGeometry& Geometry);
-
-	/** Select all the keys contained in the current marquee */
-	void HandleMarqueeSelection(const FGeometry& Geometry, const FPointerEvent& MouseEvent);
-
 private:
 
 	/** A map of child slot content that exist in our view */
 	TMap<TSharedPtr<FSequencerDisplayNode>, TWeakPtr<SSequencerTrackLane>> TrackSlots;
 
-	/** Weak pointer to the sequencer */
-	TWeakPtr<FSequencer> Sequencer;
+	/** Weak pointer to the sequencer widget */
+	TWeakPtr<SSequencer> SequencerWidget;
 
 	/** Weak pointer to the tree view (used for scrolling interactions) */
 	TWeakPtr<SSequencerTreeView> TreeView;
@@ -73,7 +56,4 @@ private:
 
 	/** Keep a hold of the size of the area so we can maintain zoom levels */
 	TOptional<FVector2D> SizeLastFrame;
-
-	/** Optional marquee selection data */
-	TUniquePtr<FMarqueeSelectData> MarqueeSelectData;
 };
diff --git a/Engine/Source/Editor/Sequencer/Private/SSequencerTreeView.cpp b/Engine/Source/Editor/Sequencer/Private/SSequencerTreeView.cpp
--- a/Engine/Source/Editor/Sequencer/Private/SSequencerTreeView.cpp
+++ b/Engine/Source/Editor/Sequencer/Private/SSequencerTreeView.cpp
@@ -152,6 +152,31 @@ void SSequencerTreeView::OnChildRowRemoved(const FDisplayNodeRef& InNode)
 	CachedRowGeometry.Remove(InNode);
 }
 
+TSharedPtr<FSequencerDisplayNode> SSequencerTreeView::HitTestNode(float InPhysical) const
+{
+	auto* Found = Utils::BinarySearch<FCachedGeometry>(PhysicalNodes, [&](const FCachedGeometry& In){
+
+		if (InPhysical < In.PhysicalTop)
+		{
+			return Utils::ESearchState::Before;
+		}
+		else if (InPhysical > In.PhysicalTop + In.PhysicalHeight)
+		{
+			return Utils::ESearchState::After;
+		}
+
+		return Utils::ESearchState::Found;
+
+	});
+
+	if (Found)
+	{
+		return Found->Node;
+	}
+	
+	return nullptr;
+}
+
 float SSequencerTreeView::PhysicalToVirtual(float InPhysical) const
 {
 	int32 SearchIndex = PhysicalNodes.Num() / 2;
diff --git a/Engine/Source/Editor/Sequencer/Private/SSequencerTreeView.h b/Engine/Source/Editor/Sequencer/Private/SSequencerTreeView.h
--- a/Engine/Source/Editor/Sequencer/Private/SSequencerTreeView.h
+++ b/Engine/Source/Editor/Sequencer/Private/SSequencerTreeView.h
@@ -47,6 +47,11 @@ public:
 	/** Access the underlying tree data */
 	TSharedPtr<FSequencerNodeTree> GetNodeTree() { return SequencerNodeTree; }
 
+public:
+
+	/** Get the display node at the specified physical vertical position */
+	TSharedPtr<FSequencerDisplayNode> HitTestNode(float InPhysical) const;
+
 	/** Convert the specified physical vertical position into an absolute virtual position, ignoring expanded states */
 	float PhysicalToVirtual(float InPhysical) const;
 
@@ -55,6 +60,8 @@ public:
 	 */
 	float VirtualToPhysical(float InVirtual) const;
 
+public:
+
 	/** Refresh this tree as a result of the underlying tree data changing */
 	void Refresh();
 
diff --git a/Engine/Source/Editor/Sequencer/Private/SectionDragOperations.cpp b/Engine/Source/Editor/Sequencer/Private/SectionDragOperations.cpp
deleted file mode 100644
diff --git a/Engine/Source/Editor/Sequencer/Private/SectionDragOperations.h b/Engine/Source/Editor/Sequencer/Private/SectionDragOperations.h
deleted file mode 100644
diff --git a/Engine/Source/Editor/Sequencer/Private/Sequencer.cpp b/Engine/Source/Editor/Sequencer/Private/Sequencer.cpp
--- a/Engine/Source/Editor/Sequencer/Private/Sequencer.cpp
+++ b/Engine/Source/Editor/Sequencer/Private/Sequencer.cpp
@@ -236,7 +236,7 @@ void FSequencer::Tick(float InDeltaTime)
 	if (AutoscrubOffset.IsSet())
 	{
 		float Offset = AutoscrubOffset.GetValue() * AutoScrollFactor;
-		SetGlobalTime(GetGlobalTime() + Offset);
+		SetGlobalTimeDirectly(GetGlobalTime() + Offset);
 	}
 
 	float NewTime = GetGlobalTime() + InDeltaTime;
@@ -568,9 +568,7 @@ float FSequencer::GetGlobalTime()
 
 void FSequencer::SetGlobalTime( float NewTime )
 {
-	float LastTime = ScrubPosition;
-
-	if (bAutoScrollEnabled && PlaybackState != EMovieScenePlayerStatus::Scrubbing)
+	if (bAutoScrollEnabled)
 	{
 		float RangeOffset = CalculateAutoscrollEncroachment(NewTime).Get(0.f);
 			
@@ -584,6 +582,13 @@ void FSequencer::SetGlobalTime( float NewTime )
 		}
 	}
 
+	SetGlobalTimeDirectly(NewTime);
+}
+
+void FSequencer::SetGlobalTimeDirectly( float NewTime )
+{
+	float LastTime = ScrubPosition;
+
 	if (Settings->GetLockInOutToStartEndRange())
 	{
 		UMovieScene* FocusedMovieScene = GetFocusedMovieSceneSequence()->GetMovieScene();
@@ -598,6 +603,51 @@ void FSequencer::SetGlobalTime( float NewTime )
 	GEditor->RedrawLevelEditingViewports();
 }
 
+void FSequencer::UpdateAutoScroll(float NewTime)
+{
+	float ThresholdPercentage = 0.025f;
+	AutoscrollOffset = CalculateAutoscrollEncroachment(NewTime, ThresholdPercentage);
+
+	if (!AutoscrollOffset.IsSet())
+	{
+		AutoscrubOffset.Reset();
+		return;
+	}
+
+	TRange<float> ViewRange = GetViewRange();
+	const float Threshold = (ViewRange.GetUpperBoundValue() - ViewRange.GetLowerBoundValue()) * ThresholdPercentage;
+
+	// If we have no autoscrub offset yet, we move the scrub position to the boundary of the autoscroll threasdhold, then autoscrub from there
+	if (!AutoscrubOffset.IsSet())
+	{
+		if (AutoscrollOffset.GetValue() < 0 && ScrubPosition > ViewRange.GetLowerBoundValue() + Threshold)
+		{
+			SetGlobalTimeDirectly( ViewRange.GetLowerBoundValue() + Threshold );
+		}
+		else if (AutoscrollOffset.GetValue() > 0 && ScrubPosition < ViewRange.GetUpperBoundValue() - Threshold)
+		{
+			SetGlobalTimeDirectly( ViewRange.GetUpperBoundValue() - Threshold );
+		}
+	}
+
+	// Set the autoscrub rate to the autoscroll rate
+	if (Settings->GetLockInOutToStartEndRange())
+	{
+		// Don't autoscrub if we're at the extremes of the movie scene range
+		UMovieScene* FocusedMovieScene = GetFocusedMovieSceneSequence()->GetMovieScene();
+		if (NewTime < FocusedMovieScene->StartTime + Threshold ||
+			NewTime > FocusedMovieScene->EndTime - Threshold
+			)
+		{
+			AutoscrubOffset.Reset();
+			return;
+		}
+	}
+
+	// Scrub at the same rate we scroll
+	AutoscrubOffset = AutoscrollOffset;
+}
+
 TOptional<float> FSequencer::CalculateAutoscrollEncroachment(float NewTime, float ThresholdPercentage) const
 {
 	enum class EDirection { Positive, Negative };
@@ -612,7 +662,7 @@ TOptional<float> FSequencer::CalculateAutoscrollEncroachment(float NewTime, floa
 		// Scrolling backwards in time, and have hit the threshold
 		return NewTime - (RangeMin + AutoScrollThreshold);
 	}
-	else if (Movement == EDirection::Positive && ScrubPosition > RangeMax - AutoScrollThreshold)
+	else if (Movement == EDirection::Positive && NewTime > RangeMax - AutoScrollThreshold)
 	{
 		// Scrolling forwards in time, and have hit the threshold
 		return NewTime - (RangeMax - AutoScrollThreshold);
@@ -1128,26 +1178,23 @@ void FSequencer::OnScrubPositionChanged( float NewScrubPosition, bool bScrubbing
 		}
 		else if (bAutoScrollEnabled)
 		{
+			UpdateAutoScroll(NewScrubPosition);
+			
 			// When scrubbing, we animate auto-scrolled scrub position in Tick()
-			AutoscrollOffset = CalculateAutoscrollEncroachment(NewScrubPosition, 0.025f);
-			if (AutoscrollOffset.IsSet())
+			if (AutoscrubOffset.IsSet())
 			{
-				AutoscrubOffset = AutoscrollOffset;
 				return;
 			}
-			else
-			{
-				AutoscrubOffset.Reset();
-			}
 		}
 	}
 
-	SetGlobalTime( NewScrubPosition );
+	SetGlobalTimeDirectly( NewScrubPosition );
 }
 
 void FSequencer::OnBeginScrubbing()
 {
 	PlaybackState = EMovieScenePlayerStatus::Scrubbing;
+	SequencerWidget->RegisterActiveTimerForPlayback();
 }
 
 void FSequencer::OnEndScrubbing()
@@ -1173,6 +1220,11 @@ void FSequencer::OnToggleAutoScroll()
 	Settings->SetAutoScrollEnabled(bAutoScrollEnabled);
 }
 
+void FSequencer::VerticalScroll(float ScrollAmountUnits)
+{
+	SequencerWidget->GetTreeView()->ScrollByDelta(ScrollAmountUnits);
+}
+
 FGuid FSequencer::AddSpawnableForAssetOrClass( UObject* Object, UObject* CounterpartGamePreviewObject )
 {
 	FGuid NewSpawnableGuid;
@@ -1853,6 +1905,11 @@ void FSequencer::SetKey()
 	}
 }
 
+ISequencerEditTool& FSequencer::GetEditTool()
+{
+	return SequencerWidget->GetEditTool();
+}
+
 void FSequencer::BindSequencerCommands()
 {
 	const FSequencerCommands& Commands = FSequencerCommands::Get();
@@ -2013,6 +2070,8 @@ void FSequencer::BindSequencerCommands()
 		FCanExecuteAction::CreateLambda( []{ return true; } ),
 		FIsActionChecked::CreateLambda( [this]{ return Settings->GetShowCurveEditor(); } ) );
 
+	SequencerWidget->BindCommands(SequencerCommandBindings);
+
 	for (int32 i = 0; i < TrackEditors.Num(); ++i)
 	{
 		TrackEditors[i]->BindCommands(SequencerCommandBindings);
diff --git a/Engine/Source/Editor/Sequencer/Private/Sequencer.h b/Engine/Source/Editor/Sequencer/Private/Sequencer.h
--- a/Engine/Source/Editor/Sequencer/Private/Sequencer.h
+++ b/Engine/Source/Editor/Sequencer/Private/Sequencer.h
@@ -6,13 +6,11 @@
 #include "SelectedKey.h"
 #include "EditorUndoClient.h"
 
-
 class UMovieScene;
 class IToolkitHost;
 class UMovieSceneSequence;
 class ISequencerObjectChangeListener;
 
-
 /**
  * Sequencer is the editing tool for MovieScene assets.
  */
@@ -91,6 +89,9 @@ public:
 	virtual FSequencerSelection& GetSelection() override;
 	virtual void NotifyMapChanged(class UWorld* NewWorld, EMapChangeType MapChangeType) override;
 
+	/** Set the global time directly, without performing any auto-scroll */
+	void SetGlobalTimeDirectly(float Time);
+
 	/** @return The current view range */
 	virtual FAnimatedRange GetViewRange() const override;
 
@@ -271,6 +272,11 @@ public:
 	/** Called to add selected objects to the movie scene */
 	void AddSelectedObjects();
 
+public:
+
+	/** Access the currently enabled edit tool */
+	ISequencerEditTool& GetEditTool();
+
 protected:
 
 	/** Reset data about a movie scene when pushing or popping a movie scene. */
@@ -339,7 +345,15 @@ public:
 	/** Stop the sequencer from autoscrolling */
 	void StopAutoscroll();
 
+	/** Scroll the sequencer vertically by the specified number of slate units */
+	void VerticalScroll(float ScrollAmountUnits);
+
 protected:
+	/**
+	 * Update autoscroll mechanics as a result of a new time position
+	 */
+	void UpdateAutoScroll(float NewTime);
+
 	/**
 	 * Calculates the amount of encroachment the specified time has into the autoscroll region, if any
 	 */
diff --git a/Engine/Source/Editor/Sequencer/Private/SequencerCommands.cpp b/Engine/Source/Editor/Sequencer/Private/SequencerCommands.cpp
--- a/Engine/Source/Editor/Sequencer/Private/SequencerCommands.cpp
+++ b/Engine/Source/Editor/Sequencer/Private/SequencerCommands.cpp
@@ -46,6 +46,9 @@ void FSequencerCommands::RegisterCommands()
 	UI_COMMAND( ToggleCleanView, "Clean View", "Enable 'Clean View' mode which displays only global tracks when no filter is applied.", EUserInterfaceActionType::ToggleButton, FInputChord() );
 	
 	UI_COMMAND( ToggleShowCurveEditor, "Curve Editor", "Show the animation keys in a curve editor.", EUserInterfaceActionType::ToggleButton, FInputChord() );
+
+	UI_COMMAND( MoveTool, "Move", "Activates the move tool.", EUserInterfaceActionType::ToggleButton, FInputChord( EKeys::V ) );
+	UI_COMMAND( MarqueeTool, "Marquee", "Activates the marquee selection tool.", EUserInterfaceActionType::ToggleButton, FInputChord( EKeys::M ) );
 }
 
 #undef LOCTEXT_NAMESPACE
\ No newline at end of file
diff --git a/Engine/Source/Editor/Sequencer/Private/SequencerCommands.h b/Engine/Source/Editor/Sequencer/Private/SequencerCommands.h
--- a/Engine/Source/Editor/Sequencer/Private/SequencerCommands.h
+++ b/Engine/Source/Editor/Sequencer/Private/SequencerCommands.h
@@ -99,6 +99,9 @@ public:
 	/** Toggles whether or not the curve editor should be shown. */
 	TSharedPtr< FUICommandInfo > ToggleShowCurveEditor;
 
+	TSharedPtr< FUICommandInfo > MoveTool;
+	TSharedPtr< FUICommandInfo > MarqueeTool;
+
 	/**
 	 * Initialize commands
 	 */
diff --git a/Engine/Source/Editor/Sequencer/Private/SequencerHotspots.cpp b/Engine/Source/Editor/Sequencer/Private/SequencerHotspots.cpp
new file mode 100644
--- /dev/null
+++ b/Engine/Source/Editor/Sequencer/Private/SequencerHotspots.cpp
@@ -0,0 +1,19 @@
+// Copyright 1998-2015 Epic Games, Inc. All Rights Reserved.
+
+#include "SequencerPrivatePCH.h"
+#include "SequencerHotspots.h"
+#include "Tools/EditToolDragOperations.h"
+#include "Sequencer.h"
+#include "MovieSceneSection.h"
+
+TSharedPtr<IEditToolDragOperation> FSectionResizeHotspot::InitiateDrag(ISequencer& Sequencer)
+{
+	const auto& SelectedSections = Sequencer.GetSelection().GetSelectedSections();
+	auto SectionHandles = StaticCastSharedRef<SSequencer>(Sequencer.GetSequencerWidget())->GetSectionHandles(SelectedSections);
+	
+	if (!SelectedSections.Contains(Section.GetSectionObject()))
+	{
+		SectionHandles.Add(Section);
+	}
+	return MakeShareable( new FResizeSection(static_cast<FSequencer&>(Sequencer), SectionHandles, Section, HandleType == Right) );
+}
\ No newline at end of file
diff --git a/Engine/Source/Editor/Sequencer/Private/SequencerHotspots.h b/Engine/Source/Editor/Sequencer/Private/SequencerHotspots.h
new file mode 100644
--- /dev/null
+++ b/Engine/Source/Editor/Sequencer/Private/SequencerHotspots.h
@@ -0,0 +1,70 @@
+// Copyright 1998-2015 Epic Games, Inc. All Rights Reserved.
+
+#pragma once
+
+#include "ISequencerHotspot.h"
+#include "SequencerDisplayNode.h"
+#include "ISequencerSection.h"
+
+/** A hotspot representing a key */
+struct FKeyHotspot : ISequencerHotspot
+{
+	FKeyHotspot(FSelectedKey InKey, TSharedPtr<FTrackNode> InTrackNode)
+		: Key(InKey)
+		, TrackNode(MoveTemp(InTrackNode))
+	{}
+
+	virtual ESequencerHotspot GetType() const override { return ESequencerHotspot::Key; }
+
+	/** The key itself */
+	FSelectedKey Key;
+	/** The track node that holds the section the key is contained in */
+	TSharedPtr<FTrackNode> TrackNode;
+};
+
+/** Structure used to encapsulate a section and its track node */
+struct FSectionHandle
+{
+	FSectionHandle(TSharedPtr<FTrackNode> InTrackNode, int32 InSectionIndex)
+		: TrackNode(MoveTemp(InTrackNode)), SectionIndex(InSectionIndex)
+	{}
+
+	UMovieSceneSection* GetSectionObject() const
+	{
+		return TrackNode->GetSections()[SectionIndex]->GetSectionObject();
+	}
+
+	int32 SectionIndex;
+	TSharedPtr<FTrackNode> TrackNode;
+};
+
+/** A hotspot representing a section */
+struct FSectionHotspot : ISequencerHotspot
+{
+	FSectionHotspot(FSectionHandle InSection)
+		: Section(InSection)
+	{}
+
+	virtual ESequencerHotspot GetType() const override { return ESequencerHotspot::Section; }
+
+	/** Handle to the section */
+	FSectionHandle Section;
+};
+
+/** A hotspot representing a resize handle on a section */
+struct FSectionResizeHotspot : ISequencerHotspot
+{
+	enum EHandle { Left, Right };
+
+	FSectionResizeHotspot(EHandle InHandleType, FSectionHandle InSection) : Section(InSection), HandleType(InHandleType) {}
+
+	virtual ESequencerHotspot GetType() const override { return HandleType == Left ? ESequencerHotspot::SectionResize_L : ESequencerHotspot::SectionResize_R; }
+	
+	virtual TSharedPtr<IEditToolDragOperation> InitiateDrag(ISequencer& Sequencer) override;
+
+	/** Handle to the section */
+	FSectionHandle Section;
+
+private:
+	EHandle HandleType;
+};
\ No newline at end of file
diff --git a/Engine/Source/Editor/Sequencer/Private/TimeSliderController.cpp b/Engine/Source/Editor/Sequencer/Private/TimeSliderController.cpp
--- a/Engine/Source/Editor/Sequencer/Private/TimeSliderController.cpp
+++ b/Engine/Source/Editor/Sequencer/Private/TimeSliderController.cpp
@@ -534,10 +534,8 @@ FReply FSequencerTimeSliderController::OnMouseMove( TSharedRef<SWidget> WidgetOw
 					{
 						NewValue = TimeSliderArgs.Settings->SnapTimeToInterval(NewValue);
 					}
-
-					// Clamp to the view range so that you can't scroll beyond the current view range. This doesn't conflict with auto-scroll because the view range will grow before you get to the bounds.
-					NewValue = FMath::Clamp(NewValue, TimeSliderArgs.ViewRange.Get().GetLowerBoundValue(), TimeSliderArgs.ViewRange.Get().GetUpperBoundValue());
-
+					
+					// Delegate responsibility for clamping to the current viewrange to the client
 					CommitScrubPosition( NewValue, /*bIsScrubbing=*/true );
 				}
 				else if (MouseDragType == DRAG_SETTING_RANGE)
diff --git a/Engine/Source/Editor/Sequencer/Private/Tools/DelayedDrag.h b/Engine/Source/Editor/Sequencer/Private/Tools/DelayedDrag.h
new file mode 100644
--- /dev/null
+++ b/Engine/Source/Editor/Sequencer/Private/Tools/DelayedDrag.h
@@ -0,0 +1,50 @@
+// Copyright 1998-2015 Epic Games, Inc. All Rights Reserved.
+
+#pragma once
+
+/** Helper class to assist with delayed drag operations */
+class FDelayedDrag
+{
+public:
+
+	/** Construct this drag helper with an initial posision, and a key (probably mouse button) required for dragging */
+	FDelayedDrag(FVector2D InInitialPosition, FKey InEffectiveKey)
+		: InitialPosition(InInitialPosition), bHasInitiatedDrag(false), DistanceDragged(0), EffectiveKey(InEffectiveKey)
+	{
+	}
+
+	/** Get the initial start position (before any drag has started) */
+	const FVector2D& GetInitialPosition() const { return InitialPosition; }
+
+	/** Check whether we have initiated a drag or not */
+	bool IsDragging() const { return bHasInitiatedDrag; }
+
+	/** Force the state of this helper to be 'dragging' */
+	void ForceDragStart() { bHasInitiatedDrag = true; }
+
+	/** Attempt to start a drag from the given mouse event. */
+	bool AttemptDragStart(const FPointerEvent& MouseEvent)
+	{
+		if (!bHasInitiatedDrag && MouseEvent.IsMouseButtonDown(EffectiveKey))
+		{
+			DistanceDragged += FMath::Abs( MouseEvent.GetCursorDelta().Size() );
+			if (!bHasInitiatedDrag && DistanceDragged > FSlateApplication::Get().GetDragTriggerDistance())
+			{
+				ForceDragStart();
+				return true;
+			}
+		}
+		return false;
+	}
+
+protected:
+
+	/** The initial position of the drag start */
+	FVector2D InitialPosition;
+	/** True where the distance dragged is sufficient to have started a drag */
+	bool bHasInitiatedDrag;
+	/** The amount we have dragged */
+	float DistanceDragged;
+	/** The key that must be pressed to initiate the drag */
+	FKey EffectiveKey;
+};
\ No newline at end of file
diff --git a/Engine/Source/Editor/Sequencer/Private/Tools/EditToolDragOperations.cpp b/Engine/Source/Editor/Sequencer/Private/Tools/EditToolDragOperations.cpp
new file mode 100644
--- /dev/null
+++ b/Engine/Source/Editor/Sequencer/Private/Tools/EditToolDragOperations.cpp
@@ -0,0 +1,568 @@
+// Copyright 1998-2015 Epic Games, Inc. All Rights Reserved.
+
+#include "SequencerPrivatePCH.h"
+#include "EditToolDragOperations.h"
+#include "Sequencer.h"
+#include "MovieSceneSection.h"
+#include "ISequencerSection.h"
+#include "IKeyArea.h"
+#include "CommonMovieSceneTools.h"
+#include "MovieScene.h"
+#include "MovieSceneTrack.h"
+#include "MovieSceneShotTrack.h"
+#include "VirtualTrackArea.h"
+
+/** How many pixels near the mouse has to be before snapping occurs */
+const float PixelSnapWidth = 10.f;
+
+TRange<float> GetSectionBoundaries(UMovieSceneSection* Section, TSharedPtr<FTrackNode> SequencerNode)
+{
+	// Find the borders of where you can drag to
+	float LowerBound = -FLT_MAX, UpperBound = FLT_MAX;
+
+	// Also get the closest borders on either side
+	const TArray< TSharedRef<ISequencerSection> >& Sections = SequencerNode->GetSections();
+	for (int32 SectionIndex = 0; SectionIndex < Sections.Num(); ++SectionIndex)
+	{
+		const UMovieSceneSection* InSection = Sections[SectionIndex]->GetSectionObject();
+		if (Section != InSection && Section->GetRowIndex() == InSection->GetRowIndex())
+		{
+			if (InSection->GetEndTime() <= Section->GetStartTime() && InSection->GetEndTime() > LowerBound)
+			{
+				LowerBound = InSection->GetEndTime();
+			}
+			if (InSection->GetStartTime() >= Section->GetEndTime() && InSection->GetStartTime() < UpperBound)
+			{
+				UpperBound = InSection->GetStartTime();
+			}
+		}
+	}
+
+	return TRange<float>(LowerBound, UpperBound);
+}
+
+void GetSectionSnapTimes(TArray<float>& OutSnapTimes, UMovieSceneSection* Section, TSharedPtr<FTrackNode> SequencerNode, bool bIgnoreOurSectionCustomSnaps)
+{
+	// @todo Sequencer handle dilation snapping better
+
+	// Collect all the potential snap times from other section borders
+	const TArray< TSharedRef<ISequencerSection> >& Sections = SequencerNode->GetSections();
+	for (int32 SectionIndex = 0; SectionIndex < Sections.Num(); ++SectionIndex)
+	{
+		const UMovieSceneSection* InSection = Sections[SectionIndex]->GetSectionObject();
+		bool bIsThisSection = Section == InSection;
+		if (!bIgnoreOurSectionCustomSnaps || !bIsThisSection)
+		{
+			InSection->GetSnapTimes(OutSnapTimes, Section != InSection);
+		}
+	}
+
+	// snap to director track if it exists, and we are not the director track
+	UMovieSceneTrack* OuterTrack = Cast<UMovieSceneTrack>(Section->GetOuter());
+	UMovieScene* MovieScene = Cast<UMovieScene>(OuterTrack->GetOuter());
+	UMovieSceneTrack* ShotTrack = MovieScene->FindMasterTrack(UMovieSceneShotTrack::StaticClass());
+	if (ShotTrack && OuterTrack != ShotTrack)
+	{
+		const TArray<UMovieSceneSection*>& ShotSections = ShotTrack->GetAllSections();
+		for (int32 SectionIndex = 0; SectionIndex < ShotSections.Num(); ++SectionIndex)
+		{
+			auto Shot = ShotSections[SectionIndex];
+			Shot->GetSnapTimes(OutSnapTimes, true);
+		}
+	}
+}
+
+bool SnapToTimes(TArray<float> InitialTimes, const TArray<float>& SnapTimes, const FTimeToPixel& TimeToPixelConverter, float& OutInitialTime, float& OutSnapTime)
+{
+	bool bSuccess = false;
+	float ClosestTimePixelDistance = PixelSnapWidth;
+	
+	for (int32 InitialTimeIndex = 0; InitialTimeIndex < InitialTimes.Num(); ++InitialTimeIndex)
+	{
+		float InitialTime = InitialTimes[InitialTimeIndex];
+		float PixelXOfTime = TimeToPixelConverter.TimeToPixel(InitialTime);
+
+		for (int32 SnapTimeIndex = 0; SnapTimeIndex < SnapTimes.Num(); ++SnapTimeIndex)
+		{
+			float SnapTime = SnapTimes[SnapTimeIndex];
+			float PixelXOfSnapTime = TimeToPixelConverter.TimeToPixel(SnapTime);
+
+			float PixelDistance = FMath::Abs(PixelXOfTime - PixelXOfSnapTime);
+			if (PixelDistance < ClosestTimePixelDistance)
+			{
+				ClosestTimePixelDistance = PixelDistance;
+				OutInitialTime = InitialTime;
+				OutSnapTime = SnapTime;
+				bSuccess = true;
+			}
+		}
+	}
+
+	return bSuccess;
+}
+
+TOptional<float> SnapToTimes(float InitialTime, const TArray<float>& SnapTimes, const FTimeToPixel& TimeToPixelConverter)
+{
+	TArray<float> InitialTimes;
+	InitialTimes.Add(InitialTime);
+
+	float OutSnapTime = 0.f;
+	float OutInitialTime = 0.f;
+	bool bSuccess = SnapToTimes(InitialTimes, SnapTimes, TimeToPixelConverter, OutInitialTime, OutSnapTime);
+	if (bSuccess)
+	{
+		return TOptional<float>(OutSnapTime);
+	}
+	return TOptional<float>();
+}
+
+
+FEditToolDragOperation::FEditToolDragOperation( FSequencer& InSequencer )
+	: Sequencer(InSequencer)
+{
+	Settings = Sequencer.GetSettings();
+}
+
+void FEditToolDragOperation::BeginTransaction( const TArray< FSectionHandle >& Sections, const FText& TransactionDesc )
+{
+	// Begin an editor transaction and mark the section as transactional so it's state will be saved
+	Transaction.Reset( new FScopedTransaction(TransactionDesc) );
+
+	for (auto& Handle : Sections)
+	{
+		UMovieSceneSection* SectionObj = Handle.GetSectionObject();
+
+		SectionObj->SetFlags( RF_Transactional );
+		// Save the current state of the section
+		SectionObj->Modify();
+	}
+}
+
+void FEditToolDragOperation::EndTransaction()
+{
+	Transaction.Reset();
+	Sequencer.UpdateRuntimeInstances();
+}
+
+FResizeSection::FResizeSection( FSequencer& Sequencer, TArray<FSectionHandle> InSections, TOptional<FSectionHandle> InCardinalSection, bool bInDraggingByEnd )
+	: FEditToolDragOperation( Sequencer )
+	, Sections( MoveTemp(InSections) )
+	, CardinalSection(InCardinalSection)
+	, bDraggingByEnd(bInDraggingByEnd)
+	, MouseDownTime(ForceInit)
+{
+}
+
+void FResizeSection::OnBeginDrag(const FVector2D& LocalMousePos, const FVirtualTrackArea& VirtualTrackArea)
+{
+	BeginTransaction( Sections, NSLOCTEXT("Sequencer", "DragSectionEdgeTransaction", "Resize section") );
+
+	MouseDownTime = VirtualTrackArea.PixelToTime(LocalMousePos.X);
+
+	DraggedKeyHandles.Empty();
+	SectionInitTimes.Empty();
+
+	for (auto& Handle : Sections)
+	{
+		UMovieSceneSection* Section = Handle.GetSectionObject();
+
+		Section->GetKeyHandles(DraggedKeyHandles);
+		SectionInitTimes.Add(Section, bDraggingByEnd ? Section->GetEndTime() : Section->GetStartTime());
+	}
+}
+
+void FResizeSection::OnEndDrag()
+{
+	EndTransaction();
+
+	DraggedKeyHandles.Empty();
+}
+
+void FResizeSection::OnDrag( const FPointerEvent& MouseEvent, const FVector2D& LocalMousePos, const FVirtualTrackArea& VirtualTrackArea )
+{
+	bool bIsDilating = MouseEvent.IsControlDown();
+
+	for (auto& Handle : Sections)
+	{
+		UMovieSceneSection* Section = Handle.GetSectionObject();
+
+		// Convert the current mouse position to a time
+		float DeltaTime = VirtualTrackArea.PixelToTime(LocalMousePos.X) - MouseDownTime;
+		float NewTime = SectionInitTimes[Section] + DeltaTime;
+
+		// Find the borders of where you can drag to
+		TRange<float> SectionBoundaries = GetSectionBoundaries(Section, Handle.TrackNode);
+		
+		// Snapping
+		if ( Settings->GetIsSnapEnabled() )
+		{
+			bool bSnappedToSection = false;
+			if ( Settings->GetSnapSectionTimesToSections() )
+			{
+				TArray<float> TimesToSnapTo;
+				GetSectionSnapTimes(TimesToSnapTo, Section, Handle.TrackNode, bIsDilating);
+
+				TOptional<float> NewSnappedTime = SnapToTimes(NewTime, TimesToSnapTo, VirtualTrackArea);
+
+				if (NewSnappedTime.IsSet())
+				{
+					NewTime = NewSnappedTime.GetValue();
+					bSnappedToSection = true;
+				}
+			}
+
+			if ( bSnappedToSection == false && Settings->GetSnapSectionTimesToInterval() )
+			{
+				NewTime = Settings->SnapTimeToInterval(NewTime);
+			}
+		}
+
+		if( bDraggingByEnd )
+		{
+			// Dragging the end of a section
+			// Ensure we aren't shrinking past the start time
+			NewTime = FMath::Clamp( NewTime, Section->GetStartTime(), SectionBoundaries.GetUpperBoundValue() );
+
+			if (bIsDilating)
+			{
+				float NewSize = NewTime - Section->GetStartTime();
+				float DilationFactor = NewSize / Section->GetTimeSize();
+				Section->DilateSection(DilationFactor, Section->GetStartTime(), DraggedKeyHandles);
+			}
+			else
+			{
+				Section->SetEndTime( NewTime );
+			}
+		}
+		else if( !bDraggingByEnd )
+		{
+			// Dragging the start of a section
+			// Ensure we arent expanding past the end time
+			NewTime = FMath::Clamp( NewTime, SectionBoundaries.GetLowerBoundValue(), Section->GetEndTime() );
+			
+			if (bIsDilating)
+			{
+				float NewSize = Section->GetEndTime() - NewTime;
+				float DilationFactor = NewSize / Section->GetTimeSize();
+				Section->DilateSection(DilationFactor, Section->GetEndTime(), DraggedKeyHandles);
+			}
+			else
+			{
+				Section->SetStartTime( NewTime );
+			}
+		}
+	}
+}
+
+FMoveSection::FMoveSection( FSequencer& Sequencer, TArray<FSectionHandle> InSections )
+	: FEditToolDragOperation( Sequencer )
+	, Sections( MoveTemp(InSections) )
+{
+}
+
+void FMoveSection::OnBeginDrag(const FVector2D& LocalMousePos, const FVirtualTrackArea& VirtualTrackArea)
+{
+	BeginTransaction( Sections, NSLOCTEXT("Sequencer", "MoveSectionTransaction", "Move Section") );
+		
+	MouseDownTime = VirtualTrackArea.PixelToTime(LocalMousePos.X);
+
+	DraggedKeyHandles.Empty();
+
+	for (auto Section : Sections)
+	{
+		Section.GetSectionObject()->GetKeyHandles(DraggedKeyHandles);
+	}
+}
+
+void FMoveSection::OnEndDrag()
+{
+	DraggedKeyHandles.Empty();
+
+	for (auto& Handle : Sections)
+	{
+		Handle.TrackNode->FixRowIndices();
+
+		UMovieSceneSection* Section = Handle.GetSectionObject();
+		UMovieSceneTrack* OuterTrack = Cast<UMovieSceneTrack>(Section->GetOuter());
+
+		if (OuterTrack)
+		{
+			OuterTrack->Modify();
+			OuterTrack->OnSectionMoved(*Section);
+		}
+	}
+
+	EndTransaction();
+}
+
+
+void FMoveSection::OnDrag( const FPointerEvent& MouseEvent, const FVector2D& LocalMousePos, const FVirtualTrackArea& VirtualTrackArea )
+{
+	float NewMouseDownTime = VirtualTrackArea.PixelToTime(LocalMousePos.X);
+	float DistanceMoved =  NewMouseDownTime - MouseDownTime;
+	float DeltaTime = DistanceMoved;
+
+	for (auto& Handle : Sections)
+	{
+		UMovieSceneSection* Section = Handle.GetSectionObject();
+
+		auto SequencerNodeSections = Handle.TrackNode->GetSections();
+		
+		TArray<UMovieSceneSection*> MovieSceneSections;
+		for (int32 i = 0; i < SequencerNodeSections.Num(); ++i)
+		{
+			MovieSceneSections.Add(SequencerNodeSections[i]->GetSectionObject());
+		}
+						
+		if ( Settings->GetIsSnapEnabled() )
+		{
+			bool bSnappedToSection = false;
+			if ( Settings->GetSnapSectionTimesToSections() )
+			{
+				TArray<float> TimesToSnapTo;
+				GetSectionSnapTimes(TimesToSnapTo, Section, Handle.TrackNode, true);
+
+				TArray<float> TimesToSnap;
+				TimesToSnap.Add(DistanceMoved + Section->GetStartTime());
+				TimesToSnap.Add(DistanceMoved + Section->GetEndTime());
+
+				float OutSnappedTime = 0.f;
+				float OutNewTime = 0.f;
+				if (SnapToTimes(TimesToSnap, TimesToSnapTo, VirtualTrackArea, OutSnappedTime, OutNewTime))
+				{
+					DeltaTime = OutNewTime - (OutSnappedTime - DistanceMoved);
+					bSnappedToSection = true;
+				}
+			}
+
+			if ( bSnappedToSection == false && Settings->GetSnapSectionTimesToInterval() )
+			{
+				float NewStartTime = DistanceMoved + Section->GetStartTime();
+				DeltaTime = Settings->SnapTimeToInterval( NewStartTime ) - Section->GetStartTime();
+			}
+		}
+
+		int32 TargetRowIndex = Section->GetRowIndex();
+
+		// vertical dragging - master tracks only
+		if (Handle.TrackNode->GetTrack()->SupportsMultipleRows() && SequencerNodeSections.Num() > 1)
+		{
+			float TrackHeight = SequencerNodeSections[0]->GetSectionHeight();
+
+			if (LocalMousePos.Y < 0.f || LocalMousePos.Y > TrackHeight)
+			{
+				int32 MaxRowIndex = 0;
+				for (int32 i = 0; i < SequencerNodeSections.Num(); ++i)
+				{
+					if (SequencerNodeSections[i]->GetSectionObject() != Section)
+					{
+						MaxRowIndex = FMath::Max(MaxRowIndex, SequencerNodeSections[i]->GetSectionObject()->GetRowIndex());
+					}
+				}
+
+				TargetRowIndex = FMath::Clamp(Section->GetRowIndex() + FMath::FloorToInt(LocalMousePos.Y / TrackHeight),
+					0, MaxRowIndex + 1);
+			}
+		}
+		
+		bool bDeltaX = !FMath::IsNearlyZero(DeltaTime);
+		bool bDeltaY = TargetRowIndex != Section->GetRowIndex();
+
+		if (bDeltaX && bDeltaY &&
+			!Section->OverlapsWithSections(MovieSceneSections, TargetRowIndex - Section->GetRowIndex(), DeltaTime))
+		{
+			Section->MoveSection(DeltaTime, DraggedKeyHandles);
+			Section->SetRowIndex(TargetRowIndex);
+		}
+		else
+		{
+			if (bDeltaY &&
+				!Section->OverlapsWithSections(MovieSceneSections, TargetRowIndex - Section->GetRowIndex(), 0.f))
+			{
+				Section->SetRowIndex(TargetRowIndex);
+			}
+
+			if (bDeltaX)
+			{
+				if (!Section->OverlapsWithSections(MovieSceneSections, 0, DeltaTime))
+				{
+					Section->MoveSection(DeltaTime, DraggedKeyHandles);
+				}
+				else
+				{
+					// Find the borders of where you can move to
+					TRange<float> SectionBoundaries = GetSectionBoundaries(Section, Handle.TrackNode);
+
+					float LeftMovementMaximum = SectionBoundaries.GetLowerBoundValue() - Section->GetStartTime();
+					float RightMovementMaximum = SectionBoundaries.GetUpperBoundValue() - Section->GetEndTime();
+
+					// Tell the section to move itself and any data it has
+					Section->MoveSection( FMath::Clamp(DeltaTime, LeftMovementMaximum, RightMovementMaximum), DraggedKeyHandles );
+				}
+			}
+		}
+	}
+
+	MouseDownTime = NewMouseDownTime;
+}
+
+void FMoveKeys::OnBeginDrag(const FVector2D& LocalMousePos, const FVirtualTrackArea& VirtualTrackArea)
+{
+	check( SelectedKeys.Num() > 0 )
+
+	// Begin an editor transaction and mark the section as transactional so it's state will be saved
+	GEditor->BeginTransaction( NSLOCTEXT("Sequencer", "MoveKeysTransaction", "Move Keys") );
+
+	TSet<UMovieSceneSection*> ModifiedSections;
+	for( FSelectedKey SelectedKey : SelectedKeys )
+	{
+		UMovieSceneSection* OwningSection = SelectedKey.Section;
+
+		// Only modify sections once
+		if( !ModifiedSections.Contains( OwningSection ) )
+		{
+			OwningSection->SetFlags( RF_Transactional );
+
+			// Save the current state of the section
+			OwningSection->Modify();
+
+			// Section has been modified
+			ModifiedSections.Add( OwningSection );
+		}
+	}
+}
+
+void FMoveKeys::OnEndDrag()
+{
+	EndTransaction();
+}
+
+void FMoveKeys::OnDrag( const FPointerEvent& MouseEvent, const FVector2D& LocalMousePos, const FVirtualTrackArea& VirtualTrackArea )
+{
+	float MouseTime = VirtualTrackArea.PixelToTime(LocalMousePos.X);
+	float DistanceMoved = MouseTime - VirtualTrackArea.PixelToTime(LocalMousePos.X - MouseEvent.GetCursorDelta().X);
+
+	if( DistanceMoved != 0.0f )
+	{
+		float TimeDelta = DistanceMoved;
+		// Snapping
+		if ( Settings->GetIsSnapEnabled() )
+		{
+			bool bSnappedToKeyTime = false;
+			if ( Settings->GetSnapKeyTimesToKeys() && CardinalTrack.IsValid() )
+			{
+				TArray<float> OutSnapTimes;
+				GetKeySnapTimes(OutSnapTimes, CardinalTrack);
+
+				TArray<float> InitialTimes;
+				for ( FSelectedKey SelectedKey : SelectedKeys )
+				{
+					InitialTimes.Add(SelectedKey.KeyArea->GetKeyTime(SelectedKey.KeyHandle.GetValue()) + DistanceMoved);
+				}
+				float OutInitialTime = 0.f;
+				float OutSnapTime = 0.f;
+				if ( SnapToTimes( InitialTimes, OutSnapTimes, VirtualTrackArea, OutInitialTime, OutSnapTime ) )
+				{
+					bSnappedToKeyTime = true;
+					TimeDelta = OutSnapTime - (OutInitialTime - DistanceMoved);
+				}
+			}
+
+			if (Settings->GetSnapKeyTimesToInterval())
+			{
+				// todo: arodham: fix snapping
+				// float SnapTimeIntervalDelta = Settings->SnapTimeToInterval(MouseTime) - SelectedKeyTime;
+
+				// // Snap to time interval only if we haven't snapped to a keyframe or if the time difference is smaller than the difference to the keyframe.
+				// if (bSnappedToKeyTime)
+				// {					
+				// 	if (FMath::Abs(SnapTimeIntervalDelta) < FMath::Abs(TimeDelta))
+				// 	{
+				// 		TimeDelta = SnapTimeIntervalDelta;
+				// 	}
+				// }
+				// else
+				// {
+				// 	TimeDelta = SnapTimeIntervalDelta;
+				// }
+			}
+		}
+
+		float PrevNewKeyTime = FLT_MAX;
+
+		for( FSelectedKey SelectedKey : SelectedKeys )
+		{
+			UMovieSceneSection* Section = SelectedKey.Section;
+
+			TSharedPtr<IKeyArea>& KeyArea = SelectedKey.KeyArea;
+
+
+			// Tell the key area to move the key.  We reset the key index as a result of the move because moving a key can change it's internal index 
+			KeyArea->MoveKey( SelectedKey.KeyHandle.GetValue(), TimeDelta );
+
+			// Update the key that moved
+			float NewKeyTime = KeyArea->GetKeyTime( SelectedKey.KeyHandle.GetValue() );
+
+			// If the key moves outside of the section resize the section to fit the key
+			// @todo Sequencer - Doesn't account for hitting other sections 
+			if( NewKeyTime > Section->GetEndTime() )
+			{
+				Section->SetEndTime( NewKeyTime );
+			}
+			else if( NewKeyTime < Section->GetStartTime() )
+			{
+				Section->SetStartTime( NewKeyTime );
+			}
+
+			if (PrevNewKeyTime == FLT_MAX)
+			{
+				PrevNewKeyTime = NewKeyTime;
+			}
+			else if (!FMath::IsNearlyEqual(NewKeyTime, PrevNewKeyTime))
+			{
+				PrevNewKeyTime = -FLT_MAX;
+			}
+		}
+
+		// Snap the play time to the new dragged key time if all the keyframes were dragged to the same time
+		if (Settings->GetSnapPlayTimeToDraggedKey() && PrevNewKeyTime != FLT_MAX && PrevNewKeyTime != -FLT_MAX)
+		{
+			Sequencer.SetGlobalTime(PrevNewKeyTime);
+		}
+	}
+}
+
+void FMoveKeys::GetKeySnapTimes(TArray<float>& OutSnapTimes, TSharedPtr<FTrackNode> SequencerNode)
+{
+	// snap to non-selected keys
+	TArray< TSharedRef<FSectionKeyAreaNode> > KeyAreaNodes;
+	SequencerNode->GetChildKeyAreaNodesRecursively(KeyAreaNodes);
+
+	for (int32 NodeIndex = 0; NodeIndex < KeyAreaNodes.Num(); ++NodeIndex)
+	{
+		TArray< TSharedRef<IKeyArea> > KeyAreas = KeyAreaNodes[NodeIndex]->GetAllKeyAreas();
+		for (int32 KeyAreaIndex = 0; KeyAreaIndex < KeyAreas.Num(); ++KeyAreaIndex)
+		{
+			auto KeyArea = KeyAreas[KeyAreaIndex];
+
+			TArray<FKeyHandle> KeyHandles = KeyArea->GetUnsortedKeyHandles();
+			for( int32 KeyIndex = 0; KeyIndex < KeyHandles.Num(); ++KeyIndex )
+			{
+				FKeyHandle KeyHandle = KeyHandles[KeyIndex];
+				bool bKeyIsSnappable = true;
+				for ( FSelectedKey SelectedKey : SelectedKeys )
+				{
+					if (SelectedKey.KeyArea == KeyArea && SelectedKey.KeyHandle.GetValue() == KeyHandle)
+					{
+						bKeyIsSnappable = false;
+						break;
+					}
+				}
+				if (bKeyIsSnappable)
+				{
+					OutSnapTimes.Add(KeyArea->GetKeyTime(KeyHandle));
+				}
+			}
+		}
+	}
+}
diff --git a/Engine/Source/Editor/Sequencer/Private/Tools/EditToolDragOperations.h b/Engine/Source/Editor/Sequencer/Private/Tools/EditToolDragOperations.h
new file mode 100644
--- /dev/null
+++ b/Engine/Source/Editor/Sequencer/Private/Tools/EditToolDragOperations.h
@@ -0,0 +1,117 @@
+// Copyright 1998-2015 Epic Games, Inc. All Rights Reserved.
+
+#pragma once
+
+#include "ISequencerEditTool.h"
+#include "ScopedTransaction.h"
+#include "SequencerHotspots.h"
+
+class USequencerSettings;
+
+/**
+ * A drag operation that handles an operation for an edit tool
+ */
+class FEditToolDragOperation : public IEditToolDragOperation
+{
+public:
+	FEditToolDragOperation( FSequencer& InSequencer );
+
+protected:
+
+	/** begin a new scoped transaction for this drag */
+	void BeginTransaction( const TArray< FSectionHandle >& Sections, const FText& TransactionDesc );
+
+	/** End an existing scoped transaction if one exists */
+	void EndTransaction();
+
+protected:
+
+	/** Scoped transaction for this drag operation */
+	TUniquePtr<FScopedTransaction> Transaction;
+
+	/** The current sequencer settings, cached on construction */
+	const USequencerSettings* Settings;
+
+	/** Reference to the sequencer itself */
+	FSequencer& Sequencer;
+};
+
+/**
+ * An operation to resize a section by dragging its left or right edge
+ */
+class FResizeSection : public FEditToolDragOperation
+{
+public:
+	FResizeSection( FSequencer& Sequencer, TArray<FSectionHandle> Sections, TOptional<FSectionHandle> InCardinalSection, bool bInDraggingByEnd );
+
+	/** FEditToolDragOperation interface */
+	virtual void OnBeginDrag(const FVector2D& LocalMousePos, const FVirtualTrackArea& VirtualTrackArea) override;
+	virtual void OnEndDrag() override;
+	virtual void OnDrag( const FPointerEvent& MouseEvent, const FVector2D& LocalMousePos, const FVirtualTrackArea& VirtualTrackArea ) override;
+	virtual FCursorReply GetCursor() const override { return FCursorReply::Cursor( EMouseCursor::ResizeLeftRight ); }
+private:
+	/** The sections we are interacting with */
+	TArray<FSectionHandle> Sections;
+	/** Optional cardinal section that the user is directly interacting with */
+	TOptional<FSectionHandle> CardinalSection;
+	/** true if dragging  the end of the section, false if dragging the start */
+	bool bDraggingByEnd;
+	/** Time where the mouse is pressed */
+	float MouseDownTime;
+	/** The section start or end times when the mouse is pressed */
+	TMap<TWeakObjectPtr<UMovieSceneSection>, float> SectionInitTimes;
+	/** The exact key handles that we're dragging */
+	TSet<FKeyHandle> DraggedKeyHandles;
+};
+
+/** Operation to move the currently selected sections */
+class FMoveSection : public FEditToolDragOperation
+{
+public:
+	FMoveSection( FSequencer& Sequencer, TArray<FSectionHandle> Sections );
+
+	/** FEditToolDragOperation interface */
+	virtual void OnBeginDrag(const FVector2D& LocalMousePos, const FVirtualTrackArea& VirtualTrackArea) override;
+	virtual void OnEndDrag() override;
+	virtual void OnDrag(const FPointerEvent& MouseEvent, const FVector2D& LocalMousePos, const FVirtualTrackArea& VirtualTrackArea) override;
+	virtual FCursorReply GetCursor() const override { return FCursorReply::Cursor( EMouseCursor::CardinalCross ); }
+
+private:
+	/** The sections we are interacting with */
+	TArray<FSectionHandle> Sections;
+	/** Optional cardinal section that the user is directly interacting with */
+	TOptional<FSectionHandle> CardinalSection;
+	/** Time where the mouse is pressed */
+	float MouseDownTime;
+	/** The exact key handles that we're dragging */
+	TSet<FKeyHandle> DraggedKeyHandles;
+};
+
+/**
+ * Operation to move the currently selected keys
+ */
+class FMoveKeys : public FEditToolDragOperation
+{
+public:
+	FMoveKeys( FSequencer& Sequencer, const TSet<FSelectedKey>& InSelectedKeys, TSharedPtr<FTrackNode> InCardinalTrack )
+		: FEditToolDragOperation(Sequencer)
+		, SelectedKeys( InSelectedKeys )
+		, CardinalTrack(MoveTemp(InCardinalTrack))
+	{}
+
+	/** FEditToolDragOperation interface */
+	virtual void OnBeginDrag(const FVector2D& LocalMousePos, const FVirtualTrackArea& VirtualTrackArea) override;
+	virtual void OnEndDrag() override;
+	virtual void OnDrag( const FPointerEvent& MouseEvent, const FVector2D& LocalMousePos, const FVirtualTrackArea& VirtualTrackArea ) override;
+
+protected:
+
+	/** Gets all potential snap times for this drag operation */
+	void GetKeySnapTimes(TArray<float>& OutSnapTimes, TSharedPtr<FTrackNode> SequencerNode);
+
+private:
+	/** The selected keys being moved. */
+	const TSet<FSelectedKey>& SelectedKeys;
+	/** Optional cardinal track to which the dragged key belongs */
+	TSharedPtr<FTrackNode> CardinalTrack;
+};
diff --git a/Engine/Source/Editor/Sequencer/Private/Tools/SequencerEditTool_Default.h b/Engine/Source/Editor/Sequencer/Private/Tools/SequencerEditTool_Default.h
new file mode 100644
--- /dev/null
+++ b/Engine/Source/Editor/Sequencer/Private/Tools/SequencerEditTool_Default.h
@@ -0,0 +1,12 @@
+// Copyright 1998-2015 Epic Games, Inc. All Rights Reserved.
+
+#pragma once
+
+#include "ISequencerEditTool.h"
+
+class FSequencerEditTool_Default : public ISequencerEditTool
+{
+public:
+
+	virtual FReply OnMouseButtonUp(SWidget& OwnerWidget, const FGeometry& MyGeometry, const FPointerEvent& MouseEvent) override;
+};
\ No newline at end of file
diff --git a/Engine/Source/Editor/Sequencer/Private/Tools/SequencerEditTool_Movement.cpp b/Engine/Source/Editor/Sequencer/Private/Tools/SequencerEditTool_Movement.cpp
new file mode 100644
--- /dev/null
+++ b/Engine/Source/Editor/Sequencer/Private/Tools/SequencerEditTool_Movement.cpp
@@ -0,0 +1,178 @@
+// Copyright 1998-2015 Epic Games, Inc. All Rights Reserved.
+
+#include "SequencerPrivatePCH.h"
+#include "SequencerEditTool_Movement.h"
+#include "Sequencer.h"
+#include "VirtualTrackArea.h"
+#include "SequencerHotspots.h"
+#include "EditToolDragOperations.h"
+
+FReply FSequencerEditTool_Default::OnMouseButtonUp(SWidget& OwnerWidget, const FGeometry& MyGeometry, const FPointerEvent& MouseEvent)
+{
+	FSequencer& Sequencer = static_cast<FSequencer&>(GetSequencer());
+	Sequencer.GetSelection().EmptySelectedSections();
+	Sequencer.GetSelection().EmptySelectedKeys();
+	return FReply::Unhandled();
+}
+
+FSequencerEditTool_Movement::FSequencerEditTool_Movement(TSharedPtr<FSequencer> InSequencer, TSharedPtr<SSequencer> InSequencerWidget)
+	: Sequencer(InSequencer)
+	, SequencerWidget(InSequencerWidget)
+{
+}
+
+ISequencer& FSequencerEditTool_Movement::GetSequencer() const
+{
+	return *Sequencer.Pin();
+}
+
+FReply FSequencerEditTool_Movement::OnMouseButtonDown(SWidget& OwnerWidget, const FGeometry& MyGeometry, const FPointerEvent& MouseEvent)
+{
+	DelayedDrag.Reset();
+
+	if (MouseEvent.GetEffectingButton() == EKeys::LeftMouseButton)
+	{
+
+		DelayedDrag = FDelayedDrag_Hotspot(MyGeometry.AbsoluteToLocal(MouseEvent.GetScreenSpacePosition()), EKeys::LeftMouseButton, Hotspot);
+		return FReply::Handled();
+	}
+	return FReply::Unhandled();
+}
+
+FReply FSequencerEditTool_Movement::OnMouseMove(SWidget& OwnerWidget, const FGeometry& MyGeometry, const FPointerEvent& MouseEvent)
+{
+	if (DelayedDrag.IsSet())
+	{
+		const FVirtualTrackArea VirtualTrackArea = SequencerWidget.Pin()->GetVirtualTrackArea(MyGeometry);
+
+		FReply Reply = FReply::Handled();
+
+		if (DelayedDrag->IsDragging())
+		{
+			// If we're already dragging, just update the drag op if it exists
+			if (DragOperation.IsValid())
+			{
+				DragOperation->OnDrag(MouseEvent, MyGeometry.AbsoluteToLocal(MouseEvent.GetScreenSpacePosition()), VirtualTrackArea);
+			}
+		}
+		// Otherwise we can attempt a new drag
+		else if (DelayedDrag->AttemptDragStart(MouseEvent))
+		{
+			DragOperation = CreateDrag();
+
+			if (DragOperation.IsValid())
+			{
+				DragOperation->OnBeginDrag(DelayedDrag->GetInitialPosition(), VirtualTrackArea);
+
+				// Steal the capture, as we're now the authoritative widget in charge of a mouse-drag operation
+				Reply.CaptureMouse(OwnerWidget.AsShared());
+			}
+		}
+
+		return Reply;
+	}
+	return FReply::Unhandled();
+}
+
+TSharedPtr<IEditToolDragOperation> FSequencerEditTool_Movement::CreateDrag()
+{
+	auto PinnedSequencer = Sequencer.Pin();
+	FSequencerSelection& Selection = PinnedSequencer->GetSelection();
+
+	if (DelayedDrag->Hotspot.IsValid())
+	{
+		// Let the hotspot start a drag first, if it wants to
+		auto HotspotDrag = DelayedDrag->Hotspot->InitiateDrag(*PinnedSequencer);
+		if (HotspotDrag.IsValid())
+		{
+			return HotspotDrag;
+		}
+
+		// Ok, the hotspot doesn't know how to drag - let's decide for ourselves
+		auto HotspotType = DelayedDrag->Hotspot->GetType();
+
+		// Moving section(s)?
+		if (HotspotType == ESequencerHotspot::Section)
+		{
+			TArray<FSectionHandle> SectionHandles;
+
+			FSectionHandle& ThisHandle = StaticCastSharedPtr<FSectionHotspot>(DelayedDrag->Hotspot)->Section;
+			UMovieSceneSection* ThisSection = ThisHandle.GetSectionObject();
+			if (Selection.IsSelected(ThisSection))
+			{
+				SectionHandles = SequencerWidget.Pin()->GetSectionHandles(Selection.GetSelectedSections());
+			}
+			else
+			{
+				Selection.EmptySelectedKeys();
+				Selection.EmptySelectedSections();
+				Selection.AddToSelection(ThisSection);
+				SectionHandles.Add(ThisHandle);
+			}
+			return MakeShareable( new FMoveSection( *PinnedSequencer, SectionHandles ) );
+		}
+		// Moving key(s)?
+		else if (HotspotType == ESequencerHotspot::Key)
+		{
+			FSelectedKey& ThisKey = StaticCastSharedPtr<FKeyHotspot>(DelayedDrag->Hotspot)->Key;
+
+			// If it's not selected, we'll treat this as a unique drag
+			if (!Selection.IsSelected(ThisKey))
+			{
+				Selection.EmptySelectedKeys();
+				Selection.EmptySelectedSections();
+				Selection.AddToSelection(ThisKey);
+			}
+
+			return MakeShareable( new FMoveKeys( *PinnedSequencer, Selection.GetSelectedKeys(), StaticCastSharedPtr<FKeyHotspot>(DelayedDrag->Hotspot)->TrackNode ) );
+		}
+	}
+	// If we're not dragging a hotspot, sections take precedence over keys
+	else if (Selection.GetSelectedSections().Num())
+	{
+		return MakeShareable( new FMoveSection( *PinnedSequencer, SequencerWidget.Pin()->GetSectionHandles(Selection.GetSelectedSections()) ) );
+	}
+	else if (Selection.GetSelectedKeys().Num())
+	{
+		return MakeShareable( new FMoveKeys( *PinnedSequencer, Selection.GetSelectedKeys(), nullptr ) );
+	}
+
+	return nullptr;
+}
+
+FReply FSequencerEditTool_Movement::OnMouseButtonUp(SWidget& OwnerWidget, const FGeometry& MyGeometry, const FPointerEvent& MouseEvent)
+{
+	DelayedDrag.Reset();
+
+	if (DragOperation.IsValid())
+	{
+		DragOperation->OnEndDrag();
+		DragOperation = nullptr;
+
+		// Only return handled if we actually started a drag
+		return FReply::Handled().ReleaseMouseCapture();
+	}
+
+	return FSequencerEditTool_Default::OnMouseButtonUp(OwnerWidget, MyGeometry, MouseEvent);
+}
+
+void FSequencerEditTool_Movement::OnMouseCaptureLost()
+{
+	DelayedDrag.Reset();
+	if (DragOperation.IsValid())
+	{
+		DragOperation->OnEndDrag();
+		DragOperation = nullptr;
+	}
+}
+
+FCursorReply FSequencerEditTool_Movement::OnCursorQuery(const FGeometry& MyGeometry, const FPointerEvent& CursorEvent) const
+{
+	return DragOperation.IsValid() ? FCursorReply::Cursor(EMouseCursor::CardinalCross) : FCursorReply::Cursor(EMouseCursor::Default);
+}
+
+FName FSequencerEditTool_Movement::GetIdentifier() const
+{
+	static FName Identifier("Movement");
+	return Identifier;
+}
\ No newline at end of file
diff --git a/Engine/Source/Editor/Sequencer/Private/Tools/SequencerEditTool_Movement.h b/Engine/Source/Editor/Sequencer/Private/Tools/SequencerEditTool_Movement.h
new file mode 100644
--- /dev/null
+++ b/Engine/Source/Editor/Sequencer/Private/Tools/SequencerEditTool_Movement.h
@@ -0,0 +1,49 @@
+// Copyright 1998-2015 Epic Games, Inc. All Rights Reserved.
+
+#pragma once
+
+#include "SequencerEditTool_Default.h"
+#include "DelayedDrag.h"
+
+class IEditToolDragOperation;
+
+class FSequencerEditTool_Movement : public FSequencerEditTool_Default
+{
+public:
+
+	FSequencerEditTool_Movement(TSharedPtr<FSequencer> InSequencer, TSharedPtr<SSequencer> InSequencerWidget);
+
+	virtual FReply OnMouseButtonDown(SWidget& OwnerWidget, const FGeometry& MyGeometry, const FPointerEvent& MouseEvent) override;
+	virtual FReply OnMouseButtonUp(SWidget& OwnerWidget, const FGeometry& MyGeometry, const FPointerEvent& MouseEvent) override;
+	virtual FReply OnMouseMove(SWidget& OwnerWidget, const FGeometry& MyGeometry, const FPointerEvent& MouseEvent) override;
+	virtual void   OnMouseCaptureLost() override;
+	virtual FCursorReply OnCursorQuery(const FGeometry& MyGeometry, const FPointerEvent& CursorEvent) const override;
+	virtual FName GetIdentifier() const override;
+	virtual ISequencer& GetSequencer() const override;
+
+private:
+
+	TSharedPtr<IEditToolDragOperation> CreateDrag();
+
+	/** The sequencer itself */
+	TWeakPtr<FSequencer> Sequencer;
+
+	/** Sequencer widget */
+	TWeakPtr<SSequencer> SequencerWidget;
+
+	struct FDelayedDrag_Hotspot : FDelayedDrag
+	{
+		FDelayedDrag_Hotspot(FVector2D InInitialPosition, FKey InApplicableKey, TSharedPtr<ISequencerHotspot> InHotspot)
+			: FDelayedDrag(InInitialPosition, InApplicableKey)
+			, Hotspot(MoveTemp(InHotspot))
+		{}
+
+		TSharedPtr<ISequencerHotspot> Hotspot;
+	};
+
+	/** Helper class responsible for handling delayed dragging */
+	TOptional<FDelayedDrag_Hotspot> DelayedDrag;
+
+	/** Current drag operation if any */
+	TSharedPtr<IEditToolDragOperation> DragOperation;
+};
diff --git a/Engine/Source/Editor/Sequencer/Private/Tools/SequencerEditTool_Selection.cpp b/Engine/Source/Editor/Sequencer/Private/Tools/SequencerEditTool_Selection.cpp
new file mode 100644
--- /dev/null
+++ b/Engine/Source/Editor/Sequencer/Private/Tools/SequencerEditTool_Selection.cpp
@@ -0,0 +1,485 @@
+// Copyright 1998-2015 Epic Games, Inc. All Rights Reserved.
+
+#include "SequencerPrivatePCH.h"
+#include "SequencerEditTool_Selection.h"
+#include "MovieSceneSection.h"
+#include "ISequencerSection.h"
+#include "VirtualTrackArea.h"
+#include "Sequencer.h"
+#include "SSequencerTreeView.h"
+
+struct FMarqueeSelectData
+{
+	FMarqueeSelectData(const FVector2D& InInitialPosition)
+		: InitialPosition(InInitialPosition)
+		, CurrentPosition(InInitialPosition)
+	{}
+
+	FVector2D TopLeft() const
+	{
+		return FVector2D(
+			FMath::Min(InitialPosition.X, CurrentPosition.X),
+			FMath::Min(InitialPosition.Y, CurrentPosition.Y)
+		);
+	}
+	
+	FVector2D BottomRight() const
+	{
+		return FVector2D(
+			FMath::Max(InitialPosition.X, CurrentPosition.X),
+			FMath::Max(InitialPosition.Y, CurrentPosition.Y)
+		);
+	}
+
+	FVector2D Size() const
+	{
+		return BottomRight() - TopLeft();
+	}
+
+	bool HasDragged() const
+	{
+		return !Size().IsNearlyZero();
+	}
+
+	FVector2D InitialPosition;
+	FVector2D CurrentPosition;
+};
+
+
+struct IMarqueeVisitor
+{
+	virtual void VisitKey(FKeyHandle KeyHandle, float KeyTime, const TSharedPtr<IKeyArea>& KeyArea, UMovieSceneSection* Section) const { }
+	virtual void VisitSection(UMovieSceneSection* Section) const { }
+};
+
+/** Struct used to iterate a marquee range with a user-supplied visitor */
+/* @todo: This could probably live in a header so it can be used elsewhere */
+/* @todo: Could probably optimize this by not walking every single node, and binary searching the begin<->end ranges instead */
+struct FMarqueeRange
+{
+	/** Construction from the range itself */
+	FMarqueeRange(FVector2D InTopLeft, FVector2D InBottomRight, FVector2D InVirtualKeySize)
+		: TopLeft(InTopLeft), BottomRight(InBottomRight), VirtualKeySize(InVirtualKeySize)
+	{}
+
+	/** Visit the specified nodes (recursively) with this range and a user-supplied visitor */
+	void Visit(const IMarqueeVisitor& Visitor, const TArray< TSharedRef<FSequencerDisplayNode> >& Nodes);
+
+private:
+
+	/** Handle visitation of a particular node */
+	void HandleNode(const IMarqueeVisitor& Visitor, FSequencerDisplayNode& InNode);
+	/** Handle visitation of a particular node, along with a set of sections */
+	void HandleNode(const IMarqueeVisitor& Visitor, FSequencerDisplayNode& InNode, const TArray<TSharedRef<ISequencerSection>>& InSections);
+	/** Handle visitation of a key area node */
+	void HandleKeyAreaNode(const IMarqueeVisitor& Visitor, FSectionKeyAreaNode& InNode, const TArray<TSharedRef<ISequencerSection>>& InSections);
+	/** Handle visitation of a key area */
+	void HandleKeyArea(const IMarqueeVisitor& Visitor, const TSharedPtr<IKeyArea>& KeyArea, UMovieSceneSection* Section);
+	
+	/** Check whether the specified section intersects this range */
+	bool SectionIntersectsRange(const UMovieSceneSection* InSection) const;
+
+	/** The bounds of the range */
+	FVector2D TopLeft, BottomRight;
+
+	/** Key size in virtual space */
+	FVector2D VirtualKeySize;
+};
+
+bool FMarqueeRange::SectionIntersectsRange(const UMovieSceneSection* InSection) const
+{
+	return InSection->GetStartTime() <= BottomRight.X && InSection->GetEndTime() >= TopLeft.X;
+}
+
+void FMarqueeRange::Visit(const IMarqueeVisitor& Visitor, const TArray< TSharedRef<FSequencerDisplayNode> >& Nodes)
+{
+	for (auto& Child : Nodes)
+	{
+		if (!Child->IsHidden())
+		{
+			HandleNode(Visitor, *Child);
+		}
+	}
+}
+
+void FMarqueeRange::HandleNode(const IMarqueeVisitor& Visitor, FSequencerDisplayNode& InNode)
+{
+	if (InNode.GetType() == ESequencerNode::Track)
+	{
+		HandleNode(Visitor, InNode, static_cast<FTrackNode&>(InNode).GetSections());
+	}
+
+	if (InNode.IsExpanded())
+	{
+		for (auto& Child : InNode.GetChildNodes())
+		{
+			if (!Child->IsHidden())
+			{
+				HandleNode(Visitor, *Child);
+			}
+		}
+	}
+}
+
+void FMarqueeRange::HandleNode(const IMarqueeVisitor& Visitor, FSequencerDisplayNode& InNode, const TArray<TSharedRef<ISequencerSection>>& InSections)
+{
+	const float NodeCenter = InNode.GetVirtualTop() + (InNode.GetVirtualBottom() - InNode.GetVirtualTop())/2;
+	if (NodeCenter + VirtualKeySize.Y/2 > TopLeft.Y &&
+		NodeCenter - VirtualKeySize.Y/2 < BottomRight.Y)
+	{
+		bool bNodeHasKeyArea = false;
+		if (InNode.GetType() == ESequencerNode::KeyArea)
+		{
+			HandleKeyAreaNode(Visitor, static_cast<FSectionKeyAreaNode&>(InNode), InSections);
+			bNodeHasKeyArea = true;
+		}
+		else if (InNode.GetType() == ESequencerNode::Track)
+		{
+			TSharedPtr<FSectionKeyAreaNode> SectionKeyNode = static_cast<FTrackNode&>(InNode).GetTopLevelKeyNode();
+			if (SectionKeyNode.IsValid())
+			{
+				HandleKeyAreaNode(Visitor, static_cast<FSectionKeyAreaNode&>(InNode), InSections);
+				bNodeHasKeyArea = true;
+			}
+		}
+
+		if (!InNode.IsExpanded() && !bNodeHasKeyArea)
+		{
+			// As a fallback, handle key groupings on collapsed parents
+			for (int32 SectionIndex = 0; SectionIndex < InSections.Num(); ++SectionIndex)
+			{
+				UMovieSceneSection* Section = InSections[SectionIndex]->GetSectionObject();
+				if (SectionIntersectsRange(Section))
+				{
+					Visitor.VisitSection(Section);
+
+					TSharedRef<IKeyArea> KeyArea = InNode.UpdateKeyGrouping(SectionIndex);
+					HandleKeyArea(Visitor, KeyArea, Section);
+				}
+			}
+		}
+	}
+
+	if (InNode.IsExpanded())
+	{
+		// Handle Children
+		for (auto& Child : InNode.GetChildNodes())
+		{
+			if (!Child->IsHidden())
+			{
+				HandleNode(Visitor, *Child, InSections);
+			}
+		}
+	}
+}
+
+void FMarqueeRange::HandleKeyAreaNode(const IMarqueeVisitor& Visitor, FSectionKeyAreaNode& KeyAreaNode, const TArray<TSharedRef<ISequencerSection>>& InSections)
+{
+	for( int32 SectionIndex = 0; SectionIndex < InSections.Num(); ++SectionIndex )
+	{
+		UMovieSceneSection* Section = InSections[SectionIndex]->GetSectionObject();
+
+		// If the section is at all within the marquee, we check its keys
+		if (SectionIntersectsRange(Section))
+		{
+			Visitor.VisitSection(Section);
+
+			TSharedRef<IKeyArea> KeyArea = KeyAreaNode.GetKeyArea(SectionIndex);
+			HandleKeyArea(Visitor, KeyArea, InSections[SectionIndex]->GetSectionObject());
+		}
+	}
+}
+
+void FMarqueeRange::HandleKeyArea(const IMarqueeVisitor& Visitor, const TSharedPtr<IKeyArea>& KeyArea, UMovieSceneSection* Section)
+{
+	for (FKeyHandle KeyHandle : KeyArea->GetUnsortedKeyHandles())
+	{
+		float KeyPosition = KeyArea->GetKeyTime(KeyHandle);
+		if (KeyPosition + VirtualKeySize.X/2 > TopLeft.X &&
+			KeyPosition - VirtualKeySize.X/2 < BottomRight.X)
+		{
+			Visitor.VisitKey(KeyHandle, KeyPosition, KeyArea, Section);
+		}
+	}
+}
+
+/** todo: Save this in a config? */
+ESequencerSelectionMode LastUsedMode = ESequencerSelectionMode::Keys;
+
+FSequencerEditTool_Selection::FSequencerEditTool_Selection(TSharedPtr<FSequencer> InSequencer, TSharedPtr<SSequencer> InSequencerWidget)
+	: Sequencer(InSequencer)
+	, SequencerWidget(InSequencerWidget)
+	, SelectionMode(LastUsedMode)
+{
+}
+FSequencerEditTool_Selection::FSequencerEditTool_Selection(TSharedPtr<FSequencer> InSequencer, TSharedPtr<SSequencer> InSequencerWidget, ESequencerSelectionMode ExplicitMode)
+	: Sequencer(InSequencer)
+	, SequencerWidget(InSequencerWidget)
+	, SelectionMode(ExplicitMode)
+{
+	LastUsedMode = SelectionMode;
+}
+
+ISequencer& FSequencerEditTool_Selection::GetSequencer() const
+{
+	return *Sequencer.Pin();
+}
+
+void FSequencerEditTool_Selection::SetSelectionMode(ESequencerSelectionMode Mode)
+{
+	SelectionMode = Mode;
+	LastUsedMode = SelectionMode;
+}
+
+ESequencerSelectionMode FSequencerEditTool_Selection::GetGlobalSelectionMode()
+{
+	return LastUsedMode;
+}
+
+FCursorReply FSequencerEditTool_Selection::OnCursorQuery(const FGeometry& MyGeometry, const FPointerEvent& CursorEvent) const
+{
+	return FCursorReply::Cursor(EMouseCursor::Crosshairs);
+}
+
+int32 FSequencerEditTool_Selection::OnPaint(const FGeometry& AllottedGeometry, const FSlateRect& MyClippingRect, FSlateWindowElementList& OutDrawElements, int32 LayerId) const
+{
+	if (MarqueeSelectData.IsValid() && MarqueeSelectData->HasDragged())
+	{
+		// Convert to physical space for rendering
+		const FVirtualTrackArea VirtualTrackArea = SequencerWidget.Pin()->GetVirtualTrackArea(AllottedGeometry);
+
+		FVector2D TopLeft = VirtualTrackArea.VirtualToPhysical(MarqueeSelectData->TopLeft());
+		FVector2D BottomRight = VirtualTrackArea.VirtualToPhysical(MarqueeSelectData->BottomRight());
+
+		FSlateDrawElement::MakeBox(
+			OutDrawElements,
+			LayerId,
+			AllottedGeometry.ToPaintGeometry(TopLeft, BottomRight - TopLeft),
+			FEditorStyle::GetBrush(TEXT("MarqueeSelection")),
+			MyClippingRect
+			);
+
+		return LayerId + 1;
+	}
+
+	return LayerId;
+}
+
+FReply FSequencerEditTool_Selection::OnMouseButtonDown(SWidget& OwnerWidget, const FGeometry& MyGeometry, const FPointerEvent& MouseEvent)
+{
+	DelayedDrag.Reset();
+
+	if (MouseEvent.GetEffectingButton() == EKeys::LeftMouseButton)
+	{
+		DelayedDrag = FDelayedDrag(MyGeometry.AbsoluteToLocal(MouseEvent.GetScreenSpacePosition()), EKeys::LeftMouseButton);
+		return FReply::Handled();
+	}
+	return FReply::Unhandled();
+}
+
+FReply FSequencerEditTool_Selection::OnMouseMove(SWidget& OwnerWidget, const FGeometry& MyGeometry, const FPointerEvent& MouseEvent)
+{
+	if (DelayedDrag.IsSet())
+	{
+		FReply Reply = FReply::Handled();
+
+		if (MarqueeSelectData.IsValid())
+		{
+			SetNewMarqueeBounds(MouseEvent, MyGeometry);
+		}
+		else if (DelayedDrag->AttemptDragStart(MouseEvent))
+		{
+			const FVirtualTrackArea VirtualTrackArea = SequencerWidget.Pin()->GetVirtualTrackArea(MyGeometry);
+
+			// Marquee data is stored in virtual space
+			MarqueeSelectData.Reset(new FMarqueeSelectData(VirtualTrackArea.PhysicalToVirtual(DelayedDrag->GetInitialPosition())));
+
+			// Steal the capture, as we're now the authoritative widget in charge of a mouse-drag operation
+			Reply.CaptureMouse(OwnerWidget.AsShared());
+		}
+
+		return Reply;
+	}
+	return FReply::Unhandled();
+}
+
+FReply FSequencerEditTool_Selection::OnMouseButtonUp(SWidget& OwnerWidget, const FGeometry& MyGeometry, const FPointerEvent& MouseEvent)
+{
+	DelayedDrag.Reset();
+	if (MarqueeSelectData.IsValid())
+	{
+		HandleMarqueeSelection(MyGeometry, MouseEvent);
+		MarqueeSelectData.Reset();
+
+		Sequencer.Pin()->StopAutoscroll();
+		return FReply::Handled().ReleaseMouseCapture();
+	}
+
+	return FSequencerEditTool_Default::OnMouseButtonUp(OwnerWidget, MyGeometry, MouseEvent);
+}
+
+void FSequencerEditTool_Selection::OnMouseCaptureLost()
+{
+	DelayedDrag.Reset();
+	MarqueeSelectData.Reset();
+}
+
+FName FSequencerEditTool_Selection::GetIdentifier() const
+{
+	static FName Identifier("Selection");
+	return Identifier;
+}
+
+void FSequencerEditTool_Selection::SetNewMarqueeBounds(const FPointerEvent& MouseEvent, const FGeometry& MyGeometry)
+{
+	const FVector2D MouseDelta = MouseEvent.GetCursorDelta();
+
+	FVector2D LocalPosition = MyGeometry.AbsoluteToLocal(MouseEvent.GetScreenSpacePosition());
+
+	auto PinnedSequencer = Sequencer.Pin();
+
+	// Handle virtual scrolling when at the vertical extremes of the widget (performed before we clamp the mouse pos)
+	{
+		const float ScrollThresholdV = MyGeometry.Size.Y * 0.025f;
+
+		float Difference = LocalPosition.Y - ScrollThresholdV;
+		if (Difference < 0 && MouseDelta.Y < 0)
+		{
+			PinnedSequencer->VerticalScroll( Difference * 0.1f );
+		}
+
+		Difference = LocalPosition.Y - (MyGeometry.Size.Y - ScrollThresholdV);
+		if (Difference > 0 && MouseDelta.Y > 0)
+		{
+			PinnedSequencer->VerticalScroll( Difference * 0.1f );
+		}
+	}
+
+	const FVirtualTrackArea VirtualTrackArea = SequencerWidget.Pin()->GetVirtualTrackArea(MyGeometry);
+
+	// Clamp the vertical position to the actual bounds of the track area
+	LocalPosition.Y = FMath::Clamp(LocalPosition.Y, 0.f, MyGeometry.GetLocalSize().Y);
+	MarqueeSelectData->CurrentPosition = VirtualTrackArea.PhysicalToVirtual(LocalPosition);
+
+	// Handle virtual scrolling when at the horizontal extremes of the widget
+	{
+		TRange<float> ViewRange = PinnedSequencer->GetViewRange();
+		const float ScrollThresholdH = ViewRange.Size<float>() * 0.025f;
+
+		float Difference = MarqueeSelectData->CurrentPosition.X - (ViewRange.GetLowerBoundValue() + ScrollThresholdH);
+		if (Difference < 0 && MouseDelta.X < 0)
+		{
+			PinnedSequencer->StartAutoscroll(Difference);
+		}
+		else
+		{
+			Difference = MarqueeSelectData->CurrentPosition.X - (ViewRange.GetUpperBoundValue() - ScrollThresholdH);
+			if (Difference > 0 && MouseDelta.X > 0)
+			{
+				PinnedSequencer->StartAutoscroll(Difference);
+			}
+			else
+			{
+				PinnedSequencer->StopAutoscroll();
+			}
+		}
+	}
+}
+
+struct FSelectionVisitor : IMarqueeVisitor
+{
+	FSelectionVisitor(FSequencerSelection& InSelection, const FPointerEvent& InMouseEvent, ESequencerSelectionMode InMode)
+		: Selection(InSelection)
+		, MouseEvent(InMouseEvent)
+		, Mode(InMode)
+	{}
+
+	virtual void VisitKey(FKeyHandle KeyHandle, float KeyTime, const TSharedPtr<IKeyArea>& KeyArea, UMovieSceneSection* Section) const override
+	{
+		if (Mode != ESequencerSelectionMode::Keys)
+		{
+			return;
+		}
+
+		FSelectedKey SelectedKey(*Section, KeyArea, KeyHandle);
+
+		if (MouseEvent.IsLeftShiftDown())
+		{
+			Selection.AddToSelection(SelectedKey);
+		}
+		else if (MouseEvent.IsLeftControlDown())
+		{
+			if (Selection.IsSelected(SelectedKey))
+			{
+				Selection.RemoveFromSelection(SelectedKey);
+			}
+			else
+			{
+				Selection.AddToSelection(SelectedKey);
+			}
+		}
+		else
+		{
+			Selection.AddToSelection(SelectedKey);
+		}
+	}
+
+	virtual void VisitSection(UMovieSceneSection* Section) const
+	{
+		if (Mode != ESequencerSelectionMode::Sections)
+		{
+			return;
+		}
+
+		if (MouseEvent.IsLeftShiftDown())
+		{
+			Selection.AddToSelection(Section);
+		}
+		else if (MouseEvent.IsLeftControlDown())
+		{
+			if (Selection.IsSelected(Section))
+			{
+				Selection.RemoveFromSelection(Section);
+			}
+			else
+			{
+				Selection.AddToSelection(Section);
+			}
+		}
+		else
+		{
+			Selection.AddToSelection(Section);
+		}
+	}
+
+private:
+	FSequencerSelection& Selection;
+	const FPointerEvent& MouseEvent;
+	ESequencerSelectionMode Mode;
+};
+
+void FSequencerEditTool_Selection::HandleMarqueeSelection(const FGeometry& Geometry, const FPointerEvent& MouseEvent)
+{
+	auto PinnedSequencer = Sequencer.Pin();
+
+	auto& Selection = PinnedSequencer->GetSelection();
+	if (!MouseEvent.IsLeftControlDown() && !MouseEvent.IsLeftShiftDown())
+	{
+		Selection.Empty();
+	}
+
+	Selection.SuspendBroadcast();
+
+	FVector2D VirtualKeySize;
+	VirtualKeySize.X = SequencerSectionConstants::KeySize.X / Geometry.Size.X * (PinnedSequencer->GetViewRange().GetUpperBoundValue() - PinnedSequencer->GetViewRange().GetLowerBoundValue());
+	// Vertically, virtual units == physical units
+	VirtualKeySize.Y = SequencerSectionConstants::KeySize.Y;
+
+	FMarqueeRange Range(MarqueeSelectData->TopLeft(), MarqueeSelectData->BottomRight(), VirtualKeySize);
+	FSelectionVisitor Visitor(Selection, MouseEvent, SelectionMode);
+
+	Range.Visit(Visitor, SequencerWidget.Pin()->GetTreeView()->GetNodeTree()->GetRootNodes());
+
+	Selection.ResumeBroadcast();
+	Selection.GetOnKeySelectionChanged().Broadcast();
+}
\ No newline at end of file
diff --git a/Engine/Source/Editor/Sequencer/Private/Tools/SequencerEditTool_Selection.h b/Engine/Source/Editor/Sequencer/Private/Tools/SequencerEditTool_Selection.h
new file mode 100644
--- /dev/null
+++ b/Engine/Source/Editor/Sequencer/Private/Tools/SequencerEditTool_Selection.h
@@ -0,0 +1,64 @@
+// Copyright 1998-2015 Epic Games, Inc. All Rights Reserved.
+
+#pragma once
+
+#include "SequencerEditTool_Default.h"
+#include "UniquePtr.h"
+#include "DelayedDrag.h"
+
+struct FMarqueeSelectData;
+
+enum class ESequencerSelectionMode
+{
+	Keys, Sections
+};
+
+class FSequencerEditTool_Selection : public FSequencerEditTool_Default
+{
+public:
+
+	FSequencerEditTool_Selection(TSharedPtr<FSequencer> InSequencer, TSharedPtr<SSequencer> InSequencerWidget);
+	FSequencerEditTool_Selection(TSharedPtr<FSequencer> InSequencer, TSharedPtr<SSequencer> InSequencerWidget, ESequencerSelectionMode ExplicitMode);
+
+	virtual ISequencer& GetSequencer() const override;
+	virtual int32 OnPaint(const FGeometry& AllottedGeometry, const FSlateRect& MyClippingRect, FSlateWindowElementList& OutDrawElements, int32 LayerId) const override;
+
+	virtual FReply OnMouseButtonDown(SWidget& OwnerWidget, const FGeometry& MyGeometry, const FPointerEvent& MouseEvent) override;
+	virtual FReply OnMouseButtonUp(SWidget& OwnerWidget, const FGeometry& MyGeometry, const FPointerEvent& MouseEvent) override;
+	virtual FReply OnMouseMove(SWidget& OwnerWidget, const FGeometry& MyGeometry, const FPointerEvent& MouseEvent) override;
+	virtual void   OnMouseCaptureLost() override;
+	virtual FCursorReply OnCursorQuery(const FGeometry& MyGeometry, const FPointerEvent& CursorEvent) const override;
+	virtual FName GetIdentifier() const override;
+
+public:
+
+	static ESequencerSelectionMode GetGlobalSelectionMode();
+	
+	ESequencerSelectionMode GetSelectionMode() const { return SelectionMode; }
+	void SetSelectionMode(ESequencerSelectionMode Mode);
+	
+private:
+
+	/** Set new bounds for the marquee selection */
+	void SetNewMarqueeBounds(const FPointerEvent& MouseEvent, const FGeometry& Geometry);
+
+	/** Select all the keys contained in the current marquee */
+	void HandleMarqueeSelection(const FGeometry& Geometry, const FPointerEvent& MouseEvent);
+
+private:
+
+	/** The sequencer itself */
+	TWeakPtr<FSequencer> Sequencer;
+
+	/** Sequencer widget */
+	TWeakPtr<SSequencer> SequencerWidget;
+
+	/** Helper class responsible for handling delayed dragging */
+	TOptional<FDelayedDrag> DelayedDrag;
+
+	/** Optional marquee selection data, valid when dragging */
+	TUniquePtr<FMarqueeSelectData> MarqueeSelectData;
+
+	/** What we are selecting */
+	ESequencerSelectionMode SelectionMode;
+};
diff --git a/Engine/Source/Editor/Sequencer/Private/VirtualTrackArea.cpp b/Engine/Source/Editor/Sequencer/Private/VirtualTrackArea.cpp
new file mode 100644
--- /dev/null
+++ b/Engine/Source/Editor/Sequencer/Private/VirtualTrackArea.cpp
@@ -0,0 +1,153 @@
+// Copyright 1998-2015 Epic Games, Inc. All Rights Reserved.
+
+#include "SequencerPrivatePCH.h"
+#include "VirtualTrackArea.h"
+#include "Sequencer.h"
+#include "SSequencerTreeView.h"
+#include "SequencerHotspots.h"
+#include "MovieSceneSection.h"
+
+FVirtualTrackArea::FVirtualTrackArea(const FSequencer& InSequencer, SSequencerTreeView& InTreeView, const FGeometry& InTrackAreaGeometry)
+	: FTimeToPixel(InTrackAreaGeometry, InSequencer.GetViewRange())
+	, TreeView(InTreeView)
+{
+}
+
+float FVirtualTrackArea::PixelToVerticalOffset(float InPixel) const
+{
+	return TreeView.PhysicalToVirtual(InPixel);
+}
+
+float FVirtualTrackArea::VerticalOffsetToPixel(float InOffset) const
+{
+	return TreeView.VirtualToPhysical(InOffset);
+}
+
+FVector2D FVirtualTrackArea::PhysicalToVirtual(FVector2D InPosition) const
+{
+	InPosition.Y = PixelToVerticalOffset(InPosition.Y);
+	InPosition.X = PixelToTime(InPosition.X);
+
+	return InPosition;
+}
+
+FVector2D FVirtualTrackArea::VirtualToPhysical(FVector2D InPosition) const
+{
+	InPosition.Y = VerticalOffsetToPixel(InPosition.Y);
+	InPosition.X = TimeToPixel(InPosition.X);
+
+	return InPosition;
+}
+
+TSharedPtr<FSequencerDisplayNode> FVirtualTrackArea::HitTestNode(float InPhysicalPosition) const
+{
+	return TreeView.HitTestNode(InPhysicalPosition);
+}
+
+TSharedPtr<FTrackNode> GetParentTrackNode(FSequencerDisplayNode& In)
+{
+	FSequencerDisplayNode* Current = &In;
+	while(Current && Current->GetType() != ESequencerNode::Object)
+	{
+		if (Current->GetType() == ESequencerNode::Track)
+		{
+			return StaticCastSharedRef<FTrackNode>(Current->AsShared());
+		}
+		Current = Current->GetParent().Get();
+	}
+
+	return nullptr;
+}
+
+TOptional<FSectionHandle> FVirtualTrackArea::HitTestSection(FVector2D InPhysicalPosition) const
+{
+	TSharedPtr<FSequencerDisplayNode> Node = HitTestNode(InPhysicalPosition.Y);
+
+	if (Node.IsValid())
+	{
+		TSharedPtr<FTrackNode> TrackNode = GetParentTrackNode(*Node);
+
+		if (TrackNode.IsValid())
+		{
+			float Time = PixelToTime(InPhysicalPosition.X);
+
+			const auto& Sections = TrackNode->GetSections();
+			for (int32 Index = 0; Index < Sections.Num(); ++Index)
+			{
+				UMovieSceneSection* Section = Sections[Index]->GetSectionObject();
+				if (Section->IsTimeWithinSection(Time))
+				{
+					return FSectionHandle(TrackNode.ToSharedRef(), Index);
+				}
+			}
+		}
+	}
+
+	return TOptional<FSectionHandle>();
+}
+
+FSelectedKey FVirtualTrackArea::HitTestKey(FVector2D InPhysicalPosition) const
+{
+	TSharedPtr<FSequencerDisplayNode> Node = HitTestNode(InPhysicalPosition.Y);
+
+	if (!Node.IsValid())
+	{
+		return FSelectedKey();
+	}
+
+	const float KeyLeft  = PixelToTime(InPhysicalPosition.X - SequencerSectionConstants::KeySize.X/2);
+	const float KeyRight = PixelToTime(InPhysicalPosition.X + SequencerSectionConstants::KeySize.X/2);
+
+	TArray<TSharedRef<IKeyArea>> KeyAreas;
+
+	// First check for a key area node on the hit-tested node
+	TSharedPtr<FSectionKeyAreaNode> KeyAreaNode;
+	switch (Node->GetType())
+	{
+		case ESequencerNode::KeyArea: 	KeyAreaNode = StaticCastSharedPtr<FSectionKeyAreaNode>(Node); break;
+		case ESequencerNode::Track:		KeyAreaNode = StaticCastSharedPtr<FTrackNode>(Node)->GetTopLevelKeyNode(); break;
+	}
+
+	if (KeyAreaNode.IsValid())
+	{
+		for (auto& KeyArea : KeyAreaNode->GetAllKeyAreas())
+		{
+			UMovieSceneSection* Section = KeyArea->GetOwningSection();
+			if (Section->GetStartTime() <= KeyRight && Section->GetEndTime() >= KeyLeft)
+			{
+				KeyAreas.Add(KeyArea);
+			}
+		}
+	}
+	// Failing that, and the node is collapsed, we check for key groupings
+	else if (!Node->IsExpanded())
+	{
+		TSharedPtr<FTrackNode> TrackNode = GetParentTrackNode(*Node);
+		const TArray< TSharedRef<ISequencerSection> >& Sections = TrackNode->GetSections();
+
+		for ( int32 SectionIndex = 0; SectionIndex < Sections.Num(); ++SectionIndex )
+		{
+			UMovieSceneSection* Section = Sections[SectionIndex]->GetSectionObject();
+			if (Section->GetStartTime() <= KeyRight && Section->GetEndTime() >= KeyLeft)
+			{
+				KeyAreas.Add(Node->GetKeyGrouping(SectionIndex));
+			}
+		}
+	}
+
+	// Search for any key that matches the position
+	// todo: investigate if this would be faster as a sort + binary search, rather than linear search
+	for (auto& KeyArea : KeyAreas)
+	{
+		for (FKeyHandle Key : KeyArea->GetUnsortedKeyHandles())
+		{
+			const float Time = KeyArea->GetKeyTime(Key);
+			if (Time >= KeyLeft && Time <= KeyRight)
+			{
+				return FSelectedKey(*KeyArea->GetOwningSection(), KeyArea, Key);
+			}
+		}
+	}
+
+	return FSelectedKey();
+}
\ No newline at end of file
diff --git a/Engine/Source/Editor/Sequencer/Public/ISequencerEditTool.h b/Engine/Source/Editor/Sequencer/Public/ISequencerEditTool.h
new file mode 100644
--- /dev/null
+++ b/Engine/Source/Editor/Sequencer/Public/ISequencerEditTool.h
@@ -0,0 +1,65 @@
+// Copyright 1998-2015 Epic Games, Inc. All Rights Reserved.
+
+#pragma once
+
+class FVirtualTrackArea;
+struct ISequencerHotspot;
+
+class IEditToolDragOperation
+{
+public:
+	virtual ~IEditToolDragOperation(){}
+
+	/**
+	 * Called to initiate a drag
+	 *
+	 * @param LocalMousePos		The current location of the mouse, relative to the top-left corner of the virtual track area
+	 * @param VirtualTrackArea	A virtual track area that can be used for pixel->time conversions and hittesting
+	 */
+	virtual void OnBeginDrag(const FVector2D& LocalMousePos, const FVirtualTrackArea& VirtualTrackArea) = 0;
+
+	/** Called when a drag has ended */
+	virtual void OnEndDrag() = 0;
+
+	/** Request the cursor for this drag operation */
+	virtual FCursorReply GetCursor() const { return FCursorReply::Cursor( EMouseCursor::Default ); }
+
+	/**
+	 * Notification called when the mouse moves while dragging
+	 *
+	 * @param MouseEvent		The associated mouse event for dragging
+	 * @param LocalMousePos		The current location of the mouse, relative to the top-left corner of the virtual track area
+	 * @param VirtualTrackArea	A virtual track area that can be used for pixel->time conversions and hittesting
+	 */
+	virtual void OnDrag( const FPointerEvent& MouseEvent, const FVector2D& LocalMousePos, const FVirtualTrackArea& VirtualTrackArea ) = 0;
+};
+
+class ISequencerEditTool
+{
+public:
+	virtual ~ISequencerEditTool() { }
+
+	virtual FReply OnMouseButtonDown(SWidget& OwnerWidget, const FGeometry& MyGeometry, const FPointerEvent& MouseEvent) { return FReply::Unhandled(); }
+	virtual FReply OnMouseButtonUp(SWidget& OwnerWidget, const FGeometry& MyGeometry, const FPointerEvent& MouseEvent) { return FReply::Unhandled(); }
+	virtual FReply OnMouseMove(SWidget& OwnerWidget, const FGeometry& MyGeometry, const FPointerEvent& MouseEvent) { return FReply::Unhandled(); }
+	virtual FReply OnMouseWheel(SWidget& OwnerWidget, const FGeometry& MyGeometry, const FPointerEvent& MouseEvent) { return FReply::Unhandled(); }
+	virtual void   OnMouseCaptureLost() { }
+	virtual int32  OnPaint(const FGeometry& AllottedGeometry, const FSlateRect& MyClippingRect, FSlateWindowElementList& OutDrawElements, int32 LayerId) const { return LayerId; }
+	virtual void   Tick(const FGeometry& AllottedGeometry, const double InCurrentTime, const float InDeltaTime) { }
+	virtual FCursorReply OnCursorQuery(const FGeometry& MyGeometry, const FPointerEvent& CursorEvent) const { return FCursorReply::Unhandled(); }
+
+	virtual ISequencer& GetSequencer() const = 0;
+
+	virtual FName GetIdentifier() const = 0;
+
+	/** Get the current active hotspot */
+	TSharedPtr<ISequencerHotspot> GetHotspot() { return Hotspot; }
+
+	/** Set the hotspot to something else */
+	void SetHotspot(TSharedPtr<ISequencerHotspot> NewHotspot) { Hotspot = MoveTemp(NewHotspot); }
+
+protected:
+
+	/** The current hotspot that can be set from anywhere to initiate drags */
+	TSharedPtr<ISequencerHotspot> Hotspot;
+};
\ No newline at end of file
diff --git a/Engine/Source/Editor/Sequencer/Public/ISequencerHotspot.h b/Engine/Source/Editor/Sequencer/Public/ISequencerHotspot.h
new file mode 100644
--- /dev/null
+++ b/Engine/Source/Editor/Sequencer/Public/ISequencerHotspot.h
@@ -0,0 +1,21 @@
+// Copyright 1998-2015 Epic Games, Inc. All Rights Reserved.
+
+#pragma once
+
+#include "ISequencerEditTool.h"
+
+enum class ESequencerHotspot
+{
+	Key,
+	Section,
+	SectionResize_L,
+	SectionResize_R,
+};
+
+/** A sequencer hotspot is used to identify specific areas on the sequencer track area */ 
+struct ISequencerHotspot
+{
+	virtual ~ISequencerHotspot() { }
+	virtual ESequencerHotspot GetType() const = 0;
+	virtual TSharedPtr<IEditToolDragOperation> InitiateDrag(ISequencer&) { return nullptr; }
+};
\ No newline at end of file
diff --git a/Engine/Source/Editor/Sequencer/Public/VirtualTrackArea.h b/Engine/Source/Editor/Sequencer/Public/VirtualTrackArea.h
new file mode 100644
--- /dev/null
+++ b/Engine/Source/Editor/Sequencer/Public/VirtualTrackArea.h
@@ -0,0 +1,44 @@
+// Copyright 1998-2015 Epic Games, Inc. All Rights Reserved.
+
+#pragma once
+
+#include "CommonMovieSceneTools.h"
+
+class SSequencerTreeView;
+
+/** Structure used for handling the virtual space of the track area */
+class SEQUENCER_API FVirtualTrackArea : public FTimeToPixel
+{
+public:
+
+	/** Construction responsibility is delegated to SSequencer. See SSequencer::GetVirtualTrackArea */
+	FVirtualTrackArea(const FSequencer& InSequencer, SSequencerTreeView& InTreeView, const FGeometry& InTrackAreaGeometry);
+
+	/** Convert the specified pixel position into a virtual vertical offset from the absolute top of the tree */
+	float PixelToVerticalOffset(float InPixel) const;
+
+	/** Convert the specified absolute vertical position into a physical vertical offset in the track area. */
+	/** @note: Use with caution - not reliable where the specified offset is not on screen */
+	float VerticalOffsetToPixel(float InOffset) const;
+
+	/** Convert the specified physical point into a virtual point the absolute top of the tree */
+	FVector2D PhysicalToVirtual(FVector2D InPosition) const;
+
+	/** Convert the specified absolute virtual point into a physical point in the track area. */
+	/** @note: Use with caution - not reliable where the specified point is not on screen */
+	FVector2D VirtualToPhysical(FVector2D InPosition) const;
+
+	/** Hit test at the specified physical position for a sequencer node */
+	TSharedPtr<FSequencerDisplayNode> HitTestNode(float InPhysicalPosition) const;
+
+	/** Hit test at the specified physical position for a section */
+	TOptional<FSectionHandle> HitTestSection(FVector2D InPhysicalPosition) const;
+
+	/** Hit test at the specified physical position for a key */
+	FSelectedKey HitTestKey(FVector2D InPhysicalPosition) const;
+
+private:
+
+	/** Reference to the sequencer tree */
+	SSequencerTreeView& TreeView;
+};
\ No newline at end of file
