Fixed some sections not being visited correctly

diff --git a/Engine/Source/Editor/Sequencer/Private/Tools/SequencerEditTool_Selection.cpp b/Engine/Source/Editor/Sequencer/Private/Tools/SequencerEditTool_Selection.cpp
--- a/Engine/Source/Editor/Sequencer/Private/Tools/SequencerEditTool_Selection.cpp
+++ b/Engine/Source/Editor/Sequencer/Private/Tools/SequencerEditTool_Selection.cpp
@@ -143,9 +143,11 @@ void FMarqueeRange::HandleNode(const IMarqueeVisitor& Visitor, FSequencerDisplay
 			}
 		}
 
-		if (!InNode.IsExpanded() && !bNodeHasKeyArea)
+		if (!bNodeHasKeyArea && (!InNode.IsExpanded() || InNode.GetChildNodes().Num() == 0))
 		{
-			// As a fallback, handle key groupings on collapsed parents
+			// As a fallback, we need to handle:
+			//  - Key groupings on collapsed parents
+			//  - Sections that have no key areas
 			for (int32 SectionIndex = 0; SectionIndex < InSections.Num(); ++SectionIndex)
 			{
 				UMovieSceneSection* Section = InSections[SectionIndex]->GetSectionObject();
@@ -153,8 +155,12 @@ void FMarqueeRange::HandleNode(const IMarqueeVisitor& Visitor, FSequencerDisplay
 				{
 					Visitor.VisitSection(Section);
 
-					TSharedRef<IKeyArea> KeyArea = InNode.UpdateKeyGrouping(SectionIndex);
-					HandleKeyArea(Visitor, KeyArea, Section);
+					// Only handle grouped keys if we actually have children
+					if (InNode.GetChildNodes().Num() != 0)
+					{
+						TSharedRef<IKeyArea> KeyArea = InNode.UpdateKeyGrouping(SectionIndex);
+						HandleKeyArea(Visitor, KeyArea, Section);
+					}
 				}
 			}
 		}
