Autoscroll is now always disabled during looped playback

diff --git a/Engine/Source/Editor/Sequencer/Private/Sequencer.cpp b/Engine/Source/Editor/Sequencer/Private/Sequencer.cpp
--- a/Engine/Source/Editor/Sequencer/Private/Sequencer.cpp
+++ b/Engine/Source/Editor/Sequencer/Private/Sequencer.cpp
@@ -568,7 +568,7 @@ float FSequencer::GetGlobalTime()
 
 void FSequencer::SetGlobalTime( float NewTime )
 {
-	if (bAutoScrollEnabled)
+	if (bAutoScrollEnabled && !bLoopingEnabled)
 	{
 		float RangeOffset = CalculateAutoscrollEncroachment(NewTime).Get(0.f);
 			
