Fixed undo on moved sequencer keys

diff --git a/Engine/Source/Editor/Sequencer/Private/Tools/EditToolDragOperations.cpp b/Engine/Source/Editor/Sequencer/Private/Tools/EditToolDragOperations.cpp
--- a/Engine/Source/Editor/Sequencer/Private/Tools/EditToolDragOperations.cpp
+++ b/Engine/Source/Editor/Sequencer/Private/Tools/EditToolDragOperations.cpp
@@ -411,7 +411,7 @@ void FMoveKeys::OnBeginDrag(const FVector2D& LocalMousePos, const FVirtualTrackA
 	check( SelectedKeys.Num() > 0 )
 
 	// Begin an editor transaction and mark the section as transactional so it's state will be saved
-	GEditor->BeginTransaction( NSLOCTEXT("Sequencer", "MoveKeysTransaction", "Move Keys") );
+	BeginTransaction( TArray< FSectionHandle >(), NSLOCTEXT("Sequencer", "MoveKeysTransaction", "Move Keys") );
 
 	TSet<UMovieSceneSection*> ModifiedSections;
 	for( FSelectedKey SelectedKey : SelectedKeys )
