Sequencer: Fix outliner node broadcast so that selected curves will load in the curve editor.

diff --git a/Engine/Source/Editor/Sequencer/Private/SAnimationOutlinerTreeNode.cpp b/Engine/Source/Editor/Sequencer/Private/SAnimationOutlinerTreeNode.cpp
--- a/Engine/Source/Editor/Sequencer/Private/SAnimationOutlinerTreeNode.cpp
+++ b/Engine/Source/Editor/Sequencer/Private/SAnimationOutlinerTreeNode.cpp
@@ -373,6 +373,7 @@ FReply SAnimationOutlinerTreeNode::OnMouseButtonDown( const FGeometry& MyGeometr
 		OnSelectionChanged( AffectedNodes );
 
 		Sequencer.GetSelection().ResumeBroadcast();
+		Sequencer.GetSelection().GetOnOutlinerNodeSelectionChanged().Broadcast();
 
 		return FReply::Handled();
 	}
