Initialization order fix

diff --git a/Engine/Source/Editor/Sequencer/Private/SequencerHotspots.h b/Engine/Source/Editor/Sequencer/Private/SequencerHotspots.h
--- a/Engine/Source/Editor/Sequencer/Private/SequencerHotspots.h
+++ b/Engine/Source/Editor/Sequencer/Private/SequencerHotspots.h
@@ -26,7 +26,7 @@ struct FKeyHotspot : ISequencerHotspot
 struct FSectionHandle
 {
 	FSectionHandle(TSharedPtr<FTrackNode> InTrackNode, int32 InSectionIndex)
-		: TrackNode(MoveTemp(InTrackNode)), SectionIndex(InSectionIndex)
+		: SectionIndex(InSectionIndex), TrackNode(MoveTemp(InTrackNode))
 	{}
 
 	UMovieSceneSection* GetSectionObject() const
