Sequencer: Clean up pass for consistency - capitalize transaction names, no periods in tooltips.

diff --git a/Engine/Source/Editor/MovieSceneTools/Private/CurveKeyEditors/SBoolCurveKeyEditor.cpp b/Engine/Source/Editor/MovieSceneTools/Private/CurveKeyEditors/SBoolCurveKeyEditor.cpp
--- a/Engine/Source/Editor/MovieSceneTools/Private/CurveKeyEditors/SBoolCurveKeyEditor.cpp
+++ b/Engine/Source/Editor/MovieSceneTools/Private/CurveKeyEditors/SBoolCurveKeyEditor.cpp
@@ -27,7 +27,7 @@ ECheckBoxState SBoolCurveKeyEditor::IsChecked() const
 
 void SBoolCurveKeyEditor::OnCheckStateChanged(ECheckBoxState NewCheckboxState)
 {
-	FScopedTransaction Transaction(LOCTEXT("SetBoolKey", "Set bool key value"));
+	FScopedTransaction Transaction(LOCTEXT("SetBoolKey", "Set Bool Key Value"));
 	OwningSection->SetFlags(RF_Transactional);
 	OwningSection->Modify();
 
diff --git a/Engine/Source/Editor/MovieSceneTools/Private/CurveKeyEditors/SEnumCurveKeyEditor.cpp b/Engine/Source/Editor/MovieSceneTools/Private/CurveKeyEditors/SEnumCurveKeyEditor.cpp
--- a/Engine/Source/Editor/MovieSceneTools/Private/CurveKeyEditors/SEnumCurveKeyEditor.cpp
+++ b/Engine/Source/Editor/MovieSceneTools/Private/CurveKeyEditors/SEnumCurveKeyEditor.cpp
@@ -51,7 +51,7 @@ TSharedRef<SWidget> SEnumCurveKeyEditor::OnGenerateWidget(TSharedPtr<FString> In
 
 void SEnumCurveKeyEditor::OnComboSelectionChanged(TSharedPtr<FString> InSelectedItem, ESelectInfo::Type SelectInfo)
 {
-	FScopedTransaction Transaction(LOCTEXT("SetEnumKey", "Set enum key value"));
+	FScopedTransaction Transaction(LOCTEXT("SetEnumKey", "Set Enum Key Value"));
 	OwningSection->SetFlags(RF_Transactional);
 	OwningSection->Modify();
 
diff --git a/Engine/Source/Editor/MovieSceneTools/Private/CurveKeyEditors/SFloatCurveKeyEditor.cpp b/Engine/Source/Editor/MovieSceneTools/Private/CurveKeyEditors/SFloatCurveKeyEditor.cpp
--- a/Engine/Source/Editor/MovieSceneTools/Private/CurveKeyEditors/SFloatCurveKeyEditor.cpp
+++ b/Engine/Source/Editor/MovieSceneTools/Private/CurveKeyEditors/SFloatCurveKeyEditor.cpp
@@ -31,7 +31,7 @@ void SFloatCurveKeyEditor::Construct(const FArguments& InArgs)
 
 void SFloatCurveKeyEditor::OnBeginSliderMovement()
 {
-	GEditor->BeginTransaction(LOCTEXT("SetFloatKey", "Set float key value"));
+	GEditor->BeginTransaction(LOCTEXT("SetFloatKey", "Set Float Key Value"));
 	OwningSection->SetFlags(RF_Transactional);
 	OwningSection->Modify();
 }
@@ -79,7 +79,7 @@ void SFloatCurveKeyEditor::OnValueCommitted(float Value, ETextCommit::Type Commi
 {
 	if (CommitInfo == ETextCommit::OnEnter)
 	{
-		const FScopedTransaction Transaction( LOCTEXT("SetFloatKey", "Set float key value") );
+		const FScopedTransaction Transaction( LOCTEXT("SetFloatKey", "Set Float Key Value") );
 
 		OnValueChanged(Value);
 	}
diff --git a/Engine/Source/Editor/MovieSceneTools/Private/CurveKeyEditors/SIntegralCurveKeyEditor.cpp b/Engine/Source/Editor/MovieSceneTools/Private/CurveKeyEditors/SIntegralCurveKeyEditor.cpp
--- a/Engine/Source/Editor/MovieSceneTools/Private/CurveKeyEditors/SIntegralCurveKeyEditor.cpp
+++ b/Engine/Source/Editor/MovieSceneTools/Private/CurveKeyEditors/SIntegralCurveKeyEditor.cpp
@@ -31,7 +31,7 @@ void SIntegralCurveKeyEditor::Construct(const FArguments& InArgs)
 
 void SIntegralCurveKeyEditor::OnBeginSliderMovement()
 {
-	GEditor->BeginTransaction(LOCTEXT("SetIntegralKey", "Set integral key value"));
+	GEditor->BeginTransaction(LOCTEXT("SetIntegralKey", "Set Integral Key Value"));
 	OwningSection->SetFlags(RF_Transactional);
 	OwningSection->Modify();
 }
@@ -77,7 +77,7 @@ void SIntegralCurveKeyEditor::OnValueCommitted(int32 Value, ETextCommit::Type Co
 {
 	if (CommitInfo == ETextCommit::OnEnter)
 	{
-		const FScopedTransaction Transaction( LOCTEXT("SetIntegralKey", "Set integral key value") );
+		const FScopedTransaction Transaction( LOCTEXT("SetIntegralKey", "Set Integral Key Value") );
 
 		OnValueChanged(Value);
 	}
diff --git a/Engine/Source/Editor/Sequencer/Private/SequencerCommands.cpp b/Engine/Source/Editor/Sequencer/Private/SequencerCommands.cpp
--- a/Engine/Source/Editor/Sequencer/Private/SequencerCommands.cpp
+++ b/Engine/Source/Editor/Sequencer/Private/SequencerCommands.cpp
@@ -7,48 +7,48 @@
 
 void FSequencerCommands::RegisterCommands()
 {
-	UI_COMMAND( TogglePlay, "Toggle Play", "Toggle the timeline playing.", EUserInterfaceActionType::Button, FInputChord(EKeys::SpaceBar) );
-	UI_COMMAND( PlayForward, "Play Forward", "Play the timeline forward.", EUserInterfaceActionType::Button, FInputChord(EKeys::Down) );
-	UI_COMMAND( Rewind, "Rewind", "Rewind the timeline.", EUserInterfaceActionType::Button, FInputChord(EKeys::Up) );
-	UI_COMMAND( StepForward, "Step Forward", "Step the timeline forward.", EUserInterfaceActionType::Button, FInputChord(EKeys::Right) );
+	UI_COMMAND( TogglePlay, "Toggle Play", "Toggle the timeline playing", EUserInterfaceActionType::Button, FInputChord(EKeys::SpaceBar) );
+	UI_COMMAND( PlayForward, "Play Forward", "Play the timeline forward", EUserInterfaceActionType::Button, FInputChord(EKeys::Down) );
+	UI_COMMAND( Rewind, "Rewind", "Rewind the timeline", EUserInterfaceActionType::Button, FInputChord(EKeys::Up) );
+	UI_COMMAND( StepForward, "Step Forward", "Step the timeline forward", EUserInterfaceActionType::Button, FInputChord(EKeys::Right) );
 	UI_COMMAND( StepBackward, "Step Backward", "Step the timeline backward", EUserInterfaceActionType::Button, FInputChord(EKeys::Left) );
 	UI_COMMAND( StepToNextKey, "Step to Next Key", "Step to the next key", EUserInterfaceActionType::Button, FInputChord(EKeys::Period) );
 	UI_COMMAND( StepToPreviousKey, "Step to Previous Key", "Step to the previous key", EUserInterfaceActionType::Button, FInputChord(EKeys::Comma) );
 	UI_COMMAND( StepToNextCameraKey, "Step to Next Camera Key", "Step to the next camera key", EUserInterfaceActionType::Button, FInputChord(EModifierKey::Alt, EKeys::Period) );
 	UI_COMMAND( StepToPreviousCameraKey, "Step to Previous Camera Key", "Step to the previous camera key", EUserInterfaceActionType::Button, FInputChord(EModifierKey::Alt, EKeys::Comma) );
 
-	UI_COMMAND( ToggleExpandCollapseNodes, "Toggle Expand/Collapse Nodes", "Toggle expand or collapse selected nodes.", EUserInterfaceActionType::Button, FInputChord(EKeys::O) );
-	UI_COMMAND( ToggleExpandCollapseNodesAndDescendants, "Toggle Expand/Collapse Nodes and Descendants", "Toggle expand or collapse selected nodes and their descendants.", EUserInterfaceActionType::Button, FInputChord(EModifierKey::Shift, EKeys::O) );
+	UI_COMMAND( ToggleExpandCollapseNodes, "Toggle Expand/Collapse Nodes", "Toggle expand or collapse selected nodes", EUserInterfaceActionType::Button, FInputChord(EKeys::O) );
+	UI_COMMAND( ToggleExpandCollapseNodesAndDescendants, "Toggle Expand/Collapse Nodes and Descendants", "Toggle expand or collapse selected nodes and their descendants", EUserInterfaceActionType::Button, FInputChord(EModifierKey::Shift, EKeys::O) );
 
-	UI_COMMAND( SetKey, "Set Key", "Sets a key at the current time for the selected actor.", EUserInterfaceActionType::Button, FInputChord(EKeys::K) );
+	UI_COMMAND( SetKey, "Set Key", "Sets a key at the current time for the selected actor", EUserInterfaceActionType::Button, FInputChord(EKeys::K) );
 
-	UI_COMMAND( ToggleAutoKeyEnabled, "Auto Key", "Enables and disables auto keying.", EUserInterfaceActionType::ToggleButton, FInputChord() );
-	UI_COMMAND( ToggleKeyAllEnabled, "Key All", "Enables and disables key all.", EUserInterfaceActionType::ToggleButton, FInputChord() );
+	UI_COMMAND( ToggleAutoKeyEnabled, "Auto Key", "Enables and disables auto keying", EUserInterfaceActionType::ToggleButton, FInputChord() );
+	UI_COMMAND( ToggleKeyAllEnabled, "Key All", "Enables and disables key all", EUserInterfaceActionType::ToggleButton, FInputChord() );
 
-	UI_COMMAND( ToggleAutoScroll, "Auto Scroll", "Toggle auto-scroll: When enabled, automatically scrolls the sequencer view to keep the current time visible.", EUserInterfaceActionType::ToggleButton, FInputChord(EModifierKey::Shift, EKeys::S) );
+	UI_COMMAND( ToggleAutoScroll, "Auto Scroll", "Toggle auto-scroll: When enabled, automatically scrolls the sequencer view to keep the current time visible", EUserInterfaceActionType::ToggleButton, FInputChord(EModifierKey::Shift, EKeys::S) );
 
-	UI_COMMAND( ToggleShowFrameNumbers, "Show Frame Numbers", "Enables and disables showing frame numbers.", EUserInterfaceActionType::ToggleButton, FInputChord() );
-	UI_COMMAND( ToggleShowRangeSlider, "Show Range Slider", "Enables and disables showing the time range slider.", EUserInterfaceActionType::ToggleButton, FInputChord() );
-	UI_COMMAND( ToggleLockInOutToStartEndRange, "Lock In/Out to Start/End", "Enables and disables locking the in/out to the start/end range while scrolling the view.", EUserInterfaceActionType::ToggleButton, FInputChord() );
-	UI_COMMAND( ToggleIsSnapEnabled, "Enable Snapping", "Enables and disables snapping.", EUserInterfaceActionType::ToggleButton, FInputChord() );
+	UI_COMMAND( ToggleShowFrameNumbers, "Show Frame Numbers", "Enables and disables showing frame numbers", EUserInterfaceActionType::ToggleButton, FInputChord() );
+	UI_COMMAND( ToggleShowRangeSlider, "Show Range Slider", "Enables and disables showing the time range slider", EUserInterfaceActionType::ToggleButton, FInputChord() );
+	UI_COMMAND( ToggleLockInOutToStartEndRange, "Lock In/Out to Start/End", "Enables and disables locking the in/out to the start/end range while scrolling the view", EUserInterfaceActionType::ToggleButton, FInputChord() );
+	UI_COMMAND( ToggleIsSnapEnabled, "Enable Snapping", "Enables and disables snapping", EUserInterfaceActionType::ToggleButton, FInputChord() );
 
-	UI_COMMAND( ToggleSnapKeyTimesToInterval, "Snap to the Interval", "Snap keys to the interval.", EUserInterfaceActionType::ToggleButton, FInputChord() );
-	UI_COMMAND( ToggleSnapKeyTimesToKeys, "Snap to Other Keys", "Snap keys to other keys in the section.", EUserInterfaceActionType::ToggleButton, FInputChord() );
+	UI_COMMAND( ToggleSnapKeyTimesToInterval, "Snap to the Interval", "Snap keys to the interval", EUserInterfaceActionType::ToggleButton, FInputChord() );
+	UI_COMMAND( ToggleSnapKeyTimesToKeys, "Snap to Other Keys", "Snap keys to other keys in the section", EUserInterfaceActionType::ToggleButton, FInputChord() );
 
-	UI_COMMAND( ToggleSnapSectionTimesToInterval, "Snap to the Interval", "Snap sections to the interval.", EUserInterfaceActionType::ToggleButton, FInputChord() );
-	UI_COMMAND( ToggleSnapSectionTimesToSections, "Snap to Other Sections", "Snap sections to other sections.", EUserInterfaceActionType::ToggleButton, FInputChord() );
+	UI_COMMAND( ToggleSnapSectionTimesToInterval, "Snap to the Interval", "Snap sections to the interval", EUserInterfaceActionType::ToggleButton, FInputChord() );
+	UI_COMMAND( ToggleSnapSectionTimesToSections, "Snap to Other Sections", "Snap sections to other sections", EUserInterfaceActionType::ToggleButton, FInputChord() );
 
-	UI_COMMAND( ToggleSnapPlayTimeToInterval, "Snap to the Interval While Scrubbing", "Snap the play time to the interval while scrubbing.", EUserInterfaceActionType::ToggleButton, FInputChord() );
-	UI_COMMAND( ToggleSnapPlayTimeToDraggedKey, "Snap to the Dragged Key", "Snap the play time to the dragged key.", EUserInterfaceActionType::ToggleButton, FInputChord() );
+	UI_COMMAND( ToggleSnapPlayTimeToInterval, "Snap to the Interval While Scrubbing", "Snap the play time to the interval while scrubbing", EUserInterfaceActionType::ToggleButton, FInputChord() );
+	UI_COMMAND( ToggleSnapPlayTimeToDraggedKey, "Snap to the Dragged Key", "Snap the play time to the dragged key", EUserInterfaceActionType::ToggleButton, FInputChord() );
 
-	UI_COMMAND( ToggleSnapCurveValueToInterval, "Snap Curve Key Values", "Snap curve keys to the value snapping interval.", EUserInterfaceActionType::ToggleButton, FInputChord() );
+	UI_COMMAND( ToggleSnapCurveValueToInterval, "Snap Curve Key Values", "Snap curve keys to the value snapping interval", EUserInterfaceActionType::ToggleButton, FInputChord() );
 
-	UI_COMMAND( ToggleCleanView, "Clean View", "Enable 'Clean View' mode which displays only global tracks when no filter is applied.", EUserInterfaceActionType::ToggleButton, FInputChord() );
+	UI_COMMAND( ToggleCleanView, "Clean View", "Enable 'Clean View' mode which displays only global tracks when no filter is applied", EUserInterfaceActionType::ToggleButton, FInputChord() );
 	
-	UI_COMMAND( ToggleShowCurveEditor, "Curve Editor", "Show the animation keys in a curve editor.", EUserInterfaceActionType::ToggleButton, FInputChord() );
+	UI_COMMAND( ToggleShowCurveEditor, "Curve Editor", "Show the animation keys in a curve editor", EUserInterfaceActionType::ToggleButton, FInputChord() );
 
-	UI_COMMAND( MoveTool, "Move", "Activates the move tool.", EUserInterfaceActionType::ToggleButton, FInputChord( EKeys::V ) );
-	UI_COMMAND( MarqueeTool, "Marquee", "Activates the marquee selection tool.", EUserInterfaceActionType::ToggleButton, FInputChord( EKeys::M ) );
+	UI_COMMAND( MoveTool, "Move", "Activates the move tool", EUserInterfaceActionType::ToggleButton, FInputChord( EKeys::V ) );
+	UI_COMMAND( MarqueeTool, "Marquee", "Activates the marquee selection tool", EUserInterfaceActionType::ToggleButton, FInputChord( EKeys::M ) );
 }
 
 #undef LOCTEXT_NAMESPACE
\ No newline at end of file
