Sequencer: Add the ability to assign an actor to an existing sequencer node (replaces the existing binding and possessable).

diff --git a/Engine/Source/Editor/Sequencer/Private/Sequencer.cpp b/Engine/Source/Editor/Sequencer/Private/Sequencer.cpp
--- a/Engine/Source/Editor/Sequencer/Private/Sequencer.cpp
+++ b/Engine/Source/Editor/Sequencer/Private/Sequencer.cpp
@@ -1823,12 +1823,57 @@ void FSequencer::DeleteSelectedItems()
 	}
 }
 
-bool FSequencer::CanDeleteSelectedItems()
+bool FSequencer::CanDeleteSelectedItems() const
 {
 	// Only delete sequencer items if the delete key was pressed in the sequencer widget and not in, for example, the viewport.
 	return SequencerWidget->HasFocusedDescendants();
 }
 
+void FSequencer::AssignActor(FGuid InObjectBinding, FObjectBindingNode* ObjectBindingNode)
+{
+	if (ObjectBindingNode != nullptr)
+	{
+		// Get the first selected actor as the actor to replace with
+		AActor* Actor = GEditor->GetSelectedActors()->GetTop<AActor>();
+		if (Actor != nullptr)
+		{
+			FScopedTransaction AssignActor( NSLOCTEXT("Sequencer", "AssignActor", "Assign Actor") );
+
+			UMovieSceneSequence* OwnerSequence = GetFocusedMovieSceneSequence();
+			UMovieScene* OwnerMovieScene = OwnerSequence->GetMovieScene();
+
+			Actor->Modify();
+			OwnerSequence->Modify();
+			OwnerMovieScene->Modify();
+
+			// Get the object guid to assign
+			FGuid NewObjectGuid = OwnerSequence->FindObjectId(*Actor);
+			FString NewActorLabel = Actor->GetActorLabel();
+
+			if (NewObjectGuid.IsValid())
+			{
+				OwnerMovieScene->RemovePossessable(NewObjectGuid);
+				OwnerSequence->UnbindPossessableObjects(NewObjectGuid);
+			}
+
+			// Add this object if it hasn't already been possessed
+			FMovieScenePossessable NewPossessable( Actor->GetActorLabel(), Actor->GetClass());
+			NewObjectGuid = NewPossessable.GetGuid();
+			OwnerSequence->BindPossessableObject(NewObjectGuid, *Actor);
+
+			// Replace
+			OwnerMovieScene->ReplacePossessable(InObjectBinding, NewObjectGuid, NewActorLabel);
+					
+			NotifyMovieSceneDataChanged();
+		}
+	}
+}
+
+bool FSequencer::CanAssignActor(FGuid ObjectBinding) const
+{
+	return GEditor->GetSelectedActors()->Num() > 0;
+}
+
 void FSequencer::TogglePlay()
 {
 	OnPlay();
@@ -2078,8 +2123,17 @@ void FSequencer::BindSequencerCommands()
 	}
 }
 
-void FSequencer::BuildObjectBindingContextMenu(FMenuBuilder& MenuBuilder, const FGuid& ObjectBinding, const UClass* ObjectClass)
+void FSequencer::BuildObjectBindingContextMenu(FMenuBuilder& MenuBuilder, const FGuid& ObjectBinding, const UClass* ObjectClass, FObjectBindingNode* ObjectBindingNode)
 {
+	FFormatNamedArguments Args;
+	MenuBuilder.AddMenuEntry(
+		FText::Format( NSLOCTEXT("Sequencer", "Assign Actor ", "Assign Actor"), Args),
+		FText::Format( NSLOCTEXT("Sequencer", "AssignActorTooltip", "Assign the selected actor to this track"), Args ),
+		FSlateIcon(),
+		FUIAction(FExecuteAction::CreateSP(this, &FSequencer::AssignActor, ObjectBinding, ObjectBindingNode),
+					FCanExecuteAction::CreateSP(this, &FSequencer::CanAssignActor, ObjectBinding)) );
+
+
 	for (int32 i = 0; i < TrackEditors.Num(); ++i)
 	{
 		TrackEditors[i]->BuildObjectBindingContextMenu(MenuBuilder, ObjectBinding, ObjectClass);
diff --git a/Engine/Source/Editor/Sequencer/Private/Sequencer.h b/Engine/Source/Editor/Sequencer/Private/Sequencer.h
--- a/Engine/Source/Editor/Sequencer/Private/Sequencer.h
+++ b/Engine/Source/Editor/Sequencer/Private/Sequencer.h
@@ -213,7 +213,7 @@ public:
 	 * @param ObjectBinding	The object binding of the selected node
 	 * @param ObjectClass	The class of the selected object
 	 */
-	void BuildObjectBindingContextMenu(class FMenuBuilder& MenuBuilder, const FGuid& ObjectBinding, const class UClass* ObjectClass);
+	void BuildObjectBindingContextMenu(class FMenuBuilder& MenuBuilder, const FGuid& ObjectBinding, const class UClass* ObjectClass, FObjectBindingNode* ObjectBindingNode);
 
 	/**
 	 * Builds up the track menu for object binding nodes in the outliner
@@ -379,8 +379,12 @@ protected:
 	
 	/** Called when a user executes the delete command to delete sections or keys */
 	void DeleteSelectedItems();
-	bool CanDeleteSelectedItems();
+	bool CanDeleteSelectedItems() const;
 	
+	/** Called when a user executes the assign actor to track menu item */
+	void AssignActor(FGuid ObjectBinding, FObjectBindingNode* ObjectBindingNode);
+	bool CanAssignActor(FGuid ObjectBinding) const;
+
 	/** Transport controls */
 	void TogglePlay();
 	void PlayForward();
diff --git a/Engine/Source/Editor/Sequencer/Private/SequencerDisplayNode.cpp b/Engine/Source/Editor/Sequencer/Private/SequencerDisplayNode.cpp
--- a/Engine/Source/Editor/Sequencer/Private/SequencerDisplayNode.cpp
+++ b/Engine/Source/Editor/Sequencer/Private/SequencerDisplayNode.cpp
@@ -707,7 +707,7 @@ TSharedPtr<SWidget> FObjectBindingNode::OnSummonContextMenu(const FGeometry& MyG
 	const bool bShouldCloseWindowAfterMenuSelection = true;
 	FMenuBuilder MenuBuilder(bShouldCloseWindowAfterMenuSelection, NULL);
 
-	GetSequencer().BuildObjectBindingContextMenu(MenuBuilder, ObjectBinding, ObjectClass);
+	GetSequencer().BuildObjectBindingContextMenu(MenuBuilder, ObjectBinding, ObjectClass, this);
 
 	return MenuBuilder.MakeWidget();
 }
diff --git a/Engine/Source/Runtime/MovieScene/Private/MovieScene.cpp b/Engine/Source/Runtime/MovieScene/Private/MovieScene.cpp
--- a/Engine/Source/Runtime/MovieScene/Private/MovieScene.cpp
+++ b/Engine/Source/Runtime/MovieScene/Private/MovieScene.cpp
@@ -161,6 +161,33 @@ bool UMovieScene::RemovePossessable( const FGuid& PossessableGuid )
 }
 
 
+bool UMovieScene::ReplacePossessable( const FGuid& OldGuid, const FGuid& NewGuid, const FString& NewName )
+{
+	bool bAnythingReplaced = false;
+
+	for( auto PossessableIter( Possessables.CreateIterator() ); PossessableIter; ++PossessableIter )
+	{
+		auto& CurPossessable = *PossessableIter;
+
+		if( CurPossessable.GetGuid() == OldGuid )
+		{	
+			Modify();
+
+			// Found it!
+			CurPossessable.SetGuid(NewGuid);
+			CurPossessable.SetName(NewName);
+
+			ReplaceBinding( OldGuid, NewGuid, NewName );
+
+			bAnythingReplaced = true;
+			break;
+		}
+	}
+
+	return bAnythingReplaced;
+}
+
+
 FMovieScenePossessable* UMovieScene::FindPossessable( const FGuid& Guid )
 {
 	for( auto CurPossessableIt( Possessables.CreateIterator() ); CurPossessableIt; ++CurPossessableIt )
@@ -390,3 +417,16 @@ void UMovieScene::RemoveBinding(const FGuid& Guid)
 		}
 	}
 }
+
+void UMovieScene::ReplaceBinding(const FGuid& OldGuid, const FGuid& NewGuid, const FString& Name)
+{
+	for (int32 BindingIndex = 0; BindingIndex < ObjectBindings.Num(); ++BindingIndex)
+	{
+		if (ObjectBindings[BindingIndex].GetObjectGuid() == OldGuid)
+		{
+			ObjectBindings[BindingIndex].SetObjectGuid(NewGuid);
+			ObjectBindings[BindingIndex].SetName(Name);
+			break;
+		}
+	}
+}
\ No newline at end of file
diff --git a/Engine/Source/Runtime/MovieScene/Public/MovieScene.h b/Engine/Source/Runtime/MovieScene/Public/MovieScene.h
--- a/Engine/Source/Runtime/MovieScene/Public/MovieScene.h
+++ b/Engine/Source/Runtime/MovieScene/Public/MovieScene.h
@@ -125,6 +125,11 @@ public:
 	 */
 	bool RemovePossessable(const FGuid& PossessableGuid);
 	
+	/*
+	* Replace an existing possessable with another 
+	*/
+	bool ReplacePossessable(const FGuid& OldGuid, const FGuid& NewGuid, const FString& Name);
+
 	/**
 	 * Tries to locate a possessable in this MovieScene for the specified possessable GUID.
 	 *
@@ -250,6 +255,11 @@ public:
 		return ObjectBindings;
 	}
 
+	/*
+	* Replace an existing binding with another 
+	*/
+	void ReplaceBinding(const FGuid& OldGuid, const FGuid& NewGuid, const FString& Name);
+
 	/**
 	 * @return The time range of the movie scene (the in to out time range)
 	 */
diff --git a/Engine/Source/Runtime/MovieScene/Public/MovieSceneBinding.h b/Engine/Source/Runtime/MovieScene/Public/MovieSceneBinding.h
--- a/Engine/Source/Runtime/MovieScene/Public/MovieSceneBinding.h
+++ b/Engine/Source/Runtime/MovieScene/Public/MovieSceneBinding.h
@@ -48,6 +48,16 @@ struct FMovieSceneBinding
 	 */
 	TRange<float> GetTimeRange() const;
 
+	/**
+	* Set the object guid
+	*
+	* @param InObjectGuid
+	*/
+	void SetObjectGuid(const FGuid& InObjectGuid)
+	{
+		ObjectGuid = InObjectGuid;
+	}
+
 	/**
 	 * @return The guid of runtime objects in this binding
 	 */
@@ -56,6 +66,14 @@ struct FMovieSceneBinding
 		return ObjectGuid;
 	}
 
+	/**
+	 * Set display name of the binding
+	 */
+	void SetName(const FString& InBindingName)
+	{
+		BindingName = InBindingName;
+	} 
+
 	/**
 	 * @return The display name of the binding
 	 */
diff --git a/Engine/Source/Runtime/MovieScene/Public/MovieScenePossessable.h b/Engine/Source/Runtime/MovieScene/Public/MovieScenePossessable.h
--- a/Engine/Source/Runtime/MovieScene/Public/MovieScenePossessable.h
+++ b/Engine/Source/Runtime/MovieScene/Public/MovieScenePossessable.h
@@ -43,6 +43,17 @@ public:
 		return Guid;
 	}
 
+	/** 
+	* Set the unique identifier
+	*
+	* @param InGuid
+	*/
+
+	void SetGuid(const FGuid& InGuid)
+	{
+		Guid = InGuid;
+	}
+
 	/**
 	 * Get the name of the possessed object.
 	 *
@@ -54,6 +65,16 @@ public:
 		return Name;
 	}
 
+	/**
+	* Set the name of the possessed object
+	*
+	* @param InName
+	*/
+	void SetName(const FString& InName)
+	{
+		Name = InName;
+	}
+
 	/**
 	 * Get the class of the possessed object.
 	 *
