Sequencer: Allow assign actor only for level editor sequencers.

diff --git a/Engine/Source/Editor/Sequencer/Private/Sequencer.cpp b/Engine/Source/Editor/Sequencer/Private/Sequencer.cpp
--- a/Engine/Source/Editor/Sequencer/Private/Sequencer.cpp
+++ b/Engine/Source/Editor/Sequencer/Private/Sequencer.cpp
@@ -2125,14 +2125,16 @@ void FSequencer::BindSequencerCommands()
 
 void FSequencer::BuildObjectBindingContextMenu(FMenuBuilder& MenuBuilder, const FGuid& ObjectBinding, const UClass* ObjectClass, FObjectBindingNode* ObjectBindingNode)
 {
-	FFormatNamedArguments Args;
-	MenuBuilder.AddMenuEntry(
-		FText::Format( NSLOCTEXT("Sequencer", "Assign Actor ", "Assign Actor"), Args),
-		FText::Format( NSLOCTEXT("Sequencer", "AssignActorTooltip", "Assign the selected actor to this track"), Args ),
-		FSlateIcon(),
-		FUIAction(FExecuteAction::CreateSP(this, &FSequencer::AssignActor, ObjectBinding, ObjectBindingNode),
-					FCanExecuteAction::CreateSP(this, &FSequencer::CanAssignActor, ObjectBinding)) );
-
+	if (IsLevelEditorSequencer())
+	{
+		FFormatNamedArguments Args;
+		MenuBuilder.AddMenuEntry(
+			FText::Format( NSLOCTEXT("Sequencer", "Assign Actor ", "Assign Actor"), Args),
+			FText::Format( NSLOCTEXT("Sequencer", "AssignActorTooltip", "Assign the selected actor to this track"), Args ),
+			FSlateIcon(),
+			FUIAction(FExecuteAction::CreateSP(this, &FSequencer::AssignActor, ObjectBinding, ObjectBindingNode),
+						FCanExecuteAction::CreateSP(this, &FSequencer::CanAssignActor, ObjectBinding)) );
+	}
 
 	for (int32 i = 0; i < TrackEditors.Num(); ++i)
 	{
