Fixed dragging keys in the sequencer with snapping enabled

  - Added FSequencerSnapField to simplify snapping calculations and related code.
  - Formalised IMarqueeVisitor as a more generic ISequencerEntityVisitor that can be used to walk a range of visible sequencer entities
  - Updated FMoveKeys drag operation to use FSequencerSnapField and corrected snapping code
  - Track area geometry is now cached in SSequencerTrackArea to simplify SSequencer::GetVirtualTrackArea
  - Fixed grouped keys swallowing each other when dragging other keys over the top

diff --git a/Engine/Source/Editor/Sequencer/Private/SSection.cpp b/Engine/Source/Editor/Sequencer/Private/SSection.cpp
--- a/Engine/Source/Editor/Sequencer/Private/SSection.cpp
+++ b/Engine/Source/Editor/Sequencer/Private/SSection.cpp
@@ -10,7 +10,19 @@
 #include "CommonMovieSceneTools.h"
 #include "SequencerHotspots.h"
 
-FThreadSafeCounter SSection::LayoutRegenerationLock;
+
+/** When 0, regeneration of dynamic key layouts is enabled, when non-zero, such behaviour is disabled */
+FThreadSafeCounter LayoutRegenerationLock;
+
+void SSection::DisableLayoutRegeneration()
+{
+	LayoutRegenerationLock.Increment();
+}
+
+void SSection::EnableLayoutRegeneration()
+{
+	LayoutRegenerationLock.Decrement();
+}
 
 void SSection::Construct( const FArguments& InArgs, TSharedRef<FTrackNode> SectionNode, int32 InSectionIndex )
 {
@@ -452,7 +464,6 @@ void SSection::ResetState()
 {
 	ResetHoveredState();
 	PressedKey = FSelectedKey();
-	LayoutRegenerationLock.Reset();
 }
 
 void SSection::ResetHoveredState()
diff --git a/Engine/Source/Editor/Sequencer/Private/SSection.h b/Engine/Source/Editor/Sequencer/Private/SSection.h
--- a/Engine/Source/Editor/Sequencer/Private/SSection.h
+++ b/Engine/Source/Editor/Sequencer/Private/SSection.h
@@ -31,6 +31,12 @@ public:
 
 	virtual FVector2D ComputeDesiredSize(float) const override;
 	
+	/** Temporarily disable dynamic layout regeneration. This prevents overlapping key groups from being amalgamated during drags. Key times will continue to update correctly. */
+	static void DisableLayoutRegeneration();
+
+	/** Re-enable dynamic layout regeneration */
+	static void EnableLayoutRegeneration();
+
 private:
 
 	/**
@@ -154,8 +160,6 @@ private:
 	bool bRightEdgeHovered;
 	/** Whether or not the right edge of the section is pressed */
 	bool bRightEdgePressed;
-	/** When 0, regeneration of dynamic key layouts is enabled, when non-zero, such behaviour is disabled */
-	static FThreadSafeCounter LayoutRegenerationLock;
 	/** Cached parent geometry to pass down to any section interfaces that need it during tick */
 	FGeometry ParentGeometry;
 };
\ No newline at end of file
diff --git a/Engine/Source/Editor/Sequencer/Private/SSequencer.cpp b/Engine/Source/Editor/Sequencer/Private/SSequencer.cpp
--- a/Engine/Source/Editor/Sequencer/Private/SSequencer.cpp
+++ b/Engine/Source/Editor/Sequencer/Private/SSequencer.cpp
@@ -1464,9 +1464,9 @@ void SSequencer::OnCurveEditorVisibilityChanged()
 	}
 }
 
-FVirtualTrackArea SSequencer::GetVirtualTrackArea(const FGeometry& InTrackAreaGeometry) const
+FVirtualTrackArea SSequencer::GetVirtualTrackArea() const
 {
-	return FVirtualTrackArea(*Sequencer.Pin(), *TreeView.Get(), InTrackAreaGeometry);
+	return FVirtualTrackArea(*Sequencer.Pin(), *TreeView.Get(), TrackArea->GetCachedGeometry());
 }
 
 #undef LOCTEXT_NAMESPACE
diff --git a/Engine/Source/Editor/Sequencer/Private/SSequencer.h b/Engine/Source/Editor/Sequencer/Private/SSequencer.h
--- a/Engine/Source/Editor/Sequencer/Private/SSequencer.h
+++ b/Engine/Source/Editor/Sequencer/Private/SSequencer.h
@@ -132,7 +132,7 @@ public:
 	ISequencerEditTool& GetEditTool() const { return *EditTool; }
 
 	/** Generate a helper structure that can be used to transform between phsyical space and virtual space in the track area */
-	FVirtualTrackArea GetVirtualTrackArea(const FGeometry& InTrackAreaGeometry) const;
+	FVirtualTrackArea GetVirtualTrackArea() const;
 
 	/** Get an array of section handles for the given set of movie scene sections */
 	TArray<FSectionHandle> GetSectionHandles(const TSet<TWeakObjectPtr<UMovieSceneSection>>& DesiredSections) const;
diff --git a/Engine/Source/Editor/Sequencer/Private/SSequencerTrackArea.cpp b/Engine/Source/Editor/Sequencer/Private/SSequencerTrackArea.cpp
--- a/Engine/Source/Editor/Sequencer/Private/SSequencerTrackArea.cpp
+++ b/Engine/Source/Editor/Sequencer/Private/SSequencerTrackArea.cpp
@@ -165,6 +165,8 @@ FCursorReply SSequencerTrackArea::OnCursorQuery( const FGeometry& MyGeometry, co
 
 void SSequencerTrackArea::Tick( const FGeometry& AllottedGeometry, const double InCurrentTime, const float InDeltaTime )
 {
+	CachedGeometry = AllottedGeometry;
+
 	auto SequencerPin = SequencerWidget.Pin();
 	if (SequencerPin.IsValid())
 	{
diff --git a/Engine/Source/Editor/Sequencer/Private/SSequencerTrackArea.h b/Engine/Source/Editor/Sequencer/Private/SSequencerTrackArea.h
--- a/Engine/Source/Editor/Sequencer/Private/SSequencerTrackArea.h
+++ b/Engine/Source/Editor/Sequencer/Private/SSequencerTrackArea.h
@@ -40,8 +40,14 @@ public:
 	virtual FCursorReply OnCursorQuery( const FGeometry& MyGeometry, const FPointerEvent& CursorEvent ) const override;
 	virtual void Tick( const FGeometry& AllottedGeometry, const double InCurrentTime, const float InDeltaTime ) override;
 
+	/** Access the cached geometry for this track area */
+	const FGeometry& GetCachedGeometry() const { return CachedGeometry; }
+	
 private:
 
+	/** Cached geometry */
+	FGeometry CachedGeometry;
+
 	/** A map of child slot content that exist in our view */
 	TMap<TSharedPtr<FSequencerDisplayNode>, TWeakPtr<SSequencerTrackLane>> TrackSlots;
 
diff --git a/Engine/Source/Editor/Sequencer/Private/SSequencerTreeView.h b/Engine/Source/Editor/Sequencer/Private/SSequencerTreeView.h
--- a/Engine/Source/Editor/Sequencer/Private/SSequencerTreeView.h
+++ b/Engine/Source/Editor/Sequencer/Private/SSequencerTreeView.h
@@ -107,6 +107,9 @@ public:
 		float PhysicalTop, PhysicalHeight;
 	};
 
+	/** Access all the physical nodes currently visible on the sequencer */
+	const TArray<FCachedGeometry>& GetAllVisibleNodes() const { return PhysicalNodes; }
+
 	/** Retrieve the last reported physical geometry for the specified node, if available */
 	TOptional<FCachedGeometry> GetPhysicalGeometryForNode(const FDisplayNodeRef& InNode) const;
 
diff --git a/Engine/Source/Editor/Sequencer/Private/Tools/EditToolDragOperations.cpp b/Engine/Source/Editor/Sequencer/Private/Tools/EditToolDragOperations.cpp
--- a/Engine/Source/Editor/Sequencer/Private/Tools/EditToolDragOperations.cpp
+++ b/Engine/Source/Editor/Sequencer/Private/Tools/EditToolDragOperations.cpp
@@ -12,6 +12,46 @@
 #include "MovieSceneShotTrack.h"
 #include "VirtualTrackArea.h"
 
+struct FDefaultKeySnappingCandidates : ISequencerSnapCandidate
+{
+	FDefaultKeySnappingCandidates(const TSet<FSelectedKey>& InSelectedKeys)
+		: SelectedKeys(InSelectedKeys)
+	{}
+
+	bool IsKeyApplicable(FKeyHandle KeyHandle, const TSharedPtr<IKeyArea>& KeyArea, UMovieSceneSection* Section)
+	{
+		return !SelectedKeys.Contains(FSelectedKey(*Section, KeyArea, KeyHandle));
+	}
+
+	bool AreSectionBoundsApplicable(UMovieSceneSection* Section) { return true; }
+	bool AreGridLinesApplicable() { return true; }
+
+	const TSet<FSelectedKey>& SelectedKeys;
+};
+
+
+TOptional<FSequencerSnapField::FSnapResult> SnapToInterval(const TArray<float>& InTimes, float Threshold, const USequencerSettings& Settings)
+{
+	TOptional<FSequencerSnapField::FSnapResult> Result;
+
+	float SnapAmount = 0.f;
+	for (float Time : InTimes)
+	{
+		float IntervalSnap = Settings.SnapTimeToInterval(Time);
+		float ThisSnapAmount = IntervalSnap - Time;
+		if (FMath::Abs(ThisSnapAmount) <= Threshold)
+		{
+			if (!Result.IsSet() || FMath::Abs(ThisSnapAmount) < SnapAmount)
+			{
+				Result = FSequencerSnapField::FSnapResult{Time, IntervalSnap};
+				SnapAmount = ThisSnapAmount;
+			}
+		}
+	}
+
+	return Result;
+}
+
 /** How many pixels near the mouse has to be before snapping occurs */
 const float PixelSnapWidth = 10.f;
 
@@ -144,8 +184,8 @@ void FEditToolDragOperation::EndTransaction()
 	Sequencer.UpdateRuntimeInstances();
 }
 
-FResizeSection::FResizeSection( FSequencer& Sequencer, TArray<FSectionHandle> InSections, TOptional<FSectionHandle> InCardinalSection, bool bInDraggingByEnd )
-	: FEditToolDragOperation( Sequencer )
+FResizeSection::FResizeSection( FSequencer& InSequencer, TArray<FSectionHandle> InSections, TOptional<FSectionHandle> InCardinalSection, bool bInDraggingByEnd )
+	: FEditToolDragOperation( InSequencer )
 	, Sections( MoveTemp(InSections) )
 	, CardinalSection(InCardinalSection)
 	, bDraggingByEnd(bInDraggingByEnd)
@@ -254,8 +294,8 @@ void FResizeSection::OnDrag( const FPointerEvent& MouseEvent, const FVector2D& L
 	}
 }
 
-FMoveSection::FMoveSection( FSequencer& Sequencer, TArray<FSectionHandle> InSections )
-	: FEditToolDragOperation( Sequencer )
+FMoveSection::FMoveSection( FSequencer& InSequencer, TArray<FSectionHandle> InSections )
+	: FEditToolDragOperation( InSequencer )
 	, Sections( MoveTemp(InSections) )
 {
 }
@@ -263,7 +303,7 @@ FMoveSection::FMoveSection( FSequencer& Sequencer, TArray<FSectionHandle> InSect
 void FMoveSection::OnBeginDrag(const FVector2D& LocalMousePos, const FVirtualTrackArea& VirtualTrackArea)
 {
 	BeginTransaction( Sections, NSLOCTEXT("Sequencer", "MoveSectionTransaction", "Move Section") );
-		
+
 	MouseDownTime = VirtualTrackArea.PixelToTime(LocalMousePos.X);
 
 	DraggedKeyHandles.Empty();
@@ -313,7 +353,7 @@ void FMoveSection::OnDrag( const FPointerEvent& MouseEvent, const FVector2D& Loc
 		{
 			MovieSceneSections.Add(SequencerNodeSections[i]->GetSectionObject());
 		}
-						
+		
 		if ( Settings->GetIsSnapEnabled() )
 		{
 			bool bSnappedToSection = false;
@@ -410,14 +450,23 @@ void FMoveKeys::OnBeginDrag(const FVector2D& LocalMousePos, const FVirtualTrackA
 {
 	check( SelectedKeys.Num() > 0 )
 
+	SSection::DisableLayoutRegeneration();
+
+	FDefaultKeySnappingCandidates SnapCandidates(SelectedKeys);
+	SnapField = FSequencerSnapField(Sequencer, SnapCandidates);
+
 	// Begin an editor transaction and mark the section as transactional so it's state will be saved
 	BeginTransaction( TArray< FSectionHandle >(), NSLOCTEXT("Sequencer", "MoveKeysTransaction", "Move Keys") );
 
+	const float MouseTime = VirtualTrackArea.PixelToTime(LocalMousePos.X);
+
 	TSet<UMovieSceneSection*> ModifiedSections;
 	for( FSelectedKey SelectedKey : SelectedKeys )
 	{
 		UMovieSceneSection* OwningSection = SelectedKey.Section;
 
+		RelativeOffsets.Add(SelectedKey, SelectedKey.KeyArea->GetKeyTime(SelectedKey.KeyHandle.GetValue()) - MouseTime);
+
 		// Only modify sections once
 		if( !ModifiedSections.Contains( OwningSection ) )
 		{
@@ -432,137 +481,89 @@ void FMoveKeys::OnBeginDrag(const FVector2D& LocalMousePos, const FVirtualTrackA
 	}
 }
 
-void FMoveKeys::OnEndDrag()
-{
-	EndTransaction();
-}
-
 void FMoveKeys::OnDrag( const FPointerEvent& MouseEvent, const FVector2D& LocalMousePos, const FVirtualTrackArea& VirtualTrackArea )
 {
 	float MouseTime = VirtualTrackArea.PixelToTime(LocalMousePos.X);
 	float DistanceMoved = MouseTime - VirtualTrackArea.PixelToTime(LocalMousePos.X - MouseEvent.GetCursorDelta().X);
 
-	if( DistanceMoved != 0.0f )
+	if( DistanceMoved == 0.0f )
 	{
-		float TimeDelta = DistanceMoved;
-		// Snapping
-		if ( Settings->GetIsSnapEnabled() )
+		return;
+	}
+
+	// Snapping
+	if ( Settings->GetIsSnapEnabled() )
+	{
+		TArray<float> KeyTimes;
+		for( FSelectedKey SelectedKey : SelectedKeys )
 		{
-			bool bSnappedToKeyTime = false;
-			if ( Settings->GetSnapKeyTimesToKeys() && CardinalTrack.IsValid() )
-			{
-				TArray<float> OutSnapTimes;
-				GetKeySnapTimes(OutSnapTimes, CardinalTrack);
+			KeyTimes.Add(MouseTime + RelativeOffsets.FindRef(SelectedKey));
+		}
 
-				TArray<float> InitialTimes;
-				for ( FSelectedKey SelectedKey : SelectedKeys )
-				{
-					InitialTimes.Add(SelectedKey.KeyArea->GetKeyTime(SelectedKey.KeyHandle.GetValue()) + DistanceMoved);
-				}
-				float OutInitialTime = 0.f;
-				float OutSnapTime = 0.f;
-				if ( SnapToTimes( InitialTimes, OutSnapTimes, VirtualTrackArea, OutInitialTime, OutSnapTime ) )
-				{
-					bSnappedToKeyTime = true;
-					TimeDelta = OutSnapTime - (OutInitialTime - DistanceMoved);
-				}
-			}
+		float SnapThresold = VirtualTrackArea.PixelToTime(10.f) - VirtualTrackArea.PixelToTime(0.f);
 
-			if (Settings->GetSnapKeyTimesToInterval())
-			{
-				// todo: arodham: fix snapping
-				// float SnapTimeIntervalDelta = Settings->SnapTimeToInterval(MouseTime) - SelectedKeyTime;
-
-				// // Snap to time interval only if we haven't snapped to a keyframe or if the time difference is smaller than the difference to the keyframe.
-				// if (bSnappedToKeyTime)
-				// {					
-				// 	if (FMath::Abs(SnapTimeIntervalDelta) < FMath::Abs(TimeDelta))
-				// 	{
-				// 		TimeDelta = SnapTimeIntervalDelta;
-				// 	}
-				// }
-				// else
-				// {
-				// 	TimeDelta = SnapTimeIntervalDelta;
-				// }
-			}
+		TOptional<FSequencerSnapField::FSnapResult> SnappedTime;
+
+		if (Settings->GetSnapKeyTimesToKeys())
+		{
+			SnappedTime = SnapField->Snap(KeyTimes, SnapThresold);
 		}
 
-		float PrevNewKeyTime = FLT_MAX;
+		if (!SnappedTime.IsSet() && Settings->GetSnapKeyTimesToInterval())
+		{
+			SnappedTime = SnapToInterval(KeyTimes, SnapThresold, *Settings);
+		}
 
-		for( FSelectedKey SelectedKey : SelectedKeys )
+		if (SnappedTime.IsSet())
 		{
-			UMovieSceneSection* Section = SelectedKey.Section;
+			MouseTime += SnappedTime->Snapped - SnappedTime->Original;
+		}
+	}
 
-			TSharedPtr<IKeyArea>& KeyArea = SelectedKey.KeyArea;
+	float PrevNewKeyTime = FLT_MAX;
 
+	for( FSelectedKey SelectedKey : SelectedKeys )
+	{
+		UMovieSceneSection* Section = SelectedKey.Section;
 
-			// Tell the key area to move the key.  We reset the key index as a result of the move because moving a key can change it's internal index 
-			KeyArea->MoveKey( SelectedKey.KeyHandle.GetValue(), TimeDelta );
+		TSharedPtr<IKeyArea>& KeyArea = SelectedKey.KeyArea;
 
-			// Update the key that moved
-			float NewKeyTime = KeyArea->GetKeyTime( SelectedKey.KeyHandle.GetValue() );
+		float NewKeyTime = MouseTime + RelativeOffsets.FindRef(SelectedKey);
+		float CurrentTime = KeyArea->GetKeyTime( SelectedKey.KeyHandle.GetValue() );
 
-			// If the key moves outside of the section resize the section to fit the key
-			// @todo Sequencer - Doesn't account for hitting other sections 
-			if( NewKeyTime > Section->GetEndTime() )
-			{
-				Section->SetEndTime( NewKeyTime );
-			}
-			else if( NewKeyTime < Section->GetStartTime() )
-			{
-				Section->SetStartTime( NewKeyTime );
-			}
+		// Tell the key area to move the key.  We reset the key index as a result of the move because moving a key can change it's internal index 
+		KeyArea->MoveKey( SelectedKey.KeyHandle.GetValue(), NewKeyTime - CurrentTime );
 
-			if (PrevNewKeyTime == FLT_MAX)
-			{
-				PrevNewKeyTime = NewKeyTime;
-			}
-			else if (!FMath::IsNearlyEqual(NewKeyTime, PrevNewKeyTime))
-			{
-				PrevNewKeyTime = -FLT_MAX;
-			}
+		// If the key moves outside of the section resize the section to fit the key
+		// @todo Sequencer - Doesn't account for hitting other sections 
+		if( NewKeyTime > Section->GetEndTime() )
+		{
+			Section->SetEndTime( NewKeyTime );
+		}
+		else if( NewKeyTime < Section->GetStartTime() )
+		{
+			Section->SetStartTime( NewKeyTime );
 		}
 
-		// Snap the play time to the new dragged key time if all the keyframes were dragged to the same time
-		if (Settings->GetSnapPlayTimeToDraggedKey() && PrevNewKeyTime != FLT_MAX && PrevNewKeyTime != -FLT_MAX)
+		if (PrevNewKeyTime == FLT_MAX)
+		{
+			PrevNewKeyTime = NewKeyTime;
+		}
+		else if (!FMath::IsNearlyEqual(NewKeyTime, PrevNewKeyTime))
 		{
-			Sequencer.SetGlobalTime(PrevNewKeyTime);
+			PrevNewKeyTime = -FLT_MAX;
 		}
 	}
-}
-
-void FMoveKeys::GetKeySnapTimes(TArray<float>& OutSnapTimes, TSharedPtr<FTrackNode> SequencerNode)
-{
-	// snap to non-selected keys
-	TArray< TSharedRef<FSectionKeyAreaNode> > KeyAreaNodes;
-	SequencerNode->GetChildKeyAreaNodesRecursively(KeyAreaNodes);
 
-	for (int32 NodeIndex = 0; NodeIndex < KeyAreaNodes.Num(); ++NodeIndex)
+	// Snap the play time to the new dragged key time if all the keyframes were dragged to the same time
+	if (Settings->GetSnapPlayTimeToDraggedKey() && PrevNewKeyTime != FLT_MAX && PrevNewKeyTime != -FLT_MAX)
 	{
-		TArray< TSharedRef<IKeyArea> > KeyAreas = KeyAreaNodes[NodeIndex]->GetAllKeyAreas();
-		for (int32 KeyAreaIndex = 0; KeyAreaIndex < KeyAreas.Num(); ++KeyAreaIndex)
-		{
-			auto KeyArea = KeyAreas[KeyAreaIndex];
-
-			TArray<FKeyHandle> KeyHandles = KeyArea->GetUnsortedKeyHandles();
-			for( int32 KeyIndex = 0; KeyIndex < KeyHandles.Num(); ++KeyIndex )
-			{
-				FKeyHandle KeyHandle = KeyHandles[KeyIndex];
-				bool bKeyIsSnappable = true;
-				for ( FSelectedKey SelectedKey : SelectedKeys )
-				{
-					if (SelectedKey.KeyArea == KeyArea && SelectedKey.KeyHandle.GetValue() == KeyHandle)
-					{
-						bKeyIsSnappable = false;
-						break;
-					}
-				}
-				if (bKeyIsSnappable)
-				{
-					OutSnapTimes.Add(KeyArea->GetKeyTime(KeyHandle));
-				}
-			}
-		}
+		Sequencer.SetGlobalTime(PrevNewKeyTime);
 	}
 }
+
+void FMoveKeys::OnEndDrag()
+{
+	EndTransaction();
+	SSection::EnableLayoutRegeneration();
+}
\ No newline at end of file
diff --git a/Engine/Source/Editor/Sequencer/Private/Tools/EditToolDragOperations.h b/Engine/Source/Editor/Sequencer/Private/Tools/EditToolDragOperations.h
--- a/Engine/Source/Editor/Sequencer/Private/Tools/EditToolDragOperations.h
+++ b/Engine/Source/Editor/Sequencer/Private/Tools/EditToolDragOperations.h
@@ -5,6 +5,7 @@
 #include "ISequencerEditTool.h"
 #include "ScopedTransaction.h"
 #include "SequencerHotspots.h"
+#include "SequencerSnapField.h"
 
 class USequencerSettings;
 
@@ -42,7 +43,7 @@ protected:
 class FResizeSection : public FEditToolDragOperation
 {
 public:
-	FResizeSection( FSequencer& Sequencer, TArray<FSectionHandle> Sections, TOptional<FSectionHandle> InCardinalSection, bool bInDraggingByEnd );
+	FResizeSection( FSequencer& InSequencer, TArray<FSectionHandle> Sections, TOptional<FSectionHandle> InCardinalSection, bool bInDraggingByEnd );
 
 	/** FEditToolDragOperation interface */
 	virtual void OnBeginDrag(const FVector2D& LocalMousePos, const FVirtualTrackArea& VirtualTrackArea) override;
@@ -68,7 +69,7 @@ private:
 class FMoveSection : public FEditToolDragOperation
 {
 public:
-	FMoveSection( FSequencer& Sequencer, TArray<FSectionHandle> Sections );
+	FMoveSection( FSequencer& InSequencer, TArray<FSectionHandle> Sections );
 
 	/** FEditToolDragOperation interface */
 	virtual void OnBeginDrag(const FVector2D& LocalMousePos, const FVirtualTrackArea& VirtualTrackArea) override;
@@ -93,11 +94,11 @@ private:
 class FMoveKeys : public FEditToolDragOperation
 {
 public:
-	FMoveKeys( FSequencer& Sequencer, const TSet<FSelectedKey>& InSelectedKeys, TSharedPtr<FTrackNode> InCardinalTrack )
-		: FEditToolDragOperation(Sequencer)
+	FMoveKeys( FSequencer& InSequencer, const TSet<FSelectedKey>& InSelectedKeys )
+		: FEditToolDragOperation(InSequencer)
 		, SelectedKeys( InSelectedKeys )
-		, CardinalTrack(MoveTemp(InCardinalTrack))
-	{}
+	{
+	}
 
 	/** FEditToolDragOperation interface */
 	virtual void OnBeginDrag(const FVector2D& LocalMousePos, const FVirtualTrackArea& VirtualTrackArea) override;
@@ -109,9 +110,11 @@ protected:
 	/** Gets all potential snap times for this drag operation */
 	void GetKeySnapTimes(TArray<float>& OutSnapTimes, TSharedPtr<FTrackNode> SequencerNode);
 
-private:
+private:	
 	/** The selected keys being moved. */
 	const TSet<FSelectedKey>& SelectedKeys;
-	/** Optional cardinal track to which the dragged key belongs */
-	TSharedPtr<FTrackNode> CardinalTrack;
+	/** Map of relative offsets from the original mouse position */
+	TMap<FSelectedKey, float> RelativeOffsets;
+	/** Snap field used to assist in snapping calculations */
+	TOptional<FSequencerSnapField> SnapField;
 };
diff --git a/Engine/Source/Editor/Sequencer/Private/Tools/SequencerEditTool_Movement.cpp b/Engine/Source/Editor/Sequencer/Private/Tools/SequencerEditTool_Movement.cpp
--- a/Engine/Source/Editor/Sequencer/Private/Tools/SequencerEditTool_Movement.cpp
+++ b/Engine/Source/Editor/Sequencer/Private/Tools/SequencerEditTool_Movement.cpp
@@ -43,7 +43,7 @@ FReply FSequencerEditTool_Movement::OnMouseMove(SWidget& OwnerWidget, const FGeo
 {
 	if (DelayedDrag.IsSet())
 	{
-		const FVirtualTrackArea VirtualTrackArea = SequencerWidget.Pin()->GetVirtualTrackArea(MyGeometry);
+		const FVirtualTrackArea VirtualTrackArea = SequencerWidget.Pin()->GetVirtualTrackArea();
 
 		FReply Reply = FReply::Handled();
 
@@ -124,7 +124,7 @@ TSharedPtr<IEditToolDragOperation> FSequencerEditTool_Movement::CreateDrag()
 				Selection.AddToSelection(ThisKey);
 			}
 
-			return MakeShareable( new FMoveKeys( *PinnedSequencer, Selection.GetSelectedKeys(), StaticCastSharedPtr<FKeyHotspot>(DelayedDrag->Hotspot)->TrackNode ) );
+			return MakeShareable( new FMoveKeys( *PinnedSequencer, Selection.GetSelectedKeys() ) );
 		}
 	}
 	// If we're not dragging a hotspot, sections take precedence over keys
@@ -134,7 +134,7 @@ TSharedPtr<IEditToolDragOperation> FSequencerEditTool_Movement::CreateDrag()
 	}
 	else if (Selection.GetSelectedKeys().Num())
 	{
-		return MakeShareable( new FMoveKeys( *PinnedSequencer, Selection.GetSelectedKeys(), nullptr ) );
+		return MakeShareable( new FMoveKeys( *PinnedSequencer, Selection.GetSelectedKeys() ) );
 	}
 
 	return nullptr;
diff --git a/Engine/Source/Editor/Sequencer/Private/Tools/SequencerEditTool_Selection.cpp b/Engine/Source/Editor/Sequencer/Private/Tools/SequencerEditTool_Selection.cpp
--- a/Engine/Source/Editor/Sequencer/Private/Tools/SequencerEditTool_Selection.cpp
+++ b/Engine/Source/Editor/Sequencer/Private/Tools/SequencerEditTool_Selection.cpp
@@ -1,6 +1,7 @@
 // Copyright 1998-2015 Epic Games, Inc. All Rights Reserved.
 
 #include "SequencerPrivatePCH.h"
+#include "SequencerEntityVisitor.h"
 #include "SequencerEditTool_Selection.h"
 #include "MovieSceneSection.h"
 #include "ISequencerSection.h"
@@ -45,170 +46,6 @@ struct FMarqueeSelectData
 	FVector2D CurrentPosition;
 };
 
-
-struct IMarqueeVisitor
-{
-	virtual void VisitKey(FKeyHandle KeyHandle, float KeyTime, const TSharedPtr<IKeyArea>& KeyArea, UMovieSceneSection* Section) const { }
-	virtual void VisitSection(UMovieSceneSection* Section) const { }
-};
-
-/** Struct used to iterate a marquee range with a user-supplied visitor */
-/* @todo: This could probably live in a header so it can be used elsewhere */
-/* @todo: Could probably optimize this by not walking every single node, and binary searching the begin<->end ranges instead */
-struct FMarqueeRange
-{
-	/** Construction from the range itself */
-	FMarqueeRange(FVector2D InTopLeft, FVector2D InBottomRight, FVector2D InVirtualKeySize)
-		: TopLeft(InTopLeft), BottomRight(InBottomRight), VirtualKeySize(InVirtualKeySize)
-	{}
-
-	/** Visit the specified nodes (recursively) with this range and a user-supplied visitor */
-	void Visit(const IMarqueeVisitor& Visitor, const TArray< TSharedRef<FSequencerDisplayNode> >& Nodes);
-
-private:
-
-	/** Handle visitation of a particular node */
-	void HandleNode(const IMarqueeVisitor& Visitor, FSequencerDisplayNode& InNode);
-	/** Handle visitation of a particular node, along with a set of sections */
-	void HandleNode(const IMarqueeVisitor& Visitor, FSequencerDisplayNode& InNode, const TArray<TSharedRef<ISequencerSection>>& InSections);
-	/** Handle visitation of a key area node */
-	void HandleKeyAreaNode(const IMarqueeVisitor& Visitor, FSectionKeyAreaNode& InNode, const TArray<TSharedRef<ISequencerSection>>& InSections);
-	/** Handle visitation of a key area */
-	void HandleKeyArea(const IMarqueeVisitor& Visitor, const TSharedPtr<IKeyArea>& KeyArea, UMovieSceneSection* Section);
-	
-	/** Check whether the specified section intersects this range */
-	bool SectionIntersectsRange(const UMovieSceneSection* InSection) const;
-
-	/** The bounds of the range */
-	FVector2D TopLeft, BottomRight;
-
-	/** Key size in virtual space */
-	FVector2D VirtualKeySize;
-};
-
-bool FMarqueeRange::SectionIntersectsRange(const UMovieSceneSection* InSection) const
-{
-	return InSection->GetStartTime() <= BottomRight.X && InSection->GetEndTime() >= TopLeft.X;
-}
-
-void FMarqueeRange::Visit(const IMarqueeVisitor& Visitor, const TArray< TSharedRef<FSequencerDisplayNode> >& Nodes)
-{
-	for (auto& Child : Nodes)
-	{
-		if (!Child->IsHidden())
-		{
-			HandleNode(Visitor, *Child);
-		}
-	}
-}
-
-void FMarqueeRange::HandleNode(const IMarqueeVisitor& Visitor, FSequencerDisplayNode& InNode)
-{
-	if (InNode.GetType() == ESequencerNode::Track)
-	{
-		HandleNode(Visitor, InNode, static_cast<FTrackNode&>(InNode).GetSections());
-	}
-
-	if (InNode.IsExpanded())
-	{
-		for (auto& Child : InNode.GetChildNodes())
-		{
-			if (!Child->IsHidden())
-			{
-				HandleNode(Visitor, *Child);
-			}
-		}
-	}
-}
-
-void FMarqueeRange::HandleNode(const IMarqueeVisitor& Visitor, FSequencerDisplayNode& InNode, const TArray<TSharedRef<ISequencerSection>>& InSections)
-{
-	const float NodeCenter = InNode.GetVirtualTop() + (InNode.GetVirtualBottom() - InNode.GetVirtualTop())/2;
-	if (NodeCenter + VirtualKeySize.Y/2 > TopLeft.Y &&
-		NodeCenter - VirtualKeySize.Y/2 < BottomRight.Y)
-	{
-		bool bNodeHasKeyArea = false;
-		if (InNode.GetType() == ESequencerNode::KeyArea)
-		{
-			HandleKeyAreaNode(Visitor, static_cast<FSectionKeyAreaNode&>(InNode), InSections);
-			bNodeHasKeyArea = true;
-		}
-		else if (InNode.GetType() == ESequencerNode::Track)
-		{
-			TSharedPtr<FSectionKeyAreaNode> SectionKeyNode = static_cast<FTrackNode&>(InNode).GetTopLevelKeyNode();
-			if (SectionKeyNode.IsValid())
-			{
-				HandleKeyAreaNode(Visitor, static_cast<FSectionKeyAreaNode&>(InNode), InSections);
-				bNodeHasKeyArea = true;
-			}
-		}
-
-		if (!bNodeHasKeyArea && (!InNode.IsExpanded() || InNode.GetChildNodes().Num() == 0))
-		{
-			// As a fallback, we need to handle:
-			//  - Key groupings on collapsed parents
-			//  - Sections that have no key areas
-			for (int32 SectionIndex = 0; SectionIndex < InSections.Num(); ++SectionIndex)
-			{
-				UMovieSceneSection* Section = InSections[SectionIndex]->GetSectionObject();
-				if (SectionIntersectsRange(Section))
-				{
-					Visitor.VisitSection(Section);
-
-					// Only handle grouped keys if we actually have children
-					if (InNode.GetChildNodes().Num() != 0)
-					{
-						TSharedRef<IKeyArea> KeyArea = InNode.UpdateKeyGrouping(SectionIndex);
-						HandleKeyArea(Visitor, KeyArea, Section);
-					}
-				}
-			}
-		}
-	}
-
-	if (InNode.IsExpanded())
-	{
-		// Handle Children
-		for (auto& Child : InNode.GetChildNodes())
-		{
-			if (!Child->IsHidden())
-			{
-				HandleNode(Visitor, *Child, InSections);
-			}
-		}
-	}
-}
-
-void FMarqueeRange::HandleKeyAreaNode(const IMarqueeVisitor& Visitor, FSectionKeyAreaNode& KeyAreaNode, const TArray<TSharedRef<ISequencerSection>>& InSections)
-{
-	for( int32 SectionIndex = 0; SectionIndex < InSections.Num(); ++SectionIndex )
-	{
-		UMovieSceneSection* Section = InSections[SectionIndex]->GetSectionObject();
-
-		// If the section is at all within the marquee, we check its keys
-		if (SectionIntersectsRange(Section))
-		{
-			Visitor.VisitSection(Section);
-
-			TSharedRef<IKeyArea> KeyArea = KeyAreaNode.GetKeyArea(SectionIndex);
-			HandleKeyArea(Visitor, KeyArea, InSections[SectionIndex]->GetSectionObject());
-		}
-	}
-}
-
-void FMarqueeRange::HandleKeyArea(const IMarqueeVisitor& Visitor, const TSharedPtr<IKeyArea>& KeyArea, UMovieSceneSection* Section)
-{
-	for (FKeyHandle KeyHandle : KeyArea->GetUnsortedKeyHandles())
-	{
-		float KeyPosition = KeyArea->GetKeyTime(KeyHandle);
-		if (KeyPosition + VirtualKeySize.X/2 > TopLeft.X &&
-			KeyPosition - VirtualKeySize.X/2 < BottomRight.X)
-		{
-			Visitor.VisitKey(KeyHandle, KeyPosition, KeyArea, Section);
-		}
-	}
-}
-
 /** todo: Save this in a config? */
 ESequencerSelectionMode LastUsedMode = ESequencerSelectionMode::Keys;
 
@@ -252,7 +89,7 @@ int32 FSequencerEditTool_Selection::OnPaint(const FGeometry& AllottedGeometry, c
 	if (MarqueeSelectData.IsValid() && MarqueeSelectData->HasDragged())
 	{
 		// Convert to physical space for rendering
-		const FVirtualTrackArea VirtualTrackArea = SequencerWidget.Pin()->GetVirtualTrackArea(AllottedGeometry);
+		const FVirtualTrackArea VirtualTrackArea = SequencerWidget.Pin()->GetVirtualTrackArea();
 
 		FVector2D TopLeft = VirtualTrackArea.VirtualToPhysical(MarqueeSelectData->TopLeft());
 		FVector2D BottomRight = VirtualTrackArea.VirtualToPhysical(MarqueeSelectData->BottomRight());
@@ -295,7 +132,7 @@ FReply FSequencerEditTool_Selection::OnMouseMove(SWidget& OwnerWidget, const FGe
 		}
 		else if (DelayedDrag->AttemptDragStart(MouseEvent))
 		{
-			const FVirtualTrackArea VirtualTrackArea = SequencerWidget.Pin()->GetVirtualTrackArea(MyGeometry);
+			const FVirtualTrackArea VirtualTrackArea = SequencerWidget.Pin()->GetVirtualTrackArea();
 
 			// Marquee data is stored in virtual space
 			MarqueeSelectData.Reset(new FMarqueeSelectData(VirtualTrackArea.PhysicalToVirtual(DelayedDrag->GetInitialPosition())));
@@ -361,7 +198,7 @@ void FSequencerEditTool_Selection::SetNewMarqueeBounds(const FPointerEvent& Mous
 		}
 	}
 
-	const FVirtualTrackArea VirtualTrackArea = SequencerWidget.Pin()->GetVirtualTrackArea(MyGeometry);
+	const FVirtualTrackArea VirtualTrackArea = SequencerWidget.Pin()->GetVirtualTrackArea();
 
 	// Clamp the vertical position to the actual bounds of the track area
 	LocalPosition.Y = FMath::Clamp(LocalPosition.Y, 0.f, MyGeometry.GetLocalSize().Y);
@@ -392,7 +229,7 @@ void FSequencerEditTool_Selection::SetNewMarqueeBounds(const FPointerEvent& Mous
 	}
 }
 
-struct FSelectionVisitor : IMarqueeVisitor
+struct FSelectionVisitor : ISequencerEntityVisitor
 {
 	FSelectionVisitor(FSequencerSelection& InSelection, const FPointerEvent& InMouseEvent, ESequencerSelectionMode InMode)
 		: Selection(InSelection)
@@ -481,10 +318,11 @@ void FSequencerEditTool_Selection::HandleMarqueeSelection(const FGeometry& Geome
 	// Vertically, virtual units == physical units
 	VirtualKeySize.Y = SequencerSectionConstants::KeySize.Y;
 
-	FMarqueeRange Range(MarqueeSelectData->TopLeft(), MarqueeSelectData->BottomRight(), VirtualKeySize);
+	FSequencerEntityWalker Walker(FSequencerEntityRange(MarqueeSelectData->TopLeft(), MarqueeSelectData->BottomRight()), VirtualKeySize);
+
 	FSelectionVisitor Visitor(Selection, MouseEvent, SelectionMode);
 
-	Range.Visit(Visitor, SequencerWidget.Pin()->GetTreeView()->GetNodeTree()->GetRootNodes());
+	Walker.Traverse(Visitor, SequencerWidget.Pin()->GetTreeView()->GetNodeTree()->GetRootNodes());
 
 	Selection.ResumeBroadcast();
 	Selection.GetOnKeySelectionChanged().Broadcast();
diff --git a/Engine/Source/Editor/Sequencer/Private/Tools/SequencerEntityVisitor.cpp b/Engine/Source/Editor/Sequencer/Private/Tools/SequencerEntityVisitor.cpp
new file mode 100644
--- /dev/null
+++ b/Engine/Source/Editor/Sequencer/Private/Tools/SequencerEntityVisitor.cpp
@@ -0,0 +1,167 @@
+// Copyright 1998-2015 Epic Games, Inc. All Rights Reserved.
+
+#include "SequencerPrivatePCH.h"
+#include "Sequencer.h"
+#include "SequencerDisplayNode.h"
+#include "SequencerEntityVisitor.h"
+#include "MovieSceneSection.h"
+#include "ISequencerSection.h"
+
+FSequencerEntityRange::FSequencerEntityRange(const TRange<float>& InRange)
+	: StartTime(InRange.GetLowerBoundValue()), EndTime(InRange.GetUpperBoundValue())
+{
+}
+
+FSequencerEntityRange::FSequencerEntityRange(FVector2D TopLeft, FVector2D BottomRight)
+	: StartTime(TopLeft.X), EndTime(BottomRight.X)
+	, VerticalTop(TopLeft.Y), VerticalBottom(BottomRight.Y)
+{
+}
+
+bool FSequencerEntityRange::IntersectSection(const UMovieSceneSection* InSection) const
+{
+	return InSection->GetStartTime() <= EndTime && InSection->GetEndTime() >= StartTime;
+}
+
+bool FSequencerEntityRange::IntersectNode(const FSequencerDisplayNode& InNode) const
+{
+	if (VerticalTop.IsSet())
+	{
+		return InNode.GetVirtualTop() <= VerticalBottom.GetValue() && InNode.GetVirtualBottom() >= VerticalTop.GetValue();
+	}
+	return true;
+}
+
+bool FSequencerEntityRange::IntersectKeyArea(const FSequencerDisplayNode& InNode, float VirtualKeyHeight) const
+{
+	if (VerticalTop.IsSet())
+	{
+		const float NodeCenter = InNode.GetVirtualTop() + (InNode.GetVirtualBottom() - InNode.GetVirtualTop())/2;
+		return NodeCenter + VirtualKeyHeight/2 > VerticalTop.GetValue() && NodeCenter - VirtualKeyHeight/2 < VerticalBottom.GetValue();
+	}
+	return true;
+}
+
+FSequencerEntityWalker::FSequencerEntityWalker(const FSequencerEntityRange& InRange, FVector2D InVirtualKeySize)
+	: Range(InRange), VirtualKeySize(InVirtualKeySize)
+{}
+
+/* @todo: Could probably optimize this by not walking every single node, and binary searching the begin<->end ranges instead */
+void FSequencerEntityWalker::Traverse(const ISequencerEntityVisitor& Visitor, const TArray< TSharedRef<FSequencerDisplayNode> >& Nodes)
+{
+	for (auto& Child : Nodes)
+	{
+		if (!Child->IsHidden())
+		{
+			HandleNode(Visitor, *Child);
+		}
+	}
+}
+
+void FSequencerEntityWalker::HandleNode(const ISequencerEntityVisitor& Visitor, FSequencerDisplayNode& InNode)
+{
+	if (InNode.GetType() == ESequencerNode::Track)
+	{
+		HandleNode(Visitor, InNode, static_cast<FTrackNode&>(InNode).GetSections());
+	}
+
+	if (InNode.IsExpanded())
+	{
+		for (auto& Child : InNode.GetChildNodes())
+		{
+			if (!Child->IsHidden())
+			{
+				HandleNode(Visitor, *Child);
+			}
+		}
+	}
+}
+
+void FSequencerEntityWalker::HandleNode(const ISequencerEntityVisitor& Visitor, FSequencerDisplayNode& InNode, const TArray<TSharedRef<ISequencerSection>>& InSections)
+{
+	if (Range.IntersectNode(InNode))
+	{
+		bool bNodeHasKeyArea = false;
+		if (InNode.GetType() == ESequencerNode::KeyArea)
+		{
+			HandleKeyAreaNode(Visitor, static_cast<FSectionKeyAreaNode&>(InNode), InSections);
+			bNodeHasKeyArea = true;
+		}
+		else if (InNode.GetType() == ESequencerNode::Track)
+		{
+			TSharedPtr<FSectionKeyAreaNode> SectionKeyNode = static_cast<FTrackNode&>(InNode).GetTopLevelKeyNode();
+			if (SectionKeyNode.IsValid())
+			{
+				HandleKeyAreaNode(Visitor, static_cast<FSectionKeyAreaNode&>(InNode), InSections);
+				bNodeHasKeyArea = true;
+			}
+		}
+
+		if (!bNodeHasKeyArea && (!InNode.IsExpanded() || InNode.GetChildNodes().Num() == 0))
+		{
+			// As a fallback, we need to handle:
+			//  - Key groupings on collapsed parents
+			//  - Sections that have no key areas
+			for (int32 SectionIndex = 0; SectionIndex < InSections.Num(); ++SectionIndex)
+			{
+				UMovieSceneSection* Section = InSections[SectionIndex]->GetSectionObject();
+				if (Range.IntersectSection(Section))
+				{
+					Visitor.VisitSection(Section);
+
+					// Only handle grouped keys if we actually have children
+					if (InNode.GetChildNodes().Num() != 0 && Range.IntersectKeyArea(InNode, VirtualKeySize.X))
+					{
+						TSharedRef<IKeyArea> KeyArea = InNode.UpdateKeyGrouping(SectionIndex);
+						HandleKeyArea(Visitor, KeyArea, Section);
+					}
+				}
+			}
+		}
+	}
+
+	if (InNode.IsExpanded())
+	{
+		// Handle Children
+		for (auto& Child : InNode.GetChildNodes())
+		{
+			if (!Child->IsHidden())
+			{
+				HandleNode(Visitor, *Child, InSections);
+			}
+		}
+	}
+}
+
+void FSequencerEntityWalker::HandleKeyAreaNode(const ISequencerEntityVisitor& Visitor, FSectionKeyAreaNode& KeyAreaNode, const TArray<TSharedRef<ISequencerSection>>& InSections)
+{
+	for( int32 SectionIndex = 0; SectionIndex < InSections.Num(); ++SectionIndex )
+	{
+		UMovieSceneSection* Section = InSections[SectionIndex]->GetSectionObject();
+
+		// If the section is at all within the marquee, we check its keys
+		if (Range.IntersectSection(Section))
+		{
+			Visitor.VisitSection(Section);
+
+			if (Range.IntersectKeyArea(KeyAreaNode, VirtualKeySize.X))
+			{
+				TSharedRef<IKeyArea> KeyArea = KeyAreaNode.GetKeyArea(SectionIndex);
+				HandleKeyArea(Visitor, KeyArea, InSections[SectionIndex]->GetSectionObject());
+			}
+		}
+	}
+}
+
+void FSequencerEntityWalker::HandleKeyArea(const ISequencerEntityVisitor& Visitor, const TSharedPtr<IKeyArea>& KeyArea, UMovieSceneSection* Section)
+{
+	for (FKeyHandle KeyHandle : KeyArea->GetUnsortedKeyHandles())
+	{
+		float KeyPosition = KeyArea->GetKeyTime(KeyHandle);
+		if (KeyPosition + VirtualKeySize.X/2 > Range.StartTime &&
+			KeyPosition - VirtualKeySize.X/2 < Range.EndTime)
+		{
+			Visitor.VisitKey(KeyHandle, KeyPosition, KeyArea, Section);
+		}
+	}
+}
\ No newline at end of file
diff --git a/Engine/Source/Editor/Sequencer/Private/Tools/SequencerEntityVisitor.h b/Engine/Source/Editor/Sequencer/Private/Tools/SequencerEntityVisitor.h
new file mode 100644
--- /dev/null
+++ b/Engine/Source/Editor/Sequencer/Private/Tools/SequencerEntityVisitor.h
@@ -0,0 +1,65 @@
+// Copyright 1998-2015 Epic Games, Inc. All Rights Reserved.
+
+#pragma once
+
+class FSequencerDisplayNode;
+class UMovieSceneSection;
+class IKeyArea;
+class FSectionKeyAreaNode;
+class ISequencerSection;
+
+/** Visitor class used to handle specific sequencer entities */
+struct ISequencerEntityVisitor
+{
+	virtual void VisitKey(FKeyHandle KeyHandle, float KeyTime, const TSharedPtr<IKeyArea>& KeyArea, UMovieSceneSection* Section) const { }
+	virtual void VisitSection(UMovieSceneSection* Section) const { }
+};
+
+/** A range specifying time (and possibly vertical) bounds in the sequencer */
+struct FSequencerEntityRange
+{
+	FSequencerEntityRange(const TRange<float>& InRange);
+	FSequencerEntityRange(FVector2D TopLeft, FVector2D BottomRight);
+
+	/** Check whether the specified section intersects this range */
+	bool IntersectSection(const UMovieSceneSection* InSection) const;
+
+	/** Check whether the specified node intersects this range */
+	bool IntersectNode(const FSequencerDisplayNode& InNode) const;
+
+	/** Check whether the specified node's key area intersects this range */
+	bool IntersectKeyArea(const FSequencerDisplayNode& InNode, float VirtualKeyHeight) const;
+
+	/** Start/end times */
+	float StartTime, EndTime;
+
+	/** Optional vertical bounds */
+	TOptional<float> VerticalTop, VerticalBottom;
+};
+
+/** Struct used to iterate a two dimensional *visible* range with a user-supplied visitor */
+struct FSequencerEntityWalker
+{
+	/** Construction from the range itself, and an optional virtual key size, where key bounds must be taken into consideration */
+	FSequencerEntityWalker(const FSequencerEntityRange& InRange, FVector2D InVirtualKeySize = FVector2D());
+
+	/** Visit the specified nodes (recursively) with this range and a user-supplied visitor */
+	void Traverse(const ISequencerEntityVisitor& Visitor, const TArray< TSharedRef<FSequencerDisplayNode> >& Nodes);
+
+private:
+
+	/** Handle visitation of a particular node */
+	void HandleNode(const ISequencerEntityVisitor& Visitor, FSequencerDisplayNode& InNode);
+	/** Handle visitation of a particular node, along with a set of sections */
+	void HandleNode(const ISequencerEntityVisitor& Visitor, FSequencerDisplayNode& InNode, const TArray<TSharedRef<ISequencerSection>>& InSections);
+	/** Handle visitation of a key area node */
+	void HandleKeyAreaNode(const ISequencerEntityVisitor& Visitor, FSectionKeyAreaNode& InNode, const TArray<TSharedRef<ISequencerSection>>& InSections);
+	/** Handle visitation of a key area */
+	void HandleKeyArea(const ISequencerEntityVisitor& Visitor, const TSharedPtr<IKeyArea>& KeyArea, UMovieSceneSection* Section);
+
+	/** The bounds of the range */
+	FSequencerEntityRange Range;
+
+	/** Key size in virtual space */
+	FVector2D VirtualKeySize;
+};
\ No newline at end of file
diff --git a/Engine/Source/Editor/Sequencer/Private/Tools/SequencerSnapField.cpp b/Engine/Source/Editor/Sequencer/Private/Tools/SequencerSnapField.cpp
new file mode 100644
--- /dev/null
+++ b/Engine/Source/Editor/Sequencer/Private/Tools/SequencerSnapField.cpp
@@ -0,0 +1,169 @@
+// Copyright 1998-2015 Epic Games, Inc. All Rights Reserved.
+
+#include "SequencerPrivatePCH.h"
+#include "SequencerSnapField.h"
+#include "SSequencerTreeView.h"
+#include "Sequencer.h"
+#include "SSequencer.h"
+#include "MovieSceneSection.h"
+#include "SequencerEntityVisitor.h"
+
+struct FSnapGridVisitor : ISequencerEntityVisitor
+{
+	FSnapGridVisitor(ISequencerSnapCandidate& InCandidate)
+		: Candidate(InCandidate)
+	{}
+
+	virtual void VisitKey(FKeyHandle KeyHandle, float KeyTime, const TSharedPtr<IKeyArea>& KeyArea, UMovieSceneSection* Section) const
+	{
+		if (Candidate.IsKeyApplicable(KeyHandle, KeyArea, Section))
+		{
+			Snaps.Add(FSequencerSnapPoint{ FSequencerSnapPoint::Key, KeyTime });
+		}
+	}
+	virtual void VisitSection(UMovieSceneSection* Section) const
+	{
+		if (Candidate.AreSectionBoundsApplicable(Section))
+		{
+			Snaps.Add(FSequencerSnapPoint{ FSequencerSnapPoint::SectionBounds, Section->GetStartTime() });
+			Snaps.Add(FSequencerSnapPoint{ FSequencerSnapPoint::SectionBounds, Section->GetEndTime() });
+		}
+
+		if (Candidate.AreSectionCustomSnapsApplicable(Section))
+		{
+			TArray<float> CustomSnaps;
+			Section->GetSnapTimes(CustomSnaps, false);
+			for (float Time : CustomSnaps)
+			{
+				Snaps.Add(FSequencerSnapPoint{ FSequencerSnapPoint::CustomSection, Time });
+			}
+		}
+	}
+
+	mutable ISequencerSnapCandidate& Candidate;
+	mutable TArray<FSequencerSnapPoint> Snaps;
+};
+
+FSequencerSnapField::FSequencerSnapField(const ISequencer& InSequencer, ISequencerSnapCandidate& Candidate)
+{
+	TSharedPtr<SSequencerTreeView> TreeView = StaticCastSharedRef<SSequencer>(InSequencer.GetSequencerWidget())->GetTreeView();
+
+	TArray<TSharedRef<FSequencerDisplayNode>> VisibleNodes;
+	for (const SSequencerTreeView::FCachedGeometry& Geometry : TreeView->GetAllVisibleNodes())
+	{
+		VisibleNodes.Add(Geometry.Node);
+	}
+
+	auto ViewRange = InSequencer.GetViewRange();
+	FSequencerEntityWalker Walker(ViewRange);
+
+	// Traverse the visible space, collecting snapping times as we go
+	FSnapGridVisitor Visitor(Candidate);
+	Walker.Traverse(Visitor, VisibleNodes);
+
+	// Sort
+	Visitor.Snaps.Sort([](const FSequencerSnapPoint& A, const FSequencerSnapPoint& B){
+		return A.Time < B.Time;
+	});
+
+	// Remove duplicates
+	for (int32 Index = 0; Index < Visitor.Snaps.Num(); ++Index)
+	{
+		const float CurrentTime = Visitor.Snaps[Index].Time;
+
+		int32 NumToMerge = 0;
+		for (int32 DuplIndex = Index + 1; DuplIndex < Visitor.Snaps.Num(); ++DuplIndex)
+		{
+			if (!FMath::IsNearlyEqual(CurrentTime, Visitor.Snaps[DuplIndex].Time))
+			{
+				break;
+			}
+			++NumToMerge;
+		}
+
+		if (NumToMerge)
+		{
+			Visitor.Snaps.RemoveAt(Index + 1, NumToMerge, false);
+		}
+	}
+
+	SortedSnaps = MoveTemp(Visitor.Snaps);
+}
+
+TOptional<float> FSequencerSnapField::Snap(float InTime, float Threshold) const
+{
+	int32 Min = 0;
+	int32 Max = SortedSnaps.Num();
+
+	// Binary search, then linearly search a range
+	for ( ; Min != Max ; )
+	{
+		int32 SearchIndex = Min + (Max - Min) / 2;
+
+		float ProspectiveSnapPos = SortedSnaps[SearchIndex].Time;
+		if (ProspectiveSnapPos > InTime + Threshold)
+		{
+			Max = SearchIndex;
+		}
+		else if (ProspectiveSnapPos < InTime - Threshold)
+		{
+			Min = SearchIndex + 1;
+		}
+		else
+		{
+			// Linearly search forwards and backwards to find the closest snap
+
+			float SnapDelta = ProspectiveSnapPos - InTime;
+
+			// Search forwards while we're in the threshold
+			for (int32 FwdIndex = SearchIndex+1; FwdIndex < Max-1 && SortedSnaps[FwdIndex].Time < InTime + Threshold; ++FwdIndex)
+			{
+				float ThisSnapDelta = InTime - SortedSnaps[FwdIndex].Time;
+				if (FMath::Abs(ThisSnapDelta) < FMath::Abs(SnapDelta))
+				{
+					SnapDelta = ThisSnapDelta;
+					ProspectiveSnapPos = SortedSnaps[FwdIndex].Time;
+				}
+			}
+
+			// Search backwards while we're in the threshold
+			for (int32 BckIndex = SearchIndex-1; BckIndex >= Min && SortedSnaps[BckIndex].Time > InTime + Threshold; --BckIndex)
+			{
+				float ThisSnapDelta = InTime - SortedSnaps[BckIndex].Time;
+				if (FMath::Abs(ThisSnapDelta) < FMath::Abs(SnapDelta))
+				{
+					SnapDelta = ThisSnapDelta;
+					ProspectiveSnapPos = SortedSnaps[BckIndex].Time;
+				}
+			}
+
+			return ProspectiveSnapPos;
+		}
+	}
+
+	return TOptional<float>();
+}
+
+TOptional<FSequencerSnapField::FSnapResult> FSequencerSnapField::Snap(const TArray<float>& InTimes, float Threshold) const
+{
+	TOptional<FSnapResult> ProspectiveSnap;
+	float SnapDelta = 0.f;
+
+	for (float Time : InTimes)
+	{
+		TOptional<float> ThisSnap = Snap(Time, Threshold);
+		if (!ThisSnap.IsSet())
+		{
+			continue;
+		}
+
+		float ThisSnapDelta = ThisSnap.GetValue() - Time;
+		if (!ProspectiveSnap.IsSet() || FMath::Abs(ThisSnapDelta) < FMath::Abs(SnapDelta))
+		{
+			ProspectiveSnap = FSnapResult{ Time, ThisSnap.GetValue() };
+			SnapDelta = ThisSnapDelta;
+		}
+	}
+
+	return ProspectiveSnap;
+}
\ No newline at end of file
diff --git a/Engine/Source/Editor/Sequencer/Private/Tools/SequencerSnapField.h b/Engine/Source/Editor/Sequencer/Private/Tools/SequencerSnapField.h
new file mode 100644
--- /dev/null
+++ b/Engine/Source/Editor/Sequencer/Private/Tools/SequencerSnapField.h
@@ -0,0 +1,66 @@
+// Copyright 1998-2015 Epic Games, Inc. All Rights Reserved.
+
+#pragma once
+
+struct FKeyHandle;
+
+class ISequencer;
+class IKeyArea;
+class UMovieSceneSection;
+
+/** Structure defining a point to snap to in the sequencer */
+struct FSequencerSnapPoint
+{
+	enum ESnapType { Key, SectionBounds, CustomSection };
+
+	/** The type of snap */
+	ESnapType Type;
+
+	/** The time of the snap */
+	float Time;
+};
+
+/** Interface that defines how to construct an FSequencerSnapField */
+struct ISequencerSnapCandidate
+{
+	/** Return true to include the specified key in the snap field */
+	virtual bool IsKeyApplicable(FKeyHandle KeyHandle, const TSharedPtr<IKeyArea>& KeyArea, UMovieSceneSection* Section) { return true; }
+
+	/** Return true to include the specified section's bounds in the snap field */
+	virtual bool AreSectionBoundsApplicable(UMovieSceneSection* Section) { return true; }
+
+	/** Return true to include the specified section's custom snap points in the snap field */
+	virtual bool AreSectionCustomSnapsApplicable(UMovieSceneSection* Section) { return true; }
+};
+
+/** A snapping field that provides efficient snapping calculations on a range of values */
+class FSequencerSnapField
+{
+public:
+
+	/** A snap result denoting the time that was snapped, and the resulting snapped time */
+	struct FSnapResult
+	{
+		/** The time before it was snapped */
+		float Original;
+		/** The time after it was snapped */
+		float Snapped;
+	};
+
+	/** Construction from a sequencer and a snap canidate implementation */
+	FSequencerSnapField(const ISequencer& InSequencer, ISequencerSnapCandidate& Candidate);
+
+	/** Move construction / assignment */
+	FSequencerSnapField(FSequencerSnapField&& In) : SortedSnaps(MoveTemp(In.SortedSnaps)) {}
+	FSequencerSnapField& operator=(FSequencerSnapField&& In) { SortedSnaps = MoveTemp(In.SortedSnaps); return *this; }
+	
+	/** Snap the specified time to this field with the given threshold */
+	TOptional<float> Snap(float InTime, float Threshold) const;
+
+	/** Snap the specified times to this field with the given threshold. Will return the closest snap value of the entire intersection. */
+	TOptional<FSnapResult> Snap(const TArray<float>& InTimes, float Threshold) const;
+
+private:
+	/** Array of snap points, approximately grouped, and sorted in ascending order by time */
+	TArray<FSequencerSnapPoint> SortedSnaps;
+};
\ No newline at end of file
