Removed mutable keyword

diff --git a/Engine/Source/Editor/Sequencer/Private/Tools/SequencerSnapField.cpp b/Engine/Source/Editor/Sequencer/Private/Tools/SequencerSnapField.cpp
--- a/Engine/Source/Editor/Sequencer/Private/Tools/SequencerSnapField.cpp
+++ b/Engine/Source/Editor/Sequencer/Private/Tools/SequencerSnapField.cpp
@@ -40,7 +40,7 @@ struct FSnapGridVisitor : ISequencerEntityVisitor
 		}
 	}
 
-	mutable ISequencerSnapCandidate& Candidate;
+	ISequencerSnapCandidate& Candidate;
 	mutable TArray<FSequencerSnapPoint> Snaps;
 };
 
