Fixed a crash caused by an invalid cast

Top level key area nodes were not being handled correctly by the entity walker. Changed HandleKeyAreaNode to require the owner node for bounds testing purposes.

diff --git a/Engine/Source/Editor/Sequencer/Private/Tools/SequencerEntityVisitor.cpp b/Engine/Source/Editor/Sequencer/Private/Tools/SequencerEntityVisitor.cpp
--- a/Engine/Source/Editor/Sequencer/Private/Tools/SequencerEntityVisitor.cpp
+++ b/Engine/Source/Editor/Sequencer/Private/Tools/SequencerEntityVisitor.cpp
@@ -84,7 +84,7 @@ void FSequencerEntityWalker::HandleNode(const ISequencerEntityVisitor& Visitor,
 		bool bNodeHasKeyArea = false;
 		if (InNode.GetType() == ESequencerNode::KeyArea)
 		{
-			HandleKeyAreaNode(Visitor, static_cast<FSectionKeyAreaNode&>(InNode), InSections);
+			HandleKeyAreaNode(Visitor, static_cast<FSectionKeyAreaNode&>(InNode), InNode, InSections);
 			bNodeHasKeyArea = true;
 		}
 		else if (InNode.GetType() == ESequencerNode::Track)
@@ -92,7 +92,7 @@ void FSequencerEntityWalker::HandleNode(const ISequencerEntityVisitor& Visitor,
 			TSharedPtr<FSectionKeyAreaNode> SectionKeyNode = static_cast<FTrackNode&>(InNode).GetTopLevelKeyNode();
 			if (SectionKeyNode.IsValid())
 			{
-				HandleKeyAreaNode(Visitor, static_cast<FSectionKeyAreaNode&>(InNode), InSections);
+				HandleKeyAreaNode(Visitor, *SectionKeyNode, InNode, InSections);
 				bNodeHasKeyArea = true;
 			}
 		}
@@ -133,7 +133,7 @@ void FSequencerEntityWalker::HandleNode(const ISequencerEntityVisitor& Visitor,
 	}
 }
 
-void FSequencerEntityWalker::HandleKeyAreaNode(const ISequencerEntityVisitor& Visitor, FSectionKeyAreaNode& KeyAreaNode, const TArray<TSharedRef<ISequencerSection>>& InSections)
+void FSequencerEntityWalker::HandleKeyAreaNode(const ISequencerEntityVisitor& Visitor, FSectionKeyAreaNode& InKeyAreaNode, FSequencerDisplayNode& InOwnerNode, const TArray<TSharedRef<ISequencerSection>>& InSections)
 {
 	for( int32 SectionIndex = 0; SectionIndex < InSections.Num(); ++SectionIndex )
 	{
@@ -144,9 +144,9 @@ void FSequencerEntityWalker::HandleKeyAreaNode(const ISequencerEntityVisitor& Vi
 		{
 			Visitor.VisitSection(Section);
 
-			if (Range.IntersectKeyArea(KeyAreaNode, VirtualKeySize.X))
+			if (Range.IntersectKeyArea(InOwnerNode, VirtualKeySize.X))
 			{
-				TSharedRef<IKeyArea> KeyArea = KeyAreaNode.GetKeyArea(SectionIndex);
+				TSharedRef<IKeyArea> KeyArea = InKeyAreaNode.GetKeyArea(SectionIndex);
 				HandleKeyArea(Visitor, KeyArea, InSections[SectionIndex]->GetSectionObject());
 			}
 		}
diff --git a/Engine/Source/Editor/Sequencer/Private/Tools/SequencerEntityVisitor.h b/Engine/Source/Editor/Sequencer/Private/Tools/SequencerEntityVisitor.h
--- a/Engine/Source/Editor/Sequencer/Private/Tools/SequencerEntityVisitor.h
+++ b/Engine/Source/Editor/Sequencer/Private/Tools/SequencerEntityVisitor.h
@@ -53,7 +53,7 @@ private:
 	/** Handle visitation of a particular node, along with a set of sections */
 	void HandleNode(const ISequencerEntityVisitor& Visitor, FSequencerDisplayNode& InNode, const TArray<TSharedRef<ISequencerSection>>& InSections);
 	/** Handle visitation of a key area node */
-	void HandleKeyAreaNode(const ISequencerEntityVisitor& Visitor, FSectionKeyAreaNode& InNode, const TArray<TSharedRef<ISequencerSection>>& InSections);
+	void HandleKeyAreaNode(const ISequencerEntityVisitor& Visitor, FSectionKeyAreaNode& InKeyAreaNode, FSequencerDisplayNode& InOwnerNode, const TArray<TSharedRef<ISequencerSection>>& InSections);
 	/** Handle visitation of a key area */
 	void HandleKeyArea(const ISequencerEntityVisitor& Visitor, const TSharedPtr<IKeyArea>& KeyArea, UMovieSceneSection* Section);
 
