Changed marquee selection to be use drag operation so we can handle section resizing with the marquee tool

diff --git a/Engine/Source/Editor/Sequencer/Private/Tools/EditToolDragOperations.cpp b/Engine/Source/Editor/Sequencer/Private/Tools/EditToolDragOperations.cpp
--- a/Engine/Source/Editor/Sequencer/Private/Tools/EditToolDragOperations.cpp
+++ b/Engine/Source/Editor/Sequencer/Private/Tools/EditToolDragOperations.cpp
@@ -193,7 +193,7 @@ FResizeSection::FResizeSection( FSequencer& InSequencer, TArray<FSectionHandle>
 {
 }
 
-void FResizeSection::OnBeginDrag(const FVector2D& LocalMousePos, const FVirtualTrackArea& VirtualTrackArea)
+void FResizeSection::OnBeginDrag(const FPointerEvent& MouseEvent, FVector2D LocalMousePos, const FVirtualTrackArea& VirtualTrackArea)
 {
 	BeginTransaction( Sections, NSLOCTEXT("Sequencer", "DragSectionEdgeTransaction", "Resize section") );
 
@@ -211,14 +211,14 @@ void FResizeSection::OnBeginDrag(const FVector2D& LocalMousePos, const FVirtualT
 	}
 }
 
-void FResizeSection::OnEndDrag()
+void FResizeSection::OnEndDrag(const FPointerEvent& MouseEvent, FVector2D LocalMousePos, const FVirtualTrackArea& VirtualTrackArea)
 {
 	EndTransaction();
 
 	DraggedKeyHandles.Empty();
 }
 
-void FResizeSection::OnDrag( const FPointerEvent& MouseEvent, const FVector2D& LocalMousePos, const FVirtualTrackArea& VirtualTrackArea )
+void FResizeSection::OnDrag(const FPointerEvent& MouseEvent, FVector2D LocalMousePos, const FVirtualTrackArea& VirtualTrackArea)
 {
 	bool bIsDilating = MouseEvent.IsControlDown();
 
@@ -300,7 +300,7 @@ FMoveSection::FMoveSection( FSequencer& InSequencer, TArray<FSectionHandle> InSe
 {
 }
 
-void FMoveSection::OnBeginDrag(const FVector2D& LocalMousePos, const FVirtualTrackArea& VirtualTrackArea)
+void FMoveSection::OnBeginDrag(const FPointerEvent& MouseEvent, FVector2D LocalMousePos, const FVirtualTrackArea& VirtualTrackArea)
 {
 	BeginTransaction( Sections, NSLOCTEXT("Sequencer", "MoveSectionTransaction", "Move Section") );
 
@@ -314,7 +314,7 @@ void FMoveSection::OnBeginDrag(const FVector2D& LocalMousePos, const FVirtualTra
 	}
 }
 
-void FMoveSection::OnEndDrag()
+void FMoveSection::OnEndDrag(const FPointerEvent& MouseEvent, FVector2D LocalMousePos, const FVirtualTrackArea& VirtualTrackArea)
 {
 	DraggedKeyHandles.Empty();
 
@@ -336,7 +336,7 @@ void FMoveSection::OnEndDrag()
 }
 
 
-void FMoveSection::OnDrag( const FPointerEvent& MouseEvent, const FVector2D& LocalMousePos, const FVirtualTrackArea& VirtualTrackArea )
+void FMoveSection::OnDrag(const FPointerEvent& MouseEvent, FVector2D LocalMousePos, const FVirtualTrackArea& VirtualTrackArea)
 {
 	float NewMouseDownTime = VirtualTrackArea.PixelToTime(LocalMousePos.X);
 	float DistanceMoved =  NewMouseDownTime - MouseDownTime;
@@ -446,7 +446,7 @@ void FMoveSection::OnDrag( const FPointerEvent& MouseEvent, const FVector2D& Loc
 	MouseDownTime = NewMouseDownTime;
 }
 
-void FMoveKeys::OnBeginDrag(const FVector2D& LocalMousePos, const FVirtualTrackArea& VirtualTrackArea)
+void FMoveKeys::OnBeginDrag(const FPointerEvent& MouseEvent, FVector2D LocalMousePos, const FVirtualTrackArea& VirtualTrackArea)
 {
 	check( SelectedKeys.Num() > 0 )
 
@@ -481,7 +481,7 @@ void FMoveKeys::OnBeginDrag(const FVector2D& LocalMousePos, const FVirtualTrackA
 	}
 }
 
-void FMoveKeys::OnDrag( const FPointerEvent& MouseEvent, const FVector2D& LocalMousePos, const FVirtualTrackArea& VirtualTrackArea )
+void FMoveKeys::OnDrag(const FPointerEvent& MouseEvent, FVector2D LocalMousePos, const FVirtualTrackArea& VirtualTrackArea)
 {
 	float MouseTime = VirtualTrackArea.PixelToTime(LocalMousePos.X);
 	float DistanceMoved = MouseTime - VirtualTrackArea.PixelToTime(LocalMousePos.X - MouseEvent.GetCursorDelta().X);
@@ -562,7 +562,7 @@ void FMoveKeys::OnDrag( const FPointerEvent& MouseEvent, const FVector2D& LocalM
 	}
 }
 
-void FMoveKeys::OnEndDrag()
+void FMoveKeys::OnEndDrag(const FPointerEvent& MouseEvent, FVector2D LocalMousePos, const FVirtualTrackArea& VirtualTrackArea)
 {
 	EndTransaction();
 	SSection::EnableLayoutRegeneration();
diff --git a/Engine/Source/Editor/Sequencer/Private/Tools/EditToolDragOperations.h b/Engine/Source/Editor/Sequencer/Private/Tools/EditToolDragOperations.h
--- a/Engine/Source/Editor/Sequencer/Private/Tools/EditToolDragOperations.h
+++ b/Engine/Source/Editor/Sequencer/Private/Tools/EditToolDragOperations.h
@@ -46,9 +46,9 @@ public:
 	FResizeSection( FSequencer& InSequencer, TArray<FSectionHandle> Sections, TOptional<FSectionHandle> InCardinalSection, bool bInDraggingByEnd );
 
 	/** FEditToolDragOperation interface */
-	virtual void OnBeginDrag(const FVector2D& LocalMousePos, const FVirtualTrackArea& VirtualTrackArea) override;
-	virtual void OnEndDrag() override;
-	virtual void OnDrag( const FPointerEvent& MouseEvent, const FVector2D& LocalMousePos, const FVirtualTrackArea& VirtualTrackArea ) override;
+	virtual void OnBeginDrag(const FPointerEvent& MouseEvent, FVector2D LocalMousePos, const FVirtualTrackArea& VirtualTrackArea) override;
+	virtual void OnDrag(const FPointerEvent& MouseEvent, FVector2D LocalMousePos, const FVirtualTrackArea& VirtualTrackArea) override;
+	virtual void OnEndDrag(const FPointerEvent& MouseEvent, FVector2D LocalMousePos, const FVirtualTrackArea& VirtualTrackArea) override;
 	virtual FCursorReply GetCursor() const override { return FCursorReply::Cursor( EMouseCursor::ResizeLeftRight ); }
 private:
 	/** The sections we are interacting with */
@@ -72,9 +72,9 @@ public:
 	FMoveSection( FSequencer& InSequencer, TArray<FSectionHandle> Sections );
 
 	/** FEditToolDragOperation interface */
-	virtual void OnBeginDrag(const FVector2D& LocalMousePos, const FVirtualTrackArea& VirtualTrackArea) override;
-	virtual void OnEndDrag() override;
-	virtual void OnDrag(const FPointerEvent& MouseEvent, const FVector2D& LocalMousePos, const FVirtualTrackArea& VirtualTrackArea) override;
+	virtual void OnBeginDrag(const FPointerEvent& MouseEvent, FVector2D LocalMousePos, const FVirtualTrackArea& VirtualTrackArea) override;
+	virtual void OnDrag(const FPointerEvent& MouseEvent, FVector2D LocalMousePos, const FVirtualTrackArea& VirtualTrackArea) override;
+	virtual void OnEndDrag(const FPointerEvent& MouseEvent, FVector2D LocalMousePos, const FVirtualTrackArea& VirtualTrackArea) override;
 	virtual FCursorReply GetCursor() const override { return FCursorReply::Cursor( EMouseCursor::CardinalCross ); }
 
 private:
@@ -101,9 +101,9 @@ public:
 	}
 
 	/** FEditToolDragOperation interface */
-	virtual void OnBeginDrag(const FVector2D& LocalMousePos, const FVirtualTrackArea& VirtualTrackArea) override;
-	virtual void OnEndDrag() override;
-	virtual void OnDrag( const FPointerEvent& MouseEvent, const FVector2D& LocalMousePos, const FVirtualTrackArea& VirtualTrackArea ) override;
+	virtual void OnBeginDrag(const FPointerEvent& MouseEvent, FVector2D LocalMousePos, const FVirtualTrackArea& VirtualTrackArea) override;
+	virtual void OnDrag(const FPointerEvent& MouseEvent, FVector2D LocalMousePos, const FVirtualTrackArea& VirtualTrackArea) override;
+	virtual void OnEndDrag(const FPointerEvent& MouseEvent, FVector2D LocalMousePos, const FVirtualTrackArea& VirtualTrackArea) override;
 
 protected:
 
diff --git a/Engine/Source/Editor/Sequencer/Private/Tools/SequencerEditTool_Movement.cpp b/Engine/Source/Editor/Sequencer/Private/Tools/SequencerEditTool_Movement.cpp
--- a/Engine/Source/Editor/Sequencer/Private/Tools/SequencerEditTool_Movement.cpp
+++ b/Engine/Source/Editor/Sequencer/Private/Tools/SequencerEditTool_Movement.cpp
@@ -32,7 +32,6 @@ FReply FSequencerEditTool_Movement::OnMouseButtonDown(SWidget& OwnerWidget, cons
 
 	if (MouseEvent.GetEffectingButton() == EKeys::LeftMouseButton)
 	{
-
 		DelayedDrag = FDelayedDrag_Hotspot(MyGeometry.AbsoluteToLocal(MouseEvent.GetScreenSpacePosition()), EKeys::LeftMouseButton, Hotspot);
 		return FReply::Handled();
 	}
@@ -62,7 +61,7 @@ FReply FSequencerEditTool_Movement::OnMouseMove(SWidget& OwnerWidget, const FGeo
 
 			if (DragOperation.IsValid())
 			{
-				DragOperation->OnBeginDrag(DelayedDrag->GetInitialPosition(), VirtualTrackArea);
+				DragOperation->OnBeginDrag(MouseEvent, DelayedDrag->GetInitialPosition(), VirtualTrackArea);
 
 				// Steal the capture, as we're now the authoritative widget in charge of a mouse-drag operation
 				Reply.CaptureMouse(OwnerWidget.AsShared());
@@ -146,7 +145,7 @@ FReply FSequencerEditTool_Movement::OnMouseButtonUp(SWidget& OwnerWidget, const
 
 	if (DragOperation.IsValid())
 	{
-		DragOperation->OnEndDrag();
+		DragOperation->OnEndDrag(MouseEvent, MyGeometry.AbsoluteToLocal(MouseEvent.GetScreenSpacePosition()), SequencerWidget.Pin()->GetVirtualTrackArea());
 		DragOperation = nullptr;
 
 		// Only return handled if we actually started a drag
@@ -159,11 +158,7 @@ FReply FSequencerEditTool_Movement::OnMouseButtonUp(SWidget& OwnerWidget, const
 void FSequencerEditTool_Movement::OnMouseCaptureLost()
 {
 	DelayedDrag.Reset();
-	if (DragOperation.IsValid())
-	{
-		DragOperation->OnEndDrag();
-		DragOperation = nullptr;
-	}
+	DragOperation = nullptr;
 }
 
 FCursorReply FSequencerEditTool_Movement::OnCursorQuery(const FGeometry& MyGeometry, const FPointerEvent& CursorEvent) const
diff --git a/Engine/Source/Editor/Sequencer/Private/Tools/SequencerEditTool_Selection.cpp b/Engine/Source/Editor/Sequencer/Private/Tools/SequencerEditTool_Selection.cpp
--- a/Engine/Source/Editor/Sequencer/Private/Tools/SequencerEditTool_Selection.cpp
+++ b/Engine/Source/Editor/Sequencer/Private/Tools/SequencerEditTool_Selection.cpp
@@ -8,14 +8,200 @@
 #include "VirtualTrackArea.h"
 #include "Sequencer.h"
 #include "SSequencerTreeView.h"
+#include "SequencerHotspots.h"
 
-struct FMarqueeSelectData
+
+struct FSelectionVisitor : ISequencerEntityVisitor
+{
+	FSelectionVisitor(FSequencerSelection& InSelection, const FPointerEvent& InMouseEvent, ESequencerSelectionMode InMode)
+		: Selection(InSelection)
+		, MouseEvent(InMouseEvent)
+		, Mode(InMode)
+	{}
+
+	virtual void VisitKey(FKeyHandle KeyHandle, float KeyTime, const TSharedPtr<IKeyArea>& KeyArea, UMovieSceneSection* Section) const override
+	{
+		if (Mode != ESequencerSelectionMode::Keys)
+		{
+			return;
+		}
+
+		FSelectedKey SelectedKey(*Section, KeyArea, KeyHandle);
+
+		if (MouseEvent.IsLeftShiftDown())
+		{
+			Selection.AddToSelection(SelectedKey);
+		}
+		else if (MouseEvent.IsLeftControlDown())
+		{
+			if (Selection.IsSelected(SelectedKey))
+			{
+				Selection.RemoveFromSelection(SelectedKey);
+			}
+			else
+			{
+				Selection.AddToSelection(SelectedKey);
+			}
+		}
+		else
+		{
+			Selection.AddToSelection(SelectedKey);
+		}
+	}
+
+	virtual void VisitSection(UMovieSceneSection* Section) const
+	{
+		if (Mode != ESequencerSelectionMode::Sections)
+		{
+			return;
+		}
+
+		if (MouseEvent.IsLeftShiftDown())
+		{
+			Selection.AddToSelection(Section);
+		}
+		else if (MouseEvent.IsLeftControlDown())
+		{
+			if (Selection.IsSelected(Section))
+			{
+				Selection.RemoveFromSelection(Section);
+			}
+			else
+			{
+				Selection.AddToSelection(Section);
+			}
+		}
+		else
+		{
+			Selection.AddToSelection(Section);
+		}
+	}
+
+private:
+	FSequencerSelection& Selection;
+	const FPointerEvent& MouseEvent;
+	ESequencerSelectionMode Mode;
+};
+
+class FMarqueeDragOperation : public IEditToolDragOperation
 {
-	FMarqueeSelectData(const FVector2D& InInitialPosition)
-		: InitialPosition(InInitialPosition)
-		, CurrentPosition(InInitialPosition)
+public:
+
+	FMarqueeDragOperation(TWeakPtr<FSequencer> InSequencer, TWeakPtr<SSequencer> InSequencerWidget, ESequencerSelectionMode InSelectionMode)
+		: Sequencer(MoveTemp(InSequencer)), SequencerWidget(MoveTemp(InSequencerWidget)), SelectionMode(InSelectionMode)
 	{}
 
+	/** Start a new marquee selection */
+	virtual void OnBeginDrag(const FPointerEvent& MouseEvent, FVector2D LocalMousePos, const FVirtualTrackArea& VirtualTrackArea) override
+	{
+		InitialPosition =  VirtualTrackArea.PhysicalToVirtual(LocalMousePos);
+	}
+
+	/** Change the current marquee selection */
+	virtual void OnDrag(const FPointerEvent& MouseEvent, FVector2D LocalMousePos, const FVirtualTrackArea& VirtualTrackArea) override
+	{
+		const FVector2D MouseDelta = MouseEvent.GetCursorDelta();
+
+		auto PinnedSequencer = Sequencer.Pin();
+
+		// Handle virtual scrolling when at the vertical extremes of the widget (performed before we clamp the mouse pos)
+		{
+			const float ScrollThresholdV = VirtualTrackArea.GetPhysicalSize().Y * 0.025f;
+
+			float Difference = LocalMousePos.Y - ScrollThresholdV;
+			if (Difference < 0 && MouseDelta.Y < 0)
+			{
+				PinnedSequencer->VerticalScroll( Difference * 0.1f );
+			}
+
+			Difference = LocalMousePos.Y - (VirtualTrackArea.GetPhysicalSize().Y - ScrollThresholdV);
+			if (Difference > 0 && MouseDelta.Y > 0)
+			{
+				PinnedSequencer->VerticalScroll( Difference * 0.1f );
+			}
+		}
+
+		// Clamp the vertical position to the actual bounds of the track area
+		LocalMousePos.Y = FMath::Clamp(LocalMousePos.Y, 0.f, VirtualTrackArea.GetPhysicalSize().Y);
+		CurrentPosition = VirtualTrackArea.PhysicalToVirtual(LocalMousePos);
+
+		// Handle virtual scrolling when at the horizontal extremes of the widget
+		{
+			TRange<float> ViewRange = PinnedSequencer->GetViewRange();
+			const float ScrollThresholdH = ViewRange.Size<float>() * 0.025f;
+
+			float Difference = CurrentPosition.X - (ViewRange.GetLowerBoundValue() + ScrollThresholdH);
+			if (Difference < 0 && MouseDelta.X < 0)
+			{
+				PinnedSequencer->StartAutoscroll(Difference);
+			}
+			else
+			{
+				Difference = CurrentPosition.X - (ViewRange.GetUpperBoundValue() - ScrollThresholdH);
+				if (Difference > 0 && MouseDelta.X > 0)
+				{
+					PinnedSequencer->StartAutoscroll(Difference);
+				}
+				else
+				{
+					PinnedSequencer->StopAutoscroll();
+				}
+			}
+		}
+	}
+
+	/** Finish dragging the marquee selection */
+	virtual void OnEndDrag(const FPointerEvent& MouseEvent, FVector2D LocalMousePos, const FVirtualTrackArea& VirtualTrackArea) override
+	{
+		auto PinnedSequencer = Sequencer.Pin();
+
+		auto& Selection = PinnedSequencer->GetSelection();
+		if (!MouseEvent.IsLeftControlDown() && !MouseEvent.IsLeftShiftDown())
+		{
+			Selection.Empty();
+		}
+
+		Selection.SuspendBroadcast();
+
+		FVector2D VirtualKeySize;
+		VirtualKeySize.X = SequencerSectionConstants::KeySize.X / VirtualTrackArea.GetPhysicalSize().X * (PinnedSequencer->GetViewRange().GetUpperBoundValue() - PinnedSequencer->GetViewRange().GetLowerBoundValue());
+		// Vertically, virtual units == physical units
+		VirtualKeySize.Y = SequencerSectionConstants::KeySize.Y;
+
+		FSequencerEntityWalker Walker(FSequencerEntityRange(TopLeft(), BottomRight()), VirtualKeySize);
+
+		FSelectionVisitor Visitor(Selection, MouseEvent, SelectionMode);
+
+		Walker.Traverse(Visitor, SequencerWidget.Pin()->GetTreeView()->GetNodeTree()->GetRootNodes());
+
+		Selection.ResumeBroadcast();
+		Selection.GetOnKeySelectionChanged().Broadcast();
+	}
+
+	virtual int32 OnPaint(const FGeometry& AllottedGeometry, const FSlateRect& MyClippingRect, FSlateWindowElementList& OutDrawElements, int32 LayerId) const
+	{
+		// Convert to physical space for rendering
+		const FVirtualTrackArea VirtualTrackArea = SequencerWidget.Pin()->GetVirtualTrackArea();
+
+		FVector2D SelectionTopLeft = VirtualTrackArea.VirtualToPhysical(TopLeft());
+		FVector2D SelectionBottomRight = VirtualTrackArea.VirtualToPhysical(BottomRight());
+
+		FSlateDrawElement::MakeBox(
+			OutDrawElements,
+			LayerId,
+			AllottedGeometry.ToPaintGeometry(SelectionTopLeft, SelectionBottomRight - SelectionTopLeft),
+			FEditorStyle::GetBrush(TEXT("MarqueeSelection")),
+			MyClippingRect
+			);
+
+		return LayerId + 1;
+	}
+	
+	/** Request the cursor for this drag operation */
+	virtual FCursorReply GetCursor() const { return FCursorReply::Cursor( EMouseCursor::Default ); }
+
+private:
+
 	FVector2D TopLeft() const
 	{
 		return FVector2D(
@@ -32,18 +218,16 @@ struct FMarqueeSelectData
 		);
 	}
 
-	FVector2D Size() const
-	{
-		return BottomRight() - TopLeft();
-	}
+	/** The sequencer itself */
+	TWeakPtr<FSequencer> Sequencer;
 
-	bool HasDragged() const
-	{
-		return !Size().IsNearlyZero();
-	}
+	/** Sequencer widget */
+	TWeakPtr<SSequencer> SequencerWidget;
 
 	FVector2D InitialPosition;
 	FVector2D CurrentPosition;
+
+	ESequencerSelectionMode SelectionMode;
 };
 
 /** todo: Save this in a config? */
@@ -86,23 +270,9 @@ FCursorReply FSequencerEditTool_Selection::OnCursorQuery(const FGeometry& MyGeom
 
 int32 FSequencerEditTool_Selection::OnPaint(const FGeometry& AllottedGeometry, const FSlateRect& MyClippingRect, FSlateWindowElementList& OutDrawElements, int32 LayerId) const
 {
-	if (MarqueeSelectData.IsValid() && MarqueeSelectData->HasDragged())
+	if (DragOperation.IsValid())
 	{
-		// Convert to physical space for rendering
-		const FVirtualTrackArea VirtualTrackArea = SequencerWidget.Pin()->GetVirtualTrackArea();
-
-		FVector2D TopLeft = VirtualTrackArea.VirtualToPhysical(MarqueeSelectData->TopLeft());
-		FVector2D BottomRight = VirtualTrackArea.VirtualToPhysical(MarqueeSelectData->BottomRight());
-
-		FSlateDrawElement::MakeBox(
-			OutDrawElements,
-			LayerId,
-			AllottedGeometry.ToPaintGeometry(TopLeft, BottomRight - TopLeft),
-			FEditorStyle::GetBrush(TEXT("MarqueeSelection")),
-			MyClippingRect
-			);
-
-		return LayerId + 1;
+		return DragOperation->OnPaint(AllottedGeometry, MyClippingRect, OutDrawElements, LayerId);
 	}
 
 	return LayerId;
@@ -114,7 +284,7 @@ FReply FSequencerEditTool_Selection::OnMouseButtonDown(SWidget& OwnerWidget, con
 
 	if (MouseEvent.GetEffectingButton() == EKeys::LeftMouseButton)
 	{
-		DelayedDrag = FDelayedDrag(MyGeometry.AbsoluteToLocal(MouseEvent.GetScreenSpacePosition()), EKeys::LeftMouseButton);
+		DelayedDrag = FDelayedDrag_Hotspot(MyGeometry.AbsoluteToLocal(MouseEvent.GetScreenSpacePosition()), EKeys::LeftMouseButton, Hotspot);
 		return FReply::Handled();
 	}
 	return FReply::Unhandled();
@@ -126,19 +296,38 @@ FReply FSequencerEditTool_Selection::OnMouseMove(SWidget& OwnerWidget, const FGe
 	{
 		FReply Reply = FReply::Handled();
 
-		if (MarqueeSelectData.IsValid())
+		const FVirtualTrackArea VirtualTrackArea = SequencerWidget.Pin()->GetVirtualTrackArea();
+
+		if (DragOperation.IsValid())
 		{
-			SetNewMarqueeBounds(MouseEvent, MyGeometry);
+			FVector2D LocalPosition = MyGeometry.AbsoluteToLocal(MouseEvent.GetScreenSpacePosition());
+			DragOperation->OnDrag(MouseEvent, LocalPosition, VirtualTrackArea );
 		}
 		else if (DelayedDrag->AttemptDragStart(MouseEvent))
 		{
-			const FVirtualTrackArea VirtualTrackArea = SequencerWidget.Pin()->GetVirtualTrackArea();
+			if (DelayedDrag->Hotspot.IsValid())
+			{
+				// We only allow resizing with the marquee selection tool enabled
+				auto HotspotType = DelayedDrag->Hotspot->GetType();
+				if (HotspotType != ESequencerHotspot::Section && HotspotType != ESequencerHotspot::Key)
+				{
+					DragOperation = DelayedDrag->Hotspot->InitiateDrag(*Sequencer.Pin());
+				}
+			}
+
+			if (!DragOperation.IsValid())
+			{
+				DragOperation = MakeShareable( new FMarqueeDragOperation(Sequencer, SequencerWidget, SelectionMode) );
+			}
 
-			// Marquee data is stored in virtual space
-			MarqueeSelectData.Reset(new FMarqueeSelectData(VirtualTrackArea.PhysicalToVirtual(DelayedDrag->GetInitialPosition())));
+			if (DragOperation.IsValid())
+			{
+				FVector2D LocalPosition = MyGeometry.AbsoluteToLocal(MouseEvent.GetScreenSpacePosition());
+				DragOperation->OnBeginDrag(MouseEvent, LocalPosition, VirtualTrackArea);
 
-			// Steal the capture, as we're now the authoritative widget in charge of a mouse-drag operation
-			Reply.CaptureMouse(OwnerWidget.AsShared());
+				// Steal the capture, as we're now the authoritative widget in charge of a mouse-drag operation
+				Reply.CaptureMouse(OwnerWidget.AsShared());
+			}
 		}
 
 		return Reply;
@@ -149,10 +338,11 @@ FReply FSequencerEditTool_Selection::OnMouseMove(SWidget& OwnerWidget, const FGe
 FReply FSequencerEditTool_Selection::OnMouseButtonUp(SWidget& OwnerWidget, const FGeometry& MyGeometry, const FPointerEvent& MouseEvent)
 {
 	DelayedDrag.Reset();
-	if (MarqueeSelectData.IsValid())
+	if (DragOperation.IsValid())
 	{
-		HandleMarqueeSelection(MyGeometry, MouseEvent);
-		MarqueeSelectData.Reset();
+		FVector2D LocalPosition = MyGeometry.AbsoluteToLocal(MouseEvent.GetScreenSpacePosition());
+		DragOperation->OnEndDrag(MouseEvent, LocalPosition, SequencerWidget.Pin()->GetVirtualTrackArea() );
+		DragOperation = nullptr;
 
 		Sequencer.Pin()->StopAutoscroll();
 		return FReply::Handled().ReleaseMouseCapture();
@@ -164,166 +354,11 @@ FReply FSequencerEditTool_Selection::OnMouseButtonUp(SWidget& OwnerWidget, const
 void FSequencerEditTool_Selection::OnMouseCaptureLost()
 {
 	DelayedDrag.Reset();
-	MarqueeSelectData.Reset();
+	DragOperation = nullptr;
 }
 
 FName FSequencerEditTool_Selection::GetIdentifier() const
 {
 	static FName Identifier("Selection");
 	return Identifier;
-}
-
-void FSequencerEditTool_Selection::SetNewMarqueeBounds(const FPointerEvent& MouseEvent, const FGeometry& MyGeometry)
-{
-	const FVector2D MouseDelta = MouseEvent.GetCursorDelta();
-
-	FVector2D LocalPosition = MyGeometry.AbsoluteToLocal(MouseEvent.GetScreenSpacePosition());
-
-	auto PinnedSequencer = Sequencer.Pin();
-
-	// Handle virtual scrolling when at the vertical extremes of the widget (performed before we clamp the mouse pos)
-	{
-		const float ScrollThresholdV = MyGeometry.Size.Y * 0.025f;
-
-		float Difference = LocalPosition.Y - ScrollThresholdV;
-		if (Difference < 0 && MouseDelta.Y < 0)
-		{
-			PinnedSequencer->VerticalScroll( Difference * 0.1f );
-		}
-
-		Difference = LocalPosition.Y - (MyGeometry.Size.Y - ScrollThresholdV);
-		if (Difference > 0 && MouseDelta.Y > 0)
-		{
-			PinnedSequencer->VerticalScroll( Difference * 0.1f );
-		}
-	}
-
-	const FVirtualTrackArea VirtualTrackArea = SequencerWidget.Pin()->GetVirtualTrackArea();
-
-	// Clamp the vertical position to the actual bounds of the track area
-	LocalPosition.Y = FMath::Clamp(LocalPosition.Y, 0.f, MyGeometry.GetLocalSize().Y);
-	MarqueeSelectData->CurrentPosition = VirtualTrackArea.PhysicalToVirtual(LocalPosition);
-
-	// Handle virtual scrolling when at the horizontal extremes of the widget
-	{
-		TRange<float> ViewRange = PinnedSequencer->GetViewRange();
-		const float ScrollThresholdH = ViewRange.Size<float>() * 0.025f;
-
-		float Difference = MarqueeSelectData->CurrentPosition.X - (ViewRange.GetLowerBoundValue() + ScrollThresholdH);
-		if (Difference < 0 && MouseDelta.X < 0)
-		{
-			PinnedSequencer->StartAutoscroll(Difference);
-		}
-		else
-		{
-			Difference = MarqueeSelectData->CurrentPosition.X - (ViewRange.GetUpperBoundValue() - ScrollThresholdH);
-			if (Difference > 0 && MouseDelta.X > 0)
-			{
-				PinnedSequencer->StartAutoscroll(Difference);
-			}
-			else
-			{
-				PinnedSequencer->StopAutoscroll();
-			}
-		}
-	}
-}
-
-struct FSelectionVisitor : ISequencerEntityVisitor
-{
-	FSelectionVisitor(FSequencerSelection& InSelection, const FPointerEvent& InMouseEvent, ESequencerSelectionMode InMode)
-		: Selection(InSelection)
-		, MouseEvent(InMouseEvent)
-		, Mode(InMode)
-	{}
-
-	virtual void VisitKey(FKeyHandle KeyHandle, float KeyTime, const TSharedPtr<IKeyArea>& KeyArea, UMovieSceneSection* Section) const override
-	{
-		if (Mode != ESequencerSelectionMode::Keys)
-		{
-			return;
-		}
-
-		FSelectedKey SelectedKey(*Section, KeyArea, KeyHandle);
-
-		if (MouseEvent.IsLeftShiftDown())
-		{
-			Selection.AddToSelection(SelectedKey);
-		}
-		else if (MouseEvent.IsLeftControlDown())
-		{
-			if (Selection.IsSelected(SelectedKey))
-			{
-				Selection.RemoveFromSelection(SelectedKey);
-			}
-			else
-			{
-				Selection.AddToSelection(SelectedKey);
-			}
-		}
-		else
-		{
-			Selection.AddToSelection(SelectedKey);
-		}
-	}
-
-	virtual void VisitSection(UMovieSceneSection* Section) const
-	{
-		if (Mode != ESequencerSelectionMode::Sections)
-		{
-			return;
-		}
-
-		if (MouseEvent.IsLeftShiftDown())
-		{
-			Selection.AddToSelection(Section);
-		}
-		else if (MouseEvent.IsLeftControlDown())
-		{
-			if (Selection.IsSelected(Section))
-			{
-				Selection.RemoveFromSelection(Section);
-			}
-			else
-			{
-				Selection.AddToSelection(Section);
-			}
-		}
-		else
-		{
-			Selection.AddToSelection(Section);
-		}
-	}
-
-private:
-	FSequencerSelection& Selection;
-	const FPointerEvent& MouseEvent;
-	ESequencerSelectionMode Mode;
-};
-
-void FSequencerEditTool_Selection::HandleMarqueeSelection(const FGeometry& Geometry, const FPointerEvent& MouseEvent)
-{
-	auto PinnedSequencer = Sequencer.Pin();
-
-	auto& Selection = PinnedSequencer->GetSelection();
-	if (!MouseEvent.IsLeftControlDown() && !MouseEvent.IsLeftShiftDown())
-	{
-		Selection.Empty();
-	}
-
-	Selection.SuspendBroadcast();
-
-	FVector2D VirtualKeySize;
-	VirtualKeySize.X = SequencerSectionConstants::KeySize.X / Geometry.Size.X * (PinnedSequencer->GetViewRange().GetUpperBoundValue() - PinnedSequencer->GetViewRange().GetLowerBoundValue());
-	// Vertically, virtual units == physical units
-	VirtualKeySize.Y = SequencerSectionConstants::KeySize.Y;
-
-	FSequencerEntityWalker Walker(FSequencerEntityRange(MarqueeSelectData->TopLeft(), MarqueeSelectData->BottomRight()), VirtualKeySize);
-
-	FSelectionVisitor Visitor(Selection, MouseEvent, SelectionMode);
-
-	Walker.Traverse(Visitor, SequencerWidget.Pin()->GetTreeView()->GetNodeTree()->GetRootNodes());
-
-	Selection.ResumeBroadcast();
-	Selection.GetOnKeySelectionChanged().Broadcast();
 }
\ No newline at end of file
diff --git a/Engine/Source/Editor/Sequencer/Private/Tools/SequencerEditTool_Selection.h b/Engine/Source/Editor/Sequencer/Private/Tools/SequencerEditTool_Selection.h
--- a/Engine/Source/Editor/Sequencer/Private/Tools/SequencerEditTool_Selection.h
+++ b/Engine/Source/Editor/Sequencer/Private/Tools/SequencerEditTool_Selection.h
@@ -37,14 +37,6 @@ public:
 	ESequencerSelectionMode GetSelectionMode() const { return SelectionMode; }
 	void SetSelectionMode(ESequencerSelectionMode Mode);
 	
-private:
-
-	/** Set new bounds for the marquee selection */
-	void SetNewMarqueeBounds(const FPointerEvent& MouseEvent, const FGeometry& Geometry);
-
-	/** Select all the keys contained in the current marquee */
-	void HandleMarqueeSelection(const FGeometry& Geometry, const FPointerEvent& MouseEvent);
-
 private:
 
 	/** The sequencer itself */
@@ -53,11 +45,21 @@ private:
 	/** Sequencer widget */
 	TWeakPtr<SSequencer> SequencerWidget;
 
-	/** Helper class responsible for handling delayed dragging */
-	TOptional<FDelayedDrag> DelayedDrag;
+	struct FDelayedDrag_Hotspot : FDelayedDrag
+	{
+		FDelayedDrag_Hotspot(FVector2D InInitialPosition, FKey InApplicableKey, TSharedPtr<ISequencerHotspot> InHotspot)
+			: FDelayedDrag(InInitialPosition, InApplicableKey)
+			, Hotspot(MoveTemp(InHotspot))
+		{}
+
+		TSharedPtr<ISequencerHotspot> Hotspot;
+	};
 
-	/** Optional marquee selection data, valid when dragging */
-	TUniquePtr<FMarqueeSelectData> MarqueeSelectData;
+	/** Helper class responsible for handling delayed dragging */
+	TOptional<FDelayedDrag_Hotspot> DelayedDrag;
+	
+	/** Current drag operation if any */
+	TSharedPtr<IEditToolDragOperation> DragOperation;
 
 	/** What we are selecting */
 	ESequencerSelectionMode SelectionMode;
diff --git a/Engine/Source/Editor/Sequencer/Private/VirtualTrackArea.cpp b/Engine/Source/Editor/Sequencer/Private/VirtualTrackArea.cpp
--- a/Engine/Source/Editor/Sequencer/Private/VirtualTrackArea.cpp
+++ b/Engine/Source/Editor/Sequencer/Private/VirtualTrackArea.cpp
@@ -10,6 +10,7 @@
 FVirtualTrackArea::FVirtualTrackArea(const FSequencer& InSequencer, SSequencerTreeView& InTreeView, const FGeometry& InTrackAreaGeometry)
 	: FTimeToPixel(InTrackAreaGeometry, InSequencer.GetViewRange())
 	, TreeView(InTreeView)
+	, TrackAreaGeometry(InTrackAreaGeometry)
 {
 }
 
@@ -39,6 +40,11 @@ FVector2D FVirtualTrackArea::VirtualToPhysical(FVector2D InPosition) const
 	return InPosition;
 }
 
+FVector2D FVirtualTrackArea::GetPhysicalSize() const
+{
+	return TrackAreaGeometry.Size;
+}
+
 TSharedPtr<FSequencerDisplayNode> FVirtualTrackArea::HitTestNode(float InPhysicalPosition) const
 {
 	return TreeView.HitTestNode(InPhysicalPosition);
diff --git a/Engine/Source/Editor/Sequencer/Public/ISequencerEditTool.h b/Engine/Source/Editor/Sequencer/Public/ISequencerEditTool.h
--- a/Engine/Source/Editor/Sequencer/Public/ISequencerEditTool.h
+++ b/Engine/Source/Editor/Sequencer/Public/ISequencerEditTool.h
@@ -13,25 +13,34 @@ public:
 	/**
 	 * Called to initiate a drag
 	 *
-	 * @param LocalMousePos		The current location of the mouse, relative to the top-left corner of the virtual track area
+	 * @param MouseEvent		The associated mouse event for dragging
+	 * @param LocalMousePos		The current location of the mouse, relative to the top-left corner of the physical track area
 	 * @param VirtualTrackArea	A virtual track area that can be used for pixel->time conversions and hittesting
 	 */
-	virtual void OnBeginDrag(const FVector2D& LocalMousePos, const FVirtualTrackArea& VirtualTrackArea) = 0;
-
-	/** Called when a drag has ended */
-	virtual void OnEndDrag() = 0;
-
-	/** Request the cursor for this drag operation */
-	virtual FCursorReply GetCursor() const { return FCursorReply::Cursor( EMouseCursor::Default ); }
+	virtual void OnBeginDrag(const FPointerEvent& MouseEvent, FVector2D LocalMousePos, const FVirtualTrackArea& VirtualTrackArea) = 0;
 
 	/**
 	 * Notification called when the mouse moves while dragging
 	 *
 	 * @param MouseEvent		The associated mouse event for dragging
-	 * @param LocalMousePos		The current location of the mouse, relative to the top-left corner of the virtual track area
+	 * @param LocalMousePos		The current location of the mouse, relative to the top-left corner of the physical track area
+	 * @param VirtualTrackArea	A virtual track area that can be used for pixel->time conversions and hittesting
+	 */
+	virtual void OnDrag(const FPointerEvent& MouseEvent, FVector2D LocalMousePos, const FVirtualTrackArea& VirtualTrackArea) = 0;
+
+	/** Called when a drag has ended
+	 *
+	 * @param MouseEvent		The associated mouse event for dragging
+	 * @param LocalMousePos		The current location of the mouse, relative to the top-left corner of the physical track area
 	 * @param VirtualTrackArea	A virtual track area that can be used for pixel->time conversions and hittesting
 	 */
-	virtual void OnDrag( const FPointerEvent& MouseEvent, const FVector2D& LocalMousePos, const FVirtualTrackArea& VirtualTrackArea ) = 0;
+	virtual void OnEndDrag( const FPointerEvent& MouseEvent, FVector2D LocalMousePos, const FVirtualTrackArea& VirtualTrackArea) = 0;
+
+	/** Request the cursor for this drag operation */
+	virtual FCursorReply GetCursor() const { return FCursorReply::Cursor( EMouseCursor::Default ); }
+
+	/** Override to implement drag-specific paint logic */
+	virtual int32 OnPaint(const FGeometry& AllottedGeometry, const FSlateRect& MyClippingRect, FSlateWindowElementList& OutDrawElements, int32 LayerId) const { return LayerId; }
 };
 
 class ISequencerEditTool
diff --git a/Engine/Source/Editor/Sequencer/Public/VirtualTrackArea.h b/Engine/Source/Editor/Sequencer/Public/VirtualTrackArea.h
--- a/Engine/Source/Editor/Sequencer/Public/VirtualTrackArea.h
+++ b/Engine/Source/Editor/Sequencer/Public/VirtualTrackArea.h
@@ -28,6 +28,9 @@ public:
 	/** @note: Use with caution - not reliable where the specified point is not on screen */
 	FVector2D VirtualToPhysical(FVector2D InPosition) const;
 
+	/** Get the physical size of the track area */
+	FVector2D GetPhysicalSize() const;
+
 	/** Hit test at the specified physical position for a sequencer node */
 	TSharedPtr<FSequencerDisplayNode> HitTestNode(float InPhysicalPosition) const;
 
@@ -41,4 +44,7 @@ private:
 
 	/** Reference to the sequencer tree */
 	SSequencerTreeView& TreeView;
+
+	/** Cached physical geometry of the track area */
+	FGeometry TrackAreaGeometry;
 };
\ No newline at end of file
