Fixed shadowed variable

diff --git a/Engine/Source/Editor/Sequencer/Private/SSection.cpp b/Engine/Source/Editor/Sequencer/Private/SSection.cpp
--- a/Engine/Source/Editor/Sequencer/Private/SSection.cpp
+++ b/Engine/Source/Editor/Sequencer/Private/SSection.cpp
@@ -488,11 +488,11 @@ void SSection::ResetHoveredState()
 	GetSequencer().GetEditTool().SetHotspot(nullptr);
 }
 
-FGeometry SSection::MakeSectionGeometryWithoutHandles( const FGeometry& AllottedGeometry, const TSharedPtr<ISequencerSection>& SectionInterface ) const
+FGeometry SSection::MakeSectionGeometryWithoutHandles( const FGeometry& AllottedGeometry, const TSharedPtr<ISequencerSection>& InSectionInterface ) const
 {
-	const bool bSectionsAreConnected = SectionInterface->AreSectionsConnected();
+	const bool bSectionsAreConnected = InSectionInterface->AreSectionsConnected();
 
-	const float SectionGripSize = !bSectionsAreConnected ? SectionInterface->GetSectionGripSize() : 0.0f;
+	const float SectionGripSize = !bSectionsAreConnected ? InSectionInterface->GetSectionGripSize() : 0.0f;
 
 	return AllottedGeometry.MakeChild( FVector2D( SectionGripSize, 0 ), AllottedGeometry.GetDrawSize() - FVector2D( SectionGripSize*2, 0.0f ) );
 }
diff --git a/Engine/Source/Editor/Sequencer/Private/SSection.h b/Engine/Source/Editor/Sequencer/Private/SSection.h
--- a/Engine/Source/Editor/Sequencer/Private/SSection.h
+++ b/Engine/Source/Editor/Sequencer/Private/SSection.h
@@ -133,7 +133,7 @@ private:
 	/** 
 	 * Creates geometry for a section without space for the handles
 	 */
-	FGeometry MakeSectionGeometryWithoutHandles( const FGeometry& AllottedGeometry, const TSharedPtr<ISequencerSection>& SectionInterface ) const;
+	FGeometry MakeSectionGeometryWithoutHandles( const FGeometry& AllottedGeometry, const TSharedPtr<ISequencerSection>& InSectionInterface ) const;
 
 private:
 	/** Interface to section data */
