Sequencer: Should be able to show the range slider regardless of frame numbers and lock the in/out range if the range slider is visible.

diff --git a/Engine/Source/Editor/Sequencer/Private/Sequencer.cpp b/Engine/Source/Editor/Sequencer/Private/Sequencer.cpp
--- a/Engine/Source/Editor/Sequencer/Private/Sequencer.cpp
+++ b/Engine/Source/Editor/Sequencer/Private/Sequencer.cpp
@@ -589,7 +589,7 @@ void FSequencer::SetGlobalTimeDirectly( float NewTime )
 {
 	float LastTime = ScrubPosition;
 
-	if (Settings->GetLockInOutToStartEndRange())
+	if (Settings->GetShowRangeSlider() && Settings->GetLockInOutToStartEndRange())
 	{
 		UMovieScene* FocusedMovieScene = GetFocusedMovieSceneSequence()->GetMovieScene();
 		NewTime = FMath::Clamp(NewTime, FocusedMovieScene->StartTime, FocusedMovieScene->EndTime);
@@ -631,7 +631,7 @@ void FSequencer::UpdateAutoScroll(float NewTime)
 	}
 
 	// Set the autoscrub rate to the autoscroll rate
-	if (Settings->GetLockInOutToStartEndRange())
+	if (Settings->GetShowRangeSlider() && Settings->GetLockInOutToStartEndRange())
 	{
 		// Don't autoscrub if we're at the extremes of the movie scene range
 		UMovieScene* FocusedMovieScene = GetFocusedMovieSceneSequence()->GetMovieScene();
@@ -1094,7 +1094,7 @@ void FSequencer::OnViewRangeChanged( TRange<float> NewViewRange, EViewRangeInter
 	{
 		UMovieScene* FocusedMovieScene = GetFocusedMovieSceneSequence()->GetMovieScene();
 		// If locked, clamp the new range to clamp range
-		if (Settings->GetLockInOutToStartEndRange() && !bExpandClampRange)
+		if (Settings->GetShowRangeSlider() && Settings->GetLockInOutToStartEndRange() && !bExpandClampRange)
 		{
 			if ( NewViewRange.GetLowerBoundValue() < FocusedMovieScene->StartTime)
 			{
@@ -2052,13 +2052,13 @@ void FSequencer::BindSequencerCommands()
 	SequencerCommandBindings->MapAction(
 		Commands.ToggleShowRangeSlider,
 		FExecuteAction::CreateLambda( [this]{ Settings->SetShowRangeSlider( !Settings->GetShowRangeSlider() ); } ),
-		FCanExecuteAction::CreateSP(this, &FSequencer::CanShowFrameNumbers ),
+		FCanExecuteAction::CreateLambda( []{ return true; } ),
 		FIsActionChecked::CreateLambda( [this]{ return Settings->GetShowRangeSlider(); } ) );
 
 	SequencerCommandBindings->MapAction(
 		Commands.ToggleLockInOutToStartEndRange,
 		FExecuteAction::CreateLambda( [this]{ Settings->SetLockInOutToStartEndRange( !Settings->GetLockInOutToStartEndRange() ); } ),
-		FCanExecuteAction::CreateSP(this, &FSequencer::CanShowFrameNumbers ),
+		FCanExecuteAction::CreateLambda( [this]{ return Settings->GetShowRangeSlider(); } ),
 		FIsActionChecked::CreateLambda( [this]{ return Settings->GetLockInOutToStartEndRange(); } ) );
 
 	SequencerCommandBindings->MapAction(
diff --git a/Engine/Source/Editor/Sequencer/Private/TimeSliderController.cpp b/Engine/Source/Editor/Sequencer/Private/TimeSliderController.cpp
--- a/Engine/Source/Editor/Sequencer/Private/TimeSliderController.cpp
+++ b/Engine/Source/Editor/Sequencer/Private/TimeSliderController.cpp
@@ -655,7 +655,7 @@ float FSequencerTimeSliderController::FrameToTime(int32 Frame) const
 void FSequencerTimeSliderController::ClampViewRange(float& NewRangeMin, float& NewRangeMax, bool bMaintainRange)
 {
 	// If locked, clamp the new range to clamp range
-	if (TimeSliderArgs.Settings->GetLockInOutToStartEndRange())
+	if (TimeSliderArgs.Settings->GetShowRangeSlider() && TimeSliderArgs.Settings->GetLockInOutToStartEndRange())
 	{
 		if (bMaintainRange)
 		{
