Sequencer: Enforce animatable properties to be Interp properties in the level editor sequencer.

diff --git a/Engine/Source/Editor/Sequencer/Private/Sequencer.cpp b/Engine/Source/Editor/Sequencer/Private/Sequencer.cpp
--- a/Engine/Source/Editor/Sequencer/Private/Sequencer.cpp
+++ b/Engine/Source/Editor/Sequencer/Private/Sequencer.cpp
@@ -549,6 +549,16 @@ void FSequencer::SetKeyAllEnabled(bool bKeyAllEnabled)
 	Settings->SetKeyAllEnabled(bKeyAllEnabled);
 }
 
+bool FSequencer::GetKeyInterpPropertiesOnly() const 
+{
+	return Settings->GetKeyInterpPropertiesOnly();
+}
+
+void FSequencer::SetKeyInterpPropertiesOnly(bool bKeyInterpPropertiesOnly) 
+{
+	Settings->SetKeyInterpPropertiesOnly(bKeyInterpPropertiesOnly);
+}
+
 bool FSequencer::IsRecordingLive() const 
 {
 	return PlaybackState == EMovieScenePlayerStatus::Recording && GIsPlayInEditorWorld;
diff --git a/Engine/Source/Editor/Sequencer/Private/Sequencer.h b/Engine/Source/Editor/Sequencer/Private/Sequencer.h
--- a/Engine/Source/Editor/Sequencer/Private/Sequencer.h
+++ b/Engine/Source/Editor/Sequencer/Private/Sequencer.h
@@ -72,6 +72,8 @@ public:
 	virtual void SetAutoKeyEnabled(bool bAutoKeyEnabled) override;
 	virtual bool GetKeyAllEnabled() const override;
 	virtual void SetKeyAllEnabled(bool bKeyAllEnabled) override;
+	virtual bool GetKeyInterpPropertiesOnly() const override;
+	virtual void SetKeyInterpPropertiesOnly(bool bKeyInterpPropertiesOnly) override;
 	virtual bool IsRecordingLive() const override;
 	virtual float GetCurrentLocalTime(UMovieSceneSequence& InMovieSceneSequence) override;
 	virtual float GetGlobalTime() override;
diff --git a/Engine/Source/Editor/Sequencer/Private/SequencerObjectChangeListener.cpp b/Engine/Source/Editor/Sequencer/Private/SequencerObjectChangeListener.cpp
--- a/Engine/Source/Editor/Sequencer/Private/SequencerObjectChangeListener.cpp
+++ b/Engine/Source/Editor/Sequencer/Private/SequencerObjectChangeListener.cpp
@@ -137,6 +137,16 @@ bool FSequencerObjectChangeListener::FindPropertySetter( const UClass& ObjectCla
 		if( Function && !Function->HasMetaData(DeprecatedFunctionName) )
 		{
 			bFound = true;
+
+			if (Sequencer.IsValid() && Sequencer.Pin()->GetKeyInterpPropertiesOnly())
+			{
+				UProperty* Property = ObjectClass.FindPropertyByName(FName(*InPropertyVarName));
+
+				if (!Property || !Property->HasAnyPropertyFlags(CPF_Interp))
+				{
+					bFound = false;
+				}
+			}
 		}
 		else
 		{
diff --git a/Engine/Source/Editor/Sequencer/Private/SequencerSettings.cpp b/Engine/Source/Editor/Sequencer/Private/SequencerSettings.cpp
--- a/Engine/Source/Editor/Sequencer/Private/SequencerSettings.cpp
+++ b/Engine/Source/Editor/Sequencer/Private/SequencerSettings.cpp
@@ -7,6 +7,7 @@ USequencerSettings::USequencerSettings( const FObjectInitializer& ObjectInitiali
 {
 	bAutoKeyEnabled = false;
 	bKeyAllEnabled = false;
+	bKeyInterpPropertiesOnly = false;
 	bShowFrameNumbers = true;
 	bShowRangeSlider = false;
 	bLockInOutToStartEndRange = false;
@@ -21,7 +22,7 @@ USequencerSettings::USequencerSettings( const FObjectInitializer& ObjectInitiali
 	CurveValueSnapInterval = 10.0f;
 	bSnapCurveValueToInterval = true;
 	bIsUsingCleanView = false;
-	bAutoScrollEnabled = true;
+	bAutoScrollEnabled = false;
 	bShowCurveEditor = false;
 	bShowCurveEditorCurveToolTips = true;
 }
@@ -54,6 +55,20 @@ void USequencerSettings::SetKeyAllEnabled(bool InbKeyAllEnabled)
 	}
 }
 
+bool USequencerSettings::GetKeyInterpPropertiesOnly() const
+{
+	return bKeyInterpPropertiesOnly;
+}
+
+void USequencerSettings::SetKeyInterpPropertiesOnly(bool InbKeyInterpPropertiesOnly)
+{
+	if ( bKeyInterpPropertiesOnly != InbKeyInterpPropertiesOnly )
+	{
+		bKeyInterpPropertiesOnly = InbKeyInterpPropertiesOnly;
+		SaveConfig();
+	}
+}
+
 bool USequencerSettings::GetShowFrameNumbers() const
 {
 	return bShowFrameNumbers;
@@ -310,6 +325,7 @@ USequencerSettings::FOnShowCurveEditorChanged& USequencerSettings::GetOnShowCurv
 ULevelEditorSequencerSettings::ULevelEditorSequencerSettings( const FObjectInitializer& ObjectInitializer )
 	: Super( ObjectInitializer )
 {
+	bKeyInterpPropertiesOnly = true;
 	TimeSnapInterval = 0.041667f;
 	bShowRangeSlider = true;
 	bLockInOutToStartEndRange = true;
diff --git a/Engine/Source/Editor/Sequencer/Private/SequencerSettings.h b/Engine/Source/Editor/Sequencer/Private/SequencerSettings.h
--- a/Engine/Source/Editor/Sequencer/Private/SequencerSettings.h
+++ b/Engine/Source/Editor/Sequencer/Private/SequencerSettings.h
@@ -54,6 +54,11 @@ public:
 	/** Sets whether or not key all is enabled. */
 	void SetKeyAllEnabled(bool InbKeyAllEnabled);
 
+	/** Gets whether or not to key interp properties only. */
+	bool GetKeyInterpPropertiesOnly() const;
+	/** Sets whether or not to key interp properties only. */
+	void SetKeyInterpPropertiesOnly(bool InbKeyInterpPropertiesOnly); 
+
 	/** Gets whether or not to show frame numbers. */
 	bool GetShowFrameNumbers() const;
 	/** Sets whether or not to show frame numbers. */
@@ -152,6 +157,9 @@ protected:
 	UPROPERTY( config )
 	bool bKeyAllEnabled;
 
+	UPROPERTY( config )
+	bool bKeyInterpPropertiesOnly;
+
 	UPROPERTY( config )
 	bool bShowFrameNumbers;
 
diff --git a/Engine/Source/Editor/Sequencer/Public/ISequencer.h b/Engine/Source/Editor/Sequencer/Public/ISequencer.h
--- a/Engine/Source/Editor/Sequencer/Public/ISequencer.h
+++ b/Engine/Source/Editor/Sequencer/Public/ISequencer.h
@@ -79,6 +79,12 @@ public:
 	/** Sets whether key all is enabled in this sequencer. */
 	virtual void SetKeyAllEnabled(bool bKeyAllEnabled) = 0;
 
+	/** @return Returns whether or not to key only  interp properties in this sequencer */
+	virtual bool GetKeyInterpPropertiesOnly() const = 0;
+
+	/** Sets whether or not to key only interp properties in this sequencer. */
+	virtual void SetKeyInterpPropertiesOnly(bool bKeyInterpPropertiesOnly) = 0;
+
 	/** @return Returns whether sequencer is currently recording live data from simulated actors */
 	virtual bool IsRecordingLive() const = 0;
 
