Sequencer: Fix infinite section drawing.

diff --git a/Engine/Source/Editor/Sequencer/Private/SSequencerSectionAreaView.cpp b/Engine/Source/Editor/Sequencer/Private/SSequencerSectionAreaView.cpp
--- a/Engine/Source/Editor/Sequencer/Private/SSequencerSectionAreaView.cpp
+++ b/Engine/Source/Editor/Sequencer/Private/SSequencerSectionAreaView.cpp
@@ -35,18 +35,11 @@ namespace SequencerSectionUtils
 		// Note the -1 pixel at the end is because the section does not actually end at the end time if there is a section starting at that same time.  It is more important that a section lines up correctly with it's true start time
 		float PixelEndX = TimeToPixelConverter.TimeToPixel( Section->GetEndTime() );
 
-		float StartX, EndX = 0;
-
 		// If the section is infinite, occupy the entire width of the geometry where the section is located.
 		if (Section->IsInfinite())
 		{
-			StartX = AllottedGeometry.Position.X;
-			EndX = AllottedGeometry.Position.X + AllottedGeometry.Size.X;
-		}
-		else
-		{
-			StartX = TimeToPixelConverter.TimeToPixel( Section->GetStartTime() );
-			EndX = TimeToPixelConverter.TimeToPixel( Section->GetEndTime() );
+			PixelStartX = AllottedGeometry.Position.X;
+			PixelEndX = AllottedGeometry.Position.X + AllottedGeometry.Size.X;
 		}
 
 		// Actual section length without grips.
