Sequencer: Toggle sections active/inactive.

diff --git a/Engine/Source/Editor/MovieSceneTools/Private/Sections/BoolPropertySection.cpp b/Engine/Source/Editor/MovieSceneTools/Private/Sections/BoolPropertySection.cpp
--- a/Engine/Source/Editor/MovieSceneTools/Private/Sections/BoolPropertySection.cpp
+++ b/Engine/Source/Editor/MovieSceneTools/Private/Sections/BoolPropertySection.cpp
@@ -13,6 +13,8 @@ void FBoolPropertySection::GenerateSectionLayout( class ISectionLayoutBuilder& L
 
 int32 FBoolPropertySection::OnPaintSection( const FGeometry& AllottedGeometry, const FSlateRect& SectionClippingRect, FSlateWindowElementList& OutDrawElements, int32 LayerId, bool bParentEnabled ) const
 {
+	const ESlateDrawEffect::Type DrawEffects = bParentEnabled ? ESlateDrawEffect::None : ESlateDrawEffect::DisabledEffect;
+
 	FTimeToPixel TimeToPixelConverter = SectionObject.IsInfinite() ?
 		FTimeToPixel( AllottedGeometry, Sequencer->GetViewRange() ) :
 		FTimeToPixel( AllottedGeometry, TRange<float>( SectionObject.GetStartTime(), SectionObject.GetEndTime() ) );
@@ -58,7 +60,7 @@ int32 FBoolPropertySection::OnPaintSection( const FGeometry& AllottedGeometry, c
 				AllottedGeometry.ToPaintGeometry( FVector2D( FirstTimeInPixels, 0 ), FVector2D( NextTimeInPixels - FirstTimeInPixels, AllottedGeometry.Size.Y ) ),
 				FEditorStyle::GetBrush( "Sequencer.GenericSection.Background" ),
 				SectionClippingRect,
-				ESlateDrawEffect::None,
+				DrawEffects,
 				FLinearColor( 0.f, 1.f, 0.f, 1.f )
 				);
 		}
@@ -70,7 +72,7 @@ int32 FBoolPropertySection::OnPaintSection( const FGeometry& AllottedGeometry, c
 				AllottedGeometry.ToPaintGeometry( FVector2D( FirstTimeInPixels, 0 ), FVector2D( NextTimeInPixels - FirstTimeInPixels, AllottedGeometry.Size.Y ) ),
 				FEditorStyle::GetBrush( "Sequencer.GenericSection.Background" ),
 				SectionClippingRect,
-				ESlateDrawEffect::None,
+				DrawEffects,
 				FLinearColor( 1.f, 0.f, 0.f, 1.f )
 				);
 		}
diff --git a/Engine/Source/Editor/MovieSceneTools/Private/Sections/ColorPropertySection.cpp b/Engine/Source/Editor/MovieSceneTools/Private/Sections/ColorPropertySection.cpp
--- a/Engine/Source/Editor/MovieSceneTools/Private/Sections/ColorPropertySection.cpp
+++ b/Engine/Source/Editor/MovieSceneTools/Private/Sections/ColorPropertySection.cpp
@@ -18,6 +18,8 @@ void FColorPropertySection::GenerateSectionLayout( class ISectionLayoutBuilder&
 
 int32 FColorPropertySection::OnPaintSection( const FGeometry& AllottedGeometry, const FSlateRect& SectionClippingRect, FSlateWindowElementList& OutDrawElements, int32 LayerId, bool bParentEnabled ) const
 {
+	const ESlateDrawEffect::Type DrawEffects = bParentEnabled ? ESlateDrawEffect::None : ESlateDrawEffect::DisabledEffect;
+
 	const UMovieSceneColorSection* ColorSection = Cast<const UMovieSceneColorSection>( &SectionObject );
 
 	float StartTime = ColorSection->GetStartTime();
@@ -34,7 +36,7 @@ int32 FColorPropertySection::OnPaintSection( const FGeometry& AllottedGeometry,
 
 		// If we are showing a background pattern and the colors is transparent, draw a checker pattern
 		const FSlateBrush* CheckerBrush = FEditorStyle::GetBrush( "Checker" );
-		FSlateDrawElement::MakeBox( OutDrawElements, LayerId, PaintGeometry, CheckerBrush, SectionClippingRect );
+		FSlateDrawElement::MakeBox( OutDrawElements, LayerId, PaintGeometry, CheckerBrush, SectionClippingRect, DrawEffects );
 
 		TArray<FSlateGradientStop> GradientStops;
 
@@ -59,7 +61,8 @@ int32 FColorPropertySection::OnPaintSection( const FGeometry& AllottedGeometry,
 				PaintGeometry,
 				GradientStops,
 				Orient_Vertical,
-				SectionClippingRect
+				SectionClippingRect,
+				DrawEffects
 				);
 		}
 	}
diff --git a/Engine/Source/Editor/MovieSceneTools/Private/TrackEditors/AttachTrackEditor.cpp b/Engine/Source/Editor/MovieSceneTools/Private/TrackEditors/AttachTrackEditor.cpp
--- a/Engine/Source/Editor/MovieSceneTools/Private/TrackEditors/AttachTrackEditor.cpp
+++ b/Engine/Source/Editor/MovieSceneTools/Private/TrackEditors/AttachTrackEditor.cpp
@@ -79,13 +79,16 @@ public:
 
 	virtual int32 OnPaintSection( const FGeometry& AllottedGeometry, const FSlateRect& SectionClippingRect, FSlateWindowElementList& OutDrawElements, int32 LayerId, bool bParentEnabled ) const override 
 	{
+		const ESlateDrawEffect::Type DrawEffects = bParentEnabled ? ESlateDrawEffect::None : ESlateDrawEffect::DisabledEffect;
+	
 		// Add a box for the section
 		FSlateDrawElement::MakeBox( 
 			OutDrawElements,
 			LayerId,
 			AllottedGeometry.ToPaintGeometry(),
 			FEditorStyle::GetBrush("Sequencer.GenericSection.Background"),
-			SectionClippingRect
+			SectionClippingRect,
+			DrawEffects
 		); 
 
 		return LayerId;
diff --git a/Engine/Source/Editor/MovieSceneTools/Private/TrackEditors/AudioTrackEditor.cpp b/Engine/Source/Editor/MovieSceneTools/Private/TrackEditors/AudioTrackEditor.cpp
--- a/Engine/Source/Editor/MovieSceneTools/Private/TrackEditors/AudioTrackEditor.cpp
+++ b/Engine/Source/Editor/MovieSceneTools/Private/TrackEditors/AudioTrackEditor.cpp
@@ -352,6 +352,8 @@ float FAudioSection::GetSectionHeight() const
 
 int32 FAudioSection::OnPaintSection( const FGeometry& AllottedGeometry, const FSlateRect& SectionClippingRect, FSlateWindowElementList& OutDrawElements, int32 LayerId, bool bParentEnabled ) const
 {
+	const ESlateDrawEffect::Type DrawEffects = bParentEnabled ? ESlateDrawEffect::None : ESlateDrawEffect::DisabledEffect;
+
 	// Add a box for the section
 	FSlateDrawElement::MakeBox(
 		OutDrawElements,
@@ -359,7 +361,7 @@ int32 FAudioSection::OnPaintSection( const FGeometry& AllottedGeometry, const FS
 		AllottedGeometry.ToPaintGeometry(),
 		FEditorStyle::GetBrush("Sequencer.GenericSection.Background"),
 		SectionClippingRect,
-		ESlateDrawEffect::None,
+		DrawEffects,
 		FLinearColor(0.4f, 0.8f, 0.4f, 1.f)
 	);
 
@@ -374,7 +376,7 @@ int32 FAudioSection::OnPaintSection( const FGeometry& AllottedGeometry, const FS
 			SectionClippingRect,
 			false,
 			false,
-			ESlateDrawEffect::None,
+			DrawEffects,
 			FLinearColor(1.f, 1.f, 1.f, 0.9f)
 		);
 	}
diff --git a/Engine/Source/Editor/MovieSceneTools/Private/TrackEditors/ParticleTrackEditor.cpp b/Engine/Source/Editor/MovieSceneTools/Private/TrackEditors/ParticleTrackEditor.cpp
--- a/Engine/Source/Editor/MovieSceneTools/Private/TrackEditors/ParticleTrackEditor.cpp
+++ b/Engine/Source/Editor/MovieSceneTools/Private/TrackEditors/ParticleTrackEditor.cpp
@@ -62,6 +62,8 @@ float FParticleSection::GetSectionHeight() const
 
 int32 FParticleSection::OnPaintSection( const FGeometry& AllottedGeometry, const FSlateRect& SectionClippingRect, FSlateWindowElementList& OutDrawElements, int32 LayerId, bool bParentEnabled ) const
 {
+	const ESlateDrawEffect::Type DrawEffects = bParentEnabled ? ESlateDrawEffect::None : ESlateDrawEffect::DisabledEffect;
+
 	UMovieSceneParticleSection* AnimSection = Cast<UMovieSceneParticleSection>(&Section);
 	
 	FTimeToPixel TimeToPixelConverter( AllottedGeometry, TRange<float>( Section.GetStartTime(), Section.GetEndTime() ) );
@@ -78,7 +80,7 @@ int32 FParticleSection::OnPaintSection( const FGeometry& AllottedGeometry, const
 			AllottedGeometry.ToPaintGeometry(),
 			FEditorStyle::GetBrush("Sequencer.GenericSection.Background"),
 			SectionClippingRect,
-			ESlateDrawEffect::None,
+			DrawEffects,
 			FLinearColor(0.6f, 0.4f, 0.3f, 1.f)
 		);
 	}
@@ -90,7 +92,7 @@ int32 FParticleSection::OnPaintSection( const FGeometry& AllottedGeometry, const
 			AllottedGeometry.ToPaintGeometry(),
 			FEditorStyle::GetBrush("Sequencer.GenericSection.Background"),
 			SectionClippingRect,
-			ESlateDrawEffect::None,
+			DrawEffects,
 			FLinearColor(0.8f, 0.4f, 0.3f, 1.f)
 		);
 	}
diff --git a/Engine/Source/Editor/MovieSceneTools/Private/TrackEditors/ShotTrackEditor.cpp b/Engine/Source/Editor/MovieSceneTools/Private/TrackEditors/ShotTrackEditor.cpp
--- a/Engine/Source/Editor/MovieSceneTools/Private/TrackEditors/ShotTrackEditor.cpp
+++ b/Engine/Source/Editor/MovieSceneTools/Private/TrackEditors/ShotTrackEditor.cpp
@@ -306,6 +306,8 @@ void FShotSection::FilterToSelectedShotSections(bool bZoomToShotBounds)
 
 int32 FShotSection::OnPaintSection( const FGeometry& AllottedGeometry, const FSlateRect& SectionClippingRect, FSlateWindowElementList& OutDrawElements, int32 LayerId, bool bParentEnabled ) const
 {
+	const ESlateDrawEffect::Type DrawEffects = bParentEnabled ? ESlateDrawEffect::None : ESlateDrawEffect::DisabledEffect;
+
 	if (Camera.IsValid())
 	{
 		FGeometry ThumbnailAreaGeometry = AllottedGeometry.MakeChild( FVector2D(GetSectionGripSize(), 0.0f), AllottedGeometry.GetDrawSize() - FVector2D( GetSectionGripSize()*2, 0.0f ) );
@@ -315,7 +317,8 @@ int32 FShotSection::OnPaintSection( const FGeometry& AllottedGeometry, const FSl
 				LayerId,
 				ThumbnailAreaGeometry.ToPaintGeometry(),
 				FEditorStyle::GetBrush("Sequencer.GenericSection.Background"),
-				SectionClippingRect
+				SectionClippingRect,
+				DrawEffects
 			);
 
 		// @todo Sequencer: Need a way to visualize the key here
@@ -336,7 +339,7 @@ int32 FShotSection::OnPaintSection( const FGeometry& AllottedGeometry, const FSl
 				SectionClippingRect,
 				false,
 				false,
-				ESlateDrawEffect::None,
+				DrawEffects,
 				FLinearColor::White
 			);
 
@@ -348,7 +351,7 @@ int32 FShotSection::OnPaintSection( const FGeometry& AllottedGeometry, const FSl
 					TruncatedGeometry.ToPaintGeometry(),
 					WhiteBrush,
 					SectionClippingRect,
-					ESlateDrawEffect::None,
+					DrawEffects,
 					FLinearColor(1.0f, 1.0f, 1.0f, Thumbnail->GetFadeInCurve())
 					);
 			}
@@ -363,7 +366,8 @@ int32 FShotSection::OnPaintSection( const FGeometry& AllottedGeometry, const FSl
 			LayerId,
 			AllottedGeometry.ToPaintGeometry(),
 			FEditorStyle::GetBrush("Sequencer.GenericSection.Background"),
-			SectionClippingRect
+			SectionClippingRect,
+			DrawEffects
 		); 
 	}
 
diff --git a/Engine/Source/Editor/MovieSceneTools/Private/TrackEditors/SkeletalAnimationTrackEditor.cpp b/Engine/Source/Editor/MovieSceneTools/Private/TrackEditors/SkeletalAnimationTrackEditor.cpp
--- a/Engine/Source/Editor/MovieSceneTools/Private/TrackEditors/SkeletalAnimationTrackEditor.cpp
+++ b/Engine/Source/Editor/MovieSceneTools/Private/TrackEditors/SkeletalAnimationTrackEditor.cpp
@@ -61,6 +61,8 @@ float FSkeletalAnimationSection::GetSectionHeight() const
 
 int32 FSkeletalAnimationSection::OnPaintSection( const FGeometry& AllottedGeometry, const FSlateRect& SectionClippingRect, FSlateWindowElementList& OutDrawElements, int32 LayerId, bool bParentEnabled ) const
 {
+	const ESlateDrawEffect::Type DrawEffects = bParentEnabled ? ESlateDrawEffect::None : ESlateDrawEffect::DisabledEffect;
+
 	UMovieSceneSkeletalAnimationSection* AnimSection = Cast<UMovieSceneSkeletalAnimationSection>(&Section);
 	
 	FTimeToPixel TimeToPixelConverter( AllottedGeometry, TRange<float>( Section.GetStartTime(), Section.GetEndTime() ) );
@@ -72,7 +74,7 @@ int32 FSkeletalAnimationSection::OnPaintSection( const FGeometry& AllottedGeomet
 		AllottedGeometry.ToPaintGeometry(),
 		FEditorStyle::GetBrush("Sequencer.GenericSection.Background"),
 		SectionClippingRect,
-		ESlateDrawEffect::None,
+		DrawEffects,
 		FLinearColor(0.7f, 0.4f, 0.7f, 1.f)
 	);
 
@@ -91,7 +93,7 @@ int32 FSkeletalAnimationSection::OnPaintSection( const FGeometry& AllottedGeomet
 			AllottedGeometry.ToPaintGeometry(FVector2D(StartPixels, 0), FVector2D(EndPixels - StartPixels, AllottedGeometry.Size.Y)),
 			FEditorStyle::GetBrush("WhiteTexture"),
 			SectionClippingRect,
-			ESlateDrawEffect::None,
+			DrawEffects,
 			FLinearColor(0.f, 0.f, 0.f, 0.3f)
 		);
 	}
@@ -113,7 +115,8 @@ int32 FSkeletalAnimationSection::OnPaintSection( const FGeometry& AllottedGeomet
 				LayerId + 2,
 				AllottedGeometry.ToPaintGeometry(),
 				Points,
-				SectionClippingRect
+				SectionClippingRect,
+				DrawEffects
 			);
 		}
 		CurrentTime += AnimSection->GetAnimationDuration();
diff --git a/Engine/Source/Editor/MovieSceneTools/Private/TrackEditors/SubMovieSceneTrackEditor.cpp b/Engine/Source/Editor/MovieSceneTools/Private/TrackEditors/SubMovieSceneTrackEditor.cpp
--- a/Engine/Source/Editor/MovieSceneTools/Private/TrackEditors/SubMovieSceneTrackEditor.cpp
+++ b/Engine/Source/Editor/MovieSceneTools/Private/TrackEditors/SubMovieSceneTrackEditor.cpp
@@ -50,6 +50,8 @@ public:
 
 	virtual int32 OnPaintSection( const FGeometry& AllottedGeometry, const FSlateRect& SectionClippingRect, FSlateWindowElementList& OutDrawElements, int32 LayerId, bool bParentEnabled ) const override 
 	{
+		const ESlateDrawEffect::Type DrawEffects = bParentEnabled ? ESlateDrawEffect::None : ESlateDrawEffect::DisabledEffect;
+
 		// Add a box for the section
 		FSlateDrawElement::MakeBox(
 			OutDrawElements,
@@ -57,7 +59,7 @@ public:
 			AllottedGeometry.ToPaintGeometry(),
 			FEditorStyle::GetBrush("Sequencer.GenericSection.Background"),
 			SectionClippingRect,
-			ESlateDrawEffect::None,
+			DrawEffects,
 			FColor( 220, 120, 120 )
 		);
 
diff --git a/Engine/Source/Editor/MovieSceneTools/Private/TrackEditors/TransformTrackEditor.cpp b/Engine/Source/Editor/MovieSceneTools/Private/TrackEditors/TransformTrackEditor.cpp
--- a/Engine/Source/Editor/MovieSceneTools/Private/TrackEditors/TransformTrackEditor.cpp
+++ b/Engine/Source/Editor/MovieSceneTools/Private/TrackEditors/TransformTrackEditor.cpp
@@ -66,13 +66,16 @@ public:
 
 	virtual int32 OnPaintSection( const FGeometry& AllottedGeometry, const FSlateRect& SectionClippingRect, FSlateWindowElementList& OutDrawElements, int32 LayerId, bool bParentEnabled ) const override 
 	{
+		const ESlateDrawEffect::Type DrawEffects = bParentEnabled ? ESlateDrawEffect::None : ESlateDrawEffect::DisabledEffect;
+	
 		// Add a box for the section
 		FSlateDrawElement::MakeBox( 
 			OutDrawElements,
 			LayerId,
 			AllottedGeometry.ToPaintGeometry(),
 			FEditorStyle::GetBrush("Sequencer.GenericSection.Background"),
-			SectionClippingRect
+			SectionClippingRect,
+			DrawEffects
 		); 
 
 		return LayerId;
diff --git a/Engine/Source/Editor/MovieSceneTools/Public/PropertySection.h b/Engine/Source/Editor/MovieSceneTools/Public/PropertySection.h
--- a/Engine/Source/Editor/MovieSceneTools/Public/PropertySection.h
+++ b/Engine/Source/Editor/MovieSceneTools/Public/PropertySection.h
@@ -30,13 +30,16 @@ public:
 
 	virtual int32 OnPaintSection( const FGeometry& AllottedGeometry, const FSlateRect& SectionClippingRect, FSlateWindowElementList& OutDrawElements, int32 LayerId, bool bParentEnabled ) const override 
 	{
+		const ESlateDrawEffect::Type DrawEffects = bParentEnabled ? ESlateDrawEffect::None : ESlateDrawEffect::DisabledEffect;
+		
 		// Add a box for the section
 		FSlateDrawElement::MakeBox( 
 			OutDrawElements,
 			LayerId,
 			AllottedGeometry.ToPaintGeometry(),
 			FEditorStyle::GetBrush("Sequencer.GenericSection.Background"),
-			SectionClippingRect
+			SectionClippingRect,
+			DrawEffects
 		);
 
 		return LayerId;
diff --git a/Engine/Source/Editor/Sequencer/Private/SSection.cpp b/Engine/Source/Editor/Sequencer/Private/SSection.cpp
--- a/Engine/Source/Editor/Sequencer/Private/SSection.cpp
+++ b/Engine/Source/Editor/Sequencer/Private/SSection.cpp
@@ -165,20 +165,26 @@ FSequencer& SSection::GetSequencer() const
 
 int32 SSection::OnPaint( const FPaintArgs& Args, const FGeometry& AllottedGeometry, const FSlateRect& MyClippingRect, FSlateWindowElementList& OutDrawElements, int32 LayerId, const FWidgetStyle& InWidgetStyle, bool bParentEnabled ) const
 {
-	int32 StartLayer = SCompoundWidget::OnPaint( Args, AllottedGeometry, MyClippingRect, OutDrawElements, LayerId, InWidgetStyle, bParentEnabled );
+	UMovieSceneSection& SectionObject = *SectionInterface->GetSectionObject();
+
+	bool bEnabled = bParentEnabled && SectionObject.IsActive();
+
+	const ESlateDrawEffect::Type DrawEffects = bEnabled ? ESlateDrawEffect::None : ESlateDrawEffect::DisabledEffect;
+
+	int32 StartLayer = SCompoundWidget::OnPaint( Args, AllottedGeometry, MyClippingRect, OutDrawElements, LayerId, InWidgetStyle, bEnabled );
 
 	FGeometry SectionGeometry = MakeSectionGeometryWithoutHandles( AllottedGeometry, SectionInterface );
 
 	FSlateRect SectionClipRect = SectionGeometry.GetClippingRect().IntersectionWith( MyClippingRect );
 
 	// Ask the interface to draw the section
-	int32 PostSectionLayer = SectionInterface->OnPaintSection( SectionGeometry, SectionClipRect, OutDrawElements, LayerId, bParentEnabled );
+	int32 PostSectionLayer = SectionInterface->OnPaintSection( SectionGeometry, SectionClipRect, OutDrawElements, LayerId, bEnabled );
 
 	const bool bDisplaySectionHandles = true;
 
-	DrawSectionHandlesAndSelection(AllottedGeometry, MyClippingRect, OutDrawElements, PostSectionLayer, bDisplaySectionHandles );
+	DrawSectionHandlesAndSelection(AllottedGeometry, MyClippingRect, OutDrawElements, PostSectionLayer, bDisplaySectionHandles, bEnabled );
 
-	PaintKeys( SectionGeometry, MyClippingRect, OutDrawElements, PostSectionLayer, InWidgetStyle );
+	PaintKeys( SectionGeometry, MyClippingRect, OutDrawElements, PostSectionLayer, InWidgetStyle, bEnabled );
 
 	// Section name with drop shadow
 	FText SectionTitle = SectionInterface->GetSectionTitle();
@@ -191,7 +197,7 @@ int32 SSection::OnPaint( const FPaintArgs& Args, const FGeometry& AllottedGeomet
 			SectionTitle,
 			FEditorStyle::GetFontStyle("NormalFont"),
 			MyClippingRect,
-			ESlateDrawEffect::None,
+			DrawEffects,
 			FLinearColor::Black
 			);
 
@@ -202,7 +208,7 @@ int32 SSection::OnPaint( const FPaintArgs& Args, const FGeometry& AllottedGeomet
 			SectionTitle,
 			FEditorStyle::GetFontStyle("NormalFont"),
 			MyClippingRect,
-			ESlateDrawEffect::None,
+			DrawEffects,
 			FLinearColor::White
 			);
 	}
@@ -210,8 +216,10 @@ int32 SSection::OnPaint( const FPaintArgs& Args, const FGeometry& AllottedGeomet
 	return LayerId;
 }
 
-void SSection::PaintKeys( const FGeometry& AllottedGeometry, const FSlateRect& MyClippingRect, FSlateWindowElementList& OutDrawElements, int32 LayerId, const FWidgetStyle& InWidgetStyle ) const
+void SSection::PaintKeys( const FGeometry& AllottedGeometry, const FSlateRect& MyClippingRect, FSlateWindowElementList& OutDrawElements, int32 LayerId, const FWidgetStyle& InWidgetStyle, bool bParentEnabled ) const
 {
+	const ESlateDrawEffect::Type DrawEffects = bParentEnabled ? ESlateDrawEffect::None : ESlateDrawEffect::DisabledEffect;
+
 	UMovieSceneSection& SectionObject = *SectionInterface->GetSectionObject();
 
 	FSequencer& Sequencer = ParentSectionArea->GetSequencer();
@@ -255,7 +263,7 @@ void SSection::PaintKeys( const FGeometry& AllottedGeometry, const FSlateRect& M
 			KeyAreaGeometry.ToPaintGeometry(),
 			BackgroundBrush,
 			MyClippingRect,
-			ESlateDrawEffect::None,
+			DrawEffects,
 			FLinearColor( .1f, .1f, .1f, 0.7f ) ); 
 
 
@@ -320,7 +328,7 @@ void SSection::PaintKeys( const FGeometry& AllottedGeometry, const FSlateRect& M
 					KeyAreaGeometry.ToPaintGeometry( FVector2D( KeyPosition - FMath::CeilToFloat(SequencerSectionConstants::KeySize.X/2.0f), ((KeyAreaGeometry.Size.Y*.5f)-(SequencerSectionConstants::KeySize.Y*.5f)) ), SequencerSectionConstants::KeySize ),
 					BrushToUse,
 					MyClippingRect,
-					ESlateDrawEffect::None,
+					DrawEffects,
 					KeyColor
 					);
 			}
@@ -328,8 +336,10 @@ void SSection::PaintKeys( const FGeometry& AllottedGeometry, const FSlateRect& M
 	}
 }
 
-void SSection::DrawSectionHandlesAndSelection( const FGeometry& AllottedGeometry, const FSlateRect& MyClippingRect, FSlateWindowElementList& OutDrawElements, int32 LayerId, bool bDisplaySectionHandles ) const
+void SSection::DrawSectionHandlesAndSelection( const FGeometry& AllottedGeometry, const FSlateRect& MyClippingRect, FSlateWindowElementList& OutDrawElements, int32 LayerId, bool bDisplaySectionHandles, bool bParentEnabled ) const
 {
+	const ESlateDrawEffect::Type DrawEffects = bParentEnabled ? ESlateDrawEffect::None : ESlateDrawEffect::DisabledEffect;
+
 	UMovieSceneSection* SectionObject = SectionInterface->GetSectionObject();
 
 	FSequencerSelection& Selection = ParentSectionArea->GetSequencer().GetSelection();
@@ -354,7 +364,7 @@ void SSection::DrawSectionHandlesAndSelection( const FGeometry& AllottedGeometry
 			AllottedGeometry.ToPaintGeometry(FVector2D(0.0f, 0.0f), FVector2D(SectionInterface->GetSectionGripSize(), AllottedGeometry.GetDrawSize().Y)),
 			LeftGripBrush,
 			MyClippingRect,
-			ESlateDrawEffect::None,
+			DrawEffects,
 			(bLeftEdgePressed || bLeftEdgeHovered) ? TransparentSelectionColor : LeftGripBrush->GetTint(FWidgetStyle())
 		);
 		
@@ -367,7 +377,7 @@ void SSection::DrawSectionHandlesAndSelection( const FGeometry& AllottedGeometry
 			AllottedGeometry.ToPaintGeometry(FVector2D(AllottedGeometry.Size.X-SectionInterface->GetSectionGripSize(), 0.0f), FVector2D(SectionInterface->GetSectionGripSize(), AllottedGeometry.GetDrawSize().Y)),
 			RightGripBrush,
 			MyClippingRect,
-			ESlateDrawEffect::None,
+			DrawEffects,
 			(bRightEdgePressed || bRightEdgeHovered) ? TransparentSelectionColor : RightGripBrush->GetTint(FWidgetStyle())
 		);
 	}
@@ -384,7 +394,7 @@ void SSection::DrawSectionHandlesAndSelection( const FGeometry& AllottedGeometry
 			AllottedGeometry.ToPaintGeometry(),
 			FEditorStyle::GetBrush(SelectionBorder),
 			MyClippingRect,
-			ESlateDrawEffect::None,
+			DrawEffects,
 			bActive ? SelectionColor : SelectionInactiveColor
 			);
 	}
@@ -416,6 +426,18 @@ TSharedPtr<SWidget> SSection::OnSummonContextMenu( const FGeometry& MyGeometry,
 	{
 		SectionInterface->BuildSectionContextMenu(MenuBuilder);
 
+		MenuBuilder.AddMenuEntry(
+			NSLOCTEXT("Sequencer", "ToggleSectionActive", "Active"),
+			NSLOCTEXT("Sequencer", "ToggleSectionActiveTooltip", "Toggle section active/inactive."),
+			FSlateIcon(),
+			FUIAction(
+				FExecuteAction::CreateSP(&Sequencer, &FSequencer::ToggleSectionActive),
+				FCanExecuteAction::CreateSP(&Sequencer, &FSequencer::CanToggleSectionActive),
+				FIsActionChecked::CreateSP(&Sequencer, &FSequencer::IsToggleSectionActive)),
+			NAME_None,
+			EUserInterfaceActionType::ToggleButton
+			);
+
 		// @todo Sequencer this should delete all selected sections
 		// delete/selection needs to be rethought in general
 		MenuBuilder.AddMenuEntry(
diff --git a/Engine/Source/Editor/Sequencer/Private/SSection.h b/Engine/Source/Editor/Sequencer/Private/SSection.h
--- a/Engine/Source/Editor/Sequencer/Private/SSection.h
+++ b/Engine/Source/Editor/Sequencer/Private/SSection.h
@@ -73,12 +73,12 @@ private:
 	 * @param OutDrawElements	List of draw elements to add to
 	 * @param LayerId			The starting draw area
 	 */
-	void PaintKeys( const FGeometry& AllottedGeometry, const FSlateRect& MyClippingRect, FSlateWindowElementList& OutDrawElements, int32 LayerId, const FWidgetStyle& InWidgetStyle ) const;
+	void PaintKeys( const FGeometry& AllottedGeometry, const FSlateRect& MyClippingRect, FSlateWindowElementList& OutDrawElements, int32 LayerId, const FWidgetStyle& InWidgetStyle, bool bParentEnabled ) const;
 
 	/**
 	 * Draw the borders around each section.  The borders are the areas which can be used to expand or collapse the size of the section.
 	 */
-	void DrawSectionHandlesAndSelection( const FGeometry& AllottedGeometry, const FSlateRect& MyClippingRect, FSlateWindowElementList& OutDrawElements, int32 LayerId, bool bDisplaySectionHandles ) const;
+	void DrawSectionHandlesAndSelection( const FGeometry& AllottedGeometry, const FSlateRect& MyClippingRect, FSlateWindowElementList& OutDrawElements, int32 LayerId, bool bDisplaySectionHandles, bool bParentEnabled ) const;
 
 	/** Summons a context menu over the associated section */
 	TSharedPtr<SWidget> OnSummonContextMenu( const FGeometry& MyGeometry, const FPointerEvent& MouseEvent );
diff --git a/Engine/Source/Editor/Sequencer/Private/Sequencer.cpp b/Engine/Source/Editor/Sequencer/Private/Sequencer.cpp
--- a/Engine/Source/Editor/Sequencer/Private/Sequencer.cpp
+++ b/Engine/Source/Editor/Sequencer/Private/Sequencer.cpp
@@ -442,6 +442,38 @@ void FSequencer::DeleteSelectedKeys()
 	Selection.EmptySelectedKeys();
 }
 
+void FSequencer::ToggleSectionActive()
+{
+	FScopedTransaction ToggleSectionActiveTransaction( NSLOCTEXT("Sequencer", "ToggleSectionActive_Transaction", "Toggle Section Active") );
+
+	bool bIsActive = !IsToggleSectionActive();
+
+	for (auto Section : Selection.GetSelectedSections())
+	{
+		Section->Modify();
+		Section->SetIsActive(bIsActive);
+	}
+}
+
+bool FSequencer::CanToggleSectionActive() const
+{
+	return Selection.GetSelectedSections().Num() > 0;
+}
+
+bool FSequencer::IsToggleSectionActive() const
+{
+	// Active only if all are active
+	for (auto Section : Selection.GetSelectedSections())
+	{
+		if (!Section->IsActive())
+		{
+			return false;
+		}
+	}
+
+	return true;
+}
+
 void FSequencer::SpawnOrDestroyPuppetObjects( TSharedRef<FMovieSceneSequenceInstance> MovieSceneInstance )
 {
 	UMovieSceneSequence* Sequence = MovieSceneInstance->GetSequence();
diff --git a/Engine/Source/Editor/Sequencer/Private/Sequencer.h b/Engine/Source/Editor/Sequencer/Private/Sequencer.h
--- a/Engine/Source/Editor/Sequencer/Private/Sequencer.h
+++ b/Engine/Source/Editor/Sequencer/Private/Sequencer.h
@@ -129,6 +129,13 @@ public:
 	 */
 	void DeleteSelectedKeys();
 
+	/**
+	* Section active/inactive toggle
+	*/
+	void ToggleSectionActive();
+	bool CanToggleSectionActive() const;
+	bool IsToggleSectionActive() const;
+
 	/**
 	 * @return Movie scene tools used by the sequencer
 	 */
diff --git a/Engine/Source/Runtime/MovieScene/Private/MovieSceneCommonHelpers.cpp b/Engine/Source/Runtime/MovieScene/Private/MovieSceneCommonHelpers.cpp
--- a/Engine/Source/Runtime/MovieScene/Private/MovieSceneCommonHelpers.cpp
+++ b/Engine/Source/Runtime/MovieScene/Private/MovieSceneCommonHelpers.cpp
@@ -36,7 +36,7 @@ UMovieSceneSection* MovieSceneHelpers::FindSectionAtTime( const TArray<UMovieSce
 		UMovieSceneSection* Section = Sections[SectionIndex];
 
 		//@todo Sequencer - There can be multiple sections overlapping in time. Returning instantly does not account for that.
-		if( Section->IsTimeWithinSection( Time ) )
+		if( Section->IsTimeWithinSection( Time ) && Section->IsActive() )
 		{
 			return Section;
 		}
@@ -59,28 +59,31 @@ UMovieSceneSection* MovieSceneHelpers::FindNearestSectionAtTime( const TArray<UM
 	{
 		UMovieSceneSection* Section = Sections[SectionIndex];
 
-		//@todo Sequencer - There can be multiple sections overlapping in time. Returning instantly does not account for that.
-		if( Section->IsTimeWithinSection( Time ) )
+		if (Section->IsActive())
 		{
-			return Section;
-		}
+			//@todo Sequencer - There can be multiple sections overlapping in time. Returning instantly does not account for that.
+			if( Section->IsTimeWithinSection( Time ) )
+			{
+				return Section;
+			}
 
-		float EndTime = Section->GetEndTime();
-		if (EndTime < Time)
-		{
-			float ClosestTime = Time - EndTime;
-			if (!ClosestSection || ClosestTime < ClosestSectionTime)
+			float EndTime = Section->GetEndTime();
+			if (EndTime < Time)
 			{
-				ClosestSection = Section;
-				ClosestSectionTime = ClosestTime;
+				float ClosestTime = Time - EndTime;
+				if (!ClosestSection || ClosestTime < ClosestSectionTime)
+				{
+					ClosestSection = Section;
+					ClosestSectionTime = ClosestTime;
+				}
 			}
-		}
 
-		float StartTime = Section->GetStartTime();
-		if (!EarliestSection || StartTime < EarliestSectionTime)
-		{
-			EarliestSection = Section;
-			EarliestSectionTime = StartTime;
+			float StartTime = Section->GetStartTime();
+			if (!EarliestSection || StartTime < EarliestSectionTime)
+			{
+				EarliestSection = Section;
+				EarliestSectionTime = StartTime;
+			}
 		}
 	}
 
diff --git a/Engine/Source/Runtime/MovieScene/Private/MovieSceneSection.cpp b/Engine/Source/Runtime/MovieScene/Private/MovieSceneSection.cpp
--- a/Engine/Source/Runtime/MovieScene/Private/MovieSceneSection.cpp
+++ b/Engine/Source/Runtime/MovieScene/Private/MovieSceneSection.cpp
@@ -5,11 +5,12 @@
 
 UMovieSceneSection::UMovieSceneSection( const FObjectInitializer& ObjectInitializer )
 	: Super( ObjectInitializer )
+	, StartTime(0.0f)
+	, EndTime(0.0f)
+	, RowIndex(0)
+	, bIsActive(true)
+	, bIsInfinite(false)
 {
-
-	StartTime = 0.0f;
-	EndTime = 0.0f;
-	RowIndex = 0;
 }
 
 const UMovieSceneSection* UMovieSceneSection::OverlapsWithSections(const TArray<UMovieSceneSection*>& Sections, int32 TrackDelta, float TimeDelta) const
diff --git a/Engine/Source/Runtime/MovieScene/Public/MovieSceneSection.h b/Engine/Source/Runtime/MovieScene/Public/MovieSceneSection.h
--- a/Engine/Source/Runtime/MovieScene/Public/MovieSceneSection.h
+++ b/Engine/Source/Runtime/MovieScene/Public/MovieSceneSection.h
@@ -176,6 +176,10 @@ public:
 	 */
 	virtual MOVIESCENE_API void InitialPlacement(const TArray<UMovieSceneSection*>& Sections, float InStartTime, float InEndTime, bool bAllowMultipleRows);
 
+	/** Whether or not this section is active. */
+	void SetIsActive(bool bInIsActive) { bIsActive = bInIsActive; }
+	bool IsActive() const { return bIsActive; }
+
 	/** Whether or not this section is infinite. An infinite section will draw the entire width of the track. StartTime and EndTime will be ignored but not discarded. */
 	void SetIsInfinite(bool bInIsInfinite) { bIsInfinite = bInIsInfinite; }
 	bool IsInfinite() const { return bIsInfinite; }
@@ -193,6 +197,10 @@ private:
 	UPROPERTY()
 	int32 RowIndex;
 
+	/** Toggle whether this section is active/inactive */
+	UPROPERTY()
+	uint32 bIsActive : 1;
+
 	/** Toggle to set this section to be infinite */
 	UPROPERTY()
 	uint32 bIsInfinite : 1;
