Fixed snapping during section moves and resizes

diff --git a/Engine/Source/Editor/Sequencer/Private/SequencerHotspots.cpp b/Engine/Source/Editor/Sequencer/Private/SequencerHotspots.cpp
--- a/Engine/Source/Editor/Sequencer/Private/SequencerHotspots.cpp
+++ b/Engine/Source/Editor/Sequencer/Private/SequencerHotspots.cpp
@@ -15,5 +15,5 @@ TSharedPtr<IEditToolDragOperation> FSectionResizeHotspot::InitiateDrag(ISequence
 	{
 		SectionHandles.Add(Section);
 	}
-	return MakeShareable( new FResizeSection(static_cast<FSequencer&>(Sequencer), SectionHandles, Section, HandleType == Right) );
+	return MakeShareable( new FResizeSection(static_cast<FSequencer&>(Sequencer), SectionHandles, HandleType == Right) );
 }
\ No newline at end of file
diff --git a/Engine/Source/Editor/Sequencer/Private/Tools/EditToolDragOperations.cpp b/Engine/Source/Editor/Sequencer/Private/Tools/EditToolDragOperations.cpp
--- a/Engine/Source/Editor/Sequencer/Private/Tools/EditToolDragOperations.cpp
+++ b/Engine/Source/Editor/Sequencer/Private/Tools/EditToolDragOperations.cpp
@@ -14,21 +14,35 @@
 
 struct FDefaultKeySnappingCandidates : ISequencerSnapCandidate
 {
-	FDefaultKeySnappingCandidates(const TSet<FSelectedKey>& InSelectedKeys)
-		: SelectedKeys(InSelectedKeys)
+	FDefaultKeySnappingCandidates(const TSet<FSelectedKey>& InKeysToExclude)
+		: KeysToExclude(InKeysToExclude)
 	{}
 
-	bool IsKeyApplicable(FKeyHandle KeyHandle, const TSharedPtr<IKeyArea>& KeyArea, UMovieSceneSection* Section)
+	virtual bool IsKeyApplicable(FKeyHandle KeyHandle, const TSharedPtr<IKeyArea>& KeyArea, UMovieSceneSection* Section) override
 	{
-		return !SelectedKeys.Contains(FSelectedKey(*Section, KeyArea, KeyHandle));
+		return !KeysToExclude.Contains(FSelectedKey(*Section, KeyArea, KeyHandle));
 	}
 
-	bool AreSectionBoundsApplicable(UMovieSceneSection* Section) { return true; }
-	bool AreGridLinesApplicable() { return true; }
-
-	const TSet<FSelectedKey>& SelectedKeys;
+	const TSet<FSelectedKey>& KeysToExclude;
 };
 
+struct FDefaultSectionSnappingCandidates : ISequencerSnapCandidate
+{
+	FDefaultSectionSnappingCandidates(const TArray<FSectionHandle>& InSectionsToIgnore)
+	{
+		for (auto& SectionHandle : InSectionsToIgnore)
+		{
+			SectionsToIgnore.Add(SectionHandle.GetSectionObject());
+		}
+	}
+
+	virtual bool AreSectionBoundsApplicable(UMovieSceneSection* Section) override
+	{
+		return !SectionsToIgnore.Contains(Section);
+	}
+
+	TSet<UMovieSceneSection*> SectionsToIgnore;
+};
 
 TOptional<FSequencerSnapField::FSnapResult> SnapToInterval(const TArray<float>& InTimes, float Threshold, const USequencerSettings& Settings)
 {
@@ -81,82 +95,6 @@ TRange<float> GetSectionBoundaries(UMovieSceneSection* Section, TSharedPtr<FTrac
 	return TRange<float>(LowerBound, UpperBound);
 }
 
-void GetSectionSnapTimes(TArray<float>& OutSnapTimes, UMovieSceneSection* Section, TSharedPtr<FTrackNode> SequencerNode, bool bIgnoreOurSectionCustomSnaps)
-{
-	// @todo Sequencer handle dilation snapping better
-
-	// Collect all the potential snap times from other section borders
-	const TArray< TSharedRef<ISequencerSection> >& Sections = SequencerNode->GetSections();
-	for (int32 SectionIndex = 0; SectionIndex < Sections.Num(); ++SectionIndex)
-	{
-		const UMovieSceneSection* InSection = Sections[SectionIndex]->GetSectionObject();
-		bool bIsThisSection = Section == InSection;
-		if (!bIgnoreOurSectionCustomSnaps || !bIsThisSection)
-		{
-			InSection->GetSnapTimes(OutSnapTimes, Section != InSection);
-		}
-	}
-
-	// snap to shots if it exists, and we are not the director track
-	UMovieSceneTrack* OuterTrack = Cast<UMovieSceneTrack>(Section->GetOuter());
-	UMovieScene* MovieScene = Cast<UMovieScene>(OuterTrack->GetOuter());
-	UMovieSceneTrack* ShotTrack = MovieScene->FindMasterTrack(UMovieSceneShotTrack::StaticClass());
-	if (ShotTrack && OuterTrack != ShotTrack)
-	{
-		const TArray<UMovieSceneSection*>& ShotSections = ShotTrack->GetAllSections();
-		for (int32 SectionIndex = 0; SectionIndex < ShotSections.Num(); ++SectionIndex)
-		{
-			auto Shot = ShotSections[SectionIndex];
-			Shot->GetSnapTimes(OutSnapTimes, true);
-		}
-	}
-}
-
-bool SnapToTimes(TArray<float> InitialTimes, const TArray<float>& SnapTimes, const FTimeToPixel& TimeToPixelConverter, float& OutInitialTime, float& OutSnapTime)
-{
-	bool bSuccess = false;
-	float ClosestTimePixelDistance = PixelSnapWidth;
-	
-	for (int32 InitialTimeIndex = 0; InitialTimeIndex < InitialTimes.Num(); ++InitialTimeIndex)
-	{
-		float InitialTime = InitialTimes[InitialTimeIndex];
-		float PixelXOfTime = TimeToPixelConverter.TimeToPixel(InitialTime);
-
-		for (int32 SnapTimeIndex = 0; SnapTimeIndex < SnapTimes.Num(); ++SnapTimeIndex)
-		{
-			float SnapTime = SnapTimes[SnapTimeIndex];
-			float PixelXOfSnapTime = TimeToPixelConverter.TimeToPixel(SnapTime);
-
-			float PixelDistance = FMath::Abs(PixelXOfTime - PixelXOfSnapTime);
-			if (PixelDistance < ClosestTimePixelDistance)
-			{
-				ClosestTimePixelDistance = PixelDistance;
-				OutInitialTime = InitialTime;
-				OutSnapTime = SnapTime;
-				bSuccess = true;
-			}
-		}
-	}
-
-	return bSuccess;
-}
-
-TOptional<float> SnapToTimes(float InitialTime, const TArray<float>& SnapTimes, const FTimeToPixel& TimeToPixelConverter)
-{
-	TArray<float> InitialTimes;
-	InitialTimes.Add(InitialTime);
-
-	float OutSnapTime = 0.f;
-	float OutInitialTime = 0.f;
-	bool bSuccess = SnapToTimes(InitialTimes, SnapTimes, TimeToPixelConverter, OutInitialTime, OutSnapTime);
-	if (bSuccess)
-	{
-		return TOptional<float>(OutSnapTime);
-	}
-	return TOptional<float>();
-}
-
-
 FEditToolDragOperation::FEditToolDragOperation( FSequencer& InSequencer )
 	: Sequencer(InSequencer)
 {
@@ -184,12 +122,11 @@ void FEditToolDragOperation::EndTransaction()
 	Sequencer.UpdateRuntimeInstances();
 }
 
-FResizeSection::FResizeSection( FSequencer& InSequencer, TArray<FSectionHandle> InSections, TOptional<FSectionHandle> InCardinalSection, bool bInDraggingByEnd )
+FResizeSection::FResizeSection( FSequencer& InSequencer, TArray<FSectionHandle> InSections, bool bInDraggingByEnd )
 	: FEditToolDragOperation( InSequencer )
 	, Sections( MoveTemp(InSections) )
-	, CardinalSection(InCardinalSection)
 	, bDraggingByEnd(bInDraggingByEnd)
-	, MouseDownTime(ForceInit)
+	, MouseDownTime(0.f)
 {
 }
 
@@ -199,6 +136,10 @@ void FResizeSection::OnBeginDrag(const FPointerEvent& MouseEvent, FVector2D Loca
 
 	MouseDownTime = VirtualTrackArea.PixelToTime(LocalMousePos.X);
 
+	// Construct a snap field of unselected sections
+	FDefaultSectionSnappingCandidates SnapCandidates(Sections);
+	SnapField = FSequencerSnapField(Sequencer, SnapCandidates, ESequencerEntity::Section);
+
 	DraggedKeyHandles.Empty();
 	SectionInitTimes.Empty();
 
@@ -222,40 +163,44 @@ void FResizeSection::OnDrag(const FPointerEvent& MouseEvent, FVector2D LocalMous
 {
 	bool bIsDilating = MouseEvent.IsControlDown();
 
-	for (auto& Handle : Sections)
+	// Convert the current mouse position to a time
+	float DeltaTime = VirtualTrackArea.PixelToTime(LocalMousePos.X) - MouseDownTime;
+
+	// Snapping
+	if ( Settings->GetIsSnapEnabled() )
 	{
-		UMovieSceneSection* Section = Handle.GetSectionObject();
+		TArray<float> SectionTimes;
+		for (auto& Handle : Sections)
+		{
+			UMovieSceneSection* Section = Handle.GetSectionObject();
+			SectionTimes.Add(SectionInitTimes[Section] + DeltaTime);
+		}
 
-		// Convert the current mouse position to a time
-		float DeltaTime = VirtualTrackArea.PixelToTime(LocalMousePos.X) - MouseDownTime;
-		float NewTime = SectionInitTimes[Section] + DeltaTime;
+		float SnapThresold = VirtualTrackArea.PixelToTime(PixelSnapWidth) - VirtualTrackArea.PixelToTime(0.f);
 
-		// Find the borders of where you can drag to
-		TRange<float> SectionBoundaries = GetSectionBoundaries(Section, Handle.TrackNode);
-		
-		// Snapping
-		if ( Settings->GetIsSnapEnabled() )
-		{
-			bool bSnappedToSection = false;
-			if ( Settings->GetSnapSectionTimesToSections() )
-			{
-				TArray<float> TimesToSnapTo;
-				GetSectionSnapTimes(TimesToSnapTo, Section, Handle.TrackNode, bIsDilating);
+		TOptional<FSequencerSnapField::FSnapResult> SnappedTime;
 
-				TOptional<float> NewSnappedTime = SnapToTimes(NewTime, TimesToSnapTo, VirtualTrackArea);
+		if (Settings->GetSnapSectionTimesToSections())
+		{
+			SnappedTime = SnapField->Snap(SectionTimes, SnapThresold);
+		}
 
-				if (NewSnappedTime.IsSet())
-				{
-					NewTime = NewSnappedTime.GetValue();
-					bSnappedToSection = true;
-				}
-			}
+		if (!SnappedTime.IsSet() && Settings->GetSnapSectionTimesToInterval())
+		{
+			SnappedTime = SnapToInterval(SectionTimes, SnapThresold, *Settings);
+		}
 
-			if ( bSnappedToSection == false && Settings->GetSnapSectionTimesToInterval() )
-			{
-				NewTime = Settings->SnapTimeToInterval(NewTime);
-			}
+		if (SnappedTime.IsSet())
+		{
+			// Add the snapped amopunt onto the delta
+			DeltaTime += SnappedTime->Snapped - SnappedTime->Original;
 		}
+	}
+
+	for (auto& Handle : Sections)
+	{
+		UMovieSceneSection* Section = Handle.GetSectionObject();
+		float NewTime = SectionInitTimes[Section] + DeltaTime;
 
 		if( bDraggingByEnd )
 		{
@@ -274,7 +219,7 @@ void FResizeSection::OnDrag(const FPointerEvent& MouseEvent, FVector2D LocalMous
 				Section->SetEndTime( NewTime );
 			}
 		}
-		else if( !bDraggingByEnd )
+		else
 		{
 			// Dragging the start of a section
 			// Ensure we arent expanding past the end time
@@ -311,13 +256,21 @@ void FMoveSection::OnBeginDrag(const FPointerEvent& MouseEvent, FVector2D LocalM
 {
 	BeginTransaction( Sections, NSLOCTEXT("Sequencer", "MoveSectionTransaction", "Move Section") );
 
-	MouseDownTime = VirtualTrackArea.PixelToTime(LocalMousePos.X);
+	// Construct a snap field of unselected sections
+	FDefaultSectionSnappingCandidates SnapCandidates(Sections);
+	SnapField = FSequencerSnapField(Sequencer, SnapCandidates, ESequencerEntity::Section);
 
 	DraggedKeyHandles.Empty();
 
-	for (auto Section : Sections)
+	const FVector2D InitialPosition = VirtualTrackArea.PhysicalToVirtual(LocalMousePos);
+
+	RelativeOffsets.Reserve(Sections.Num());
+	for (int32 Index = 0; Index < Sections.Num(); ++Index)
 	{
-		Section.GetSectionObject()->GetKeyHandles(DraggedKeyHandles);
+		auto* Section = Sections[Index].GetSectionObject();
+
+		Section->GetKeyHandles(DraggedKeyHandles);
+		RelativeOffsets.Add(FRelativeOffset{ Section->GetStartTime() - InitialPosition.X, Section->GetEndTime() - InitialPosition.X });
 	}
 }
 
@@ -342,17 +295,53 @@ void FMoveSection::OnEndDrag(const FPointerEvent& MouseEvent, FVector2D LocalMou
 	EndTransaction();
 }
 
-
 void FMoveSection::OnDrag(const FPointerEvent& MouseEvent, FVector2D LocalMousePos, const FVirtualTrackArea& VirtualTrackArea)
 {
-	float NewMouseDownTime = VirtualTrackArea.PixelToTime(LocalMousePos.X);
-	float DistanceMoved =  NewMouseDownTime - MouseDownTime;
-	float DeltaTime = DistanceMoved;
+	LocalMousePos.Y = FMath::Clamp(LocalMousePos.Y, 0.f, VirtualTrackArea.GetPhysicalSize().Y);
 
-	for (auto& Handle : Sections)
+	// Convert the current mouse position to a time
+	FVector2D VirtualMousePos = VirtualTrackArea.PhysicalToVirtual(LocalMousePos);
+
+	// Snapping
+	if ( Settings->GetIsSnapEnabled() )
+	{
+		TArray<float> SectionTimes;
+		SectionTimes.Reserve(RelativeOffsets.Num());
+		for (const auto& Offset : RelativeOffsets)
+		{
+			SectionTimes.Add(VirtualMousePos.X + Offset.StartTime);
+			SectionTimes.Add(VirtualMousePos.X + Offset.EndTime);
+		}
+
+		float SnapThresold = VirtualTrackArea.PixelToTime(PixelSnapWidth) - VirtualTrackArea.PixelToTime(0.f);
+
+		TOptional<FSequencerSnapField::FSnapResult> SnappedTime;
+
+		if (Settings->GetSnapSectionTimesToSections())
+		{
+			SnappedTime = SnapField->Snap(SectionTimes, SnapThresold);
+		}
+
+		if (!SnappedTime.IsSet() && Settings->GetSnapSectionTimesToInterval())
+		{
+			SnappedTime = SnapToInterval(SectionTimes, SnapThresold, *Settings);
+		}
+
+		if (SnappedTime.IsSet())
+		{
+			// Add the snapped amopunt onto the delta
+			VirtualMousePos.X += SnappedTime->Snapped - SnappedTime->Original;
+		}
+	}
+
+	for (int32 Index = 0; Index < Sections.Num(); ++Index)
 	{
+		auto& Handle = Sections[Index];
 		UMovieSceneSection* Section = Handle.GetSectionObject();
 
+		// *2 because we store start/end times
+		const float DeltaTime = VirtualMousePos.X + RelativeOffsets[Index].StartTime - Section->GetStartTime();
+
 		auto SequencerNodeSections = Handle.TrackNode->GetSections();
 		
 		TArray<UMovieSceneSection*> MovieSceneSections;
@@ -360,58 +349,44 @@ void FMoveSection::OnDrag(const FPointerEvent& MouseEvent, FVector2D LocalMouseP
 		{
 			MovieSceneSections.Add(SequencerNodeSections[i]->GetSectionObject());
 		}
-		
-		if ( Settings->GetIsSnapEnabled() )
-		{
-			bool bSnappedToSection = false;
-			if ( Settings->GetSnapSectionTimesToSections() )
-			{
-				TArray<float> TimesToSnapTo;
-				GetSectionSnapTimes(TimesToSnapTo, Section, Handle.TrackNode, true);
 
-				TArray<float> TimesToSnap;
-				TimesToSnap.Add(DistanceMoved + Section->GetStartTime());
-				TimesToSnap.Add(DistanceMoved + Section->GetEndTime());
+		int32 TargetRowIndex = Section->GetRowIndex();
 
-				float OutSnappedTime = 0.f;
-				float OutNewTime = 0.f;
-				if (SnapToTimes(TimesToSnap, TimesToSnapTo, VirtualTrackArea, OutSnappedTime, OutNewTime))
+		// vertical dragging - master tracks only
+		if (Handle.TrackNode->GetTrack()->SupportsMultipleRows() && SequencerNodeSections.Num() > 1)
+		{
+			int32 NumRows = 0;
+			for (auto* Sec : MovieSceneSections)
+			{
+				if (Sec != Section)
 				{
-					DeltaTime = OutNewTime - (OutSnappedTime - DistanceMoved);
-					bSnappedToSection = true;
+					NumRows = FMath::Max(Sec->GetRowIndex() + 1, NumRows);
 				}
 			}
+			
+			// Compute the max row index whilst disregarding the one we're dragging
+			const int32 MaxRowIndex = NumRows;
 
-			if ( bSnappedToSection == false && Settings->GetSnapSectionTimesToInterval() )
-			{
-				float NewStartTime = DistanceMoved + Section->GetStartTime();
-				DeltaTime = Settings->SnapTimeToInterval( NewStartTime ) - Section->GetStartTime();
-			}
-		}
+			// Now factor in the row index of the one we're dragging
+			NumRows = FMath::Max(Section->GetRowIndex() + 1, NumRows);
 
-		int32 TargetRowIndex = Section->GetRowIndex();
+			const float VirtualSectionHeight = Handle.TrackNode->GetVirtualBottom() - Handle.TrackNode->GetVirtualTop();
+			const float VirtualRowHeight = VirtualSectionHeight / NumRows;
+			const float MouseOffsetWithinRow = VirtualMousePos.Y - (Handle.TrackNode->GetVirtualTop() + VirtualRowHeight * TargetRowIndex);
 
-		// vertical dragging - master tracks only
-		if (Handle.TrackNode->GetTrack()->SupportsMultipleRows() && SequencerNodeSections.Num() > 1)
-		{
-			float TrackHeight = SequencerNodeSections[0]->GetSectionHeight();
-
-			if (LocalMousePos.Y < 0.f || LocalMousePos.Y > TrackHeight)
+			if (MouseOffsetWithinRow < 0 || MouseOffsetWithinRow > VirtualRowHeight)
 			{
-				int32 MaxRowIndex = 0;
-				for (int32 i = 0; i < SequencerNodeSections.Num(); ++i)
+				const int32 NewIndex = FMath::FloorToInt((VirtualMousePos.Y - Handle.TrackNode->GetVirtualTop()) / VirtualRowHeight);
+
+				if (NewIndex < 0)
 				{
-					if (SequencerNodeSections[i]->GetSectionObject() != Section)
-					{
-						MaxRowIndex = FMath::Max(MaxRowIndex, SequencerNodeSections[i]->GetSectionObject()->GetRowIndex());
-					}
+					// todo: Move everything else down?
 				}
 
-				TargetRowIndex = FMath::Clamp(Section->GetRowIndex() + FMath::FloorToInt(LocalMousePos.Y / TrackHeight),
-					0, MaxRowIndex + 1);
+				TargetRowIndex = FMath::Clamp(NewIndex,	0, MaxRowIndex);
 			}
 		}
-		
+
 		bool bDeltaX = !FMath::IsNearlyZero(DeltaTime);
 		bool bDeltaY = TargetRowIndex != Section->GetRowIndex();
 
@@ -449,8 +424,6 @@ void FMoveSection::OnDrag(const FPointerEvent& MouseEvent, FVector2D LocalMouseP
 			}
 		}
 	}
-
-	MouseDownTime = NewMouseDownTime;
 }
 
 void FMoveKeys::OnBeginDrag(const FPointerEvent& MouseEvent, FVector2D LocalMousePos, const FVirtualTrackArea& VirtualTrackArea)
@@ -507,7 +480,7 @@ void FMoveKeys::OnDrag(const FPointerEvent& MouseEvent, FVector2D LocalMousePos,
 			KeyTimes.Add(MouseTime + RelativeOffsets.FindRef(SelectedKey));
 		}
 
-		float SnapThresold = VirtualTrackArea.PixelToTime(10.f) - VirtualTrackArea.PixelToTime(0.f);
+		float SnapThresold = VirtualTrackArea.PixelToTime(PixelSnapWidth) - VirtualTrackArea.PixelToTime(0.f);
 
 		TOptional<FSequencerSnapField::FSnapResult> SnappedTime;
 
diff --git a/Engine/Source/Editor/Sequencer/Private/Tools/EditToolDragOperations.h b/Engine/Source/Editor/Sequencer/Private/Tools/EditToolDragOperations.h
--- a/Engine/Source/Editor/Sequencer/Private/Tools/EditToolDragOperations.h
+++ b/Engine/Source/Editor/Sequencer/Private/Tools/EditToolDragOperations.h
@@ -43,7 +43,7 @@ protected:
 class FResizeSection : public FEditToolDragOperation
 {
 public:
-	FResizeSection( FSequencer& InSequencer, TArray<FSectionHandle> Sections, TOptional<FSectionHandle> InCardinalSection, bool bInDraggingByEnd );
+	FResizeSection( FSequencer& InSequencer, TArray<FSectionHandle> Sections, bool bInDraggingByEnd );
 
 	/** FEditToolDragOperation interface */
 	virtual void OnBeginDrag(const FPointerEvent& MouseEvent, FVector2D LocalMousePos, const FVirtualTrackArea& VirtualTrackArea) override;
@@ -53,8 +53,6 @@ public:
 private:
 	/** The sections we are interacting with */
 	TArray<FSectionHandle> Sections;
-	/** Optional cardinal section that the user is directly interacting with */
-	TOptional<FSectionHandle> CardinalSection;
 	/** true if dragging  the end of the section, false if dragging the start */
 	bool bDraggingByEnd;
 	/** Time where the mouse is pressed */
@@ -63,6 +61,8 @@ private:
 	TMap<TWeakObjectPtr<UMovieSceneSection>, float> SectionInitTimes;
 	/** The exact key handles that we're dragging */
 	TSet<FKeyHandle> DraggedKeyHandles;
+	/** Optional snap field to use when dragging */
+	TOptional<FSequencerSnapField> SnapField;
 };
 
 /** Operation to move the currently selected sections */
@@ -80,12 +80,13 @@ public:
 private:
 	/** The sections we are interacting with */
 	TArray<FSectionHandle> Sections;
-	/** Optional cardinal section that the user is directly interacting with */
-	TOptional<FSectionHandle> CardinalSection;
-	/** Time where the mouse is pressed */
-	float MouseDownTime;
 	/** The exact key handles that we're dragging */
 	TSet<FKeyHandle> DraggedKeyHandles;
+	/** Array of desired offsets relative to the mouse position. Representes a contiguous array of start/end times directly corresponding to Sections array */
+	struct FRelativeOffset { float StartTime, EndTime; };
+	TArray<FRelativeOffset> RelativeOffsets;
+	/** Optional snap field to use when dragging */
+	TOptional<FSequencerSnapField> SnapField;
 };
 
 /**
@@ -105,11 +106,6 @@ public:
 	virtual void OnDrag(const FPointerEvent& MouseEvent, FVector2D LocalMousePos, const FVirtualTrackArea& VirtualTrackArea) override;
 	virtual void OnEndDrag(const FPointerEvent& MouseEvent, FVector2D LocalMousePos, const FVirtualTrackArea& VirtualTrackArea) override;
 
-protected:
-
-	/** Gets all potential snap times for this drag operation */
-	void GetKeySnapTimes(TArray<float>& OutSnapTimes, TSharedPtr<FTrackNode> SequencerNode);
-
 private:	
 	/** The selected keys being moved. */
 	const TSet<FSelectedKey>& SelectedKeys;
diff --git a/Engine/Source/Editor/Sequencer/Private/Tools/SequencerEntityVisitor.cpp b/Engine/Source/Editor/Sequencer/Private/Tools/SequencerEntityVisitor.cpp
--- a/Engine/Source/Editor/Sequencer/Private/Tools/SequencerEntityVisitor.cpp
+++ b/Engine/Source/Editor/Sequencer/Private/Tools/SequencerEntityVisitor.cpp
@@ -107,13 +107,19 @@ void FSequencerEntityWalker::HandleNode(const ISequencerEntityVisitor& Visitor,
 				UMovieSceneSection* Section = InSections[SectionIndex]->GetSectionObject();
 				if (Range.IntersectSection(Section))
 				{
-					Visitor.VisitSection(Section);
+					if (Visitor.CheckEntityMask(ESequencerEntity::Section))
+					{
+						Visitor.VisitSection(Section);
+					}
 
-					// Only handle grouped keys if we actually have children
-					if (InNode.GetChildNodes().Num() != 0 && Range.IntersectKeyArea(InNode, VirtualKeySize.X))
+					if (Visitor.CheckEntityMask(ESequencerEntity::Key))
 					{
-						TSharedRef<IKeyArea> KeyArea = InNode.UpdateKeyGrouping(SectionIndex);
-						HandleKeyArea(Visitor, KeyArea, Section);
+						// Only handle grouped keys if we actually have children
+						if (InNode.GetChildNodes().Num() != 0 && Range.IntersectKeyArea(InNode, VirtualKeySize.X))
+						{
+							TSharedRef<IKeyArea> KeyArea = InNode.UpdateKeyGrouping(SectionIndex);
+							HandleKeyArea(Visitor, KeyArea, Section);
+						}
 					}
 				}
 			}
@@ -142,12 +148,18 @@ void FSequencerEntityWalker::HandleKeyAreaNode(const ISequencerEntityVisitor& Vi
 		// If the section is at all within the marquee, we check its keys
 		if (Range.IntersectSection(Section))
 		{
-			Visitor.VisitSection(Section);
+			if (Visitor.CheckEntityMask(ESequencerEntity::Section))
+			{
+				Visitor.VisitSection(Section);
+			}
 
-			if (Range.IntersectKeyArea(InOwnerNode, VirtualKeySize.X))
+			if (Visitor.CheckEntityMask(ESequencerEntity::Key))
 			{
-				TSharedRef<IKeyArea> KeyArea = InKeyAreaNode.GetKeyArea(SectionIndex);
-				HandleKeyArea(Visitor, KeyArea, InSections[SectionIndex]->GetSectionObject());
+				if (Range.IntersectKeyArea(InOwnerNode, VirtualKeySize.X))
+				{
+					TSharedRef<IKeyArea> KeyArea = InKeyAreaNode.GetKeyArea(SectionIndex);
+					HandleKeyArea(Visitor, KeyArea, InSections[SectionIndex]->GetSectionObject());
+				}
 			}
 		}
 	}
diff --git a/Engine/Source/Editor/Sequencer/Private/Tools/SequencerEntityVisitor.h b/Engine/Source/Editor/Sequencer/Private/Tools/SequencerEntityVisitor.h
--- a/Engine/Source/Editor/Sequencer/Private/Tools/SequencerEntityVisitor.h
+++ b/Engine/Source/Editor/Sequencer/Private/Tools/SequencerEntityVisitor.h
@@ -8,11 +8,33 @@ class IKeyArea;
 class FSectionKeyAreaNode;
 class ISequencerSection;
 
+/** Enum of different types of entities that are available in the sequencer */
+namespace ESequencerEntity
+{
+	enum Type
+	{
+		Key			= 1<<0,
+		Section		= 1<<1,
+	};
+
+	static const uint32 Everything = (uint32)-1;
+};
+
 /** Visitor class used to handle specific sequencer entities */
 struct ISequencerEntityVisitor
 {
+	ISequencerEntityVisitor(uint32 InEntityMask = ESequencerEntity::Everything) : EntityMask(InEntityMask) {}
+
 	virtual void VisitKey(FKeyHandle KeyHandle, float KeyTime, const TSharedPtr<IKeyArea>& KeyArea, UMovieSceneSection* Section) const { }
 	virtual void VisitSection(UMovieSceneSection* Section) const { }
+	
+	/** Check if the specified type of entity is applicable to this visitor */
+	bool CheckEntityMask(ESequencerEntity::Type Type) const { return (EntityMask & Type) != 0; }
+
+protected:
+
+	/** Bitmask of allowable entities */
+	uint32 EntityMask;
 };
 
 /** A range specifying time (and possibly vertical) bounds in the sequencer */
diff --git a/Engine/Source/Editor/Sequencer/Private/Tools/SequencerSnapField.cpp b/Engine/Source/Editor/Sequencer/Private/Tools/SequencerSnapField.cpp
--- a/Engine/Source/Editor/Sequencer/Private/Tools/SequencerSnapField.cpp
+++ b/Engine/Source/Editor/Sequencer/Private/Tools/SequencerSnapField.cpp
@@ -10,8 +10,9 @@
 
 struct FSnapGridVisitor : ISequencerEntityVisitor
 {
-	FSnapGridVisitor(ISequencerSnapCandidate& InCandidate)
-		: Candidate(InCandidate)
+	FSnapGridVisitor(ISequencerSnapCandidate& InCandidate, uint32 EntityMask)
+		: ISequencerEntityVisitor(EntityMask)
+		, Candidate(InCandidate)
 	{}
 
 	virtual void VisitKey(FKeyHandle KeyHandle, float KeyTime, const TSharedPtr<IKeyArea>& KeyArea, UMovieSceneSection* Section) const
@@ -44,7 +45,7 @@ struct FSnapGridVisitor : ISequencerEntityVisitor
 	mutable TArray<FSequencerSnapPoint> Snaps;
 };
 
-FSequencerSnapField::FSequencerSnapField(const ISequencer& InSequencer, ISequencerSnapCandidate& Candidate)
+FSequencerSnapField::FSequencerSnapField(const ISequencer& InSequencer, ISequencerSnapCandidate& Candidate, uint32 EntityMask)
 {
 	TSharedPtr<SSequencerTreeView> TreeView = StaticCastSharedRef<SSequencer>(InSequencer.GetSequencerWidget())->GetTreeView();
 
@@ -58,7 +59,7 @@ FSequencerSnapField::FSequencerSnapField(const ISequencer& InSequencer, ISequenc
 	FSequencerEntityWalker Walker(ViewRange);
 
 	// Traverse the visible space, collecting snapping times as we go
-	FSnapGridVisitor Visitor(Candidate);
+	FSnapGridVisitor Visitor(Candidate, EntityMask);
 	Walker.Traverse(Visitor, VisibleNodes);
 
 	// Sort
diff --git a/Engine/Source/Editor/Sequencer/Private/Tools/SequencerSnapField.h b/Engine/Source/Editor/Sequencer/Private/Tools/SequencerSnapField.h
--- a/Engine/Source/Editor/Sequencer/Private/Tools/SequencerSnapField.h
+++ b/Engine/Source/Editor/Sequencer/Private/Tools/SequencerSnapField.h
@@ -2,6 +2,8 @@
 
 #pragma once
 
+#include "SequencerEntityVisitor.h"
+
 struct FKeyHandle;
 
 class ISequencer;
@@ -47,8 +49,8 @@ public:
 		float Snapped;
 	};
 
-	/** Construction from a sequencer and a snap canidate implementation */
-	FSequencerSnapField(const ISequencer& InSequencer, ISequencerSnapCandidate& Candidate);
+	/** Construction from a sequencer and a snap canidate implementation. Optionally provide an entity mask to completely ignore some entity types */
+	FSequencerSnapField(const ISequencer& InSequencer, ISequencerSnapCandidate& Candidate, uint32 EntityMask = ESequencerEntity::Everything);
 
 	/** Move construction / assignment */
 	FSequencerSnapField(FSequencerSnapField&& In) : SortedSnaps(MoveTemp(In.SortedSnaps)) {}
