Sequencer: Fix assign actor for components.

diff --git a/Engine/Source/Editor/Sequencer/Private/Sequencer.cpp b/Engine/Source/Editor/Sequencer/Private/Sequencer.cpp
--- a/Engine/Source/Editor/Sequencer/Private/Sequencer.cpp
+++ b/Engine/Source/Editor/Sequencer/Private/Sequencer.cpp
@@ -1888,18 +1888,50 @@ void FSequencer::AssignActor(FGuid InObjectBinding, FObjectBindingNode* ObjectBi
 			OwnerSequence->Modify();
 			OwnerMovieScene->Modify();
 
-			// Get the object guid to assign
+			// Handle components
+			UObject* RuntimeObject = OwnerSequence->FindObject(InObjectBinding);
+			if (AActor* ActorToReplace = Cast<AActor>(RuntimeObject))
+			{
+				for (UActorComponent* ComponentToReplace : ActorToReplace->GetComponents())
+				{
+					FGuid ComponentGuid = OwnerSequence->FindObjectId(*ComponentToReplace);
+					if (ComponentGuid.IsValid())
+					{
+						UActorComponent* NewComponent = Actor->GetComponentByClass(ComponentToReplace->GetClass());
+						if (NewComponent)
+						{
+							// Get the object guid to assign, remove the binding if it already exists
+							FGuid NewComponentGuid = OwnerSequence->FindObjectId(*NewComponent);
+							FString NewComponentLabel = NewComponent->GetName();
+							if (NewComponentGuid.IsValid())
+							{
+								OwnerMovieScene->RemovePossessable(NewComponentGuid);
+								OwnerSequence->UnbindPossessableObjects(NewComponentGuid);
+							}
+
+							// Add this object
+							FMovieScenePossessable NewPossessable( NewComponentLabel, NewComponent->GetClass());
+							NewComponentGuid = NewPossessable.GetGuid();
+							OwnerSequence->BindPossessableObject(NewComponentGuid, *NewComponent);
+
+							// Replace
+							OwnerMovieScene->ReplacePossessable(ComponentGuid, NewComponentGuid, NewComponentLabel);
+						}
+					}
+				}
+			}
+
+			// Get the object guid to assign, remove the binding if it already exists
 			FGuid NewObjectGuid = OwnerSequence->FindObjectId(*Actor);
 			FString NewActorLabel = Actor->GetActorLabel();
-
 			if (NewObjectGuid.IsValid())
 			{
 				OwnerMovieScene->RemovePossessable(NewObjectGuid);
 				OwnerSequence->UnbindPossessableObjects(NewObjectGuid);
 			}
 
-			// Add this object if it hasn't already been possessed
-			FMovieScenePossessable NewPossessable( Actor->GetActorLabel(), Actor->GetClass());
+			// Add this object
+			FMovieScenePossessable NewPossessable( NewActorLabel, Actor->GetClass());
 			NewObjectGuid = NewPossessable.GetGuid();
 			OwnerSequence->BindPossessableObject(NewObjectGuid, *Actor);
 
@@ -2169,13 +2201,16 @@ void FSequencer::BuildObjectBindingContextMenu(FMenuBuilder& MenuBuilder, const
 {
 	if (IsLevelEditorSequencer())
 	{
-		FFormatNamedArguments Args;
-		MenuBuilder.AddMenuEntry(
-			FText::Format( NSLOCTEXT("Sequencer", "Assign Actor ", "Assign Actor"), Args),
-			FText::Format( NSLOCTEXT("Sequencer", "AssignActorTooltip", "Assign the selected actor to this track"), Args ),
-			FSlateIcon(),
-			FUIAction(FExecuteAction::CreateSP(this, &FSequencer::AssignActor, ObjectBinding, ObjectBindingNode),
-						FCanExecuteAction::CreateSP(this, &FSequencer::CanAssignActor, ObjectBinding)) );
+		if (ObjectClass->IsChildOf(AActor::StaticClass()))
+		{
+			FFormatNamedArguments Args;
+			MenuBuilder.AddMenuEntry(
+				FText::Format( NSLOCTEXT("Sequencer", "Assign Actor ", "Assign Actor"), Args),
+				FText::Format( NSLOCTEXT("Sequencer", "AssignActorTooltip", "Assign the selected actor to this track"), Args ),
+				FSlateIcon(),
+				FUIAction(FExecuteAction::CreateSP(this, &FSequencer::AssignActor, ObjectBinding, ObjectBindingNode),
+							FCanExecuteAction::CreateSP(this, &FSequencer::CanAssignActor, ObjectBinding)) );
+		}
 	}
 
 	for (int32 i = 0; i < TrackEditors.Num(); ++i)
