Sequencer: Add expand/collapse nodes to menu.

diff --git a/Engine/Source/Editor/Sequencer/Private/SSequencer.cpp b/Engine/Source/Editor/Sequencer/Private/SSequencer.cpp
--- a/Engine/Source/Editor/Sequencer/Private/SSequencer.cpp
+++ b/Engine/Source/Editor/Sequencer/Private/SSequencer.cpp
@@ -670,7 +670,7 @@ TSharedRef<SWidget> SSequencer::MakeAddMenu()
 
 TSharedRef<SWidget> SSequencer::MakeGeneralMenu()
 {
-	FMenuBuilder MenuBuilder( false, Sequencer.Pin()->GetCommandBindings() );
+	FMenuBuilder MenuBuilder( true, Sequencer.Pin()->GetCommandBindings() );
 
 	MenuBuilder.BeginSection( "ViewOptions", LOCTEXT( "ViewMenuHeader", "View" ) );
 	{
@@ -679,6 +679,8 @@ TSharedRef<SWidget> SSequencer::MakeGeneralMenu()
 			MenuBuilder.AddMenuEntry( FSequencerCommands::Get().ToggleCleanView );
 		}
 		MenuBuilder.AddMenuEntry( FSequencerCommands::Get().ToggleAutoScroll );
+		MenuBuilder.AddMenuEntry( FSequencerCommands::Get().ExpandNodesAndDescendants );
+		MenuBuilder.AddMenuEntry( FSequencerCommands::Get().CollapseNodesAndDescendants );
 	}
 	MenuBuilder.EndSection();
 
diff --git a/Engine/Source/Editor/Sequencer/Private/SSequencerTreeView.cpp b/Engine/Source/Editor/Sequencer/Private/SSequencerTreeView.cpp
--- a/Engine/Source/Editor/Sequencer/Private/SSequencerTreeView.cpp
+++ b/Engine/Source/Editor/Sequencer/Private/SSequencerTreeView.cpp
@@ -359,7 +359,38 @@ bool ShouldExpand(const T& InContainer, ETreeRecursion Recursion)
 	return !bAllExpanded;
 }
 
-void SSequencerTreeView::ToggleSelectedNodeExpansion(ETreeRecursion Recursion)
+void SSequencerTreeView::ToggleExpandCollapseNodes(ETreeRecursion Recursion)
+{
+	FSequencer& Sequencer = SequencerNodeTree->GetSequencer();
+
+	const TSet< FDisplayNodeRef >& SelectedNodes = Sequencer.GetSelection().GetSelectedOutlinerNodes();
+
+	bool bExpand = false;
+	if (SelectedNodes.Num() != 0)
+	{
+		bExpand = ShouldExpand(SelectedNodes, Recursion);
+	}
+	else
+	{
+		bExpand = ShouldExpand(SequencerNodeTree->GetRootNodes(), Recursion);
+	}
+
+	ExpandOrCollapseNodes(Recursion, bExpand);
+}
+
+void SSequencerTreeView::ExpandNodes(ETreeRecursion Recursion)
+{
+	const bool bExpand = true;
+	ExpandOrCollapseNodes(Recursion, bExpand);
+}
+
+void SSequencerTreeView::CollapseNodes(ETreeRecursion Recursion)
+{
+	const bool bExpand = false;
+	ExpandOrCollapseNodes(Recursion, bExpand);
+}
+
+void SSequencerTreeView::ExpandOrCollapseNodes(ETreeRecursion Recursion, bool bExpand)
 {
 	FSequencer& Sequencer = SequencerNodeTree->GetSequencer();
 
@@ -367,7 +398,6 @@ void SSequencerTreeView::ToggleSelectedNodeExpansion(ETreeRecursion Recursion)
 
 	if (SelectedNodes.Num() != 0)
 	{
-		const bool bExpand = ShouldExpand(SelectedNodes, Recursion);
 		for (auto& Item : SelectedNodes)
 		{
 			ExpandCollapseNode(Item, bExpand, Recursion);
@@ -375,7 +405,6 @@ void SSequencerTreeView::ToggleSelectedNodeExpansion(ETreeRecursion Recursion)
 	}
 	else
 	{
-		const bool bExpand = ShouldExpand(SequencerNodeTree->GetRootNodes(), Recursion);
 		for (auto& Item : SequencerNodeTree->GetRootNodes())
 		{
 			ExpandCollapseNode(Item, bExpand, Recursion);
diff --git a/Engine/Source/Editor/Sequencer/Private/SSequencerTreeView.h b/Engine/Source/Editor/Sequencer/Private/SSequencerTreeView.h
--- a/Engine/Source/Editor/Sequencer/Private/SSequencerTreeView.h
+++ b/Engine/Source/Editor/Sequencer/Private/SSequencerTreeView.h
@@ -65,14 +65,26 @@ public:
 	/** Refresh this tree as a result of the underlying tree data changing */
 	void Refresh();
 
+	/** Expand or collapse nodes */
+	void ToggleExpandCollapseNodes(ETreeRecursion Recursion = ETreeRecursion::Recursive);
+
+	/** Expand nodes */
+	void ExpandNodes(ETreeRecursion Recursion = ETreeRecursion::Recursive);
+
+	/** Collapse nodes */
+	void CollapseNodes(ETreeRecursion Recursion = ETreeRecursion::Recursive);
+
 	/** Expand or collapse selected nodes */
-	void ToggleSelectedNodeExpansion(ETreeRecursion Recursion);
+	void ExpandSelectedNodes(ETreeRecursion Recursion);
 
 	/** Scroll this tree view by the specified number of slate units */
 	void ScrollByDelta(float DeltaInSlateUnits);
 
 protected:
 
+	/** Expand or collapse nodes */
+	void ExpandOrCollapseNodes(ETreeRecursion Recursion, bool bExpand);
+
 	/** Set the item's expansion state, including all of its children */
 	void ExpandCollapseNode(const FDisplayNodeRef& InNode, bool bExpansionState, ETreeRecursion Recursion);
 
diff --git a/Engine/Source/Editor/Sequencer/Private/Sequencer.cpp b/Engine/Source/Editor/Sequencer/Private/Sequencer.cpp
--- a/Engine/Source/Editor/Sequencer/Private/Sequencer.cpp
+++ b/Engine/Source/Editor/Sequencer/Private/Sequencer.cpp
@@ -1993,14 +1993,24 @@ void FSequencer::StepToPreviousCameraKey()
 	SequencerWidget->StepToPreviousCameraKey();
 }
 
+void FSequencer::ExpandNodesAndDescendants()
+{
+	SequencerWidget->GetTreeView()->ExpandNodes(ETreeRecursion::Recursive);
+}
+
+void FSequencer::CollapseNodesAndDescendants()
+{
+	SequencerWidget->GetTreeView()->CollapseNodes(ETreeRecursion::Recursive);
+}
+
 void FSequencer::ToggleExpandCollapseNodes()
 {
-	SequencerWidget->GetTreeView()->ToggleSelectedNodeExpansion(ETreeRecursion::NonRecursive);
+	SequencerWidget->GetTreeView()->ToggleExpandCollapseNodes(ETreeRecursion::NonRecursive);
 }
 
 void FSequencer::ToggleExpandCollapseNodesAndDescendants()
 {
-	SequencerWidget->GetTreeView()->ToggleSelectedNodeExpansion(ETreeRecursion::Recursive);
+	SequencerWidget->GetTreeView()->ToggleExpandCollapseNodes(ETreeRecursion::Recursive);
 }
 
 void FSequencer::SetKey()
@@ -2074,6 +2084,14 @@ void FSequencer::BindSequencerCommands()
 		Commands.StepToPreviousCameraKey,
 		FExecuteAction::CreateSP( this, &FSequencer::StepToPreviousCameraKey ) );
 
+	SequencerCommandBindings->MapAction(
+		Commands.ExpandNodesAndDescendants,
+		FExecuteAction::CreateSP(this, &FSequencer::ExpandNodesAndDescendants));
+
+	SequencerCommandBindings->MapAction(
+		Commands.CollapseNodesAndDescendants,
+		FExecuteAction::CreateSP(this, &FSequencer::CollapseNodesAndDescendants));
+
 	SequencerCommandBindings->MapAction(
 		Commands.ToggleExpandCollapseNodes,
 		FExecuteAction::CreateSP(this, &FSequencer::ToggleExpandCollapseNodes));
diff --git a/Engine/Source/Editor/Sequencer/Private/Sequencer.h b/Engine/Source/Editor/Sequencer/Private/Sequencer.h
--- a/Engine/Source/Editor/Sequencer/Private/Sequencer.h
+++ b/Engine/Source/Editor/Sequencer/Private/Sequencer.h
@@ -405,6 +405,9 @@ protected:
 	void StepToNextCameraKey();
 	void StepToPreviousCameraKey();
 
+	void ExpandNodesAndDescendants();
+	void CollapseNodesAndDescendants();
+
 	/** Expand or collapse selected nodes */
 	void ToggleExpandCollapseNodes();
 
diff --git a/Engine/Source/Editor/Sequencer/Private/SequencerCommands.cpp b/Engine/Source/Editor/Sequencer/Private/SequencerCommands.cpp
--- a/Engine/Source/Editor/Sequencer/Private/SequencerCommands.cpp
+++ b/Engine/Source/Editor/Sequencer/Private/SequencerCommands.cpp
@@ -17,6 +17,8 @@ void FSequencerCommands::RegisterCommands()
 	UI_COMMAND( StepToNextCameraKey, "Step to Next Camera Key", "Step to the next camera key", EUserInterfaceActionType::Button, FInputChord(EModifierKey::Alt, EKeys::Period) );
 	UI_COMMAND( StepToPreviousCameraKey, "Step to Previous Camera Key", "Step to the previous camera key", EUserInterfaceActionType::Button, FInputChord(EModifierKey::Alt, EKeys::Comma) );
 
+	UI_COMMAND( ExpandNodesAndDescendants, "Expand Nodes", "Expand selected nodes", EUserInterfaceActionType::Button, FInputChord() );
+	UI_COMMAND( CollapseNodesAndDescendants, "Collapse Nodes", "Collapse selected nodes", EUserInterfaceActionType::Button, FInputChord() );
 	UI_COMMAND( ToggleExpandCollapseNodes, "Toggle Expand/Collapse Nodes", "Toggle expand or collapse selected nodes", EUserInterfaceActionType::Button, FInputChord(EKeys::O) );
 	UI_COMMAND( ToggleExpandCollapseNodesAndDescendants, "Toggle Expand/Collapse Nodes and Descendants", "Toggle expand or collapse selected nodes and their descendants", EUserInterfaceActionType::Button, FInputChord(EModifierKey::Shift, EKeys::O) );
 
diff --git a/Engine/Source/Editor/Sequencer/Private/SequencerCommands.h b/Engine/Source/Editor/Sequencer/Private/SequencerCommands.h
--- a/Engine/Source/Editor/Sequencer/Private/SequencerCommands.h
+++ b/Engine/Source/Editor/Sequencer/Private/SequencerCommands.h
@@ -42,6 +42,12 @@ public:
 	/** Step to previous camera key */
 	TSharedPtr< FUICommandInfo > StepToPreviousCameraKey;
 
+	/** Expand nodes and descendants */
+	TSharedPtr< FUICommandInfo > ExpandNodesAndDescendants;
+
+	/** Collapse nodes and descendants */
+	TSharedPtr< FUICommandInfo > CollapseNodesAndDescendants;
+
 	/** Expand/collapse nodes */
 	TSharedPtr< FUICommandInfo > ToggleExpandCollapseNodes;
 
