Sequencer: Fix to enable tick for SSequencerTrackArea since SOverlay no longer ticks by default.

diff --git a/Engine/Source/Editor/Sequencer/Private/SSequencerTrackArea.cpp b/Engine/Source/Editor/Sequencer/Private/SSequencerTrackArea.cpp
--- a/Engine/Source/Editor/Sequencer/Private/SSequencerTrackArea.cpp
+++ b/Engine/Source/Editor/Sequencer/Private/SSequencerTrackArea.cpp
@@ -16,6 +16,9 @@ void SSequencerTrackArea::Construct( const FArguments& InArgs, TSharedRef<FSeque
 	SequencerWidget = InSequencerWidget;
 	TimeSliderController = InTimeSliderController;
 
+	//@todo Overlays don't tick by default, but we need to. Maybe try using Composition? 
+	bCanTick = true;
+
 	SOverlay::Construct(SOverlay::FArguments());
 }
 
