Sequencer: Fix sequencer editor mode so that it's enabled when a map is loaded or a new level is created.

diff --git a/Engine/Source/Editor/Sequencer/Private/Sequencer.cpp b/Engine/Source/Editor/Sequencer/Private/Sequencer.cpp
--- a/Engine/Source/Editor/Sequencer/Private/Sequencer.cpp
+++ b/Engine/Source/Editor/Sequencer/Private/Sequencer.cpp
@@ -74,6 +74,9 @@ void FSequencer::InitSequencer(const FSequencerInitParams& InitParams, const TSh
 			// Register for saving the level so that the state of the scene can be restored before saving and updated after saving.
 			FEditorDelegates::PreSaveWorld.AddSP(this, &FSequencer::OnPreSaveWorld);
 			FEditorDelegates::PostSaveWorld.AddSP(this, &FSequencer::OnPostSaveWorld);
+
+			//FEditorDelegates::NewCurrentLevel.AddSP(this, &FSequencer::OnNewCurrentLevel);
+			FEditorDelegates::OnMapOpened.AddSP(this, &FSequencer::OnMapOpened);
 		}
 		else
 		{
@@ -150,15 +153,7 @@ void FSequencer::InitSequencer(const FSequencerInitParams& InitParams, const TSh
 
 		BindSequencerCommands();
 
-		if( GLevelEditorModeTools().IsModeActive( FSequencerEdMode::EM_SequencerMode ) )
-		{
-			GLevelEditorModeTools().DeactivateMode( FSequencerEdMode::EM_SequencerMode );
-		}
-
-		GLevelEditorModeTools().ActivateMode( FSequencerEdMode::EM_SequencerMode );
-
-		FSequencerEdMode* SequencerEdMode = (FSequencerEdMode*)(GLevelEditorModeTools().GetActiveMode(FSequencerEdMode::EM_SequencerMode));
-		SequencerEdMode->SetSequencer(ActiveSequencer.Get());
+		ActivateSequencerEditorMode();
 	}
 }
 
@@ -198,6 +193,9 @@ void FSequencer::Close()
 
 		FEditorDelegates::PreSaveWorld.RemoveAll(this);
 		FEditorDelegates::PostSaveWorld.RemoveAll(this);
+
+		//FEditorDelegates::NewCurrentLevel.RemoveAll(this);
+		FEditorDelegates::OnMapOpened.RemoveAll(this);
 		
 		for (auto TrackEditor : TrackEditors)
 		{
@@ -1525,6 +1523,29 @@ void FSequencer::OnPostSaveWorld(uint32 SaveFlags, class UWorld* World, bool bSu
 	SetGlobalTime(GetGlobalTime());
 }
 
+void FSequencer::OnNewCurrentLevel()
+{
+	ActivateSequencerEditorMode();
+}
+
+void FSequencer::OnMapOpened(const FString& Filename, bool bLoadAsTemplate)
+{
+	ActivateSequencerEditorMode();
+}
+
+void FSequencer::ActivateSequencerEditorMode()
+{
+	if( GLevelEditorModeTools().IsModeActive( FSequencerEdMode::EM_SequencerMode ) )
+	{
+		GLevelEditorModeTools().DeactivateMode( FSequencerEdMode::EM_SequencerMode );
+	}
+
+	GLevelEditorModeTools().ActivateMode( FSequencerEdMode::EM_SequencerMode );
+
+	FSequencerEdMode* SequencerEdMode = (FSequencerEdMode*)(GLevelEditorModeTools().GetActiveMode(FSequencerEdMode::EM_SequencerMode));
+	SequencerEdMode->SetSequencer(ActiveSequencer.Get());
+}
+
 void FSequencer::UpdatePreviewLevelViewportClientFromCameraCut( FLevelEditorViewportClient& InViewportClient, UObject* InCameraObject, bool bNewCameraCut ) const
 {
 	ACameraActor* Cam = Cast<ACameraActor>(InCameraObject);
diff --git a/Engine/Source/Editor/Sequencer/Private/Sequencer.h b/Engine/Source/Editor/Sequencer/Private/Sequencer.h
--- a/Engine/Source/Editor/Sequencer/Private/Sequencer.h
+++ b/Engine/Source/Editor/Sequencer/Private/Sequencer.h
@@ -420,6 +420,8 @@ protected:
 	/** Generates command bindings for UI commands */
 	void BindSequencerCommands();
 
+	void ActivateSequencerEditorMode();
+
 	// Begin FEditorUndoClient Interface
 	virtual void PostUndo(bool bSuccess) override;
 	virtual void PostRedo(bool bSuccess) override { PostUndo(bSuccess); }
@@ -435,6 +437,12 @@ protected:
 	/** Called after the world has been saved. The sequencer updates to the animated state. */
 	void OnPostSaveWorld(uint32 SaveFlags, class UWorld* World, bool bSuccess);
 
+	/** Called after a new level has been created. The sequencer editor mode needs to be enabled. */
+	void OnNewCurrentLevel();
+
+	/** Called after a map has been opened. The sequencer editor mode needs to be enabled. */
+	void OnMapOpened(const FString& Filename, bool bLoadAsTemplate);
+
 	/** Updates a viewport client from camera cut data */
 	void UpdatePreviewLevelViewportClientFromCameraCut( FLevelEditorViewportClient& InViewportClient, UObject* InCameraObject, bool bNewCameraCut ) const;
 
