Sequencer: Uncommented out debug code.

diff --git a/Engine/Source/Editor/Sequencer/Private/Sequencer.cpp b/Engine/Source/Editor/Sequencer/Private/Sequencer.cpp
--- a/Engine/Source/Editor/Sequencer/Private/Sequencer.cpp
+++ b/Engine/Source/Editor/Sequencer/Private/Sequencer.cpp
@@ -75,7 +75,7 @@ void FSequencer::InitSequencer(const FSequencerInitParams& InitParams, const TSh
 			FEditorDelegates::PreSaveWorld.AddSP(this, &FSequencer::OnPreSaveWorld);
 			FEditorDelegates::PostSaveWorld.AddSP(this, &FSequencer::OnPostSaveWorld);
 
-			//FEditorDelegates::NewCurrentLevel.AddSP(this, &FSequencer::OnNewCurrentLevel);
+			FEditorDelegates::NewCurrentLevel.AddSP(this, &FSequencer::OnNewCurrentLevel);
 			FEditorDelegates::OnMapOpened.AddSP(this, &FSequencer::OnMapOpened);
 		}
 		else
@@ -194,7 +194,7 @@ void FSequencer::Close()
 		FEditorDelegates::PreSaveWorld.RemoveAll(this);
 		FEditorDelegates::PostSaveWorld.RemoveAll(this);
 
-		//FEditorDelegates::NewCurrentLevel.RemoveAll(this);
+		FEditorDelegates::NewCurrentLevel.RemoveAll(this);
 		FEditorDelegates::OnMapOpened.RemoveAll(this);
 		
 		for (auto TrackEditor : TrackEditors)
