Sequencer: Add support for keyframing properties within structs.

diff --git a/Engine/Source/Editor/Sequencer/Private/SequencerObjectChangeListener.cpp b/Engine/Source/Editor/Sequencer/Private/SequencerObjectChangeListener.cpp
--- a/Engine/Source/Editor/Sequencer/Private/SequencerObjectChangeListener.cpp
+++ b/Engine/Source/Editor/Sequencer/Private/SequencerObjectChangeListener.cpp
@@ -111,7 +111,7 @@ FOnPropagateObjectChanges& FSequencerObjectChangeListener::GetOnPropagateObjectC
 	return OnPropagateObjectChanges;
 }
 
-bool FSequencerObjectChangeListener::FindPropertySetter( const UClass& ObjectClass, const FName PropertyTypeName, const FString& InPropertyVarName ) const
+bool FSequencerObjectChangeListener::FindPropertySetter( const UClass& ObjectClass, const FName PropertyTypeName, const FString& InPropertyVarName, const UStructProperty* StructProperty ) const
 {
 	bool bFound = ClassToPropertyChangedMap.Contains( PropertyTypeName );
 
@@ -140,11 +140,21 @@ bool FSequencerObjectChangeListener::FindPropertySetter( const UClass& ObjectCla
 
 			if (Sequencer.IsValid() && Sequencer.Pin()->GetKeyInterpPropertiesOnly())
 			{
-				UProperty* Property = ObjectClass.FindPropertyByName(FName(*InPropertyVarName));
-
-				if (!Property || !Property->HasAnyPropertyFlags(CPF_Interp))
+				if (StructProperty != 0)
+				{
+					if (!StructProperty->HasAnyPropertyFlags(CPF_Interp))
+					{
+						bFound = false;
+					}
+				}
+				else
 				{
-					bFound = false;
+					UProperty* Property = ObjectClass.FindPropertyByName(FName(*InPropertyVarName));
+
+					if (!Property || !Property->HasAnyPropertyFlags(CPF_Interp))
+					{
+						bFound = false;
+					}
 				}
 			}
 		}
@@ -171,13 +181,13 @@ bool FSequencerObjectChangeListener::CanKeyProperty(FCanKeyPropertyParams CanKey
 	bool bFound = false;
 	if ( StructProperty )
 	{
-		bFound = FindPropertySetter(*CanKeyPropertyParams.ObjectClass, StructProperty->Struct->GetFName(), StructProperty->GetName() );
+		bFound = FindPropertySetter(*CanKeyPropertyParams.ObjectClass, StructProperty->Struct->GetFName(), StructProperty->GetName(), StructProperty );
 	}
 	
 	if( !bFound && ParentStructProperty )
 	{
 		// If the property parent is a struct, see if this property parent can be keyed.
-		bFound = FindPropertySetter(*CanKeyPropertyParams.ObjectClass, ParentStructProperty->Struct->GetFName(), ParentStructProperty->GetName());
+		bFound = FindPropertySetter(*CanKeyPropertyParams.ObjectClass, ParentStructProperty->Struct->GetFName(), ParentStructProperty->GetName(), ParentStructProperty );
 	}
 
 	if( !bFound )
diff --git a/Engine/Source/Editor/Sequencer/Private/SequencerObjectChangeListener.h b/Engine/Source/Editor/Sequencer/Private/SequencerObjectChangeListener.h
--- a/Engine/Source/Editor/Sequencer/Private/SequencerObjectChangeListener.h
+++ b/Engine/Source/Editor/Sequencer/Private/SequencerObjectChangeListener.h
@@ -61,7 +61,7 @@ private:
 	bool IsObjectValidForListening( UObject* Object ) const;
 
 	/** @return Whether or not a property setter could be found for a property on a class */
-	bool FindPropertySetter( const UClass& ObjectClass, const FName PropertyTypeName, const FString& PropertyVarName ) const;
+	bool FindPropertySetter( const UClass& ObjectClass, const FName PropertyTypeName, const FString& PropertyVarName, const UStructProperty* StructProperty = 0 ) const;
 private:
 	/** Mapping of object to a listener used to check for property changes */
 	TMap< TWeakObjectPtr<UObject>, TSharedPtr<class IPropertyChangeListener> > ActivePropertyChangeListeners;
