Sequencer: Add sequencer detail keyframe handler.

diff --git a/Engine/Source/Editor/PropertyEditor/Private/IDetailsViewPrivate.h b/Engine/Source/Editor/PropertyEditor/Private/IDetailsViewPrivate.h
--- a/Engine/Source/Editor/PropertyEditor/Private/IDetailsViewPrivate.h
+++ b/Engine/Source/Editor/PropertyEditor/Private/IDetailsViewPrivate.h
@@ -82,8 +82,6 @@ public:
 	 */
 	virtual bool IsPropertyReadOnly( const struct FPropertyAndParent& PropertyAndParent ) const = 0;
 
-	virtual TSharedPtr<IDetailKeyframeHandler> GetKeyframeHandler() = 0;
-
 	virtual TSharedPtr<IDetailPropertyExtensionHandler> GetExtensionHandler() = 0;
 
 	/**
diff --git a/Engine/Source/Editor/PropertyEditor/Public/IDetailsView.h b/Engine/Source/Editor/PropertyEditor/Public/IDetailsView.h
--- a/Engine/Source/Editor/PropertyEditor/Public/IDetailsView.h
+++ b/Engine/Source/Editor/PropertyEditor/Public/IDetailsView.h
@@ -192,6 +192,8 @@ public:
 
 	virtual void SetKeyframeHandler( TSharedPtr<class IDetailKeyframeHandler> InKeyframeHandler ) = 0;
 
+	virtual TSharedPtr<IDetailKeyframeHandler> GetKeyframeHandler() = 0;
+
 	virtual void SetExtensionHandler(TSharedPtr<class IDetailPropertyExtensionHandler> InExtensionandler) = 0;
 
 	/**
diff --git a/Engine/Source/Editor/Sequencer/Private/Sequencer.cpp b/Engine/Source/Editor/Sequencer/Private/Sequencer.cpp
--- a/Engine/Source/Editor/Sequencer/Private/Sequencer.cpp
+++ b/Engine/Source/Editor/Sequencer/Private/Sequencer.cpp
@@ -3,6 +3,7 @@
 #include "SequencerPrivatePCH.h"
 #include "Sequencer.h"
 #include "SequencerEdMode.h"
+#include "SequencerDetailKeyframeHandler.h"
 #include "MovieSceneSequence.h"
 #include "MovieScene.h"
 #include "Engine/LevelScriptBlueprint.h"
@@ -32,6 +33,7 @@
 #include "ISequencerSection.h"
 #include "MovieSceneSequenceInstance.h"
 #include "IKeyArea.h"
+#include "IDetailsView.h"
 #include "SnappingUtils.h"
 #include "GenericCommands.h"
 #include "Engine/BlueprintGeneratedClass.h"
@@ -54,7 +56,7 @@ namespace
 	TSharedPtr<FSequencer> ActiveSequencer;
 }
 
-void FSequencer::InitSequencer(const FSequencerInitParams& InitParams, const TSharedRef<ISequencerObjectChangeListener>& InObjectChangeListener, const TArray<FOnCreateTrackEditor>& TrackEditorDelegates)
+void FSequencer::InitSequencer(const FSequencerInitParams& InitParams, const TSharedRef<ISequencerObjectChangeListener>& InObjectChangeListener, const TSharedRef<IDetailKeyframeHandler>& InDetailKeyframeHandler, const TArray<FOnCreateTrackEditor>& TrackEditorDelegates)
 {
 	if( IsSequencerEnabled() )
 	{
@@ -88,6 +90,7 @@ void FSequencer::InitSequencer(const FSequencerInitParams& InitParams, const TSh
 		LastViewRange = TargetViewRange = InitParams.ViewParams.InitalViewRange;
 		ScrubPosition = InitParams.ViewParams.InitialScrubPosition;
 		ObjectChangeListener = InObjectChangeListener;
+		DetailKeyframeHandler = InDetailKeyframeHandler;
 
 		check( ObjectChangeListener.IsValid() );
 		
@@ -117,6 +120,8 @@ void FSequencer::InitSequencer(const FSequencerInitParams& InitParams, const TSh
 			// @todo remove when world-centric mode is added
 			// Hook into the editor's mechanism for checking whether we need live capture of PIE/SIE actor state
 			GEditor->GetActorRecordingState().AddSP(this, &FSequencer::GetActorRecordingState);
+
+			ActivateDetailKeyframeHandler();
 		}
 
 		// Create tools and bind them to this sequencer
@@ -208,6 +213,7 @@ void FSequencer::Close()
 			GLevelEditorModeTools().DeactivateMode( FSequencerEdMode::EM_SequencerMode );
 		}
 	}
+	DeactivateDetailKeyframeHandler();
 }
 
 void FSequencer::Tick(float InDeltaTime)
@@ -1546,6 +1552,42 @@ void FSequencer::ActivateSequencerEditorMode()
 	SequencerEdMode->SetSequencer(ActiveSequencer.Get());
 }
 
+void FSequencer::ActivateDetailKeyframeHandler()
+{			
+	// Add sequencer detail keyframe handler
+	static const FName DetailsTabIdentifiers[] = { "LevelEditorSelectionDetails", "LevelEditorSelectionDetails2", "LevelEditorSelectionDetails3", "LevelEditorSelectionDetails4" };
+	FPropertyEditorModule& EditModule = FModuleManager::Get().GetModuleChecked<FPropertyEditorModule>("PropertyEditor");
+
+	for(const FName& DetailsTabIdentifier : DetailsTabIdentifiers)
+	{
+		TSharedPtr<IDetailsView> DetailsView = EditModule.FindDetailView(DetailsTabIdentifier);
+
+		if(DetailsView.IsValid())
+		{
+			DetailsView->SetKeyframeHandler(DetailKeyframeHandler);
+		}
+	}
+}
+
+void FSequencer::DeactivateDetailKeyframeHandler()
+{
+	static const FName DetailsTabIdentifiers[] = { "LevelEditorSelectionDetails", "LevelEditorSelectionDetails2", "LevelEditorSelectionDetails3", "LevelEditorSelectionDetails4" };
+	FPropertyEditorModule& EditModule = FModuleManager::Get().GetModuleChecked<FPropertyEditorModule>("PropertyEditor");
+
+	for(const FName& DetailsTabIdentifier : DetailsTabIdentifiers)
+	{
+		TSharedPtr<IDetailsView> DetailsView = EditModule.FindDetailView(DetailsTabIdentifier);
+
+		if(DetailsView.IsValid())
+		{
+			if (DetailsView->GetKeyframeHandler() == DetailKeyframeHandler)
+			{
+				DetailsView->SetKeyframeHandler(0);
+			}
+		}
+	}
+}
+
 void FSequencer::UpdatePreviewLevelViewportClientFromCameraCut( FLevelEditorViewportClient& InViewportClient, UObject* InCameraObject, bool bNewCameraCut ) const
 {
 	ACameraActor* Cam = Cast<ACameraActor>(InCameraObject);
diff --git a/Engine/Source/Editor/Sequencer/Private/Sequencer.h b/Engine/Source/Editor/Sequencer/Private/Sequencer.h
--- a/Engine/Source/Editor/Sequencer/Private/Sequencer.h
+++ b/Engine/Source/Editor/Sequencer/Private/Sequencer.h
@@ -10,6 +10,7 @@ class UMovieScene;
 class IToolkitHost;
 class UMovieSceneSequence;
 class ISequencerObjectChangeListener;
+class IDetailKeyframeHandler;
 
 /**
  * Sequencer is the editing tool for MovieScene assets.
@@ -29,9 +30,10 @@ public:
 	 *
 	 * @param InitParams Initialization parameters.
 	 * @param InObjectChangeListener The object change listener to use.
+	 * @param InDetailKeyframeHandler The detail keyframe handler to use.
 	 * @param TrackEditorDelegates Delegates to call to create auto-key handlers for this sequencer.
 	 */
-	void InitSequencer(const FSequencerInitParams& InitParams, const TSharedRef<ISequencerObjectChangeListener>& InObjectChangeListener, const TArray<FOnCreateTrackEditor>& TrackEditorDelegates);
+	void InitSequencer(const FSequencerInitParams& InitParams, const TSharedRef<ISequencerObjectChangeListener>& InObjectChangeListener, const TSharedRef<IDetailKeyframeHandler>& InDetailKeyframeHandler, const TArray<FOnCreateTrackEditor>& TrackEditorDelegates);
 
 	/** Constructor */
 	FSequencer();
@@ -422,6 +424,9 @@ protected:
 
 	void ActivateSequencerEditorMode();
 
+	void ActivateDetailKeyframeHandler();
+	void DeactivateDetailKeyframeHandler();
+
 	// Begin FEditorUndoClient Interface
 	virtual void PostUndo(bool bSuccess) override;
 	virtual void PostRedo(bool bSuccess) override { PostUndo(bSuccess); }
@@ -466,6 +471,9 @@ private:
 	/** Listener for object changes being made while this sequencer is open*/
 	TSharedPtr< class ISequencerObjectChangeListener > ObjectChangeListener;
 
+	/** Listener for object changes being made while this sequencer is open*/
+	TSharedPtr< class IDetailKeyframeHandler > DetailKeyframeHandler;
+
 	/** The runtime instance for the root movie scene */
 	TSharedPtr< class FMovieSceneSequenceInstance > RootMovieSceneSequenceInstance;
 
diff --git a/Engine/Source/Editor/Sequencer/Private/SequencerDetailKeyframeHandler.cpp b/Engine/Source/Editor/Sequencer/Private/SequencerDetailKeyframeHandler.cpp
new file mode 100644
--- /dev/null
+++ b/Engine/Source/Editor/Sequencer/Private/SequencerDetailKeyframeHandler.cpp
@@ -0,0 +1,32 @@
+// Copyright 1998-2015 Epic Games, Inc. All Rights Reserved.
+
+#include "SequencerPrivatePCH.h"
+#include "SequencerDetailKeyframeHandler.h"
+#include "PropertyHandle.h"
+
+FSequencerDetailKeyframeHandler::FSequencerDetailKeyframeHandler(TSharedRef<ISequencer> InSequencer) 
+	: Sequencer(InSequencer)
+{}
+
+bool FSequencerDetailKeyframeHandler::IsPropertyKeyable(UClass* InObjectClass, const IPropertyHandle& InPropertyHandle) const
+{
+	return Sequencer.Pin()->CanKeyProperty(FCanKeyPropertyParams(InObjectClass, InPropertyHandle));
+}
+
+bool FSequencerDetailKeyframeHandler::IsPropertyKeyingEnabled() const
+{
+	return Sequencer.IsValid() && Sequencer.Pin()->GetFocusedMovieSceneSequence();
+}
+
+void FSequencerDetailKeyframeHandler::OnKeyPropertyClicked(const IPropertyHandle& KeyedPropertyHandle)
+{
+	TArray<UObject*> Objects;
+	KeyedPropertyHandle.GetOuterObjects( Objects );
+
+	FKeyPropertyParams KeyPropertyParams(Objects, KeyedPropertyHandle);
+	KeyPropertyParams.KeyParams.bCreateHandleIfMissing = true;
+	KeyPropertyParams.KeyParams.bCreateTrackIfMissing = true;
+	KeyPropertyParams.KeyParams.bAddKeyEvenIfUnchanged = true;
+
+	Sequencer.Pin()->KeyProperty(KeyPropertyParams);
+}
diff --git a/Engine/Source/Editor/Sequencer/Private/SequencerDetailKeyframeHandler.h b/Engine/Source/Editor/Sequencer/Private/SequencerDetailKeyframeHandler.h
new file mode 100644
--- /dev/null
+++ b/Engine/Source/Editor/Sequencer/Private/SequencerDetailKeyframeHandler.h
@@ -0,0 +1,20 @@
+// Copyright 1998-2015 Epic Games, Inc. All Rights Reserved.
+
+#pragma once
+
+#include "Editor/PropertyEditor/Public/IDetailKeyframeHandler.h"
+
+class FSequencerDetailKeyframeHandler : public IDetailKeyframeHandler
+{
+public:
+	FSequencerDetailKeyframeHandler(TSharedRef<ISequencer> InSequencer);
+
+	virtual bool IsPropertyKeyable(UClass* InObjectClass, const class IPropertyHandle& PropertyHandle) const override;
+
+	virtual bool IsPropertyKeyingEnabled() const override;
+
+	virtual void OnKeyPropertyClicked(const IPropertyHandle& KeyedPropertyHandle) override;
+
+private:
+	TWeakPtr<class ISequencer> Sequencer;
+};
\ No newline at end of file
diff --git a/Engine/Source/Editor/Sequencer/Private/SequencerModule.cpp b/Engine/Source/Editor/Sequencer/Private/SequencerModule.cpp
--- a/Engine/Source/Editor/Sequencer/Private/SequencerModule.cpp
+++ b/Engine/Source/Editor/Sequencer/Private/SequencerModule.cpp
@@ -8,6 +8,7 @@
 #include "SequencerCommands.h"
 #include "SequencerEdMode.h"
 #include "SequencerObjectChangeListener.h"
+#include "SequencerDetailKeyframeHandler.h"
 
 
 // We disable the deprecation warnings here because otherwise it'll complain about us
@@ -34,8 +35,9 @@ public:
 	{
 		TSharedRef<FSequencer> Sequencer = MakeShareable(new FSequencer);
 		TSharedRef<ISequencerObjectChangeListener> ObjectChangeListener = MakeShareable(new FSequencerObjectChangeListener(Sequencer, InitParams.bEditWithinLevelEditor));
+		TSharedRef<IDetailKeyframeHandler> KeyframeHandler = MakeShareable( new FSequencerDetailKeyframeHandler( Sequencer ));
 
-		Sequencer->InitSequencer(InitParams, ObjectChangeListener, TrackEditorDelegates);
+		Sequencer->InitSequencer(InitParams, ObjectChangeListener, KeyframeHandler, TrackEditorDelegates);
 
 		return Sequencer;
 	}
