Sequencer: Separate out struct properties to submenus.

diff --git a/Engine/Source/Editor/Sequencer/Private/SequencerDisplayNode.cpp b/Engine/Source/Editor/Sequencer/Private/SequencerDisplayNode.cpp
--- a/Engine/Source/Editor/Sequencer/Private/SequencerDisplayNode.cpp
+++ b/Engine/Source/Editor/Sequencer/Private/SequencerDisplayNode.cpp
@@ -771,15 +771,70 @@ TSharedRef<SWidget> FObjectBindingNode::OnGetAddTrackMenuContent()
 
 	AddTrackMenuBuilder.BeginSection( SequencerMenuExtensionPoints::AddTrackMenu_PropertiesSection, LOCTEXT("PropertiesMenuHeader" , "Properties"));
 	{
+		// Gather list of property path chains, ie:
+		//   [Aspect Ratio]
+		//   [Postprocess Settings] [ColorGamma]
+		//   [Postprocess Settings] [Gain]
+		//   [Focal Length]
+		TArray< TArray<FString> > PropertyChainPaths;
 		for ( TArray<UProperty*>& KeyablePropertyPath : KeyablePropertyPaths )
 		{
-			FUIAction AddTrackMenuAction( FExecuteAction::CreateSP( this, &FObjectBindingNode::AddTrackForProperty, KeyablePropertyPath ) );
 			TArray<FString> PropertyNames;
 			for ( UProperty* Property : KeyablePropertyPath )
 			{
 				PropertyNames.Add( Property->GetName() );
 			}
-			AddTrackMenuBuilder.AddMenuEntry( FText::FromString( FString::Join( PropertyNames, TEXT( "." ) ) ), FText(), FSlateIcon(), AddTrackMenuAction );
+
+			PropertyChainPaths.Add(PropertyNames);
+		}
+
+		//@todo this currently only supports two level deep property chains. it should be arbitrary.
+		for (int32 PropertyPathIndex = 0; PropertyPathIndex < PropertyChainPaths.Num(); ++PropertyPathIndex )
+		{ 
+			TArray<TArray<FString> > SubPropertyChainPaths;
+			TArray<TArray<UProperty*> > SubKeyablePropertyPaths;
+
+			TArray<FString> PropertyChainPath = PropertyChainPaths[PropertyPathIndex];
+			SubPropertyChainPaths.Add(PropertyChainPath);
+			SubKeyablePropertyPaths.Add(KeyablePropertyPaths[PropertyPathIndex]);
+
+			// Look for the next chain that is not the same length
+			int32 NextPropertyPathIndex = PropertyPathIndex + 1; 
+			for (; NextPropertyPathIndex < PropertyChainPaths.Num(); ++NextPropertyPathIndex)
+			{
+				if (PropertyChainPaths[PropertyPathIndex].Num() != PropertyChainPaths[NextPropertyPathIndex].Num())
+				{
+					break;
+				}
+
+				TArray<FString> NextPropertyChainPath = PropertyChainPaths[NextPropertyPathIndex];
+				SubPropertyChainPaths.Add(NextPropertyChainPath);
+				SubKeyablePropertyPaths.Add(KeyablePropertyPaths[NextPropertyPathIndex]);
+			}
+			
+			// Don't use a submenu if there's only one menu item
+			if (PropertyChainPaths[PropertyPathIndex].Num() > 1 && SubPropertyChainPaths.Num() > 1)
+			{
+				// Strip off the first member of the chain since it's the submenu title
+				for (int32 SubPropertyChainIndex = 0; SubPropertyChainIndex < SubPropertyChainPaths.Num(); ++SubPropertyChainIndex)
+				{
+					SubPropertyChainPaths[SubPropertyChainIndex].RemoveAt(0);
+				}
+
+				int32 ChainLength = PropertyChainPaths[PropertyPathIndex].Num();
+				const FText& ParentDisplayName = FText::FromString(PropertyChainPaths[PropertyPathIndex][0]);
+				AddTrackMenuBuilder.AddSubMenu(
+					ParentDisplayName, 
+					FText::GetEmpty(), 
+					FNewMenuDelegate::CreateSP(this, &FObjectBindingNode::AddPropertyMenuItems, SubPropertyChainPaths, SubKeyablePropertyPaths) 
+					);
+			}
+			else
+			{
+				AddPropertyMenuItems(AddTrackMenuBuilder, SubPropertyChainPaths, SubKeyablePropertyPaths);
+			}
+
+			PropertyPathIndex = NextPropertyPathIndex;
 		}
 	}
 	AddTrackMenuBuilder.EndSection();
@@ -787,6 +842,15 @@ TSharedRef<SWidget> FObjectBindingNode::OnGetAddTrackMenuContent()
 	return AddTrackMenuBuilder.MakeWidget();
 }
 
+void FObjectBindingNode::AddPropertyMenuItems(FMenuBuilder& AddTrackMenuBuilder, TArray<TArray<FString> > PropertyChainPaths, TArray<TArray<UProperty*> > KeyablePropertyPaths)
+{
+	for (int32 PropertyPathIndex = 0; PropertyPathIndex < PropertyChainPaths.Num(); ++PropertyPathIndex)
+	{
+		FUIAction AddTrackMenuAction( FExecuteAction::CreateSP( this, &FObjectBindingNode::AddTrackForProperty, KeyablePropertyPaths[PropertyPathIndex] ) );
+		AddTrackMenuBuilder.AddMenuEntry( FText::FromString( FString::Join( PropertyChainPaths[PropertyPathIndex], TEXT( "." ) ) ), FText(), FSlateIcon(), AddTrackMenuAction );
+	}
+}
+
 void FObjectBindingNode::AddTrackForProperty(TArray<UProperty*> PropertyPath)
 {
 	FSequencer& Sequencer = GetSequencer();
diff --git a/Engine/Source/Editor/Sequencer/Private/SequencerDisplayNode.h b/Engine/Source/Editor/Sequencer/Private/SequencerDisplayNode.h
--- a/Engine/Source/Editor/Sequencer/Private/SequencerDisplayNode.h
+++ b/Engine/Source/Editor/Sequencer/Private/SequencerDisplayNode.h
@@ -485,6 +485,8 @@ public:
 private:
 	TSharedRef<SWidget> OnGetAddTrackMenuContent();
 
+	void AddPropertyMenuItems(FMenuBuilder& AddTrackMenuBuilder, TArray<TArray<FString> > PropertyChainPaths, TArray<TArray<UProperty*> > KeyablePropertyPaths);
+
 	void AddTrackForProperty(TArray<UProperty*> PropertyPath);
 
 	/** Get class for object binding */
