Sequencer: Clamp scrub time to view range, not outer range.

diff --git a/Engine/Source/Editor/Sequencer/Private/Sequencer.cpp b/Engine/Source/Editor/Sequencer/Private/Sequencer.cpp
--- a/Engine/Source/Editor/Sequencer/Private/Sequencer.cpp
+++ b/Engine/Source/Editor/Sequencer/Private/Sequencer.cpp
@@ -638,7 +638,7 @@ void FSequencer::SetGlobalTimeDirectly( float NewTime )
 	if (Settings->GetShowRangeSlider() && Settings->GetLockInOutToStartEndRange())
 	{
 		UMovieScene* FocusedMovieScene = GetFocusedMovieSceneSequence()->GetMovieScene();
-		NewTime = FMath::Clamp(NewTime, FocusedMovieScene->StartTime, FocusedMovieScene->EndTime);
+		NewTime = FMath::Clamp(NewTime, FocusedMovieScene->InTime, FocusedMovieScene->OutTime);
 	}
 
 	// Update the position
