Sequencer: Lock the view range to the window dimensions if the view range is locked.

diff --git a/Engine/Source/Editor/Sequencer/Private/SSequencer.cpp b/Engine/Source/Editor/Sequencer/Private/SSequencer.cpp
--- a/Engine/Source/Editor/Sequencer/Private/SSequencer.cpp
+++ b/Engine/Source/Editor/Sequencer/Private/SSequencer.cpp
@@ -234,7 +234,8 @@ void SSequencer::Construct( const FArguments& InArgs, TSharedRef< class FSequenc
 	SAssignNew( TrackOutliner, SSequencerTrackOutliner );
 
 	SAssignNew( TrackArea, SSequencerTrackArea, TimeSliderController, SharedThis(this) )
-		.Visibility( this, &SSequencer::GetTrackAreaVisibility );
+		.Visibility( this, &SSequencer::GetTrackAreaVisibility )
+		.LockInOutToStartEndRange(this, &SSequencer::GetLockInOutToStartEndRange );
 	
 	SAssignNew( TreeView, SSequencerTreeView, SequencerNodeTree.ToSharedRef(), TrackArea.ToSharedRef() )
 	.ExternalScrollbar( ScrollBar );
@@ -1457,6 +1458,11 @@ EVisibility SSequencer::GetCurveEditorVisibility() const
 	return Settings->GetShowCurveEditor() ? EVisibility::Visible : EVisibility::Collapsed;
 }
 
+bool SSequencer::GetLockInOutToStartEndRange() const
+{
+	return Settings->GetLockInOutToStartEndRange();
+}
+
 void SSequencer::OnCurveEditorVisibilityChanged()
 {
 	if (CurveEditor.IsValid())
diff --git a/Engine/Source/Editor/Sequencer/Private/SSequencer.h b/Engine/Source/Editor/Sequencer/Private/SSequencer.h
--- a/Engine/Source/Editor/Sequencer/Private/SSequencer.h
+++ b/Engine/Source/Editor/Sequencer/Private/SSequencer.h
@@ -203,6 +203,9 @@ private:
 	/** Get the visibility of the track area */
 	EVisibility GetTrackAreaVisibility() const;
 
+	/** Get the lock in/out to start/end range value */
+	bool GetLockInOutToStartEndRange() const;
+
 	/** SWidget interface */
 	/** @todo Sequencer Basic drag and drop support.  Doesn't belong here most likely */
 	virtual void OnDragEnter( const FGeometry& MyGeometry, const FDragDropEvent& DragDropEvent ) override;
diff --git a/Engine/Source/Editor/Sequencer/Private/SSequencerTrackArea.cpp b/Engine/Source/Editor/Sequencer/Private/SSequencerTrackArea.cpp
--- a/Engine/Source/Editor/Sequencer/Private/SSequencerTrackArea.cpp
+++ b/Engine/Source/Editor/Sequencer/Private/SSequencerTrackArea.cpp
@@ -16,6 +16,8 @@ void SSequencerTrackArea::Construct( const FArguments& InArgs, TSharedRef<FSeque
 	SequencerWidget = InSequencerWidget;
 	TimeSliderController = InTimeSliderController;
 
+	bLockInOutToStartEndRange = InArgs._LockInOutToStartEndRange;
+
 	//@todo Overlays don't tick by default, but we need to. Maybe try using Composition? 
 	bCanTick = true;
 
@@ -178,17 +180,20 @@ void SSequencerTrackArea::Tick( const FGeometry& AllottedGeometry, const double
 
 	FVector2D Size = AllottedGeometry.GetLocalSize();
 
-	if (SizeLastFrame.IsSet() && Size.X != SizeLastFrame->X)
+	if (!bLockInOutToStartEndRange.Get())
 	{
-		// Zoom by the difference in horizontal size
-		const float Difference = Size.X - SizeLastFrame->X;
-		TRange<float> OldRange = TimeSliderController->GetViewRange().GetAnimationTarget();
-
-		TimeSliderController->SetViewRange(
-			OldRange.GetLowerBoundValue(),
-			OldRange.GetUpperBoundValue() + (Difference * OldRange.Size<float>() / SizeLastFrame->X),
-			EViewRangeInterpolation::Immediate
-			);
+		if (SizeLastFrame.IsSet() && Size.X != SizeLastFrame->X)
+		{
+			// Zoom by the difference in horizontal size
+			const float Difference = Size.X - SizeLastFrame->X;
+			TRange<float> OldRange = TimeSliderController->GetViewRange().GetAnimationTarget();
+
+			TimeSliderController->SetViewRange(
+				OldRange.GetLowerBoundValue(),
+				OldRange.GetUpperBoundValue() + (Difference * OldRange.Size<float>() / SizeLastFrame->X),
+				EViewRangeInterpolation::Immediate
+				);
+		}
 	}
 
 	SizeLastFrame = Size;
diff --git a/Engine/Source/Editor/Sequencer/Private/SSequencerTrackArea.h b/Engine/Source/Editor/Sequencer/Private/SSequencerTrackArea.h
--- a/Engine/Source/Editor/Sequencer/Private/SSequencerTrackArea.h
+++ b/Engine/Source/Editor/Sequencer/Private/SSequencerTrackArea.h
@@ -15,7 +15,10 @@ class FSequencer;
 class SSequencerTrackArea : public SOverlay
 {
 public:
-	SLATE_BEGIN_ARGS( SSequencerTrackArea ){}
+	SLATE_BEGIN_ARGS( SSequencerTrackArea )
+	{}		
+		/** Lock the in/out to the start/end range */
+		SLATE_ATTRIBUTE( bool, LockInOutToStartEndRange )
 	SLATE_END_ARGS()
 
 	/** Construct this widget */
@@ -62,4 +65,7 @@ private:
 
 	/** Keep a hold of the size of the area so we can maintain zoom levels */
 	TOptional<FVector2D> SizeLastFrame;
+
+	/** Whether to lock the in/out to the start/end range */
+	TAttribute<bool> bLockInOutToStartEndRange;
 };
