Sequencer: Use display name text instead of property name.

diff --git a/Engine/Source/Editor/Sequencer/Private/SequencerDisplayNode.cpp b/Engine/Source/Editor/Sequencer/Private/SequencerDisplayNode.cpp
--- a/Engine/Source/Editor/Sequencer/Private/SequencerDisplayNode.cpp
+++ b/Engine/Source/Editor/Sequencer/Private/SequencerDisplayNode.cpp
@@ -750,7 +750,7 @@ TSharedRef<SWidget> FObjectBindingNode::OnGetAddTrackMenuContent()
 	{
 		for (int32 IndexToCheck = 0; IndexToCheck < A.Num() && IndexToCheck < B.Num(); IndexToCheck++)
 		{
-			int32 CompareResult = A[IndexToCheck]->GetName().Compare(B[IndexToCheck]->GetName());
+			int32 CompareResult = A[IndexToCheck]->GetDisplayNameText().ToString().Compare(B[IndexToCheck]->GetDisplayNameText().ToString());
 			if (CompareResult != 0)
 			{
 				return CompareResult < 0;
@@ -782,7 +782,7 @@ TSharedRef<SWidget> FObjectBindingNode::OnGetAddTrackMenuContent()
 			TArray<FString> PropertyNames;
 			for ( UProperty* Property : KeyablePropertyPath )
 			{
-				PropertyNames.Add( Property->GetName() );
+				PropertyNames.Add( Property->GetDisplayNameText().ToString() );
 			}
 
 			PropertyChainPaths.Add(PropertyNames);
