Sequencer: Disregard deprecated properties as a keyable property.

diff --git a/Engine/Source/Editor/Sequencer/Private/SequencerDisplayNode.cpp b/Engine/Source/Editor/Sequencer/Private/SequencerDisplayNode.cpp
--- a/Engine/Source/Editor/Sequencer/Private/SequencerDisplayNode.cpp
+++ b/Engine/Source/Editor/Sequencer/Private/SequencerDisplayNode.cpp
@@ -717,20 +717,24 @@ void GetKeyablePropertyPaths(UClass* Class, UStruct* PropertySource, TArray<UPro
 	for (TFieldIterator<UProperty> PropertyIterator(PropertySource); PropertyIterator; ++PropertyIterator)
 	{
 		UProperty* Property = *PropertyIterator;
-		PropertyPath.Add(Property);
-		if (Sequencer.CanKeyProperty(FCanKeyPropertyParams(Class, PropertyPath)))
-		{
-			KeyablePropertyPaths.Add(PropertyPath);
-		}
-		else
+
+		if (Property && !Property->HasAnyPropertyFlags(CPF_Deprecated))
 		{
-			UStructProperty* StructProperty = Cast<UStructProperty>(Property);
-			if (StructProperty != nullptr)
+			PropertyPath.Add(Property);
+			if (Sequencer.CanKeyProperty(FCanKeyPropertyParams(Class, PropertyPath)))
 			{
-				GetKeyablePropertyPaths(Class, StructProperty->Struct, PropertyPath, Sequencer, KeyablePropertyPaths);
+				KeyablePropertyPaths.Add(PropertyPath);
+			}
+			else
+			{
+				UStructProperty* StructProperty = Cast<UStructProperty>(Property);
+				if (StructProperty != nullptr)
+				{
+					GetKeyablePropertyPaths(Class, StructProperty->Struct, PropertyPath, Sequencer, KeyablePropertyPaths);
+				}
 			}
+			PropertyPath.RemoveAt(PropertyPath.Num() - 1);
 		}
-		PropertyPath.RemoveAt(PropertyPath.Num() - 1);
 	}
 }
 
