Sequencer: Draw the decimal frame when not on a frame.

diff --git a/Engine/Source/Editor/Sequencer/Private/TimeSliderController.cpp b/Engine/Source/Editor/Sequencer/Private/TimeSliderController.cpp
--- a/Engine/Source/Editor/Sequencer/Private/TimeSliderController.cpp
+++ b/Engine/Source/Editor/Sequencer/Private/TimeSliderController.cpp
@@ -303,7 +303,18 @@ int32 FSequencerTimeSliderController::OnPaintTimeSlider( bool bMirrorLabels, con
 		FString FrameString;
 		if (SequencerSnapValues::IsTimeSnapIntervalFrameRate(TimeSliderArgs.Settings->GetTimeSnapInterval()) && TimeSliderArgs.Settings->GetShowFrameNumbers())
 		{
-			FrameString = FString::Printf( TEXT("%d"), TimeToFrame(Time));
+			float FrameRate = 1.0f/TimeSliderArgs.Settings->GetTimeSnapInterval();
+			float FrameTime = Time * FrameRate;
+			int32 Frame = SequencerHelpers::TimeToFrame(Time, FrameRate);
+					
+			if (FMath::IsNearlyEqual(FrameTime, (float)Frame, KINDA_SMALL_NUMBER))
+			{
+				FrameString = FString::Printf( TEXT("%d"), TimeToFrame(Time));
+			}
+			else
+			{
+				FrameString = FString::Printf( TEXT("%.3f"), FrameTime);
+			}
 		}
 		else
 		{
