Sequencer: Right click menu for keys: Interpolation Modes, Extrapolation, Snap To Frame, Select All Keys (in Section).

diff --git a/Engine/Source/Editor/MovieSceneTools/Public/MovieSceneToolHelpers.h b/Engine/Source/Editor/MovieSceneTools/Public/MovieSceneToolHelpers.h
--- a/Engine/Source/Editor/MovieSceneTools/Public/MovieSceneToolHelpers.h
+++ b/Engine/Source/Editor/MovieSceneTools/Public/MovieSceneToolHelpers.h
@@ -28,12 +28,17 @@ public:
 		return OutKeyHandles;
 	}
 
-	virtual float GetKeyTime( FKeyHandle KeyHandle ) const override
+	virtual void SetKeyTime(FKeyHandle KeyHandle, float NewKeyTime) const override
 	{
-		return Curve->GetKeyTime( KeyHandle );
+		Curve->SetKeyTime(KeyHandle, NewKeyTime);
 	}
 
-	virtual FKeyHandle MoveKey( FKeyHandle KeyHandle, float DeltaPosition ) override
+	virtual float GetKeyTime(FKeyHandle KeyHandle) const override
+	{
+		return Curve->GetKeyTime(KeyHandle);
+	}
+
+	virtual FKeyHandle MoveKey(FKeyHandle KeyHandle, float DeltaPosition) override
 	{
 		return Curve->SetKeyTime( KeyHandle, Curve->GetKeyTime( KeyHandle ) + DeltaPosition );
 	}
@@ -46,6 +51,65 @@ public:
 		}
 	}
 
+	virtual void SetKeyInterpMode(FKeyHandle KeyHandle, ERichCurveInterpMode InterpMode) override
+	{
+		if (Curve->IsKeyHandleValid(KeyHandle))
+		{
+			Curve->SetKeyInterpMode(KeyHandle, InterpMode);
+		}
+	}
+
+	virtual ERichCurveInterpMode GetKeyInterpMode(FKeyHandle KeyHandle) const override
+	{
+		if (Curve->IsKeyHandleValid(KeyHandle))
+		{
+			return Curve->GetKeyInterpMode(KeyHandle);
+		}
+		return RCIM_None;
+	}
+
+	virtual void SetKeyTangentMode(FKeyHandle KeyHandle, ERichCurveTangentMode TangentMode) override
+	{
+		if (Curve->IsKeyHandleValid(KeyHandle))
+		{
+			Curve->SetKeyTangentMode(KeyHandle, TangentMode);
+		}
+	}
+
+	virtual ERichCurveTangentMode GetKeyTangentMode(FKeyHandle KeyHandle) const override
+	{
+		if (Curve->IsKeyHandleValid(KeyHandle))
+		{
+			return Curve->GetKeyTangentMode(KeyHandle);
+		}
+		return RCTM_None;
+	}
+
+	virtual void SetExtrapolationMode(ERichCurveExtrapolation ExtrapMode, bool bPreInfinity) override
+	{
+		if (bPreInfinity)
+		{
+			Curve->PreInfinityExtrap = ExtrapMode;
+		}
+		else
+		{
+			Curve->PostInfinityExtrap = ExtrapMode;
+		}
+	}
+
+	virtual ERichCurveExtrapolation GetExtrapolationMode(bool bPreInfinity) const override
+	{
+		if (bPreInfinity)
+		{
+			return Curve->PreInfinityExtrap;
+		}
+		else
+		{
+			return Curve->PostInfinityExtrap;
+		}
+		return RCCE_None;
+	}
+
 	virtual void AddKeyUnique(float Time) override;
 
 	virtual FRichCurve* GetRichCurve() override { return Curve; }
@@ -88,12 +152,17 @@ public:
 		return OutKeyHandles;
 	}
 
-	virtual float GetKeyTime( FKeyHandle KeyHandle ) const override
+	virtual void SetKeyTime(FKeyHandle KeyHandle, float NewKeyTime) const override
 	{
-		return Curve.GetKeyTime( KeyHandle );
+		Curve.SetKeyTime(KeyHandle, NewKeyTime);
 	}
 
-	virtual FKeyHandle MoveKey( FKeyHandle KeyHandle, float DeltaPosition ) override
+	virtual float GetKeyTime(FKeyHandle KeyHandle) const override
+	{
+		return Curve.GetKeyTime(KeyHandle);
+	}
+
+	virtual FKeyHandle MoveKey(FKeyHandle KeyHandle, float DeltaPosition) override
 	{
 		return Curve.SetKeyTime( KeyHandle, Curve.GetKeyTime( KeyHandle ) + DeltaPosition );
 	}
@@ -103,6 +172,33 @@ public:
 		Curve.DeleteKey(KeyHandle);
 	}
 
+	virtual void SetKeyInterpMode(FKeyHandle KeyHandle, ERichCurveInterpMode InterpMode) override
+	{
+	}
+
+	virtual ERichCurveInterpMode GetKeyInterpMode(FKeyHandle KeyHandle) const override
+	{
+		return RCIM_None;
+	}
+
+	virtual void SetKeyTangentMode(FKeyHandle KeyHandle, ERichCurveTangentMode TangentMode) override
+	{
+	}
+
+	virtual ERichCurveTangentMode GetKeyTangentMode(FKeyHandle KeyHandle) const override
+	{
+		return RCTM_None;
+	}
+
+	virtual void SetExtrapolationMode(ERichCurveExtrapolation ExtrapMode, bool bPreInfinity) override
+	{
+	}
+
+	virtual ERichCurveExtrapolation GetExtrapolationMode(bool bPreInfinity) const override
+	{
+		return RCCE_None;
+	}
+
 	virtual void AddKeyUnique(float Time) override;
 
 	virtual FRichCurve* GetRichCurve() override { return nullptr; };
diff --git a/Engine/Source/Editor/Sequencer/Private/GroupedKeyArea.cpp b/Engine/Source/Editor/Sequencer/Private/GroupedKeyArea.cpp
--- a/Engine/Source/Editor/Sequencer/Private/GroupedKeyArea.cpp
+++ b/Engine/Source/Editor/Sequencer/Private/GroupedKeyArea.cpp
@@ -193,13 +193,23 @@ TArray<FKeyHandle> FGroupedKeyArea::GetUnsortedKeyHandles() const
 	return Array;
 }
 
-float FGroupedKeyArea::GetKeyTime( FKeyHandle KeyHandle ) const
+void FGroupedKeyArea::SetKeyTime(FKeyHandle KeyHandle, float NewKeyTime) const
+{
+	auto* Group = FindGroup(KeyHandle);
+
+	for (auto& Key : Group->Keys)
+	{
+		KeyAreas[Key.AreaIndex]->SetKeyTime(Key.KeyHandle, NewKeyTime);
+	}
+}
+
+float FGroupedKeyArea::GetKeyTime(FKeyHandle KeyHandle) const
 {
 	auto* Group = FindGroup(KeyHandle);
 	return Group ? Group->RepresentativeTime : 0.f;
 }
 
-FKeyHandle FGroupedKeyArea::MoveKey( FKeyHandle KeyHandle, float DeltaPosition )
+FKeyHandle FGroupedKeyArea::MoveKey(FKeyHandle KeyHandle, float DeltaPosition)
 {
 	int32* GroupIndex = nullptr;
 
@@ -245,6 +255,99 @@ void FGroupedKeyArea::DeleteKey(FKeyHandle KeyHandle)
 	}
 }
 
+void                    
+FGroupedKeyArea::SetKeyInterpMode(FKeyHandle KeyHandle, ERichCurveInterpMode InterpMode)
+{
+	if (auto* Group = FindGroup(KeyHandle))
+	{
+		for (auto& Key : Group->Keys)
+		{
+			KeyAreas[Key.AreaIndex]->SetKeyInterpMode(Key.KeyHandle, InterpMode);
+		}
+	}
+}
+
+ERichCurveInterpMode
+FGroupedKeyArea::GetKeyInterpMode(FKeyHandle KeyHandle) const
+{
+	// Return RCIM_None if the keys don't all have the same interp mode
+	ERichCurveInterpMode InterpMode = RCIM_None;
+	if (auto* Group = FindGroup(KeyHandle))
+	{
+		for (auto& Key : Group->Keys)
+		{
+			if (InterpMode == RCIM_None)
+			{
+				InterpMode = KeyAreas[Key.AreaIndex]->GetKeyInterpMode(Key.KeyHandle);
+			}
+			else if (InterpMode != KeyAreas[Key.AreaIndex]->GetKeyInterpMode(Key.KeyHandle))
+			{
+				return RCIM_None;
+			}
+		}
+	}
+	return InterpMode;
+}
+
+void                    
+FGroupedKeyArea::SetKeyTangentMode(FKeyHandle KeyHandle, ERichCurveTangentMode TangentMode)
+{
+	if (auto* Group = FindGroup(KeyHandle))
+	{
+		for (auto& Key : Group->Keys)
+		{
+			KeyAreas[Key.AreaIndex]->SetKeyTangentMode(Key.KeyHandle, TangentMode);
+		}
+	}
+}
+
+ERichCurveTangentMode
+FGroupedKeyArea::GetKeyTangentMode(FKeyHandle KeyHandle) const
+{
+	// Return RCTM_None if the keys don't all have the same tangent mode
+	ERichCurveTangentMode TangentMode = RCTM_None;
+	if (auto* Group = FindGroup(KeyHandle))
+	{
+		for (auto& Key : Group->Keys)
+		{
+			if (TangentMode == RCTM_None)
+			{
+				TangentMode = KeyAreas[Key.AreaIndex]->GetKeyTangentMode(Key.KeyHandle);
+			}
+			else if (TangentMode != KeyAreas[Key.AreaIndex]->GetKeyTangentMode(Key.KeyHandle))
+			{
+				return RCTM_None;
+			}
+		}
+	}
+	return TangentMode;
+}
+
+void FGroupedKeyArea::SetExtrapolationMode(ERichCurveExtrapolation ExtrapMode, bool bPreInfinity)
+{
+	for (auto& Area : KeyAreas)
+	{
+		Area->SetExtrapolationMode(ExtrapMode, bPreInfinity);
+	}
+}
+
+ERichCurveExtrapolation FGroupedKeyArea::GetExtrapolationMode(bool bPreInfinity) const
+{
+	ERichCurveExtrapolation ExtrapMode = RCCE_None;
+	for (auto& Area : KeyAreas)
+	{
+		if (ExtrapMode == RCCE_None)
+		{
+			ExtrapMode = Area->GetExtrapolationMode(bPreInfinity);
+		}
+		else if (Area->GetExtrapolationMode(bPreInfinity) != ExtrapMode)
+		{
+			return RCCE_None;
+		}
+	}
+	return ExtrapMode;
+}
+
 void FGroupedKeyArea::AddKeyUnique(float Time)
 {
 	for (auto& Area : KeyAreas)
diff --git a/Engine/Source/Editor/Sequencer/Private/GroupedKeyArea.h b/Engine/Source/Editor/Sequencer/Private/GroupedKeyArea.h
--- a/Engine/Source/Editor/Sequencer/Private/GroupedKeyArea.h
+++ b/Engine/Source/Editor/Sequencer/Private/GroupedKeyArea.h
@@ -71,9 +71,16 @@ public:
 	
 	/* IKeyArea interface */
 	virtual TArray<FKeyHandle>		GetUnsortedKeyHandles() const override;
-	virtual float					GetKeyTime( FKeyHandle KeyHandle ) const override;
-	virtual FKeyHandle				MoveKey( FKeyHandle KeyHandle, float DeltaPosition ) override;
+	virtual void 					SetKeyTime(FKeyHandle KeyHandle, float NewKeyTime) const override;
+	virtual float					GetKeyTime(FKeyHandle KeyHandle) const override;
+	virtual FKeyHandle				MoveKey(FKeyHandle KeyHandle, float DeltaPosition) override;
 	virtual void					DeleteKey(FKeyHandle KeyHandle) override;
+	virtual void                    SetKeyInterpMode(FKeyHandle KeyHandle, ERichCurveInterpMode InterpMode) override;
+	virtual ERichCurveInterpMode    GetKeyInterpMode(FKeyHandle KeyHandle) const override;
+	virtual void                    SetKeyTangentMode(FKeyHandle KeyHandle, ERichCurveTangentMode TangentMode) override;
+	virtual ERichCurveTangentMode   GetKeyTangentMode(FKeyHandle KeyHandle) const override;
+	virtual void                    SetExtrapolationMode(ERichCurveExtrapolation ExtrapMode, bool bPreInfinity) override;
+	virtual ERichCurveExtrapolation GetExtrapolationMode(bool bPreInfinity) const override;
 	virtual void					AddKeyUnique(float Time) override;
 	virtual FRichCurve*				GetRichCurve() override;
 	virtual UMovieSceneSection*		GetOwningSection() override;
diff --git a/Engine/Source/Editor/Sequencer/Private/SSection.cpp b/Engine/Source/Editor/Sequencer/Private/SSection.cpp
--- a/Engine/Source/Editor/Sequencer/Private/SSection.cpp
+++ b/Engine/Source/Editor/Sequencer/Private/SSection.cpp
@@ -9,6 +9,7 @@
 #include "MovieSceneShotSection.h"
 #include "CommonMovieSceneTools.h"
 #include "SequencerHotspots.h"
+#include "ScopedTransaction.h"
 
 
 /** When 0, regeneration of dynamic key layouts is enabled, when non-zero, such behaviour is disabled */
@@ -415,42 +416,205 @@ TSharedPtr<SWidget> SSection::OnSummonContextMenu( const FGeometry& MyGeometry,
 
 	if (Key.IsValid() || Sequencer.GetSelection().GetSelectedKeys().Num())
 	{
-		MenuBuilder.AddMenuEntry(
-			NSLOCTEXT("Sequencer", "DeleteKey", "Delete"),
-			NSLOCTEXT("Sequencer", "DeleteKeyToolTip", "Deletes the selected keys."),
-			FSlateIcon(),
-			FUIAction(FExecuteAction::CreateSP(&Sequencer, &FSequencer::DeleteSelectedKeys))
+		MenuBuilder.BeginSection("SequencerInterpolation", NSLOCTEXT("Sequencer", "KeyInterpolationMenu", "Key Interpolation"));
+		{
+			MenuBuilder.AddMenuEntry(
+				NSLOCTEXT("Sequencer", "SetKeyInterpolationAuto", "Auto"),
+				NSLOCTEXT("Sequencer", "SetKeyInterpolationAutoTooltip", "Set key interpolation to auto"),
+				FSlateIcon(),
+				FUIAction(
+					FExecuteAction::CreateSP(&Sequencer, &FSequencer::SetInterpTangentMode, RCIM_Cubic, RCTM_Auto),
+					FCanExecuteAction(),
+					FIsActionChecked::CreateSP(&Sequencer, &FSequencer::IsInterpTangentModeSelected, RCIM_Cubic, RCTM_Auto) ),
+				NAME_None,
+				EUserInterfaceActionType::ToggleButton
+			);
+
+			MenuBuilder.AddMenuEntry(
+				NSLOCTEXT("Sequencer", "SetKeyInterpolationUser", "User"),
+				NSLOCTEXT("Sequencer", "SetKeyInterpolationUserTooltip", "Set key interpolation to user"),
+				FSlateIcon(),
+				FUIAction(
+					FExecuteAction::CreateSP(&Sequencer, &FSequencer::SetInterpTangentMode, RCIM_Cubic, RCTM_User),
+					FCanExecuteAction(),
+					FIsActionChecked::CreateSP(&Sequencer, &FSequencer::IsInterpTangentModeSelected, RCIM_Cubic, RCTM_User) ),
+				NAME_None,
+				EUserInterfaceActionType::ToggleButton
+			);
+
+			MenuBuilder.AddMenuEntry(
+				NSLOCTEXT("Sequencer", "SetKeyInterpolationBreak", "Break"),
+				NSLOCTEXT("Sequencer", "SetKeyInterpolationBreakTooltip", "Set key interpolation to break"),
+				FSlateIcon(),
+				FUIAction(
+					FExecuteAction::CreateSP(&Sequencer, &FSequencer::SetInterpTangentMode, RCIM_Cubic, RCTM_Break),
+					FCanExecuteAction(),
+					FIsActionChecked::CreateSP(&Sequencer, &FSequencer::IsInterpTangentModeSelected, RCIM_Cubic, RCTM_Break) ),
+				NAME_None,
+				EUserInterfaceActionType::ToggleButton
+			);
+
+			MenuBuilder.AddMenuEntry(
+				NSLOCTEXT("Sequencer", "SetKeyInterpolationLinear", "Linear"),
+				NSLOCTEXT("Sequencer", "SetKeyInterpolationLinearTooltip", "Set key interpolation to linear"),
+				FSlateIcon(),
+				FUIAction(
+					FExecuteAction::CreateSP(&Sequencer, &FSequencer::SetInterpTangentMode, RCIM_Linear, RCTM_Auto),
+					FCanExecuteAction(),
+					FIsActionChecked::CreateSP(&Sequencer, &FSequencer::IsInterpTangentModeSelected, RCIM_Linear, RCTM_Auto) ),
+				NAME_None,
+				EUserInterfaceActionType::ToggleButton
+			);
+
+			MenuBuilder.AddMenuEntry(
+				NSLOCTEXT("Sequencer", "SetKeyInterpolationConstant", "Constant"),
+				NSLOCTEXT("Sequencer", "SetKeyInterpolationConstantTooltip", "Set key interpolation to constant"),
+				FSlateIcon(),
+				FUIAction(
+					FExecuteAction::CreateSP(&Sequencer, &FSequencer::SetInterpTangentMode, RCIM_Constant, RCTM_Auto),
+					FCanExecuteAction(),
+					FIsActionChecked::CreateSP(&Sequencer, &FSequencer::IsInterpTangentModeSelected, RCIM_Constant, RCTM_Auto) ),
+				NAME_None,
+				EUserInterfaceActionType::ToggleButton
+			);
+		}
+		MenuBuilder.EndSection(); // SequencerInterpolation
+
+		MenuBuilder.BeginSection("SequencerKeys", NSLOCTEXT("Sequencer", "KeysMenu", "Keys"));
+		{
+			MenuBuilder.AddMenuEntry(
+				NSLOCTEXT("Sequencer", "SnapToFrame", "Snap to Frame"),
+				NSLOCTEXT("Sequencer", "SnapToFrameToolTip", "Snap selected keys to frame"),
+				FSlateIcon(),
+				FUIAction(FExecuteAction::CreateSP(&Sequencer, &FSequencer::SnapToFrame))
+			);
+
+			MenuBuilder.AddMenuEntry(
+				NSLOCTEXT("Sequencer", "DeleteKey", "Delete"),
+				NSLOCTEXT("Sequencer", "DeleteKeyToolTip", "Deletes the selected keys"),
+				FSlateIcon(),
+				FUIAction(FExecuteAction::CreateSP(&Sequencer, &FSequencer::DeleteSelectedKeys))
 			);
+		}
+		MenuBuilder.EndSection(); // SequencerKeys
 	}
-	else
+	
+	if (SceneSection != nullptr)
 	{
 		SectionInterface->BuildSectionContextMenu(MenuBuilder);
 
-		MenuBuilder.AddMenuEntry(
-			NSLOCTEXT("Sequencer", "ToggleSectionActive", "Active"),
-			NSLOCTEXT("Sequencer", "ToggleSectionActiveTooltip", "Toggle section active/inactive."),
-			FSlateIcon(),
-			FUIAction(
-				FExecuteAction::CreateSP(&Sequencer, &FSequencer::ToggleSectionActive),
-				FCanExecuteAction::CreateSP(&Sequencer, &FSequencer::CanToggleSectionActive),
-				FIsActionChecked::CreateSP(&Sequencer, &FSequencer::IsToggleSectionActive)),
-			NAME_None,
-			EUserInterfaceActionType::ToggleButton
+		MenuBuilder.BeginSection("SequencerSections", NSLOCTEXT("Sequencer", "SectionsMenu", "Sections"));
+		{
+			MenuBuilder.AddMenuEntry(
+				NSLOCTEXT("Sequencer", "SelectAllKeys", "Select All Keys"),
+				NSLOCTEXT("Sequencer", "SelectAllKeysTooltip", "Select all keys in section"),
+				FSlateIcon(),
+				FUIAction(FExecuteAction::CreateSP(this, &SSection::SelectAllKeys))
 			);
 
-		// @todo Sequencer this should delete all selected sections
-		// delete/selection needs to be rethought in general
-		MenuBuilder.AddMenuEntry(
-			NSLOCTEXT("Sequencer", "DeleteSection", "Delete"),
-			NSLOCTEXT("Sequencer", "DeleteSectionToolTip", "Deletes this section."),
-			FSlateIcon(),
-			FUIAction(FExecuteAction::CreateSP(&Sequencer, &FSequencer::DeleteSection, SceneSection))
+			MenuBuilder.AddSubMenu(
+				NSLOCTEXT("Sequencer", "SetPreInfinityExtrap", "Pre-Infinity"), 
+				NSLOCTEXT("Sequencer", "SetPreInfinityExtrapTooltip", "Set pre-infinity extrapolation"),
+				FNewMenuDelegate::CreateRaw(this, &SSection::AddExtrapolationMenu, true));
+
+			MenuBuilder.AddSubMenu(
+				NSLOCTEXT("Sequencer", "SetPostInfinityExtrap", "Post-Infinity"), 
+				NSLOCTEXT("Sequencer", "SetPostInfinityExtrapTooltip", "Set post-infinity extrapolation"),
+				FNewMenuDelegate::CreateRaw(this, &SSection::AddExtrapolationMenu, false));
+
+			MenuBuilder.AddMenuEntry(
+				NSLOCTEXT("Sequencer", "ToggleSectionActive", "Active"),
+				NSLOCTEXT("Sequencer", "ToggleSectionActiveTooltip", "Toggle section active/inactive"),
+				FSlateIcon(),
+				FUIAction(
+					FExecuteAction::CreateSP(this, &SSection::ToggleSectionActive),
+					FCanExecuteAction(),
+					FIsActionChecked::CreateSP(this, &SSection::IsToggleSectionActive)),
+				NAME_None,
+				EUserInterfaceActionType::ToggleButton
 			);
+
+			// @todo Sequencer this should delete all selected sections
+			// delete/selection needs to be rethought in general
+			MenuBuilder.AddMenuEntry(
+				NSLOCTEXT("Sequencer", "DeleteSection", "Delete"),
+				NSLOCTEXT("Sequencer", "DeleteSectionToolTip", "Deletes this section"),
+				FSlateIcon(),
+				FUIAction(FExecuteAction::CreateSP(this, &SSection::DeleteSection))
+			);
+		}
+		MenuBuilder.EndSection(); // SequencerSections
 	}
 	
+	ResetState();
+
 	return MenuBuilder.MakeWidget();
 }
 
+void SSection::AddExtrapolationMenu(FMenuBuilder& MenuBuilder, bool bPreInfinity)
+{
+	FSequencer& Sequencer = GetSequencer();
+
+	MenuBuilder.AddMenuEntry(
+		NSLOCTEXT("Sequencer", "SetExtrapCycle", "Cycle"),
+		NSLOCTEXT("Sequencer", "SetExtrapCycleTooltip", "Set extrapolation cycle"),
+		FSlateIcon(),
+		FUIAction(
+			FExecuteAction::CreateSP(this, &SSection::SetExtrapolationMode, RCCE_Cycle, bPreInfinity),
+			FCanExecuteAction(),
+			FIsActionChecked::CreateSP(this, &SSection::IsExtrapolationModeSelected, RCCE_Cycle, bPreInfinity) ),
+		NAME_None,
+		EUserInterfaceActionType::RadioButton
+	);
+
+	MenuBuilder.AddMenuEntry(
+		NSLOCTEXT("Sequencer", "SetExtrapCycleWithOffset", "Cycle with Offset"),
+		NSLOCTEXT("Sequencer", "SetExtrapCycleWithOffsetTooltip", "Set extrapolation cycle with offset"),
+		FSlateIcon(),
+		FUIAction(
+			FExecuteAction::CreateSP(this, &SSection::SetExtrapolationMode, RCCE_CycleWithOffset, bPreInfinity),
+			FCanExecuteAction(),
+			FIsActionChecked::CreateSP(this, &SSection::IsExtrapolationModeSelected, RCCE_CycleWithOffset, bPreInfinity) ),
+		NAME_None,
+		EUserInterfaceActionType::RadioButton
+	);
+
+	MenuBuilder.AddMenuEntry(
+		NSLOCTEXT("Sequencer", "SetExtrapOscillate", "Oscillate"),
+		NSLOCTEXT("Sequencer", "SetExtrapOscillateTooltip", "Set extrapolation oscillate"),
+		FSlateIcon(),
+		FUIAction(
+			FExecuteAction::CreateSP(this, &SSection::SetExtrapolationMode, RCCE_Oscillate, bPreInfinity),
+			FCanExecuteAction(),
+			FIsActionChecked::CreateSP(this, &SSection::IsExtrapolationModeSelected, RCCE_Oscillate, bPreInfinity) ),
+		NAME_None,
+		EUserInterfaceActionType::RadioButton
+	);
+
+	MenuBuilder.AddMenuEntry(
+		NSLOCTEXT("Sequencer", "SetExtrapLinear", "Linear"),
+		NSLOCTEXT("Sequencer", "SetExtrapLinearTooltip", "Set extrapolation linear"),
+		FSlateIcon(),
+		FUIAction(
+			FExecuteAction::CreateSP(this, &SSection::SetExtrapolationMode, RCCE_Linear, bPreInfinity),
+			FCanExecuteAction(),
+			FIsActionChecked::CreateSP(this, &SSection::IsExtrapolationModeSelected, RCCE_Linear, bPreInfinity) ),
+		NAME_None,
+		EUserInterfaceActionType::RadioButton
+	);
+
+	MenuBuilder.AddMenuEntry(
+		NSLOCTEXT("Sequencer", "SetExtrapConstant", "Constant"),
+		NSLOCTEXT("Sequencer", "SetExtrapConstantTooltip", "Set extrapolation constant"),
+		FSlateIcon(),
+		FUIAction(
+			FExecuteAction::CreateSP(this, &SSection::SetExtrapolationMode, RCCE_Constant, bPreInfinity),
+			FCanExecuteAction(),
+			FIsActionChecked::CreateSP(this, &SSection::IsExtrapolationModeSelected, RCCE_Constant, bPreInfinity) ),
+		NAME_None,
+		EUserInterfaceActionType::RadioButton
+	);
+}
 
 void SSection::Tick( const FGeometry& AllottedGeometry, const double InCurrentTime, const float InDeltaTime )
 {	
@@ -565,15 +729,15 @@ FReply SSection::OnMouseButtonDoubleClick( const FGeometry& MyGeometry, const FP
 {
 	ResetState();
 
-	UMovieSceneSection* SectionObject = SectionInterface->GetSectionObject();
-	if( GetSequencer().IsSectionVisible( SectionObject ) )
+	if( MouseEvent.GetEffectingButton() == EKeys::LeftMouseButton )
 	{
-		FReply Reply = SectionInterface->OnSectionDoubleClicked( MyGeometry, MouseEvent );
+		UMovieSceneSection* SectionObject = SectionInterface->GetSectionObject();
+		if( GetSequencer().IsSectionVisible( SectionObject ) )
+		{
+			FReply Reply = SectionInterface->OnSectionDoubleClicked( MyGeometry, MouseEvent );
 
-		if (Reply.IsEventHandled()) {return Reply;}
+			if (Reply.IsEventHandled()) {return Reply;}
 
-		if( MouseEvent.GetEffectingButton() == EKeys::LeftMouseButton )
-		{
 			GetSequencer().ZoomToSelectedSections();
 
 			return FReply::Handled();
@@ -713,3 +877,143 @@ void SSection::HandleSectionSelection( const FPointerEvent& MouseEvent, bool bSe
 	}
 }
 
+void SSection::SelectAllKeys()
+{
+	// @todo Sequencer should operate on selected sections
+	UMovieSceneSection* Section = SectionInterface->GetSectionObject();
+
+	for (const FKeyAreaLayoutElement& Element : Layout->GetElements())
+	{
+		TSharedPtr<IKeyArea> KeyArea = Element.GetKeyArea();
+
+		TArray<FKeyHandle> KeyHandles = KeyArea->GetUnsortedKeyHandles();
+		for( int32 KeyIndex = 0; KeyIndex < KeyHandles.Num(); ++KeyIndex )
+		{
+			FKeyHandle KeyHandle = KeyHandles[KeyIndex];
+			FSelectedKey SelectKey(*Section, KeyArea, KeyHandle);
+			GetSequencer().GetSelection().AddToSelection(SelectKey);
+		}
+	}
+}
+
+void SSection::SetExtrapolationMode(ERichCurveExtrapolation ExtrapMode, bool bPreInfinity)
+{
+	// @todo Sequencer should operate on selected sections
+	UMovieSceneSection* Section = SectionInterface->GetSectionObject();
+
+	FScopedTransaction SetExtrapolationModeTransaction(NSLOCTEXT("Sequencer", "SetExtrapolationMode_Transaction", "Set Extrapolation Mode"));
+	bool bAnythingChanged = false;
+
+	Section->Modify();
+	for (const FKeyAreaLayoutElement& Element : Layout->GetElements())
+	{
+		TSharedPtr<IKeyArea> KeyArea = Element.GetKeyArea();
+		bAnythingChanged = true;
+		KeyArea->SetExtrapolationMode(ExtrapMode, bPreInfinity);
+	}
+
+	if (bAnythingChanged)
+	{
+		GetSequencer().UpdateRuntimeInstances();
+	}
+}
+
+bool SSection::IsExtrapolationModeSelected(ERichCurveExtrapolation ExtrapMode, bool bPreInfinity) const
+{
+	// @todo Sequencer should operate on selected sections
+	bool bAllSelected = false;
+
+	for (const FKeyAreaLayoutElement& Element : Layout->GetElements())
+	{
+		TSharedPtr<IKeyArea> KeyArea = Element.GetKeyArea();
+
+		bAllSelected = true;
+		if (KeyArea->GetExtrapolationMode(bPreInfinity) != ExtrapMode)
+		{
+			bAllSelected = false;
+			break;
+		}
+	}
+
+	return bAllSelected;
+}
+
+void SSection::ToggleSectionActive()
+{
+	FScopedTransaction ToggleSectionActiveTransaction( NSLOCTEXT("Sequencer", "ToggleSectionActive_Transaction", "Toggle Section Active") );
+
+	bool bIsActive = !IsToggleSectionActive();
+
+	bool bAnythingChanged = false;
+	if (GetSequencer().GetSelection().GetSelectedSections().Num() != 0)
+	{
+		for (auto Section : GetSequencer().GetSelection().GetSelectedSections())
+		{
+			bAnythingChanged = true;
+			Section->Modify();
+			Section->SetIsActive(bIsActive);
+		}
+	}
+	else
+	{
+		UMovieSceneSection* Section = SectionInterface->GetSectionObject();
+		if (Section)
+		{
+			bAnythingChanged = true;
+			Section->Modify();
+			Section->SetIsActive(bIsActive);
+		}
+	}
+
+	if (bAnythingChanged)
+	{
+		GetSequencer().UpdateRuntimeInstances();
+	}
+}
+
+bool SSection::IsToggleSectionActive() const
+{
+	// Active only if all are active
+	if (GetSequencer().GetSelection().GetSelectedSections().Num() != 0)
+	{
+		for (auto Section : GetSequencer().GetSelection().GetSelectedSections())
+		{
+			if (!Section->IsActive())
+			{
+				return false;
+			}
+		}
+	}
+	else
+	{
+		UMovieSceneSection* Section = SectionInterface->GetSectionObject();
+		if (Section)
+		{
+			if (!Section->IsActive())
+			{
+				return false;
+			}
+		}
+	}
+
+	return true;
+}
+
+void SSection::DeleteSection()
+{
+	if (GetSequencer().GetSelection().GetSelectedSections().Num() != 0)
+	{
+		GetSequencer().DeleteSections(GetSequencer().GetSelection().GetSelectedSections());
+	}
+	else
+	{
+		UMovieSceneSection* Section = SectionInterface->GetSectionObject();
+		if (Section)
+		{
+			TSet<TWeakObjectPtr<UMovieSceneSection> > Sections;
+			Sections.Add(Section);
+
+			GetSequencer().DeleteSections(Sections);
+		}
+	}
+}
\ No newline at end of file
diff --git a/Engine/Source/Editor/Sequencer/Private/SSection.h b/Engine/Source/Editor/Sequencer/Private/SSection.h
--- a/Engine/Source/Editor/Sequencer/Private/SSection.h
+++ b/Engine/Source/Editor/Sequencer/Private/SSection.h
@@ -83,6 +83,9 @@ private:
 	/** Summons a context menu over the associated section */
 	TSharedPtr<SWidget> OnSummonContextMenu( const FGeometry& MyGeometry, const FPointerEvent& MouseEvent );
 
+	/** Add extrapolation menu for pre and post infinity */
+	void AddExtrapolationMenu(FMenuBuilder& MenuBuilder, bool bPreInfinity);
+
 	/**
 	 * Selects the section or keys if they are found under the mouse
 	 *
@@ -135,6 +138,32 @@ private:
 	 */
 	FGeometry MakeSectionGeometryWithoutHandles( const FGeometry& AllottedGeometry, const TSharedPtr<ISequencerSection>& InSectionInterface ) const;
 
+	/** 
+	 * Select all keys within the section
+	 */
+	void SelectAllKeys();
+
+	/**
+	 * Set extrapolation mode
+	 */
+	void SetExtrapolationMode(ERichCurveExtrapolation ExtrapMode, bool bPreInfinity);
+
+	/**
+	 * Is extrapolation mode selected
+	 */
+	bool IsExtrapolationModeSelected(ERichCurveExtrapolation ExtrapMode, bool bPreInfinity) const;
+
+	/**
+	 * Section active/inactive toggle
+	 */
+	void ToggleSectionActive();
+	bool IsToggleSectionActive() const;
+
+	/*
+	 * Delete section
+	 */ 
+	void DeleteSection();
+
 private:
 	/** Interface to section data */
 	TSharedPtr<ISequencerSection> SectionInterface;
diff --git a/Engine/Source/Editor/Sequencer/Private/Sequencer.cpp b/Engine/Source/Editor/Sequencer/Private/Sequencer.cpp
--- a/Engine/Source/Editor/Sequencer/Private/Sequencer.cpp
+++ b/Engine/Source/Editor/Sequencer/Private/Sequencer.cpp
@@ -392,34 +392,41 @@ void FSequencer::AddAnimation(FGuid ObjectGuid, class UAnimSequence* AnimSequenc
 	}
 }
 
-void FSequencer::DeleteSection(class UMovieSceneSection* Section)
+void FSequencer::DeleteSections(const TSet<TWeakObjectPtr<UMovieSceneSection> >& Sections)
 {
 	UMovieScene* MovieScene = GetFocusedMovieSceneSequence()->GetMovieScene();
 	
 	bool bAnythingRemoved = false;
 
-	UMovieSceneTrack* Track = CastChecked<UMovieSceneTrack>( Section->GetOuter() );
+	FScopedTransaction DeleteSectionTransaction( NSLOCTEXT("Sequencer", "DeleteSection_Transaction", "Delete Section") );
+		
+	for (auto It = Sections.CreateConstIterator(); It; ++It)
+	{
+		UMovieSceneSection* Section = (*It).Get();
+		if (Section != nullptr)
+		{
+			UMovieSceneTrack* Track = CastChecked<UMovieSceneTrack>( Section->GetOuter() );
 
-	// If this check fails then the section is outered to a type that doesnt know about the section
-	//checkSlow( Track->HasSection(Section) );
+			// If this check fails then the section is outered to a type that doesnt know about the section
+			//checkSlow( Track->HasSection(Section) );
 	
-	Track->SetFlags( RF_Transactional );
+			Track->SetFlags( RF_Transactional );
 
-	FScopedTransaction DeleteSectionTransaction( NSLOCTEXT("Sequencer", "DeleteSection_Transaction", "Delete Section") );
-	
-	Track->Modify();
+			Track->Modify();
 
-	Track->RemoveSection(Section);
+			Track->RemoveSection(Section);
+
+			bAnythingRemoved = true;
+
+			Selection.RemoveFromSelection(Section);
+		}
+	}
 
-	bAnythingRemoved = true;
-	
 	if( bAnythingRemoved )
 	{
 		// Full refresh required just in case the last section was removed from any track.
 		NotifyMovieSceneDataChanged();
 	}
-
-	Selection.RemoveFromSelection(Section);
 }
 
 void FSequencer::DeleteSelectedKeys()
@@ -446,38 +453,80 @@ void FSequencer::DeleteSelectedKeys()
 	Selection.EmptySelectedKeys();
 }
 
-void FSequencer::ToggleSectionActive()
+void FSequencer::SetInterpTangentMode(ERichCurveInterpMode InterpMode, ERichCurveTangentMode TangentMode)
 {
-	FScopedTransaction ToggleSectionActiveTransaction( NSLOCTEXT("Sequencer", "ToggleSectionActive_Transaction", "Toggle Section Active") );
+	FScopedTransaction SetInterpTangentModeTransaction(NSLOCTEXT("Sequencer", "SetInterpTangentMode_Transaction", "Set Interpolation and Tangent Mode"));
+	bool bAnythingChanged = false;
+	TArray<FSelectedKey> SelectedKeysArray = Selection.GetSelectedKeys().Array();
 
-	bool bIsActive = !IsToggleSectionActive();
+	for (const FSelectedKey& Key : SelectedKeysArray)
+	{
+		if (Key.IsValid())
+		{
+			Key.Section->Modify();
+			Key.KeyArea->SetKeyInterpMode(Key.KeyHandle.GetValue(), InterpMode);
+			Key.KeyArea->SetKeyTangentMode(Key.KeyHandle.GetValue(), TangentMode);
+			bAnythingChanged = true;
+		}
+	}
 
-	for (auto Section : Selection.GetSelectedSections())
+	if (bAnythingChanged)
 	{
-		Section->Modify();
-		Section->SetIsActive(bIsActive);
+		UpdateRuntimeInstances();
 	}
 }
 
-bool FSequencer::CanToggleSectionActive() const
+bool FSequencer::IsInterpTangentModeSelected(ERichCurveInterpMode InterpMode, ERichCurveTangentMode TangentMode) const
 {
-	return Selection.GetSelectedSections().Num() > 0;
+	TArray<FSelectedKey> SelectedKeysArray = Selection.GetSelectedKeys().Array();
+
+	bool bAllSelected = false;
+	for (const FSelectedKey& Key : SelectedKeysArray)
+	{
+		if (Key.IsValid())
+		{
+			bAllSelected = true;
+			if (Key.KeyArea->GetKeyInterpMode(Key.KeyHandle.GetValue()) != InterpMode || 
+				Key.KeyArea->GetKeyTangentMode(Key.KeyHandle.GetValue()) != TangentMode)
+			{
+				bAllSelected = false;
+				break;
+			}
+		}
+	}
+	return bAllSelected;
 }
 
-bool FSequencer::IsToggleSectionActive() const
+void FSequencer::SnapToFrame()
 {
-	// Active only if all are active
-	for (auto Section : Selection.GetSelectedSections())
+	FScopedTransaction SnapToFrameTransaction(NSLOCTEXT("Sequencer", "SnapToFrame_Transaction", "Snap Selected Keys to Frame"));
+	bool bAnythingChanged = false;
+	TArray<FSelectedKey> SelectedKeysArray = Selection.GetSelectedKeys().Array();
+
+	for (const FSelectedKey& Key : SelectedKeysArray)
 	{
-		if (!Section->IsActive())
+		if (Key.IsValid())
 		{
-			return false;
+			Key.Section->Modify();
+			float NewKeyTime = Key.KeyArea->GetKeyTime(Key.KeyHandle.GetValue());
+
+			// Convert to frame
+			float FrameRate = 1.0f / Settings->GetTimeSnapInterval();
+			int32 NewFrame = SequencerHelpers::TimeToFrame(NewKeyTime, FrameRate);
+
+			// Convert back to time
+			NewKeyTime = SequencerHelpers::FrameToTime(NewFrame, FrameRate);
+
+			Key.KeyArea->SetKeyTime(Key.KeyHandle.GetValue(), NewKeyTime);
+			bAnythingChanged = true;
 		}
 	}
 
-	return true;
+	if (bAnythingChanged)
+	{
+		UpdateRuntimeInstances();
+	}
 }
-
 void FSequencer::SpawnOrDestroyPuppetObjects( TSharedRef<FMovieSceneSequenceInstance> MovieSceneInstance )
 {
 	UMovieSceneSequence* Sequence = MovieSceneInstance->GetSequence();
@@ -1917,10 +1966,8 @@ void FSequencer::DeleteSelectedItems()
 	if (Selection.GetActiveSelection() == FSequencerSelection::EActiveSelection::KeyAndSection)
 	{
 		DeleteSelectedKeys();
-		for (TWeakObjectPtr<UMovieSceneSection> SelectedSection : Selection.GetSelectedSections())
-		{
-			DeleteSection(SelectedSection.Get());
-		}
+
+		DeleteSections(Selection.GetSelectedSections());
 	}
 	else if (Selection.GetActiveSelection() == FSequencerSelection::EActiveSelection::OutlinerNode)
 	{
diff --git a/Engine/Source/Editor/Sequencer/Private/Sequencer.h b/Engine/Source/Editor/Sequencer/Private/Sequencer.h
--- a/Engine/Source/Editor/Sequencer/Private/Sequencer.h
+++ b/Engine/Source/Editor/Sequencer/Private/Sequencer.h
@@ -122,9 +122,9 @@ public:
 	void SpawnOrDestroyPuppetObjects( TSharedRef<FMovieSceneSequenceInstance> MovieSceneInstance );
 
 	/** 
-	 * Deletes the passed in section
+	 * Deletes the passed in sections
 	 */
-	void DeleteSection(class UMovieSceneSection* Section);
+	void DeleteSections(const TSet<TWeakObjectPtr<UMovieSceneSection> > & Sections);
 
 	/**
 	 * Deletes the currently selected in keys
@@ -132,11 +132,19 @@ public:
 	void DeleteSelectedKeys();
 
 	/**
-	* Section active/inactive toggle
-	*/
-	void ToggleSectionActive();
-	bool CanToggleSectionActive() const;
-	bool IsToggleSectionActive() const;
+	 * Set interpolation modes
+	 */
+	void SetInterpTangentMode(ERichCurveInterpMode InterpMode, ERichCurveTangentMode TangentMode);
+
+	/**
+	 * Is interpolation mode selected
+	 */
+	bool IsInterpTangentModeSelected(ERichCurveInterpMode InterpMode, ERichCurveTangentMode TangentMode) const;
+
+	/**
+	 * Snap the currently selected keys to frame
+	 */
+	void SnapToFrame();
 
 	/**
 	 * @return Movie scene tools used by the sequencer
diff --git a/Engine/Source/Editor/Sequencer/Public/IKeyArea.h b/Engine/Source/Editor/Sequencer/Public/IKeyArea.h
--- a/Engine/Source/Editor/Sequencer/Public/IKeyArea.h
+++ b/Engine/Source/Editor/Sequencer/Public/IKeyArea.h
@@ -27,13 +27,21 @@ public:
 	/** @return The array of unsorted key handles in the key area */
 	virtual TArray<FKeyHandle> GetUnsortedKeyHandles() const = 0;
 
-	/** 
+	/**
+	 * Sets the time of a key given a handle
+	 *
+	 * @param KeyHandle Handle of the key
+	 * @param The new time of the key
+	 */
+	virtual void SetKeyTime(FKeyHandle KeyHandle, float NewKeyTime) const = 0;
+
+	/**
 	 * Gets the time of a key given a handle
 	 *
 	 * @param KeyHandle Handle of the key
 	 * @return The time of the key
 	 */
-	virtual float GetKeyTime( FKeyHandle KeyHandle ) const = 0;
+	virtual float GetKeyTime(FKeyHandle KeyHandle) const = 0;
 
 	/**
 	 * Moves a key
@@ -51,6 +59,54 @@ public:
 	 */
 	virtual void DeleteKey(FKeyHandle KeyHandle) = 0;
 
+	/**
+	 * Set key interpolation
+	 *
+	 * @param KeyHandle The key handle
+	 * @param InterpMode The interpolation mode
+	 */
+	virtual void SetKeyInterpMode(FKeyHandle KeyHandle, ERichCurveInterpMode InterpMode) = 0;
+
+	/**
+	 * Get key interpolation
+	 *
+	 * @param KeyHandle The key handle
+	 * @return InterpMode The interpolation mode
+	 */
+	virtual ERichCurveInterpMode GetKeyInterpMode(FKeyHandle Keyhandle) const = 0;
+
+	/**
+	 * Set key tangent
+	 *
+	 * @param KeyHandle The key handle
+	 * @param TangentMode The tangent mode
+	 */
+	virtual void SetKeyTangentMode(FKeyHandle KeyHandle, ERichCurveTangentMode TangentMode) = 0;
+
+	/**
+	 * Get key tangent
+	 *
+	 * @param KeyHandle The key handle
+	 * @return TangentMode The tangent mode
+	 */
+	virtual ERichCurveTangentMode GetKeyTangentMode(FKeyHandle KeyHandle) const = 0;
+
+	/**
+	 * Set extrapolation
+	 *
+	 * @param ExtrapMode The extrapolation mode
+	 * @param bPreInfinity Set pre-infinity (or post-infinity)
+	 */
+	virtual void SetExtrapolationMode(ERichCurveExtrapolation ExtrapMode, bool bPreInfinity) = 0;
+
+	/**
+	 * Get extrapolation
+	 *
+	 * @param bPreInfinity Get pre-infinity (or post-infinity)
+	 * @return ExtrapMode The extrapolation mode
+	 */
+	virtual ERichCurveExtrapolation GetExtrapolationMode(bool bPreInfinity) const = 0;
+
 	/**
 	 * Adds a key at the specified time if there isn't already a key present.  The value of the added key should
 	 * be the value which would be returned if the animation containing this key area was evaluated at the specified time.
diff --git a/Engine/Source/Runtime/Engine/Classes/Curves/CurveBase.h b/Engine/Source/Runtime/Engine/Classes/Curves/CurveBase.h
--- a/Engine/Source/Runtime/Engine/Classes/Curves/CurveBase.h
+++ b/Engine/Source/Runtime/Engine/Classes/Curves/CurveBase.h
@@ -134,7 +134,8 @@ enum ERichCurveInterpMode
 {
 	RCIM_Linear,
 	RCIM_Constant,
-	RCIM_Cubic
+	RCIM_Cubic,
+	RCIM_None
 };
 
 /** If using RCIM_Cubic, this enum describes how the tangents should be controlled in editor. */
@@ -143,7 +144,8 @@ enum ERichCurveTangentMode
 {
 	RCTM_Auto,
 	RCTM_User,
-	RCTM_Break
+	RCTM_Break,
+	RCTM_None
 };
 
 /** Enum to indicate whether if a tangent is 'weighted' (ie can be stretched). */
@@ -164,7 +166,8 @@ enum ERichCurveExtrapolation
 	RCCE_CycleWithOffset,
 	RCCE_Oscillate,
 	RCCE_Linear,
-	RCCE_Constant
+	RCCE_Constant,
+	RCCE_None
 };
 
 /** One key in a rich, editable float curve */
