Sequencer: Fix key selection sensitivity. Getting the key under the mouse was taking into account the section handle geometry, but it doesn't need to.

diff --git a/Engine/Source/Editor/Sequencer/Private/SSection.cpp b/Engine/Source/Editor/Sequencer/Private/SSection.cpp
--- a/Engine/Source/Editor/Sequencer/Private/SSection.cpp
+++ b/Engine/Source/Editor/Sequencer/Private/SSection.cpp
@@ -53,6 +53,8 @@ FGeometry SSection::GetKeyAreaGeometry( const FKeyAreaLayoutElement& KeyArea, co
 
 FSelectedKey SSection::GetKeyUnderMouse( const FVector2D& MousePosition, const FGeometry& AllottedGeometry ) const
 {
+	FGeometry SectionGeometry = MakeSectionGeometryWithoutHandles( AllottedGeometry, SectionInterface );
+
 	UMovieSceneSection& Section = *SectionInterface->GetSectionObject();
 
 	// Search every key area until we find the one under the mouse
@@ -66,12 +68,6 @@ FSelectedKey SSection::GetKeyUnderMouse( const FVector2D& MousePosition, const F
 		// Is the key area under the mouse
 		if( KeyAreaGeometryPadded.IsUnderLocation( MousePosition ) )
 		{
-			// Do not display section handles for connected sections
-			const bool bDisplaySectionHandles = SectionInterface->AreSectionsConnected();
-
-			const float SectionGripSize = bDisplaySectionHandles ? SectionInterface->GetSectionGripSize() : 0.0f;
-
-			FGeometry SectionGeometry = AllottedGeometry.MakeChild(FVector2D(SectionGripSize, 0), AllottedGeometry.GetDrawSize() - FVector2D(SectionGripSize*2, 0.0f));
 			FGeometry KeyAreaGeometry = GetKeyAreaGeometry( Element, SectionGeometry );
 
 			FVector2D LocalSpaceMousePosition = KeyAreaGeometry.AbsoluteToLocal( MousePosition );
@@ -86,9 +82,8 @@ FSelectedKey SSection::GetKeyUnderMouse( const FVector2D& MousePosition, const F
 			{
 				FKeyHandle KeyHandle = KeyHandles[KeyIndex];
 				float KeyPosition = TimeToPixelConverter.TimeToPixel( KeyArea->GetKeyTime(KeyHandle) );
-
 				FGeometry KeyGeometry = KeyAreaGeometry.MakeChild( 
-					FVector2D( KeyPosition - FMath::TruncToFloat(SequencerSectionConstants::KeySize.X/2.0f), ((KeyAreaGeometry.Size.Y*.5f)-(SequencerSectionConstants::KeySize.Y*.5f)) ),
+					FVector2D( KeyPosition - FMath::CeilToFloat(SequencerSectionConstants::KeySize.X/2.0f), ((KeyAreaGeometry.Size.Y*.5f)-(SequencerSectionConstants::KeySize.Y*.5f)) ),
 					SequencerSectionConstants::KeySize );
 
 				if( KeyGeometry.IsUnderLocation( MousePosition ) )
@@ -96,7 +91,6 @@ FSelectedKey SSection::GetKeyUnderMouse( const FVector2D& MousePosition, const F
 					// The current key is under the mouse
 					return FSelectedKey( Section, KeyArea, KeyHandle );
 				}
-				
 			}
 
 			// no key was selected in the current key area but the mouse is in the key area so it cannot possibly be in any other key area
