Sequencer: Added details panel

diff --git a/Engine/Source/Editor/Sequencer/Private/SSequencer.cpp b/Engine/Source/Editor/Sequencer/Private/SSequencer.cpp
--- a/Engine/Source/Editor/Sequencer/Private/SSequencer.cpp
+++ b/Engine/Source/Editor/Sequencer/Private/SSequencer.cpp
@@ -34,6 +34,10 @@
 #include "VirtualTrackArea.h"
 #include "SequencerHotspots.h"
 
+#include "IDetailsView.h"
+#include "PropertyEditorModule.h"
+
+
 #define LOCTEXT_NAMESPACE "Sequencer"
 
 
@@ -250,6 +254,27 @@ void SSequencer::Construct( const FArguments& InArgs, TSharedRef< class FSequenc
 
 	EditTool.Reset(new FSequencerEditTool_Movement(PinnedSequencer, SharedThis(this)));
 
+	// initialize details view
+	FDetailsViewArgs DetailsViewArgs;
+	{
+		DetailsViewArgs.bAllowSearch = false;
+		DetailsViewArgs.bCustomFilterAreaLocation = true;
+		DetailsViewArgs.bCustomNameAreaLocation = true;
+		DetailsViewArgs.bHideSelectionTip = true;
+		DetailsViewArgs.bLockable = false;
+		DetailsViewArgs.bSearchInitialKeyFocus = true;
+		DetailsViewArgs.bUpdatesFromSelection = false;
+		DetailsViewArgs.NotifyHook = this;
+		DetailsViewArgs.bShowOptions = false;
+		DetailsViewArgs.bShowModifiedPropertiesOption = false;
+	}
+
+	DetailsView = FModuleManager::GetModuleChecked<FPropertyEditorModule>("PropertyEditor").CreateDetailView(DetailsViewArgs);
+	{
+		DetailsView->SetEnabled(TAttribute<bool>::Create(TAttribute<bool>::FGetter::CreateSP(this, &SSequencer::HandleDetailsViewEnabled)));
+		DetailsView->SetVisibility(TAttribute<EVisibility>::Create(TAttribute<EVisibility>::FGetter::CreateSP(this, &SSequencer::HandleDetailsViewVisibility)));
+	}
+
 	const int32 				Column0	= 0,	Column1	= 1;
 	const int32 Row0	= 0,
 				Row1	= 1,
@@ -285,140 +310,153 @@ void SSequencer::Construct( const FArguments& InArgs, TSharedRef< class FSequenc
 
 				+ SOverlay::Slot()
 				[
-					SNew( SGridPanel )
-					.FillRow( 1, 1.f )
-					.FillColumn( 0, FillCoefficient_0 )
-					.FillColumn( 1, FillCoefficient_1 )
-
-					+ SGridPanel::Slot( Column0, Row0 )
-					.VAlign( VAlign_Center )
-					[
-						// Search box for searching through the outliner
-						SNew( SSearchBox )
-						.OnTextChanged( this, &SSequencer::OnOutlinerSearchChanged )
-					]
-
-					+ SGridPanel::Slot( Column0, Row1 )
-					.ColumnSpan(2)
-					[
-						SNew(SHorizontalBox)
-
-						+ SHorizontalBox::Slot()
+					SNew(SSplitter)
+						.Orientation(Orient_Horizontal)
+						
+					+ SSplitter::Slot()
+						.Value(0.8f)
 						[
-							SNew( SOverlay )
+							SNew( SGridPanel )
+								.FillRow( 1, 1.f )
+								.FillColumn( 0, FillCoefficient_0 )
+								.FillColumn( 1, FillCoefficient_1 )
 
-							+ SOverlay::Slot()
-							[
-								SNew(SHorizontalBox)
-								
-								+ SHorizontalBox::Slot()
-								.FillWidth( FillCoefficient_0 )
+							+ SGridPanel::Slot( Column0, Row0 )
+								.VAlign( VAlign_Center )
 								[
-									SNew(SSequencerTreeViewBox, PinnedSequencer)
-									// Padding to allow space for the scroll bar
-									.Padding(FMargin(0,0,10.f,0))
-									[
-										TreeView.ToSharedRef()
-									]
+									// Search box for searching through the outliner
+									SNew( SSearchBox )
+										.OnTextChanged( this, &SSequencer::OnOutlinerSearchChanged )
 								]
 
-								+ SHorizontalBox::Slot()
-								.FillWidth( FillCoefficient_1 )
+							+ SGridPanel::Slot( Column0, Row1 )
+								.ColumnSpan(2)
 								[
-									TrackArea.ToSharedRef()
-								]
-							]
+									SNew(SHorizontalBox)
 
-							+ SOverlay::Slot()
-							.HAlign( HAlign_Right )
-							[
-								ScrollBar
-							]
-						]
+									+ SHorizontalBox::Slot()
+										[
+											SNew( SOverlay )
 
-						+ SHorizontalBox::Slot()
-						.FillWidth( TAttribute<float>( this, &SSequencer::GetOutlinerSpacerFill ) )
-						[
-							SNew(SSpacer)
-						]
-					]
-
-					+ SGridPanel::Slot( Column0, Row2 )
-					.HAlign( HAlign_Center )
-					[
-						MakeTransportControls()
-					]
-
-					// Second column
-					+ SGridPanel::Slot( Column1, Row0 )
-					[
-						SNew( SBorder )
-						.BorderImage( FEditorStyle::GetBrush("ToolPanel.GroupBorder") )
-						.BorderBackgroundColor( FLinearColor(.50f, .50f, .50f, 1.0f ) )
-						.Padding(0)
-						[
-							TopTimeSlider
-						]
-					]
-
-					// Overlay that draws the tick lines
-					+ SGridPanel::Slot( Column1, Row1, SGridPanel::Layer(0) )
-					[
-						SNew( SSequencerSectionOverlay, TimeSliderController )
-						.Visibility( EVisibility::HitTestInvisible )
-						.DisplayScrubPosition( false )
-						.DisplayTickLines( true )
-					]
+											+ SOverlay::Slot()
+												[
+													SNew(SHorizontalBox)
+								
+													+ SHorizontalBox::Slot()
+														.FillWidth( FillCoefficient_0 )
+														[
+															SNew(SSequencerTreeViewBox, PinnedSequencer)
+																.Padding(FMargin(0, 0, 10.f, 0)) // Padding to allow space for the scroll bar
+																[
+																	TreeView.ToSharedRef()
+																]
+														]
+
+													+ SHorizontalBox::Slot()
+														.FillWidth( FillCoefficient_1 )
+														[
+															TrackArea.ToSharedRef()
+														]
+												]
+
+											+ SOverlay::Slot()
+												.HAlign( HAlign_Right )
+												[
+													ScrollBar
+												]
+										]
+
+									+ SHorizontalBox::Slot()
+										.FillWidth( TAttribute<float>( this, &SSequencer::GetOutlinerSpacerFill ) )
+										[
+											SNew(SSpacer)
+										]
+								]
 
-					// Curve editor
-					+ SGridPanel::Slot( Column1, Row1 )
-					[
-						CurveEditor.ToSharedRef()
-					]
+							+ SGridPanel::Slot( Column0, Row2 )
+								.HAlign( HAlign_Center )
+								[
+									MakeTransportControls()
+								]
 
-					// Overlay that draws the scrub position
-					+ SGridPanel::Slot( Column1, Row1, SGridPanel::Layer(30) )
-					[
-						SNew( SSequencerSectionOverlay, TimeSliderController )
-						.Visibility( EVisibility::HitTestInvisible )
-						.DisplayScrubPosition( true )
-						.DisplayTickLines( false )
-					]
+							// Second column
+							+ SGridPanel::Slot( Column1, Row0 )
+								[
+									SNew( SBorder )
+										.BorderImage( FEditorStyle::GetBrush("ToolPanel.GroupBorder") )
+										.BorderBackgroundColor( FLinearColor(.50f, .50f, .50f, 1.0f ) )
+										.Padding(0)
+										[
+											TopTimeSlider
+										]
+								]
 
-					+ SGridPanel::Slot( Column1, Row2 )
-					[
-						SNew( SBorder )
-						.BorderImage( FEditorStyle::GetBrush("ToolPanel.GroupBorder") )
-						.BorderBackgroundColor( FLinearColor(.50f, .50f, .50f, 1.0f ) )
-						.Padding(0)
-						[
-							SNew(SVerticalBox)
+							// Overlay that draws the tick lines
+							+ SGridPanel::Slot( Column1, Row1, SGridPanel::Layer(0) )
+								[
+									SNew( SSequencerSectionOverlay, TimeSliderController )
+										.Visibility( EVisibility::HitTestInvisible )
+										.DisplayScrubPosition( false )
+										.DisplayTickLines( true )
+								]
 
-							+ SVerticalBox::Slot()
-							.AutoHeight()
-							[
-								SNew( SOverlay )
+							// Curve editor
+							+ SGridPanel::Slot( Column1, Row1 )
+								[
+									CurveEditor.ToSharedRef()
+								]
 
-								+ SOverlay::Slot()
+							// Overlay that draws the scrub position
+							+ SGridPanel::Slot( Column1, Row1, SGridPanel::Layer(30) )
 								[
-									BottomTimeSlider
+									SNew( SSequencerSectionOverlay, TimeSliderController )
+										.Visibility( EVisibility::HitTestInvisible )
+										.DisplayScrubPosition( true )
+										.DisplayTickLines( false )
 								]
 
-								+ SOverlay::Slot()
+							+ SGridPanel::Slot( Column1, Row2 )
 								[
-									BottomTimeRange
+									SNew( SBorder )
+										.BorderImage( FEditorStyle::GetBrush("ToolPanel.GroupBorder") )
+										.BorderBackgroundColor( FLinearColor(.50f, .50f, .50f, 1.0f ) )
+										.Padding(0)
+										[
+											SNew(SVerticalBox)
+
+											+ SVerticalBox::Slot()
+												.AutoHeight()
+												[
+													SNew( SOverlay )
+
+													+ SOverlay::Slot()
+														[
+															BottomTimeSlider
+														]
+
+													+ SOverlay::Slot()
+														[
+															BottomTimeRange
+														]
+												]
+
+											+ SVerticalBox::Slot()
+												.AutoHeight()
+												[
+													SAssignNew( BreadcrumbTrail, SBreadcrumbTrail<FSequencerBreadcrumb> )
+													.Visibility( this, &SSequencer::GetBreadcrumbTrailVisibility )
+													.OnCrumbClicked( this, &SSequencer::OnCrumbClicked )
+												]
+										]
 								]
-							]
-
-							+ SVerticalBox::Slot()
-							.AutoHeight()
-							[
-								SAssignNew( BreadcrumbTrail, SBreadcrumbTrail<FSequencerBreadcrumb> )
-								.Visibility( this, &SSequencer::GetBreadcrumbTrailVisibility )
-								.OnCrumbClicked( this, &SSequencer::OnCrumbClicked )
-							]
 						]
-					]
+
+					+ SSplitter::Slot()
+						.SizeRule(SSplitter::SizeToContent)
+						.Value(0.2f)
+						[
+							DetailsView.ToSharedRef()
+						]
 				]
 
 				+ SOverlay::Slot()
@@ -445,6 +483,9 @@ void SSequencer::Construct( const FArguments& InArgs, TSharedRef< class FSequenc
 		]
 	];
 
+	InSequencer->GetSelection().GetOnKeySelectionChanged().AddSP(this, &SSequencer::HandleKeySelectionChanged);
+	InSequencer->GetSelection().GetOnSectionSelectionChanged().AddSP(this, &SSequencer::HandleSectionSelectionChanged);
+
 	ResetBreadcrumbs();
 }
 
@@ -483,6 +524,49 @@ void SSequencer::BindCommands(TSharedRef<FUICommandList> SequencerCommandBinding
 		);
 }
 
+
+void SSequencer::NotifyPostChange(const FPropertyChangedEvent& PropertyChangedEvent, class FEditPropertyChain* PropertyThatChanged)
+{
+}
+
+
+bool SSequencer::HandleDetailsViewEnabled() const
+{
+	return true;
+}
+
+
+EVisibility SSequencer::HandleDetailsViewVisibility() const
+{/*
+	if (Model->GetSelectedMessage().IsValid())
+	{
+		return EVisibility::Visible;
+	}
+
+	return EVisibility::Hidden;*/
+	return EVisibility::Visible;
+}
+
+
+void SSequencer::HandleKeySelectionChanged()
+{
+
+}
+
+
+void SSequencer::HandleSectionSelectionChanged()
+{
+	TArray<TWeakObjectPtr<UObject>> Sections;
+
+	for (auto Section : Sequencer.Pin()->GetSelection().GetSelectedSections())
+	{
+		Sections.Add(Section);
+	}
+
+	DetailsView->SetObjects(Sections, true);
+}
+
+
 bool SSequencer::IsEditToolEnabled(FName InIdentifier)
 {
 	return GetEditTool().GetIdentifier() == InIdentifier;
diff --git a/Engine/Source/Editor/Sequencer/Private/SSequencer.h b/Engine/Source/Editor/Sequencer/Private/SSequencer.h
--- a/Engine/Source/Editor/Sequencer/Private/SSequencer.h
+++ b/Engine/Source/Editor/Sequencer/Private/SSequencer.h
@@ -6,9 +6,12 @@
 #include "ISequencerEditTool.h"
 #include "SequencerCommonHelpers.h"
 
+
+class IDetailsView;
 class SSequencerTreeView;
 struct FSectionHandle;
 
+
 namespace SequencerLayoutConstants
 {
 	/** The amount to indent child nodes of the layout tree */
@@ -23,6 +26,7 @@ namespace SequencerLayoutConstants
 	const float CategoryNodeHeight = 15.0f;
 }
 
+
 /**
  * The kind of breadcrumbs that sequencer uses
  */
@@ -47,10 +51,14 @@ struct FSequencerBreadcrumb
 		: BreadcrumbType(FSequencerBreadcrumb::ShotType) {}
 };
 
+
 /**
  * Main sequencer UI widget
  */
-class SSequencer : public SCompoundWidget, public FGCObject
+class SSequencer
+	: public SCompoundWidget
+	, public FGCObject
+	, public FNotifyHook
 {
 public:
 	DECLARE_DELEGATE_OneParam( FOnToggleBoolOption, bool )
@@ -137,8 +145,26 @@ public:
 	/** Get an array of section handles for the given set of movie scene sections */
 	TArray<FSectionHandle> GetSectionHandles(const TSet<TWeakObjectPtr<UMovieSceneSection>>& DesiredSections) const;
 
+public:
+
+	// FNotifyHook overrides
+
+	void NotifyPostChange(const FPropertyChangedEvent& PropertyChangedEvent, class FEditPropertyChain* PropertyThatChanged);
+
 private:
 	
+	/** Handles checking whether the details view is enabled. */
+	bool HandleDetailsViewEnabled() const;
+
+	/** Handles determining the visibility of the details view. */
+	EVisibility HandleDetailsViewVisibility() const;
+
+	/** Handles key selection changes. */
+	void HandleKeySelectionChanged();
+
+	/** Handles section selection changes. */
+	void HandleSectionSelectionChanged();
+
 	/** Check whether the specified edit tool is enabled */
 	bool IsEditToolEnabled(FName InIdentifier);
 
@@ -280,6 +306,9 @@ private:
 
 private:
 
+	/** Holds the details view. */
+	TSharedPtr<IDetailsView> DetailsView;
+
 	/** Section area widget */
 	TSharedPtr<class SSequencerTrackArea> TrackArea;
 	/** Outliner widget */
diff --git a/Engine/Source/Editor/Sequencer/Sequencer.Build.cs b/Engine/Source/Editor/Sequencer/Sequencer.Build.cs
--- a/Engine/Source/Editor/Sequencer/Sequencer.Build.cs
+++ b/Engine/Source/Editor/Sequencer/Sequencer.Build.cs
@@ -23,21 +23,21 @@ public class Sequencer : ModuleRules
 				"MovieSceneTools", 
 				"EditorWidgets", 
 				"SequencerWidgets",
-				"BlueprintGraph"
+				"BlueprintGraph",
 			}
 		);
 
 		PrivateIncludePathModuleNames.AddRange(
 			new string[] {
 				"PropertyEditor",
-				"Kismet"
+				"Kismet",
 			}
 		);
 
 		DynamicallyLoadedModuleNames.AddRange(
 			new string[] {
 				"PropertyEditor",
-				"WorkspaceMenuStructure"
+				"WorkspaceMenuStructure",
 			}
 		);
 
