Sequencer: Adjusted padding because race car

diff --git a/Engine/Source/Editor/Sequencer/Private/SSequencer.cpp b/Engine/Source/Editor/Sequencer/Private/SSequencer.cpp
--- a/Engine/Source/Editor/Sequencer/Private/SSequencer.cpp
+++ b/Engine/Source/Editor/Sequencer/Private/SSequencer.cpp
@@ -305,159 +305,162 @@ void SSequencer::Construct( const FArguments& InArgs, TSharedRef< class FSequenc
 			]
 
 			+ SVerticalBox::Slot()
-			[
-				SNew(SOverlay)
-
-				+ SOverlay::Slot()
 				[
-					SNew(SSplitter)
-						.Orientation(Orient_Horizontal)
+					SNew(SOverlay)
+
+					+ SOverlay::Slot()
+					[
+						SNew(SSplitter)
+							.Orientation(Orient_Horizontal)
 						
-					+ SSplitter::Slot()
-						.Value(0.8f)
-						[
-							SNew( SGridPanel )
-								.FillRow( 1, 1.f )
-								.FillColumn( 0, FillCoefficient_0 )
-								.FillColumn( 1, FillCoefficient_1 )
-
-							+ SGridPanel::Slot( Column0, Row0 )
-								.VAlign( VAlign_Center )
-								[
-									// Search box for searching through the outliner
-									SNew( SSearchBox )
-										.OnTextChanged( this, &SSequencer::OnOutlinerSearchChanged )
-								]
-
-							+ SGridPanel::Slot( Column0, Row1 )
-								.ColumnSpan(2)
-								[
-									SNew(SHorizontalBox)
-
-									+ SHorizontalBox::Slot()
-										[
-											SNew( SOverlay )
-
-											+ SOverlay::Slot()
-												[
-													SNew(SHorizontalBox)
+						+ SSplitter::Slot()
+							.Value(0.8f)
+							[
+								SNew(SBox)
+									.Padding(FMargin(0.0f, 2.0f, 0.0f, 0.0f))
+									[
+										SNew( SGridPanel )
+											.FillRow( 1, 1.f )
+											.FillColumn( 0, FillCoefficient_0 )
+											.FillColumn( 1, FillCoefficient_1 )
+
+										+ SGridPanel::Slot( Column0, Row0 )
+											.VAlign( VAlign_Center )
+											[
+												// Search box for searching through the outliner
+												SNew( SSearchBox )
+													.OnTextChanged( this, &SSequencer::OnOutlinerSearchChanged )
+											]
+
+										+ SGridPanel::Slot( Column0, Row1 )
+											.ColumnSpan(2)
+											[
+												SNew(SHorizontalBox)
+
+												+ SHorizontalBox::Slot()
+													[
+														SNew( SOverlay )
+
+														+ SOverlay::Slot()
+															[
+																SNew(SHorizontalBox)
 								
-													+ SHorizontalBox::Slot()
-														.FillWidth( FillCoefficient_0 )
-														[
-															SNew(SSequencerTreeViewBox, PinnedSequencer)
-																.Padding(FMargin(0, 0, 10.f, 0)) // Padding to allow space for the scroll bar
-																[
-																	TreeView.ToSharedRef()
-																]
-														]
-
-													+ SHorizontalBox::Slot()
-														.FillWidth( FillCoefficient_1 )
-														[
-															TrackArea.ToSharedRef()
-														]
-												]
-
-											+ SOverlay::Slot()
-												.HAlign( HAlign_Right )
-												[
-													ScrollBar
-												]
-										]
-
-									+ SHorizontalBox::Slot()
-										.FillWidth( TAttribute<float>( this, &SSequencer::GetOutlinerSpacerFill ) )
-										[
-											SNew(SSpacer)
-										]
-								]
-
-							+ SGridPanel::Slot( Column0, Row2 )
-								.HAlign( HAlign_Center )
-								[
-									MakeTransportControls()
-								]
-
-							// Second column
-							+ SGridPanel::Slot( Column1, Row0 )
-								[
-									SNew( SBorder )
-										.BorderImage( FEditorStyle::GetBrush("ToolPanel.GroupBorder") )
-										.BorderBackgroundColor( FLinearColor(.50f, .50f, .50f, 1.0f ) )
-										.Padding(0)
-										[
-											TopTimeSlider
-										]
-								]
-
-							// Overlay that draws the tick lines
-							+ SGridPanel::Slot( Column1, Row1, SGridPanel::Layer(0) )
-								[
-									SNew( SSequencerSectionOverlay, TimeSliderController )
-										.Visibility( EVisibility::HitTestInvisible )
-										.DisplayScrubPosition( false )
-										.DisplayTickLines( true )
-								]
-
-							// Curve editor
-							+ SGridPanel::Slot( Column1, Row1 )
-								[
-									CurveEditor.ToSharedRef()
-								]
-
-							// Overlay that draws the scrub position
-							+ SGridPanel::Slot( Column1, Row1, SGridPanel::Layer(30) )
-								[
-									SNew( SSequencerSectionOverlay, TimeSliderController )
-										.Visibility( EVisibility::HitTestInvisible )
-										.DisplayScrubPosition( true )
-										.DisplayTickLines( false )
-								]
-
-							+ SGridPanel::Slot( Column1, Row2 )
-								[
-									SNew( SBorder )
-										.BorderImage( FEditorStyle::GetBrush("ToolPanel.GroupBorder") )
-										.BorderBackgroundColor( FLinearColor(.50f, .50f, .50f, 1.0f ) )
-										.Padding(0)
-										[
-											SNew(SVerticalBox)
-
-											+ SVerticalBox::Slot()
-												.AutoHeight()
-												[
-													SNew( SOverlay )
-
-													+ SOverlay::Slot()
-														[
-															BottomTimeSlider
-														]
-
-													+ SOverlay::Slot()
-														[
-															BottomTimeRange
-														]
-												]
-
-											+ SVerticalBox::Slot()
-												.AutoHeight()
-												[
-													SAssignNew( BreadcrumbTrail, SBreadcrumbTrail<FSequencerBreadcrumb> )
-													.Visibility( this, &SSequencer::GetBreadcrumbTrailVisibility )
-													.OnCrumbClicked( this, &SSequencer::OnCrumbClicked )
-												]
-										]
-								]
-						]
-
-					+ SSplitter::Slot()
-						.SizeRule(SSplitter::SizeToContent)
-						.Value(0.2f)
-						[
-							DetailsView.ToSharedRef()
-						]
-				]
+																+ SHorizontalBox::Slot()
+																	.FillWidth( FillCoefficient_0 )
+																	[
+																		SNew(SSequencerTreeViewBox, PinnedSequencer)
+																			.Padding(FMargin(0, 0, 10.f, 0)) // Padding to allow space for the scroll bar
+																			[
+																				TreeView.ToSharedRef()
+																			]
+																	]
+
+																+ SHorizontalBox::Slot()
+																	.FillWidth( FillCoefficient_1 )
+																	[
+																		TrackArea.ToSharedRef()
+																	]
+															]
+
+														+ SOverlay::Slot()
+															.HAlign( HAlign_Right )
+															[
+																ScrollBar
+															]
+													]
+
+												+ SHorizontalBox::Slot()
+													.FillWidth( TAttribute<float>( this, &SSequencer::GetOutlinerSpacerFill ) )
+													[
+														SNew(SSpacer)
+													]
+											]
+
+										+ SGridPanel::Slot( Column0, Row2 )
+											.HAlign( HAlign_Center )
+											[
+												MakeTransportControls()
+											]
+
+										// Second column
+										+ SGridPanel::Slot( Column1, Row0 )
+											[
+												SNew( SBorder )
+													.BorderImage( FEditorStyle::GetBrush("ToolPanel.GroupBorder") )
+													.BorderBackgroundColor( FLinearColor(.50f, .50f, .50f, 1.0f ) )
+													[
+														TopTimeSlider
+													]
+											]
+
+										// Overlay that draws the tick lines
+										+ SGridPanel::Slot( Column1, Row1, SGridPanel::Layer(0) )
+											[
+												SNew( SSequencerSectionOverlay, TimeSliderController )
+													.Visibility( EVisibility::HitTestInvisible )
+													.DisplayScrubPosition( false )
+													.DisplayTickLines( true )
+											]
+
+										// Curve editor
+										+ SGridPanel::Slot( Column1, Row1 )
+											[
+												CurveEditor.ToSharedRef()
+											]
+
+										// Overlay that draws the scrub position
+										+ SGridPanel::Slot( Column1, Row1, SGridPanel::Layer(30) )
+											[
+												SNew( SSequencerSectionOverlay, TimeSliderController )
+													.Visibility( EVisibility::HitTestInvisible )
+													.DisplayScrubPosition( true )
+													.DisplayTickLines( false )
+											]
+
+										+ SGridPanel::Slot( Column1, Row2 )
+											[
+												SNew( SBorder )
+													.BorderImage( FEditorStyle::GetBrush("ToolPanel.GroupBorder") )
+													.BorderBackgroundColor( FLinearColor(.50f, .50f, .50f, 1.0f ) )
+													.Padding(0)
+													[
+														SNew(SVerticalBox)
+
+														+ SVerticalBox::Slot()
+															.AutoHeight()
+															[
+																SNew( SOverlay )
+
+																+ SOverlay::Slot()
+																	[
+																		BottomTimeSlider
+																	]
+
+																+ SOverlay::Slot()
+																	[
+																		BottomTimeRange
+																	]
+															]
+
+														+ SVerticalBox::Slot()
+															.AutoHeight()
+															[
+																SAssignNew( BreadcrumbTrail, SBreadcrumbTrail<FSequencerBreadcrumb> )
+																.Visibility( this, &SSequencer::GetBreadcrumbTrailVisibility )
+																.OnCrumbClicked( this, &SSequencer::OnCrumbClicked )
+															]
+													]
+											]
+									]
+							]
+
+						+ SSplitter::Slot()
+							.SizeRule(SSplitter::SizeToContent)
+							.Value(0.2f)
+							[
+								DetailsView.ToSharedRef()
+							]
+					]
 
 				+ SOverlay::Slot()
 				[
