Sequencer: Add middle click copy key from maya feature.

Middle click on a key copies that key and initiates dragging of the new key. Middle click on an empty key area will create a key and set the value at the curve's value at that time. Key groupings are supported as well (only channels that are keyed are copied).

diff --git a/Engine/Source/Editor/MovieSceneTools/Private/MovieSceneToolHelpers.cpp b/Engine/Source/Editor/MovieSceneTools/Private/MovieSceneToolHelpers.cpp
--- a/Engine/Source/Editor/MovieSceneTools/Private/MovieSceneToolHelpers.cpp
+++ b/Engine/Source/Editor/MovieSceneTools/Private/MovieSceneToolHelpers.cpp
@@ -16,8 +16,10 @@
 #include "SEnumCurveKeyEditor.h"
 #include "SBoolCurveKeyEditor.h"
 
-void FFloatCurveKeyArea::AddKeyUnique(float Time)
+TArray<FKeyHandle> FFloatCurveKeyArea::AddKeyUnique(float Time, float TimeToCopyFrom)
 {
+	TArray<FKeyHandle> AddedKeyHandles;
+
 	FKeyHandle CurrentKeyHandle = Curve->FindKey(Time);
 	if (Curve->IsKeyHandleValid(CurrentKeyHandle) == false)
 	{
@@ -29,8 +31,33 @@ void FFloatCurveKeyArea::AddKeyUnique(float Time)
 		{
 			OwningSection->SetEndTime(Time);
 		}
-		Curve->AddKey(Time, Curve->Eval(Time), false, CurrentKeyHandle);
+
+		float Value = Curve->Eval(Time);
+		if (TimeToCopyFrom != FLT_MAX)
+		{
+			Value = Curve->Eval(TimeToCopyFrom);
+		}
+
+		Curve->AddKey(Time, Value, false, CurrentKeyHandle);
+		AddedKeyHandles.Add(CurrentKeyHandle);
+		
+		// Copy the properties from the key if it exists
+		FKeyHandle KeyHandleToCopy = Curve->FindKey(TimeToCopyFrom);
+		if (Curve->IsKeyHandleValid(KeyHandleToCopy))
+		{
+			FRichCurveKey& CurrentKey = Curve->GetKey(CurrentKeyHandle);
+			FRichCurveKey& KeyToCopy = Curve->GetKey(KeyHandleToCopy);
+			CurrentKey.InterpMode = KeyToCopy.InterpMode;
+			CurrentKey.TangentMode = KeyToCopy.TangentMode;
+			CurrentKey.TangentWeightMode = KeyToCopy.TangentWeightMode;
+			CurrentKey.ArriveTangent = KeyToCopy.ArriveTangent;
+			CurrentKey.LeaveTangent = KeyToCopy.LeaveTangent;
+			CurrentKey.ArriveTangentWeight = KeyToCopy.ArriveTangentWeight;
+			CurrentKey.LeaveTangentWeight = KeyToCopy.LeaveTangentWeight;
+		}
 	}
+
+	return AddedKeyHandles;
 }
 
 TSharedRef<SWidget> FFloatCurveKeyArea::CreateKeyEditor(ISequencer* Sequencer)
@@ -41,8 +68,10 @@ TSharedRef<SWidget> FFloatCurveKeyArea::CreateKeyEditor(ISequencer* Sequencer)
 		.Curve(Curve);
 };
 
-void FIntegralKeyArea::AddKeyUnique(float Time)
+TArray<FKeyHandle> FIntegralKeyArea::AddKeyUnique(float Time, float TimeToCopyFrom)
 {
+	TArray<FKeyHandle> AddedKeyHandles;
+
 	FKeyHandle CurrentKey = Curve.FindKey(Time);
 	if (Curve.IsKeyHandleValid(CurrentKey) == false)
 	{
@@ -54,8 +83,18 @@ void FIntegralKeyArea::AddKeyUnique(float Time)
 		{
 			OwningSection->SetEndTime(Time);
 		}
-		Curve.AddKey(Time, Curve.Evaluate(Time), CurrentKey);
+
+		int32 Value = Curve.Evaluate(Time);
+		if (TimeToCopyFrom != FLT_MAX)
+		{
+			Value = Curve.Evaluate(TimeToCopyFrom);
+		}
+
+		Curve.AddKey(Time, Value, CurrentKey);
+		AddedKeyHandles.Add(CurrentKey);
 	}
+
+	return AddedKeyHandles;
 }
 
 TSharedRef<SWidget> FIntegralKeyArea::CreateKeyEditor(ISequencer* Sequencer)
diff --git a/Engine/Source/Editor/MovieSceneTools/Public/MovieSceneToolHelpers.h b/Engine/Source/Editor/MovieSceneTools/Public/MovieSceneToolHelpers.h
--- a/Engine/Source/Editor/MovieSceneTools/Public/MovieSceneToolHelpers.h
+++ b/Engine/Source/Editor/MovieSceneTools/Public/MovieSceneToolHelpers.h
@@ -110,7 +110,7 @@ public:
 		return RCCE_None;
 	}
 
-	virtual void AddKeyUnique(float Time) override;
+	virtual TArray<FKeyHandle> AddKeyUnique(float Time, float TimeToCopyFrom = FLT_MAX) override;
 
 	virtual FRichCurve* GetRichCurve() override { return Curve; }
 
@@ -199,7 +199,7 @@ public:
 		return RCCE_None;
 	}
 
-	virtual void AddKeyUnique(float Time) override;
+	virtual TArray<FKeyHandle> AddKeyUnique(float Time, float TimeToCopyFrom = FLT_MAX) override;
 
 	virtual FRichCurve* GetRichCurve() override { return nullptr; };
 
diff --git a/Engine/Source/Editor/Sequencer/Private/GroupedKeyArea.cpp b/Engine/Source/Editor/Sequencer/Private/GroupedKeyArea.cpp
--- a/Engine/Source/Editor/Sequencer/Private/GroupedKeyArea.cpp
+++ b/Engine/Source/Editor/Sequencer/Private/GroupedKeyArea.cpp
@@ -348,12 +348,29 @@ ERichCurveExtrapolation FGroupedKeyArea::GetExtrapolationMode(bool bPreInfinity)
 	return ExtrapMode;
 }
 
-void FGroupedKeyArea::AddKeyUnique(float Time)
+TArray<FKeyHandle> FGroupedKeyArea::AddKeyUnique(float Time, float TimeToCopyFrom)
 {
+	TArray<FKeyHandle> AddedKeyHandles;
+
 	for (auto& Area : KeyAreas)
 	{
-		Area->AddKeyUnique(Time);
+		// If TimeToCopyFrom is valid, add a key only if there is a key to copy from
+		if (TimeToCopyFrom != FLT_MAX)
+		{
+			if (FRichCurve* Curve = Area->GetRichCurve())
+			{
+				if (!Curve->IsKeyHandleValid(Curve->FindKey(TimeToCopyFrom)))
+				{
+					continue;
+				}
+			}
+		}
+
+		TArray<FKeyHandle> AddedGroupKeyHandles = Area->AddKeyUnique(Time, TimeToCopyFrom);
+		AddedKeyHandles.Append(AddedGroupKeyHandles);
 	}
+
+	return AddedKeyHandles;
 }
 
 FRichCurve* FGroupedKeyArea::GetRichCurve()
diff --git a/Engine/Source/Editor/Sequencer/Private/GroupedKeyArea.h b/Engine/Source/Editor/Sequencer/Private/GroupedKeyArea.h
--- a/Engine/Source/Editor/Sequencer/Private/GroupedKeyArea.h
+++ b/Engine/Source/Editor/Sequencer/Private/GroupedKeyArea.h
@@ -81,7 +81,7 @@ public:
 	virtual ERichCurveTangentMode   GetKeyTangentMode(FKeyHandle KeyHandle) const override;
 	virtual void                    SetExtrapolationMode(ERichCurveExtrapolation ExtrapMode, bool bPreInfinity) override;
 	virtual ERichCurveExtrapolation GetExtrapolationMode(bool bPreInfinity) const override;
-	virtual void					AddKeyUnique(float Time) override;
+	virtual TArray<FKeyHandle>		AddKeyUnique(float Time, float TimeToCopyFrom = FLT_MAX) override;
 	virtual FRichCurve*				GetRichCurve() override;
 	virtual UMovieSceneSection*		GetOwningSection() override;
 	virtual bool					CanCreateKeyEditor() override;
diff --git a/Engine/Source/Editor/Sequencer/Private/SSection.cpp b/Engine/Source/Editor/Sequencer/Private/SSection.cpp
--- a/Engine/Source/Editor/Sequencer/Private/SSection.cpp
+++ b/Engine/Source/Editor/Sequencer/Private/SSection.cpp
@@ -102,6 +102,85 @@ FSelectedKey SSection::GetKeyUnderMouse( const FVector2D& MousePosition, const F
 	return FSelectedKey();
 }
 
+FSelectedKey SSection::CreateKeyUnderMouse( const FVector2D& MousePosition, const FGeometry& AllottedGeometry, FSelectedKey InPressedKey )
+{
+	UMovieSceneSection& Section = *SectionInterface->GetSectionObject();
+
+	FGeometry SectionGeometry = MakeSectionGeometryWithoutHandles( AllottedGeometry, SectionInterface );
+
+	// Search every key area until we find the one under the mouse
+	for (const FKeyAreaLayoutElement& Element : Layout->GetElements())
+	{
+		TSharedPtr<IKeyArea> KeyArea = Element.GetKeyArea();
+
+		// Compute the current key area geometry
+		FGeometry KeyAreaGeometryPadded = GetKeyAreaGeometry( Element, AllottedGeometry );
+
+		// Is the key area under the mouse
+		if( KeyAreaGeometryPadded.IsUnderLocation( MousePosition ) )
+		{
+			FTimeToPixel TimeToPixelConverter = Section.IsInfinite() ? 			
+				FTimeToPixel( ParentGeometry, GetSequencer().GetViewRange()) : 
+				FTimeToPixel( SectionGeometry, TRange<float>( Section.GetStartTime(), Section.GetEndTime() ) );
+
+			// If a key was pressed on, get the pressed on key's time to duplicate that key
+			float KeyTime;
+			if (InPressedKey.IsValid())
+			{
+				KeyTime = KeyArea->GetKeyTime(InPressedKey.KeyHandle.GetValue());
+			}
+			// Otherwise, use the time where the mouse is pressed
+			else
+			{
+				FVector2D LocalSpaceMousePosition = SectionGeometry.AbsoluteToLocal( MousePosition );
+				KeyTime = TimeToPixelConverter.PixelToTime(LocalSpaceMousePosition.X);
+			}
+
+			FScopedTransaction CreateKeyTransaction(NSLOCTEXT("Sequencer", "CreateKey_Transaction", "Create Key"));
+
+			Section.Modify();
+
+			// If the pressed key exists, offset the new key and look for it in the newly laid out key areas
+			if (InPressedKey.IsValid())
+			{
+				// Offset by 1 pixel worth of time
+				const float TimeFuzz = (GetSequencer().GetViewRange().GetUpperBoundValue() - GetSequencer().GetViewRange().GetLowerBoundValue()) / ParentGeometry.GetLocalSize().X;
+
+				TArray<FKeyHandle> KeyHandles = KeyArea->AddKeyUnique(KeyTime+TimeFuzz, KeyTime);
+
+				Layout = FKeyAreaLayout(*ParentSectionArea, SectionIndex);
+
+				// Look specifically for the key with the offset key time
+				for (const FKeyAreaLayoutElement& NewElement : Layout->GetElements())
+				{
+					TSharedPtr<IKeyArea> NewKeyArea = NewElement.GetKeyArea();
+
+					for (auto KeyHandle : KeyHandles)
+					{
+						for (auto UnsortedKeyHandle : NewKeyArea->GetUnsortedKeyHandles())
+						{
+							if (FMath::IsNearlyEqual(KeyTime+TimeFuzz, NewKeyArea->GetKeyTime(UnsortedKeyHandle), KINDA_SMALL_NUMBER))
+							{
+								return FSelectedKey(Section, NewKeyArea, UnsortedKeyHandle);
+							}
+						}
+					}
+				}
+			}
+			else
+			{
+				KeyArea->AddKeyUnique(KeyTime);
+
+				Layout = FKeyAreaLayout(*ParentSectionArea, SectionIndex);
+						
+				return GetKeyUnderMouse(MousePosition, AllottedGeometry);
+			}
+		}
+	}
+
+	return FSelectedKey();
+}
+
 void SSection::CheckForEdgeInteraction( const FPointerEvent& MouseEvent, const FGeometry& SectionGeometry )
 {
 	bLeftEdgeHovered = false;
@@ -629,11 +708,8 @@ FReply SSection::OnMouseButtonDown( const FGeometry& MyGeometry, const FPointerE
 {
 	FSequencer& Sequencer = GetSequencer();
 
-	if( MouseEvent.GetEffectingButton() == EKeys::LeftMouseButton || MouseEvent.GetEffectingButton() == EKeys::RightMouseButton )
-	{
-		// Check for clicking on a key and mark it as the pressed key for drag detection (if necessary) later
-		PressedKey = GetKeyUnderMouse( MouseEvent.GetScreenSpacePosition(), MyGeometry );
-	}
+	// Check for clicking on a key and mark it as the pressed key for drag detection (if necessary) later
+	PressedKey = GetKeyUnderMouse( MouseEvent.GetScreenSpacePosition(), MyGeometry );
 
 	if (MouseEvent.GetEffectingButton() == EKeys::LeftMouseButton)
 	{
@@ -642,6 +718,31 @@ FReply SSection::OnMouseButtonDown( const FGeometry& MyGeometry, const FPointerE
 			CheckForEdgeInteraction( MouseEvent, MyGeometry );
 		}
 	}
+	else if (MouseEvent.GetEffectingButton() == EKeys::MiddleMouseButton)
+	{
+		// Generate a key and set it as the PressedKey
+		PressedKey = CreateKeyUnderMouse(MouseEvent.GetScreenSpacePosition(), MyGeometry, PressedKey);
+		HoveredKey = PressedKey;
+
+		Sequencer.GetSelection().EmptySelectedKeys();
+		Sequencer.GetSelection().AddToSelection(PressedKey);
+
+		// Pass the event to the tool to copy the hovered key and move it
+		auto& EditTool = GetSequencer().GetEditTool();
+
+		EditTool.SetHotspot( MakeShareable( new FSectionHotspot(FSectionHandle(ParentSectionArea, SectionIndex))) );
+
+		if ( HoveredKey.IsValid() )
+		{
+			EditTool.SetHotspot( MakeShareable( new FKeyHotspot(HoveredKey, ParentSectionArea) ) );
+			FReply Reply = EditTool.OnMouseButtonDown(*this, MyGeometry, MouseEvent);
+			if (Reply.IsEventHandled())
+			{
+				return Reply;
+			}
+		}
+		return FReply::Unhandled();
+	}
 	else if (MouseEvent.GetEffectingButton() == EKeys::RightMouseButton && (PressedKey.IsValid() || Sequencer.GetSelection().GetSelectedKeys().Num() || Sequencer.GetSelection().GetSelectedSections().Num()))
 	{
 		// Don't return handled if we didn't click on a key so that right-click panning gets a look in
diff --git a/Engine/Source/Editor/Sequencer/Private/SSection.h b/Engine/Source/Editor/Sequencer/Private/SSection.h
--- a/Engine/Source/Editor/Sequencer/Private/SSection.h
+++ b/Engine/Source/Editor/Sequencer/Private/SSection.h
@@ -47,6 +47,16 @@ private:
 	 */
 	FSelectedKey GetKeyUnderMouse( const FVector2D& MousePosition, const FGeometry& AllottedGeometry ) const;
 
+	/**
+	 * Creates a key at the mouse position
+	 *
+	 * @param MousePosition		The current screen space position of the mouse
+	 * @param AllottedGeometry	The geometry of the mouse event
+	 * @param InPressedKey      Key if pressed
+	 * @return The newly created key
+	 */
+	FSelectedKey CreateKeyUnderMouse( const FVector2D& MousePosition, const FGeometry& AllottedGeometry, FSelectedKey InPressedKey );
+
 	/**
 	 * Checks for user interaction (via the mouse) with the left and right edge of a section
 	 *
diff --git a/Engine/Source/Editor/Sequencer/Private/Tools/SequencerEditTool_Movement.cpp b/Engine/Source/Editor/Sequencer/Private/Tools/SequencerEditTool_Movement.cpp
--- a/Engine/Source/Editor/Sequencer/Private/Tools/SequencerEditTool_Movement.cpp
+++ b/Engine/Source/Editor/Sequencer/Private/Tools/SequencerEditTool_Movement.cpp
@@ -30,9 +30,11 @@ FReply FSequencerEditTool_Movement::OnMouseButtonDown(SWidget& OwnerWidget, cons
 {
 	DelayedDrag.Reset();
 
-	if (MouseEvent.GetEffectingButton() == EKeys::LeftMouseButton)
+	if (MouseEvent.GetEffectingButton() == EKeys::LeftMouseButton || MouseEvent.GetEffectingButton() == EKeys::MiddleMouseButton)
 	{
-		DelayedDrag = FDelayedDrag_Hotspot(MyGeometry.AbsoluteToLocal(MouseEvent.GetScreenSpacePosition()), EKeys::LeftMouseButton, Hotspot);
+		const FVirtualTrackArea VirtualTrackArea = SequencerWidget.Pin()->GetVirtualTrackArea();
+
+		DelayedDrag = FDelayedDrag_Hotspot(VirtualTrackArea.CachedTrackAreaGeometry().AbsoluteToLocal(MouseEvent.GetScreenSpacePosition()), MouseEvent.GetEffectingButton(), Hotspot);
 		return FReply::Handled();
 	}
 	return FReply::Unhandled();
diff --git a/Engine/Source/Editor/Sequencer/Public/IKeyArea.h b/Engine/Source/Editor/Sequencer/Public/IKeyArea.h
--- a/Engine/Source/Editor/Sequencer/Public/IKeyArea.h
+++ b/Engine/Source/Editor/Sequencer/Public/IKeyArea.h
@@ -110,8 +110,11 @@ public:
 	/**
 	 * Adds a key at the specified time if there isn't already a key present.  The value of the added key should
 	 * be the value which would be returned if the animation containing this key area was evaluated at the specified time.
+	 *
+	 * @param TimeToCopyFrom Optional time to copy key parameters from
+	 * @param return      The new keys that were added
 	 */
-	virtual void AddKeyUnique(float Time) = 0;
+	virtual TArray<FKeyHandle> AddKeyUnique(float Time, float TimeToCopyFrom = FLT_MAX) = 0;
 
 	/** Gets the rich curve associated with this key area.  This can be null but it must be valid in order to be
 	  * edited by the curve editor. */
diff --git a/Engine/Source/Editor/Sequencer/Public/VirtualTrackArea.h b/Engine/Source/Editor/Sequencer/Public/VirtualTrackArea.h
--- a/Engine/Source/Editor/Sequencer/Public/VirtualTrackArea.h
+++ b/Engine/Source/Editor/Sequencer/Public/VirtualTrackArea.h
@@ -40,6 +40,9 @@ public:
 	/** Hit test at the specified physical position for a key */
 	FSelectedKey HitTestKey(FVector2D InPhysicalPosition) const;
 
+	/** Cached track area geometry */
+	FGeometry CachedTrackAreaGeometry() const { return TrackAreaGeometry; }
+
 private:
 
 	/** Reference to the sequencer tree */
