Sequencer: Change default expansion state so that only objects are expanded.

diff --git a/Engine/Source/Editor/Sequencer/Private/SequencerNodeTree.cpp b/Engine/Source/Editor/Sequencer/Private/SequencerNodeTree.cpp
--- a/Engine/Source/Editor/Sequencer/Private/SequencerNodeTree.cpp
+++ b/Engine/Source/Editor/Sequencer/Private/SequencerNodeTree.cpp
@@ -272,8 +272,8 @@ bool FSequencerNodeTree::GetSavedExpansionState( const FSequencerDisplayNode& No
 
 bool FSequencerNodeTree::GetDefaultExpansionState( const FSequencerDisplayNode& Node ) const
 {
-	// For now, all types except categories and key areas are expanded by default
-	return Node.GetType() != ESequencerNode::Category && Node.GetType() != ESequencerNode::KeyArea;
+	// For now, only object nodes are expanded by default
+	return Node.GetType() == ESequencerNode::Object;
 }
 
 bool FSequencerNodeTree::IsNodeFiltered( const TSharedRef<const FSequencerDisplayNode> Node ) const
