Sequencer: Add hotkeys for setting the interp/tangent mode of the selected keys.

diff --git a/Engine/Source/Editor/Sequencer/Private/Sequencer.cpp b/Engine/Source/Editor/Sequencer/Private/Sequencer.cpp
--- a/Engine/Source/Editor/Sequencer/Private/Sequencer.cpp
+++ b/Engine/Source/Editor/Sequencer/Private/Sequencer.cpp
@@ -455,9 +455,14 @@ void FSequencer::DeleteSelectedKeys()
 
 void FSequencer::SetInterpTangentMode(ERichCurveInterpMode InterpMode, ERichCurveTangentMode TangentMode)
 {
+	TArray<FSelectedKey> SelectedKeysArray = Selection.GetSelectedKeys().Array();
+	if (SelectedKeysArray.Num() == 0)
+	{
+		return;
+	}
+
 	FScopedTransaction SetInterpTangentModeTransaction(NSLOCTEXT("Sequencer", "SetInterpTangentMode_Transaction", "Set Interpolation and Tangent Mode"));
 	bool bAnythingChanged = false;
-	TArray<FSelectedKey> SelectedKeysArray = Selection.GetSelectedKeys().Array();
 
 	for (const FSelectedKey& Key : SelectedKeysArray)
 	{
@@ -2263,6 +2268,26 @@ void FSequencer::BindSequencerCommands()
 		Commands.SetKey,
 		FExecuteAction::CreateSP( this, &FSequencer::SetKey ) );
 
+	SequencerCommandBindings->MapAction(
+		Commands.SetInterpolationCubicAuto,
+		FExecuteAction::CreateSP( this, &FSequencer::SetInterpTangentMode, ERichCurveInterpMode::RCIM_Cubic, ERichCurveTangentMode::RCTM_Auto ) );
+
+	SequencerCommandBindings->MapAction(
+		Commands.SetInterpolationCubicUser,
+		FExecuteAction::CreateSP( this, &FSequencer::SetInterpTangentMode, ERichCurveInterpMode::RCIM_Cubic, ERichCurveTangentMode::RCTM_User ) );
+
+	SequencerCommandBindings->MapAction(
+		Commands.SetInterpolationCubicBreak,
+		FExecuteAction::CreateSP( this, &FSequencer::SetInterpTangentMode, ERichCurveInterpMode::RCIM_Cubic, ERichCurveTangentMode::RCTM_Break ) );
+
+	SequencerCommandBindings->MapAction(
+		Commands.SetInterpolationLinear,
+		FExecuteAction::CreateSP( this, &FSequencer::SetInterpTangentMode, ERichCurveInterpMode::RCIM_Linear, ERichCurveTangentMode::RCTM_Auto ) );
+
+	SequencerCommandBindings->MapAction(
+		Commands.SetInterpolationConstant,
+		FExecuteAction::CreateSP( this, &FSequencer::SetInterpTangentMode, ERichCurveInterpMode::RCIM_Constant, ERichCurveTangentMode::RCTM_Auto ) );
+
 	SequencerCommandBindings->MapAction(
 		Commands.ToggleAutoKeyEnabled,
 		FExecuteAction::CreateLambda( [this]{ Settings->SetAutoKeyEnabled( !Settings->GetAutoKeyEnabled() ); } ),
diff --git a/Engine/Source/Editor/Sequencer/Private/SequencerCommands.cpp b/Engine/Source/Editor/Sequencer/Private/SequencerCommands.cpp
--- a/Engine/Source/Editor/Sequencer/Private/SequencerCommands.cpp
+++ b/Engine/Source/Editor/Sequencer/Private/SequencerCommands.cpp
@@ -23,6 +23,11 @@ void FSequencerCommands::RegisterCommands()
 	UI_COMMAND( ToggleExpandCollapseNodesAndDescendants, "Toggle Expand/Collapse Nodes and Descendants", "Toggle expand or collapse selected nodes and their descendants", EUserInterfaceActionType::Button, FInputChord(EModifierKey::Shift, EKeys::O) );
 
 	UI_COMMAND( SetKey, "Set Key", "Sets a key at the current time for the selected actor", EUserInterfaceActionType::Button, FInputChord(EKeys::K) );
+	UI_COMMAND( SetInterpolationCubicAuto, "Set Key Auto", "Cubic interpolation - Automatic tangents", EUserInterfaceActionType::Button, FInputChord(EKeys::One));
+	UI_COMMAND( SetInterpolationCubicUser, "Set Key User", "Cubic interpolation - User flat tangents", EUserInterfaceActionType::Button, FInputChord(EKeys::Two));
+	UI_COMMAND( SetInterpolationCubicBreak, "Set Key Break", "Cubic interpolation - User broken tangents", EUserInterfaceActionType::Button, FInputChord(EKeys::Three));
+	UI_COMMAND( SetInterpolationLinear, "Set Key Linear", "Linear interpolation", EUserInterfaceActionType::Button, FInputChord(EKeys::Four));
+	UI_COMMAND( SetInterpolationConstant, "Set Key Constant", "Constant interpolation", EUserInterfaceActionType::Button, FInputChord(EKeys::Five));
 
 	UI_COMMAND( ToggleAutoKeyEnabled, "Auto Key", "Enables and disables auto keying", EUserInterfaceActionType::ToggleButton, FInputChord() );
 	UI_COMMAND( ToggleKeyAllEnabled, "Key All", "Enables and disables key all", EUserInterfaceActionType::ToggleButton, FInputChord() );
diff --git a/Engine/Source/Editor/Sequencer/Private/SequencerCommands.h b/Engine/Source/Editor/Sequencer/Private/SequencerCommands.h
--- a/Engine/Source/Editor/Sequencer/Private/SequencerCommands.h
+++ b/Engine/Source/Editor/Sequencer/Private/SequencerCommands.h
@@ -57,6 +57,21 @@ public:
 	/** Sets a key at the current time for the selected actor */
 	TSharedPtr< FUICommandInfo > SetKey;
 
+	/** Sets the interp tangent mode for the selected keys to auto */
+	TSharedPtr< FUICommandInfo > SetInterpolationCubicAuto;
+
+	/** Sets the interp tangent mode for the selected keys to user */
+	TSharedPtr< FUICommandInfo > SetInterpolationCubicUser;
+
+	/** Sets the interp tangent mode for the selected keys to break */
+	TSharedPtr< FUICommandInfo > SetInterpolationCubicBreak;
+
+	/** Sets the interp tangent mode for the selected keys to linear */
+	TSharedPtr< FUICommandInfo > SetInterpolationLinear;
+
+	/** Sets the interp tangent mode for the selected keys to constant */
+	TSharedPtr< FUICommandInfo > SetInterpolationConstant;
+
 	/** Turns auto keying on and off. */
 	TSharedPtr< FUICommandInfo > ToggleAutoKeyEnabled;
 
diff --git a/Engine/Source/Editor/UnrealEd/Private/RichCurveEditorCommands.cpp b/Engine/Source/Editor/UnrealEd/Private/RichCurveEditorCommands.cpp
--- a/Engine/Source/Editor/UnrealEd/Private/RichCurveEditorCommands.cpp
+++ b/Engine/Source/Editor/UnrealEd/Private/RichCurveEditorCommands.cpp
@@ -13,11 +13,11 @@ void FRichCurveEditorCommands::RegisterCommands()
 
 	UI_COMMAND(ToggleSnapping, "Snapping", "Toggle Snapping", EUserInterfaceActionType::ToggleButton, FInputChord());
 
-	UI_COMMAND(InterpolationConstant, "Constant", "Constant interpolation", EUserInterfaceActionType::ToggleButton, FInputChord());
-	UI_COMMAND(InterpolationLinear, "Linear", "Linear interpolation", EUserInterfaceActionType::ToggleButton, FInputChord());
-	UI_COMMAND(InterpolationCubicAuto, "Auto", "Cubic interpolation - Automatic tangents", EUserInterfaceActionType::ToggleButton, FInputChord());
-	UI_COMMAND(InterpolationCubicUser, "User", "Cubic interpolation - User flat tangents", EUserInterfaceActionType::ToggleButton, FInputChord());
-	UI_COMMAND(InterpolationCubicBreak, "Break", "Cubic interpolation - User broken tangents", EUserInterfaceActionType::ToggleButton, FInputChord());
+	UI_COMMAND(InterpolationConstant, "Constant", "Constant interpolation", EUserInterfaceActionType::ToggleButton, FInputChord(EKeys::Five));
+	UI_COMMAND(InterpolationLinear, "Linear", "Linear interpolation", EUserInterfaceActionType::ToggleButton, FInputChord(EKeys::Four));
+	UI_COMMAND(InterpolationCubicAuto, "Auto", "Cubic interpolation - Automatic tangents", EUserInterfaceActionType::ToggleButton, FInputChord(EKeys::One));
+	UI_COMMAND(InterpolationCubicUser, "User", "Cubic interpolation - User flat tangents", EUserInterfaceActionType::ToggleButton, FInputChord(EKeys::Two));
+	UI_COMMAND(InterpolationCubicBreak, "Break", "Cubic interpolation - User broken tangents", EUserInterfaceActionType::ToggleButton, FInputChord(EKeys::Three));
 
 	UI_COMMAND(FlattenTangents, "Flatten", "Flatten tangents", EUserInterfaceActionType::Button, FInputChord());
 	UI_COMMAND(StraightenTangents, "Straighten", "Straighten tangents", EUserInterfaceActionType::Button, FInputChord());
