Sequencer: Update runtime instances and redraw when dragging sections/keys.

diff --git a/Engine/Source/Editor/Sequencer/Private/Sequencer.cpp b/Engine/Source/Editor/Sequencer/Private/Sequencer.cpp
--- a/Engine/Source/Editor/Sequencer/Private/Sequencer.cpp
+++ b/Engine/Source/Editor/Sequencer/Private/Sequencer.cpp
@@ -944,6 +944,8 @@ void FSequencer::UpdateRuntimeInstances()
 	// Refresh the current root instance
 	RootMovieSceneSequenceInstance->RefreshInstance( *this );
 	RootMovieSceneSequenceInstance->Update(ScrubPosition, ScrubPosition, *this);
+
+	GEditor->RedrawLevelEditingViewports();
 }
 
 void FSequencer::DestroySpawnablesForAllSequences()
diff --git a/Engine/Source/Editor/Sequencer/Private/Tools/EditToolDragOperations.cpp b/Engine/Source/Editor/Sequencer/Private/Tools/EditToolDragOperations.cpp
--- a/Engine/Source/Editor/Sequencer/Private/Tools/EditToolDragOperations.cpp
+++ b/Engine/Source/Editor/Sequencer/Private/Tools/EditToolDragOperations.cpp
@@ -244,6 +244,8 @@ void FResizeSection::OnDrag(const FPointerEvent& MouseEvent, FVector2D LocalMous
 			OuterTrack->OnSectionMoved(*Section);
 		}
 	}
+
+	Sequencer.UpdateRuntimeInstances();
 }
 
 FMoveSection::FMoveSection( FSequencer& InSequencer, TArray<FSectionHandle> InSections )
@@ -424,6 +426,8 @@ void FMoveSection::OnDrag(const FPointerEvent& MouseEvent, FVector2D LocalMouseP
 			}
 		}
 	}
+
+	Sequencer.UpdateRuntimeInstances();
 }
 
 void FMoveKeys::OnBeginDrag(const FPointerEvent& MouseEvent, FVector2D LocalMousePos, const FVirtualTrackArea& VirtualTrackArea)
@@ -540,6 +544,8 @@ void FMoveKeys::OnDrag(const FPointerEvent& MouseEvent, FVector2D LocalMousePos,
 	{
 		Sequencer.SetGlobalTime(PrevNewKeyTime);
 	}
+
+	Sequencer.UpdateRuntimeInstances();
 }
 
 void FMoveKeys::OnEndDrag(const FPointerEvent& MouseEvent, FVector2D LocalMousePos, const FVirtualTrackArea& VirtualTrackArea)
