Sequencer: Fix SetKey so that it keys selected tracks.

diff --git a/Engine/Source/Editor/Sequencer/Private/Sequencer.cpp b/Engine/Source/Editor/Sequencer/Private/Sequencer.cpp
--- a/Engine/Source/Editor/Sequencer/Private/Sequencer.cpp
+++ b/Engine/Source/Editor/Sequencer/Private/Sequencer.cpp
@@ -2191,21 +2191,21 @@ void FSequencer::ToggleExpandCollapseNodesAndDescendants()
 
 void FSequencer::SetKey()
 {
-	USelection* CurrentSelection = GEditor->GetSelectedActors();
-	TArray<UObject*> SelectedActors;
-	CurrentSelection->GetSelectedObjects( AActor::StaticClass(), SelectedActors );
-	for (TArray<UObject*>::TIterator It(SelectedActors); It; ++It)
+	TSet<TSharedPtr<IKeyArea> > KeyAreas;
+	for (auto OutlinerNode : Selection.GetSelectedOutlinerNodes())
 	{
-		// @todo Handle case of actors which aren't in sequencer yet
+		SequencerHelpers::GetAllKeyAreas(OutlinerNode, KeyAreas);
+	}
 
-		FGuid ObjectGuid = GetHandleToObject(*It);
-		for ( auto& TrackEditor : TrackEditors )
+	if (KeyAreas.Num() > 0)
+	{
+		FScopedTransaction SetKeyTransaction( NSLOCTEXT("Sequencer", "SetKey_Transaction", "Set Key") );
+	
+		for (auto KeyArea : KeyAreas)
 		{
-			// @todo Handle this shot track business better
-			if (TrackEditor != ShotTrackEditor.Pin())
-			{
-				TrackEditor->AddKey(ObjectGuid);
-			}
+			KeyArea->GetOwningSection()->Modify();
+	
+			KeyArea->AddKeyUnique(GetGlobalTime(), GetKeyInterpolation());
 		}
 	}
 }
diff --git a/Engine/Source/Editor/Sequencer/Private/SequencerCommands.cpp b/Engine/Source/Editor/Sequencer/Private/SequencerCommands.cpp
--- a/Engine/Source/Editor/Sequencer/Private/SequencerCommands.cpp
+++ b/Engine/Source/Editor/Sequencer/Private/SequencerCommands.cpp
@@ -22,7 +22,7 @@ void FSequencerCommands::RegisterCommands()
 	UI_COMMAND( ToggleExpandCollapseNodes, "Toggle Expand/Collapse Nodes", "Toggle expand or collapse selected nodes", EUserInterfaceActionType::Button, FInputChord(EKeys::O) );
 	UI_COMMAND( ToggleExpandCollapseNodesAndDescendants, "Toggle Expand/Collapse Nodes and Descendants", "Toggle expand or collapse selected nodes and their descendants", EUserInterfaceActionType::Button, FInputChord(EModifierKey::Shift, EKeys::O) );
 
-	UI_COMMAND( SetKey, "Set Key", "Sets a key at the current time for the selected actor", EUserInterfaceActionType::Button, FInputChord(EKeys::K) );
+	UI_COMMAND( SetKey, "Set Key", "Sets a key on the selected tracks", EUserInterfaceActionType::Button, FInputChord(EKeys::K) );
 	UI_COMMAND( SetInterpolationCubicAuto, "Set Key Auto", "Cubic interpolation - Automatic tangents", EUserInterfaceActionType::Button, FInputChord(EKeys::One));
 	UI_COMMAND( SetInterpolationCubicUser, "Set Key User", "Cubic interpolation - User flat tangents", EUserInterfaceActionType::Button, FInputChord(EKeys::Two));
 	UI_COMMAND( SetInterpolationCubicBreak, "Set Key Break", "Cubic interpolation - User broken tangents", EUserInterfaceActionType::Button, FInputChord(EKeys::Three));
