Sequencer: Lower the frame rounding tolerance.

diff --git a/Engine/Source/Editor/Sequencer/Private/TimeSliderController.cpp b/Engine/Source/Editor/Sequencer/Private/TimeSliderController.cpp
--- a/Engine/Source/Editor/Sequencer/Private/TimeSliderController.cpp
+++ b/Engine/Source/Editor/Sequencer/Private/TimeSliderController.cpp
@@ -307,7 +307,8 @@ int32 FSequencerTimeSliderController::OnPaintTimeSlider( bool bMirrorLabels, con
 			float FrameTime = Time * FrameRate;
 			int32 Frame = SequencerHelpers::TimeToFrame(Time, FrameRate);
 					
-			if (FMath::IsNearlyEqual(FrameTime, (float)Frame, KINDA_SMALL_NUMBER))
+			const float FrameTolerance = 0.001f;
+			if (FMath::IsNearlyEqual(FrameTime, (float)Frame, FrameTolerance))
 			{
 				FrameString = FString::Printf( TEXT("%d"), TimeToFrame(Time));
 			}
