Sequencer: Only show right click menu when section is selected.

diff --git a/Engine/Source/Editor/Sequencer/Private/SSection.cpp b/Engine/Source/Editor/Sequencer/Private/SSection.cpp
--- a/Engine/Source/Editor/Sequencer/Private/SSection.cpp
+++ b/Engine/Source/Editor/Sequencer/Private/SSection.cpp
@@ -487,8 +487,10 @@ TSharedPtr<SWidget> SSection::OnSummonContextMenu( const FGeometry& MyGeometry,
 	const bool bShouldCloseWindowAfterMenuSelection = true;
 	FMenuBuilder MenuBuilder(bShouldCloseWindowAfterMenuSelection, NULL);
 
+	bool bHasMenuItems = false;
 	if (Key.IsValid() || Sequencer.GetSelection().GetSelectedKeys().Num())
 	{
+		bHasMenuItems = true;
 		MenuBuilder.BeginSection("SequencerInterpolation", NSLOCTEXT("Sequencer", "KeyInterpolationMenu", "Key Interpolation"));
 		{
 			MenuBuilder.AddMenuEntry(
@@ -572,8 +574,9 @@ TSharedPtr<SWidget> SSection::OnSummonContextMenu( const FGeometry& MyGeometry,
 		MenuBuilder.EndSection(); // SequencerKeys
 	}
 	
-	if (SceneSection != nullptr)
+	if (SceneSection != nullptr && Sequencer.GetSelection().IsSelected(SceneSection))
 	{
+		bHasMenuItems = true;
 		SectionInterface->BuildSectionContextMenu(MenuBuilder);
 
 		MenuBuilder.BeginSection("SequencerSections", NSLOCTEXT("Sequencer", "SectionsMenu", "Sections"));
@@ -621,7 +624,12 @@ TSharedPtr<SWidget> SSection::OnSummonContextMenu( const FGeometry& MyGeometry,
 	
 	ResetState();
 
-	return MenuBuilder.MakeWidget();
+	if (bHasMenuItems)
+	{
+		return MenuBuilder.MakeWidget();
+	}
+
+	return TSharedPtr<SWidget>();
 }
 
 void SSection::AddExtrapolationMenu(FMenuBuilder& MenuBuilder, bool bPreInfinity)
