Sequencer: Animate only editable properties.

diff --git a/Engine/Source/Editor/Sequencer/Private/SequencerObjectChangeListener.cpp b/Engine/Source/Editor/Sequencer/Private/SequencerObjectChangeListener.cpp
--- a/Engine/Source/Editor/Sequencer/Private/SequencerObjectChangeListener.cpp
+++ b/Engine/Source/Editor/Sequencer/Private/SequencerObjectChangeListener.cpp
@@ -142,24 +142,45 @@ bool FSequencerObjectChangeListener::FindPropertySetter( const UClass& ObjectCla
 		}
 
 		bool bFoundValidInterp = false;
+		bool bFoundEditDefaultsOnly = false;
+		bool bFoundEdit = false;
 		if (StructProperty != 0)
 		{
 			if (StructProperty->HasAnyPropertyFlags(CPF_Interp))
 			{
 				bFoundValidInterp = true;
 			}
+			if (StructProperty->HasAnyPropertyFlags(CPF_DisableEditOnInstance))
+			{
+				bFoundEditDefaultsOnly = true;
+			}
+			if (StructProperty->HasAnyPropertyFlags(CPF_Edit))
+			{
+				bFoundEdit = true;
+			}
 		}
 		else
 		{
 			UProperty* Property = ObjectClass.FindPropertyByName(FName(*InPropertyVarName));
-			if (Property && Property->HasAnyPropertyFlags(CPF_Interp))
+			if (Property)
 			{
-				bFoundValidInterp = true;
+				if (Property->HasAnyPropertyFlags(CPF_Interp))
+				{
+					bFoundValidInterp = true;
+				}
+				if (Property->HasAnyPropertyFlags(CPF_DisableEditOnInstance))
+				{
+					bFoundEditDefaultsOnly = true;
+				}
+				if (Property->HasAnyPropertyFlags(CPF_Edit))
+				{
+					bFoundEdit = true;
+				}
 			}
 		}
 		
-		// Valid either if a setter function is foud or a property with the Interp keyword is found.
-		bFound = bFoundValidFunction || bFoundValidInterp;
+		// Valid either if a setter function is found or a property with the Interp keyword is found. The property also needs to be editable.
+		bFound = (bFoundValidFunction || bFoundValidInterp) && bFoundEdit && !bFoundEditDefaultsOnly;
 	}
 
 	return bFound;
