Sequencer: Fix transaction when deleting both keys and sections.

diff --git a/Engine/Source/Editor/Sequencer/Private/Sequencer.cpp b/Engine/Source/Editor/Sequencer/Private/Sequencer.cpp
--- a/Engine/Source/Editor/Sequencer/Private/Sequencer.cpp
+++ b/Engine/Source/Editor/Sequencer/Private/Sequencer.cpp
@@ -2031,6 +2031,8 @@ void FSequencer::DeleteSelectedItems()
 {
 	if (Selection.GetActiveSelection() == FSequencerSelection::EActiveSelection::KeyAndSection)
 	{
+		FScopedTransaction DeleteKeysTransaction( NSLOCTEXT("Sequencer", "DeleteKeysAndSections_Transaction", "Delete Keys and Sections") );
+		
 		DeleteSelectedKeys();
 
 		DeleteSections(Selection.GetSelectedSections());
