Sequencer: Fix section drawing when collapsed. Audio sections were being drawn with a background section when the display node is collapsed.

diff --git a/Engine/Source/Editor/Sequencer/Private/SSection.cpp b/Engine/Source/Editor/Sequencer/Private/SSection.cpp
--- a/Engine/Source/Editor/Sequencer/Private/SSection.cpp
+++ b/Engine/Source/Editor/Sequencer/Private/SSection.cpp
@@ -322,6 +322,7 @@ void SSection::PaintKeys( const FGeometry& AllottedGeometry, const FSlateRect& M
 	{
 		// Get the key area at the same index of the section.  Each section in this widget has the same layout and the same number of key areas
 		TSharedPtr<IKeyArea> KeyArea = Element.GetKeyArea();
+		TArray<FKeyHandle> KeyHandles = KeyArea->GetUnsortedKeyHandles();
 
 		FGeometry KeyAreaGeometry = GetKeyAreaGeometry( Element, AllottedGeometry );
 
@@ -331,19 +332,20 @@ void SSection::PaintKeys( const FGeometry& AllottedGeometry, const FSlateRect& M
 
 		// Draw a box for the key area 
 		// @todo Sequencer - Allow the IKeyArea to do this
-		FSlateDrawElement::MakeBox( 
-			OutDrawElements,
-			LayerId,
-			KeyAreaGeometry.ToPaintGeometry(),
-			BackgroundBrush,
-			MyClippingRect,
-			DrawEffects,
-			FLinearColor( .1f, .1f, .1f, 0.7f ) ); 
-
+		if (KeyHandles.Num())
+		{
+			FSlateDrawElement::MakeBox( 
+				OutDrawElements,
+				LayerId,
+				KeyAreaGeometry.ToPaintGeometry(),
+				BackgroundBrush,
+				MyClippingRect,
+				DrawEffects,
+				FLinearColor( .1f, .1f, .1f, 0.7f ) ); 
+		}
 
 		int32 KeyLayer = LayerId + 1;
 
-		TArray<FKeyHandle> KeyHandles = KeyArea->GetUnsortedKeyHandles();
 		for( int32 KeyIndex = 0; KeyIndex < KeyHandles.Num(); ++KeyIndex )
 		{
 			FKeyHandle KeyHandle = KeyHandles[KeyIndex];
