Sequencer: Add Find in Content Browser to Sequencer view menu.

diff --git a/Engine/Source/Editor/Sequencer/Private/SSequencer.cpp b/Engine/Source/Editor/Sequencer/Private/SSequencer.cpp
--- a/Engine/Source/Editor/Sequencer/Private/SSequencer.cpp
+++ b/Engine/Source/Editor/Sequencer/Private/SSequencer.cpp
@@ -766,7 +766,14 @@ TSharedRef<SWidget> SSequencer::MakeGeneralMenu()
 		{
 			MenuBuilder.AddMenuEntry( FSequencerCommands::Get().ToggleCleanView );
 		}
+
 		MenuBuilder.AddMenuEntry( FSequencerCommands::Get().ToggleAutoScroll );
+				
+		if (Sequencer.Pin()->IsLevelEditorSequencer())
+		{
+			MenuBuilder.AddMenuEntry( FSequencerCommands::Get().FindInContentBrowser );
+		}
+
 		MenuBuilder.AddMenuEntry( FSequencerCommands::Get().ExpandNodesAndDescendants );
 		MenuBuilder.AddMenuEntry( FSequencerCommands::Get().CollapseNodesAndDescendants );
 	}
diff --git a/Engine/Source/Editor/Sequencer/Private/Sequencer.cpp b/Engine/Source/Editor/Sequencer/Private/Sequencer.cpp
--- a/Engine/Source/Editor/Sequencer/Private/Sequencer.cpp
+++ b/Engine/Source/Editor/Sequencer/Private/Sequencer.cpp
@@ -1332,6 +1332,17 @@ void FSequencer::OnToggleAutoScroll()
 	Settings->SetAutoScrollEnabled(bAutoScrollEnabled);
 }
 
+void FSequencer::FindInContentBrowser()
+{
+	if (GetFocusedMovieSceneSequence())
+	{
+		TArray<UObject*> ObjectsToFocus;
+		ObjectsToFocus.Add(GetFocusedMovieSceneSequence());
+
+		GEditor->SyncBrowserToObjects(ObjectsToFocus);
+	}
+}
+
 void FSequencer::VerticalScroll(float ScrollAmountUnits)
 {
 	SequencerWidget->GetTreeView()->ScrollByDelta(ScrollAmountUnits);
@@ -2321,6 +2332,10 @@ void FSequencer::BindSequencerCommands()
 		FCanExecuteAction::CreateLambda( []{ return true; } ),
 		FIsActionChecked::CreateSP( this, &FSequencer::GetAutoScrollEnabled ) );
 
+	SequencerCommandBindings->MapAction(
+		Commands.FindInContentBrowser,
+		FExecuteAction::CreateSP( this, &FSequencer::FindInContentBrowser ) );
+
 	SequencerCommandBindings->MapAction(
 		Commands.ToggleCleanView,
 		FExecuteAction::CreateLambda( [this]{
diff --git a/Engine/Source/Editor/Sequencer/Private/Sequencer.h b/Engine/Source/Editor/Sequencer/Private/Sequencer.h
--- a/Engine/Source/Editor/Sequencer/Private/Sequencer.h
+++ b/Engine/Source/Editor/Sequencer/Private/Sequencer.h
@@ -393,6 +393,11 @@ protected:
 		return bAutoScrollEnabled;
 	}
 
+	/** 
+	 * Find the viewed actor animation asset in the content browser
+	 */
+	void FindInContentBrowser();
+
 protected:
 
 	/** Called via UEditorEngine::GetActorRecordingStateEvent to check to see whether we need to record actor state */
diff --git a/Engine/Source/Editor/Sequencer/Private/SequencerCommands.cpp b/Engine/Source/Editor/Sequencer/Private/SequencerCommands.cpp
--- a/Engine/Source/Editor/Sequencer/Private/SequencerCommands.cpp
+++ b/Engine/Source/Editor/Sequencer/Private/SequencerCommands.cpp
@@ -48,8 +48,9 @@ void FSequencerCommands::RegisterCommands()
 	UI_COMMAND( ToggleSnapPlayTimeToInterval, "Snap to the Interval While Scrubbing", "Snap the play time to the interval while scrubbing", EUserInterfaceActionType::ToggleButton, FInputChord() );
 	UI_COMMAND( ToggleSnapPlayTimeToDraggedKey, "Snap to the Dragged Key", "Snap the play time to the dragged key", EUserInterfaceActionType::ToggleButton, FInputChord() );
 
-	UI_COMMAND( ToggleSnapCurveValueToInterval, "Snap Curve Key Values", "Snap curve keys to the value snapping interval", EUserInterfaceActionType::ToggleButton, FInputChord() );
+	UI_COMMAND( ToggleSnapCurveValueToInterval, "Snap Curve Key Values", "Snap curve keys to the value snapping interval", EUserInterfaceActionType::Button, FInputChord() );
 
+	UI_COMMAND( FindInContentBrowser, "Find in Content Browser", "Find the viewed actor animation in the content browser", EUserInterfaceActionType::Button, FInputChord() );
 	UI_COMMAND( ToggleCleanView, "Clean View", "Enable 'Clean View' mode which displays only global tracks when no filter is applied", EUserInterfaceActionType::ToggleButton, FInputChord() );
 	
 	UI_COMMAND( ToggleShowCurveEditor, "Curve Editor", "Show the animation keys in a curve editor", EUserInterfaceActionType::ToggleButton, FInputChord() );
diff --git a/Engine/Source/Editor/Sequencer/Private/SequencerCommands.h b/Engine/Source/Editor/Sequencer/Private/SequencerCommands.h
--- a/Engine/Source/Editor/Sequencer/Private/SequencerCommands.h
+++ b/Engine/Source/Editor/Sequencer/Private/SequencerCommands.h
@@ -111,6 +111,9 @@ public:
 	/** Toggles whether or not to snap curve values to the interval. */
 	TSharedPtr< FUICommandInfo > ToggleSnapCurveValueToInterval;
 
+	/** Finds the viewed actor animation in the content browser. */
+	TSharedPtr< FUICommandInfo > FindInContentBrowser;
+
 	/** Toggles whether the "clean view" is enabled in the level editor. */
 	TSharedPtr< FUICommandInfo > ToggleCleanView;
 	
