Sequencer: Update pivot location when scrubbing to fix case where the pivot doesn't update when realtime is off.

diff --git a/Engine/Source/Editor/Sequencer/Private/Sequencer.cpp b/Engine/Source/Editor/Sequencer/Private/Sequencer.cpp
--- a/Engine/Source/Editor/Sequencer/Private/Sequencer.cpp
+++ b/Engine/Source/Editor/Sequencer/Private/Sequencer.cpp
@@ -710,6 +710,9 @@ void FSequencer::SetGlobalTimeDirectly( float NewTime )
 
 	RootMovieSceneSequenceInstance->Update( ScrubPosition, LastTime, *this );
 
+	// If realtime is off, this needs to be called to update the pivot location when scrubbing.
+	GUnrealEd->UpdatePivotLocationForSelection();
+
 	GEditor->RedrawLevelEditingViewports();
 }
 
