Fixed not being able to save actor animation assets from within sequencer

diff --git a/Engine/Source/Editor/Sequencer/Private/Sequencer.cpp b/Engine/Source/Editor/Sequencer/Private/Sequencer.cpp
--- a/Engine/Source/Editor/Sequencer/Private/Sequencer.cpp
+++ b/Engine/Source/Editor/Sequencer/Private/Sequencer.cpp
@@ -1340,12 +1340,19 @@ void FSequencer::FindInContentBrowser()
 	if (GetFocusedMovieSceneSequence())
 	{
 		TArray<UObject*> ObjectsToFocus;
-		ObjectsToFocus.Add(GetFocusedMovieSceneSequence());
+		ObjectsToFocus.Add(GetCurrentAsset());
 
 		GEditor->SyncBrowserToObjects(ObjectsToFocus);
 	}
 }
 
+UObject* FSequencer::GetCurrentAsset() const
+{
+	// For now we find the asset by looking at the root movie scene's outer.
+	// @todo: this may need refining if/when we support editing movie scene instances
+	return GetFocusedMovieSceneSequence()->GetMovieScene()->GetOuter();
+}
+
 void FSequencer::VerticalScroll(float ScrollAmountUnits)
 {
 	SequencerWidget->GetTreeView()->ScrollByDelta(ScrollAmountUnits);
@@ -1776,7 +1783,7 @@ void FSequencer::UpdatePreviewLevelViewportClientFromCameraCut( FLevelEditorView
 
 void FSequencer::SaveCurrentMovieScene()
 {
-	UPackage* MovieScenePackage = GetFocusedMovieSceneSequence()->GetOutermost();
+	UPackage* MovieScenePackage = GetCurrentAsset()->GetOutermost();
 
 	TArray<UPackage*> PackagesToSave;
 
diff --git a/Engine/Source/Editor/Sequencer/Private/Sequencer.h b/Engine/Source/Editor/Sequencer/Private/Sequencer.h
--- a/Engine/Source/Editor/Sequencer/Private/Sequencer.h
+++ b/Engine/Source/Editor/Sequencer/Private/Sequencer.h
@@ -398,6 +398,11 @@ protected:
 	 */
 	void FindInContentBrowser();
 
+	/**
+	 * Get the asset we're currently editing, if applicable
+	 */
+	UObject* GetCurrentAsset() const;
+
 protected:
 
 	/** Called via UEditorEngine::GetActorRecordingStateEvent to check to see whether we need to record actor state */
