Sequencer: Make the details view toggleable.

diff --git a/Engine/Source/Editor/Sequencer/Private/SSequencer.cpp b/Engine/Source/Editor/Sequencer/Private/SSequencer.cpp
--- a/Engine/Source/Editor/Sequencer/Private/SSequencer.cpp
+++ b/Engine/Source/Editor/Sequencer/Private/SSequencer.cpp
@@ -314,7 +314,7 @@ void SSequencer::Construct( const FArguments& InArgs, TSharedRef< class FSequenc
 							.Orientation(Orient_Horizontal)
 						
 						+ SSplitter::Slot()
-							.Value(0.8f)
+							.Value(0.85f)
 							[
 								SNew(SBox)
 									.Padding(FMargin(0.0f, 2.0f, 0.0f, 0.0f))
@@ -455,8 +455,7 @@ void SSequencer::Construct( const FArguments& InArgs, TSharedRef< class FSequenc
 							]
 
 						+ SSplitter::Slot()
-							.SizeRule(SSplitter::SizeToContent)
-							.Value(0.2f)
+							.Value(0.15f)
 							[
 								DetailsView.ToSharedRef()
 							]
@@ -540,14 +539,15 @@ bool SSequencer::HandleDetailsViewEnabled() const
 
 
 EVisibility SSequencer::HandleDetailsViewVisibility() const
-{/*
-	if (Model->GetSelectedMessage().IsValid())
+{
+	if (Sequencer.Pin()->GetSettings()->GetDetailsViewVisible())
 	{
 		return EVisibility::Visible;
 	}
-
-	return EVisibility::Hidden;*/
-	return EVisibility::Visible;
+	else
+	{
+		return EVisibility::Collapsed;
+	}
 }
 
 
@@ -783,6 +783,7 @@ TSharedRef<SWidget> SSequencer::MakeGeneralMenu()
 	{
 		if (Sequencer.Pin()->IsLevelEditorSequencer())
 		{
+			MenuBuilder.AddMenuEntry( FSequencerCommands::Get().ToggleDetailsView );
 			MenuBuilder.AddMenuEntry( FSequencerCommands::Get().ToggleCleanView );
 		}
 
diff --git a/Engine/Source/Editor/Sequencer/Private/Sequencer.cpp b/Engine/Source/Editor/Sequencer/Private/Sequencer.cpp
--- a/Engine/Source/Editor/Sequencer/Private/Sequencer.cpp
+++ b/Engine/Source/Editor/Sequencer/Private/Sequencer.cpp
@@ -2346,6 +2346,14 @@ void FSequencer::BindSequencerCommands()
 		Commands.FindInContentBrowser,
 		FExecuteAction::CreateSP( this, &FSequencer::FindInContentBrowser ) );
 
+	SequencerCommandBindings->MapAction(
+		Commands.ToggleDetailsView,
+		FExecuteAction::CreateLambda( [this]{
+			Settings->SetDetailsViewVisible( !Settings->GetDetailsViewVisible() );
+		} ),
+		FCanExecuteAction::CreateLambda( []{ return true; } ),
+		FIsActionChecked::CreateLambda( [this]{ return Settings->GetDetailsViewVisible(); } ) );
+
 	SequencerCommandBindings->MapAction(
 		Commands.ToggleCleanView,
 		FExecuteAction::CreateLambda( [this]{
diff --git a/Engine/Source/Editor/Sequencer/Private/SequencerCommands.cpp b/Engine/Source/Editor/Sequencer/Private/SequencerCommands.cpp
--- a/Engine/Source/Editor/Sequencer/Private/SequencerCommands.cpp
+++ b/Engine/Source/Editor/Sequencer/Private/SequencerCommands.cpp
@@ -51,6 +51,7 @@ void FSequencerCommands::RegisterCommands()
 	UI_COMMAND( ToggleSnapCurveValueToInterval, "Snap Curve Key Values", "Snap curve keys to the value snapping interval", EUserInterfaceActionType::Button, FInputChord() );
 
 	UI_COMMAND( FindInContentBrowser, "Find in Content Browser", "Find the viewed actor animation in the content browser", EUserInterfaceActionType::Button, FInputChord() );
+	UI_COMMAND( ToggleDetailsView, "Details View", "Enable details view for selected sections and keys", EUserInterfaceActionType::ToggleButton, FInputChord() );
 	UI_COMMAND( ToggleCleanView, "Clean View", "Enable 'Clean View' mode which displays only global tracks when no filter is applied", EUserInterfaceActionType::ToggleButton, FInputChord() );
 	
 	UI_COMMAND( ToggleShowCurveEditor, "Curve Editor", "Show the animation keys in a curve editor", EUserInterfaceActionType::ToggleButton, FInputChord() );
diff --git a/Engine/Source/Editor/Sequencer/Private/SequencerCommands.h b/Engine/Source/Editor/Sequencer/Private/SequencerCommands.h
--- a/Engine/Source/Editor/Sequencer/Private/SequencerCommands.h
+++ b/Engine/Source/Editor/Sequencer/Private/SequencerCommands.h
@@ -114,6 +114,9 @@ public:
 	/** Finds the viewed actor animation in the content browser. */
 	TSharedPtr< FUICommandInfo > FindInContentBrowser;
 
+	/** Toggles whether the details view is enabled in the level editor. */
+	TSharedPtr< FUICommandInfo > ToggleDetailsView;
+	
 	/** Toggles whether the "clean view" is enabled in the level editor. */
 	TSharedPtr< FUICommandInfo > ToggleCleanView;
 	
diff --git a/Engine/Source/Editor/Sequencer/Private/SequencerSettings.cpp b/Engine/Source/Editor/Sequencer/Private/SequencerSettings.cpp
--- a/Engine/Source/Editor/Sequencer/Private/SequencerSettings.cpp
+++ b/Engine/Source/Editor/Sequencer/Private/SequencerSettings.cpp
@@ -22,6 +22,7 @@ USequencerSettings::USequencerSettings( const FObjectInitializer& ObjectInitiali
 	bSnapPlayTimeToDraggedKey = false;
 	CurveValueSnapInterval = 10.0f;
 	bSnapCurveValueToInterval = true;
+	bDetailsViewVisible = false;
 	bIsUsingCleanView = false;
 	bAutoScrollEnabled = false;
 	bShowCurveEditor = false;
@@ -266,6 +267,20 @@ void USequencerSettings::SetSnapCurveValueToInterval( bool InbSnapCurveValueToIn
 	}
 }
 
+bool USequencerSettings::GetDetailsViewVisible() const
+{
+	return bDetailsViewVisible;
+}
+
+void USequencerSettings::SetDetailsViewVisible(bool InbDetailsViewVisible )
+{
+	if (bDetailsViewVisible != InbDetailsViewVisible)
+	{
+		bDetailsViewVisible = InbDetailsViewVisible;
+		SaveConfig();
+	}
+}
+
 bool USequencerSettings::GetIsUsingCleanView() const
 {
 	return bIsUsingCleanView;
diff --git a/Engine/Source/Editor/Sequencer/Private/SequencerSettings.h b/Engine/Source/Editor/Sequencer/Private/SequencerSettings.h
--- a/Engine/Source/Editor/Sequencer/Private/SequencerSettings.h
+++ b/Engine/Source/Editor/Sequencer/Private/SequencerSettings.h
@@ -129,6 +129,11 @@ public:
 	/** Sets whether or not to snap curve values to the interval. */
 	void SetSnapCurveValueToInterval(bool InbSnapCurveValueToInterval);
 
+	/** Gets whether or not the details view is visible. */
+	bool GetDetailsViewVisible() const;
+	/** Sets whether or not the details view is visible. */
+	void SetDetailsViewVisible(bool InbDetailsViewVisible);
+
 	/** Gets whether or not the 'Clean View' is enabled. In 'Clean View' mode only global tracks are displayed when no filter is applied. */
 	bool GetIsUsingCleanView() const;
 	/** Sets whether or not the 'Clean View' is enabled. In 'Clean View' mode only global tracks are displayed when no filter is applied. */
@@ -207,6 +212,9 @@ protected:
 	UPROPERTY( config )
 	bool bSnapCurveValueToInterval;
 
+	UPROPERTY( config )
+	bool bDetailsViewVisible;
+
 	UPROPERTY( config )
 	bool bIsUsingCleanView;
 
