Sequencer: Trim and split sections. Add missing AddSection() virtual for skeletal animation track.

diff --git a/Engine/Source/Editor/MovieSceneTools/Private/MovieSceneToolHelpers.cpp b/Engine/Source/Editor/MovieSceneTools/Private/MovieSceneToolHelpers.cpp
--- a/Engine/Source/Editor/MovieSceneTools/Private/MovieSceneToolHelpers.cpp
+++ b/Engine/Source/Editor/MovieSceneTools/Private/MovieSceneToolHelpers.cpp
@@ -151,3 +151,47 @@ TSharedRef<SWidget> FBoolKeyArea::CreateKeyEditor(ISequencer* Sequencer)
 		.Curve(&Curve);
 };
 
+void MovieSceneToolHelpers::TrimSection(const TSet<TWeakObjectPtr<UMovieSceneSection>>& Sections, float Time, bool bTrimLeft)
+{
+	for (auto Section : Sections)
+	{
+		if (Section.IsValid() && Section->IsTimeWithinSection(Time))
+		{
+			Section->Modify();
+
+			if (bTrimLeft)
+			{
+				Section->SetStartTime(Time);
+			}
+			else
+			{
+				Section->SetEndTime(Time);
+			}
+		}
+	}
+}
+
+void MovieSceneToolHelpers::SplitSection(const TSet<TWeakObjectPtr<UMovieSceneSection>>& Sections, float Time)
+{
+	for (auto Section : Sections)
+	{
+		if (Section.IsValid() && Section->IsTimeWithinSection(Time))
+		{
+			Section->Modify();
+
+			float SectionEndTime = Section->GetEndTime();
+				
+			// Trim off the right
+			Section->SetEndTime(Time);
+
+			// Create a new section
+			UMovieSceneTrack* Track = CastChecked<UMovieSceneTrack>(Section->GetOuter() );
+			Track->Modify();
+
+			UMovieSceneSection* NewSection = DuplicateObject<UMovieSceneSection>(Section.Get(), Track);
+			Track->AddSection(NewSection);
+			NewSection->SetStartTime(Time);
+			NewSection->SetEndTime(SectionEndTime);
+		}
+	}
+}
\ No newline at end of file
diff --git a/Engine/Source/Editor/MovieSceneTools/Public/MovieSceneToolHelpers.h b/Engine/Source/Editor/MovieSceneTools/Public/MovieSceneToolHelpers.h
--- a/Engine/Source/Editor/MovieSceneTools/Public/MovieSceneToolHelpers.h
+++ b/Engine/Source/Editor/MovieSceneTools/Public/MovieSceneToolHelpers.h
@@ -256,4 +256,25 @@ public:
 	virtual TSharedRef<SWidget> CreateKeyEditor(ISequencer* Sequencer) override;
 };
 
+class MOVIESCENETOOLS_API MovieSceneToolHelpers
+{
+public:
+	/**
+	* Trim section at the given time
+	*
+	* @param Sections The sections to trim
+	* @param Time	The time at which to trim
+	* @param bTrimLeft Trim left or trim right
+	*/
+	static void TrimSection(const TSet<TWeakObjectPtr<UMovieSceneSection>>& Sections, float Time, bool bTrimLeft);
+
+	/**
+	* Splits sections at the given time
+	*
+	* @param Sections The sections to split
+	* @param Time	The time at which to split
+	*/
+	static void SplitSection(const TSet<TWeakObjectPtr<UMovieSceneSection>>& Sections, float Time);
+};
+
 
diff --git a/Engine/Source/Editor/Sequencer/Private/SSection.cpp b/Engine/Source/Editor/Sequencer/Private/SSection.cpp
--- a/Engine/Source/Editor/Sequencer/Private/SSection.cpp
+++ b/Engine/Source/Editor/Sequencer/Private/SSection.cpp
@@ -7,6 +7,7 @@
 #include "ISequencerSection.h"
 #include "MovieSceneSection.h"
 #include "MovieSceneShotSection.h"
+#include "MovieSceneToolHelpers.h"
 #include "CommonMovieSceneTools.h"
 #include "SequencerHotspots.h"
 #include "ScopedTransaction.h"
@@ -590,6 +591,33 @@ TSharedPtr<SWidget> SSection::OnSummonContextMenu( const FGeometry& MyGeometry,
 				FUIAction(FExecuteAction::CreateSP(this, &SSection::SelectAllKeys))
 			);
 
+			MenuBuilder.AddMenuEntry(
+				NSLOCTEXT("Sequencer", "TrimSectionLeft", "Trim Left"),
+				NSLOCTEXT("Sequencer", "TrimSectionLeftTooltip", "Trim section at current time to the left"),
+				FSlateIcon(),
+				FUIAction(
+					FExecuteAction::CreateSP(this, &SSection::TrimSection, true),
+					FCanExecuteAction::CreateSP(this, &SSection::IsTrimmable))
+			);
+
+			MenuBuilder.AddMenuEntry(
+				NSLOCTEXT("Sequencer", "TrimSectionRight", "Trim Right"),
+				NSLOCTEXT("Sequencer", "TrimSectionRightTooltip", "Trim section at current time to the right"),
+				FSlateIcon(),
+				FUIAction(
+					FExecuteAction::CreateSP(this, &SSection::TrimSection, false),
+					FCanExecuteAction::CreateSP(this, &SSection::IsTrimmable))
+			);
+
+			MenuBuilder.AddMenuEntry(
+				NSLOCTEXT("Sequencer", "SplitSection", "Split"),
+				NSLOCTEXT("Sequencer", "SplitSectionTooltip", "Split section at current time"),
+				FSlateIcon(),
+				FUIAction(
+					FExecuteAction::CreateSP(this, &SSection::SplitSection),
+					FCanExecuteAction::CreateSP(this, &SSection::IsTrimmable))
+			);
+
 			MenuBuilder.AddSubMenu(
 				NSLOCTEXT("Sequencer", "SetPreInfinityExtrap", "Pre-Infinity"), 
 				NSLOCTEXT("Sequencer", "SetPreInfinityExtrapTooltip", "Set pre-infinity extrapolation"),
@@ -1001,6 +1029,36 @@ void SSection::SelectAllKeys()
 	}
 }
 
+void SSection::TrimSection(bool bTrimLeft)
+{
+	FScopedTransaction TrimSectionTransaction(NSLOCTEXT("Sequencer", "TrimSection_Transaction", "Trim Section"));
+
+	MovieSceneToolHelpers::TrimSection(GetSequencer().GetSelection().GetSelectedSections(), GetSequencer().GetGlobalTime(), bTrimLeft);
+
+	GetSequencer().NotifyMovieSceneDataChanged();
+}
+
+void SSection::SplitSection()
+{
+	FScopedTransaction SplitSectionTransaction(NSLOCTEXT("Sequencer", "SplitSection_Transaction", "Split Section"));
+
+	MovieSceneToolHelpers::SplitSection(GetSequencer().GetSelection().GetSelectedSections(), GetSequencer().GetGlobalTime());
+
+	GetSequencer().NotifyMovieSceneDataChanged();
+}
+
+bool SSection::IsTrimmable() const
+{
+	for (auto Section : GetSequencer().GetSelection().GetSelectedSections())
+	{
+		if (Section.IsValid() && Section->IsTimeWithinSection(GetSequencer().GetGlobalTime()))
+		{
+			return true;
+		}
+	}
+	return false;
+}
+
 void SSection::SetExtrapolationMode(ERichCurveExtrapolation ExtrapMode, bool bPreInfinity)
 {
 	// @todo Sequencer should operate on selected sections
diff --git a/Engine/Source/Editor/Sequencer/Private/SSection.h b/Engine/Source/Editor/Sequencer/Private/SSection.h
--- a/Engine/Source/Editor/Sequencer/Private/SSection.h
+++ b/Engine/Source/Editor/Sequencer/Private/SSection.h
@@ -153,6 +153,21 @@ private:
 	 */
 	void SelectAllKeys();
 
+	/*
+	 * Trim section 
+	 */
+	void TrimSection(bool bTrimLeft);
+
+	/*
+	 * Split section
+	 */
+	void SplitSection();
+
+	/*
+	 * Is there a trimmable section?
+	 */
+	bool IsTrimmable() const;
+
 	/**
 	 * Set extrapolation mode
 	 */
diff --git a/Engine/Source/Editor/Sequencer/Private/Sequencer.cpp b/Engine/Source/Editor/Sequencer/Private/Sequencer.cpp
--- a/Engine/Source/Editor/Sequencer/Private/Sequencer.cpp
+++ b/Engine/Source/Editor/Sequencer/Private/Sequencer.cpp
@@ -24,6 +24,7 @@
 #include "MovieSceneAudioTrack.h"
 #include "MovieSceneSkeletalAnimationTrack.h"
 #include "MovieSceneTrackEditor.h"
+#include "MovieSceneToolHelpers.h"
 #include "Kismet2/KismetEditorUtilities.h"
 #include "Kismet2/BlueprintEditorUtils.h"
 #include "AssetSelection.h"
@@ -2233,6 +2234,24 @@ void FSequencer::SetKey()
 	}
 }
 
+void FSequencer::TrimSection(bool bTrimLeft)
+{
+	FScopedTransaction TrimSectionTransaction( NSLOCTEXT("Sequencer", "TrimSection_Transaction", "Trim Section") );
+
+	MovieSceneToolHelpers::TrimSection(Selection.GetSelectedSections(), GetGlobalTime(), bTrimLeft);
+
+	NotifyMovieSceneDataChanged();
+}
+
+void FSequencer::SplitSection()
+{
+	FScopedTransaction SplitSectionTransaction( NSLOCTEXT("Sequencer", "SplitSection_Transaction", "Split Section") );
+
+	MovieSceneToolHelpers::SplitSection(Selection.GetSelectedSections(), GetGlobalTime());
+
+	NotifyMovieSceneDataChanged();
+}
+
 ISequencerEditTool& FSequencer::GetEditTool()
 {
 	return SequencerWidget->GetEditTool();
@@ -2323,6 +2342,18 @@ void FSequencer::BindSequencerCommands()
 		Commands.SetInterpolationConstant,
 		FExecuteAction::CreateSP( this, &FSequencer::SetInterpTangentMode, ERichCurveInterpMode::RCIM_Constant, ERichCurveTangentMode::RCTM_Auto ) );
 
+	SequencerCommandBindings->MapAction(
+		Commands.TrimSectionLeft,
+		FExecuteAction::CreateSP( this, &FSequencer::TrimSection, true ) );
+
+	SequencerCommandBindings->MapAction(
+		Commands.TrimSectionRight,
+		FExecuteAction::CreateSP( this, &FSequencer::TrimSection, false ) );
+
+	SequencerCommandBindings->MapAction(
+		Commands.SplitSection,
+		FExecuteAction::CreateSP( this, &FSequencer::SplitSection ) );
+
 	SequencerCommandBindings->MapAction(
 		Commands.ToggleAutoKeyEnabled,
 		FExecuteAction::CreateLambda( [this]{ Settings->SetAutoKeyEnabled( !Settings->GetAutoKeyEnabled() ); } ),
diff --git a/Engine/Source/Editor/Sequencer/Private/Sequencer.h b/Engine/Source/Editor/Sequencer/Private/Sequencer.h
--- a/Engine/Source/Editor/Sequencer/Private/Sequencer.h
+++ b/Engine/Source/Editor/Sequencer/Private/Sequencer.h
@@ -439,6 +439,12 @@ protected:
 	/** Manually sets a key for the selected objects at the current time */
 	void SetKey();
 
+	/** Trim a section to the left or right */
+	void TrimSection(bool bTrimLeft);
+
+	/** Split a section */
+	void SplitSection();
+
 	/** Generates command bindings for UI commands */
 	void BindSequencerCommands();
 
diff --git a/Engine/Source/Editor/Sequencer/Private/SequencerCommands.cpp b/Engine/Source/Editor/Sequencer/Private/SequencerCommands.cpp
--- a/Engine/Source/Editor/Sequencer/Private/SequencerCommands.cpp
+++ b/Engine/Source/Editor/Sequencer/Private/SequencerCommands.cpp
@@ -29,6 +29,10 @@ void FSequencerCommands::RegisterCommands()
 	UI_COMMAND( SetInterpolationLinear, "Set Key Linear", "Linear interpolation", EUserInterfaceActionType::Button, FInputChord(EKeys::Four));
 	UI_COMMAND( SetInterpolationConstant, "Set Key Constant", "Constant interpolation", EUserInterfaceActionType::Button, FInputChord(EKeys::Five));
 
+	UI_COMMAND( TrimSectionLeft, "Trim Section Left", "Trim section at current time to the left (keeps the right)", EUserInterfaceActionType::Button, FInputChord(EModifierKey::Control, EKeys::Comma) );
+	UI_COMMAND( TrimSectionRight, "Trim Section Right", "Trim section at current time to the right (keeps the left)", EUserInterfaceActionType::Button, FInputChord(EModifierKey::Control, EKeys::Period) );
+	UI_COMMAND( SplitSection, "Split Section", "Split section at current time", EUserInterfaceActionType::Button, FInputChord(EModifierKey::Control, EKeys::Slash) );
+
 	UI_COMMAND( ToggleAutoKeyEnabled, "Auto Key", "Enables and disables auto keying", EUserInterfaceActionType::ToggleButton, FInputChord() );
 	UI_COMMAND( ToggleKeyAllEnabled, "Key All", "Enables and disables key all", EUserInterfaceActionType::ToggleButton, FInputChord() );
 
diff --git a/Engine/Source/Editor/Sequencer/Private/SequencerCommands.h b/Engine/Source/Editor/Sequencer/Private/SequencerCommands.h
--- a/Engine/Source/Editor/Sequencer/Private/SequencerCommands.h
+++ b/Engine/Source/Editor/Sequencer/Private/SequencerCommands.h
@@ -72,6 +72,15 @@ public:
 	/** Sets the interp tangent mode for the selected keys to constant */
 	TSharedPtr< FUICommandInfo > SetInterpolationConstant;
 
+	/** Trim section to the left, keeping the right portion */
+	TSharedPtr< FUICommandInfo > TrimSectionLeft;
+
+	/** Trim section to the right, keeping the left portion */
+	TSharedPtr< FUICommandInfo > TrimSectionRight;
+
+	/** Split section */
+	TSharedPtr< FUICommandInfo > SplitSection;
+
 	/** Turns auto keying on and off. */
 	TSharedPtr< FUICommandInfo > ToggleAutoKeyEnabled;
 
diff --git a/Engine/Source/Runtime/MovieSceneTracks/Private/Tracks/MovieSceneSkeletalAnimationTrack.cpp b/Engine/Source/Runtime/MovieSceneTracks/Private/Tracks/MovieSceneSkeletalAnimationTrack.cpp
--- a/Engine/Source/Runtime/MovieSceneTracks/Private/Tracks/MovieSceneSkeletalAnimationTrack.cpp
+++ b/Engine/Source/Runtime/MovieSceneTracks/Private/Tracks/MovieSceneSkeletalAnimationTrack.cpp
@@ -45,6 +45,12 @@ bool UMovieSceneSkeletalAnimationTrack::HasSection( UMovieSceneSection* Section
 }
 
 
+void UMovieSceneSkeletalAnimationTrack::AddSection( UMovieSceneSection* Section)
+{
+	AnimationSections.Add(Section);
+}
+
+
 void UMovieSceneSkeletalAnimationTrack::RemoveSection( UMovieSceneSection* Section )
 {
 	AnimationSections.Remove( Section );
@@ -76,7 +82,7 @@ void UMovieSceneSkeletalAnimationTrack::AddNewAnimation(float KeyTime, UAnimSequ
 	NewSection->SetAnimationStartTime( KeyTime );
 	NewSection->SetAnimSequence(AnimSequence);
 
-	AnimationSections.Add(NewSection);
+	AddSection(NewSection);
 }
 
 
diff --git a/Engine/Source/Runtime/MovieSceneTracks/Public/Tracks/MovieSceneSkeletalAnimationTrack.h b/Engine/Source/Runtime/MovieSceneTracks/Public/Tracks/MovieSceneSkeletalAnimationTrack.h
--- a/Engine/Source/Runtime/MovieSceneTracks/Public/Tracks/MovieSceneSkeletalAnimationTrack.h
+++ b/Engine/Source/Runtime/MovieSceneTracks/Public/Tracks/MovieSceneSkeletalAnimationTrack.h
@@ -21,6 +21,7 @@ public:
 	virtual TSharedPtr<IMovieSceneTrackInstance> CreateInstance() override;
 	virtual void RemoveAllAnimationData() override;
 	virtual bool HasSection( UMovieSceneSection* Section ) const override;
+	virtual void AddSection( UMovieSceneSection* Section ) override;
 	virtual void RemoveSection( UMovieSceneSection* Section ) override;
 	virtual bool IsEmpty() const override;
 	virtual TRange<float> GetSectionBoundaries() const override;
