Sequencer: Removed unused FilterToSelectedShotSections feature

As per discussion with MattK.

diff --git a/Engine/Source/Editor/MovieSceneTools/Private/TrackEditors/ShotTrackEditor.cpp b/Engine/Source/Editor/MovieSceneTools/Private/TrackEditors/ShotTrackEditor.cpp
--- a/Engine/Source/Editor/MovieSceneTools/Private/TrackEditors/ShotTrackEditor.cpp
+++ b/Engine/Source/Editor/MovieSceneTools/Private/TrackEditors/ShotTrackEditor.cpp
@@ -288,22 +288,6 @@ FReply FShotSection::OnSectionDoubleClicked( const FGeometry& SectionGeometry, c
 	return FReply::Handled();
 }
 
-void FShotSection::BuildSectionContextMenu(FMenuBuilder& MenuBuilder)
-{
-	MenuBuilder.AddMenuEntry(
-		NSLOCTEXT("FShotTrackEditor", "FilterToShots", "Filter To Shots"),
-		NSLOCTEXT("FShotTrackEditor", "FilterToShotsToolTip", "Filters to the selected shot sections"),
-		FSlateIcon(),
-		FUIAction(FExecuteAction::CreateSP(this, &FShotSection::FilterToSelectedShotSections, true))
-		);
-}
-
-
-void FShotSection::FilterToSelectedShotSections(bool bZoomToShotBounds)
-{
-	Sequencer.Pin()->FilterToSelectedShotSections(bZoomToShotBounds);
-}
-
 int32 FShotSection::OnPaintSection( const FGeometry& AllottedGeometry, const FSlateRect& SectionClippingRect, FSlateWindowElementList& OutDrawElements, int32 LayerId, bool bParentEnabled ) const
 {
 	const ESlateDrawEffect::Type DrawEffects = bParentEnabled ? ESlateDrawEffect::None : ESlateDrawEffect::DisabledEffect;
diff --git a/Engine/Source/Editor/MovieSceneTools/Private/TrackEditors/ShotTrackEditor.h b/Engine/Source/Editor/MovieSceneTools/Private/TrackEditors/ShotTrackEditor.h
--- a/Engine/Source/Editor/MovieSceneTools/Private/TrackEditors/ShotTrackEditor.h
+++ b/Engine/Source/Editor/MovieSceneTools/Private/TrackEditors/ShotTrackEditor.h
@@ -154,10 +154,6 @@ public:
 	virtual bool AreSectionsConnected() const override { return true; }
 	virtual void GenerateSectionLayout( class ISectionLayoutBuilder& LayoutBuilder ) const override {}
 	virtual FReply OnSectionDoubleClicked( const FGeometry& SectionGeometry, const FPointerEvent& MouseEvent ) override;
-	virtual void BuildSectionContextMenu(FMenuBuilder& MenuBuilder) override;
-
-	/** Filter to selected shot sections */
-	void FilterToSelectedShotSections(bool bZoomToShotBounds);
 
 	/** Gets the thumbnail width */
 	uint32 GetThumbnailWidth() const;
diff --git a/Engine/Source/Editor/Sequencer/Private/SSequencer.cpp b/Engine/Source/Editor/Sequencer/Private/SSequencer.cpp
--- a/Engine/Source/Editor/Sequencer/Private/SSequencer.cpp
+++ b/Engine/Source/Editor/Sequencer/Private/SSequencer.cpp
@@ -1318,13 +1318,6 @@ void SSequencer::OnCrumbClicked(const FSequencerBreadcrumb& Item)
 		{
 			Sequencer.Pin()->PopToMovieScene( Item.MovieSceneInstance.Pin().ToSharedRef() );
 		}
-
-		Sequencer.Pin()->FilterToShotSections( TArray< TWeakObjectPtr<UMovieSceneSection> >() );
-	}
-	else
-	{
-		// refilter what we already have
-		Sequencer.Pin()->FilterToShotSections(Sequencer.Pin()->GetFilteringShotSections());
 	}
 }
 
diff --git a/Engine/Source/Editor/Sequencer/Private/Sequencer.cpp b/Engine/Source/Editor/Sequencer/Private/Sequencer.cpp
--- a/Engine/Source/Editor/Sequencer/Private/Sequencer.cpp
+++ b/Engine/Source/Editor/Sequencer/Private/Sequencer.cpp
@@ -923,8 +923,6 @@ void FSequencer::AddReferencedObjects( FReferenceCollector& Collector )
 
 void FSequencer::ResetPerMovieSceneData()
 {
-	FilterToShotSections( TArray< TWeakObjectPtr<UMovieSceneSection> >(), false );
-
 	//@todo Sequencer - We may want to preserve selections when moving between movie scenes
 	Selection.Empty();
 
@@ -1897,79 +1895,6 @@ void FSequencer::ZoomToSelectedSections()
 	}
 }
 
-void FSequencer::FilterToShotSections(const TArray< TWeakObjectPtr<class UMovieSceneSection> >& ShotSections, bool bZoomToShotBounds )
-{
-	TArray< TWeakObjectPtr<UMovieSceneSection> > ActualShotSections;
-	for (int32 i = 0; i < ShotSections.Num(); ++i)
-	{
-		if (ShotSections[i]->IsA<UMovieSceneShotSection>())
-		{
-			ActualShotSections.Add(ShotSections[i]);
-		}
-	}
-
-	bool bWasPreviouslyFiltering = IsShotFilteringOn();
-	bool bIsNowFiltering = ActualShotSections.Num() > 0;
-
-	FilteringShots.Empty();
-	UnfilterableSections.Empty();
-	UnfilterableObjects.Empty();
-
-	if (bIsNowFiltering)
-	{
-		for ( int32 ShotSectionIndex = 0; ShotSectionIndex < ActualShotSections.Num(); ++ShotSectionIndex )
-		{
-			FilteringShots.Add( ActualShotSections[ShotSectionIndex] );
-		}
-
-		// populate unfilterable shots - shots that started not filtered
-		TArray< TWeakObjectPtr<class UMovieSceneSection> > TempUnfilterableSections;
-		const TArray<UMovieSceneSection*>& AllSections = GetFocusedMovieSceneSequence()->GetMovieScene()->GetAllSections();
-		for (int32 SectionIndex = 0; SectionIndex < AllSections.Num(); ++SectionIndex)
-		{
-			UMovieSceneSection* Section = AllSections[SectionIndex];
-			if (!Section->IsA<UMovieSceneShotSection>() && IsSectionVisible(Section))
-			{
-				TempUnfilterableSections.Add(Section);
-			}
-		}
-		// wait until after we've collected them all before applying in order to
-		// prevent wastefully searching through UnfilterableSections in IsSectionVisible
-		UnfilterableSections = TempUnfilterableSections;
-	}
-
-	if (!bWasPreviouslyFiltering && bIsNowFiltering)
-	{
-		OverlayAnimation.Play( SequencerWidget.ToSharedRef() );
-	}
-	else if (bWasPreviouslyFiltering && !bIsNowFiltering) 
-	{
-		OverlayAnimation.PlayReverse( SequencerWidget.ToSharedRef() );
-	}
-
-	if( bZoomToShotBounds )
-	{
-		// zoom in
-		bool bExpandClampRange = true;
-		OnViewRangeChanged(GetTimeBounds(), EViewRangeInterpolation::Animated, bExpandClampRange);
-	}
-
-	SequencerWidget->UpdateBreadcrumbs(ActualShotSections);
-}
-
-void FSequencer::FilterToSelectedShotSections(bool bZoomToShotBounds)
-{
-	TArray< TWeakObjectPtr<UMovieSceneSection> > SelectedShotSections;
-	for (TWeakObjectPtr<UMovieSceneSection> SelectedSection : Selection.GetSelectedSections())
-	{
-		if (SelectedSection->IsA<UMovieSceneShotSection>())
-		{
-			SelectedShotSections.Add(SelectedSection);
-		}
-	}
-	FilterToShotSections(SelectedShotSections, bZoomToShotBounds);
-}
-
 bool FSequencer::CanKeyProperty(FCanKeyPropertyParams CanKeyPropertyParams) const
 {
 	return ObjectChangeListener->CanKeyProperty(CanKeyPropertyParams);
diff --git a/Engine/Source/Editor/Sequencer/Private/Sequencer.h b/Engine/Source/Editor/Sequencer/Private/Sequencer.h
--- a/Engine/Source/Editor/Sequencer/Private/Sequencer.h
+++ b/Engine/Source/Editor/Sequencer/Private/Sequencer.h
@@ -86,8 +86,6 @@ public:
 	virtual void NotifyMovieSceneDataChanged() override;
 	virtual void UpdateRuntimeInstances() override;
 	virtual void AddSubMovieScene(UMovieSceneSequence* SubMovieSceneSequence) override;
-	virtual void FilterToShotSections(const TArray< TWeakObjectPtr<class UMovieSceneSection> >& ShotSections, bool bZoomToShotBounds = false) override;
-	virtual void FilterToSelectedShotSections(bool bZoomToShotBounds = true) override;
 	virtual bool CanKeyProperty(FCanKeyPropertyParams CanKeyPropertyParams) const override;
 	virtual void KeyProperty(FKeyPropertyParams KeyPropertyParams) override;
 	virtual FSequencerSelection& GetSelection() override;
diff --git a/Engine/Source/Editor/Sequencer/Public/ISequencer.h b/Engine/Source/Editor/Sequencer/Public/ISequencer.h
--- a/Engine/Source/Editor/Sequencer/Public/ISequencer.h
+++ b/Engine/Source/Editor/Sequencer/Public/ISequencer.h
@@ -94,23 +94,6 @@ public:
 	/** @return Returns whether sequencer is currently recording live data from simulated actors */
 	virtual bool IsRecordingLive() const = 0;
 
-	/**
-	 * Sets which shot sections need to be filtered to, or none to disable shot filtering.
-	 *
-	 * @param ShotSections The shot sections to filter.
-	 * @param bZoomToShotBounds Whether or not to zoom time to the bounds of the shots.
-	 * @see FilterToSelectedShotSections
-	 */
-	virtual void FilterToShotSections(const TArray<TWeakObjectPtr<UMovieSceneSection>>& ShotSections, bool bZoomToShotBounds = true) = 0;
-	
-	/**
-	 * Sets shot filtering to current selection.
-	 *
-	 * @param bZoomToShotBounds Whether or not to zoom time to the bounds of the shots.
-	 * @see FilterToShotSections
-	 */
-	virtual void FilterToSelectedShotSections(bool bZoomToShotBounds = true) = 0;
-
 	/**
 	 * Gets the current time of the time slider relative to the passed in movie scene
 	 *
