Sequencer: Removed more code related to shot track filtering

diff --git a/Engine/Source/Editor/Sequencer/Private/SSection.cpp b/Engine/Source/Editor/Sequencer/Private/SSection.cpp
--- a/Engine/Source/Editor/Sequencer/Private/SSection.cpp
+++ b/Engine/Source/Editor/Sequencer/Private/SSection.cpp
@@ -864,17 +864,16 @@ FReply SSection::OnMouseButtonDoubleClick( const FGeometry& MyGeometry, const FP
 
 	if( MouseEvent.GetEffectingButton() == EKeys::LeftMouseButton )
 	{
-		UMovieSceneSection* SectionObject = SectionInterface->GetSectionObject();
-		if( GetSequencer().IsSectionVisible( SectionObject ) )
-		{
-			FReply Reply = SectionInterface->OnSectionDoubleClicked( MyGeometry, MouseEvent );
+		FReply Reply = SectionInterface->OnSectionDoubleClicked( MyGeometry, MouseEvent );
 
-			if (Reply.IsEventHandled()) {return Reply;}
+		if (Reply.IsEventHandled())
+		{
+			return Reply;
+		}
 
-			GetSequencer().ZoomToSelectedSections();
+		GetSequencer().ZoomToSelectedSections();
 
-			return FReply::Handled();
-		}
+		return FReply::Handled();
 	}
 
 	return FReply::Unhandled();
diff --git a/Engine/Source/Editor/Sequencer/Private/SSequencer.cpp b/Engine/Source/Editor/Sequencer/Private/SSequencer.cpp
--- a/Engine/Source/Editor/Sequencer/Private/SSequencer.cpp
+++ b/Engine/Source/Editor/Sequencer/Private/SSequencer.cpp
@@ -57,22 +57,6 @@ public:
 		ViewRange = InArgs._ViewRange;
 		Sequencer = InSequencer;
 	}
-
-	virtual void Tick( const FGeometry& AllottedGeometry, const double InCurrentTime, const float InDeltaTime ) override
-	{
-		if (Sequencer.Pin()->IsShotFilteringOn())
-		{
-			TArray< TWeakObjectPtr<UMovieSceneSection> > FilterShots = Sequencer.Pin()->GetFilteringShotSections();
-			// if there are filtering shots, continuously update the cached filter zones
-			// we do this in tick, and cache it in order to make animation work properly
-			CachedFilteredRanges.Empty();
-			for (int32 i = 0; i < FilterShots.Num(); ++i)
-			{
-				UMovieSceneSection* Filter = FilterShots[i].Get();
-				CachedFilteredRanges.Add(Filter->GetRange());
-			}
-		}
-	}
 	
 	virtual int32 OnPaint( const FPaintArgs& Args, const FGeometry& AllottedGeometry, const FSlateRect& MyClippingRect, FSlateWindowElementList& OutDrawElements, int32 LayerId, const FWidgetStyle& InWidgetStyle, bool bParentEnabled ) const override
 	{
@@ -953,21 +937,6 @@ void SSequencer::UpdateBreadcrumbs(const TArray< TWeakObjectPtr<class UMovieScen
 		BreadcrumbTrail->PushCrumb( CrumbName, FSequencerBreadcrumb( FocusedMovieSceneInstance ) );
 	}
 
-	if (Sequencer.Pin()->IsShotFilteringOn())
-	{
-		TAttribute<FText> CrumbString;
-		if (FilteringShots.Num() == 1)
-		{
-			CrumbString = TAttribute<FText>::Create(TAttribute<FText>::FGetter::CreateSP(this, &SSequencer::GetShotSectionTitle, FilteringShots[0].Get()));
-		}
-		else 
-		{
-			CrumbString = LOCTEXT("MultipleShots", "Multiple Shots");
-		}
-
-		BreadcrumbTrail->PushCrumb(CrumbString, FSequencerBreadcrumb());
-	}
-
 	SequencerNodeTree->UpdateCachedVisibilityBasedOnShotFiltersChanged();
 }
 
diff --git a/Engine/Source/Editor/Sequencer/Private/SSequencerSectionAreaView.cpp b/Engine/Source/Editor/Sequencer/Private/SSequencerSectionAreaView.cpp
--- a/Engine/Source/Editor/Sequencer/Private/SSequencerSectionAreaView.cpp
+++ b/Engine/Source/Editor/Sequencer/Private/SSequencerSectionAreaView.cpp
@@ -110,7 +110,7 @@ void SSequencerSectionAreaView::GenerateSectionWidgets()
 
 EVisibility SSequencerSectionAreaView::GetSectionVisibility( UMovieSceneSection* SectionObject ) const
 {
-	return GetSequencer().IsSectionVisible( SectionObject ) ? EVisibility::Visible : EVisibility::Collapsed;
+	return EVisibility::Visible;
 }
 
 
diff --git a/Engine/Source/Editor/Sequencer/Private/Sequencer.cpp b/Engine/Source/Editor/Sequencer/Private/Sequencer.cpp
--- a/Engine/Source/Editor/Sequencer/Private/Sequencer.cpp
+++ b/Engine/Source/Editor/Sequencer/Private/Sequencer.cpp
@@ -131,17 +131,6 @@ void FSequencer::InitSequencer(const FSequencerInitParams& InitParams, const TSh
 			check( TrackEditorDelegates[DelegateIndex].IsBound() );
 			// Tools may exist in other modules, call a delegate that will create one for us 
 			TSharedRef<FMovieSceneTrackEditor> TrackEditor = TrackEditorDelegates[DelegateIndex].Execute( SharedThis( this ) );
-
-			// Keep track of certain editors
-			if ( TrackEditor->SupportsType( UMovieSceneShotTrack::StaticClass() ) )
-			{
-				ShotTrackEditor = TrackEditor;
-			}
-			else if ( TrackEditor->SupportsType( UMovieSceneSkeletalAnimationTrack::StaticClass() ) )
-			{
-				SkeletalAnimationTrackEditor = TrackEditor;
-			}
-
 			TrackEditors.Add( TrackEditor );
 		}
 
@@ -304,9 +293,6 @@ void FSequencer::ResetToNewRootSequence(UMovieSceneSequence& NewSequence)
 	//@todo Sequencer - Encapsulate this better
 	MovieSceneStack.Empty();
 	Selection.Empty();
-	FilteringShots.Empty();
-	UnfilterableSections.Empty();
-	UnfilterableObjects.Empty();
 	MovieSceneSectionToInstanceMap.Empty();
 
 	// Focusing the initial movie scene needs to be done before the first time NewSequence or GetRootMovieSceneInstance is used
@@ -560,11 +546,6 @@ void FSequencer::OnActorsDropped( const TArray<TWeakObjectPtr<AActor> >& Actors
 				
 			const FGuid PossessableGuid = OwnerMovieScene->AddPossessable(Actor->GetActorLabel(), Actor->GetClass());
 
-			if (IsShotFilteringOn())
-			{
-				AddUnfilterableObject(PossessableGuid);
-			}
-
 			OwnerSequence->BindPossessableObject(PossessableGuid, *Actor);
 			bPossessableAdded = true;
 		}
@@ -822,12 +803,6 @@ FGuid FSequencer::GetHandleToObject( UObject* Object, bool bCreateHandleIfMissin
 			FocusedMovieScene->Modify();
 
 			ObjectGuid = FocusedMovieScene->AddPossessable(Object->GetName(), Object->GetClass());
-
-			if(IsShotFilteringOn())
-			{
-				AddUnfilterableObject(ObjectGuid);
-			}
-
 			FocusedMovieSceneSequence->BindPossessableObject(ObjectGuid, *Object);
 			bPossessableAdded = true;
 		}
@@ -1152,34 +1127,19 @@ TRange<float> FSequencer::GetTimeBounds() const
 
 	const UMovieScene* MovieScene = GetFocusedMovieSceneSequence()->GetMovieScene();
 	const UMovieSceneTrack* AnimatableShot = MovieScene->FindMasterTrack( UMovieSceneShotTrack::StaticClass() );
+
 	if (AnimatableShot)
 	{
-		// try getting filtered shot boundaries
-		TRange<float> Bounds = GetFilteringShotsTimeBounds();
-		if (!Bounds.IsEmpty()) {return Bounds;}
-
 		// try getting the bounds of all shots
-		Bounds = AnimatableShot->GetSectionBoundaries();
-		if (!Bounds.IsEmpty()) {return Bounds;}
-	}
-	
-	return MovieScene->GetTimeRange();
-}
-
-TRange<float> FSequencer::GetFilteringShotsTimeBounds() const
-{
-	if (IsShotFilteringOn())
-	{
-		TArray< TRange<float> > Bounds;
-		for (int32 i = 0; i < FilteringShots.Num(); ++i)
+			TRange<float> Bounds = AnimatableShot->GetSectionBoundaries();
+		
+		if (!Bounds.IsEmpty())
 		{
-			Bounds.Add(FilteringShots[i]->GetRange());
+			return Bounds;
 		}
-
-		return TRange<float>::Hull(Bounds);
 	}
-
-	return TRange<float>::Empty();
+	
+	return MovieScene->GetTimeRange();
 }
 
 void FSequencer::OnViewRangeChanged( TRange<float> NewViewRange, EViewRangeInterpolation Interpolation, bool bExpandClampRange )
@@ -1416,11 +1376,6 @@ FGuid FSequencer::AddSpawnableForAssetOrClass( UObject* Object, UObject* Counter
 			}
 
 			NewSpawnableGuid = OwnerMovieScene->AddSpawnable( NewSpawnableName, NewBlueprint, CounterpartGamePreviewObject );
-
-			if (IsShotFilteringOn())
-			{
-				AddUnfilterableObject(NewSpawnableGuid);
-			}
 		}
 	
 	}
@@ -1818,12 +1773,6 @@ void FSequencer::AddSelectedObjects()
 		for (auto Actor : SelectedActors)
 		{
 			const FGuid PossessableGuid = OwnerMovieScene->AddPossessable(Actor->GetActorLabel(), Actor->GetClass());
-
-			if (IsShotFilteringOn())
-			{
-				AddUnfilterableObject(PossessableGuid);
-			}
-
 			OwnerSequence->BindPossessableObject(PossessableGuid, *Actor);
 			bPossessableAdded = true;
 		}
@@ -1837,11 +1786,6 @@ void FSequencer::AddSelectedObjects()
 
 void FSequencer::OnSectionSelectionChanged()
 {
-	for (TWeakObjectPtr<UMovieSceneSection> SelectedSection : Selection.GetSelectedSections())
-	{
-		// if we select something, consider it unfilterable until we change shot filters
-		UnfilterableSections.AddUnique(TWeakObjectPtr<UMovieSceneSection>(SelectedSection));
-	}
 }
 
 void FSequencer::OnSelectedOutlinerNodesChanged()
@@ -1910,52 +1854,11 @@ FSequencerSelection& FSequencer::GetSelection()
 	return Selection;
 }
 
-const TArray< TWeakObjectPtr<UMovieSceneSection> >& FSequencer::GetFilteringShotSections() const
-{
-	return FilteringShots;
-}
-
-bool FSequencer::IsObjectUnfilterable(const FGuid& ObjectGuid) const
-{
-	return UnfilterableObjects.Find(ObjectGuid) != INDEX_NONE;
-}
-
-void FSequencer::AddUnfilterableObject(const FGuid& ObjectGuid)
-{
-	UnfilterableObjects.AddUnique(ObjectGuid);
-}
-
-bool FSequencer::IsShotFilteringOn() const
-{
-	return FilteringShots.Num() > 0;
-}
-
 float FSequencer::GetOverlayFadeCurve() const
 {
 	return OverlayCurve.GetLerp();
 }
 
-bool FSequencer::IsSectionVisible(UMovieSceneSection* Section) const
-{
-	// if no shots are being filtered, don't filter at all
-	bool bShowAll = !IsShotFilteringOn();
-
-	bool bIsAShotSection = Section->IsA<UMovieSceneShotSection>();
-
-	bool bIsUnfilterable = UnfilterableSections.Find(Section) != INDEX_NONE;
-
-	// do not draw if not within the filtering shot sections
-	TRange<float> SectionRange = Section->GetRange();
-	bool bIsVisible = bShowAll || bIsAShotSection || bIsUnfilterable;
-	for (int32 i = 0; i < FilteringShots.Num() && !bIsVisible; ++i)
-	{
-		TRange<float> ShotRange = FilteringShots[i]->GetRange();
-		bIsVisible |= SectionRange.Overlaps(ShotRange);
-	}
-
-	return bIsVisible;
-}
-
 void FSequencer::DeleteSelectedItems()
 {
 	if (Selection.GetActiveSelection() == FSequencerSelection::EActiveSelection::KeyAndSection)
diff --git a/Engine/Source/Editor/Sequencer/Private/Sequencer.h b/Engine/Source/Editor/Sequencer/Private/Sequencer.h
--- a/Engine/Source/Editor/Sequencer/Private/Sequencer.h
+++ b/Engine/Source/Editor/Sequencer/Private/Sequencer.h
@@ -190,36 +190,11 @@ public:
 	 */
 	void ZoomToSelectedSections();
 
-	/**
-	 * Gets all shots that are filtering currently
-	 */
-	const TArray< TWeakObjectPtr<UMovieSceneSection> >& GetFilteringShotSections() const;
-
-	/**
-	 * Checks to see if an object is unfilterable
-	 */
-	bool IsObjectUnfilterable(const FGuid& ObjectGuid) const;
-
-	/**
-	 * Declares an object unfilterable
-	 */
-	void AddUnfilterableObject(const FGuid& ObjectGuid);
-
-	/**
-	 * Checks to see if shot filtering is on
-	 */
-	bool IsShotFilteringOn() const;
-
 	/**
 	 * Gets the overlay fading animation curve lerp
 	 */
 	float GetOverlayFadeCurve() const;
 
-	/**
-	 * Checks if a section is visible, given current shot filtering
-	 */
-	bool IsSectionVisible(UMovieSceneSection* Section) const;
-
 	/** Gets the command bindings for the sequencer */
 	TSharedPtr<FUICommandList> GetCommandBindings() { return SequencerCommandBindings; }
 
@@ -486,10 +461,6 @@ private:
 	/** List of tools we own */
 	TArray< TSharedPtr<FMovieSceneTrackEditor> > TrackEditors;
 
-	/** The editors we keep track of for special behaviors */
-	TWeakPtr<FMovieSceneTrackEditor> ShotTrackEditor;
-	TWeakPtr<FMovieSceneTrackEditor> SkeletalAnimationTrackEditor;
-
 	/** Listener for object changes being made while this sequencer is open*/
 	TSharedPtr< class ISequencerObjectChangeListener > ObjectChangeListener;
 
@@ -505,15 +476,6 @@ private:
 	/** The asset editor that created this Sequencer if any */
 	TWeakPtr<IToolkitHost> ToolkitHost;
 
-	/** A list of all shots that are acting as filters */
-	TArray< TWeakObjectPtr<class UMovieSceneSection> > FilteringShots;
-
-	/** A list of sections that are considered unfilterable */
-	TArray< TWeakObjectPtr<class UMovieSceneSection> > UnfilterableSections;
-
-	/** A list of object guids that will be visible, regardless of shot filters */
-	TArray<FGuid> UnfilterableObjects;
-
 	/** Stack of movie scenes.  The first element is always the root movie scene.  The last element is the focused movie scene */
 	TArray< TSharedRef<FMovieSceneSequenceInstance> > MovieSceneStack;
 
diff --git a/Engine/Source/Editor/Sequencer/Private/SequencerDisplayNode.cpp b/Engine/Source/Editor/Sequencer/Private/SequencerDisplayNode.cpp
--- a/Engine/Source/Editor/Sequencer/Private/SequencerDisplayNode.cpp
+++ b/Engine/Source/Editor/Sequencer/Private/SequencerDisplayNode.cpp
@@ -398,7 +398,7 @@ bool FSequencerDisplayNode::IsHidden() const
 	// If shot filtering is off and clean view is on, but the node isn't pinned
 	return !bCachedShotFilteredVisibility ||
 		(ParentTree.HasActiveFilter() && !ParentTree.IsNodeFiltered(AsShared())) ||
-		(!GetSequencer().IsShotFilteringOn() && GetSequencer().GetSettings()->GetIsUsingCleanView() && !bNodeIsPinned);
+		(GetSequencer().GetSettings()->GetIsUsingCleanView() && !bNodeIsPinned);
 }
 
 bool FSequencerDisplayNode::HasVisibleChildren() const
@@ -544,13 +544,7 @@ bool FTrackNode::GetShotFilteredVisibilityToCache() const
 		return AssociatedType->IsVisibleWhenEmpty();
 	}
 
-	// if no child sections are visible, neither is the entire section
-	bool bAnySectionsVisible = false;
-	for (int32 i = 0; i < Sections.Num() && !bAnySectionsVisible; ++i)
-	{
-		bAnySectionsVisible = GetSequencer().IsSectionVisible(Sections[i]->GetSectionObject());
-	}
-	return bAnySectionsVisible;
+	return (Sections.Num() > 0);
 }
 
 void FTrackNode::GetChildKeyAreaNodesRecursively(TArray< TSharedRef<class FSectionKeyAreaNode> >& OutNodes) const
@@ -662,7 +656,7 @@ bool FObjectBindingNode::GetShotFilteredVisibilityToCache() const
 {
 	// if shot filtering is off and we are not unfilterable
 	// always show the object nodes (clean view is handled elsewhere)
-	return !GetSequencer().IsShotFilteringOn() || GetSequencer().IsObjectUnfilterable(ObjectBinding) || HasVisibleChildren();
+	return HasVisibleChildren();
 }
 
 const UClass* FObjectBindingNode::GetClassForObjectBinding()
