Forgot to check in

diff --git a/Engine/Source/Editor/Sequencer/Public/ISequencerTrackEditor.h b/Engine/Source/Editor/Sequencer/Public/ISequencerTrackEditor.h
new file mode 100644
--- /dev/null
+++ b/Engine/Source/Editor/Sequencer/Public/ISequencerTrackEditor.h
@@ -0,0 +1,119 @@
+// Copyright 1998-2015 Epic Games, Inc. All Rights Reserved.
+
+#pragma once
+
+
+class ISequencerSection;
+class UMovieSceneSection;
+class UMovieSceneTrack;
+
+
+/**
+ * Interface for sequencer track editors.
+ */
+class ISequencerTrackEditor
+{
+public:
+
+	/**
+	 * Manually adds a key.
+
+	 * @param ObjectGuid The Guid of the object that we are adding a key to.
+	 * @param AdditionalAsset An optional asset that can be added with the key.
+	 */
+	virtual void AddKey(const FGuid& ObjectGuid, UObject* AdditionalAsset = NULL) = 0;
+	
+	/**
+	 * Add a new track to the sequence.
+	 */
+	virtual UMovieSceneTrack* AddTrack(UMovieScene* FocusedMovieScene, const FGuid& ObjectHandle, TSubclassOf<class UMovieSceneTrack> TrackClass, FName UniqueTypeName) = 0;
+
+	/**
+	 * Allows the track editors to bind commands.
+	 *
+	 * @param SequencerCommandBindings The command bindings to map to.
+	*/
+	virtual void BindCommands(TSharedRef<FUICommandList> SequencerCommandBindings) = 0;
+
+	/**
+	 * Builds up the object binding context menu for the outliner.
+	 *
+	 * @param MenuBuilder The menu builder to change.
+	 * @param ObjectBinding The object binding this is for.
+	 * @param ObjectClass The class of the object this is for.
+	 */
+	virtual void BuildObjectBindingContextMenu(FMenuBuilder& MenuBuilder, const FGuid& ObjectBinding, const UClass* ObjectClass) = 0;
+
+	/**
+	 * Builds up the object binding edit buttons for the outliner.
+	 *
+	 * @param EditBox The edit box to add buttons to.
+	 * @param ObjectBinding The object binding this is for.
+	 * @param ObjectClass The class of the object this is for.
+	 */
+	virtual void BuildObjectBindingEditButtons(TSharedPtr<SHorizontalBox> EditBox, const FGuid& ObjectBinding, const UClass* ObjectClass) = 0;
+
+	/**
+	 * Builds up the object binding track menu for the outliner.
+	 *
+	 * @param MenuBuilder The menu builder to change.
+	 * @param ObjectBinding The object binding this is for.
+	 * @param ObjectClass The class of the object this is for.
+	 */
+	virtual void BuildObjectBindingTrackMenu(FMenuBuilder& MenuBuilder, const FGuid& ObjectBinding, const UClass* ObjectClass) = 0;
+
+	/**
+	 * Builds an edit widget for the outliner nodes which represent tracks which are edited by this editor.
+	 * @param ObjectBinding The object binding associated with the track being edited by this editor.
+	 * @param Track The track being edited by this editor.
+	 * @returns The the widget to display in the outliner, or an empty shared ptr if not widget is to be displayed.
+	 */
+	virtual TSharedPtr<SWidget> BuildOutlinerEditWidget(const FGuid& ObjectBinding, UMovieSceneTrack* Track) = 0;
+
+	/**
+	 * Called when an asset is dropped into Sequencer. Can potentially consume the asset
+	 * so it doesn't get added as a spawnable.
+	 *
+	 * @param Asset The asset that is dropped in.
+	 * @param TargetObjectGuid The object guid this asset is dropped onto, if applicable.
+	 * @return true if we want to consume this asset, false otherwise.
+	 */
+	virtual bool HandleAssetAdded(UObject* Asset, const FGuid& TargetObjectGuid) = 0;
+
+	/** Gets whether the tool can key all*/
+	virtual bool IsAllowedKeyAll() const = 0;
+
+	/** Gets whether the tool can legally autokey */
+	virtual bool IsAllowedToAutoKey() const = 0;
+
+	/**
+	 * Called to generate a section layout for a particular section.
+	 *
+	 * @param SectionObject The section to make UI for.
+	 * @param SectionName The name of the section.
+	 */
+	virtual TSharedRef<ISequencerSection> MakeSectionInterface(UMovieSceneSection& SectionObject, UMovieSceneTrack* Track) = 0;
+
+	/** Called when the instance of this track editor is released */
+	virtual void OnRelease() = 0;
+
+	/**
+	 * Returns whether a track class is supported by this tool.
+	 *
+	 * @param TrackClass The track class that could be supported.
+	 * @return true if the type is supported.
+	 */
+	virtual bool SupportsType(TSubclassOf<UMovieSceneTrack> TrackClass) const = 0;
+
+	/**
+	 * Ticks this tool.
+	 *
+	 * @param DeltaTime The time since the last tick.
+	 */
+	virtual void Tick(float DeltaTime) = 0;
+
+public:
+
+	/** Virtual destructor. */
+	virtual ~ISequencerTrackEditor() { }
+};
