Sequencer: Clamp scrubbing when auto scroll isn't enabled.

diff --git a/Engine/Source/Editor/Sequencer/Private/Sequencer.cpp b/Engine/Source/Editor/Sequencer/Private/Sequencer.cpp
--- a/Engine/Source/Editor/Sequencer/Private/Sequencer.cpp
+++ b/Engine/Source/Editor/Sequencer/Private/Sequencer.cpp
@@ -666,7 +666,7 @@ void FSequencer::SetGlobalTimeDirectly( float NewTime )
 {
 	float LastTime = ScrubPosition;
 
-	if (Settings->GetShowRangeSlider() && Settings->GetLockInOutToStartEndRange())
+	if (!Settings->GetAutoScrollEnabled())
 	{
 		UMovieScene* FocusedMovieScene = GetFocusedMovieSceneSequence()->GetMovieScene();
 		NewTime = FMath::Clamp(NewTime, FocusedMovieScene->InTime, FocusedMovieScene->OutTime);
