Sequencer: Fix padding in the top timeline slider so that it's in line with the bototm timeline slider.

diff --git a/Engine/Source/Editor/Sequencer/Private/SSequencer.cpp b/Engine/Source/Editor/Sequencer/Private/SSequencer.cpp
--- a/Engine/Source/Editor/Sequencer/Private/SSequencer.cpp
+++ b/Engine/Source/Editor/Sequencer/Private/SSequencer.cpp
@@ -372,6 +372,7 @@ void SSequencer::Construct( const FArguments& InArgs, TSharedRef< class FSequenc
 												SNew( SBorder )
 													.BorderImage( FEditorStyle::GetBrush("ToolPanel.GroupBorder") )
 													.BorderBackgroundColor( FLinearColor(.50f, .50f, .50f, 1.0f ) )
+													.Padding(0)
 													[
 														TopTimeSlider
 													]
