Sequencer: Fix add objects not doing anything.

diff --git a/Engine/Source/Editor/Sequencer/Private/SequencerDisplayNode.cpp b/Engine/Source/Editor/Sequencer/Private/SequencerDisplayNode.cpp
--- a/Engine/Source/Editor/Sequencer/Private/SequencerDisplayNode.cpp
+++ b/Engine/Source/Editor/Sequencer/Private/SequencerDisplayNode.cpp
@@ -654,9 +654,8 @@ FNodePadding FObjectBindingNode::GetNodePadding() const
 
 bool FObjectBindingNode::GetShotFilteredVisibilityToCache() const
 {
-	// if shot filtering is off and we are not unfilterable
-	// always show the object nodes (clean view is handled elsewhere)
-	return HasVisibleChildren();
+	// Clean view is handled elsewhere
+	return true;
 }
 
 const UClass* FObjectBindingNode::GetClassForObjectBinding()
