Sequencer: Add (s) for seconds for time snap values.

diff --git a/Engine/Source/Editor/Sequencer/Private/SequencerCommonHelpers.h b/Engine/Source/Editor/Sequencer/Private/SequencerCommonHelpers.h
--- a/Engine/Source/Editor/Sequencer/Private/SequencerCommonHelpers.h
+++ b/Engine/Source/Editor/Sequencer/Private/SequencerCommonHelpers.h
@@ -52,6 +52,23 @@ public:
 		return SnapValues;
 	}
 
+	static const TArray<SNumericDropDown<float>::FNamedValue>& GetSecondsSnapValues()
+	{
+		static TArray<SNumericDropDown<float>::FNamedValue> SecondsSnapValues;
+		static bool bInitSecondsSnapValues = false;
+		if (!bInitSecondsSnapValues)
+		{
+			bInitSecondsSnapValues = true;
+			SecondsSnapValues.Add( SNumericDropDown<float>::FNamedValue( 0.001f, LOCTEXT( "Snap_OneThousandthSeconds", "0.001s" ), LOCTEXT( "SnapDescription_OneThousandthSeconds", "Set snap to 1/1000th of a second" ) ) );
+			SecondsSnapValues.Add( SNumericDropDown<float>::FNamedValue( 0.01f, LOCTEXT( "Snap_OneHundredthSeconds", "0.01s" ), LOCTEXT( "SnapDescription_OneHundredthSeconds", "Set snap to 1/100th of a second" ) ) );
+			SecondsSnapValues.Add( SNumericDropDown<float>::FNamedValue( 0.1f, LOCTEXT( "Snap_OneTenthSeconds", "0.1s" ), LOCTEXT( "SnapDescription_OneTenthSeconds", "Set snap to 1/10th of a second" ) ) );
+			SecondsSnapValues.Add( SNumericDropDown<float>::FNamedValue( 1.0f, LOCTEXT( "Snap_OneSeconds", "1s" ), LOCTEXT( "SnapDescription_OneSeconds", "Set snap to 1 second" ) ) );
+			SecondsSnapValues.Add( SNumericDropDown<float>::FNamedValue( 10.0f, LOCTEXT( "Snap_TenSeconds", "10s" ), LOCTEXT( "SnapDescription_TenSeconds", "Set snap to 10 seconds" ) ) );
+			SecondsSnapValues.Add( SNumericDropDown<float>::FNamedValue( 100.0f, LOCTEXT( "Snap_OneHundredSeconds", "100s" ), LOCTEXT( "SnapDescription_OneHundredSeconds", "Set snap to 100 seconds" ) ) );
+		}
+		return SecondsSnapValues;
+	}
+
 	static const TArray<SNumericDropDown<float>::FNamedValue>& GetFrameRateSnapValues()
 	{
 		static TArray<SNumericDropDown<float>::FNamedValue> FrameRateSnapValues;
@@ -80,7 +97,7 @@ public:
 		if (!bInitTimeSnapValues)
 		{
 			bInitTimeSnapValues = true;
-			TimeSnapValues.Append(GetSnapValues());
+			TimeSnapValues.Append(GetSecondsSnapValues());
 			TimeSnapValues.Append(GetFrameRateSnapValues());
 		}
 		return TimeSnapValues;
