Sequencer: Don't ensure when trying to add spawnables for assets that aren't actors.

diff --git a/Engine/Source/Editor/Sequencer/Private/Sequencer.cpp b/Engine/Source/Editor/Sequencer/Private/Sequencer.cpp
--- a/Engine/Source/Editor/Sequencer/Private/Sequencer.cpp
+++ b/Engine/Source/Editor/Sequencer/Private/Sequencer.cpp
@@ -1356,7 +1356,7 @@ FGuid FSequencer::AddSpawnableForAssetOrClass( UObject* Object, UObject* Counter
 			}
 		}
 
-		if (ensure(NewBlueprint != nullptr))
+		if (NewBlueprint != nullptr)
 		{
 			if ((NewBlueprint->GeneratedClass != nullptr) && FBlueprintEditorUtils::IsActorBased(NewBlueprint))
 			{
