Sequencer - Simplify visibility so that it's just based on the text filter for now.  We will revisit this in the future once we have more use cases.
+ Removed "Clean View" it was not being used.
+ Removed pinned nodes there was no UI for pinning and pinning in code is not neccessary with no shot filtering.
+ Tracks are now always visible, with shot filtering gone this was only being used to hide empty tracks.  If we want this behavior we should just delete tracks when they become empty.
+ Removed caching code which is not neccessary now that visibility is simpler.

diff --git a/Engine/Source/Editor/Sequencer/Private/SSequencer.cpp b/Engine/Source/Editor/Sequencer/Private/SSequencer.cpp
--- a/Engine/Source/Editor/Sequencer/Private/SSequencer.cpp
+++ b/Engine/Source/Editor/Sequencer/Private/SSequencer.cpp
@@ -792,7 +792,6 @@ TSharedRef<SWidget> SSequencer::MakeGeneralMenu()
 		if (Sequencer.Pin()->IsLevelEditorSequencer())
 		{
 			MenuBuilder.AddMenuEntry( FSequencerCommands::Get().ToggleDetailsView );
-			MenuBuilder.AddMenuEntry( FSequencerCommands::Get().ToggleCleanView );
 		}
 
 		MenuBuilder.AddMenuEntry( FSequencerCommands::Get().ToggleAutoScroll );
@@ -939,8 +938,6 @@ void SSequencer::UpdateLayoutTree()
 	// This must come after the selection state has been restored so that the curve editor is populated with the correctly selected nodes
 	CurveEditor->SetSequencerNodeTree(SequencerNodeTree);
 
-	SequencerNodeTree->UpdateCachedVisibilityBasedOnShotFiltersChanged();
-
 	// Continue broadcasting selection changes
 	Sequencer.Pin()->GetSelection().ResumeBroadcast();
 }
@@ -960,8 +957,6 @@ void SSequencer::UpdateBreadcrumbs(const TArray< TWeakObjectPtr<class UMovieScen
 		// The current breadcrumb is not a moviescene so we need to make a new breadcrumb in order return to the parent moviescene later
 		BreadcrumbTrail->PushCrumb( CrumbName, FSequencerBreadcrumb( FocusedMovieSceneInstance ) );
 	}
-
-	SequencerNodeTree->UpdateCachedVisibilityBasedOnShotFiltersChanged();
 }
 
 void SSequencer::ResetBreadcrumbs()
diff --git a/Engine/Source/Editor/Sequencer/Private/SSequencerTreeView.cpp b/Engine/Source/Editor/Sequencer/Private/SSequencerTreeView.cpp
--- a/Engine/Source/Editor/Sequencer/Private/SSequencerTreeView.cpp
+++ b/Engine/Source/Editor/Sequencer/Private/SSequencerTreeView.cpp
@@ -437,8 +437,6 @@ void SSequencerTreeView::OnExpansionChanged(FDisplayNodeRef InItem, bool bIsExpa
 			SetItemExpansion(Child, true);
 		}
 	}
-
-	InItem->UpdateCachedShotFilteredVisibility();
 }
 
 void SSequencerTreeView::OnGetChildren(FDisplayNodeRef InParent, TArray<FDisplayNodeRef>& OutChildren) const
diff --git a/Engine/Source/Editor/Sequencer/Private/Sequencer.cpp b/Engine/Source/Editor/Sequencer/Private/Sequencer.cpp
--- a/Engine/Source/Editor/Sequencer/Private/Sequencer.cpp
+++ b/Engine/Source/Editor/Sequencer/Private/Sequencer.cpp
@@ -2198,18 +2198,6 @@ void FSequencer::BindSequencerCommands()
 		FCanExecuteAction::CreateLambda( []{ return true; } ),
 		FIsActionChecked::CreateLambda( [this]{ return Settings->GetDetailsViewVisible(); } ) );
 
-	SequencerCommandBindings->MapAction(
-		Commands.ToggleCleanView,
-		FExecuteAction::CreateLambda( [this]{
-			Settings->SetIsUsingCleanView( !Settings->GetIsUsingCleanView() );
-			if (SequencerWidget.IsValid())
-			{
-				SequencerWidget->GetTreeView()->Refresh();
-			}
-		} ),
-		FCanExecuteAction::CreateLambda( []{ return true; } ),
-		FIsActionChecked::CreateLambda( [this]{ return Settings->GetIsUsingCleanView(); } ) );
-
 	SequencerCommandBindings->MapAction(
 		Commands.ToggleShowFrameNumbers,
 		FExecuteAction::CreateLambda( [this]{ Settings->SetShowFrameNumbers( !Settings->GetShowFrameNumbers() ); } ),
diff --git a/Engine/Source/Editor/Sequencer/Private/SequencerCommands.cpp b/Engine/Source/Editor/Sequencer/Private/SequencerCommands.cpp
--- a/Engine/Source/Editor/Sequencer/Private/SequencerCommands.cpp
+++ b/Engine/Source/Editor/Sequencer/Private/SequencerCommands.cpp
@@ -56,7 +56,6 @@ void FSequencerCommands::RegisterCommands()
 
 	UI_COMMAND( FindInContentBrowser, "Find in Content Browser", "Find the viewed actor animation in the content browser", EUserInterfaceActionType::Button, FInputChord() );
 	UI_COMMAND( ToggleDetailsView, "Details View", "Enable details view for selected sections and keys", EUserInterfaceActionType::ToggleButton, FInputChord() );
-	UI_COMMAND( ToggleCleanView, "Clean View", "Enable 'Clean View' mode which displays only global tracks when no filter is applied", EUserInterfaceActionType::ToggleButton, FInputChord() );
 	
 	UI_COMMAND( ToggleShowCurveEditor, "Curve Editor", "Show the animation keys in a curve editor", EUserInterfaceActionType::ToggleButton, FInputChord() );
 
diff --git a/Engine/Source/Editor/Sequencer/Private/SequencerCommands.h b/Engine/Source/Editor/Sequencer/Private/SequencerCommands.h
--- a/Engine/Source/Editor/Sequencer/Private/SequencerCommands.h
+++ b/Engine/Source/Editor/Sequencer/Private/SequencerCommands.h
@@ -125,10 +125,7 @@ public:
 
 	/** Toggles whether the details view is enabled in the level editor. */
 	TSharedPtr< FUICommandInfo > ToggleDetailsView;
-	
-	/** Toggles whether the "clean view" is enabled in the level editor. */
-	TSharedPtr< FUICommandInfo > ToggleCleanView;
-	
+
 	/** Turns auto scroll on and off. */
 	TSharedPtr< FUICommandInfo > ToggleAutoScroll;
 
diff --git a/Engine/Source/Editor/Sequencer/Private/SequencerDisplayNode.cpp b/Engine/Source/Editor/Sequencer/Private/SequencerDisplayNode.cpp
--- a/Engine/Source/Editor/Sequencer/Private/SequencerDisplayNode.cpp
+++ b/Engine/Source/Editor/Sequencer/Private/SequencerDisplayNode.cpp
@@ -128,8 +128,6 @@ FSequencerDisplayNode::FSequencerDisplayNode( FName InNodeName, TSharedPtr<FSequ
 	, ParentTree( InParentTree )
 	, NodeName( InNodeName )
 	, bExpanded( false )
-	, bCachedShotFilteredVisibility( true )
-	, bNodeIsPinned( false )
 {
 	
 }
@@ -381,51 +379,16 @@ void FSequencerDisplayNode::SetExpansionState(bool bInExpanded)
 
 	// Expansion state has changed, save it to the movie scene now
 	ParentTree.SaveExpansionState( *this, bExpanded );
-
-	UpdateCachedShotFilteredVisibility();
 }
 
 bool FSequencerDisplayNode::IsExpanded() const
 {
-	return ParentTree.HasActiveFilter() ? ParentTree.IsNodeFiltered( AsShared() ) : bExpanded;
+	return ParentTree.HasActiveFilter() ? true : bExpanded;
 }
 
 bool FSequencerDisplayNode::IsHidden() const
 {
-	// Not visible if:
-	// Its cached shot visibility is false OR
-	// There is a search filter, and it's not filtered, OR
-	// If shot filtering is off and clean view is on, but the node isn't pinned
-	return !bCachedShotFilteredVisibility ||
-		(ParentTree.HasActiveFilter() && !ParentTree.IsNodeFiltered(AsShared())) ||
-		(GetSequencer().GetSettings()->GetIsUsingCleanView() && !bNodeIsPinned);
-}
-
-bool FSequencerDisplayNode::HasVisibleChildren() const
-{
-	for (int32 i = 0; i < ChildNodes.Num(); ++i)
-	{
-		if (ChildNodes[i]->bCachedShotFilteredVisibility) {return true;}
-	}
-	return false;
-}
-
-void FSequencerDisplayNode::UpdateCachedShotFilteredVisibility()
-{
-	// Tell our children to update their visibility first
-	for( int32 ChildIndex = 0; ChildIndex < ChildNodes.Num(); ++ChildIndex )
-	{
-		ChildNodes[ChildIndex]->UpdateCachedShotFilteredVisibility();
-	}
-	
-	// then cache our visibility
-	// this must be done after the children, because it relies on child cached visibility
-	bCachedShotFilteredVisibility = GetShotFilteredVisibilityToCache();
-}
-
-void FSequencerDisplayNode::PinNode()
-{
-	bNodeIsPinned = true;
+	return ParentTree.HasActiveFilter() && !ParentTree.IsNodeFiltered(AsShared());
 }
 
 TSharedRef<FGroupedKeyArea> FSequencerDisplayNode::GetKeyGrouping(int32 InSectionIndex)
@@ -474,11 +437,6 @@ FNodePadding FSectionKeyAreaNode::GetNodePadding() const
 	return FNodePadding(0.f, 1.f);
 }
 
-bool FSectionKeyAreaNode::GetShotFilteredVisibilityToCache() const
-{
-	return true;
-}
-
 TSharedRef<SWidget> FSectionKeyAreaNode::GenerateEditWidgetForOutliner()
 {
 	// @todo - Sequencer - Support multiple sections/key areas?
@@ -532,21 +490,6 @@ void FTrackNode::SetSectionAsKeyArea( TSharedRef<IKeyArea>& KeyArea )
 	TopLevelKeyNode->AddKeyArea( KeyArea );
 }
 
-bool FTrackNode::GetShotFilteredVisibilityToCache() const
-{
-	if (AssociatedType->HasShowableData() == false)
-	{
-		return false;
-	}
-
-	if ( Sections.Num() == 0 )
-	{
-		return AssociatedType->IsVisibleWhenEmpty();
-	}
-
-	return (Sections.Num() > 0);
-}
-
 void FTrackNode::GetChildKeyAreaNodesRecursively(TArray< TSharedRef<class FSectionKeyAreaNode> >& OutNodes) const
 {
 	FSequencerDisplayNode::GetChildKeyAreaNodesRecursively(OutNodes);
@@ -652,12 +595,6 @@ FNodePadding FObjectBindingNode::GetNodePadding() const
 	return FNodePadding(SequencerNodeConstants::CommonPadding * 2, 0.f);
 }
 
-bool FObjectBindingNode::GetShotFilteredVisibilityToCache() const
-{
-	// Clean view is handled elsewhere
-	return true;
-}
-
 const UClass* FObjectBindingNode::GetClassForObjectBinding()
 {
 	FSequencer& ParentSequencer = GetSequencer();
@@ -908,10 +845,4 @@ FNodePadding FSectionCategoryNode::GetNodePadding() const
 	return FNodePadding(SequencerNodeConstants::CommonPadding/2);
 }
 
-bool FSectionCategoryNode::GetShotFilteredVisibilityToCache() const
-{
-	// this node is only visible if at least one child node is visible
-	return HasVisibleChildren();
-}
-
 #undef LOCTEXT_NAMESPACE
diff --git a/Engine/Source/Editor/Sequencer/Private/SequencerDisplayNode.h b/Engine/Source/Editor/Sequencer/Private/SequencerDisplayNode.h
--- a/Engine/Source/Editor/Sequencer/Private/SequencerDisplayNode.h
+++ b/Engine/Source/Editor/Sequencer/Private/SequencerDisplayNode.h
@@ -234,12 +234,6 @@ public:
 	 */
 	bool IsHidden() const;
 
-	/** Updates the cached shot filtered visibility flag */
-	void UpdateCachedShotFilteredVisibility();
-
-	/** Pins this node, forcing it to the top of the sequencer. Affects visibility */
-	void PinNode();
-
 	/** Initialize this node with expansion states and virtual offsets */
 	void Initialize(float InVirtualTop, float InVirtualBottom);
 
@@ -255,41 +249,6 @@ public:
 	/** Get the key grouping for the specified section index */
 	TSharedRef<FGroupedKeyArea> GetKeyGrouping(int32 InSectionIndex);
 
-protected:
-	/**
-	 * Visibility of nodes is a complex situation. There are 7 different factors to account for:
-	 * - Expansion State of Parents
-	 * - Search Filtering
-	 * - Shot Filtering
-	 * - Clean View
-	 * - Pinned Nodes
-	 * - Manually keyed nodes with a single key are not displayed
-	 * - Sub Movie Scene Filtering (maybe, we might just ignore this completely)
-	 *
-	 * It is handled by first caching the shot filtering state in each node.
-	 * Manually keyed nodes with a single key are cached at this point as well.
-	 * This cached state is updated only when shot filtering changes.
-	 * This is done for a) Simplicity, to break things apart
-	 * and b) Efficiency, since that is likely the most expensive visibility flag
-	 * and, most importantly, c) because it is the only pass that is dependent on
-	 * the way the tree structure is laid out (it requires looking at child nodes).
-	 *
-	 * Then, the remaining visibility flags are AND'd together in these groups
-	 * 1) CachedShotFilteringFlag
-	 * 2) Search Filtering, which overrides Expansion State of Parents
-	 * 3) Clean View and Pinned Nodes, which intersects with shot filtering global enabling
-	 * 4) Sub Movie Scenes (not implemented yet)
-	 */
-
-	/**
-	 * Returns visibility, calculating the conditions for which this
-	 * node's shot filtering visibility is cached
-	 */
-	virtual bool GetShotFilteredVisibilityToCache() const = 0;
-	
-	/** Whether this node has visible children, based on cached shot filtering visibility only */
-	bool HasVisibleChildren() const;
-
 protected:
 	/** The virtual offset of this item from the top of the tree, irrespective of expansion states */
 	float VirtualTop;
@@ -307,10 +266,6 @@ protected:
 	FName NodeName;
 	/** Whether or not the node is expanded */
 	bool bExpanded;
-	/** A cached state of what this node's visibility is based only on shot filtering */
-	bool bCachedShotFilteredVisibility;
-	/** Whether this node is pinned to the top of the sequencer */
-	bool bNodeIsPinned;
 	/** Transient grouped keys for this node */
 	TArray< TSharedPtr<FGroupedKeyArea> > KeyGroupings;
 };
@@ -342,7 +297,6 @@ public:
 	virtual float GetNodeHeight() const override;
 	virtual FNodePadding GetNodePadding() const override;
 	virtual FText GetDisplayName() const override { return DisplayName; }
-	virtual bool GetShotFilteredVisibilityToCache() const override;
 	virtual TSharedRef<SWidget> GenerateEditWidgetForOutliner() override;
 
 	/**
@@ -401,7 +355,6 @@ public:
 	virtual float GetNodeHeight() const override;
 	virtual FNodePadding GetNodePadding() const override;
 	virtual FText GetDisplayName() const override;
-	virtual bool GetShotFilteredVisibilityToCache() const override;
 	virtual void GetChildKeyAreaNodesRecursively(TArray< TSharedRef<class FSectionKeyAreaNode> >& OutNodes) const override;
 	virtual TSharedRef<SWidget> GenerateEditWidgetForOutliner() override;
 
@@ -477,7 +430,6 @@ public:
 	virtual FText GetDisplayName() const override;
 	virtual float GetNodeHeight() const override;
 	virtual FNodePadding GetNodePadding() const override;
-	virtual bool GetShotFilteredVisibilityToCache() const override;
 	TSharedRef<SWidget> GenerateEditWidgetForOutliner() override;
 	
 	/** @return The object binding on this node */
@@ -526,7 +478,6 @@ public:
 	virtual float GetNodeHeight() const override;
 	virtual FNodePadding GetNodePadding() const override;
 	virtual FText GetDisplayName() const override { return DisplayName; }
-	virtual bool GetShotFilteredVisibilityToCache() const override;
 private:
 	/** The display name of the category */
 	FText DisplayName;
diff --git a/Engine/Source/Editor/Sequencer/Private/SequencerNodeTree.cpp b/Engine/Source/Editor/Sequencer/Private/SequencerNodeTree.cpp
--- a/Engine/Source/Editor/Sequencer/Private/SequencerNodeTree.cpp
+++ b/Engine/Source/Editor/Sequencer/Private/SequencerNodeTree.cpp
@@ -43,8 +43,6 @@ void FSequencerNodeTree::Update()
 		NewRootNodes.Add( SectionNode );
 	
 		MakeSectionInterfaces( TrackRef, SectionNode );
-
-		SectionNode->PinNode();
 	}
 
 
@@ -114,8 +112,6 @@ void FSequencerNodeTree::Update()
 		RootNodes.Add( SectionNode );
 
 		MakeSectionInterfaces( *ShotTrack, SectionNode );
-
-		SectionNode->PinNode();
 	}
 
 	// Add all other nodes after the shot track
@@ -360,11 +356,3 @@ void FSequencerNodeTree::FilterNodes( const FString& InFilter )
 		}
 	}
 }
-
-void FSequencerNodeTree::UpdateCachedVisibilityBasedOnShotFiltersChanged()
-{
-	for (int32 i = 0; i < RootNodes.Num(); ++i)
-	{
-		RootNodes[i]->UpdateCachedShotFilteredVisibility();
-	}
-}
diff --git a/Engine/Source/Editor/Sequencer/Private/SequencerNodeTree.h b/Engine/Source/Editor/Sequencer/Private/SequencerNodeTree.h
--- a/Engine/Source/Editor/Sequencer/Private/SequencerNodeTree.h
+++ b/Engine/Source/Editor/Sequencer/Private/SequencerNodeTree.h
@@ -78,9 +78,6 @@ public:
 	 * @return true if the node is to be expanded, false otherwise
 	 */
 	bool GetDefaultExpansionState( const FSequencerDisplayNode& Node ) const;
-	
-	/** Updates nodes visibility for use when shot filters change at all */
-	void UpdateCachedVisibilityBasedOnShotFiltersChanged();
 
 private:
 	/**
diff --git a/Engine/Source/Editor/Sequencer/Private/SequencerSettings.cpp b/Engine/Source/Editor/Sequencer/Private/SequencerSettings.cpp
--- a/Engine/Source/Editor/Sequencer/Private/SequencerSettings.cpp
+++ b/Engine/Source/Editor/Sequencer/Private/SequencerSettings.cpp
@@ -23,7 +23,6 @@ USequencerSettings::USequencerSettings( const FObjectInitializer& ObjectInitiali
 	CurveValueSnapInterval = 10.0f;
 	bSnapCurveValueToInterval = true;
 	bDetailsViewVisible = false;
-	bIsUsingCleanView = false;
 	bAutoScrollEnabled = false;
 	bShowCurveEditor = false;
 	bShowCurveEditorCurveToolTips = true;
@@ -281,20 +280,6 @@ void USequencerSettings::SetDetailsViewVisible(bool InbDetailsViewVisible )
 	}
 }
 
-bool USequencerSettings::GetIsUsingCleanView() const
-{
-	return bIsUsingCleanView;
-}
-
-void USequencerSettings::SetIsUsingCleanView(bool InbIsUsingCleanView )
-{
-	if (bIsUsingCleanView != InbIsUsingCleanView)
-	{
-		bIsUsingCleanView = InbIsUsingCleanView;
-		SaveConfig();
-	}
-}
-
 bool USequencerSettings::GetAutoScrollEnabled() const
 {
 	return bAutoScrollEnabled;
diff --git a/Engine/Source/Editor/Sequencer/Private/SequencerSettings.h b/Engine/Source/Editor/Sequencer/Private/SequencerSettings.h
--- a/Engine/Source/Editor/Sequencer/Private/SequencerSettings.h
+++ b/Engine/Source/Editor/Sequencer/Private/SequencerSettings.h
@@ -134,11 +134,6 @@ public:
 	/** Sets whether or not the details view is visible. */
 	void SetDetailsViewVisible(bool InbDetailsViewVisible);
 
-	/** Gets whether or not the 'Clean View' is enabled. In 'Clean View' mode only global tracks are displayed when no filter is applied. */
-	bool GetIsUsingCleanView() const;
-	/** Sets whether or not the 'Clean View' is enabled. In 'Clean View' mode only global tracks are displayed when no filter is applied. */
-	void SetIsUsingCleanView(bool InbIsUsingCleanView);
-
 	/** Gets whether or not auto-scroll is enabled. */
 	bool GetAutoScrollEnabled() const;
 	/** Sets whether or not auto-scroll is enabled. */
@@ -215,9 +210,6 @@ protected:
 	UPROPERTY( config )
 	bool bDetailsViewVisible;
 
-	UPROPERTY( config )
-	bool bIsUsingCleanView;
-
 	UPROPERTY( config )
 	bool bAutoScrollEnabled;
 
diff --git a/Engine/Source/Runtime/MovieScene/Public/MovieSceneTrack.h b/Engine/Source/Runtime/MovieScene/Public/MovieSceneTrack.h
--- a/Engine/Source/Runtime/MovieScene/Public/MovieSceneTrack.h
+++ b/Engine/Source/Runtime/MovieScene/Public/MovieSceneTrack.h
@@ -43,16 +43,6 @@ public:
 	 * @return Whether or not this track has any data in it
 	 */
 	virtual bool IsEmpty() const PURE_VIRTUAL( UMovieSceneTrack::IsEmpty, return true; );
-	
-	/**
-	 * @return True if the data of this track show generate display nodes by default
-	 */
-	virtual bool HasShowableData() const {return true;}
-
-	/**
-	 * @returns True if this track should generate a display node even when it's empty.
-	 */
-	virtual bool IsVisibleWhenEmpty() const { return false; }
 
 	/**
 	 * @return Whether or not this track supports multiple row indices
diff --git a/Engine/Source/Runtime/MovieSceneTracks/Private/Tracks/MovieSceneSlomoTrack.cpp b/Engine/Source/Runtime/MovieSceneTracks/Private/Tracks/MovieSceneSlomoTrack.cpp
--- a/Engine/Source/Runtime/MovieSceneTracks/Private/Tracks/MovieSceneSlomoTrack.cpp
+++ b/Engine/Source/Runtime/MovieSceneTracks/Private/Tracks/MovieSceneSlomoTrack.cpp
@@ -5,11 +5,3 @@
 #include "MovieSceneSlomoTrack.h"
 #include "MovieSceneSlomoTrackInstance.h"
 
-
-/* UMovieSceneFloatTrack overrides
- *****************************************************************************/
-
-bool UMovieSceneSlomoTrack::IsVisibleWhenEmpty() const
-{
-	return true;
-}
diff --git a/Engine/Source/Runtime/MovieSceneTracks/Public/Tracks/MovieSceneMaterialTrack.h b/Engine/Source/Runtime/MovieSceneTracks/Public/Tracks/MovieSceneMaterialTrack.h
--- a/Engine/Source/Runtime/MovieSceneTracks/Public/Tracks/MovieSceneMaterialTrack.h
+++ b/Engine/Source/Runtime/MovieSceneTracks/Public/Tracks/MovieSceneMaterialTrack.h
@@ -61,8 +61,6 @@ public:
 	virtual bool IsEmpty() const override;
 	virtual TRange<float> GetSectionBoundaries() const override;
 	virtual const TArray<UMovieSceneSection*>& GetAllSections() const override;
-	virtual bool HasShowableData() const override { return true; }
-	virtual bool IsVisibleWhenEmpty() const override { return true; }
 
 	/**
 	 * Adds a scalar parameter key to the track. 
diff --git a/Engine/Source/Runtime/MovieSceneTracks/Public/Tracks/MovieScenePropertyTrack.h b/Engine/Source/Runtime/MovieSceneTracks/Public/Tracks/MovieScenePropertyTrack.h
--- a/Engine/Source/Runtime/MovieSceneTracks/Public/Tracks/MovieScenePropertyTrack.h
+++ b/Engine/Source/Runtime/MovieSceneTracks/Public/Tracks/MovieScenePropertyTrack.h
@@ -23,7 +23,6 @@ public:
 	virtual bool IsEmpty() const override;
 	virtual TRange<float> GetSectionBoundaries() const override;
 	virtual const TArray<UMovieSceneSection*>& GetAllSections() const override;
-	virtual bool HasShowableData() const override {return bSectionsAreShowable;}
 	
 	/**
 	 * Sets the property name for this animatable property
diff --git a/Engine/Source/Runtime/MovieSceneTracks/Public/Tracks/MovieSceneSlomoTrack.h b/Engine/Source/Runtime/MovieSceneTracks/Public/Tracks/MovieSceneSlomoTrack.h
--- a/Engine/Source/Runtime/MovieSceneTracks/Public/Tracks/MovieSceneSlomoTrack.h
+++ b/Engine/Source/Runtime/MovieSceneTracks/Public/Tracks/MovieSceneSlomoTrack.h
@@ -14,10 +14,4 @@ class UMovieSceneSlomoTrack
 	: public UMovieSceneFloatTrack
 {
 	GENERATED_BODY()
-
-public:
-
-	// UMovieSceneFloatTrack overrides
-
-	virtual bool IsVisibleWhenEmpty() const override;
 };
