Sequencer: Fix snap curve key value to interval button so that it's actually a toggle button.

diff --git a/Engine/Source/Editor/Sequencer/Private/SequencerCommands.cpp b/Engine/Source/Editor/Sequencer/Private/SequencerCommands.cpp
--- a/Engine/Source/Editor/Sequencer/Private/SequencerCommands.cpp
+++ b/Engine/Source/Editor/Sequencer/Private/SequencerCommands.cpp
@@ -52,7 +52,7 @@ void FSequencerCommands::RegisterCommands()
 	UI_COMMAND( ToggleSnapPlayTimeToInterval, "Snap to the Interval While Scrubbing", "Snap the play time to the interval while scrubbing", EUserInterfaceActionType::ToggleButton, FInputChord() );
 	UI_COMMAND( ToggleSnapPlayTimeToDraggedKey, "Snap to the Dragged Key", "Snap the play time to the dragged key", EUserInterfaceActionType::ToggleButton, FInputChord() );
 
-	UI_COMMAND( ToggleSnapCurveValueToInterval, "Snap Curve Key Values", "Snap curve keys to the value snapping interval", EUserInterfaceActionType::Button, FInputChord() );
+	UI_COMMAND( ToggleSnapCurveValueToInterval, "Snap Curve Key Values", "Snap curve keys to the value snapping interval", EUserInterfaceActionType::ToggleButton, FInputChord() );
 
 	UI_COMMAND( FindInContentBrowser, "Find in Content Browser", "Find the viewed actor animation in the content browser", EUserInterfaceActionType::Button, FInputChord() );
 	UI_COMMAND( ToggleDetailsView, "Details View", "Enable details view for selected sections and keys", EUserInterfaceActionType::ToggleButton, FInputChord() );
