Sequencer: Move auto scroll to frame range menu. Also, change the behavior so that auto scroll overrides lock in/out to start/end time range.

diff --git a/Engine/Source/Editor/Sequencer/Private/Sequencer.cpp b/Engine/Source/Editor/Sequencer/Private/Sequencer.cpp
--- a/Engine/Source/Editor/Sequencer/Private/Sequencer.cpp
+++ b/Engine/Source/Editor/Sequencer/Private/Sequencer.cpp
@@ -1153,7 +1153,7 @@ void FSequencer::OnViewRangeChanged( TRange<float> NewViewRange, EViewRangeInter
 	{
 		UMovieScene* FocusedMovieScene = GetFocusedMovieSceneSequence()->GetMovieScene();
 		// If locked, clamp the new range to clamp range
-		if (Settings->GetShowRangeSlider() && Settings->GetLockInOutToStartEndRange() && !bExpandClampRange)
+		if (!Settings->GetAutoScrollEnabled() && Settings->GetShowRangeSlider() && Settings->GetLockInOutToStartEndRange() && !bExpandClampRange)
 		{
 			if ( NewViewRange.GetLowerBoundValue() < FocusedMovieScene->StartTime)
 			{
@@ -2213,8 +2213,8 @@ void FSequencer::BindSequencerCommands()
 	SequencerCommandBindings->MapAction(
 		Commands.ToggleLockInOutToStartEndRange,
 		FExecuteAction::CreateLambda( [this]{ Settings->SetLockInOutToStartEndRange( !Settings->GetLockInOutToStartEndRange() ); } ),
-		FCanExecuteAction::CreateLambda( [this]{ return Settings->GetShowRangeSlider(); } ),
-		FIsActionChecked::CreateLambda( [this]{ return Settings->GetLockInOutToStartEndRange(); } ) );
+		FCanExecuteAction::CreateLambda( [this]{ return Settings->GetShowRangeSlider() && !Settings->GetAutoScrollEnabled(); } ),
+		FIsActionChecked::CreateLambda( [this]{ return !Settings->GetShowRangeSlider() || Settings->GetAutoScrollEnabled() ? false : Settings->GetLockInOutToStartEndRange(); } ) );
 
 	SequencerCommandBindings->MapAction(
 		Commands.ToggleIsSnapEnabled,
