Changed sequencer track area to be its own panel rather than an SOverlay

This means we can control the desired sizes of its child slots and ensure they get their desired size (SOverlay clips to its bounds)

diff --git a/Engine/Source/Editor/Sequencer/Private/SSequencerTrackArea.cpp b/Engine/Source/Editor/Sequencer/Private/SSequencerTrackArea.cpp
--- a/Engine/Source/Editor/Sequencer/Private/SSequencerTrackArea.cpp
+++ b/Engine/Source/Editor/Sequencer/Private/SSequencerTrackArea.cpp
@@ -11,17 +11,45 @@
 #include "ISequencerSection.h"
 #include "SSection.h"
 
+/** Structure representing a slot in the track area */
+class FTrackAreaSlot : public TSlotBase<FTrackAreaSlot>
+{
+public:
+	/** Construction from a track lane */
+	FTrackAreaSlot(const TSharedPtr<SSequencerTrackLane>& InSlotContent)
+	{
+		TrackLane = InSlotContent;
+		
+		HAlignment = HAlign_Fill;
+		VAlignment = VAlign_Top;
+
+		this->AttachWidget(
+			SNew(SWeakWidget)
+			.PossiblyNullContent(InSlotContent)
+			);
+	}
+
+	/** Get the vertical position of this slot inside its parent */
+	float GetVerticalOffset() const
+	{
+		auto PinnedTrackLane = TrackLane.Pin();
+		return PinnedTrackLane.IsValid() ? PinnedTrackLane->GetPhysicalPosition() : 0.f;
+	}
+
+	/** Horizontal/Vertical alignment for the slot */
+	EHorizontalAlignment HAlignment;
+	EVerticalAlignment VAlignment;
+
+	/** The track lane that we represent */
+	TWeakPtr<SSequencerTrackLane> TrackLane;
+};
+
 void SSequencerTrackArea::Construct( const FArguments& InArgs, TSharedRef<FSequencerTimeSliderController> InTimeSliderController, TSharedRef<SSequencer> InSequencerWidget )
 {
 	SequencerWidget = InSequencerWidget;
 	TimeSliderController = InTimeSliderController;
 
 	bLockInOutToStartEndRange = InArgs._LockInOutToStartEndRange;
-
-	//@todo Overlays don't tick by default, but we need to. Maybe try using Composition? 
-	bCanTick = true;
-
-	SOverlay::Construct(SOverlay::FArguments());
 }
 
 void SSequencerTrackArea::SetTreeView(const TSharedPtr<SSequencerTreeView>& InTreeView)
@@ -31,24 +59,8 @@ void SSequencerTrackArea::SetTreeView(const TSharedPtr<SSequencerTreeView>& InTr
 
 void SSequencerTrackArea::AddTrackSlot(const TSharedRef<FSequencerDisplayNode>& InNode, const TSharedPtr<SSequencerTrackLane>& InSlot)
 {
-	// Capture a weak reference to ensure that the slot can get cleaned up ok
-	TWeakPtr<SSequencerTrackLane> WeakSlot = InSlot;
-	auto Padding = [WeakSlot]{
-		auto PinnedSlot = WeakSlot.Pin();
-		return PinnedSlot.IsValid() ? FMargin(0.f, PinnedSlot->GetPhysicalPosition(), 0, 0) : FMargin(0);
-	};
-
 	TrackSlots.Add(InNode, InSlot);
-
-	typedef TAttribute<FMargin> Attr;
-
-	AddSlot()
-	.VAlign(VAlign_Top)
-	.Padding(Attr::Create(Attr::FGetter::CreateLambda(Padding)))
-	[
-		SNew(SWeakWidget)
-		.PossiblyNullContent(InSlot)
-	];
+	Children.Add(new FTrackAreaSlot(InSlot));
 }
 
 TSharedPtr<SSequencerTrackLane> SSequencerTrackArea::FindTrackSlot(const TSharedRef<FSequencerDisplayNode>& InNode)
@@ -56,15 +68,77 @@ TSharedPtr<SSequencerTrackLane> SSequencerTrackArea::FindTrackSlot(const TShared
 	return TrackSlots.FindRef(InNode).Pin();
 }
 
+void SSequencerTrackArea::OnArrangeChildren( const FGeometry& AllottedGeometry, FArrangedChildren& ArrangedChildren ) const
+{
+	for (int32 ChildIndex = 0; ChildIndex < Children.Num(); ++ChildIndex)
+	{
+		const FTrackAreaSlot& CurChild = Children[ChildIndex];
+
+		const EVisibility ChildVisibility = CurChild.GetWidget()->GetVisibility();
+		if (!ArrangedChildren.Accepts(ChildVisibility))
+		{
+			continue;
+		}
+
+		const FMargin Padding(0, CurChild.GetVerticalOffset(), 0, 0);
+
+		AlignmentArrangeResult XResult = AlignChild<Orient_Horizontal>(AllottedGeometry.Size.X, CurChild, Padding, 1.0f, false);
+		AlignmentArrangeResult YResult = AlignChild<Orient_Vertical>(AllottedGeometry.Size.Y, CurChild, Padding, 1.0f, false);
+
+		ArrangedChildren.AddWidget(ChildVisibility,
+			AllottedGeometry.MakeChild(
+				CurChild.GetWidget(),
+				FVector2D(XResult.Offset,YResult.Offset),
+				FVector2D(XResult.Size, YResult.Size)
+			)
+		);
+	}
+}
+
+FVector2D SSequencerTrackArea::ComputeDesiredSize( float ) const
+{
+	FVector2D MaxSize(0,0);
+	for (int32 ChildIndex = 0; ChildIndex < Children.Num(); ++ChildIndex)
+	{
+		const FTrackAreaSlot& CurChild = Children[ChildIndex];
+
+		const EVisibility ChildVisibilty = CurChild.GetWidget()->GetVisibility();
+		if (ChildVisibilty != EVisibility::Collapsed)
+		{
+			FVector2D ChildDesiredSize = CurChild.GetWidget()->GetDesiredSize();
+			MaxSize.X = FMath::Max(MaxSize.X, ChildDesiredSize.X);
+			MaxSize.Y = FMath::Max(MaxSize.Y, ChildDesiredSize.Y);
+		}
+	}
+
+	return MaxSize;
+}
+
+FChildren* SSequencerTrackArea::GetChildren()
+{
+	return &Children;
+}
+
 int32 SSequencerTrackArea::OnPaint(const FPaintArgs& Args, const FGeometry& AllottedGeometry, const FSlateRect& MyClippingRect, FSlateWindowElementList& OutDrawElements, int32 LayerId, const FWidgetStyle& InWidgetStyle, bool bParentEnabled ) const
 {
-	LayerId = SOverlay::OnPaint(Args, AllottedGeometry, MyClippingRect, OutDrawElements, LayerId, InWidgetStyle, bParentEnabled);
+	FArrangedChildren ArrangedChildren(EVisibility::Visible);
+	ArrangeChildren(AllottedGeometry, ArrangedChildren);
+
+	for (int32 ChildIndex = 0; ChildIndex < ArrangedChildren.Num(); ++ChildIndex)
+	{
+		FArrangedWidget& CurWidget = ArrangedChildren[ChildIndex];
+		FSlateRect ChildClipRect = MyClippingRect.IntersectionWith( CurWidget.Geometry.GetClippingRect() );
+		const int32 ThisWidgetLayerId = CurWidget.Widget->Paint( Args.WithNewParent(this), CurWidget.Geometry, ChildClipRect, OutDrawElements, LayerId + 1, InWidgetStyle, ShouldBeEnabled( bParentEnabled ) );
+
+		LayerId = FMath::Max(LayerId, ThisWidgetLayerId);
+	}
 
 	auto SequencerPin = SequencerWidget.Pin();
 	if (SequencerPin.IsValid())
 	{
-		return SequencerPin->GetEditTool().OnPaint(AllottedGeometry, MyClippingRect, OutDrawElements, LayerId);
+		return SequencerPin->GetEditTool().OnPaint(AllottedGeometry, MyClippingRect, OutDrawElements, LayerId + 1);
 	}
+
 	return LayerId;
 }
 
diff --git a/Engine/Source/Editor/Sequencer/Private/SSequencerTrackArea.h b/Engine/Source/Editor/Sequencer/Private/SSequencerTrackArea.h
--- a/Engine/Source/Editor/Sequencer/Private/SSequencerTrackArea.h
+++ b/Engine/Source/Editor/Sequencer/Private/SSequencerTrackArea.h
@@ -8,11 +8,12 @@ class SSequencerTrackLane;
 
 class FSequencerTimeSliderController;
 class FSequencer;
+class FTrackAreaSlot;
 
 /**
  * The area where tracks( rows of sections ) are displayed
  */
-class SSequencerTrackArea : public SOverlay
+class SSequencerTrackArea : public SPanel
 {
 public:
 	SLATE_BEGIN_ARGS( SSequencerTrackArea )
@@ -42,10 +43,18 @@ public:
 	virtual void OnMouseCaptureLost() override;
 	virtual FCursorReply OnCursorQuery( const FGeometry& MyGeometry, const FPointerEvent& CursorEvent ) const override;
 	virtual void Tick( const FGeometry& AllottedGeometry, const double InCurrentTime, const float InDeltaTime ) override;
+	virtual void OnArrangeChildren( const FGeometry& AllottedGeometry, FArrangedChildren& ArrangedChildren ) const override;
+	virtual FVector2D ComputeDesiredSize(float) const override;
+	virtual FChildren* GetChildren() override;
 
 	/** Access the cached geometry for this track area */
 	const FGeometry& GetCachedGeometry() const { return CachedGeometry; }
-	
+
+private:
+
+	/** The track area's children */
+	TPanelChildren<FTrackAreaSlot> Children;
+
 private:
 
 	/** Cached geometry */
