Moved definition to a header to appease Visual Studio's poor template instantiation

diff --git a/Engine/Source/Editor/Sequencer/Private/SSequencerTrackArea.cpp b/Engine/Source/Editor/Sequencer/Private/SSequencerTrackArea.cpp
--- a/Engine/Source/Editor/Sequencer/Private/SSequencerTrackArea.cpp
+++ b/Engine/Source/Editor/Sequencer/Private/SSequencerTrackArea.cpp
@@ -11,38 +11,24 @@
 #include "ISequencerSection.h"
 #include "SSection.h"
 
-/** Structure representing a slot in the track area */
-class FTrackAreaSlot : public TSlotBase<FTrackAreaSlot>
+FTrackAreaSlot::FTrackAreaSlot(const TSharedPtr<SSequencerTrackLane>& InSlotContent)
 {
-public:
-	/** Construction from a track lane */
-	FTrackAreaSlot(const TSharedPtr<SSequencerTrackLane>& InSlotContent)
-	{
-		TrackLane = InSlotContent;
-		
-		HAlignment = HAlign_Fill;
-		VAlignment = VAlign_Top;
-
-		this->AttachWidget(
-			SNew(SWeakWidget)
-			.PossiblyNullContent(InSlotContent)
-			);
-	}
-
-	/** Get the vertical position of this slot inside its parent */
-	float GetVerticalOffset() const
-	{
-		auto PinnedTrackLane = TrackLane.Pin();
-		return PinnedTrackLane.IsValid() ? PinnedTrackLane->GetPhysicalPosition() : 0.f;
-	}
-
-	/** Horizontal/Vertical alignment for the slot */
-	EHorizontalAlignment HAlignment;
-	EVerticalAlignment VAlignment;
+	TrackLane = InSlotContent;
+	
+	HAlignment = HAlign_Fill;
+	VAlignment = VAlign_Top;
+
+	this->AttachWidget(
+		SNew(SWeakWidget)
+		.PossiblyNullContent(InSlotContent)
+		);
+}
 
-	/** The track lane that we represent */
-	TWeakPtr<SSequencerTrackLane> TrackLane;
-};
+float FTrackAreaSlot::GetVerticalOffset() const
+{
+	auto PinnedTrackLane = TrackLane.Pin();
+	return PinnedTrackLane.IsValid() ? PinnedTrackLane->GetPhysicalPosition() : 0.f;
+}
 
 void SSequencerTrackArea::Construct( const FArguments& InArgs, TSharedRef<FSequencerTimeSliderController> InTimeSliderController, TSharedRef<SSequencer> InSequencerWidget )
 {
diff --git a/Engine/Source/Editor/Sequencer/Private/SSequencerTrackArea.h b/Engine/Source/Editor/Sequencer/Private/SSequencerTrackArea.h
--- a/Engine/Source/Editor/Sequencer/Private/SSequencerTrackArea.h
+++ b/Engine/Source/Editor/Sequencer/Private/SSequencerTrackArea.h
@@ -8,7 +8,24 @@ class SSequencerTrackLane;
 
 class FSequencerTimeSliderController;
 class FSequencer;
-class FTrackAreaSlot;
+
+/** Structure representing a slot in the track area */
+class FTrackAreaSlot : public TSlotBase<FTrackAreaSlot>
+{
+public:
+	/** Construction from a track lane */
+	FTrackAreaSlot(const TSharedPtr<SSequencerTrackLane>& InSlotContent);
+
+	/** Get the vertical position of this slot inside its parent */
+	float GetVerticalOffset() const;
+
+	/** Horizontal/Vertical alignment for the slot */
+	EHorizontalAlignment HAlignment;
+	EVerticalAlignment VAlignment;
+
+	/** The track lane that we represent */
+	TWeakPtr<SSequencerTrackLane> TrackLane;
+};
 
 /**
  * The area where tracks( rows of sections ) are displayed
