Fixed sequencer sections overriding default mode cursors

diff --git a/Engine/Source/Editor/Sequencer/Private/SSection.cpp b/Engine/Source/Editor/Sequencer/Private/SSection.cpp
--- a/Engine/Source/Editor/Sequencer/Private/SSection.cpp
+++ b/Engine/Source/Editor/Sequencer/Private/SSection.cpp
@@ -918,7 +918,7 @@ FCursorReply SSection::OnCursorQuery( const FGeometry& MyGeometry, const FPointe
 		return FCursorReply::Cursor( EMouseCursor::ResizeLeftRight );
 	}
 
-	return FCursorReply::Cursor( EMouseCursor::Default );
+	return FCursorReply::Unhandled();
 }
 
 
