Sequencer: Add delete node to right click menu of nodes

diff --git a/Engine/Source/Editor/Sequencer/Private/SAnimationOutlinerTreeNode.cpp b/Engine/Source/Editor/Sequencer/Private/SAnimationOutlinerTreeNode.cpp
--- a/Engine/Source/Editor/Sequencer/Private/SAnimationOutlinerTreeNode.cpp
+++ b/Engine/Source/Editor/Sequencer/Private/SAnimationOutlinerTreeNode.cpp
@@ -377,6 +377,10 @@ FReply SAnimationOutlinerTreeNode::OnMouseButtonDown( const FGeometry& MyGeometr
 
 		return FReply::Handled();
 	}
+	else if (MouseEvent.GetEffectingButton() == EKeys::RightMouseButton)
+	{
+		return FReply::Handled();
+	}
 
 	return FReply::Unhandled();
 }
diff --git a/Engine/Source/Editor/Sequencer/Private/SSequencer.cpp b/Engine/Source/Editor/Sequencer/Private/SSequencer.cpp
--- a/Engine/Source/Editor/Sequencer/Private/SSequencer.cpp
+++ b/Engine/Source/Editor/Sequencer/Private/SSequencer.cpp
@@ -1331,28 +1331,6 @@ TSharedPtr<SSequencerTreeView> SSequencer::GetTreeView() const
 	return TreeView;
 }
 
-void SSequencer::DeleteSelectedNodes()
-{
-	TSet< TSharedRef<FSequencerDisplayNode> > SelectedNodesCopy = Sequencer.Pin()->GetSelection().GetSelectedOutlinerNodes();
-
-	if( SelectedNodesCopy.Num() > 0 )
-	{
-		const FScopedTransaction Transaction( LOCTEXT("UndoDeletingObject", "Delete Object from MovieScene") );
-
-		FSequencer& SequencerRef = *Sequencer.Pin();
-
-		for( const TSharedRef<FSequencerDisplayNode>& SelectedNode : SelectedNodesCopy )
-		{
-			if( !SelectedNode->IsHidden() )
-			{
-				// Delete everything in the entire node
-				TSharedRef<const FSequencerDisplayNode> NodeToBeDeleted = StaticCastSharedRef<const FSequencerDisplayNode>(SelectedNode);
-				SequencerRef.OnRequestNodeDeleted( NodeToBeDeleted );
-			}
-		}
-	}
-}
-
 TArray<FSectionHandle> SSequencer::GetSectionHandles(const TSet<TWeakObjectPtr<UMovieSceneSection>>& DesiredSections) const
 {
 	TArray<FSectionHandle> SectionHandles;
diff --git a/Engine/Source/Editor/Sequencer/Private/SSequencer.h b/Engine/Source/Editor/Sequencer/Private/SSequencer.h
--- a/Engine/Source/Editor/Sequencer/Private/SSequencer.h
+++ b/Engine/Source/Editor/Sequencer/Private/SSequencer.h
@@ -113,9 +113,6 @@ public:
 	void UpdateBreadcrumbs(const TArray< TWeakObjectPtr<class UMovieSceneSection> >& FilteringShots);
 	void ResetBreadcrumbs();
 
-	/** Deletes selected nodes out of the sequencer node tree */
-	void DeleteSelectedNodes();
-
 	/** Step to next and previous keyframes */
 	void StepToNextKey();
 	void StepToPreviousKey();
diff --git a/Engine/Source/Editor/Sequencer/Private/Sequencer.cpp b/Engine/Source/Editor/Sequencer/Private/Sequencer.cpp
--- a/Engine/Source/Editor/Sequencer/Private/Sequencer.cpp
+++ b/Engine/Source/Editor/Sequencer/Private/Sequencer.cpp
@@ -1445,7 +1445,7 @@ void GetParentTrackNodeAndNamePath(TSharedRef<const FSequencerDisplayNode> Displ
 	}
 }
 
-void FSequencer::OnRequestNodeDeleted( TSharedRef<const FSequencerDisplayNode>& NodeToBeDeleted )
+void FSequencer::OnRequestNodeDeleted( TSharedRef<const FSequencerDisplayNode> NodeToBeDeleted )
 {
 	bool bAnySpawnablesRemoved = false;
 	bool bAnythingRemoved = false;
@@ -1874,7 +1874,7 @@ void FSequencer::DeleteSelectedItems()
 	}
 	else if (Selection.GetActiveSelection() == FSequencerSelection::EActiveSelection::OutlinerNode)
 	{
-		SequencerWidget->DeleteSelectedNodes();
+		DeleteSelectedNodes();
 	}
 }
 
@@ -1961,6 +1961,41 @@ bool FSequencer::CanAssignActor(FGuid ObjectBinding) const
 	return GEditor->GetSelectedActors()->Num() > 0;
 }
 
+void FSequencer::DeleteNode(TSharedRef<FSequencerDisplayNode> NodeToBeDeleted)
+{
+	// If this node is selected, delete all selected nodes
+	if (GetSelection().IsSelected(NodeToBeDeleted))
+	{
+		DeleteSelectedNodes();
+	}
+	else
+	{
+		const FScopedTransaction Transaction( NSLOCTEXT("Sequencer", "UndoDeletingObject", "Delete Node") );
+
+		OnRequestNodeDeleted(NodeToBeDeleted);
+	}
+}
+
+void FSequencer::DeleteSelectedNodes()
+{
+	TSet< TSharedRef<FSequencerDisplayNode> > SelectedNodesCopy = GetSelection().GetSelectedOutlinerNodes();
+
+	if( SelectedNodesCopy.Num() > 0 )
+	{
+		const FScopedTransaction Transaction( NSLOCTEXT("Sequencer", "UndoDeletingObject", "Delete Node") );
+
+		for( const TSharedRef<FSequencerDisplayNode>& SelectedNode : SelectedNodesCopy )
+		{
+			if( !SelectedNode->IsHidden() )
+			{
+				// Delete everything in the entire node
+				TSharedRef<const FSequencerDisplayNode> NodeToBeDeleted = StaticCastSharedRef<const FSequencerDisplayNode>(SelectedNode);
+				OnRequestNodeDeleted( NodeToBeDeleted );
+			}
+		}
+	}
+}
+
 void FSequencer::TogglePlay()
 {
 	OnPlay();
@@ -2311,28 +2346,6 @@ void FSequencer::BuildAddTrackMenu(class FMenuBuilder& MenuBuilder)
 	}
 }
 
-void FSequencer::BuildObjectBindingContextMenu(FMenuBuilder& MenuBuilder, const FGuid& ObjectBinding, const UClass* ObjectClass, FObjectBindingNode* ObjectBindingNode)
-{
-	if (IsLevelEditorSequencer())
-	{
-		if (ObjectClass->IsChildOf(AActor::StaticClass()))
-		{
-			FFormatNamedArguments Args;
-			MenuBuilder.AddMenuEntry(
-				FText::Format( NSLOCTEXT("Sequencer", "Assign Actor ", "Assign Actor"), Args),
-				FText::Format( NSLOCTEXT("Sequencer", "AssignActorTooltip", "Assign the selected actor to this track"), Args ),
-				FSlateIcon(),
-				FUIAction(FExecuteAction::CreateSP(this, &FSequencer::AssignActor, ObjectBinding, ObjectBindingNode),
-							FCanExecuteAction::CreateSP(this, &FSequencer::CanAssignActor, ObjectBinding)) );
-		}
-	}
-
-	for (int32 i = 0; i < TrackEditors.Num(); ++i)
-	{
-		TrackEditors[i]->BuildObjectBindingContextMenu(MenuBuilder, ObjectBinding, ObjectClass);
-	}
-}
-
 void FSequencer::BuildObjectBindingTrackMenu(FMenuBuilder& MenuBuilder, const FGuid& ObjectBinding, const UClass* ObjectClass)
 {
 	for (int32 i = 0; i < TrackEditors.Num(); ++i)
diff --git a/Engine/Source/Editor/Sequencer/Private/Sequencer.h b/Engine/Source/Editor/Sequencer/Private/Sequencer.h
--- a/Engine/Source/Editor/Sequencer/Private/Sequencer.h
+++ b/Engine/Source/Editor/Sequencer/Private/Sequencer.h
@@ -177,7 +177,7 @@ public:
 	 *
 	 * @param NodeToBeDeleted	Node with data that should be deleted
 	 */
-	virtual void OnRequestNodeDeleted( TSharedRef<const FSequencerDisplayNode>& NodeToBeDeleted );
+	virtual void OnRequestNodeDeleted( TSharedRef<const FSequencerDisplayNode> NodeToBeDeleted );
 
 	/**
 	 * Copies properties from one actor to another.  Only properties that are different are copied.  This is used
@@ -208,15 +208,6 @@ public:
 	 */
 	void BuildAddTrackMenu(class FMenuBuilder& MenuBuilder);
 
-	/**
-	 * Builds up the context menu for object binding nodes in the outliner
-	 * 
-	 * @param MenuBuilder	The context menu builder to add things to
-	 * @param ObjectBinding	The object binding of the selected node
-	 * @param ObjectClass	The class of the selected object
-	 */
-	void BuildObjectBindingContextMenu(class FMenuBuilder& MenuBuilder, const FGuid& ObjectBinding, const class UClass* ObjectClass, FObjectBindingNode* ObjectBindingNode);
-
 	/**
 	 * Builds up the track menu for object binding nodes in the outliner
 	 * 
@@ -274,6 +265,14 @@ public:
 	/** Called to add selected objects to the movie scene */
 	void AddSelectedObjects();
 
+	/** Called when a user executes the assign actor to track menu item */
+	void AssignActor(FGuid ObjectBinding, FObjectBindingNode* ObjectBindingNode);
+	bool CanAssignActor(FGuid ObjectBinding) const;
+
+	/** Called when a user executes the delete node menu item */
+	void DeleteNode(TSharedRef<FSequencerDisplayNode> NodeToBeDeleted);
+	void DeleteSelectedNodes();
+
 public:
 
 	/** Access the currently enabled edit tool */
@@ -393,10 +392,6 @@ protected:
 	void DeleteSelectedItems();
 	bool CanDeleteSelectedItems() const;
 	
-	/** Called when a user executes the assign actor to track menu item */
-	void AssignActor(FGuid ObjectBinding, FObjectBindingNode* ObjectBindingNode);
-	bool CanAssignActor(FGuid ObjectBinding) const;
-
 	/** Transport controls */
 	void TogglePlay();
 	void PlayForward();
diff --git a/Engine/Source/Editor/Sequencer/Private/SequencerDisplayNode.cpp b/Engine/Source/Editor/Sequencer/Private/SequencerDisplayNode.cpp
--- a/Engine/Source/Editor/Sequencer/Private/SequencerDisplayNode.cpp
+++ b/Engine/Source/Editor/Sequencer/Private/SequencerDisplayNode.cpp
@@ -361,6 +361,28 @@ FString FSequencerDisplayNode::GetPathName() const
 	return PathName;
 }
 
+TSharedPtr<SWidget> FSequencerDisplayNode::OnSummonContextMenu(const FGeometry& MyGeometry, const FPointerEvent& MouseEvent)
+{
+	// @todo sequencer replace with UI Commands instead of faking it
+	const bool bShouldCloseWindowAfterMenuSelection = true;
+	FMenuBuilder MenuBuilder(bShouldCloseWindowAfterMenuSelection, NULL);
+
+	BuildContextMenu(MenuBuilder);
+
+	return MenuBuilder.MakeWidget();
+}
+
+void FSequencerDisplayNode::BuildContextMenu(FMenuBuilder& MenuBuilder)
+{
+	TSharedRef<FSequencerDisplayNode> NodeToBeDeleted = SharedThis(this);
+
+	MenuBuilder.AddMenuEntry(
+		LOCTEXT("DeleteNode", "Delete"),
+		LOCTEXT("DeleteNodeTooltip", "Delete this or selected nodes"),
+		FSlateIcon(),
+		FUIAction(FExecuteAction::CreateSP(&GetSequencer(), &FSequencer::DeleteNode, NodeToBeDeleted)));
+}
+
 void FSequencerDisplayNode::GetChildKeyAreaNodesRecursively(TArray< TSharedRef<FSectionKeyAreaNode> >& OutNodes) const
 {
 	for (int32 i = 0; i < ChildNodes.Num(); ++i)
@@ -656,17 +678,25 @@ TSharedRef<SWidget> FObjectBindingNode::GenerateEditWidgetForOutliner()
 	return EditBox.ToSharedRef();
 }
 
-TSharedPtr<SWidget> FObjectBindingNode::OnSummonContextMenu(const FGeometry& MyGeometry, const FPointerEvent& MouseEvent)
+void FObjectBindingNode::BuildContextMenu(FMenuBuilder& MenuBuilder)
 {
-	const UClass* ObjectClass = GetClassForObjectBinding();
-
-	// @todo sequencer replace with UI Commands instead of faking it
-	const bool bShouldCloseWindowAfterMenuSelection = true;
-	FMenuBuilder MenuBuilder(bShouldCloseWindowAfterMenuSelection, NULL);
-
-	GetSequencer().BuildObjectBindingContextMenu(MenuBuilder, ObjectBinding, ObjectClass, this);
+	if (GetSequencer().IsLevelEditorSequencer())
+	{
+		const UClass* ObjectClass = GetClassForObjectBinding();
+		
+		if (ObjectClass->IsChildOf(AActor::StaticClass()))
+		{
+			FFormatNamedArguments Args;
+			MenuBuilder.AddMenuEntry(
+				FText::Format( LOCTEXT("Assign Actor ", "Assign Actor"), Args),
+				FText::Format( LOCTEXT("AssignActorTooltip", "Assign the selected actor to this track"), Args ),
+				FSlateIcon(),
+				FUIAction(FExecuteAction::CreateSP(&GetSequencer(), &FSequencer::AssignActor, ObjectBinding, this),
+							FCanExecuteAction::CreateSP(&GetSequencer(), &FSequencer::CanAssignActor, ObjectBinding)) );
+		}
+	}
 
-	return MenuBuilder.MakeWidget();
+	FSequencerDisplayNode::BuildContextMenu(MenuBuilder);
 }
 
 void GetKeyablePropertyPaths(UClass* Class, UStruct* PropertySource, TArray<UProperty*>& PropertyPath, FSequencer& Sequencer, TArray<TArray<UProperty*>>& KeyablePropertyPaths)
diff --git a/Engine/Source/Editor/Sequencer/Private/SequencerDisplayNode.h b/Engine/Source/Editor/Sequencer/Private/SequencerDisplayNode.h
--- a/Engine/Source/Editor/Sequencer/Private/SequencerDisplayNode.h
+++ b/Engine/Source/Editor/Sequencer/Private/SequencerDisplayNode.h
@@ -142,8 +142,11 @@ public:
 	 */
 	FString GetPathName() const;
 
+	/** Summon context menu */
+	TSharedPtr<SWidget> OnSummonContextMenu(const FGeometry& MyGeometry, const FPointerEvent& MouseEvent);
+
 	/** What sort of context menu this node summons */
-	virtual TSharedPtr<SWidget> OnSummonContextMenu(const FGeometry& MyGeometry, const FPointerEvent& MouseEvent) {return TSharedPtr<SWidget>();}
+	virtual void BuildContextMenu(FMenuBuilder& MenuBuilder);
 
 	/**
 	 * @return The name of the node (for identification purposes)
@@ -436,7 +439,7 @@ public:
 	const FGuid& GetObjectBinding() const { return ObjectBinding; }
 	
 	/** What sort of context menu this node summons */
-	virtual TSharedPtr<SWidget> OnSummonContextMenu(const FGeometry& MyGeometry, const FPointerEvent& MouseEvent) override;
+	virtual void BuildContextMenu(FMenuBuilder& MenuBuilder) override;
 
 private:
 	TSharedRef<SWidget> OnGetAddTrackMenuContent();
diff --git a/Engine/Source/Editor/Sequencer/Public/ISequencerTrackEditor.h b/Engine/Source/Editor/Sequencer/Public/ISequencerTrackEditor.h
--- a/Engine/Source/Editor/Sequencer/Public/ISequencerTrackEditor.h
+++ b/Engine/Source/Editor/Sequencer/Public/ISequencerTrackEditor.h
@@ -42,15 +42,6 @@ public:
 	 */
 	virtual void BuildAddTrackMenu(FMenuBuilder& MenuBuilder) = 0;
 
-	/**
-	 * Builds up the object binding context menu for the outliner.
-	 *
-	 * @param MenuBuilder The menu builder to change.
-	 * @param ObjectBinding The object binding this is for.
-	 * @param ObjectClass The class of the object this is for.
-	 */
-	virtual void BuildObjectBindingContextMenu(FMenuBuilder& MenuBuilder, const FGuid& ObjectBinding, const UClass* ObjectClass) = 0;
-
 	/**
 	 * Builds up the object binding edit buttons for the outliner.
 	 *
diff --git a/Engine/Source/Editor/Sequencer/Public/MovieSceneTrackEditor.h b/Engine/Source/Editor/Sequencer/Public/MovieSceneTrackEditor.h
--- a/Engine/Source/Editor/Sequencer/Public/MovieSceneTrackEditor.h
+++ b/Engine/Source/Editor/Sequencer/Public/MovieSceneTrackEditor.h
@@ -163,7 +163,6 @@ public:
 
 	virtual void BindCommands(TSharedRef<FUICommandList> SequencerCommandBindings) override { }
 	virtual void BuildAddTrackMenu(FMenuBuilder& MenuBuilder) override { }
-	virtual void BuildObjectBindingContextMenu(FMenuBuilder& MenuBuilder, const FGuid& ObjectBinding, const UClass* ObjectClass) override { }
 	virtual void BuildObjectBindingEditButtons(TSharedPtr<SHorizontalBox> EditBox, const FGuid& ObjectBinding, const UClass* ObjectClass) override { }
 	virtual void BuildObjectBindingTrackMenu(FMenuBuilder& MenuBuilder, const FGuid& ObjectBinding, const UClass* ObjectClass) override { }
 	virtual TSharedPtr<SWidget> BuildOutlinerEditWidget(const FGuid& ObjectBinding, UMovieSceneTrack* Track) override  { return TSharedPtr<SWidget>(); }
