Moved Render Movie button on the sequencer

diff --git a/Engine/Source/Editor/Sequencer/Private/SSequencer.cpp b/Engine/Source/Editor/Sequencer/Private/SSequencer.cpp
--- a/Engine/Source/Editor/Sequencer/Private/SSequencer.cpp
+++ b/Engine/Source/Editor/Sequencer/Private/SSequencer.cpp
@@ -626,6 +626,8 @@ TSharedRef<SWidget> SSequencer::MakeToolBar()
 				, LOCTEXT("SaveDirtyPackagesTooltip", "Saves the current movie scene")
 				, FSlateIcon(FEditorStyle::GetStyleSetName(), "Sequencer.Save"));
 
+			ToolBarBuilder.AddToolBarButton( FSequencerCommands::Get().RenderMovie );
+
 			ToolBarBuilder.AddSeparator();
 		}
 
@@ -759,16 +761,6 @@ TSharedRef<SWidget> SSequencer::MakeToolBar()
 	}
 	ToolBarBuilder.EndSection();
 
-	if( Sequencer.Pin()->IsLevelEditorSequencer() )
-	{
-		ToolBarBuilder.BeginSection("Render Movie");
-
-		ToolBarBuilder.AddToolBarButton( FSequencerCommands::Get().RenderMovie );
-
-		ToolBarBuilder.EndSection();
-	}
-	
-
 	return ToolBarBuilder.MakeWidget();
 }
 
