Sequencer: Improved details view updates based on key/section selection

diff --git a/Engine/Source/Editor/Sequencer/Private/SSequencer.cpp b/Engine/Source/Editor/Sequencer/Private/SSequencer.cpp
--- a/Engine/Source/Editor/Sequencer/Private/SSequencer.cpp
+++ b/Engine/Source/Editor/Sequencer/Private/SSequencer.cpp
@@ -516,6 +516,40 @@ void SSequencer::NotifyPostChange(const FPropertyChangedEvent& PropertyChangedEv
 }
 
 
+/* SSequencer implementation
+ *****************************************************************************/
+
+void SSequencer::UpdateDetailsView()
+{
+	TArray<TWeakObjectPtr<UObject>> Sections;
+
+	// get selected sections
+	for (auto Section : Sequencer.Pin()->GetSelection().GetSelectedSections())
+	{
+		Sections.Add(Section);
+	}
+
+	// get sections from selected keys
+	const TArray<FSelectedKey> SelectedKeys = Sequencer.Pin()->GetSelection().GetSelectedKeys().Array();
+
+	for (const auto& Key : SelectedKeys)
+	{
+		if (Key.Section != nullptr)
+		{
+			Sections.AddUnique(Key.Section);
+		}
+	}
+
+	// @todo sequencer: highlight selected keys in details view
+
+	// update details view
+	DetailsView->SetObjects(Sections, true);
+}
+
+
+/* SSequencer callbacks
+ *****************************************************************************/
+
 bool SSequencer::HandleDetailsViewEnabled() const
 {
 	return true;
@@ -528,39 +562,20 @@ EVisibility SSequencer::HandleDetailsViewVisibility() const
 	{
 		return EVisibility::Visible;
 	}
-	else
-	{
-		return EVisibility::Collapsed;
-	}
+
+	return EVisibility::Collapsed;
 }
 
 
 void SSequencer::HandleKeySelectionChanged()
 {
-	TSharedPtr<FStructOnScope> Struct;
-	const TArray<FSelectedKey> SelectedKeys = Sequencer.Pin()->GetSelection().GetSelectedKeys().Array();
-
-	if (SelectedKeys.Num() > 0)
-	{
-		// @todo sequencer: highlight selected keys in details view
-	}
-	else
-	{
-		// @todo sequencer: remove key highlights in details view
-	}
+	UpdateDetailsView();
 }
 
 
 void SSequencer::HandleSectionSelectionChanged()
 {
-	TArray<TWeakObjectPtr<UObject>> Sections;
-
-	for (auto Section : Sequencer.Pin()->GetSelection().GetSelectedSections())
-	{
-		Sections.Add(Section);
-	}
-
-	DetailsView->SetObjects(Sections, true);
+	UpdateDetailsView();
 }
 
 
diff --git a/Engine/Source/Editor/Sequencer/Private/SSequencer.h b/Engine/Source/Editor/Sequencer/Private/SSequencer.h
--- a/Engine/Source/Editor/Sequencer/Private/SSequencer.h
+++ b/Engine/Source/Editor/Sequencer/Private/SSequencer.h
@@ -147,6 +147,11 @@ public:
 
 	void NotifyPostChange(const FPropertyChangedEvent& PropertyChangedEvent, class FEditPropertyChain* PropertyThatChanged);
 
+protected:
+
+	/** Update the details view from the currently selected keys and sections. */
+	void UpdateDetailsView();
+
 private:
 	
 	/** Handles checking whether the details view is enabled. */
