Sequencer: Giving the details panel a little more space by default

diff --git a/Engine/Source/Editor/Sequencer/Private/SSequencer.cpp b/Engine/Source/Editor/Sequencer/Private/SSequencer.cpp
--- a/Engine/Source/Editor/Sequencer/Private/SSequencer.cpp
+++ b/Engine/Source/Editor/Sequencer/Private/SSequencer.cpp
@@ -297,7 +297,7 @@ void SSequencer::Construct( const FArguments& InArgs, TSharedRef< class FSequenc
 							.Orientation(Orient_Horizontal)
 						
 						+ SSplitter::Slot()
-							.Value(0.85f)
+							.Value(0.80f)
 							[
 								SNew(SBox)
 									.Padding(FMargin(0.0f, 2.0f, 0.0f, 0.0f))
@@ -439,7 +439,7 @@ void SSequencer::Construct( const FArguments& InArgs, TSharedRef< class FSequenc
 							]
 
 						+ SSplitter::Slot()
-							.Value(0.15f)
+							.Value(0.2f)
 							[
 								DetailsView.ToSharedRef()
 							]
