Sequencer: Add Fade track.

diff --git a/Engine/Source/Editor/EditorStyle/Private/SlateEditorStyle.cpp b/Engine/Source/Editor/EditorStyle/Private/SlateEditorStyle.cpp
--- a/Engine/Source/Editor/EditorStyle/Private/SlateEditorStyle.cpp
+++ b/Engine/Source/Editor/EditorStyle/Private/SlateEditorStyle.cpp
@@ -1630,6 +1630,7 @@ void FSlateEditorStyle::FStyle::SetupGeneralStyles()
 		Set( "Sequencer.Tracks.Event", new IMAGE_BRUSH("Icons/SequencerIcons/icon_Sequencer_Event_16x", Icon16x16));
 		Set( "Sequencer.Tracks.Shot", new IMAGE_BRUSH("Icons/SequencerIcons/icon_Sequencer_Shot_16x", Icon16x16));
 		Set( "Sequencer.Tracks.Slomo", new IMAGE_BRUSH("Icons/SequencerIcons/icon_Sequencer_Slomo_16x", Icon16x16));
+		Set( "Sequencer.Tracks.Fade", new IMAGE_BRUSH("Icons/SequencerIcons/icon_Sequencer_Fade_16x", Icon16x16));
 
 		const FButtonStyle DetailsKeyButton = FButtonStyle(NoBorder)
 			.SetNormal( IMAGE_BRUSH("Sequencer/AddKey_Details", FVector2D(11,11) )  )
diff --git a/Engine/Source/Editor/MovieSceneTools/Private/MovieSceneToolsModule.cpp b/Engine/Source/Editor/MovieSceneTools/Private/MovieSceneToolsModule.cpp
--- a/Engine/Source/Editor/MovieSceneTools/Private/MovieSceneToolsModule.cpp
+++ b/Engine/Source/Editor/MovieSceneTools/Private/MovieSceneToolsModule.cpp
@@ -28,6 +28,7 @@
 #include "EventTrackEditor.h"
 #include "PathTrackEditor.h"
 #include "MaterialTrackEditor.h"
+#include "FadeTrackEditor.h"
 
 
 /**
@@ -64,6 +65,7 @@ public:
 		SubMovieSceneTrackCreateEditorHandle = SequencerModule.RegisterTrackEditor_Handle( FOnCreateTrackEditor::CreateStatic( &FSubMovieSceneTrackEditor::CreateTrackEditor ) );
 		TransformTrackCreateEditorHandle = SequencerModule.RegisterTrackEditor_Handle( FOnCreateTrackEditor::CreateStatic( &F3DTransformTrackEditor::CreateTrackEditor ) );
 		ComponentMaterialTrackCreatEditorHandle = SequencerModule.RegisterTrackEditor_Handle( FOnCreateTrackEditor::CreateStatic( &FComponentMaterialTrackEditor::CreateTrackEditor ) );
+		FadeTrackCreatEditorHandle = SequencerModule.RegisterTrackEditor_Handle( FOnCreateTrackEditor::CreateStatic( &FFadeTrackEditor::CreateTrackEditor ) );
 	}
 
 	virtual void ShutdownModule() override
@@ -95,6 +97,7 @@ public:
 		SequencerModule.UnRegisterTrackEditor_Handle( SubMovieSceneTrackCreateEditorHandle );
 		SequencerModule.UnRegisterTrackEditor_Handle( TransformTrackCreateEditorHandle );
 		SequencerModule.UnRegisterTrackEditor_Handle( ComponentMaterialTrackCreatEditorHandle );
+		SequencerModule.UnRegisterTrackEditor_Handle( FadeTrackCreatEditorHandle );
 	}
 
 private:
@@ -118,6 +121,7 @@ private:
 	FDelegateHandle SubMovieSceneTrackCreateEditorHandle;
 	FDelegateHandle TransformTrackCreateEditorHandle;
 	FDelegateHandle ComponentMaterialTrackCreatEditorHandle;
+	FDelegateHandle FadeTrackCreatEditorHandle;
 };
 
 
diff --git a/Engine/Source/Editor/MovieSceneTools/Private/TrackEditors/FadeTrackEditor.cpp b/Engine/Source/Editor/MovieSceneTools/Private/TrackEditors/FadeTrackEditor.cpp
new file mode 100644
--- /dev/null
+++ b/Engine/Source/Editor/MovieSceneTools/Private/TrackEditors/FadeTrackEditor.cpp
@@ -0,0 +1,81 @@
+// Copyright 1998-2015 Epic Games, Inc. All Rights Reserved.
+
+#include "MovieSceneToolsPrivatePCH.h"
+#include "MovieSceneFadeTrack.h"
+#include "MovieSceneTrack.h"
+#include "FadeTrackEditor.h"
+
+
+#define LOCTEXT_NAMESPACE "FFadeTrackEditor"
+
+
+/* FFadeTrackEditor static functions
+ *****************************************************************************/
+
+TSharedRef<ISequencerTrackEditor> FFadeTrackEditor::CreateTrackEditor(TSharedRef<ISequencer> InSequencer)
+{
+	return MakeShareable(new FFadeTrackEditor(InSequencer));
+}
+
+
+/* FFadeTrackEditor structors
+ *****************************************************************************/
+
+FFadeTrackEditor::FFadeTrackEditor(TSharedRef<ISequencer> InSequencer)
+	: FFloatPropertyTrackEditor(InSequencer)
+{ }
+
+
+/* ISequencerTrackEditor interface
+ *****************************************************************************/
+
+void FFadeTrackEditor::BuildAddTrackMenu(FMenuBuilder& MenuBuilder)
+{
+	UMovieSceneSequence* RootMovieSceneSequence = GetSequencer()->GetRootMovieSceneSequence();
+
+	if ((RootMovieSceneSequence == nullptr) || (RootMovieSceneSequence->GetClass()->GetName() != TEXT("ActorAnimationInstance")))
+	{
+		return;
+	}
+
+	MenuBuilder.AddMenuEntry(
+		LOCTEXT("AddFadeTrack", "Fade Track"),
+		LOCTEXT("AddFadeTrackTooltip", "Adds a new track that controls the fade of the sequence."),
+		FSlateIcon(FEditorStyle::GetStyleSetName(), "Sequencer.Tracks.Fade"),
+		FUIAction(
+			FExecuteAction::CreateRaw(this, &FFadeTrackEditor::HandleAddFadeTrackMenuEntryExecute)
+		)
+	);
+}
+
+
+bool FFadeTrackEditor::SupportsType(TSubclassOf<UMovieSceneTrack> Type) const
+{
+	return (Type == UMovieSceneFadeTrack::StaticClass());
+}
+
+
+/* FFadeTrackEditor callbacks
+ *****************************************************************************/
+
+void FFadeTrackEditor::HandleAddFadeTrackMenuEntryExecute()
+{
+	if (FadeTrack.IsValid())
+	{
+		return;
+	}
+
+	UMovieSceneSequence* FocusedSequence = GetSequencer()->GetFocusedMovieSceneSequence();
+	UMovieScene* MovieScene = FocusedSequence->GetMovieScene();
+
+	if (MovieScene != nullptr)
+	{
+		FadeTrack = MovieScene->AddMasterTrack(UMovieSceneFadeTrack::StaticClass());
+		FadeTrack->AddSection(FadeTrack->CreateNewSection());
+
+		GetSequencer()->NotifyMovieSceneDataChanged();
+	}
+}
+
+
+#undef LOCTEXT_NAMESPACE
diff --git a/Engine/Source/Editor/MovieSceneTools/Private/TrackEditors/FadeTrackEditor.h b/Engine/Source/Editor/MovieSceneTools/Private/TrackEditors/FadeTrackEditor.h
new file mode 100644
--- /dev/null
+++ b/Engine/Source/Editor/MovieSceneTools/Private/TrackEditors/FadeTrackEditor.h
@@ -0,0 +1,53 @@
+// Copyright 1998-2015 Epic Games, Inc. All Rights Reserved.
+
+#pragma once
+
+#include "FloatPropertyTrackEditor.h"
+
+
+class ISequencer;
+class UMovieSceneTrack;
+
+
+/**
+* A property track editor for fade control.
+*/
+class FFadeTrackEditor
+	: public FFloatPropertyTrackEditor
+{
+public:
+
+	/**
+	 * Factory function to create an instance of this class (called by a sequencer).
+	 *
+	 * @param InSequencer The sequencer instance to be used by this tool.
+	 * @return The new instance of this class.
+	 */
+	static TSharedRef<ISequencerTrackEditor> CreateTrackEditor(TSharedRef<ISequencer> InSequencer);
+
+public:
+
+	/**
+	 * Creates and initializes a new instance.
+	 *
+	 * @param InSequencer The sequencer instance to be used by this tool.
+	 */
+	FFadeTrackEditor(TSharedRef<ISequencer> InSequencer);	
+
+public:
+
+	// ISequencerTrackEditor interface
+
+	virtual void BuildAddTrackMenu(FMenuBuilder& MenuBuilder) override;
+	virtual bool SupportsType(TSubclassOf<UMovieSceneTrack> Type) const override;
+
+private:
+
+	/** Callback for executing the "Add Fade Track" menu entry. */
+	void HandleAddFadeTrackMenuEntryExecute();
+
+private:
+
+	/** Pointer to the sequence's fade track, if it exists. */
+	TWeakObjectPtr<UMovieSceneTrack> FadeTrack;
+};
diff --git a/Engine/Source/Editor/Sequencer/Private/Sequencer.cpp b/Engine/Source/Editor/Sequencer/Private/Sequencer.cpp
--- a/Engine/Source/Editor/Sequencer/Private/Sequencer.cpp
+++ b/Engine/Source/Editor/Sequencer/Private/Sequencer.cpp
@@ -866,6 +866,56 @@ void FSequencer::UpdateCameraCut(UObject* ObjectToViewThrough, bool bNewCameraCu
 	}
 }
 
+void FSequencer::SetViewportSettings(const TMap<FViewportClient*, EMovieSceneViewportParams>& ViewportParamsMap)
+{
+	if(!IsPerspectiveViewportPosessionEnabled())
+	{
+		return;
+	}
+
+	for(FLevelEditorViewportClient* LevelVC : GEditor->LevelViewportClients)
+	{
+		if(LevelVC && LevelVC->IsPerspective() && LevelVC->AllowsCinematicPreview())
+		{
+			if (ViewportParamsMap.Contains(LevelVC))
+			{
+				const EMovieSceneViewportParams* ViewportParams = ViewportParamsMap.Find(LevelVC);
+				if (ViewportParams->SetWhichViewportParam & EMovieSceneViewportParams::SVP_FadeAmount)
+				{
+					LevelVC->FadeAmount = ViewportParams->FadeAmount;
+					LevelVC->bEnableFading = true;
+				}
+				if (ViewportParams->SetWhichViewportParam & EMovieSceneViewportParams::SVP_FadeColor)
+				{
+					LevelVC->FadeColor = ViewportParams->FadeColor.ToRGBE();
+					LevelVC->bEnableFading = true;
+				}
+				if (ViewportParams->SetWhichViewportParam & EMovieSceneViewportParams::SVP_ColorScaling)
+				{
+					LevelVC->bEnableColorScaling = ViewportParams->bEnableColorScaling;
+					LevelVC->ColorScale = ViewportParams->ColorScale;
+				}
+			}
+		}
+	}
+}
+
+void FSequencer::GetViewportSettings(TMap<FViewportClient*, EMovieSceneViewportParams>& ViewportParamsMap) const
+{
+	for(FLevelEditorViewportClient* LevelVC : GEditor->LevelViewportClients)
+	{
+		if(LevelVC && LevelVC->IsPerspective() && LevelVC->AllowsCinematicPreview())
+		{
+			EMovieSceneViewportParams ViewportParams;
+			ViewportParams.FadeAmount = LevelVC->FadeAmount;
+			ViewportParams.FadeColor = FLinearColor(LevelVC->FadeColor);
+			ViewportParams.ColorScale = LevelVC->ColorScale;
+
+			ViewportParamsMap.Add(LevelVC, ViewportParams);
+		}
+	}
+}
+
 EMovieScenePlayerStatus::Type FSequencer::GetPlaybackStatus() const
 {
 	return PlaybackState;
@@ -1672,23 +1722,12 @@ void FSequencer::UpdatePreviewLevelViewportClientFromCameraCut( FLevelEditorView
 		InViewportClient.SetViewLocation(Cam->GetActorLocation());
 		InViewportClient.SetViewRotation(Cam->GetActorRotation());
 
-		// @todo Sequencer previewing
-		/*InViewportClient.FadeAmount = InFadeAmount;
-		InViewportClient.bEnableFading = true;*/
-
-		/*InViewportClient.bEnableColorScaling = bInEnableColorScaling;
-		InViewportClient.ColorScale = InColorScale;*/
-
 		InViewportClient.bEditorCameraCut = bNewCameraCut;
 	}
 	else
 	{
 		InViewportClient.ViewFOV = InViewportClient.FOVAngle;
 
-		// @todo Sequencer previewing
-		/*InViewportClient.>FadeAmount = InFadeAmount;
-		InViewportClient.bEnableFading = true;*/
-
 		InViewportClient.bEditorCameraCut = false;
 	}
 
diff --git a/Engine/Source/Editor/Sequencer/Private/Sequencer.h b/Engine/Source/Editor/Sequencer/Private/Sequencer.h
--- a/Engine/Source/Editor/Sequencer/Private/Sequencer.h
+++ b/Engine/Source/Editor/Sequencer/Private/Sequencer.h
@@ -229,6 +229,8 @@ public:
 	/** IMovieScenePlayer interface */
 	virtual void GetRuntimeObjects( TSharedRef<FMovieSceneSequenceInstance> MovieSceneInstance, const FGuid& ObjectHandle, TArray< UObject* >& OutObjects ) const override;
 	virtual void UpdateCameraCut(UObject* ObjectToViewThrough, bool bNewCameraCut) const override;
+	virtual void SetViewportSettings(const TMap<FViewportClient*, EMovieSceneViewportParams>& ViewportParamsMap) override;
+	virtual void GetViewportSettings(TMap<FViewportClient*, EMovieSceneViewportParams>& ViewportParamsMap) const override;
 	virtual EMovieScenePlayerStatus::Type GetPlaybackStatus() const override;
 	virtual void AddOrUpdateMovieSceneInstance( class UMovieSceneSection& MovieSceneSection, TSharedRef<FMovieSceneSequenceInstance> InstanceToAdd ) override;
 	virtual void RemoveMovieSceneInstance( class UMovieSceneSection& MovieSceneSection, TSharedRef<FMovieSceneSequenceInstance> InstanceToRemove ) override;
diff --git a/Engine/Source/Runtime/ActorAnimation/Public/ActorAnimationPlayer.h b/Engine/Source/Runtime/ActorAnimation/Public/ActorAnimationPlayer.h
--- a/Engine/Source/Runtime/ActorAnimation/Public/ActorAnimationPlayer.h
+++ b/Engine/Source/Runtime/ActorAnimation/Public/ActorAnimationPlayer.h
@@ -126,6 +126,8 @@ protected:
 	virtual void DestroyActorsForMovie(TSharedRef<FMovieSceneSequenceInstance> MovieSceneInstance);
 	virtual void GetRuntimeObjects(TSharedRef<FMovieSceneSequenceInstance> MovieSceneInstance, const FGuid& ObjectHandle, TArray<UObject*>& OutObjects) const override;
 	virtual void UpdateCameraCut(UObject* ObjectToViewThrough, bool bNewCameraCut) const override;
+	virtual void SetViewportSettings(const TMap<FViewportClient*, EMovieSceneViewportParams>& ViewportParamsMap) override;
+	virtual void GetViewportSettings(TMap<FViewportClient*, EMovieSceneViewportParams>& ViewportParamsMap) const override;
 	virtual EMovieScenePlayerStatus::Type GetPlaybackStatus() const override;
 	virtual void AddOrUpdateMovieSceneInstance(UMovieSceneSection& MovieSceneSection, TSharedRef<FMovieSceneSequenceInstance> InstanceToAdd) override;
 	virtual void RemoveMovieSceneInstance(UMovieSceneSection& MovieSceneSection, TSharedRef<FMovieSceneSequenceInstance> InstanceToRemove) override;
diff --git a/Engine/Source/Runtime/MovieScene/Public/IMovieScenePlayer.h b/Engine/Source/Runtime/MovieScene/Public/IMovieScenePlayer.h
--- a/Engine/Source/Runtime/MovieScene/Public/IMovieScenePlayer.h
+++ b/Engine/Source/Runtime/MovieScene/Public/IMovieScenePlayer.h
@@ -14,6 +14,31 @@ namespace EMovieScenePlayerStatus
 	};
 }
 
+struct EMovieSceneViewportParams
+{
+	EMovieSceneViewportParams()
+	{
+		FadeAmount = 0.f;
+		FadeColor = FLinearColor::Black;
+		bEnableColorScaling = false;
+	}
+
+	enum SetViewportParam
+	{
+		SVP_FadeAmount   = 0x00000001,
+		SVP_FadeColor    = 0x00000002,
+		SVP_ColorScaling = 0x00000004,
+		SVP_All          = SVP_FadeAmount | SVP_FadeColor | SVP_ColorScaling
+	};
+
+	SetViewportParam SetWhichViewportParam;
+
+	float FadeAmount;
+	FLinearColor FadeColor;
+	FVector ColorScale; 
+	bool bEnableColorScaling;
+};
+
 class FMovieSceneSequenceInstance;
 
 /**
@@ -41,6 +66,20 @@ public:
 	 */
 	virtual void UpdateCameraCut(UObject* ObjectToViewThrough, bool bNewCameraCut) const = 0;
 
+	/*
+	 * Set the perspective viewport settings
+	 *
+	 * @param ViewportParamMap A map from the viewport client to its settings
+	 */
+	virtual void SetViewportSettings(const TMap<FViewportClient*, EMovieSceneViewportParams>& ViewportParamsMap) = 0;
+
+	/*
+	 * Get the current perspective viewport settings
+	 *
+	 * @param ViewportParamMap A map from the viewport client to its settings
+	 */
+	virtual void GetViewportSettings(TMap<FViewportClient*, EMovieSceneViewportParams>& ViewportParamsMap) const = 0;
+
 	/**
 	 * Adds a MovieScene instance to the player.  MovieScene players need to know about each instance for actor spawning
 	 *
diff --git a/Engine/Source/Runtime/MovieSceneTracks/Private/MovieSceneTracksPrivatePCH.h b/Engine/Source/Runtime/MovieSceneTracks/Private/MovieSceneTracksPrivatePCH.h
--- a/Engine/Source/Runtime/MovieSceneTracks/Private/MovieSceneTracksPrivatePCH.h
+++ b/Engine/Source/Runtime/MovieSceneTracks/Private/MovieSceneTracksPrivatePCH.h
@@ -8,6 +8,7 @@
 
 #include "Core.h"
 #include "Engine.h"
+#include "IMovieScenePlayer.h"
 #include "KeyParams.h"
 #include "MovieScene.h"
 #include "MovieSceneCommonHelpers.h"
diff --git a/Engine/Source/Runtime/MovieSceneTracks/Private/Sections/MovieSceneFadeSection.cpp b/Engine/Source/Runtime/MovieSceneTracks/Private/Sections/MovieSceneFadeSection.cpp
new file mode 100644
--- /dev/null
+++ b/Engine/Source/Runtime/MovieSceneTracks/Private/Sections/MovieSceneFadeSection.cpp
@@ -0,0 +1,16 @@
+// Copyright 1998-2015 Epic Games, Inc. All Rights Reserved.
+
+#include "MovieSceneTracksPrivatePCH.h"
+#include "MovieSceneFadeSection.h"
+
+
+/* UMovieSceneFadeSection structors
+ *****************************************************************************/
+
+UMovieSceneFadeSection::UMovieSceneFadeSection()
+	: UMovieSceneFloatSection()
+	, FadeColor(FLinearColor::Black)
+	, bFadeAudio(false)
+{
+	SetIsInfinite(true);
+}
diff --git a/Engine/Source/Runtime/MovieSceneTracks/Private/TrackInstances/MovieSceneFadeTrackInstance.cpp b/Engine/Source/Runtime/MovieSceneTracks/Private/TrackInstances/MovieSceneFadeTrackInstance.cpp
new file mode 100644
--- /dev/null
+++ b/Engine/Source/Runtime/MovieSceneTracks/Private/TrackInstances/MovieSceneFadeTrackInstance.cpp
@@ -0,0 +1,68 @@
+// Copyright 1998-2015 Epic Games, Inc. All Rights Reserved.
+
+#include "MovieSceneTracksPrivatePCH.h"
+#include "MovieSceneFadeSection.h"
+#include "MovieSceneFadeTrack.h"
+#include "MovieSceneFadeTrackInstance.h"
+
+
+/* FMovieSceneFadeTrackInstance structors
+ *****************************************************************************/
+
+FMovieSceneFadeTrackInstance::FMovieSceneFadeTrackInstance(UMovieSceneFadeTrack& InFadeTrack)
+	: FadeTrack(&InFadeTrack)
+{ }
+
+
+/* IMovieSceneTrackInstance interface
+ *****************************************************************************/
+
+void FMovieSceneFadeTrackInstance::Update(float Position, float LastPosition, const TArray<UObject*>& RuntimeObjects, IMovieScenePlayer& Player)
+{
+	FLinearColor FadeColor = FLinearColor::Black;
+	bool bFadeAudio = false;
+
+	float FloatValue = 0.0f;
+
+	if (FadeTrack->Eval(Position, LastPosition, FloatValue))
+	{
+		const UMovieSceneSection* NearestSection = MovieSceneHelpers::FindNearestSectionAtTime(FadeTrack->GetAllSections(), Position);
+		const UMovieSceneFadeSection* FadeSection = CastChecked<const UMovieSceneFadeSection>(NearestSection);
+		if (FadeSection)
+		{
+			FadeColor = FadeSection->FadeColor;
+			bFadeAudio = FadeSection->bFadeAudio;
+		}
+
+		// Set editor preview/fade
+		EMovieSceneViewportParams ViewportParams;
+		ViewportParams.SetWhichViewportParam = (EMovieSceneViewportParams::SetViewportParam)(EMovieSceneViewportParams::SVP_FadeAmount | EMovieSceneViewportParams::SVP_FadeColor);
+		ViewportParams.FadeAmount = FloatValue;
+		ViewportParams.FadeColor = FadeColor;
+
+		TMap<FViewportClient*, EMovieSceneViewportParams> ViewportParamsMap;
+		Player.GetViewportSettings(ViewportParamsMap);
+		for (auto ViewportParamsPair : ViewportParamsMap)
+		{
+			ViewportParamsMap[ViewportParamsPair.Key] = ViewportParams;
+		}
+		Player.SetViewportSettings(ViewportParamsMap);
+
+		// Set runtime fade
+		for (const FWorldContext& Context : GEngine->GetWorldContexts())
+		{
+			if (Context.WorldType == EWorldType::Game || Context.WorldType == EWorldType::PIE)
+			{
+				UWorld* World = Context.World();
+				if (World != nullptr)
+				{
+					APlayerController* PlayerController = World->GetGameInstance()->GetFirstLocalPlayerController();
+					if (PlayerController != nullptr && PlayerController->PlayerCameraManager && !PlayerController->PlayerCameraManager->IsPendingKill())
+					{
+						PlayerController->PlayerCameraManager->SetManualCameraFade(FloatValue, FadeColor, bFadeAudio);
+					}
+				}
+			}
+		}
+	}
+}
diff --git a/Engine/Source/Runtime/MovieSceneTracks/Private/Tracks/MovieSceneFadeTrack.cpp b/Engine/Source/Runtime/MovieSceneTracks/Private/Tracks/MovieSceneFadeTrack.cpp
new file mode 100644
--- /dev/null
+++ b/Engine/Source/Runtime/MovieSceneTracks/Private/Tracks/MovieSceneFadeTrack.cpp
@@ -0,0 +1,27 @@
+// Copyright 1998-2015 Epic Games, Inc. All Rights Reserved.
+
+#include "MovieSceneTracksPrivatePCH.h"
+#include "MovieSceneFadeSection.h"
+#include "MovieSceneFadeTrack.h"
+#include "MovieSceneFadeTrackInstance.h"
+
+
+/* UMovieSceneEventTrack interface
+ *****************************************************************************/
+
+TSharedPtr<IMovieSceneTrackInstance> UMovieSceneFadeTrack::CreateInstance()
+{
+	return MakeShareable(new FMovieSceneFadeTrackInstance(*this)); 
+}
+
+
+UMovieSceneSection* UMovieSceneFadeTrack::CreateNewSection()
+{
+	return NewObject<UMovieSceneSection>(this, UMovieSceneFadeSection::StaticClass(), NAME_None, RF_Transactional);
+}
+
+
+FName UMovieSceneFadeTrack::GetTrackName() const
+{
+	return FName("Fade");
+}
diff --git a/Engine/Source/Runtime/MovieSceneTracks/Public/Sections/MovieSceneFadeSection.h b/Engine/Source/Runtime/MovieSceneTracks/Public/Sections/MovieSceneFadeSection.h
new file mode 100644
--- /dev/null
+++ b/Engine/Source/Runtime/MovieSceneTracks/Public/Sections/MovieSceneFadeSection.h
@@ -0,0 +1,30 @@
+// Copyright 1998-2015 Epic Games, Inc. All Rights Reserved.
+
+#pragma once
+
+#include "MovieSceneFloatSection.h"
+#include "MovieSceneFadeSection.generated.h"
+
+
+/**
+ * A single floating point section
+ */
+UCLASS(MinimalAPI)
+class UMovieSceneFadeSection
+	: public UMovieSceneFloatSection
+{
+	GENERATED_BODY()
+
+	/** Default constructor. */
+	UMovieSceneFadeSection();
+
+public:
+
+	/** Fade color */
+	UPROPERTY(EditAnywhere, Category="Fade")
+	FLinearColor FadeColor;
+
+	/** Fade audio */
+	UPROPERTY(EditAnywhere, Category="Fade")
+	uint32 bFadeAudio:1;
+};
diff --git a/Engine/Source/Runtime/MovieSceneTracks/Public/TrackInstances/MovieSceneFadeTrackInstance.h b/Engine/Source/Runtime/MovieSceneTracks/Public/TrackInstances/MovieSceneFadeTrackInstance.h
new file mode 100644
--- /dev/null
+++ b/Engine/Source/Runtime/MovieSceneTracks/Public/TrackInstances/MovieSceneFadeTrackInstance.h
@@ -0,0 +1,46 @@
+// Copyright 1998-2015 Epic Games, Inc. All Rights Reserved.
+
+#pragma once
+
+#include "IMovieSceneTrackInstance.h"
+
+
+class UMovieSceneFadeTrack;
+
+
+/**
+ * Instance of a UMovieSceneFadeTrack
+ */
+class FMovieSceneFadeTrackInstance
+	: public IMovieSceneTrackInstance
+{
+public:
+
+	/**
+	 * Creates and initializes a new instance.
+	 *
+	 * @param InFadeTrack The event track to create an instance for.
+	 */
+	FMovieSceneFadeTrackInstance(UMovieSceneFadeTrack& InFadeTrack);
+
+	/** Virtual destructor. */
+	virtual ~FMovieSceneFadeTrackInstance() { }
+
+public:
+
+	// IMovieSceneTrackInstance interface
+
+	virtual void ClearInstance(IMovieScenePlayer& Player) override { }
+	virtual void RefreshInstance(const TArray<UObject*>& RuntimeObjects, IMovieScenePlayer& Player) override { }
+	virtual void RestoreState(const TArray<UObject*>& RuntimeObjects) override { }
+	virtual void SaveState(const TArray<UObject*>& RuntimeObjects) override { }
+	virtual void Update(float Position, float LastPosition, const TArray<UObject*>& RuntimeObjects, IMovieScenePlayer& Player) override;
+
+private:
+
+	/** Track that is being instanced */
+	UMovieSceneFadeTrack* FadeTrack;
+
+	/** A map from the viewport client to the viewport params that is stored before the track takes control and restored when the track relinquishes control */
+	TMap<TSharedPtr<FViewportClient>, EMovieSceneViewportParams> ViewportParamMap;
+};
diff --git a/Engine/Source/Runtime/UMG/Public/Animation/UMGSequencePlayer.h b/Engine/Source/Runtime/UMG/Public/Animation/UMGSequencePlayer.h
--- a/Engine/Source/Runtime/UMG/Public/Animation/UMGSequencePlayer.h
+++ b/Engine/Source/Runtime/UMG/Public/Animation/UMGSequencePlayer.h
@@ -37,6 +37,8 @@ public:
 	/** IMovieScenePlayer interface */
 	virtual void GetRuntimeObjects( TSharedRef<FMovieSceneSequenceInstance> MovieSceneInstance, const FGuid& ObjectHandle, TArray< UObject* >& OutObjects ) const override;
 	virtual void UpdateCameraCut(UObject* ObjectToViewThrough, bool bNewCameraCut) const override {}
+	virtual void SetViewportSettings(const TMap<FViewportClient*, EMovieSceneViewportParams>& ViewportParamsMap) override {}
+	virtual void GetViewportSettings(TMap<FViewportClient*, EMovieSceneViewportParams>& ViewportParamsMap) const override {}
 	virtual void AddOrUpdateMovieSceneInstance( class UMovieSceneSection& MovieSceneSection, TSharedRef<FMovieSceneSequenceInstance> InstanceToAdd ) override {}
 	virtual void RemoveMovieSceneInstance( class UMovieSceneSection& MovieSceneSection, TSharedRef<FMovieSceneSequenceInstance> InstanceToRemove ) override {}
 	virtual TSharedRef<FMovieSceneSequenceInstance> GetRootMovieSceneSequenceInstance() const override { return RootMovieSceneInstance.ToSharedRef(); }
