Sequencer: Default to 30 fps.

diff --git a/Engine/Source/Editor/Sequencer/Private/SequencerSettings.cpp b/Engine/Source/Editor/Sequencer/Private/SequencerSettings.cpp
--- a/Engine/Source/Editor/Sequencer/Private/SequencerSettings.cpp
+++ b/Engine/Source/Editor/Sequencer/Private/SequencerSettings.cpp
@@ -341,7 +341,7 @@ ULevelEditorSequencerSettings::ULevelEditorSequencerSettings( const FObjectIniti
 	: Super( ObjectInitializer )
 {
 	bKeyInterpPropertiesOnly = true;
-	TimeSnapInterval = 0.041667f;
+	TimeSnapInterval = 0.033334f;
 	bShowRangeSlider = true;
 	bLockInOutToStartEndRange = true;
 }
\ No newline at end of file
