Sequencer: Another round of fixes for keying properties within a struct. In the +Track button, recurse into the struct for its properties instead testing the struct property itself for keyability. This allows the post process settings struct and all of it's properties to be keyable.

diff --git a/Engine/Source/Editor/Sequencer/Private/SequencerDisplayNode.cpp b/Engine/Source/Editor/Sequencer/Private/SequencerDisplayNode.cpp
--- a/Engine/Source/Editor/Sequencer/Private/SequencerDisplayNode.cpp
+++ b/Engine/Source/Editor/Sequencer/Private/SequencerDisplayNode.cpp
@@ -14,6 +14,7 @@
 #include "CommonMovieSceneTools.h"
 #include "IKeyArea.h"
 #include "GroupedKeyArea.h"
+#include "ObjectEditorUtils.h"
 
 #define LOCTEXT_NAMESPACE "SequencerDisplayNode"
 
@@ -701,6 +702,9 @@ void FObjectBindingNode::BuildContextMenu(FMenuBuilder& MenuBuilder)
 
 void GetKeyablePropertyPaths(UClass* Class, UStruct* PropertySource, TArray<UProperty*>& PropertyPath, FSequencer& Sequencer, TArray<TArray<UProperty*>>& KeyablePropertyPaths)
 {
+	//@todo need to resolve this between UMG and the level editor sequencer
+	const bool bRecurseAllProperties = Sequencer.IsLevelEditorSequencer();
+
 	for (TFieldIterator<UProperty> PropertyIterator(PropertySource); PropertyIterator; ++PropertyIterator)
 	{
 		UProperty* Property = *PropertyIterator;
@@ -708,11 +712,14 @@ void GetKeyablePropertyPaths(UClass* Class, UStruct* PropertySource, TArray<UPro
 		if (Property && !Property->HasAnyPropertyFlags(CPF_Deprecated))
 		{
 			PropertyPath.Add(Property);
-			if (Sequencer.CanKeyProperty(FCanKeyPropertyParams(Class, PropertyPath)))
+
+			bool bIsPropertyKeyable = Sequencer.CanKeyProperty(FCanKeyPropertyParams(Class, PropertyPath));
+			if (bIsPropertyKeyable)
 			{
 				KeyablePropertyPaths.Add(PropertyPath);
 			}
-			else
+
+			if (!bIsPropertyKeyable || bRecurseAllProperties)
 			{
 				UStructProperty* StructProperty = Cast<UStructProperty>(Property);
 				if (StructProperty != nullptr)
@@ -720,36 +727,26 @@ void GetKeyablePropertyPaths(UClass* Class, UStruct* PropertySource, TArray<UPro
 					GetKeyablePropertyPaths(Class, StructProperty->Struct, PropertyPath, Sequencer, KeyablePropertyPaths);
 				}
 			}
+
 			PropertyPath.RemoveAt(PropertyPath.Num() - 1);
 		}
 	}
 }
 
+struct PropertyMenuData
+{
+	FString MenuName;
+	TArray<UProperty*> PropertyPath;
+};
+	
 TSharedRef<SWidget> FObjectBindingNode::OnGetAddTrackMenuContent()
 {
-	TArray<TArray<UProperty*>> KeyablePropertyPaths;
 	FSequencer& Sequencer = GetSequencer();
-	UObject* BoundObject = Sequencer.GetFocusedMovieSceneSequence()->FindObject(ObjectBinding);
 
-	if (BoundObject != nullptr)
-	{
-		TArray<UProperty*> PropertyPath;
-		GetKeyablePropertyPaths(BoundObject->GetClass(), BoundObject->GetClass(), PropertyPath, Sequencer, KeyablePropertyPaths);
-	}
+	//@todo need to resolve this between UMG and the level editor sequencer
+	const bool bUseSubMenus = Sequencer.IsLevelEditorSequencer();
 
-	KeyablePropertyPaths.Sort([](const TArray<UProperty*>& A,const TArray<UProperty*>& B)
-	{
-		for (int32 IndexToCheck = 0; IndexToCheck < A.Num() && IndexToCheck < B.Num(); IndexToCheck++)
-		{
-			int32 CompareResult = A[IndexToCheck]->GetDisplayNameText().ToString().Compare(B[IndexToCheck]->GetDisplayNameText().ToString());
-			if (CompareResult != 0)
-			{
-				return CompareResult < 0;
-			}
-		}
-		// If a difference in name wasn't found, the shorter path should be first.
-		return A.Num() < B.Num();
-	});
+	UObject* BoundObject = Sequencer.GetFocusedMovieSceneSequence()->FindObject(ObjectBinding);
 
 	ISequencerModule& SequencerModule = FModuleManager::GetModuleChecked<ISequencerModule>( "Sequencer" );
 	TSharedRef<FUICommandList> CommandList(new FUICommandList);
@@ -760,85 +757,193 @@ TSharedRef<SWidget> FObjectBindingNode::OnGetAddTrackMenuContent()
 	GetSequencer().BuildObjectBindingTrackMenu(AddTrackMenuBuilder, ObjectBinding, ObjectClass);
 	AddTrackMenuBuilder.EndSection();
 
+	TArray<TArray<UProperty*>> KeyablePropertyPaths;
+
+	if (BoundObject != nullptr)
+	{
+		TArray<UProperty*> PropertyPath;
+		GetKeyablePropertyPaths(BoundObject->GetClass(), BoundObject->GetClass(), PropertyPath, Sequencer, KeyablePropertyPaths);
+	}
+
+	// [Aspect Ratio]
+	// [PostProcess Settings] [Bloom1Tint] [X]
+	// [PostProcess Settings] [Bloom1Tint] [Y]
+	// [PostProcess Settings] [ColorGrading]
+	// [Ortho View]
+
+	// Create property menu data based on keyable property paths
+	TArray<PropertyMenuData> KeyablePropertyMenuData;
+	for (auto KeyablePropertyPath : KeyablePropertyPaths)
+	{
+		PropertyMenuData KeyableMenuData;
+		KeyableMenuData.PropertyPath = KeyablePropertyPath;
+		KeyableMenuData.MenuName = KeyablePropertyPath[0]->GetDisplayNameText().ToString();
+		KeyablePropertyMenuData.Add(KeyableMenuData);
+	}
+
+	// Sort on the menu name
+	KeyablePropertyMenuData.Sort([](const PropertyMenuData& A, const PropertyMenuData& B)
+	{
+		int32 CompareResult = A.MenuName.Compare(B.MenuName);
+		return CompareResult < 0;
+	});
+	
+
+	// Add menu items
 	AddTrackMenuBuilder.BeginSection( SequencerMenuExtensionPoints::AddTrackMenu_PropertiesSection, LOCTEXT("PropertiesMenuHeader" , "Properties"));
+	for (int32 MenuDataIndex = 0; MenuDataIndex < KeyablePropertyMenuData.Num(); )
 	{
-		// Gather list of property path chains, ie:
-		//   [Aspect Ratio]
-		//   [Postprocess Settings] [ColorGamma]
-		//   [Postprocess Settings] [Gain]
-		//   [Focal Length]
-		TArray< TArray<FString> > PropertyChainPaths;
-		for ( TArray<UProperty*>& KeyablePropertyPath : KeyablePropertyPaths )
+		TArray<TArray<UProperty*> > KeyableSubMenuPropertyPaths;
+
+		KeyableSubMenuPropertyPaths.Add(KeyablePropertyMenuData[MenuDataIndex].PropertyPath);
+
+		// If this menu data only has one property name, add the menu item
+		if (KeyablePropertyMenuData[MenuDataIndex].PropertyPath.Num() == 1 || !bUseSubMenus)
+		{
+			AddPropertyMenuItems(AddTrackMenuBuilder, KeyableSubMenuPropertyPaths);
+			++MenuDataIndex;
+		}
+		// Otherwise, look to the next menu data to gather up new data
+		else
 		{
-			TArray<FString> PropertyNames;
-			for ( UProperty* Property : KeyablePropertyPath )
+			for (; MenuDataIndex < KeyablePropertyMenuData.Num()-1; )
 			{
-				PropertyNames.Add( Property->GetDisplayNameText().ToString() );
+				if (KeyablePropertyMenuData[MenuDataIndex].MenuName == KeyablePropertyMenuData[MenuDataIndex+1].MenuName)
+				{	
+					++MenuDataIndex;
+					KeyableSubMenuPropertyPaths.Add(KeyablePropertyMenuData[MenuDataIndex].PropertyPath);
+				}
+				else
+				{
+					break;
+				}
 			}
 
-			PropertyChainPaths.Add(PropertyNames);
+			AddTrackMenuBuilder.AddSubMenu(
+				FText::FromString(KeyablePropertyMenuData[MenuDataIndex].MenuName),
+				FText::GetEmpty(), 
+				FNewMenuDelegate::CreateSP(this, &FObjectBindingNode::AddPropertyMenuItemsWithCategories, KeyableSubMenuPropertyPaths));
+
+			++MenuDataIndex;
 		}
+	}
+	AddTrackMenuBuilder.EndSection();
 
-		//@todo this currently only supports two level deep property chains. it should be arbitrary.
-		for (int32 PropertyPathIndex = 0; PropertyPathIndex < PropertyChainPaths.Num(); ++PropertyPathIndex )
-		{ 
-			TArray<TArray<FString> > SubPropertyChainPaths;
-			TArray<TArray<UProperty*> > SubKeyablePropertyPaths;
+	return AddTrackMenuBuilder.MakeWidget();
+}
 
-			TArray<FString> PropertyChainPath = PropertyChainPaths[PropertyPathIndex];
-			SubPropertyChainPaths.Add(PropertyChainPath);
-			SubKeyablePropertyPaths.Add(KeyablePropertyPaths[PropertyPathIndex]);
+void FObjectBindingNode::AddPropertyMenuItemsWithCategories(FMenuBuilder& AddTrackMenuBuilder, TArray<TArray<UProperty*> > KeyablePropertyPaths)
+{
+	// [PostProcessSettings] [Bloom1Tint] [X]
+	// [PostProcessSettings] [Bloom1Tint] [Y]
+	// [PostProcessSettings] [ColorGrading]
 
-			// Look for the next chain that is not the same length
-			int32 NextPropertyPathIndex = PropertyPathIndex + 1; 
-			for (; NextPropertyPathIndex < PropertyChainPaths.Num(); ++NextPropertyPathIndex)
-			{
-				if (PropertyChainPaths[PropertyPathIndex].Num() != PropertyChainPaths[NextPropertyPathIndex].Num())
-				{
-					break;
-				}
+	// Create property menu data based on keyable property paths
+	TSet<UProperty*> PropertiesTraversed;
+	TArray<PropertyMenuData> KeyablePropertyMenuData;
+	for (auto KeyablePropertyPath : KeyablePropertyPaths)
+	{		
+		PropertyMenuData KeyableMenuData;
+		KeyableMenuData.PropertyPath = KeyablePropertyPath;
 
-				TArray<FString> NextPropertyChainPath = PropertyChainPaths[NextPropertyPathIndex];
-				SubPropertyChainPaths.Add(NextPropertyChainPath);
-				SubKeyablePropertyPaths.Add(KeyablePropertyPaths[NextPropertyPathIndex]);
-			}
-			
-			// Don't use a submenu if there's only one menu item
-			if (PropertyChainPaths[PropertyPathIndex].Num() > 1 && SubPropertyChainPaths.Num() > 1)
+		// If the path is greater than 1, keep track of the actual properties (not channels) and only add these properties once since we can't do single channel keying of a property yet.
+		if (KeyablePropertyPath.Num() > 1) //@todo
+		{
+			if (PropertiesTraversed.Find(KeyablePropertyPath[1]) != nullptr)
 			{
-				// Strip off the first member of the chain since it's the submenu title
-				for (int32 SubPropertyChainIndex = 0; SubPropertyChainIndex < SubPropertyChainPaths.Num(); ++SubPropertyChainIndex)
-				{
-					SubPropertyChainPaths[SubPropertyChainIndex].RemoveAt(0);
-				}
+				continue;
+			}
 
-				int32 ChainLength = PropertyChainPaths[PropertyPathIndex].Num();
-				const FText& ParentDisplayName = FText::FromString(PropertyChainPaths[PropertyPathIndex][0]);
-				AddTrackMenuBuilder.AddSubMenu(
-					ParentDisplayName, 
-					FText::GetEmpty(), 
-					FNewMenuDelegate::CreateSP(this, &FObjectBindingNode::AddPropertyMenuItems, SubPropertyChainPaths, SubKeyablePropertyPaths) 
-					);
+			KeyableMenuData.MenuName = FObjectEditorUtils::GetCategoryFName(KeyablePropertyPath[1]).ToString();
+			PropertiesTraversed.Add(KeyablePropertyPath[1]);
+		}
+		else
+		{
+			// No sub menus items, so skip
+			continue; 
+		}
+		KeyablePropertyMenuData.Add(KeyableMenuData);
+	}
+
+	// Sort on the menu name
+	KeyablePropertyMenuData.Sort([](const PropertyMenuData& A, const PropertyMenuData& B)
+	{
+		int32 CompareResult = A.MenuName.Compare(B.MenuName);
+		return CompareResult < 0;
+	});
+
+	// Add menu items
+	for (int32 MenuDataIndex = 0; MenuDataIndex < KeyablePropertyMenuData.Num(); )
+	{
+		TArray<TArray<UProperty*> > KeyableSubMenuPropertyPaths;
+		KeyableSubMenuPropertyPaths.Add(KeyablePropertyMenuData[MenuDataIndex].PropertyPath);
+
+		for (; MenuDataIndex < KeyablePropertyMenuData.Num()-1; )
+		{
+			if (KeyablePropertyMenuData[MenuDataIndex].MenuName == KeyablePropertyMenuData[MenuDataIndex+1].MenuName)
+			{
+				++MenuDataIndex;
+				KeyableSubMenuPropertyPaths.Add(KeyablePropertyMenuData[MenuDataIndex].PropertyPath);
 			}
 			else
 			{
-				AddPropertyMenuItems(AddTrackMenuBuilder, SubPropertyChainPaths, SubKeyablePropertyPaths);
+				break;
 			}
-
-			PropertyPathIndex = NextPropertyPathIndex;
 		}
-	}
-	AddTrackMenuBuilder.EndSection();
 
-	return AddTrackMenuBuilder.MakeWidget();
+		const int32 PropertyNameIndexStart = 1; // Strip off the struct property name
+		const int32 PropertyNameIndexEnd = 2; // Stop at the property name, don't descend into the channels
+
+		AddTrackMenuBuilder.AddSubMenu(
+			FText::FromString(KeyablePropertyMenuData[MenuDataIndex].MenuName),
+			FText::GetEmpty(), 
+			FNewMenuDelegate::CreateSP(this, &FObjectBindingNode::AddPropertyMenuItems, KeyableSubMenuPropertyPaths, PropertyNameIndexStart, PropertyNameIndexEnd));
+
+		++MenuDataIndex;
+	}
 }
 
-void FObjectBindingNode::AddPropertyMenuItems(FMenuBuilder& AddTrackMenuBuilder, TArray<TArray<FString> > PropertyChainPaths, TArray<TArray<UProperty*> > KeyablePropertyPaths)
+void FObjectBindingNode::AddPropertyMenuItems(FMenuBuilder& AddTrackMenuBuilder, TArray<TArray<UProperty*> > KeyableProperties, int32 PropertyNameIndexStart, int32 PropertyNameIndexEnd)
 {
-	for (int32 PropertyPathIndex = 0; PropertyPathIndex < PropertyChainPaths.Num(); ++PropertyPathIndex)
+	TArray<PropertyMenuData> KeyablePropertyMenuData;
+
+	for (auto KeyableProperty : KeyableProperties)
+	{
+		TArray<FString> PropertyNames;
+		if (PropertyNameIndexEnd == -1)
+		{
+			PropertyNameIndexEnd = KeyableProperty.Num();
+		}
+
+		//@todo
+		if (PropertyNameIndexStart >= KeyableProperty.Num())
+		{
+			continue;
+		}
+
+		for (int32 PropertyNameIndex = PropertyNameIndexStart; PropertyNameIndex < PropertyNameIndexEnd; ++PropertyNameIndex)
+		{
+			PropertyNames.Add(KeyableProperty[PropertyNameIndex]->GetDisplayNameText().ToString());
+		}
+
+		PropertyMenuData KeyableMenuData;
+		KeyableMenuData.PropertyPath = KeyableProperty;
+		KeyableMenuData.MenuName = FString::Join( PropertyNames, TEXT( "." ) );
+		KeyablePropertyMenuData.Add(KeyableMenuData);
+	}
+
+	// Sort on the menu name
+	KeyablePropertyMenuData.Sort([](const PropertyMenuData& A, const PropertyMenuData& B)
+	{
+		int32 CompareResult = A.MenuName.Compare(B.MenuName);
+		return CompareResult < 0;
+	});
+
+	// Add menu items
+	for (int32 MenuDataIndex = 0; MenuDataIndex < KeyablePropertyMenuData.Num(); ++MenuDataIndex)
 	{
-		FUIAction AddTrackMenuAction( FExecuteAction::CreateSP( this, &FObjectBindingNode::AddTrackForProperty, KeyablePropertyPaths[PropertyPathIndex] ) );
-		AddTrackMenuBuilder.AddMenuEntry( FText::FromString( FString::Join( PropertyChainPaths[PropertyPathIndex], TEXT( "." ) ) ), FText(), FSlateIcon(), AddTrackMenuAction );
+		FUIAction AddTrackMenuAction( FExecuteAction::CreateSP( this, &FObjectBindingNode::AddTrackForProperty, KeyablePropertyMenuData[MenuDataIndex].PropertyPath ) );
+		AddTrackMenuBuilder.AddMenuEntry( FText::FromString(KeyablePropertyMenuData[MenuDataIndex].MenuName), FText(), FSlateIcon(), AddTrackMenuAction );
 	}
 }
 
diff --git a/Engine/Source/Editor/Sequencer/Private/SequencerDisplayNode.h b/Engine/Source/Editor/Sequencer/Private/SequencerDisplayNode.h
--- a/Engine/Source/Editor/Sequencer/Private/SequencerDisplayNode.h
+++ b/Engine/Source/Editor/Sequencer/Private/SequencerDisplayNode.h
@@ -444,7 +444,8 @@ public:
 private:
 	TSharedRef<SWidget> OnGetAddTrackMenuContent();
 
-	void AddPropertyMenuItems(FMenuBuilder& AddTrackMenuBuilder, TArray<TArray<FString> > PropertyChainPaths, TArray<TArray<UProperty*> > KeyablePropertyPaths);
+	void AddPropertyMenuItemsWithCategories(FMenuBuilder& AddTrackMenuBuilder, TArray<TArray<UProperty*> > KeyablePropertyPath);
+	void AddPropertyMenuItems(FMenuBuilder& AddTrackMenuBuilder, TArray<TArray<UProperty*> > KeyableProperties, int32 PropertyNameIndexStart = 0, int32 PropertyNameIndexEnd = -1);
 
 	void AddTrackForProperty(TArray<UProperty*> PropertyPath);
 
diff --git a/Engine/Source/Editor/Sequencer/Private/SequencerObjectChangeListener.cpp b/Engine/Source/Editor/Sequencer/Private/SequencerObjectChangeListener.cpp
--- a/Engine/Source/Editor/Sequencer/Private/SequencerObjectChangeListener.cpp
+++ b/Engine/Source/Editor/Sequencer/Private/SequencerObjectChangeListener.cpp
@@ -215,6 +215,29 @@ bool FSequencerObjectChangeListener::CanKeyProperty(FCanKeyPropertyParams CanKey
 		bFound = FindPropertySetter(*CanKeyPropertyParams.ObjectClass, CanKeyPropertyParams.PropertyPath.Last()->GetClass()->GetFName(), CanKeyPropertyParams.PropertyPath.Last()->GetName());
 	}
 
+	if ( !bFound )
+	{
+		bool bFoundValidInterp = false;
+		bool bFoundEditDefaultsOnly = false;
+		bool bFoundEdit = false;
+		UProperty* Property = CanKeyPropertyParams.PropertyPath.Last();
+		if (Property->HasAnyPropertyFlags(CPF_Interp))
+		{
+			bFoundValidInterp = true;
+		}
+		if (Property->HasAnyPropertyFlags(CPF_DisableEditOnInstance))
+		{
+			bFoundEditDefaultsOnly = true;
+		}
+		if (Property->HasAnyPropertyFlags(CPF_Edit))
+		{
+			bFoundEdit = true;
+		}
+
+		// Valid Interp keyword is found. The property also needs to be editable.
+		bFound = bFoundValidInterp && bFoundEdit && !bFoundEditDefaultsOnly;
+	}
+
 	return bFound;
 }
 
