Sequencer: Fixed play rate track not working when playing in Sequencer

diff --git a/Engine/Source/Editor/Sequencer/Private/Sequencer.cpp b/Engine/Source/Editor/Sequencer/Private/Sequencer.cpp
--- a/Engine/Source/Editor/Sequencer/Private/Sequencer.cpp
+++ b/Engine/Source/Editor/Sequencer/Private/Sequencer.cpp
@@ -237,7 +237,8 @@ void FSequencer::Tick(float InDeltaTime)
 		SetGlobalTimeDirectly(GetGlobalTime() + Offset);
 	}
 
-	float NewTime = GetGlobalTime() + InDeltaTime;
+	float NewTime = GetGlobalTime() + InDeltaTime * GWorld->GetWorldSettings()->MatineeTimeDilation;
+
 	if (PlaybackState == EMovieScenePlayerStatus::Playing ||
 		PlaybackState == EMovieScenePlayerStatus::Recording)
 	{
