Sequencer: Section menu tweaks (put trim and split in edit menu, add can select all keys and can snap keys to frame).

diff --git a/Engine/Source/Editor/Sequencer/Private/SSection.cpp b/Engine/Source/Editor/Sequencer/Private/SSection.cpp
--- a/Engine/Source/Editor/Sequencer/Private/SSection.cpp
+++ b/Engine/Source/Editor/Sequencer/Private/SSection.cpp
@@ -568,7 +568,9 @@ TSharedPtr<SWidget> SSection::OnSummonContextMenu( const FGeometry& MyGeometry,
 				NSLOCTEXT("Sequencer", "SnapToFrame", "Snap to Frame"),
 				NSLOCTEXT("Sequencer", "SnapToFrameToolTip", "Snap selected keys to frame"),
 				FSlateIcon(),
-				FUIAction(FExecuteAction::CreateSP(&Sequencer, &FSequencer::SnapToFrame))
+				FUIAction(
+					FExecuteAction::CreateSP(&Sequencer, &FSequencer::SnapToFrame),
+					FCanExecuteAction::CreateSP(&Sequencer, &FSequencer::CanSnapToFrame))
 			);
 
 			MenuBuilder.AddMenuEntry(
@@ -592,35 +594,15 @@ TSharedPtr<SWidget> SSection::OnSummonContextMenu( const FGeometry& MyGeometry,
 				NSLOCTEXT("Sequencer", "SelectAllKeys", "Select All Keys"),
 				NSLOCTEXT("Sequencer", "SelectAllKeysTooltip", "Select all keys in section"),
 				FSlateIcon(),
-				FUIAction(FExecuteAction::CreateSP(this, &SSection::SelectAllKeys))
-			);
-
-			MenuBuilder.AddMenuEntry(
-				NSLOCTEXT("Sequencer", "TrimSectionLeft", "Trim Left"),
-				NSLOCTEXT("Sequencer", "TrimSectionLeftTooltip", "Trim section at current time to the left"),
-				FSlateIcon(),
 				FUIAction(
-					FExecuteAction::CreateSP(this, &SSection::TrimSection, true),
-					FCanExecuteAction::CreateSP(this, &SSection::IsTrimmable))
+					FExecuteAction::CreateSP(this, &SSection::SelectAllKeys),
+					FCanExecuteAction::CreateSP(this, &SSection::CanSelectAllKeys))
 			);
 
-			MenuBuilder.AddMenuEntry(
-				NSLOCTEXT("Sequencer", "TrimSectionRight", "Trim Right"),
-				NSLOCTEXT("Sequencer", "TrimSectionRightTooltip", "Trim section at current time to the right"),
-				FSlateIcon(),
-				FUIAction(
-					FExecuteAction::CreateSP(this, &SSection::TrimSection, false),
-					FCanExecuteAction::CreateSP(this, &SSection::IsTrimmable))
-			);
-
-			MenuBuilder.AddMenuEntry(
-				NSLOCTEXT("Sequencer", "SplitSection", "Split"),
-				NSLOCTEXT("Sequencer", "SplitSectionTooltip", "Split section at current time"),
-				FSlateIcon(),
-				FUIAction(
-					FExecuteAction::CreateSP(this, &SSection::SplitSection),
-					FCanExecuteAction::CreateSP(this, &SSection::IsTrimmable))
-			);
+			MenuBuilder.AddSubMenu(
+				NSLOCTEXT("Sequencer", "EditSection", "Edit"),
+				NSLOCTEXT("Sequencer", "EditSectionTooltip", "Edit section"),
+				FNewMenuDelegate::CreateRaw(this, &SSection::AddEditMenu));
 
 			MenuBuilder.AddSubMenu(
 				NSLOCTEXT("Sequencer", "SetPreInfinityExtrap", "Pre-Infinity"), 
@@ -666,6 +648,36 @@ TSharedPtr<SWidget> SSection::OnSummonContextMenu( const FGeometry& MyGeometry,
 	return TSharedPtr<SWidget>();
 }
 
+void SSection::AddEditMenu(FMenuBuilder& MenuBuilder)
+{
+	MenuBuilder.AddMenuEntry(
+		NSLOCTEXT("Sequencer", "TrimSectionLeft", "Trim Left"),
+		NSLOCTEXT("Sequencer", "TrimSectionLeftTooltip", "Trim section at current time to the left"),
+		FSlateIcon(),
+		FUIAction(
+			FExecuteAction::CreateSP(this, &SSection::TrimSection, true),
+			FCanExecuteAction::CreateSP(this, &SSection::IsTrimmable))
+	);
+
+	MenuBuilder.AddMenuEntry(
+		NSLOCTEXT("Sequencer", "TrimSectionRight", "Trim Right"),
+		NSLOCTEXT("Sequencer", "TrimSectionRightTooltip", "Trim section at current time to the right"),
+		FSlateIcon(),
+		FUIAction(
+			FExecuteAction::CreateSP(this, &SSection::TrimSection, false),
+			FCanExecuteAction::CreateSP(this, &SSection::IsTrimmable))
+	);
+
+	MenuBuilder.AddMenuEntry(
+		NSLOCTEXT("Sequencer", "SplitSection", "Split"),
+		NSLOCTEXT("Sequencer", "SplitSectionTooltip", "Split section at current time"),
+		FSlateIcon(),
+		FUIAction(
+			FExecuteAction::CreateSP(this, &SSection::SplitSection),
+			FCanExecuteAction::CreateSP(this, &SSection::IsTrimmable))
+	);
+}
+
 void SSection::AddExtrapolationMenu(FMenuBuilder& MenuBuilder, bool bPreInfinity)
 {
 	FSequencer& Sequencer = GetSequencer();
@@ -1032,6 +1044,25 @@ void SSection::SelectAllKeys()
 	}
 }
 
+bool SSection::CanSelectAllKeys() const
+{
+	// @todo Sequencer should operate on selected sections
+	UMovieSceneSection* Section = SectionInterface->GetSectionObject();
+
+	for (const FKeyAreaLayoutElement& Element : Layout->GetElements())
+	{
+		TSharedPtr<IKeyArea> KeyArea = Element.GetKeyArea();
+
+		TArray<FKeyHandle> KeyHandles = KeyArea->GetUnsortedKeyHandles();
+		if (KeyHandles.Num() > 0)
+		{
+			return true;
+		}
+	}
+
+	return false;
+}
+
 void SSection::TrimSection(bool bTrimLeft)
 {
 	FScopedTransaction TrimSectionTransaction(NSLOCTEXT("Sequencer", "TrimSection_Transaction", "Trim Section"));
diff --git a/Engine/Source/Editor/Sequencer/Private/SSection.h b/Engine/Source/Editor/Sequencer/Private/SSection.h
--- a/Engine/Source/Editor/Sequencer/Private/SSection.h
+++ b/Engine/Source/Editor/Sequencer/Private/SSection.h
@@ -93,6 +93,9 @@ private:
 	/** Summons a context menu over the associated section */
 	TSharedPtr<SWidget> OnSummonContextMenu( const FGeometry& MyGeometry, const FPointerEvent& MouseEvent );
 
+	/** Add edit menu for trim and split */
+	void AddEditMenu(FMenuBuilder& MenuBuilder);
+
 	/** Add extrapolation menu for pre and post infinity */
 	void AddExtrapolationMenu(FMenuBuilder& MenuBuilder, bool bPreInfinity);
 
@@ -153,17 +156,22 @@ private:
 	 */
 	void SelectAllKeys();
 
-	/*
+	/**
+	 * Are there keys to select? 
+	 */
+	bool CanSelectAllKeys() const;
+
+	/**
 	 * Trim section 
 	 */
 	void TrimSection(bool bTrimLeft);
 
-	/*
+	/**
 	 * Split section
 	 */
 	void SplitSection();
 
-	/*
+	/**
 	 * Is there a trimmable section?
 	 */
 	bool IsTrimmable() const;
diff --git a/Engine/Source/Editor/Sequencer/Private/Sequencer.cpp b/Engine/Source/Editor/Sequencer/Private/Sequencer.cpp
--- a/Engine/Source/Editor/Sequencer/Private/Sequencer.cpp
+++ b/Engine/Source/Editor/Sequencer/Private/Sequencer.cpp
@@ -506,6 +506,14 @@ void FSequencer::SnapToFrame()
 		UpdateRuntimeInstances();
 	}
 }
+
+bool FSequencer::CanSnapToFrame() const
+{
+	const bool bKeysSelected = Selection.GetSelectedKeys().Num() > 0;
+
+	return bKeysSelected && CanShowFrameNumbers();
+}
+
 void FSequencer::SpawnOrDestroyPuppetObjects( TSharedRef<FMovieSceneSequenceInstance> MovieSceneInstance )
 {
 	UMovieSceneSequence* Sequence = MovieSceneInstance->GetSequence();
diff --git a/Engine/Source/Editor/Sequencer/Private/Sequencer.h b/Engine/Source/Editor/Sequencer/Private/Sequencer.h
--- a/Engine/Source/Editor/Sequencer/Private/Sequencer.h
+++ b/Engine/Source/Editor/Sequencer/Private/Sequencer.h
@@ -144,6 +144,11 @@ public:
 	 */
 	void SnapToFrame();
 
+	/**
+ 	 * Are there keys to snap? 
+	 */
+	bool CanSnapToFrame() const;
+
 	/**
 	 * @return Movie scene tools used by the sequencer
 	 */
