Sequencer: Snap time to selected key if setting is on.

diff --git a/Engine/Source/Editor/Sequencer/Private/Tools/SequencerEditTool_Movement.cpp b/Engine/Source/Editor/Sequencer/Private/Tools/SequencerEditTool_Movement.cpp
--- a/Engine/Source/Editor/Sequencer/Private/Tools/SequencerEditTool_Movement.cpp
+++ b/Engine/Source/Editor/Sequencer/Private/Tools/SequencerEditTool_Movement.cpp
@@ -35,6 +35,18 @@ FReply FSequencerEditTool_Movement::OnMouseButtonDown(SWidget& OwnerWidget, cons
 		const FVirtualTrackArea VirtualTrackArea = SequencerWidget.Pin()->GetVirtualTrackArea();
 
 		DelayedDrag = FDelayedDrag_Hotspot(VirtualTrackArea.CachedTrackAreaGeometry().AbsoluteToLocal(MouseEvent.GetScreenSpacePosition()), MouseEvent.GetEffectingButton(), Hotspot);
+
+		if (Sequencer.Pin()->GetSettings()->GetSnapPlayTimeToDraggedKey())
+		{
+			if (DelayedDrag->Hotspot.IsValid())
+			{
+				if (DelayedDrag->Hotspot->GetType() == ESequencerHotspot::Key)
+				{
+					FSelectedKey& ThisKey = StaticCastSharedPtr<FKeyHotspot>(DelayedDrag->Hotspot)->Key;
+					Sequencer.Pin()->SetGlobalTime(ThisKey.KeyArea->GetKeyTime(ThisKey.KeyHandle.GetValue()));
+				}
+			}
+		}
 		return FReply::Handled();
 	}
 	return FReply::Unhandled();
