Sequencer: "Movie scene" -> "Level sequence"

diff --git a/Engine/Source/Editor/Sequencer/Private/SSequencer.cpp b/Engine/Source/Editor/Sequencer/Private/SSequencer.cpp
--- a/Engine/Source/Editor/Sequencer/Private/SSequencer.cpp
+++ b/Engine/Source/Editor/Sequencer/Private/SSequencer.cpp
@@ -639,7 +639,7 @@ TSharedRef<SWidget> SSequencer::MakeToolBar()
 				FUIAction(FExecuteAction::CreateSP(this, &SSequencer::OnSaveMovieSceneClicked))
 				, NAME_None
 				, LOCTEXT("SaveDirtyPackages", "Save")
-				, LOCTEXT("SaveDirtyPackagesTooltip", "Saves the current movie scene")
+				, LOCTEXT("SaveDirtyPackagesTooltip", "Saves the current level sequence")
 				, FSlateIcon(FEditorStyle::GetStyleSetName(), "Sequencer.Save"));
 
 			ToolBarBuilder.AddToolBarButton( FSequencerCommands::Get().RenderMovie );
