Sequencer: Adjust timeline project settings to be in frame (default to 1) and duration (default to 5s)

diff --git a/Engine/Source/Editor/Sequencer/Private/Sequencer.cpp b/Engine/Source/Editor/Sequencer/Private/Sequencer.cpp
--- a/Engine/Source/Editor/Sequencer/Private/Sequencer.cpp
+++ b/Engine/Source/Editor/Sequencer/Private/Sequencer.cpp
@@ -1110,16 +1110,18 @@ void FSequencer::UpdateTimeBoundsToFocusedMovieScene()
 	USequencerProjectSettings* ProjectSettings = GetMutableDefault<USequencerProjectSettings>();
 
 	const int32 InFrame = ProjectSettings->InFrame;
-	const int32 OutFrame = ProjectSettings->OutFrame;
+	const float Duration = ProjectSettings->Duration;
+	const bool bUsingFrameRate = SequencerSnapValues::IsTimeSnapIntervalFrameRate(Settings->GetTimeSnapInterval());
+	float FrameRate = 1.0f / Settings->GetTimeSnapInterval();
 
 	UMovieScene* FocusedMovieScene = GetFocusedMovieSceneSequence()->GetMovieScene();
 
 	float InTime = FocusedMovieScene->InTime;
 	if (InTime >= FLT_MAX)
 	{
-		if (SequencerSnapValues::IsTimeSnapIntervalFrameRate(Settings->GetTimeSnapInterval()))
+		if (bUsingFrameRate)
 		{
-			InTime = SequencerHelpers::FrameToTime(InFrame, 1.0f/Settings->GetTimeSnapInterval());
+			InTime = SequencerHelpers::FrameToTime(InFrame, FrameRate);
 		}
 		else
 		{
@@ -1130,22 +1132,24 @@ void FSequencer::UpdateTimeBoundsToFocusedMovieScene()
 	float OutTime = FocusedMovieScene->OutTime;
 	if (OutTime <= -FLT_MAX)
 	{
-		if (SequencerSnapValues::IsTimeSnapIntervalFrameRate(Settings->GetTimeSnapInterval()))
+		if (bUsingFrameRate)
 		{
-			OutTime = SequencerHelpers::FrameToTime(OutFrame, 1.0f/Settings->GetTimeSnapInterval());
+			OutTime = SequencerHelpers::FrameToTime(InFrame, FrameRate) + Duration;
+			int32 OutFrame = SequencerHelpers::TimeToFrame(OutTime, FrameRate);
+			OutTime = SequencerHelpers::FrameToTime(OutFrame, FrameRate);
 		}
 		else
 		{
-			OutTime = 5.0f;
+			OutTime = Duration;
 		}
 	}
 
 	float StartTime = FocusedMovieScene->StartTime;
 	if (StartTime >= FLT_MAX)
 	{
-		if (SequencerSnapValues::IsTimeSnapIntervalFrameRate(Settings->GetTimeSnapInterval()))
+		if (bUsingFrameRate)
 		{
-			StartTime = SequencerHelpers::FrameToTime(InFrame, 1.0f/Settings->GetTimeSnapInterval());
+			StartTime = SequencerHelpers::FrameToTime(InFrame, FrameRate);
 		}
 		else
 		{
@@ -1156,13 +1160,15 @@ void FSequencer::UpdateTimeBoundsToFocusedMovieScene()
 	float EndTime = FocusedMovieScene->EndTime;
 	if (EndTime <= -FLT_MAX)
 	{
-		if (SequencerSnapValues::IsTimeSnapIntervalFrameRate(Settings->GetTimeSnapInterval()))
+		if (bUsingFrameRate)
 		{
-			EndTime = SequencerHelpers::FrameToTime(OutFrame, 1.0f/Settings->GetTimeSnapInterval());
+			EndTime = SequencerHelpers::FrameToTime(InFrame, FrameRate) + Duration;
+			int32 EndFrame = SequencerHelpers::TimeToFrame(EndTime, FrameRate);
+			EndTime = SequencerHelpers::FrameToTime(EndFrame, FrameRate);
 		}
 		else
 		{
-			EndTime = 5.0f;
+			EndTime = Duration;
 		}
 	}
 
diff --git a/Engine/Source/Editor/Sequencer/Private/SequencerProjectSettings.cpp b/Engine/Source/Editor/Sequencer/Private/SequencerProjectSettings.cpp
--- a/Engine/Source/Editor/Sequencer/Private/SequencerProjectSettings.cpp
+++ b/Engine/Source/Editor/Sequencer/Private/SequencerProjectSettings.cpp
@@ -7,7 +7,7 @@
 // USequencerProjectSettings
 
 USequencerProjectSettings::USequencerProjectSettings()
-	: InFrame(101)
-	, OutFrame(200)
+	: InFrame(1)
+	, Duration(5.f)
 {
 }
diff --git a/Engine/Source/Editor/Sequencer/Private/SequencerProjectSettings.h b/Engine/Source/Editor/Sequencer/Private/SequencerProjectSettings.h
--- a/Engine/Source/Editor/Sequencer/Private/SequencerProjectSettings.h
+++ b/Engine/Source/Editor/Sequencer/Private/SequencerProjectSettings.h
@@ -13,12 +13,12 @@ class USequencerProjectSettings : public UObject
 public:
 	USequencerProjectSettings();
 
-	/** The in frame. */
+	/** The default in frame. */
 	UPROPERTY(config, EditAnywhere, Category=Timeline)
 	int32 InFrame;
 
-	/** The out frame. */
+	/** The default duration in seconds. */
 	UPROPERTY(config, EditAnywhere, Category=Timeline)
-	int32 OutFrame;
+	float Duration;
 
 };
