Sequencer: Fix time scrubbing so that it respects the snap play time to interval toggle.

diff --git a/Engine/Source/Editor/Sequencer/Private/TimeSliderController.cpp b/Engine/Source/Editor/Sequencer/Private/TimeSliderController.cpp
--- a/Engine/Source/Editor/Sequencer/Private/TimeSliderController.cpp
+++ b/Engine/Source/Editor/Sequencer/Private/TimeSliderController.cpp
@@ -470,7 +470,7 @@ FReply FSequencerTimeSliderController::OnMouseButtonUp( TSharedRef<SWidget> Widg
 			FVector2D CursorPos = MyGeometry.AbsoluteToLocal(MouseEvent.GetLastScreenSpacePosition());
 			float NewValue = RangeToScreen.LocalXToInput(CursorPos.X);
 
-			if ( TimeSliderArgs.Settings->GetIsSnapEnabled() )
+			if ( TimeSliderArgs.Settings->GetSnapPlayTimeToInterval() )
 			{
 				NewValue = TimeSliderArgs.Settings->SnapTimeToInterval(NewValue);
 			}
@@ -542,7 +542,7 @@ FReply FSequencerTimeSliderController::OnMouseMove( TSharedRef<SWidget> WidgetOw
 
 				if (MouseDragType == DRAG_SCRUBBING_TIME)
 				{
-					if ( TimeSliderArgs.Settings->GetIsSnapEnabled() )
+					if ( TimeSliderArgs.Settings->GetSnapPlayTimeToInterval() )
 					{
 						NewValue = TimeSliderArgs.Settings->SnapTimeToInterval(NewValue);
 					}
