Sequencer: Set time or frame popup for keys.

diff --git a/Engine/Source/Editor/Sequencer/Private/SSection.cpp b/Engine/Source/Editor/Sequencer/Private/SSection.cpp
--- a/Engine/Source/Editor/Sequencer/Private/SSection.cpp
+++ b/Engine/Source/Editor/Sequencer/Private/SSection.cpp
@@ -564,6 +564,17 @@ TSharedPtr<SWidget> SSection::OnSummonContextMenu( const FGeometry& MyGeometry,
 
 		MenuBuilder.BeginSection("SequencerKeys", NSLOCTEXT("Sequencer", "KeysMenu", "Keys"));
 		{
+			const bool bUseFrames = SequencerSnapValues::IsTimeSnapIntervalFrameRate(Sequencer.GetSettings()->GetTimeSnapInterval());
+
+			MenuBuilder.AddMenuEntry(
+				bUseFrames ? NSLOCTEXT("Sequencer", "SetKeyFrame", "Set Key Frame") : NSLOCTEXT("Sequencer", "SetKeyTime", "Set key time"),
+				bUseFrames ? NSLOCTEXT("Sequencer", "SetKeyFrameTooltip", "Set key frame") : NSLOCTEXT("Sequencer", "SetKeyTimeTooltip", "Set key time"),
+				FSlateIcon(),
+				FUIAction(
+					FExecuteAction::CreateSP(&Sequencer, &FSequencer::SetKeyTime, bUseFrames),
+					FCanExecuteAction::CreateSP(&Sequencer, &FSequencer::CanSetKeyTime))
+			);
+
 			MenuBuilder.AddMenuEntry(
 				NSLOCTEXT("Sequencer", "SnapToFrame", "Snap to Frame"),
 				NSLOCTEXT("Sequencer", "SnapToFrameToolTip", "Snap selected keys to frame"),
diff --git a/Engine/Source/Editor/Sequencer/Private/Sequencer.cpp b/Engine/Source/Editor/Sequencer/Private/Sequencer.cpp
--- a/Engine/Source/Editor/Sequencer/Private/Sequencer.cpp
+++ b/Engine/Source/Editor/Sequencer/Private/Sequencer.cpp
@@ -41,6 +41,7 @@
 #include "Engine/Selection.h"
 #include "LevelEditor.h"
 #include "IMenu.h"
+#include "STextEntryPopup.h"
 
 #include "MovieSceneCaptureModule.h"
 #include "AutomatedLevelSequenceCapture.h"
@@ -2137,6 +2138,90 @@ void FSequencer::SetKey()
 	}
 }
 
+bool FSequencer::CanSetKeyTime() const
+{
+	return Selection.GetSelectedKeys().Num() > 0;
+}
+
+void FSequencer::SetKeyTime(const bool bUseFrames)
+{
+	TArray<FSelectedKey> SelectedKeysArray = Selection.GetSelectedKeys().Array();
+
+	float KeyTime = 0.f;
+	for ( const FSelectedKey& Key : SelectedKeysArray )
+	{
+		if (Key.IsValid())
+		{
+			KeyTime = Key.KeyArea->GetKeyTime(Key.KeyHandle.GetValue());
+			break;
+		}
+	}
+
+	float FrameRate = 1.0f / Settings->GetTimeSnapInterval();
+	
+	GenericTextEntryModeless(
+		bUseFrames ? NSLOCTEXT("Sequencer.Popups", "SetKeyFramePopup", "New Frame") : NSLOCTEXT("Sequencer.Popups", "SetKeyTimePopup", "New Time"),
+		bUseFrames ? FText::AsNumber( SequencerHelpers::TimeToFrame( KeyTime, FrameRate )) : FText::AsNumber( KeyTime ),
+		FOnTextCommitted::CreateSP(this, &FSequencer::OnSetKeyTimeTextCommitted, bUseFrames)
+		);
+}
+
+void FSequencer::OnSetKeyTimeTextCommitted(const FText& InText, ETextCommit::Type CommitInfo, const bool bUseFrames)
+{
+	CloseEntryPopupMenu();
+	if (CommitInfo == ETextCommit::OnEnter)
+	{
+		float FrameRate = 1.0f / Settings->GetTimeSnapInterval();
+		double dNewTime = bUseFrames ? SequencerHelpers::FrameToTime(FCString::Atod(*InText.ToString()), FrameRate) : FCString::Atod(*InText.ToString());
+		const bool bIsNumber = InText.IsNumeric(); 
+		if(!bIsNumber)
+			return;
+
+		const float NewKeyTime = (float)dNewTime;
+
+		FScopedTransaction SetKeyTimeTransaction(NSLOCTEXT("Sequencer", "SetKeyTime_Transaction", "Set Key Time"));
+
+		TArray<FSelectedKey> SelectedKeysArray = Selection.GetSelectedKeys().Array();
+	
+		for ( const FSelectedKey& Key : SelectedKeysArray )
+		{
+			if (Key.IsValid())
+			{
+				Key.Section->Modify();
+				Key.KeyArea->SetKeyTime(Key.KeyHandle.GetValue(), NewKeyTime);
+			}
+		}
+	}
+}
+
+void FSequencer::GenericTextEntryModeless(const FText& DialogText, const FText& DefaultText, FOnTextCommitted OnTextComitted)
+{
+	TSharedRef<STextEntryPopup> TextEntryPopup = 
+		SNew(STextEntryPopup)
+		.Label(DialogText)
+		.DefaultText(DefaultText)
+		.OnTextCommitted(OnTextComitted)
+		.ClearKeyboardFocusOnCommit(false)
+		.SelectAllTextWhenFocused(true)
+		.MaxWidth(1024.0f);
+
+	EntryPopupMenu = FSlateApplication::Get().PushMenu(
+		ToolkitHost.Pin()->GetParentWidget(),
+		FWidgetPath(),
+		TextEntryPopup,
+		FSlateApplication::Get().GetCursorPos(),
+		FPopupTransitionEffect(FPopupTransitionEffect::TypeInPopup)
+		);
+}
+
+void FSequencer::CloseEntryPopupMenu()
+{
+	if (EntryPopupMenu.IsValid())
+	{
+		EntryPopupMenu.Pin()->Dismiss();
+	}
+}
+
 void FSequencer::TrimSection(bool bTrimLeft)
 {
 	FScopedTransaction TrimSectionTransaction( NSLOCTEXT("Sequencer", "TrimSection_Transaction", "Trim Section") );
diff --git a/Engine/Source/Editor/Sequencer/Private/Sequencer.h b/Engine/Source/Editor/Sequencer/Private/Sequencer.h
--- a/Engine/Source/Editor/Sequencer/Private/Sequencer.h
+++ b/Engine/Source/Editor/Sequencer/Private/Sequencer.h
@@ -280,6 +280,11 @@ public:
 	void DeleteNode(TSharedRef<FSequencerDisplayNode> NodeToBeDeleted);
 	void DeleteSelectedNodes();
 
+	/** Called when a user executes the set key time for selected keys */
+	bool CanSetKeyTime() const;
+	void SetKeyTime(const bool bUseFrames);
+	void OnSetKeyTimeTextCommitted(const FText& InText, ETextCommit::Type CommitInfo, const bool bUseFrames);
+
 public:
 
 	/** Access the currently enabled edit tool */
@@ -422,6 +427,12 @@ protected:
 	/** Manually sets a key for the selected objects at the current time */
 	void SetKey();
 
+	/** Modeless Version of the String Entry Box */
+	void GenericTextEntryModeless(const FText& DialogText, const FText& DefaultText, FOnTextCommitted OnTextComitted);
+	
+	/** Closes the popup created by GenericTextEntryModeless*/
+	void CloseEntryPopupMenu();
+
 	/** Trim a section to the left or right */
 	void TrimSection(bool bTrimLeft);
 
@@ -531,6 +542,9 @@ private:
 	/** True if this sequencer is being edited within the level editor */
 	bool bIsEditingWithinLevelEditor;
 
+	/** Generic Popup Entry */
+	TWeakPtr<class IMenu> EntryPopupMenu;
+
 	/** Stores a dirty bit for whether the sequencer tree (and other UI bits) may need to be refreshed.  We
 	    do this simply to avoid refreshing the UI more than once per frame. (e.g. during live recording where
 		the MovieScene data can change many times per frame.) */
