Sequencer: Fix case.

diff --git a/Engine/Source/Editor/Sequencer/Private/SSection.cpp b/Engine/Source/Editor/Sequencer/Private/SSection.cpp
--- a/Engine/Source/Editor/Sequencer/Private/SSection.cpp
+++ b/Engine/Source/Editor/Sequencer/Private/SSection.cpp
@@ -567,7 +567,7 @@ TSharedPtr<SWidget> SSection::OnSummonContextMenu( const FGeometry& MyGeometry,
 			const bool bUseFrames = SequencerSnapValues::IsTimeSnapIntervalFrameRate(Sequencer.GetSettings()->GetTimeSnapInterval());
 
 			MenuBuilder.AddMenuEntry(
-				bUseFrames ? NSLOCTEXT("Sequencer", "SetKeyFrame", "Set Key Frame") : NSLOCTEXT("Sequencer", "SetKeyTime", "Set key time"),
+				bUseFrames ? NSLOCTEXT("Sequencer", "SetKeyFrame", "Set Key Frame") : NSLOCTEXT("Sequencer", "SetKeyTime", "Set Key Time"),
 				bUseFrames ? NSLOCTEXT("Sequencer", "SetKeyFrameTooltip", "Set key frame") : NSLOCTEXT("Sequencer", "SetKeyTimeTooltip", "Set key time"),
 				FSlateIcon(),
 				FUIAction(
