Sequencer marquee selection improvements

Marquee selection now intelligently selects keys and sections depending on the context. If a section is within the marquee it is eligible for selection, unless there are keys from that section also in the marquee that would change selection state as a result of the operation.

Selections are also now previewed in real-time

diff --git a/Engine/Source/Editor/EditorStyle/Private/SlateEditorStyle.cpp b/Engine/Source/Editor/EditorStyle/Private/SlateEditorStyle.cpp
--- a/Engine/Source/Editor/EditorStyle/Private/SlateEditorStyle.cpp
+++ b/Engine/Source/Editor/EditorStyle/Private/SlateEditorStyle.cpp
@@ -1616,8 +1616,7 @@ void FSlateEditorStyle::FStyle::SetupGeneralStyles()
 		Set( "Sequencer.ToggleAutoScroll", new IMAGE_BRUSH( "Icons/icon_Sequencer_ToggleAutoScroll_40x", Icon40x40 ) );
 		Set( "Sequencer.ToggleAutoScroll.Small", new IMAGE_BRUSH( "Icons/icon_Sequencer_ToggleAutoScroll_16x", Icon16x16 ) );
 		Set( "Sequencer.MoveTool.Small", new IMAGE_BRUSH( "Icons/SequencerIcons/icon_Sequencer_Move_24x", Icon16x16 ) );
-		Set( "Sequencer.MarqueeTool_Keys.Small", new IMAGE_BRUSH( "Sequencer/MarqueeTool_Keys_16px", Icon16x16 ) ); //@todo
-		Set( "Sequencer.MarqueeTool_Sections.Small", new IMAGE_BRUSH( "Sequencer/MarqueeTool_Sections_16px", Icon16x16 ) ); //@todo
+		Set( "Sequencer.MarqueeTool.Small", new IMAGE_BRUSH( "Icons/SequencerIcons/icon_Sequencer_Marquee_24x", Icon16x16 ) );
 		Set( "Sequencer.RenderMovie.Small", new IMAGE_BRUSH( "Sequencer/RenderMovie_16px", Icon16x16 ) );
 		Set( "Sequencer.LockCamera", new IMAGE_BRUSH( "Icons/icon_ToggleActorPilotCameraView_16x", Icon16x16 ) );
 		Set( "Sequencer.UnlockCamera", new IMAGE_BRUSH( "Icons/icon_ToggleActorPilotCameraView_16x", Icon16x16, FLinearColor(1.f, 1.f, 1.f, 0.5f) ) );
@@ -1634,6 +1633,10 @@ void FSlateEditorStyle::FStyle::SetupGeneralStyles()
 		Set( "Sequencer.Tracks.Slomo", new IMAGE_BRUSH("Icons/SequencerIcons/icon_Sequencer_Slomo_16x", Icon16x16));
 		Set( "Sequencer.Tracks.Fade", new IMAGE_BRUSH("Icons/SequencerIcons/icon_Sequencer_Fade_16x", Icon16x16));
 
+		Set( "Sequencer.CursorDecorator_Marquee", new IMAGE_BRUSH( "Sequencer/CursorDecorator_Marquee", Icon16x16));
+		Set( "Sequencer.CursorDecorator_MarqueeAdd", new IMAGE_BRUSH( "Sequencer/CursorDecorator_MarqueeAdd", Icon16x16));
+		Set( "Sequencer.CursorDecorator_MarqueeSubtract", new IMAGE_BRUSH( "Sequencer/CursorDecorator_MarqueeSubtract", Icon16x16));
+
 		const FButtonStyle DetailsKeyButton = FButtonStyle(NoBorder)
 			.SetNormal( IMAGE_BRUSH("Sequencer/AddKey_Details", FVector2D(11,11) )  )
 			.SetHovered( IMAGE_BRUSH("Sequencer/AddKey_Details", FVector2D(11,11), SelectionColor ) )
diff --git a/Engine/Source/Editor/Sequencer/Private/SSection.cpp b/Engine/Source/Editor/Sequencer/Private/SSection.cpp
--- a/Engine/Source/Editor/Sequencer/Private/SSection.cpp
+++ b/Engine/Source/Editor/Sequencer/Private/SSection.cpp
@@ -255,9 +255,11 @@ int32 SSection::OnPaint( const FPaintArgs& Args, const FGeometry& AllottedGeomet
 	// Ask the interface to draw the section
 	int32 PostSectionLayer = SectionInterface->OnPaintSection( SectionGeometry, SectionClipRect, OutDrawElements, LayerId, bEnabled );
 
-	const bool bDisplaySectionHandles = true;
+	FLinearColor SelectionColor = FEditorStyle::GetSlateColor(SequencerSectionConstants::SelectionColorName).GetColor(FWidgetStyle());
+
+	DrawSectionHandles(AllottedGeometry, MyClippingRect, OutDrawElements, PostSectionLayer, DrawEffects, SelectionColor);
 
-	DrawSectionHandlesAndSelection(AllottedGeometry, MyClippingRect, OutDrawElements, PostSectionLayer, bDisplaySectionHandles, bEnabled );
+	DrawSelectionBorder(AllottedGeometry, MyClippingRect, OutDrawElements, PostSectionLayer, DrawEffects, SelectionColor);
 
 	PaintKeys( SectionGeometry, MyClippingRect, OutDrawElements, PostSectionLayer, InWidgetStyle, bEnabled );
 
@@ -311,13 +313,22 @@ void SSection::PaintKeys( const FGeometry& AllottedGeometry, const FSlateRect& M
 	static const FName SelectionColorPressedName("SelectionColor_Pressed");
 
 	const FLinearColor PressedKeyColor = FEditorStyle::GetSlateColor(SelectionColorPressedName).GetColor( InWidgetStyle );
-	const FLinearColor SelectedKeyColor = FEditorStyle::GetSlateColor(SelectionColorName).GetColor( InWidgetStyle );
-	const FLinearColor SelectedInactiveColor = FEditorStyle::GetSlateColor(SelectionInactiveColorName).GetColor( InWidgetStyle )
-		* FLinearColor(.25, .25, .25, 1);  // Make the color a little darker since it's not very visible next to white keyframes.
-
+	FLinearColor SelectionColor = FEditorStyle::GetSlateColor(SelectionColorName).GetColor( InWidgetStyle );
+	FLinearColor SelectedKeyColor = SelectionColor;
 	// @todo Sequencer temp color, make hovered brighter than selected.
 	FLinearColor HoveredKeyColor = SelectedKeyColor * FLinearColor(1.5,1.5,1.5,1.0f);
 
+
+	auto& Selection = Sequencer.GetSelection();
+	auto& SelectionPreview = Sequencer.GetSelectionPreview();
+
+	if (Selection.GetActiveSelection() != FSequencerSelection::EActiveSelection::KeyAndSection)
+	{
+		// Selected key color is different when the active selection is not keys
+		SelectedKeyColor = FEditorStyle::GetSlateColor(SelectionInactiveColorName).GetColor( InWidgetStyle )
+			* FLinearColor(.25, .25, .25, 1);  // Make the color a little darker since it's not very visible next to white keyframes.
+	}
+
 	// Draw all keys in each key area
 	for (const FKeyAreaLayoutElement& Element : Layout->GetElements())
 	{
@@ -325,161 +336,182 @@ void SSection::PaintKeys( const FGeometry& AllottedGeometry, const FSlateRect& M
 		TSharedPtr<IKeyArea> KeyArea = Element.GetKeyArea();
 		TArray<FKeyHandle> KeyHandles = KeyArea->GetUnsortedKeyHandles();
 
+		if (!KeyHandles.Num())
+		{
+			continue;
+		}
+
 		FGeometry KeyAreaGeometry = GetKeyAreaGeometry( Element, AllottedGeometry );
 
-		FTimeToPixel TimeToPixelConverter = SectionObject.IsInfinite() ? 			
-			FTimeToPixel( ParentGeometry, GetSequencer().GetViewRange()) : 
+		FTimeToPixel TimeToPixelConverter = SectionObject.IsInfinite() ?
+			FTimeToPixel( ParentGeometry, GetSequencer().GetViewRange()) :
 			FTimeToPixel( KeyAreaGeometry, TRange<float>( SectionObject.GetStartTime(), SectionObject.GetEndTime() ) );
 
 		// Draw a box for the key area 
 		// @todo Sequencer - Allow the IKeyArea to do this
-		if (KeyHandles.Num())
-		{
-			FSlateDrawElement::MakeBox( 
-				OutDrawElements,
-				LayerId,
-				KeyAreaGeometry.ToPaintGeometry(),
-				BackgroundBrush,
-				MyClippingRect,
-				DrawEffects,
-				FLinearColor( .1f, .1f, .1f, 0.7f ) ); 
-		}
+		FSlateDrawElement::MakeBox( 
+			OutDrawElements,
+			LayerId,
+			KeyAreaGeometry.ToPaintGeometry(),
+			BackgroundBrush,
+			MyClippingRect,
+			DrawEffects,
+			FLinearColor( .1f, .1f, .1f, 0.7f ) ); 
 
-		int32 KeyLayer = LayerId + 1;
+		const int32 KeyLayer = LayerId + 1;
 
-		for( int32 KeyIndex = 0; KeyIndex < KeyHandles.Num(); ++KeyIndex )
+		for (const FKeyHandle& KeyHandle : KeyHandles)
 		{
-			FKeyHandle KeyHandle = KeyHandles[KeyIndex];
 			float KeyTime = KeyArea->GetKeyTime(KeyHandle);
 
 			// Omit keys which would not be visible
-			if( SectionObject.IsTimeWithinSection( KeyTime ) )
+			if( !SectionObject.IsTimeWithinSection( KeyTime ) )
 			{
-				const FSlateBrush* BrushToUse = SectionInterface->GetKeyBrush(KeyHandle);
-				if(BrushToUse == nullptr)
-				{
-					BrushToUse = DefaultKeyBrush;
-				}
-				FLinearColor KeyColor( 1.0f, 1.0f, 1.0f, 1.0f );
-				FLinearColor KeyTint(1.f, 1.f, 1.f, 1.f);
+				continue;
+			}
 
-				if (Element.GetType() == FKeyAreaLayoutElement::Group)
-				{
-					auto Group = StaticCastSharedPtr<FGroupedKeyArea>(KeyArea);
-					KeyTint = Group->GetKeyTint(KeyHandle);
-					BrushToUse = Group->GetBrush(KeyHandle);
-				}
+			const FSlateBrush* BrushToUse = SectionInterface->GetKeyBrush(KeyHandle);
+			if(BrushToUse == nullptr)
+			{
+				BrushToUse = DefaultKeyBrush;
+			}
+			FLinearColor KeyColor( 1.0f, 1.0f, 1.0f, 1.0f );
+			FLinearColor KeyTint(1.f, 1.f, 1.f, 1.f);
 
-				// Where to start drawing the key (relative to the section)
-				float KeyPosition =  TimeToPixelConverter.TimeToPixel( KeyTime );
+			if (Element.GetType() == FKeyAreaLayoutElement::Group)
+			{
+				auto Group = StaticCastSharedPtr<FGroupedKeyArea>(KeyArea);
+				KeyTint = Group->GetKeyTint(KeyHandle);
+				BrushToUse = Group->GetBrush(KeyHandle);
+			}
 
-				FSelectedKey TestKey( SectionObject, KeyArea, KeyHandle );
+			// Where to start drawing the key (relative to the section)
+			float KeyPosition =  TimeToPixelConverter.TimeToPixel( KeyTime );
 
-				bool bSelected = Sequencer.GetSelection().IsSelected( TestKey );
-				bool bActive = Sequencer.GetSelection().GetActiveSelection() == FSequencerSelection::EActiveSelection::KeyAndSection;
+			FSelectedKey TestKey( SectionObject, KeyArea, KeyHandle );
 
-				if( TestKey == PressedKey )
-				{
-					KeyColor = PressedKeyColor;
-				}
-				else if( TestKey == HoveredKey )
-				{
-					KeyColor = HoveredKeyColor;
-				}
-				else if( bSelected )
-				{
-					if (bActive)
-					{
-						KeyColor = SelectedKeyColor;
-					}
-					else
-					{
-						KeyColor = SelectedInactiveColor;
-					}
-				}
+			bool bSelected = Selection.IsSelected( TestKey );
+			ESelectionPreviewState SelectionPreviewState = SelectionPreview.GetSelectionState( TestKey );
 
-				KeyColor *= KeyTint;
-
-				// Draw the key
-				FSlateDrawElement::MakeBox(
-					OutDrawElements,
-					// always draw selected keys on top of other keys
-					bSelected ? KeyLayer+1 : KeyLayer,
-					// Center the key along Y.  Ensure the middle of the key is at the actual key time
-					KeyAreaGeometry.ToPaintGeometry( FVector2D( KeyPosition - FMath::CeilToFloat(SequencerSectionConstants::KeySize.X/2.0f), ((KeyAreaGeometry.Size.Y*.5f)-(SequencerSectionConstants::KeySize.Y*.5f)) ), SequencerSectionConstants::KeySize ),
-					BrushToUse,
-					MyClippingRect,
-					DrawEffects,
-					KeyColor
-					);
+			if( TestKey == PressedKey )
+			{
+				KeyColor = PressedKeyColor;
+			}
+			else if( TestKey == HoveredKey )
+			{
+				KeyColor = HoveredKeyColor;
+			}
+			else if( SelectionPreviewState == ESelectionPreviewState::Selected )
+			{
+				FLinearColor PreviewSelectionColor = SelectionColor.LinearRGBToHSV();
+				PreviewSelectionColor.R += 0.1f; // +10% hue
+				PreviewSelectionColor.G = 0.6f; // 60% saturation
+				KeyColor = PreviewSelectionColor.HSVToLinearRGB();
 			}
+			else if( SelectionPreviewState == ESelectionPreviewState::NotSelected )
+			{
+				// Default white selection color
+			}
+			else if( bSelected )
+			{
+				KeyColor = SelectedKeyColor;
+			}
+
+			KeyColor *= KeyTint;
+
+			// Draw the key
+			FSlateDrawElement::MakeBox(
+				OutDrawElements,
+				// always draw selected keys on top of other keys
+				bSelected ? KeyLayer+1 : KeyLayer,
+				// Center the key along Y.  Ensure the middle of the key is at the actual key time
+				KeyAreaGeometry.ToPaintGeometry( FVector2D( KeyPosition - FMath::CeilToFloat(SequencerSectionConstants::KeySize.X/2.0f), ((KeyAreaGeometry.Size.Y*.5f)-(SequencerSectionConstants::KeySize.Y*.5f)) ), SequencerSectionConstants::KeySize ),
+				BrushToUse,
+				MyClippingRect,
+				DrawEffects,
+				KeyColor
+				);
 		}
 	}
 }
 
-void SSection::DrawSectionHandlesAndSelection( const FGeometry& AllottedGeometry, const FSlateRect& MyClippingRect, FSlateWindowElementList& OutDrawElements, int32 LayerId, bool bDisplaySectionHandles, bool bParentEnabled ) const
+void SSection::DrawSectionHandles( const FGeometry& AllottedGeometry, const FSlateRect& MyClippingRect, FSlateWindowElementList& OutDrawElements, int32 LayerId, ESlateDrawEffect::Type DrawEffects, FLinearColor SelectionColor ) const
 {
-	const ESlateDrawEffect::Type DrawEffects = bParentEnabled ? ESlateDrawEffect::None : ESlateDrawEffect::DisabledEffect;
-
-	UMovieSceneSection* SectionObject = SectionInterface->GetSectionObject();
+	const FSlateBrush* LeftGripBrush = FEditorStyle::GetBrush(SectionInterface->GetSectionGripLeftBrushName());
+	const FSlateBrush* RightGripBrush = FEditorStyle::GetBrush(SectionInterface->GetSectionGripRightBrushName());
+
+	// Left Grip
+	FSlateDrawElement::MakeBox
+	(
+		OutDrawElements,
+		LayerId,
+		// Center the key along Y.  Ensure the middle of the key is at the actual key time
+		AllottedGeometry.ToPaintGeometry(FVector2D(0.0f, 0.0f), FVector2D(SectionInterface->GetSectionGripSize(), AllottedGeometry.GetDrawSize().Y)),
+		LeftGripBrush,
+		MyClippingRect,
+		DrawEffects,
+		(bLeftEdgePressed || bLeftEdgeHovered) ? SelectionColor : LeftGripBrush->GetTint(FWidgetStyle())
+	);
+	
+	// Right Grip
+	FSlateDrawElement::MakeBox
+	(
+		OutDrawElements,
+		LayerId,
+		// Center the key along Y.  Ensure the middle of the key is at the actual key time
+		AllottedGeometry.ToPaintGeometry(FVector2D(AllottedGeometry.Size.X-SectionInterface->GetSectionGripSize(), 0.0f), FVector2D(SectionInterface->GetSectionGripSize(), AllottedGeometry.GetDrawSize().Y)),
+		RightGripBrush,
+		MyClippingRect,
+		DrawEffects,
+		(bRightEdgePressed || bRightEdgeHovered) ? SelectionColor : RightGripBrush->GetTint(FWidgetStyle())
+	);
+}
 
+void SSection::DrawSelectionBorder( const FGeometry& AllottedGeometry, const FSlateRect& MyClippingRect, FSlateWindowElementList& OutDrawElements, int32 LayerId, ESlateDrawEffect::Type DrawEffects, FLinearColor SelectionColor ) const
+{
 	FSequencerSelection& Selection = ParentSectionArea->GetSequencer().GetSelection();
-	const bool bSelected = Selection.IsSelected(SectionObject);
-	const bool bActive = Selection.GetActiveSelection() == FSequencerSelection::EActiveSelection::KeyAndSection;
-
-	FLinearColor SelectionColor = FEditorStyle::GetSlateColor(SequencerSectionConstants::SelectionColorName).GetColor(FWidgetStyle());
-	FLinearColor SelectionInactiveColor = FEditorStyle::GetSlateColor(SequencerSectionConstants::SelectionInactiveColorName).GetColor(FWidgetStyle());
-	FLinearColor TransparentSelectionColor = SelectionColor;
+	FSequencerSelectionPreview& SelectionPreview = ParentSectionArea->GetSequencer().GetSelectionPreview();
+	
+	UMovieSceneSection* SectionObject = SectionInterface->GetSectionObject();
+	ESelectionPreviewState SelectionPreviewState = SelectionPreview.GetSelectionState( SectionObject );
 
-	if( bDisplaySectionHandles )
+	if (SelectionPreviewState == ESelectionPreviewState::NotSelected)
 	{
-		const FSlateBrush* LeftGripBrush = FEditorStyle::GetBrush(SectionInterface->GetSectionGripLeftBrushName());
-		const FSlateBrush* RightGripBrush = FEditorStyle::GetBrush(SectionInterface->GetSectionGripRightBrushName());
-
-		// Left Grip
-		FSlateDrawElement::MakeBox
-		(
-			OutDrawElements,
-			LayerId,
-			// Center the key along Y.  Ensure the middle of the key is at the actual key time
-			AllottedGeometry.ToPaintGeometry(FVector2D(0.0f, 0.0f), FVector2D(SectionInterface->GetSectionGripSize(), AllottedGeometry.GetDrawSize().Y)),
-			LeftGripBrush,
-			MyClippingRect,
-			DrawEffects,
-			(bLeftEdgePressed || bLeftEdgeHovered) ? TransparentSelectionColor : LeftGripBrush->GetTint(FWidgetStyle())
-		);
-		
-		// Right Grip
-		FSlateDrawElement::MakeBox
-		(
-			OutDrawElements,
-			LayerId,
-			// Center the key along Y.  Ensure the middle of the key is at the actual key time
-			AllottedGeometry.ToPaintGeometry(FVector2D(AllottedGeometry.Size.X-SectionInterface->GetSectionGripSize(), 0.0f), FVector2D(SectionInterface->GetSectionGripSize(), AllottedGeometry.GetDrawSize().Y)),
-			RightGripBrush,
-			MyClippingRect,
-			DrawEffects,
-			(bRightEdgePressed || bRightEdgeHovered) ? TransparentSelectionColor : RightGripBrush->GetTint(FWidgetStyle())
-		);
+		// Explicitly not selected in the preview selection
+		return;
 	}
-
-
-	// draw selection box
-	if(bSelected)
+	else if (SelectionPreviewState == ESelectionPreviewState::Undefined && !Selection.IsSelected(SectionObject))
 	{
-		static const FName SelectionBorder("Sequencer.Section.SelectionBorder");
-
-		FSlateDrawElement::MakeBox(
-			OutDrawElements,
-			LayerId+1,
-			AllottedGeometry.ToPaintGeometry(),
-			FEditorStyle::GetBrush(SelectionBorder),
-			MyClippingRect,
-			DrawEffects,
-			bActive ? SelectionColor : SelectionInactiveColor
-			);
+		// No preview selection for this section, and it's not selected
+		return;
+	}
+	
+	// Use a muted selection color for selection previews
+	if( SelectionPreviewState == ESelectionPreviewState::Selected )
+	{
+		SelectionColor = SelectionColor.LinearRGBToHSV();
+		SelectionColor.R += 0.1f; // +10% hue
+		SelectionColor.G = 0.6f; // 60% saturation
+		SelectionColor = SelectionColor.HSVToLinearRGB();
+	}
+	else if (Selection.GetActiveSelection() != FSequencerSelection::EActiveSelection::KeyAndSection)
+	{
+		// Use an inactive selection color for existing selections that are not active
+		SelectionColor = FEditorStyle::GetSlateColor(SequencerSectionConstants::SelectionInactiveColorName).GetColor(FWidgetStyle());
 	}
 
+	// draw selection box
+	static const FName SelectionBorder("Sequencer.Section.SelectionBorder");
+
+	FSlateDrawElement::MakeBox(
+		OutDrawElements,
+		LayerId+1,
+		AllottedGeometry.ToPaintGeometry(),
+		FEditorStyle::GetBrush(SelectionBorder),
+		MyClippingRect,
+		DrawEffects,
+		SelectionColor
+		);
 }
 
 TSharedPtr<SWidget> SSection::OnSummonContextMenu( const FGeometry& MyGeometry, const FPointerEvent& MouseEvent )
diff --git a/Engine/Source/Editor/Sequencer/Private/SSection.h b/Engine/Source/Editor/Sequencer/Private/SSection.h
--- a/Engine/Source/Editor/Sequencer/Private/SSection.h
+++ b/Engine/Source/Editor/Sequencer/Private/SSection.h
@@ -86,9 +86,14 @@ private:
 	void PaintKeys( const FGeometry& AllottedGeometry, const FSlateRect& MyClippingRect, FSlateWindowElementList& OutDrawElements, int32 LayerId, const FWidgetStyle& InWidgetStyle, bool bParentEnabled ) const;
 
 	/**
-	 * Draw the borders around each section.  The borders are the areas which can be used to expand or collapse the size of the section.
+	 * Draw the section resize handles.
 	 */
-	void DrawSectionHandlesAndSelection( const FGeometry& AllottedGeometry, const FSlateRect& MyClippingRect, FSlateWindowElementList& OutDrawElements, int32 LayerId, bool bDisplaySectionHandles, bool bParentEnabled ) const;
+	void DrawSectionHandles( const FGeometry& AllottedGeometry, const FSlateRect& MyClippingRect, FSlateWindowElementList& OutDrawElements, int32 LayerId, ESlateDrawEffect::Type DrawEffects, FLinearColor SelectionColor ) const;
+
+	/**
+	 * Draw a box representing this section's selection
+	 */
+	void DrawSelectionBorder( const FGeometry& AllottedGeometry, const FSlateRect& MyClippingRect, FSlateWindowElementList& OutDrawElements, int32 LayerId, ESlateDrawEffect::Type DrawEffects, FLinearColor SelectionColor ) const;
 
 	/** Summons a context menu over the associated section */
 	TSharedPtr<SWidget> OnSummonContextMenu( const FGeometry& MyGeometry, const FPointerEvent& MouseEvent );
diff --git a/Engine/Source/Editor/Sequencer/Private/SSequencer.cpp b/Engine/Source/Editor/Sequencer/Private/SSequencer.cpp
--- a/Engine/Source/Editor/Sequencer/Private/SSequencer.cpp
+++ b/Engine/Source/Editor/Sequencer/Private/SSequencer.cpp
@@ -489,21 +489,7 @@ void SSequencer::BindCommands(TSharedRef<FUICommandList> SequencerCommandBinding
 	SequencerCommandBindings->MapAction(
 		FSequencerCommands::Get().MarqueeTool,
 		FExecuteAction::CreateLambda( [&]{
-
-			if (IsEditToolEnabled("Selection"))
-			{
-				// If the selection tool is already enabled, we toggle the selection mode
-				auto* SelectionTool = static_cast<FSequencerEditTool_Selection*>(EditTool.Get());
-
-				ESequencerSelectionMode CurrentMode = SelectionTool->GetSelectionMode();
-				SelectionTool->SetSelectionMode(CurrentMode == ESequencerSelectionMode::Keys ? ESequencerSelectionMode::Sections : ESequencerSelectionMode::Keys);
-			}
-			else
-			{
-				// Otherwise, switch the tool
-				EditTool.Reset( new FSequencerEditTool_Selection(Sequencer.Pin(), SharedThis(this)) );
-			}
-
+			EditTool.Reset( new FSequencerEditTool_Selection(Sequencer.Pin(), SharedThis(this)) );
 		} ),
 		FCanExecuteAction::CreateLambda( []{ return true; } ),
 		FIsActionChecked::CreateSP( this, &SSequencer::IsEditToolEnabled, FName("Selection") )
@@ -677,43 +663,7 @@ TSharedRef<SWidget> SSequencer::MakeToolBar()
 		ToolBarBuilder.SetLabelVisibility( EVisibility::Collapsed );
 
 		ToolBarBuilder.AddToolBarButton( FSequencerCommands::Get().MoveTool );
-
-		// The marquee selection button is a little more complicated as toggling it when enabled causes it to toggle between key/section selection
-		TAttribute<FSlateIcon> MarqueeIcon;
-		MarqueeIcon.Bind(TAttribute<FSlateIcon>::FGetter::CreateLambda([]{
-			static FSlateIcon MarqueeTool_Keys(FEditorStyle::GetStyleSetName(), "Sequencer.MarqueeTool_Keys");
-			static FSlateIcon MarqueeTool_Sections(FEditorStyle::GetStyleSetName(), "Sequencer.MarqueeTool_Sections");
-
-			return FSequencerEditTool_Selection::GetGlobalSelectionMode() == ESequencerSelectionMode::Keys ? MarqueeTool_Keys : MarqueeTool_Sections;
-		}));
-
-		TAttribute<FText> MarqueeTooltip;
-		MarqueeTooltip.Bind(TAttribute<FText>::FGetter::CreateLambda([&]{
-			if (FSequencerEditTool_Selection::GetGlobalSelectionMode() == ESequencerSelectionMode::Keys)
-			{
-				if (IsEditToolEnabled("Selection"))
-				{
-					return LOCTEXT("MarqueeTooltip_Enabled_Keys", "Change marquee selection to operate on keys rather than sections");
-				}
-				else
-				{
-					return LOCTEXT("MarqueeTooltip_Key", "Activate marquee key selection tool");
-				}
-			}
-			else
-			{
-				if (IsEditToolEnabled("Selection"))
-				{
-					return LOCTEXT("MarqueeTooltip_Enabled_Sections", "Change marquee selection to operate on sections rather than keys");
-				}
-				else
-				{
-					return LOCTEXT("MarqueeTooltip_Section", "Activate marquee section selection tool");
-				}
-			}
-		}));
-
-		ToolBarBuilder.AddToolBarButton( FSequencerCommands::Get().MarqueeTool, NAME_None, TAttribute<FText>(), MarqueeTooltip, MarqueeIcon );
+		ToolBarBuilder.AddToolBarButton( FSequencerCommands::Get().MarqueeTool );
 	}
 	ToolBarBuilder.EndSection();
 
diff --git a/Engine/Source/Editor/Sequencer/Private/SSequencerTrackArea.cpp b/Engine/Source/Editor/Sequencer/Private/SSequencerTrackArea.cpp
--- a/Engine/Source/Editor/Sequencer/Private/SSequencerTrackArea.cpp
+++ b/Engine/Source/Editor/Sequencer/Private/SSequencerTrackArea.cpp
@@ -203,6 +203,24 @@ FReply SSequencerTrackArea::OnMouseWheel( const FGeometry& MyGeometry, const FPo
 	return FReply::Handled();
 }
 
+void SSequencerTrackArea::OnMouseEnter(const FGeometry& MyGeometry, const FPointerEvent& MouseEvent)
+{
+	auto SequencerPin = SequencerWidget.Pin();
+	if (SequencerPin.IsValid())
+	{
+		SequencerPin->GetEditTool().OnMouseEnter(*this, MyGeometry, MouseEvent);
+	}
+}
+
+void SSequencerTrackArea::OnMouseLeave(const FPointerEvent& MouseEvent)
+{
+	auto SequencerPin = SequencerWidget.Pin();
+	if (SequencerPin.IsValid())
+	{
+		SequencerPin->GetEditTool().OnMouseLeave(*this, MouseEvent);
+	}
+}
+
 void SSequencerTrackArea::OnMouseCaptureLost()
 {
 	auto SequencerPin = SequencerWidget.Pin();
diff --git a/Engine/Source/Editor/Sequencer/Private/SSequencerTrackArea.h b/Engine/Source/Editor/Sequencer/Private/SSequencerTrackArea.h
--- a/Engine/Source/Editor/Sequencer/Private/SSequencerTrackArea.h
+++ b/Engine/Source/Editor/Sequencer/Private/SSequencerTrackArea.h
@@ -56,6 +56,8 @@ public:
 	virtual FReply OnMouseButtonUp( const FGeometry& MyGeometry, const FPointerEvent& MouseEvent ) override;
 	virtual FReply OnMouseMove( const FGeometry& MyGeometry, const FPointerEvent& MouseEvent ) override;
 	virtual FReply OnMouseWheel( const FGeometry& MyGeometry, const FPointerEvent& MouseEvent ) override;
+	virtual void OnMouseEnter(const FGeometry& MyGeometry, const FPointerEvent& MouseEvent) override;
+	virtual void OnMouseLeave(const FPointerEvent& MouseEvent) override;
 	virtual int32 OnPaint(const FPaintArgs& Args, const FGeometry& AllottedGeometry, const FSlateRect& MyClippingRect, FSlateWindowElementList& OutDrawElements, int32 LayerId, const FWidgetStyle& InWidgetStyle, bool bParentEnabled ) const override;
 	virtual void OnMouseCaptureLost() override;
 	virtual FCursorReply OnCursorQuery( const FGeometry& MyGeometry, const FPointerEvent& CursorEvent ) const override;
diff --git a/Engine/Source/Editor/Sequencer/Private/Sequencer.cpp b/Engine/Source/Editor/Sequencer/Private/Sequencer.cpp
--- a/Engine/Source/Editor/Sequencer/Private/Sequencer.cpp
+++ b/Engine/Source/Editor/Sequencer/Private/Sequencer.cpp
@@ -1916,6 +1916,11 @@ FSequencerSelection& FSequencer::GetSelection()
 	return Selection;
 }
 
+FSequencerSelectionPreview& FSequencer::GetSelectionPreview()
+{
+	return SelectionPreview;
+}
+
 float FSequencer::GetOverlayFadeCurve() const
 {
 	return OverlayCurve.GetLerp();
diff --git a/Engine/Source/Editor/Sequencer/Private/Sequencer.h b/Engine/Source/Editor/Sequencer/Private/Sequencer.h
--- a/Engine/Source/Editor/Sequencer/Private/Sequencer.h
+++ b/Engine/Source/Editor/Sequencer/Private/Sequencer.h
@@ -89,6 +89,7 @@ public:
 	virtual bool CanKeyProperty(FCanKeyPropertyParams CanKeyPropertyParams) const override;
 	virtual void KeyProperty(FKeyPropertyParams KeyPropertyParams) override;
 	virtual FSequencerSelection& GetSelection() override;
+	virtual FSequencerSelectionPreview& GetSelectionPreview() override;
 	virtual void NotifyMapChanged(class UWorld* NewWorld, EMapChangeType MapChangeType) override;
 
 	/** Set the global time directly, without performing any auto-scroll */
@@ -551,4 +552,5 @@ private:
 	bool bNeedTreeRefresh;
 
 	FSequencerSelection Selection;
+	FSequencerSelectionPreview SelectionPreview;
 };
diff --git a/Engine/Source/Editor/Sequencer/Private/SequencerPrivatePCH.h b/Engine/Source/Editor/Sequencer/Private/SequencerPrivatePCH.h
--- a/Engine/Source/Editor/Sequencer/Private/SequencerPrivatePCH.h
+++ b/Engine/Source/Editor/Sequencer/Private/SequencerPrivatePCH.h
@@ -11,6 +11,7 @@
 #include "ISequencerModule.h"
 #include "ISequencerObjectChangeListener.h"
 #include "SequencerSelection.h"
+#include "SequencerSelectionPreview.h"
 #include "SequencerCommands.h"
 #include "SequencerNodeTree.h"
 #include "SequencerDisplayNode.h"
diff --git a/Engine/Source/Editor/Sequencer/Private/SequencerSelectionPreview.cpp b/Engine/Source/Editor/Sequencer/Private/SequencerSelectionPreview.cpp
new file mode 100644
--- /dev/null
+++ b/Engine/Source/Editor/Sequencer/Private/SequencerSelectionPreview.cpp
@@ -0,0 +1,63 @@
+// Copyright 1998-2015 Epic Games, Inc. All Rights Reserved.
+
+#include "SequencerPrivatePCH.h"
+#include "SequencerSelectionPreview.h"
+#include "MovieSceneSection.h"
+
+void FSequencerSelectionPreview::SetSelectionState(FSelectedKey Key, ESelectionPreviewState InState)
+{
+	if (InState == ESelectionPreviewState::Undefined)
+	{
+		DefinedKeyStates.Remove(Key);
+	}
+	else
+	{
+		DefinedKeyStates.Add(Key, InState);
+	}
+}
+
+void FSequencerSelectionPreview::SetSelectionState(UMovieSceneSection* Section, ESelectionPreviewState InState)
+{
+	if (InState == ESelectionPreviewState::Undefined)
+	{
+		DefinedSectionStates.Remove(Section);
+	}
+	else
+	{
+		DefinedSectionStates.Add(Section, InState);
+	}
+}
+
+ESelectionPreviewState FSequencerSelectionPreview::GetSelectionState(FSelectedKey Key) const
+{
+	if (auto* State = DefinedKeyStates.Find(Key))
+	{
+		return *State;
+	}
+	return ESelectionPreviewState::Undefined;
+}
+
+ESelectionPreviewState FSequencerSelectionPreview::GetSelectionState(UMovieSceneSection* Section) const
+{
+	if (auto* State = DefinedSectionStates.Find(Section))
+	{
+		return *State;
+	}
+	return ESelectionPreviewState::Undefined;
+}
+
+void FSequencerSelectionPreview::Empty()
+{
+	EmptyDefinedKeyStates();
+	EmptyDefinedSectionStates();
+}
+
+void FSequencerSelectionPreview::EmptyDefinedKeyStates()
+{
+	DefinedKeyStates.Reset();
+}
+
+void FSequencerSelectionPreview::EmptyDefinedSectionStates()
+{
+	DefinedSectionStates.Reset();
+}
\ No newline at end of file
diff --git a/Engine/Source/Editor/Sequencer/Private/SequencerSelectionPreview.h b/Engine/Source/Editor/Sequencer/Private/SequencerSelectionPreview.h
new file mode 100644
--- /dev/null
+++ b/Engine/Source/Editor/Sequencer/Private/SequencerSelectionPreview.h
@@ -0,0 +1,55 @@
+// Copyright 1998-2015 Epic Games, Inc. All Rights Reserved.
+
+#pragma once
+
+#include "SelectedKey.h"
+
+class UMovieSceneSection;
+class FSequencerDisplayNode;
+
+enum class ESelectionPreviewState
+{
+	Undefined,
+	Selected,
+	NotSelected
+};
+
+/**
+ * Manages the selection of keys, sections, and outliner nodes for the sequencer.
+ */
+class FSequencerSelectionPreview
+{
+public:
+
+	/** Access the defined key states */
+	const TMap<FSelectedKey, ESelectionPreviewState>& GetDefinedKeyStates() const { return DefinedKeyStates; }
+
+	/** Access the defined section states */
+	const TMap<TWeakObjectPtr<UMovieSceneSection>, ESelectionPreviewState>& GetDefinedSectionStates() const { return DefinedSectionStates; }
+
+	/** Adds a key to the selection */
+	void SetSelectionState(FSelectedKey Key, ESelectionPreviewState InState);
+
+	/** Adds a section to the selection */
+	void SetSelectionState(UMovieSceneSection* Section, ESelectionPreviewState InState);
+
+	/** Returns the selection state for the the specified key. */
+	ESelectionPreviewState GetSelectionState(FSelectedKey Key) const;
+
+	/** Returns the selection state for the the specified section. */
+	ESelectionPreviewState GetSelectionState(UMovieSceneSection* Section) const;
+
+	/** Empties all selections. */
+	void Empty();
+
+	/** Empties the key selection. */
+	void EmptyDefinedKeyStates();
+
+	/** Empties the section selection. */
+	void EmptyDefinedSectionStates();
+
+private:
+
+	TMap<FSelectedKey, ESelectionPreviewState> DefinedKeyStates;
+	TMap<TWeakObjectPtr<UMovieSceneSection>, ESelectionPreviewState> DefinedSectionStates;
+};
\ No newline at end of file
diff --git a/Engine/Source/Editor/Sequencer/Private/Tools/SequencerEditTool_Movement.cpp b/Engine/Source/Editor/Sequencer/Private/Tools/SequencerEditTool_Movement.cpp
--- a/Engine/Source/Editor/Sequencer/Private/Tools/SequencerEditTool_Movement.cpp
+++ b/Engine/Source/Editor/Sequencer/Private/Tools/SequencerEditTool_Movement.cpp
@@ -9,9 +9,12 @@
 
 FReply FSequencerEditTool_Default::OnMouseButtonUp(SWidget& OwnerWidget, const FGeometry& MyGeometry, const FPointerEvent& MouseEvent)
 {
-	FSequencer& Sequencer = static_cast<FSequencer&>(GetSequencer());
-	Sequencer.GetSelection().EmptySelectedSections();
-	Sequencer.GetSelection().EmptySelectedKeys();
+	if (MouseEvent.GetEffectingButton() == EKeys::LeftMouseButton)
+	{
+		FSequencer& Sequencer = static_cast<FSequencer&>(GetSequencer());
+		Sequencer.GetSelection().EmptySelectedSections();
+		Sequencer.GetSelection().EmptySelectedKeys();
+	}
 	return FReply::Unhandled();
 }
 
diff --git a/Engine/Source/Editor/Sequencer/Private/Tools/SequencerEditTool_Selection.cpp b/Engine/Source/Editor/Sequencer/Private/Tools/SequencerEditTool_Selection.cpp
--- a/Engine/Source/Editor/Sequencer/Private/Tools/SequencerEditTool_Selection.cpp
+++ b/Engine/Source/Editor/Sequencer/Private/Tools/SequencerEditTool_Selection.cpp
@@ -10,91 +10,80 @@
 #include "SSequencerTreeView.h"
 #include "SequencerHotspots.h"
 
-
-struct FSelectionVisitor : ISequencerEntityVisitor
+struct FSelectionPreviewVisitor : ISequencerEntityVisitor
 {
-	FSelectionVisitor(FSequencerSelection& InSelection, const FPointerEvent& InMouseEvent, ESequencerSelectionMode InMode)
-		: Selection(InSelection)
-		, MouseEvent(InMouseEvent)
-		, Mode(InMode)
+	FSelectionPreviewVisitor(FSequencerSelectionPreview& InSelectionPreview, FSequencerSelection& InSelection, ESelectionPreviewState InSetStateTo)
+		: SelectionPreview(InSelectionPreview)
+		, ExistingSelection(InSelection)
+		, SetStateTo(InSetStateTo)
 	{}
 
 	virtual void VisitKey(FKeyHandle KeyHandle, float KeyTime, const TSharedPtr<IKeyArea>& KeyArea, UMovieSceneSection* Section) const override
 	{
-		if (Mode != ESequencerSelectionMode::Keys)
-		{
-			return;
-		}
+		FSelectedKey Key(*Section, KeyArea, KeyHandle);
 
-		FSelectedKey SelectedKey(*Section, KeyArea, KeyHandle);
+		bool bResetSectionSelection = false;
 
-		if (MouseEvent.IsLeftShiftDown())
+		// If we're trying to change this key's selection state, we reset the selection state of the section
+		bool bKeyIsSelected = ExistingSelection.IsSelected(Key);
+		if ( (bKeyIsSelected && SetStateTo == ESelectionPreviewState::NotSelected) ||
+			(!bKeyIsSelected && SetStateTo == ESelectionPreviewState::Selected) )
 		{
-			Selection.AddToSelection(SelectedKey);
-		}
-		else if (MouseEvent.IsLeftControlDown())
-		{
-			if (Selection.IsSelected(SelectedKey))
-			{
-				Selection.RemoveFromSelection(SelectedKey);
-			}
-			else
-			{
-				Selection.AddToSelection(SelectedKey);
-			}
-		}
-		else
-		{
-			Selection.AddToSelection(SelectedKey);
+			SelectionPreview.SetSelectionState(Section, ESelectionPreviewState::Undefined);
+			SectionsWithKeysSelected.Add(Section);
 		}
+
+		SelectionPreview.SetSelectionState(Key, SetStateTo);
 	}
 
 	virtual void VisitSection(UMovieSceneSection* Section) const
 	{
-		if (Mode != ESequencerSelectionMode::Sections)
-		{
-			return;
-		}
-
-		if (MouseEvent.IsLeftShiftDown())
+		if (!SectionsWithKeysSelected.Contains(Section))
 		{
-			Selection.AddToSelection(Section);
-		}
-		else if (MouseEvent.IsLeftControlDown())
-		{
-			if (Selection.IsSelected(Section))
-			{
-				Selection.RemoveFromSelection(Section);
-			}
-			else
-			{
-				Selection.AddToSelection(Section);
-			}
-		}
-		else
-		{
-			Selection.AddToSelection(Section);
+			SelectionPreview.SetSelectionState(Section, SetStateTo);
 		}
 	}
 
 private:
-	FSequencerSelection& Selection;
-	const FPointerEvent& MouseEvent;
-	ESequencerSelectionMode Mode;
+	mutable TSet<UMovieSceneSection*> SectionsWithKeysSelected;
+
+	FSequencerSelectionPreview& SelectionPreview;
+	FSequencerSelection& ExistingSelection;
+	ESelectionPreviewState SetStateTo;
 };
 
 class FMarqueeDragOperation : public IEditToolDragOperation
 {
 public:
 
-	FMarqueeDragOperation(TWeakPtr<FSequencer> InSequencer, TWeakPtr<SSequencer> InSequencerWidget, ESequencerSelectionMode InSelectionMode)
-		: Sequencer(MoveTemp(InSequencer)), SequencerWidget(MoveTemp(InSequencerWidget)), SelectionMode(InSelectionMode)
+	FMarqueeDragOperation(TWeakPtr<FSequencer> InSequencer, TWeakPtr<SSequencer> InSequencerWidget)
+		: Sequencer(MoveTemp(InSequencer))
+		, SequencerWidget(MoveTemp(InSequencerWidget))
+		, PreviewState(ESelectionPreviewState::Selected)
 	{}
 
 	/** Start a new marquee selection */
 	virtual void OnBeginDrag(const FPointerEvent& MouseEvent, FVector2D LocalMousePos, const FVirtualTrackArea& VirtualTrackArea) override
 	{
 		InitialPosition =  VirtualTrackArea.PhysicalToVirtual(LocalMousePos);
+
+		CurrentMousePos = LocalMousePos;
+
+		if (MouseEvent.IsShiftDown())
+		{
+			PreviewState = ESelectionPreviewState::Selected;
+		}
+		else if (MouseEvent.IsAltDown())
+		{
+			PreviewState = ESelectionPreviewState::NotSelected;
+		}
+		else
+		{
+			PreviewState = ESelectionPreviewState::Selected;
+
+			// @todo: selection in transactions
+			Sequencer.Pin()->GetSelection().Empty();
+		}
 	}
 
 	/** Change the current marquee selection */
@@ -125,9 +114,14 @@ public:
 		LocalMousePos.Y = FMath::Clamp(LocalMousePos.Y, 0.f, VirtualTrackArea.GetPhysicalSize().Y);
 		CurrentPosition = VirtualTrackArea.PhysicalToVirtual(LocalMousePos);
 
+		// Clamp software cursor position to bounds of the track area
+		CurrentMousePos = LocalMousePos;
+		CurrentMousePos.X = FMath::Clamp(CurrentMousePos.X, 0.f, VirtualTrackArea.GetPhysicalSize().X);
+
+		TRange<float> ViewRange = PinnedSequencer->GetViewRange();
+
 		// Handle virtual scrolling when at the horizontal extremes of the widget
 		{
-			TRange<float> ViewRange = PinnedSequencer->GetViewRange();
 			const float ScrollThresholdH = ViewRange.Size<float>() * 0.025f;
 
 			float Difference = CurrentPosition.X - (ViewRange.GetLowerBoundValue() + ScrollThresholdH);
@@ -148,6 +142,24 @@ public:
 				}
 			}
 		}
+
+		// Calculate the size of a key in virtual space
+		FVector2D VirtualKeySize;
+		VirtualKeySize.X = SequencerSectionConstants::KeySize.X / VirtualTrackArea.GetPhysicalSize().X * ViewRange.Size<float>();
+		// Vertically, virtual units == physical units
+		VirtualKeySize.Y = SequencerSectionConstants::KeySize.Y;
+
+		// Visit everything using the preview selection primarily as well as the 
+		auto& SelectionPreview = PinnedSequencer->GetSelectionPreview();
+
+		// Ensure the preview is empty before calculating the intersection
+		SelectionPreview.Empty();
+
+		const auto& RootNodes = SequencerWidget.Pin()->GetTreeView()->GetNodeTree()->GetRootNodes();
+
+		// Now walk everything within the current marquee range, setting preview selection states as we go
+		FSequencerEntityWalker Walker(FSequencerEntityRange(TopLeft(), BottomRight()), VirtualKeySize);
+		Walker.Traverse(FSelectionPreviewVisitor(SelectionPreview, PinnedSequencer->GetSelection(), PreviewState), RootNodes);
 	}
 
 	/** Finish dragging the marquee selection */
@@ -156,26 +168,38 @@ public:
 		auto PinnedSequencer = Sequencer.Pin();
 
 		auto& Selection = PinnedSequencer->GetSelection();
-		if (!MouseEvent.IsLeftControlDown() && !MouseEvent.IsLeftShiftDown())
+		auto& SelectionPreview = PinnedSequencer->GetSelectionPreview();
+
+		// Patch everything from the selection preview into the actual selection
+		for (const auto& Pair : SelectionPreview.GetDefinedKeyStates())
 		{
-			Selection.Empty();
+			if (Pair.Value == ESelectionPreviewState::Selected)
+			{
+				// Select it in the main selection
+				Selection.AddToSelection(Pair.Key);
+			}
+			else
+			{
+				Selection.RemoveFromSelection(Pair.Key);
+			}
 		}
 
-		Selection.SuspendBroadcast();
-
-		FVector2D VirtualKeySize;
-		VirtualKeySize.X = SequencerSectionConstants::KeySize.X / VirtualTrackArea.GetPhysicalSize().X * (PinnedSequencer->GetViewRange().GetUpperBoundValue() - PinnedSequencer->GetViewRange().GetLowerBoundValue());
-		// Vertically, virtual units == physical units
-		VirtualKeySize.Y = SequencerSectionConstants::KeySize.Y;
-
-		FSequencerEntityWalker Walker(FSequencerEntityRange(TopLeft(), BottomRight()), VirtualKeySize);
-
-		FSelectionVisitor Visitor(Selection, MouseEvent, SelectionMode);
-
-		Walker.Traverse(Visitor, SequencerWidget.Pin()->GetTreeView()->GetNodeTree()->GetRootNodes());
+		for (const auto& Pair : SelectionPreview.GetDefinedSectionStates())
+		{
+			UMovieSceneSection* Section = Pair.Key.Get();
+			if (Pair.Value == ESelectionPreviewState::Selected)
+			{
+				// Select it in the main selection
+				Selection.AddToSelection(Section);
+			}
+			else
+			{
+				Selection.RemoveFromSelection(Section);
+			}
+		}
 
-		Selection.ResumeBroadcast();
-		Selection.GetOnKeySelectionChanged().Broadcast();
+		// We're done with this now
+		SelectionPreview.Empty();
 	}
 
 	virtual int32 OnPaint(const FGeometry& AllottedGeometry, const FSlateRect& MyClippingRect, FSlateWindowElementList& OutDrawElements, int32 LayerId) const
@@ -196,9 +220,6 @@ public:
 
 		return LayerId + 1;
 	}
-	
-	/** Request the cursor for this drag operation */
-	virtual FCursorReply GetCursor() const { return FCursorReply::Cursor( EMouseCursor::Default ); }
 
 private:
 
@@ -224,27 +245,19 @@ private:
 	/** Sequencer widget */
 	TWeakPtr<SSequencer> SequencerWidget;
 
+	/** Whether we should select/deselect things in this marquee operation */
+	ESelectionPreviewState PreviewState;
+
 	FVector2D InitialPosition;
 	FVector2D CurrentPosition;
-
-	ESequencerSelectionMode SelectionMode;
+	FVector2D CurrentMousePos;
 };
 
-/** todo: Save this in a config? */
-ESequencerSelectionMode LastUsedMode = ESequencerSelectionMode::Keys;
-
 FSequencerEditTool_Selection::FSequencerEditTool_Selection(TSharedPtr<FSequencer> InSequencer, TSharedPtr<SSequencer> InSequencerWidget)
 	: Sequencer(InSequencer)
 	, SequencerWidget(InSequencerWidget)
-	, SelectionMode(LastUsedMode)
-{
-}
-FSequencerEditTool_Selection::FSequencerEditTool_Selection(TSharedPtr<FSequencer> InSequencer, TSharedPtr<SSequencer> InSequencerWidget, ESequencerSelectionMode ExplicitMode)
-	: Sequencer(InSequencer)
-	, SequencerWidget(InSequencerWidget)
-	, SelectionMode(ExplicitMode)
+	, CursorDecorator(nullptr)
 {
-	LastUsedMode = SelectionMode;
 }
 
 ISequencer& FSequencerEditTool_Selection::GetSequencer() const
@@ -252,17 +265,6 @@ ISequencer& FSequencerEditTool_Selection::GetSequencer() const
 	return *Sequencer.Pin();
 }
 
-void FSequencerEditTool_Selection::SetSelectionMode(ESequencerSelectionMode Mode)
-{
-	SelectionMode = Mode;
-	LastUsedMode = SelectionMode;
-}
-
-ESequencerSelectionMode FSequencerEditTool_Selection::GetGlobalSelectionMode()
-{
-	return LastUsedMode;
-}
-
 FCursorReply FSequencerEditTool_Selection::OnCursorQuery(const FGeometry& MyGeometry, const FPointerEvent& CursorEvent) const
 {
 	return FCursorReply::Cursor(EMouseCursor::Crosshairs);
@@ -272,7 +274,18 @@ int32 FSequencerEditTool_Selection::OnPaint(const FGeometry& AllottedGeometry, c
 {
 	if (DragOperation.IsValid())
 	{
-		return DragOperation->OnPaint(AllottedGeometry, MyClippingRect, OutDrawElements, LayerId);
+		LayerId = DragOperation->OnPaint(AllottedGeometry, MyClippingRect, OutDrawElements, LayerId);
+	}
+
+	if (CursorDecorator)
+	{
+		FSlateDrawElement::MakeBox(
+			OutDrawElements,
+			++LayerId,
+			AllottedGeometry.ToPaintGeometry(MousePosition + FVector2D(5, 5), CursorDecorator->ImageSize),
+			CursorDecorator,
+			MyClippingRect
+			);
 	}
 
 	return LayerId;
@@ -280,6 +293,8 @@ int32 FSequencerEditTool_Selection::OnPaint(const FGeometry& AllottedGeometry, c
 
 FReply FSequencerEditTool_Selection::OnMouseButtonDown(SWidget& OwnerWidget, const FGeometry& MyGeometry, const FPointerEvent& MouseEvent)
 {
+	UpdateCursor(MyGeometry, MouseEvent);
+
 	DelayedDrag.Reset();
 
 	if (MouseEvent.GetEffectingButton() == EKeys::LeftMouseButton)
@@ -292,6 +307,8 @@ FReply FSequencerEditTool_Selection::OnMouseButtonDown(SWidget& OwnerWidget, con
 
 FReply FSequencerEditTool_Selection::OnMouseMove(SWidget& OwnerWidget, const FGeometry& MyGeometry, const FPointerEvent& MouseEvent)
 {
+	UpdateCursor(MyGeometry, MouseEvent);
+
 	if (DelayedDrag.IsSet())
 	{
 		FReply Reply = FReply::Handled();
@@ -317,7 +334,7 @@ FReply FSequencerEditTool_Selection::OnMouseMove(SWidget& OwnerWidget, const FGe
 
 			if (!DragOperation.IsValid())
 			{
-				DragOperation = MakeShareable( new FMarqueeDragOperation(Sequencer, SequencerWidget, SelectionMode) );
+				DragOperation = MakeShareable( new FMarqueeDragOperation(Sequencer, SequencerWidget) );
 			}
 
 			if (DragOperation.IsValid())
@@ -337,6 +354,8 @@ FReply FSequencerEditTool_Selection::OnMouseMove(SWidget& OwnerWidget, const FGe
 
 FReply FSequencerEditTool_Selection::OnMouseButtonUp(SWidget& OwnerWidget, const FGeometry& MyGeometry, const FPointerEvent& MouseEvent)
 {
+	UpdateCursor(MyGeometry, MouseEvent);
+
 	DelayedDrag.Reset();
 	if (DragOperation.IsValid())
 	{
@@ -344,6 +363,8 @@ FReply FSequencerEditTool_Selection::OnMouseButtonUp(SWidget& OwnerWidget, const
 		DragOperation->OnEndDrag(MouseEvent, LocalPosition, SequencerWidget.Pin()->GetVirtualTrackArea() );
 		DragOperation = nullptr;
 
+		CursorDecorator = nullptr;
+
 		Sequencer.Pin()->StopAutoscroll();
 		return FReply::Handled().ReleaseMouseCapture();
 	}
@@ -351,14 +372,45 @@ FReply FSequencerEditTool_Selection::OnMouseButtonUp(SWidget& OwnerWidget, const
 	return FSequencerEditTool_Default::OnMouseButtonUp(OwnerWidget, MyGeometry, MouseEvent);
 }
 
+void FSequencerEditTool_Selection::OnMouseLeave(SWidget& OwnerWidget, const FPointerEvent& MouseEvent)
+{
+	if (!DragOperation.IsValid())
+	{
+		CursorDecorator = nullptr;
+	}
+}
+
 void FSequencerEditTool_Selection::OnMouseCaptureLost()
 {
 	DelayedDrag.Reset();
 	DragOperation = nullptr;
+	CursorDecorator = nullptr;
 }
 
 FName FSequencerEditTool_Selection::GetIdentifier() const
 {
 	static FName Identifier("Selection");
 	return Identifier;
-}
\ No newline at end of file
+}
+
+void FSequencerEditTool_Selection::UpdateCursor(const FGeometry& MyGeometry, const FPointerEvent& MouseEvent)
+{
+	MousePosition = MyGeometry.AbsoluteToLocal(MouseEvent.GetScreenSpacePosition());
+
+	// Don't update the brush if we have a drag operation
+	if (!DragOperation.IsValid())
+	{
+		if (MouseEvent.IsShiftDown())
+		{
+			CursorDecorator = FEditorStyle::Get().GetBrush(TEXT("Sequencer.CursorDecorator_MarqueeAdd"));
+		}
+		else if (MouseEvent.IsAltDown())
+		{
+			CursorDecorator = FEditorStyle::Get().GetBrush(TEXT("Sequencer.CursorDecorator_MarqueeSubtract"));
+		}
+		else
+		{
+			CursorDecorator = FEditorStyle::Get().GetBrush(TEXT("Sequencer.CursorDecorator_Marquee"));
+		}
+	}
+}
diff --git a/Engine/Source/Editor/Sequencer/Private/Tools/SequencerEditTool_Selection.h b/Engine/Source/Editor/Sequencer/Private/Tools/SequencerEditTool_Selection.h
--- a/Engine/Source/Editor/Sequencer/Private/Tools/SequencerEditTool_Selection.h
+++ b/Engine/Source/Editor/Sequencer/Private/Tools/SequencerEditTool_Selection.h
@@ -8,17 +8,11 @@
 
 struct FMarqueeSelectData;
 
-enum class ESequencerSelectionMode
-{
-	Keys, Sections
-};
-
 class FSequencerEditTool_Selection : public FSequencerEditTool_Default
 {
 public:
 
 	FSequencerEditTool_Selection(TSharedPtr<FSequencer> InSequencer, TSharedPtr<SSequencer> InSequencerWidget);
-	FSequencerEditTool_Selection(TSharedPtr<FSequencer> InSequencer, TSharedPtr<SSequencer> InSequencerWidget, ESequencerSelectionMode ExplicitMode);
 
 	virtual ISequencer& GetSequencer() const override;
 	virtual int32 OnPaint(const FGeometry& AllottedGeometry, const FSlateRect& MyClippingRect, FSlateWindowElementList& OutDrawElements, int32 LayerId) const override;
@@ -26,17 +20,16 @@ public:
 	virtual FReply OnMouseButtonDown(SWidget& OwnerWidget, const FGeometry& MyGeometry, const FPointerEvent& MouseEvent) override;
 	virtual FReply OnMouseButtonUp(SWidget& OwnerWidget, const FGeometry& MyGeometry, const FPointerEvent& MouseEvent) override;
 	virtual FReply OnMouseMove(SWidget& OwnerWidget, const FGeometry& MyGeometry, const FPointerEvent& MouseEvent) override;
+	virtual void   OnMouseLeave(SWidget& OwnerWidget, const FPointerEvent& MouseEvent) override;
 	virtual void   OnMouseCaptureLost() override;
 	virtual FCursorReply OnCursorQuery(const FGeometry& MyGeometry, const FPointerEvent& CursorEvent) const override;
 	virtual FName GetIdentifier() const override;
+	
+private:
 
-public:
+	/** Update the software cursor */
+	void UpdateCursor(const FGeometry& MyGeometry, const FPointerEvent& MouseEvent);
 
-	static ESequencerSelectionMode GetGlobalSelectionMode();
-	
-	ESequencerSelectionMode GetSelectionMode() const { return SelectionMode; }
-	void SetSelectionMode(ESequencerSelectionMode Mode);
-	
 private:
 
 	/** The sequencer itself */
@@ -61,6 +54,9 @@ private:
 	/** Current drag operation if any */
 	TSharedPtr<IEditToolDragOperation> DragOperation;
 
-	/** What we are selecting */
-	ESequencerSelectionMode SelectionMode;
+	/** Cached mouse position for software cursor rendering */
+	FVector2D MousePosition;
+
+	/** Software cursor decorator brush */
+	const FSlateBrush* CursorDecorator;
 };
diff --git a/Engine/Source/Editor/Sequencer/Private/Tools/SequencerEntityVisitor.cpp b/Engine/Source/Editor/Sequencer/Private/Tools/SequencerEntityVisitor.cpp
--- a/Engine/Source/Editor/Sequencer/Private/Tools/SequencerEntityVisitor.cpp
+++ b/Engine/Source/Editor/Sequencer/Private/Tools/SequencerEntityVisitor.cpp
@@ -77,10 +77,28 @@ void FSequencerEntityWalker::HandleNode(const ISequencerEntityVisitor& Visitor,
 	}
 }
 
-void FSequencerEntityWalker::HandleNode(const ISequencerEntityVisitor& Visitor, FSequencerDisplayNode& InNode, const TArray<TSharedRef<ISequencerSection>>& InSections)
+void FSequencerEntityWalker::HandleNode(const ISequencerEntityVisitor& Visitor, FSequencerDisplayNode& InNode, TArray<TSharedRef<ISequencerSection>> InSections)
 {
 	if (Range.IntersectNode(InNode))
 	{
+		// Prune the selections to anything that is in the range, visiting if necessary
+		for (int32 SectionIndex = 0; SectionIndex < InSections.Num();)
+		{
+			UMovieSceneSection* Section = InSections[SectionIndex]->GetSectionObject();
+			if (!Range.IntersectSection(Section))
+			{
+				InSections.RemoveAtSwap(SectionIndex, 1, false);
+				continue;
+			}
+
+			if (Visitor.CheckEntityMask(ESequencerEntity::Section))
+			{
+				Visitor.VisitSection(Section);
+			}
+
+			++SectionIndex;
+		}
+
 		bool bNodeHasKeyArea = false;
 		if (InNode.GetType() == ESequencerNode::KeyArea)
 		{
@@ -97,29 +115,26 @@ void FSequencerEntityWalker::HandleNode(const ISequencerEntityVisitor& Visitor,
 			}
 		}
 
-		if (!bNodeHasKeyArea && (!InNode.IsExpanded() || InNode.GetChildNodes().Num() == 0))
+		// As a fallback, we need to handle:
+		//  - Key groupings on collapsed parents
+		//  - Sections that have no key areas
+		const bool bIterateKeyGroupings = Visitor.CheckEntityMask(ESequencerEntity::Key) &&
+			!bNodeHasKeyArea &&
+			(!InNode.IsExpanded() || InNode.GetChildNodes().Num() == 0);
+
+		if (bIterateKeyGroupings)
 		{
-			// As a fallback, we need to handle:
-			//  - Key groupings on collapsed parents
-			//  - Sections that have no key areas
 			for (int32 SectionIndex = 0; SectionIndex < InSections.Num(); ++SectionIndex)
 			{
 				UMovieSceneSection* Section = InSections[SectionIndex]->GetSectionObject();
-				if (Range.IntersectSection(Section))
-				{
-					if (Visitor.CheckEntityMask(ESequencerEntity::Section))
-					{
-						Visitor.VisitSection(Section);
-					}
 
-					if (Visitor.CheckEntityMask(ESequencerEntity::Key))
+				if (Visitor.CheckEntityMask(ESequencerEntity::Key))
+				{
+					// Only handle grouped keys if we actually have children
+					if (InNode.GetChildNodes().Num() != 0 && Range.IntersectKeyArea(InNode, VirtualKeySize.X))
 					{
-						// Only handle grouped keys if we actually have children
-						if (InNode.GetChildNodes().Num() != 0 && Range.IntersectKeyArea(InNode, VirtualKeySize.X))
-						{
-							TSharedRef<IKeyArea> KeyArea = InNode.UpdateKeyGrouping(SectionIndex);
-							HandleKeyArea(Visitor, KeyArea, Section);
-						}
+						TSharedRef<IKeyArea> KeyArea = InNode.UpdateKeyGrouping(SectionIndex);
+						HandleKeyArea(Visitor, KeyArea, Section);
 					}
 				}
 			}
@@ -144,22 +159,12 @@ void FSequencerEntityWalker::HandleKeyAreaNode(const ISequencerEntityVisitor& Vi
 	for( int32 SectionIndex = 0; SectionIndex < InSections.Num(); ++SectionIndex )
 	{
 		UMovieSceneSection* Section = InSections[SectionIndex]->GetSectionObject();
-
-		// If the section is at all within the marquee, we check its keys
-		if (Range.IntersectSection(Section))
+		if (Visitor.CheckEntityMask(ESequencerEntity::Key))
 		{
-			if (Visitor.CheckEntityMask(ESequencerEntity::Section))
+			if (Range.IntersectKeyArea(InOwnerNode, VirtualKeySize.X))
 			{
-				Visitor.VisitSection(Section);
-			}
-
-			if (Visitor.CheckEntityMask(ESequencerEntity::Key))
-			{
-				if (Range.IntersectKeyArea(InOwnerNode, VirtualKeySize.X))
-				{
-					TSharedRef<IKeyArea> KeyArea = InKeyAreaNode.GetKeyArea(SectionIndex);
-					HandleKeyArea(Visitor, KeyArea, InSections[SectionIndex]->GetSectionObject());
-				}
+				TSharedRef<IKeyArea> KeyArea = InKeyAreaNode.GetKeyArea(SectionIndex);
+				HandleKeyArea(Visitor, KeyArea, InSections[SectionIndex]->GetSectionObject());
 			}
 		}
 	}
diff --git a/Engine/Source/Editor/Sequencer/Private/Tools/SequencerEntityVisitor.h b/Engine/Source/Editor/Sequencer/Private/Tools/SequencerEntityVisitor.h
--- a/Engine/Source/Editor/Sequencer/Private/Tools/SequencerEntityVisitor.h
+++ b/Engine/Source/Editor/Sequencer/Private/Tools/SequencerEntityVisitor.h
@@ -73,7 +73,7 @@ private:
 	/** Handle visitation of a particular node */
 	void HandleNode(const ISequencerEntityVisitor& Visitor, FSequencerDisplayNode& InNode);
 	/** Handle visitation of a particular node, along with a set of sections */
-	void HandleNode(const ISequencerEntityVisitor& Visitor, FSequencerDisplayNode& InNode, const TArray<TSharedRef<ISequencerSection>>& InSections);
+	void HandleNode(const ISequencerEntityVisitor& Visitor, FSequencerDisplayNode& InNode, TArray<TSharedRef<ISequencerSection>> InSections);
 	/** Handle visitation of a key area node */
 	void HandleKeyAreaNode(const ISequencerEntityVisitor& Visitor, FSectionKeyAreaNode& InKeyAreaNode, FSequencerDisplayNode& InOwnerNode, const TArray<TSharedRef<ISequencerSection>>& InSections);
 	/** Handle visitation of a key area */
diff --git a/Engine/Source/Editor/Sequencer/Public/ISequencer.h b/Engine/Source/Editor/Sequencer/Public/ISequencer.h
--- a/Engine/Source/Editor/Sequencer/Public/ISequencer.h
+++ b/Engine/Source/Editor/Sequencer/Public/ISequencer.h
@@ -9,6 +9,7 @@
 
 
 class FSequencerSelection;
+class FSequencerSelectionPreview;
 class UMovieSceneSequence;
 class UAnimSequence;
 class UMovieScene;
@@ -153,6 +154,7 @@ public:
 	virtual void UpdateRuntimeInstances() = 0;
 
 	virtual FSequencerSelection& GetSelection() = 0;
+	virtual FSequencerSelectionPreview& GetSelectionPreview() = 0;
 
 	virtual void NotifyMapChanged(class UWorld* NewWorld, EMapChangeType MapChangeType) = 0;
 };
diff --git a/Engine/Source/Editor/Sequencer/Public/ISequencerEditTool.h b/Engine/Source/Editor/Sequencer/Public/ISequencerEditTool.h
--- a/Engine/Source/Editor/Sequencer/Public/ISequencerEditTool.h
+++ b/Engine/Source/Editor/Sequencer/Public/ISequencerEditTool.h
@@ -53,6 +53,8 @@ public:
 	virtual FReply OnMouseMove(SWidget& OwnerWidget, const FGeometry& MyGeometry, const FPointerEvent& MouseEvent) { return FReply::Unhandled(); }
 	virtual FReply OnMouseWheel(SWidget& OwnerWidget, const FGeometry& MyGeometry, const FPointerEvent& MouseEvent) { return FReply::Unhandled(); }
 	virtual void   OnMouseCaptureLost() { }
+	virtual void OnMouseEnter(SWidget& OwnerWidget, const FGeometry& MyGeometry, const FPointerEvent& MouseEvent) { }
+	virtual void OnMouseLeave(SWidget& OwnerWidget, const FPointerEvent& MouseEvent) { }
 	virtual int32  OnPaint(const FGeometry& AllottedGeometry, const FSlateRect& MyClippingRect, FSlateWindowElementList& OutDrawElements, int32 LayerId) const { return LayerId; }
 	virtual void   Tick(const FGeometry& AllottedGeometry, const double InCurrentTime, const float InDeltaTime) { }
 	virtual FCursorReply OnCursorQuery(const FGeometry& MyGeometry, const FPointerEvent& CursorEvent) const { return FCursorReply::Unhandled(); }
