Fixed assets getting the wrong names in the sequencer breadcrumbs
This addresses UE-21678

diff --git a/Engine/Source/Editor/Sequencer/Private/SSequencer.cpp b/Engine/Source/Editor/Sequencer/Private/SSequencer.cpp
--- a/Engine/Source/Editor/Sequencer/Private/SSequencer.cpp
+++ b/Engine/Source/Editor/Sequencer/Private/SSequencer.cpp
@@ -917,7 +917,7 @@ void SSequencer::UpdateBreadcrumbs(const TArray< TWeakObjectPtr<class UMovieScen
 
 	if( BreadcrumbTrail->PeekCrumb().BreadcrumbType == FSequencerBreadcrumb::MovieSceneType && BreadcrumbTrail->PeekCrumb().MovieSceneInstance.Pin() != FocusedMovieSceneInstance )
 	{
-		FText CrumbName = FText::FromString( FocusedMovieSceneInstance->GetSequence()->GetName() );
+		FText CrumbName = FocusedMovieSceneInstance->GetSequence()->GetDisplayName();
 		// The current breadcrumb is not a moviescene so we need to make a new breadcrumb in order return to the parent moviescene later
 		BreadcrumbTrail->PushCrumb( CrumbName, FSequencerBreadcrumb( FocusedMovieSceneInstance ) );
 	}
@@ -1285,7 +1285,7 @@ void SSequencer::OnCrumbClicked(const FSequencerBreadcrumb& Item)
 
 FText SSequencer::GetRootAnimationName() const
 {
-	return FText::FromName(Sequencer.Pin()->GetRootMovieSceneSequence()->GetFName());
+	return Sequencer.Pin()->GetRootMovieSceneSequence()->GetDisplayName();
 }
 
 FText SSequencer::GetShotSectionTitle(UMovieSceneSection* ShotSection) const
diff --git a/Engine/Source/Runtime/LevelSequence/Private/LevelSequence.cpp b/Engine/Source/Runtime/LevelSequence/Private/LevelSequence.cpp
--- a/Engine/Source/Runtime/LevelSequence/Private/LevelSequence.cpp
+++ b/Engine/Source/Runtime/LevelSequence/Private/LevelSequence.cpp
@@ -362,6 +362,12 @@ bool ULevelSequenceInstance::TryGetObjectDisplayName(const FGuid& ObjectId, FTex
 
 	return true;
 }
+
+FText ULevelSequenceInstance::GetDisplayName() const
+{
+	return LevelSequence ? FText::FromName(LevelSequence->GetFName()) : Super::GetDisplayName();
+}
+
 #endif
 
 void ULevelSequenceInstance::UnbindPossessableObjects(const FGuid& ObjectId)
diff --git a/Engine/Source/Runtime/LevelSequence/Public/LevelSequence.h b/Engine/Source/Runtime/LevelSequence/Public/LevelSequence.h
--- a/Engine/Source/Runtime/LevelSequence/Public/LevelSequence.h
+++ b/Engine/Source/Runtime/LevelSequence/Public/LevelSequence.h
@@ -97,6 +97,7 @@ public:
 
 #if WITH_EDITOR
 	virtual bool TryGetObjectDisplayName(const FGuid& ObjectId, FText& OutDisplayName) const override;
+	virtual FText GetDisplayName() const override;
 #endif
 
 protected:
diff --git a/Engine/Source/Runtime/MovieScene/Public/MovieSceneSequence.h b/Engine/Source/Runtime/MovieScene/Public/MovieSceneSequence.h
--- a/Engine/Source/Runtime/MovieScene/Public/MovieSceneSequence.h
+++ b/Engine/Source/Runtime/MovieScene/Public/MovieSceneSequence.h
@@ -114,5 +114,10 @@ public:
 	 * @returns true if DisplayName has been set to a valid display name, otherwise false.
 	 */
 	virtual bool TryGetObjectDisplayName(const FGuid& ObjectId, FText& OutDisplayName) const PURE_VIRTUAL(UMovieSceneSequence::TryGetObjectDisplayName, return false;);
+
+	/**
+	 * Get the display name for this movie sequence
+	 */
+	virtual FText GetDisplayName() const { return FText::FromName(GetFName()); }
 #endif
 };
