Mergine changes from //depot/UE4 to //UE4/Main

diff --git a/Engine/Source/Editor/MovieSceneTools/Private/TrackEditors/PathTrackEditor.h b/Engine/Source/Editor/MovieSceneTools/Private/TrackEditors/PathTrackEditor.h
--- a/Engine/Source/Editor/MovieSceneTools/Private/TrackEditors/PathTrackEditor.h
+++ b/Engine/Source/Editor/MovieSceneTools/Private/TrackEditors/PathTrackEditor.h
@@ -4,13 +4,14 @@
 
 #include "MovieScene3DPathSection.h"
 #include "MovieScene3DPathTrack.h"
+#include "ActorPickerTrackEditor.h"
 
 
 /**
  * Tools for animating a transform with a path curve
  */
 class F3DPathTrackEditor
-	: public FMovieSceneTrackEditor
+	: public FActorPickerTrackEditor
 {
 public:
 
@@ -22,7 +23,7 @@ public:
 	F3DPathTrackEditor( TSharedRef<ISequencer> InSequencer );
 
 	/** Virtual destructor. */
-	virtual ~F3DPathTrackEditor();
+	virtual ~F3DPathTrackEditor() { }
 
 	/**
 	 * Creates an instance of this class.  Called by a sequencer 
@@ -32,14 +33,6 @@ public:
 	 */
 	static TSharedRef<ISequencerTrackEditor> CreateTrackEditor( TSharedRef<ISequencer> OwningSequencer );
 
-public:
-
-	/** Add path */
-	void AddPath(FGuid ObjectGuid, UObject* AdditionalAsset);
-
-	/** Set path */
-	void SetPath(UMovieSceneSection* Section, AActor* ActorWithSplineComponent);
-
 public:
 
 	// ISequencerTrackEditor interface
@@ -49,13 +42,14 @@ public:
 	virtual TSharedRef<ISequencerSection> MakeSectionInterface( UMovieSceneSection& SectionObject, UMovieSceneTrack& Track ) override;
 	virtual bool SupportsType( TSubclassOf<UMovieSceneTrack> Type ) const override;
 
-private:
+	// FTrackEditorActorPicker
+	virtual bool IsActorPickable( const AActor* const ParentActor ) override;
+	virtual void ActorSocketPicked(const FName SocketName, AActor* ParentActor, FGuid ObjectBinding, UMovieSceneSection* Section) override;
 
-	/** Add path sub menu */
-	void AddPathSubMenu(FMenuBuilder& MenuBuilder, FGuid ObjectBinding);
+private:
 
 	/** Delegate for AnimatablePropertyChanged in AddKey */
-	bool AddKeyInternal(float KeyTime, const TArray<UObject*> Objects, UObject* AdditionalAsset);
+	bool AddKeyInternal(float KeyTime, const TArray<UObject*> Objects, AActor* ParentActor);
 };
 
 
diff --git a/Engine/Source/Editor/MovieSceneTools/Private/TrackEditors/PropertyTrackEditors/BoolPropertyTrackEditor.cpp b/Engine/Source/Editor/MovieSceneTools/Private/TrackEditors/PropertyTrackEditors/BoolPropertyTrackEditor.cpp
--- a/Engine/Source/Editor/MovieSceneTools/Private/TrackEditors/PropertyTrackEditors/BoolPropertyTrackEditor.cpp
+++ b/Engine/Source/Editor/MovieSceneTools/Private/TrackEditors/PropertyTrackEditors/BoolPropertyTrackEditor.cpp
@@ -6,19 +6,19 @@
 #include "BoolPropertySection.h"
 
 
-TSharedRef<ISequencerTrackEditor> FBoolPropertyTrackEditor::CreateTrackEditor( TSharedRef<ISequencer> OwningSequencer )
+TSharedRef<ISequencerTrackEditor> FBoolPropertyTrackEditor::CreateTrackEditor(TSharedRef<ISequencer> OwningSequencer)
 {
 	return MakeShareable(new FBoolPropertyTrackEditor(OwningSequencer));
 }
 
 
-TSharedRef<FPropertySection> FBoolPropertyTrackEditor::MakePropertySectionInterface( UMovieSceneSection& SectionObject, UMovieSceneTrack& Track )
+TSharedRef<FPropertySection> FBoolPropertyTrackEditor::MakePropertySectionInterface(UMovieSceneSection& SectionObject, UMovieSceneTrack& Track)
 {
-	return MakeShareable(new FBoolPropertySection( SectionObject, Track.GetTrackName(), GetSequencer().Get() ));
+	return MakeShareable(new FBoolPropertySection(SectionObject, Track.GetDisplayName(), GetSequencer().Get()));
 }
 
 
-void FBoolPropertyTrackEditor::GenerateKeysFromPropertyChanged( const FPropertyChangedParams& PropertyChangedParams, TArray<bool>& GeneratedKeys )
+void FBoolPropertyTrackEditor::GenerateKeysFromPropertyChanged(const FPropertyChangedParams& PropertyChangedParams, TArray<bool>& GeneratedKeys)
 {
 	GeneratedKeys.Add( PropertyChangedParams.GetPropertyValue<bool>() );
 }
diff --git a/Engine/Source/Editor/MovieSceneTools/Private/TrackEditors/PropertyTrackEditors/BytePropertyTrackEditor.cpp b/Engine/Source/Editor/MovieSceneTools/Private/TrackEditors/PropertyTrackEditors/BytePropertyTrackEditor.cpp
--- a/Engine/Source/Editor/MovieSceneTools/Private/TrackEditors/PropertyTrackEditors/BytePropertyTrackEditor.cpp
+++ b/Engine/Source/Editor/MovieSceneTools/Private/TrackEditors/PropertyTrackEditors/BytePropertyTrackEditor.cpp
@@ -15,7 +15,7 @@ TSharedRef<ISequencerTrackEditor> FBytePropertyTrackEditor::CreateTrackEditor( T
 
 TSharedRef<FPropertySection> FBytePropertyTrackEditor::MakePropertySectionInterface( UMovieSceneSection& SectionObject, UMovieSceneTrack& Track )
 {
-	return MakeShareable( new FBytePropertySection( SectionObject, Track.GetTrackName(), Cast<UMovieSceneByteTrack>( SectionObject.GetOuter() )->GetEnum() ) );
+	return MakeShareable( new FBytePropertySection( SectionObject, Track.GetDisplayName(), Cast<UMovieSceneByteTrack>( SectionObject.GetOuter() )->GetEnum() ) );
 }
 
 
diff --git a/Engine/Source/Editor/MovieSceneTools/Private/TrackEditors/PropertyTrackEditors/FloatPropertyTrackEditor.cpp b/Engine/Source/Editor/MovieSceneTools/Private/TrackEditors/PropertyTrackEditors/FloatPropertyTrackEditor.cpp
--- a/Engine/Source/Editor/MovieSceneTools/Private/TrackEditors/PropertyTrackEditors/FloatPropertyTrackEditor.cpp
+++ b/Engine/Source/Editor/MovieSceneTools/Private/TrackEditors/PropertyTrackEditors/FloatPropertyTrackEditor.cpp
@@ -13,7 +13,7 @@ TSharedRef<ISequencerTrackEditor> FFloatPropertyTrackEditor::CreateTrackEditor(
 
 TSharedRef<FPropertySection> FFloatPropertyTrackEditor::MakePropertySectionInterface( UMovieSceneSection& SectionObject, UMovieSceneTrack& Track )
 {
-	return MakeShareable(new FFloatPropertySection( SectionObject, Track.GetTrackName() ));
+	return MakeShareable(new FFloatPropertySection(SectionObject, Track.GetDisplayName()));
 }
 
 
diff --git a/Engine/Source/Editor/MovieSceneTools/Private/TrackEditors/PropertyTrackEditors/VectorPropertyTrackEditor.cpp b/Engine/Source/Editor/MovieSceneTools/Private/TrackEditors/PropertyTrackEditors/VectorPropertyTrackEditor.cpp
--- a/Engine/Source/Editor/MovieSceneTools/Private/TrackEditors/PropertyTrackEditors/VectorPropertyTrackEditor.cpp
+++ b/Engine/Source/Editor/MovieSceneTools/Private/TrackEditors/PropertyTrackEditors/VectorPropertyTrackEditor.cpp
@@ -19,7 +19,7 @@ TSharedRef<ISequencerTrackEditor> FVectorPropertyTrackEditor::CreateTrackEditor(
 
 TSharedRef<FPropertySection> FVectorPropertyTrackEditor::MakePropertySectionInterface( UMovieSceneSection& SectionObject, UMovieSceneTrack& Track )
 {
-	return MakeShareable(new FVectorPropertySection( SectionObject, Track.GetTrackName() ));
+	return MakeShareable(new FVectorPropertySection(SectionObject, Track.GetDisplayName()));
 }
 
 
diff --git a/Engine/Source/Editor/MovieSceneTools/Private/TrackEditors/PropertyTrackEditors/VisibilityPropertyTrackEditor.cpp b/Engine/Source/Editor/MovieSceneTools/Private/TrackEditors/PropertyTrackEditors/VisibilityPropertyTrackEditor.cpp
--- a/Engine/Source/Editor/MovieSceneTools/Private/TrackEditors/PropertyTrackEditors/VisibilityPropertyTrackEditor.cpp
+++ b/Engine/Source/Editor/MovieSceneTools/Private/TrackEditors/PropertyTrackEditors/VisibilityPropertyTrackEditor.cpp
@@ -14,7 +14,7 @@ TSharedRef<ISequencerTrackEditor> FVisibilityPropertyTrackEditor::CreateTrackEdi
 
 TSharedRef<FPropertySection> FVisibilityPropertyTrackEditor::MakePropertySectionInterface( UMovieSceneSection& SectionObject, UMovieSceneTrack& Track )
 {
-	return MakeShareable(new FVisibilityPropertySection( SectionObject, Track.GetTrackName(), GetSequencer().Get() ));
+	return MakeShareable(new FVisibilityPropertySection(SectionObject, Track.GetDisplayName(), GetSequencer().Get()));
 }
 
 
diff --git a/Engine/Source/Editor/MovieSceneTools/Private/TrackEditors/ShotTrackEditor/ShotTrackEditor.cpp b/Engine/Source/Editor/MovieSceneTools/Private/TrackEditors/ShotTrackEditor/ShotTrackEditor.cpp
--- a/Engine/Source/Editor/MovieSceneTools/Private/TrackEditors/ShotTrackEditor/ShotTrackEditor.cpp
+++ b/Engine/Source/Editor/MovieSceneTools/Private/TrackEditors/ShotTrackEditor/ShotTrackEditor.cpp
@@ -114,23 +114,23 @@ TSharedPtr<SWidget> FShotTrackEditor::BuildOutlinerEditWidget(const FGuid& Objec
 	];
 
 	EditBox.Get()->AddSlot()
-	.VAlign(VAlign_Center)
-	.HAlign(HAlign_Right)
-	.AutoWidth()
-	.Padding(4, 0, 0, 0)
-	[
-		SNew(SCheckBox)
-		.IsChecked(this, &FShotTrackEditor::IsCameraLocked)
-		.OnCheckStateChanged(this, &FShotTrackEditor::OnLockCameraClicked)
-		.ToolTipText(this, &FShotTrackEditor::GetLockCameraToolTip)
-		.ForegroundColor(FLinearColor::White)
-		.CheckedImage(FEditorStyle::GetBrush("Sequencer.LockCamera"))
-		.CheckedHoveredImage(FEditorStyle::GetBrush("Sequencer.LockCamera"))
-		.CheckedPressedImage(FEditorStyle::GetBrush("Sequencer.LockCamera"))
-		.UncheckedImage(FEditorStyle::GetBrush("Sequencer.UnlockCamera"))
-		.UncheckedHoveredImage(FEditorStyle::GetBrush("Sequencer.UnlockCamera"))
-		.UncheckedPressedImage(FEditorStyle::GetBrush("Sequencer.UnlockCamera"))
-	];
+		.VAlign(VAlign_Center)
+		.HAlign(HAlign_Right)
+		.AutoWidth()
+		.Padding(4, 0, 0, 0)
+		[
+			SNew(SCheckBox)
+				.IsChecked(this, &FShotTrackEditor::IsCameraLocked)
+				.OnCheckStateChanged(this, &FShotTrackEditor::OnLockCameraClicked)
+				.ToolTipText(this, &FShotTrackEditor::GetLockCameraToolTip)
+				.ForegroundColor(FLinearColor::White)
+				.CheckedImage(FEditorStyle::GetBrush("Sequencer.LockCamera"))
+				.CheckedHoveredImage(FEditorStyle::GetBrush("Sequencer.LockCamera"))
+				.CheckedPressedImage(FEditorStyle::GetBrush("Sequencer.LockCamera"))
+				.UncheckedImage(FEditorStyle::GetBrush("Sequencer.UnlockCamera"))
+				.UncheckedHoveredImage(FEditorStyle::GetBrush("Sequencer.UnlockCamera"))
+				.UncheckedPressedImage(FEditorStyle::GetBrush("Sequencer.UnlockCamera"))
+		];
 
 	return EditBox.ToSharedRef();
 }
@@ -164,21 +164,9 @@ void FShotTrackEditor::Tick(float DeltaTime)
 
 bool FShotTrackEditor::AddKeyInternal( float KeyTime, const FGuid ObjectGuid )
 {
-	UMovieSceneSequence* FocusedSequence = GetSequencer()->GetFocusedMovieSceneSequence();
-	UMovieScene* MovieScene = FocusedSequence->GetMovieScene();
-	UMovieSceneTrack* Track = MovieScene->GetShotTrack();
-
-	if (!Track)
-	{
-		Track = MovieScene->AddShotTrack(UMovieSceneShotTrack::StaticClass());
-	}
-
-	UMovieSceneShotTrack* ShotTrack = CastChecked<UMovieSceneShotTrack>(Track);
-
-	// Get the current path of the movie scene the shot is in
+	UMovieSceneShotTrack* ShotTrack = FindOrCreateShotTrack();
 	const TArray<UMovieSceneSection*>& AllSections = ShotTrack->GetAllSections();
 
-	// Find where the section will be inserted
 	int32 SectionIndex = FindIndexForNewShot(AllSections, KeyTime);
 	const int32 ShotNumber = GenerateShotNumber(*ShotTrack, SectionIndex);
 	FString ShotName = FString::Printf(TEXT("Shot_%04d"), ShotNumber);
@@ -189,6 +177,20 @@ bool FShotTrackEditor::AddKeyInternal( float KeyTime, const FGuid ObjectGuid )
 }
 
 
+UMovieSceneShotTrack* FShotTrackEditor::FindOrCreateShotTrack()
+{
+	UMovieScene* FocusedMovieScene = GetFocusedMovieScene();
+	UMovieSceneTrack* ShotTrack = FocusedMovieScene->GetShotTrack();
+
+	if (ShotTrack == nullptr)
+	{
+		ShotTrack = FocusedMovieScene->AddShotTrack(UMovieSceneShotTrack::StaticClass());
+	}
+
+	return CastChecked<UMovieSceneShotTrack>(ShotTrack);
+}
+
+
 int32 FShotTrackEditor::FindIndexForNewShot(const TArray<UMovieSceneSection*>& ShotSections, float NewShotTime) const
 {
 	// @todo Sequencer: This could be a binary search. All shots should be sorted 
@@ -297,16 +299,7 @@ bool FShotTrackEditor::HandleAddShotTrackMenuEntryCanExecute() const
 
 void FShotTrackEditor::HandleAddShotTrackMenuEntryExecute()
 {
-	UMovieScene* FocusedMovieScene = GetFocusedMovieScene();
-
-	if ((FocusedMovieScene == nullptr) || (FocusedMovieScene->GetShotTrack() != nullptr))
-	{
-		return;
-	}
-
-	UMovieSceneTrack* NewTrack = FocusedMovieScene->AddShotTrack(UMovieSceneShotTrack::StaticClass());
-	ensure(NewTrack);
-
+	FindOrCreateShotTrack();
 	GetSequencer()->NotifyMovieSceneDataChanged();
 }
 
diff --git a/Engine/Source/Editor/MovieSceneTools/Private/TrackEditors/ShotTrackEditor/ShotTrackEditor.h b/Engine/Source/Editor/MovieSceneTools/Private/TrackEditors/ShotTrackEditor/ShotTrackEditor.h
--- a/Engine/Source/Editor/MovieSceneTools/Private/TrackEditors/ShotTrackEditor/ShotTrackEditor.h
+++ b/Engine/Source/Editor/MovieSceneTools/Private/TrackEditors/ShotTrackEditor/ShotTrackEditor.h
@@ -50,6 +50,9 @@ protected:
 	/** Delegate for AnimatablePropertyChanged in AddKey */
 	bool AddKeyInternal(float AutoKeyTime, const FGuid ObjectGuid);
 
+	/** Find or create a shot track in the currently focused movie scene. */
+	UMovieSceneShotTrack* FindOrCreateShotTrack();
+
 	/** Finds the index in the ShotSections array where a new shot should be inserted */
 	int32 FindIndexForNewShot(const TArray<UMovieSceneSection*>& ShotSections, float NewShotTime) const;
 
diff --git a/Engine/Source/Editor/MovieSceneTools/Private/TrackEditors/SkeletalAnimationTrackEditor.cpp b/Engine/Source/Editor/MovieSceneTools/Private/TrackEditors/SkeletalAnimationTrackEditor.cpp
--- a/Engine/Source/Editor/MovieSceneTools/Private/TrackEditors/SkeletalAnimationTrackEditor.cpp
+++ b/Engine/Source/Editor/MovieSceneTools/Private/TrackEditors/SkeletalAnimationTrackEditor.cpp
@@ -29,15 +29,12 @@ namespace SkeletalAnimationEditorConstants
 }
 
 
-FSkeletalAnimationSection::FSkeletalAnimationSection( UMovieSceneSection& InSection )
-	: Section( InSection )
-{
-}
+#define LOCTEXT_NAMESPACE "FSkeletalAnimationTrackEditor"
 
 
-FSkeletalAnimationSection::~FSkeletalAnimationSection()
-{
-}
+FSkeletalAnimationSection::FSkeletalAnimationSection( UMovieSceneSection& InSection )
+	: Section( InSection )
+{ }
 
 
 UMovieSceneSection* FSkeletalAnimationSection::GetSectionObject()
@@ -48,7 +45,7 @@ UMovieSceneSection* FSkeletalAnimationSection::GetSectionObject()
 
 FText FSkeletalAnimationSection::GetDisplayName() const
 {
-	return NSLOCTEXT("FAnimationSection", "AnimationSection", "Animation");
+	return LOCTEXT("AnimationSection", "Animation");
 }
 
 
@@ -59,7 +56,7 @@ FText FSkeletalAnimationSection::GetSectionTitle() const
 	{
 		return FText::FromString( AnimSection->GetAnimSequence()->GetName() );
 	}
-	return NSLOCTEXT("FAnimationSection", "NoAnimationSection", "No Animation");
+	return LOCTEXT("NoAnimationSection", "No Animation");
 }
 
 
@@ -120,13 +117,7 @@ int32 FSkeletalAnimationSection::OnPaintSection( const FGeometry& AllottedGeomet
 
 FSkeletalAnimationTrackEditor::FSkeletalAnimationTrackEditor( TSharedRef<ISequencer> InSequencer )
 	: FMovieSceneTrackEditor( InSequencer ) 
-{
-}
-
-
-FSkeletalAnimationTrackEditor::~FSkeletalAnimationTrackEditor()
-{
-}
+{ }
 
 
 TSharedRef<ISequencerTrackEditor> FSkeletalAnimationTrackEditor::CreateTrackEditor( TSharedRef<ISequencer> InSequencer )
@@ -156,8 +147,8 @@ void FSkeletalAnimationTrackEditor::AddKey(const FGuid& ObjectGuid, UObject* Add
 	if (AnimSequence)
 	{
 		TArray<UObject*> OutObjects;
-		GetSequencer()->GetRuntimeObjects( GetSequencer()->GetFocusedMovieSceneSequenceInstance(), ObjectGuid, OutObjects);
 
+		GetSequencer()->GetRuntimeObjects( GetSequencer()->GetFocusedMovieSceneSequenceInstance(), ObjectGuid, OutObjects);
 		AnimatablePropertyChanged( FOnKeyProperty::CreateRaw( this, &FSkeletalAnimationTrackEditor::AddKeyInternal, OutObjects, AnimSequence) );
 	}
 }
@@ -176,8 +167,8 @@ bool FSkeletalAnimationTrackEditor::HandleAssetAdded(UObject* Asset, const FGuid
 			if (Skeleton && Skeleton == AnimSequence->GetSkeleton())
 			{
 				TArray<UObject*> OutObjects;
-				GetSequencer()->GetRuntimeObjects(GetSequencer()->GetFocusedMovieSceneSequenceInstance(), TargetObjectGuid, OutObjects);
 
+				GetSequencer()->GetRuntimeObjects(GetSequencer()->GetFocusedMovieSceneSequenceInstance(), TargetObjectGuid, OutObjects);
 				AnimatablePropertyChanged(FOnKeyProperty::CreateRaw(this, &FSkeletalAnimationTrackEditor::AddKeyInternal, OutObjects, AnimSequence));
 
 				return true;
@@ -208,46 +199,51 @@ void FSkeletalAnimationTrackEditor::BuildObjectBindingTrackMenu(FMenuBuilder& Me
 			if (AssetDataList.Num())
 			{
 				MenuBuilder.AddSubMenu(
-					NSLOCTEXT("Sequencer", "AddAnimation", "Animation"), NSLOCTEXT("Sequencer", "AddAnimationTooltip", "Adds an animation track."),
-					FNewMenuDelegate::CreateRaw(this, &FSkeletalAnimationTrackEditor::BuildAnimationSubMenu, ObjectBinding, Skeleton));
+					LOCTEXT("AddAnimation", "Animation"), NSLOCTEXT("Sequencer", "AddAnimationTooltip", "Adds an animation track."),
+					FNewMenuDelegate::CreateRaw(this, &FSkeletalAnimationTrackEditor::BuildAnimationSubMenu, ObjectBinding, Skeleton)
+				);
 			}
 		}
 	}
 }
 
+
 void FSkeletalAnimationTrackEditor::BuildAnimationSubMenu(FMenuBuilder& MenuBuilder, FGuid ObjectBinding, USkeleton* Skeleton)
 {
 	FAssetPickerConfig AssetPickerConfig;
-	AssetPickerConfig.OnAssetSelected = FOnAssetSelected::CreateRaw( this, &FSkeletalAnimationTrackEditor::OnAnimationAssetSelected, ObjectBinding);
-	AssetPickerConfig.bAllowNullSelection = false;
-	AssetPickerConfig.InitialAssetViewType = EAssetViewType::List;
-
-	// Filter config
-	AssetPickerConfig.Filter.ClassNames.Add(UAnimSequence::StaticClass()->GetFName());
-	AssetPickerConfig.Filter.TagsAndValues.Add(TEXT("Skeleton"), FAssetData(Skeleton).GetExportTextName());
+	{
+		AssetPickerConfig.OnAssetSelected = FOnAssetSelected::CreateRaw( this, &FSkeletalAnimationTrackEditor::OnAnimationAssetSelected, ObjectBinding);
+		AssetPickerConfig.bAllowNullSelection = false;
+		AssetPickerConfig.InitialAssetViewType = EAssetViewType::List;
+		AssetPickerConfig.Filter.ClassNames.Add(UAnimSequence::StaticClass()->GetFName());
+		AssetPickerConfig.Filter.TagsAndValues.Add(TEXT("Skeleton"), FAssetData(Skeleton).GetExportTextName());
+	}
 
 	FContentBrowserModule& ContentBrowserModule = FModuleManager::Get().LoadModuleChecked<FContentBrowserModule>(TEXT("ContentBrowser"));
 
 	TSharedPtr<SBox> MenuEntry = SNew(SBox)
-	.WidthOverride(300.0f)
-	.HeightOverride(300.f)
-	[
-		ContentBrowserModule.Get().CreateAssetPicker(AssetPickerConfig)
-	];
+		.WidthOverride(300.0f)
+		.HeightOverride(300.f)
+		[
+			ContentBrowserModule.Get().CreateAssetPicker(AssetPickerConfig)
+		];
+
 	MenuBuilder.AddWidget(MenuEntry.ToSharedRef(), FText::GetEmpty(), true);
 }
 
+
 void FSkeletalAnimationTrackEditor::OnAnimationAssetSelected(const FAssetData& AssetData, FGuid ObjectBinding)
 {
 	UObject* SelectedObject = AssetData.GetAsset();
+
 	if (SelectedObject && SelectedObject->IsA(UAnimSequence::StaticClass()))
 	{
 		UAnimSequence* AnimSequence = CastChecked<UAnimSequence>(AssetData.GetAsset());
-
 		AddKey(ObjectBinding, AnimSequence);
 	}
 }
 
+
 bool FSkeletalAnimationTrackEditor::AddKeyInternal( float KeyTime, const TArray<UObject*> Objects, class UAnimSequence* AnimSequence )
 {
 	bool bHandleCreated = false;
@@ -263,7 +259,7 @@ bool FSkeletalAnimationTrackEditor::AddKeyInternal( float KeyTime, const TArray<
 		bHandleCreated |= HandleResult.bWasCreated;
 		if (ObjectHandle.IsValid())
 		{
-			FFindOrCreateTrackResult TrackResult = FindOrCreateTrackForObject( ObjectHandle, UMovieSceneSkeletalAnimationTrack::StaticClass(), FName( "Animation" ) );
+			FFindOrCreateTrackResult TrackResult = FindOrCreateTrackForObject(ObjectHandle, UMovieSceneSkeletalAnimationTrack::StaticClass());
 			UMovieSceneTrack* Track = TrackResult.Track;
 			bTrackCreated |= TrackResult.bWasCreated;
 
@@ -309,3 +305,6 @@ USkeleton* FSkeletalAnimationTrackEditor::AcquireSkeletonFromObjectGuid(const FG
 
 	return Skeleton;
 }
+
+
+#undef LOCTEXT_NAMESPACE
diff --git a/Engine/Source/Editor/MovieSceneTools/Private/TrackEditors/SkeletalAnimationTrackEditor.h b/Engine/Source/Editor/MovieSceneTools/Private/TrackEditors/SkeletalAnimationTrackEditor.h
--- a/Engine/Source/Editor/MovieSceneTools/Private/TrackEditors/SkeletalAnimationTrackEditor.h
+++ b/Engine/Source/Editor/MovieSceneTools/Private/TrackEditors/SkeletalAnimationTrackEditor.h
@@ -17,7 +17,7 @@ public:
 	FSkeletalAnimationTrackEditor( TSharedRef<ISequencer> InSequencer );
 
 	/** Virtual destructor. */
-	virtual ~FSkeletalAnimationTrackEditor();
+	virtual ~FSkeletalAnimationTrackEditor() { }
 
 	/**
 	 * Creates an instance of this class.  Called by a sequencer 
@@ -64,7 +64,7 @@ public:
 	FSkeletalAnimationSection( UMovieSceneSection& InSection );
 
 	/** Virtual destructor. */
-	virtual ~FSkeletalAnimationSection();
+	virtual ~FSkeletalAnimationSection() { }
 
 public:
 
diff --git a/Engine/Source/Editor/MovieSceneTools/Private/TrackEditors/SpawnTrackEditor.cpp b/Engine/Source/Editor/MovieSceneTools/Private/TrackEditors/SpawnTrackEditor.cpp
--- a/Engine/Source/Editor/MovieSceneTools/Private/TrackEditors/SpawnTrackEditor.cpp
+++ b/Engine/Source/Editor/MovieSceneTools/Private/TrackEditors/SpawnTrackEditor.cpp
@@ -8,15 +8,17 @@
 
 #define LOCTEXT_NAMESPACE "FSpawnTrackEditor"
 
+
 TSharedRef<ISequencerTrackEditor> FSpawnTrackEditor::CreateTrackEditor(TSharedRef<ISequencer> InSequencer)
 {
 	return MakeShareable(new FSpawnTrackEditor(InSequencer));
 }
 
+
 FSpawnTrackEditor::FSpawnTrackEditor(TSharedRef<ISequencer> InSequencer)
 	: FBoolPropertyTrackEditor(InSequencer)
-{
-}
+{ }
+
 
 UMovieSceneTrack* FSpawnTrackEditor::AddTrack(UMovieScene* FocusedMovieScene, const FGuid& ObjectHandle, TSubclassOf<UMovieSceneTrack> TrackClass, FName UniqueTypeName)
 {
@@ -31,6 +33,7 @@ UMovieSceneTrack* FSpawnTrackEditor::AddTrack(UMovieScene* FocusedMovieScene, co
 	return NewTrack;
 }
 
+
 void FSpawnTrackEditor::BuildObjectBindingTrackMenu(FMenuBuilder& MenuBuilder, const FGuid& ObjectBinding, const UClass* ObjectClass)
 {
 	UMovieSceneSequence* MovieSequence = GetSequencer()->GetFocusedMovieSceneSequence();
@@ -51,21 +54,25 @@ void FSpawnTrackEditor::BuildObjectBindingTrackMenu(FMenuBuilder& MenuBuilder, c
 	);
 }
 
+
 bool FSpawnTrackEditor::SupportsType(TSubclassOf<UMovieSceneTrack> Type) const
 {
 	return (Type == UMovieSceneSpawnTrack::StaticClass());
 }
 
+
 void FSpawnTrackEditor::HandleAddSpawnTrackMenuEntryExecute(FGuid ObjectBinding)
 {
 	FScopedTransaction AddSpawnTrackTransaction(LOCTEXT("AddSpawnTrack_Transaction", "Add Spawn Track"));
-	AddTrack(GetSequencer()->GetFocusedMovieSceneSequence()->GetMovieScene(), ObjectBinding, UMovieSceneSpawnTrack::StaticClass(), "Spawned");
+	AddTrack(GetSequencer()->GetFocusedMovieSceneSequence()->GetMovieScene(), ObjectBinding, UMovieSceneSpawnTrack::StaticClass(), NAME_None);
 	NotifyMovieSceneDataChanged();
 }
 
+
 bool FSpawnTrackEditor::CanAddSpawnTrack(FGuid ObjectBinding) const
 {
-	return !GetSequencer()->GetFocusedMovieSceneSequence()->GetMovieScene()->FindTrack<UMovieSceneSpawnTrack>(ObjectBinding, "Spawned");
+	return !GetSequencer()->GetFocusedMovieSceneSequence()->GetMovieScene()->FindTrack<UMovieSceneSpawnTrack>(ObjectBinding);
 }
 
+
 #undef LOCTEXT_NAMESPACE
diff --git a/Engine/Source/Editor/MovieSceneTools/Private/TrackEditors/SubTrackEditor.cpp b/Engine/Source/Editor/MovieSceneTools/Private/TrackEditors/SubTrackEditor.cpp
--- a/Engine/Source/Editor/MovieSceneTools/Private/TrackEditors/SubTrackEditor.cpp
+++ b/Engine/Source/Editor/MovieSceneTools/Private/TrackEditors/SubTrackEditor.cpp
@@ -24,8 +24,8 @@ class FSubSection
 {
 public:
 
-	FSubSection(TSharedPtr<ISequencer> InSequencer, UMovieSceneSection& InSection, FName SectionName)
-		: DisplayName(NSLOCTEXT("FSubSection", "DisplayName", "Scenes"))
+	FSubSection(TSharedPtr<ISequencer> InSequencer, UMovieSceneSection& InSection, const FText& InDisplayName)
+		: DisplayName(InDisplayName)
 		, SectionObject(*CastChecked<UMovieSceneSubSection>(&InSection))
 		, Sequencer(InSequencer)
 	{
@@ -209,12 +209,6 @@ private:
 };
 
 
-#undef LOCTEXT_NAMESPACE
-
-
-#define LOCTEXT_NAMESPACE "FSubTrackEditor"
-
-
 /* FSubTrackEditor structors
  *****************************************************************************/
 
@@ -254,7 +248,7 @@ TSharedRef<ISequencerTrackEditor> FSubTrackEditor::CreateTrackEditor(TSharedRef<
 
 TSharedRef<ISequencerSection> FSubTrackEditor::MakeSectionInterface(UMovieSceneSection& SectionObject, UMovieSceneTrack& Track)
 {
-	return MakeShareable(new FSubSection(GetSequencer(), SectionObject, Track.GetTrackName()));
+	return MakeShareable(new FSubSection(GetSequencer(), SectionObject, Track.GetDisplayName()));
 }
 
 
@@ -329,20 +323,22 @@ void FSubTrackEditor::HandleAddSubTrackMenuEntryExecute()
 		return;
 	}
 
-	const FScopedTransaction Transaction(NSLOCTEXT("Sequencer", "AddSubTrack_Transaction", "Add Sub Track"));
+	const FScopedTransaction Transaction(LOCTEXT("AddSubTrack_Transaction", "Add Sub Track"));
 	FocusedMovieScene->Modify();
 
-	UMovieSceneTrack* NewTrack = FocusedMovieScene->AddMasterTrack<UMovieSceneSubTrack>();
+	auto NewTrack = FocusedMovieScene->AddMasterTrack<UMovieSceneSubTrack>();
 	ensure(NewTrack);
 
 	GetSequencer()->NotifyMovieSceneDataChanged();
 }
 
 
-bool FSubTrackEditor::HandleSequenceAdded(float KeyTime, class UMovieSceneSequence* Sequence)
+bool FSubTrackEditor::HandleSequenceAdded(float KeyTime, UMovieSceneSequence* Sequence)
 {
 	auto SubTrack = FindOrCreateMasterTrack<UMovieSceneSubTrack>().Track;
 	SubTrack->AddSequence(*Sequence, KeyTime);
+	SubTrack->SetDisplayName(LOCTEXT("SubTrackName", "Sequences"));
+
 	return true;
 }
 
diff --git a/Engine/Source/Editor/MovieSceneTools/Private/TrackEditors/SubTrackEditor.h b/Engine/Source/Editor/MovieSceneTools/Private/TrackEditors/SubTrackEditor.h
--- a/Engine/Source/Editor/MovieSceneTools/Private/TrackEditors/SubTrackEditor.h
+++ b/Engine/Source/Editor/MovieSceneTools/Private/TrackEditors/SubTrackEditor.h
@@ -3,6 +3,9 @@
 #pragma once
 
 
+class UMovieSceneSequence;
+
+
 /**
  * Tools for animatable property types such as floats ands vectors
  */
@@ -57,5 +60,5 @@ private:
 	void HandleAddSubTrackMenuEntryExecute();
 
 	/** Callback for AnimatablePropertyChanged in HandleAssetAdded. */
-	bool HandleSequenceAdded(float KeyTime, class UMovieSceneSequence* Sequence);
+	bool HandleSequenceAdded(float KeyTime, UMovieSceneSequence* Sequence);
 };
diff --git a/Engine/Source/Editor/MovieSceneTools/Private/TrackEditors/TransformTrackEditor.cpp b/Engine/Source/Editor/MovieSceneTools/Private/TrackEditors/TransformTrackEditor.cpp
--- a/Engine/Source/Editor/MovieSceneTools/Private/TrackEditors/TransformTrackEditor.cpp
+++ b/Engine/Source/Editor/MovieSceneTools/Private/TrackEditors/TransformTrackEditor.cpp
@@ -198,7 +198,7 @@ void F3DTransformTrackCommands::RegisterCommands()
 	UI_COMMAND( AddScaleKey, "Add Scale Key", "Add a scale key at the current time for the selected actor.", EUserInterfaceActionType::Button, FInputChord(EModifierKey::Shift, EKeys::R) );
 }
 
-FName F3DTransformTrackEditor::TransformTrackName("Transform");
+FName F3DTransformTrackEditor::TransformPropertyName("Transform");
 
 F3DTransformTrackEditor::F3DTransformTrackEditor( TSharedRef<ISequencer> InSequencer )
 	: FKeyframeTrackEditor<UMovieScene3DTransformTrack, UMovieScene3DTransformSection, FTransformKey>( InSequencer ) 
@@ -378,11 +378,13 @@ void F3DTransformTrackEditor::OnTransformChanged( UObject& InObject )
 		FTransformData NewTransformData( SceneComponentThatChanged );
 
 		FKeyParams KeyParams;
-		KeyParams.bCreateHandleIfMissing = false;
-		KeyParams.bCreateTrackIfMissing = false;
-		KeyParams.bCreateKeyIfUnchanged = false;
-		KeyParams.bCreateKeyIfEmpty = false;
-		KeyParams.bCreateKeyOnlyWhenAutoKeying = true;
+		{
+			KeyParams.bCreateHandleIfMissing = false;
+			KeyParams.bCreateTrackIfMissing = false;
+			KeyParams.bCreateKeyIfUnchanged = false;
+			KeyParams.bCreateKeyIfEmpty = false;
+			KeyParams.bCreateKeyOnlyWhenAutoKeying = true;
+		}
 
 		const bool bUnwindRotation = GetSequencer()->IsRecordingLive();
 
@@ -459,7 +461,7 @@ void F3DTransformTrackEditor::BuildObjectBindingEditButtons(TSharedPtr<SHorizont
 
 	TWeakObjectPtr<ACameraActor> CameraActor;
 
-	for ( UObject* Object : OutObjects )
+	for (UObject* Object : OutObjects)
 	{
 		ACameraActor* Actor = Cast<ACameraActor>( Object );
 		if (Actor)
@@ -476,23 +478,23 @@ void F3DTransformTrackEditor::BuildObjectBindingEditButtons(TSharedPtr<SHorizont
 
 	// If this is a camera track, add a button to lock the viewport to the camera
 	EditBox.Get()->AddSlot()
-	.VAlign(VAlign_Center)
-	.HAlign(HAlign_Right)
-	.AutoWidth()
-	.Padding(4, 0, 0, 0)
-	[
-		SNew(SCheckBox)
-		.IsChecked(this, &F3DTransformTrackEditor::IsCameraLocked, CameraActor)
-		.OnCheckStateChanged(this, &F3DTransformTrackEditor::OnLockCameraClicked, CameraActor)
-		.ToolTipText(this, &F3DTransformTrackEditor::GetLockCameraToolTip, CameraActor)
-		.ForegroundColor(FLinearColor::White)
-		.CheckedImage(FEditorStyle::GetBrush("Sequencer.LockCamera"))
-		.CheckedHoveredImage(FEditorStyle::GetBrush("Sequencer.LockCamera"))
-		.CheckedPressedImage(FEditorStyle::GetBrush("Sequencer.LockCamera"))
-		.UncheckedImage(FEditorStyle::GetBrush("Sequencer.UnlockCamera"))
-		.UncheckedHoveredImage(FEditorStyle::GetBrush("Sequencer.UnlockCamera"))
-		.UncheckedPressedImage(FEditorStyle::GetBrush("Sequencer.UnlockCamera"))
-	];
+		.VAlign(VAlign_Center)
+		.HAlign(HAlign_Right)
+		.AutoWidth()
+		.Padding(4, 0, 0, 0)
+		[
+			SNew(SCheckBox)
+				.IsChecked(this, &F3DTransformTrackEditor::IsCameraLocked, ObjectGuid)
+				.OnCheckStateChanged(this, &F3DTransformTrackEditor::OnLockCameraClicked, ObjectGuid)
+				.ToolTipText(this, &F3DTransformTrackEditor::GetLockCameraToolTip, ObjectGuid)
+				.ForegroundColor(FLinearColor::White)
+				.CheckedImage(FEditorStyle::GetBrush("Sequencer.LockCamera"))
+				.CheckedHoveredImage(FEditorStyle::GetBrush("Sequencer.LockCamera"))
+				.CheckedPressedImage(FEditorStyle::GetBrush("Sequencer.LockCamera"))
+				.UncheckedImage(FEditorStyle::GetBrush("Sequencer.UnlockCamera"))
+				.UncheckedHoveredImage(FEditorStyle::GetBrush("Sequencer.UnlockCamera"))
+				.UncheckedPressedImage(FEditorStyle::GetBrush("Sequencer.UnlockCamera"))
+		];
 }
 
 
@@ -526,14 +528,32 @@ bool F3DTransformTrackEditor::CanAddTransformTrackForActorHandle( FGuid ObjectBi
 }
 
 
-ECheckBoxState F3DTransformTrackEditor::IsCameraLocked(TWeakObjectPtr<ACameraActor> CameraActor) const
+ECheckBoxState F3DTransformTrackEditor::IsCameraLocked(FGuid ObjectGuid) const
 {
-	for (int32 i = 0; i < GEditor->LevelViewportClients.Num(); ++i)
-	{		
-		FLevelEditorViewportClient* LevelVC = GEditor->LevelViewportClients[i];
-		if (LevelVC && LevelVC->IsPerspective() && LevelVC->AllowsCinematicPreview() && LevelVC->GetViewMode() != VMI_Unknown && CameraActor.IsValid() && LevelVC->IsActorLocked(CameraActor.Get()))
+	TArray<UObject*> OutObjects;
+	GetSequencer()->GetRuntimeObjects( GetSequencer()->GetFocusedMovieSceneSequenceInstance(), ObjectGuid, OutObjects);
+
+	TWeakObjectPtr<ACameraActor> CameraActor;
+
+	for (UObject* Object : OutObjects)
+	{
+		ACameraActor* Actor = Cast<ACameraActor>( Object );
+		if (Actor)
 		{
-			return ECheckBoxState::Checked;
+			CameraActor = Actor;
+			break;
+		}
+	}
+
+	if (CameraActor.IsValid())
+	{
+		for (int32 i = 0; i < GEditor->LevelViewportClients.Num(); ++i)
+		{		
+			FLevelEditorViewportClient* LevelVC = GEditor->LevelViewportClients[i];
+			if (LevelVC && LevelVC->IsPerspective() && LevelVC->AllowsCinematicPreview() && LevelVC->GetViewMode() != VMI_Unknown && CameraActor.IsValid() && LevelVC->IsActorLocked(CameraActor.Get()))
+			{
+				return ECheckBoxState::Checked;
+			}
 		}
 	}
 
@@ -541,8 +561,23 @@ ECheckBoxState F3DTransformTrackEditor::IsCameraLocked(TWeakObjectPtr<ACameraAct
 }
 
 
-void F3DTransformTrackEditor::OnLockCameraClicked(ECheckBoxState CheckBoxState, TWeakObjectPtr<ACameraActor> CameraActor)
+void F3DTransformTrackEditor::OnLockCameraClicked(ECheckBoxState CheckBoxState, FGuid ObjectGuid)
 {
+	TArray<UObject*> OutObjects;
+	GetSequencer()->GetRuntimeObjects( GetSequencer()->GetFocusedMovieSceneSequenceInstance(), ObjectGuid, OutObjects);
+
+	TWeakObjectPtr<ACameraActor> CameraActor;
+
+	for (UObject* Object : OutObjects)
+	{
+		ACameraActor* Actor = Cast<ACameraActor>( Object );
+		if (Actor)
+		{
+			CameraActor = Actor;
+			break;
+		}
+	}
+
 	// If toggle is on, lock the active viewport to the camera
 	if (CheckBoxState == ECheckBoxState::Checked)
 	{
@@ -596,11 +631,26 @@ void F3DTransformTrackEditor::OnLockCameraClicked(ECheckBoxState CheckBoxState,
 }
 
 
-FText F3DTransformTrackEditor::GetLockCameraToolTip(TWeakObjectPtr<ACameraActor> CameraActor) const
+FText F3DTransformTrackEditor::GetLockCameraToolTip(FGuid ObjectGuid) const
 {
+	TArray<UObject*> OutObjects;
+	GetSequencer()->GetRuntimeObjects( GetSequencer()->GetFocusedMovieSceneSequenceInstance(), ObjectGuid, OutObjects);
+
+	TWeakObjectPtr<ACameraActor> CameraActor;
+
+	for (UObject* Object : OutObjects)
+	{
+		ACameraActor* Actor = Cast<ACameraActor>( Object );
+		if (Actor)
+		{
+			CameraActor = Actor;
+			break;
+		}
+	}
+
 	if (CameraActor.IsValid())
 	{
-		return IsCameraLocked(CameraActor) == ECheckBoxState::Checked ?
+		return IsCameraLocked(ObjectGuid) == ECheckBoxState::Checked ?
 			FText::Format(LOCTEXT("UnlockCamera", "Unlock {0} from Viewport"), FText::FromString(CameraActor.Get()->GetActorLabel())) :
 			FText::Format(LOCTEXT("LockCamera", "Lock {0} to Selected Viewport"), FText::FromString(CameraActor.Get()->GetActorLabel()));
 	}
@@ -677,7 +727,6 @@ void F3DTransformTrackEditor::AddTransformKeysForObject( UObject* Object, EKey3D
 	GetActorAndSceneComponentFromObject( Object, Actor, SceneComponent );
 	if ( Actor != nullptr && SceneComponent != nullptr )
 	{
-		MovieSceneHelpers::SetRuntimeObjectMobility( SceneComponent );
 		FTransformData CurrentTransform( SceneComponent );
 		AddTransformKeys( Actor, FTransformData(), CurrentTransform, ChannelToKey, false, KeyParams );
 	}
@@ -702,9 +751,11 @@ bool F3DTransformTrackEditor::OnAddTransformKeys( float Time, AActor* ActorToKey
 	OnSetIntermediateValue.BindRaw( this, &F3DTransformTrackEditor::SetIntermediateValueFromTransformChange, CurrentTransform );
 
 	FOnInitializeNewTrack OnInitializeNewTrack;
-	OnInitializeNewTrack.BindLambda([this](UMovieScene3DTransformTrack* NewTrack) { NewTrack->SetPropertyNameAndPath( TransformTrackName, TransformTrackName.ToString() ); } );
+	OnInitializeNewTrack.BindLambda([this](UMovieScene3DTransformTrack* NewTrack) {
+		NewTrack->SetPropertyNameAndPath(TransformPropertyName, TransformPropertyName.ToString());
+	});
 	
-	return AddKeysToObjects( ObjectsToKey, Time, *Keys, KeyParams, UMovieScene3DTransformTrack::StaticClass(), TransformTrackName, OnInitializeNewTrack, OnSetIntermediateValue );
+	return AddKeysToObjects(ObjectsToKey, Time, *Keys, KeyParams, UMovieScene3DTransformTrack::StaticClass(), TransformPropertyName, OnInitializeNewTrack, OnSetIntermediateValue);
 }
 
 
diff --git a/Engine/Source/Editor/MovieSceneTools/Private/TrackEditors/TransformTrackEditor.h b/Engine/Source/Editor/MovieSceneTools/Private/TrackEditors/TransformTrackEditor.h
--- a/Engine/Source/Editor/MovieSceneTools/Private/TrackEditors/TransformTrackEditor.h
+++ b/Engine/Source/Editor/MovieSceneTools/Private/TrackEditors/TransformTrackEditor.h
@@ -66,13 +66,13 @@ private:
 	void OnTransformChanged( UObject& InObject );
 
 	/** Delegate for camera button lock state */
-	ECheckBoxState IsCameraLocked(TWeakObjectPtr<ACameraActor> CameraActor) const; 
+	ECheckBoxState IsCameraLocked(FGuid ObjectGuid) const; 
 
 	/** Delegate for locked camera button */
-	void OnLockCameraClicked(ECheckBoxState CheckBoxState, TWeakObjectPtr<ACameraActor> CameraActor);
+	void OnLockCameraClicked(ECheckBoxState CheckBoxState, FGuid ObjectGuid);
 
 	/** Delegate for camera button lock tooltip */
-	FText GetLockCameraToolTip(TWeakObjectPtr<ACameraActor> CameraActor) const; 
+	FText GetLockCameraToolTip(FGuid ObjectGuid) const; 
 
 	/** Tries to get an actor and root scene component from an object which can be either an actor or a component. */
 	void GetActorAndSceneComponentFromObject(UObject* Object, AActor*& OutActor, USceneComponent*& OutSceneComponent);
@@ -134,7 +134,8 @@ private:
 	void SetIntermediateValueFromTransformChange( UMovieSceneTrack* Track, FTransformData TransformData );
 
 private:
-	static FName TransformTrackName;
+
+	static FName TransformPropertyName;
 
 	/** Mapping of objects to their existing transform data (for comparing against new transform data) */
 	TMap< TWeakObjectPtr<UObject>, FTransformData > ObjectToExistingTransform;
diff --git a/Engine/Source/Editor/MovieSceneTools/Public/KeyframeTrackEditor.h b/Engine/Source/Editor/MovieSceneTools/Public/KeyframeTrackEditor.h
--- a/Engine/Source/Editor/MovieSceneTools/Public/KeyframeTrackEditor.h
+++ b/Engine/Source/Editor/MovieSceneTools/Public/KeyframeTrackEditor.h
@@ -55,14 +55,14 @@ protected:
 	 * @param NewKeys The new keys to add.
 	 * @param KeyParams The parameters to control keyframing behavior.
 	 * @param TrackClass The class of track which should be created if specified in the parameters.
-	 * @param TrackName The name which should be used to look up existing tracks, and which should be used to create new tracks.
+	 * @param PropertyName The name of the property to add keys for.
 	 * @param OnInitializeNewTrack A delegate which allows for custom initialization for new tracks.  This is called after the 
 	 *        track is created, but before any sections or keys have been added.
 	 * @param OnSetIntermediateValue A delegate which is called when the key params specify that keying should only happen when
 	 * 		  auto-keying is on, but auto-keying is off.  This allows handling the changed but not keyed intermediate value.
 	 */
 	bool AddKeysToObjects(TArray<UObject*> ObjectsToKey, float KeyTime, const TArray<KeyDataType> NewKeys, FKeyParams KeyParams,
-		TSubclassOf<UMovieSceneTrack> TrackClass, FName TrackName, FOnInitializeNewTrack OnIntializeNewTrack, FOnSetIntermediateValue OnSetIntermediateValue )
+		TSubclassOf<UMovieSceneTrack> TrackClass, FName PropertyName, FOnInitializeNewTrack OnIntializeNewTrack, FOnSetIntermediateValue OnSetIntermediateValue )
 	{
 		bool bHandleCreated = false;
 		bool bTrackCreatedOrModified = false;
@@ -75,7 +75,7 @@ protected:
 
 			if ( ObjectHandle.IsValid() )
 			{
-				bTrackCreatedOrModified |= AddKeysToHandle( ObjectHandle, KeyTime, NewKeys, KeyParams, TrackClass, TrackName, OnIntializeNewTrack, OnSetIntermediateValue );
+				bTrackCreatedOrModified |= AddKeysToHandle( ObjectHandle, KeyTime, NewKeys, KeyParams, TrackClass, PropertyName, OnIntializeNewTrack, OnSetIntermediateValue );
 			}
 		}
 		return bHandleCreated || bTrackCreatedOrModified;
@@ -84,13 +84,13 @@ protected:
 
 private:
 
-	bool AddKeysToHandle( FGuid ObjectHandle, float KeyTime, const TArray<KeyDataType> NewKeys, FKeyParams KeyParams, TSubclassOf<UMovieSceneTrack> TrackClass, FName TrackName, FOnInitializeNewTrack OnIntializeNewTrack, FOnSetIntermediateValue OnSetIntermediateValue )
+	bool AddKeysToHandle( FGuid ObjectHandle, float KeyTime, const TArray<KeyDataType> NewKeys, FKeyParams KeyParams, TSubclassOf<UMovieSceneTrack> TrackClass, FName PropertyName, FOnInitializeNewTrack OnIntializeNewTrack, FOnSetIntermediateValue OnSetIntermediateValue )
 	{
 		bool bTrackCreated = false;
 		bool bTrackModified = false;
 
 		// Try to find an existing Track, and if one doesn't exist check the key params and create one if requested.
-		FFindOrCreateTrackResult TrackResult = FindOrCreateTrackForObject( ObjectHandle, TrackClass, TrackName, KeyParams.bCreateTrackIfMissing );
+		FFindOrCreateTrackResult TrackResult = FindOrCreateTrackForObject( ObjectHandle, TrackClass, PropertyName, KeyParams.bCreateTrackIfMissing );
 		TrackType* Track = CastChecked<TrackType>( TrackResult.Track, ECastCheckedType::NullAllowed );
 		bTrackCreated = TrackResult.bWasCreated;
 
diff --git a/Engine/Source/Editor/MovieSceneTools/Public/PropertySection.h b/Engine/Source/Editor/MovieSceneTools/Public/PropertySection.h
--- a/Engine/Source/Editor/MovieSceneTools/Public/PropertySection.h
+++ b/Engine/Source/Editor/MovieSceneTools/Public/PropertySection.h
@@ -6,29 +6,45 @@
 #include "MovieSceneSection.h"
 #include "ISequencerSection.h"
 
+
 /**
  * A generic implementation for displaying simple property sections
  */
-class FPropertySection : public ISequencerSection
+class FPropertySection
+	: public ISequencerSection
 {
 public:
-	FPropertySection( UMovieSceneSection& InSectionObject, FName SectionName )
-		: SectionObject( InSectionObject )
+
+	/**
+	 * Create and initialize a new instance.
+	 *
+	 * @param The section's human readable display name.
+	 */
+	FPropertySection(UMovieSceneSection& InSectionObject, const FText& InDisplayName)
+		: DisplayName(InDisplayName)
+		, SectionObject(InSectionObject)
+	{ }
+
+public:
+
+	// ISequencerSection interface
+
+	virtual UMovieSceneSection* GetSectionObject() override
 	{
-		DisplayName = FText::FromString( FName::NameToDisplayString( SectionName.ToString(), false ) );
+		return &SectionObject;
 	}
 
-	/** ISequencerSection interface */
-	virtual UMovieSceneSection* GetSectionObject() override { return &SectionObject; }
-
 	virtual FText GetDisplayName() const override
 	{
 		return DisplayName;
 	}
 	
-	virtual FText GetSectionTitle() const override { return FText::GetEmpty(); }
+	virtual FText GetSectionTitle() const override
+	{
+		return FText::GetEmpty();
+	}
 
-	virtual void GenerateSectionLayout( class ISectionLayoutBuilder& LayoutBuilder ) const override {}
+	virtual void GenerateSectionLayout( class ISectionLayoutBuilder& LayoutBuilder ) const override { }
 
 	virtual int32 OnPaintSection( const FGeometry& AllottedGeometry, const FSlateRect& SectionClippingRect, FSlateWindowElementList& OutDrawElements, int32 LayerId, bool bParentEnabled ) const override 
 	{
@@ -63,8 +79,10 @@ public:
 	virtual void ClearIntermediateValue() = 0;
 
 protected:
+
 	/** Display name of the section */
 	FText DisplayName;
+
 	/** The section we are visualizing */
 	UMovieSceneSection& SectionObject;
 };
diff --git a/Engine/Source/Editor/NiagaraEditor/Private/NiagaraSequencer.cpp b/Engine/Source/Editor/NiagaraEditor/Private/NiagaraSequencer.cpp
--- a/Engine/Source/Editor/NiagaraEditor/Private/NiagaraSequencer.cpp
+++ b/Engine/Source/Editor/NiagaraEditor/Private/NiagaraSequencer.cpp
@@ -3,6 +3,9 @@
 #include "NiagaraSequencer.h"
 
 
+#define LOCTEXT_NAMESPACE "EmitterMovieSceneTrack"
+
+
 UEmitterMovieSceneTrack::UEmitterMovieSceneTrack(const FObjectInitializer& ObjectInitializer)
 : Super(ObjectInitializer)
 {
@@ -17,7 +20,7 @@ UNiagaraMovieSceneSection::UNiagaraMovieSceneSection(const FObjectInitializer& O
 /*
  *  This is called when the user edits a track in the effect editor timeline
  */
-void INiagaraTrackInstance::RefreshInstance(const TArray<UObject*>& RuntimeObjects, class IMovieScenePlayer& Player, FMovieSceneSequenceInstance& SequenceInstance)
+void INiagaraTrackInstance::RefreshInstance(const TArray<UObject*>& RuntimeObjects, IMovieScenePlayer& Player, FMovieSceneSequenceInstance& SequenceInstance)
 {
 	// every track should have exactly one section, always
 	check(Track->GetAllSections().Num() == 1)
@@ -29,3 +32,14 @@ void INiagaraTrackInstance::RefreshInstance(const TArray<UObject*>& RuntimeObjec
 	Track->GetEmitter()->GetProperties()->StartTime = Section->GetStartTime();
 	Track->GetEmitter()->GetProperties()->EndTime = Section->GetEndTime();
 }
+
+
+#if WITH_EDITORONLY_DATA
+FText UEmitterMovieSceneTrack::GetDisplayName() const
+{
+	return LOCTEXT("TrackName", "Emitter");
+}
+#endif
+
+
+#undef LOCTEXT_NAMESPACE
diff --git a/Engine/Source/Editor/NiagaraEditor/Private/NiagaraSequencer.h b/Engine/Source/Editor/NiagaraEditor/Private/NiagaraSequencer.h
--- a/Engine/Source/Editor/NiagaraEditor/Private/NiagaraSequencer.h
+++ b/Engine/Source/Editor/NiagaraEditor/Private/NiagaraSequencer.h
@@ -120,7 +120,7 @@ public:
 	{
 	}
 
-	virtual void Update(float Position, float LastPosition, const TArray<UObject*>& RuntimeObjects, IMovieScenePlayer& Player, FMovieSceneSequenceInstance& SequenceInstance)
+	virtual void Update(float Position, float LastPosition, const TArray<UObject*>& RuntimeObjects, IMovieScenePlayer& Player, FMovieSceneSequenceInstance& SequenceInstance, EMovieSceneUpdatePass UpdatePass)
 	{
 	}
 
@@ -146,34 +146,58 @@ class UEmitterMovieSceneTrack
 	GENERATED_UCLASS_BODY()
 public:
 
-	/** UMovieSceneTrack interface */
-	virtual FName GetTrackName() const override { return *Emitter->GetProperties()->EmitterName; }
+	TSharedPtr<FNiagaraSimulation> GetEmitter()
+	{
+		return Emitter;
+	}
+
 	void SetEmitter(TSharedPtr<FNiagaraSimulation> InEmitter)
 	{
-		if (InEmitter.IsValid())
+		if (!InEmitter.IsValid())
+		{
+			return;
+		}
+
+		Emitter = InEmitter;
+
+		if (const UNiagaraEmitterProperties* EmitterProps = InEmitter->GetProperties().Get())
 		{
-			Emitter = InEmitter;
-			if (const UNiagaraEmitterProperties* EmitterProps = InEmitter->GetProperties().Get())
+			UNiagaraMovieSceneSection *Section = NewObject<UNiagaraMovieSceneSection>(this, *Emitter->GetProperties()->EmitterName);
 			{
-				UNiagaraMovieSceneSection *Section = NewObject<UNiagaraMovieSceneSection>(this, *Emitter->GetProperties()->EmitterName);
 				Section->SetEmitterProps(EmitterProps);
-
 				Section->SetStartTime(EmitterProps->StartTime);
 				Section->SetEndTime(EmitterProps->EndTime);
 				Section->SetEmitterName(FText::FromString(EmitterProps->EmitterName));
-
-				Sections.Add(Section);
 			}
+
+			Sections.Add(Section);
 		}
 	}
 
-	virtual void RemoveAllAnimationData() override	{};
-	virtual bool HasSection(const UMovieSceneSection& Section) const override { return false; }
-	virtual bool IsEmpty() const override { return false; }
-	virtual TSharedPtr<class IMovieSceneTrackInstance> CreateInstance()
+public:
+
+	// UMovieSceneTrack interface
+
+	virtual void RemoveAllAnimationData() override
+	{
+		// do nothing
+	};
+
+	virtual bool HasSection(const UMovieSceneSection& Section) const override
+	{
+		return false;
+	}
+
+	virtual bool IsEmpty() const override
+	{
+		return false;
+	}
+
+	virtual TSharedPtr<IMovieSceneTrackInstance> CreateInstance()
 	{
 		return MakeShareable(new INiagaraTrackInstance(this));
 	}
+
 	virtual const TArray<UMovieSceneSection*>& GetAllSections() const override
 	{
 		return Sections;
@@ -184,9 +208,12 @@ public:
 		return TRange<float>(0, FLT_MAX);
 	}
 
-	TSharedPtr<FNiagaraSimulation> GetEmitter() { return Emitter; }
+#if WITH_EDITORONLY_DATA
+	virtual FText GetDisplayName() const override;
+#endif
 
 private:
+
 	TSharedPtr<FNiagaraSimulation> Emitter;
 	TArray<UMovieSceneSection*> Sections;
 };
diff --git a/Engine/Source/Editor/SceneOutliner/Private/SSceneOutliner.cpp b/Engine/Source/Editor/SceneOutliner/Private/SSceneOutliner.cpp
--- a/Engine/Source/Editor/SceneOutliner/Private/SSceneOutliner.cpp
+++ b/Engine/Source/Editor/SceneOutliner/Private/SSceneOutliner.cpp
@@ -1533,7 +1533,8 @@ namespace SceneOutliner
 		return	!SharedData->bOnlyShowFolders && 												// Don't show actors if we're only showing folders
 				Actor->IsEditable() &&															// Only show actors that are allowed to be selected and drawn in editor
 				Actor->IsListedInSceneOutliner() &&
-				(SharedData->bRepresentingPlayWorld || !Actor->HasAnyFlags(RF_Transient)) &&    // Don't show transient actors in non-play worlds
+				( (SharedData->bRepresentingPlayWorld || !Actor->HasAnyFlags(RF_Transient)) ||
+				  (SharedData->bShowTransient && Actor->HasAnyFlags(RF_Transient) ) ) &&		// Don't show transient actors in non-play worlds
 				!Actor->IsTemplate() &&															// Should never happen, but we never want CDOs displayed
 				!FActorEditorUtils::IsABuilderBrush(Actor) &&									// Don't show the builder brush
 				!Actor->IsA( AWorldSettings::StaticClass() ) &&									// Don't show the WorldSettings actor, even though it is technically editable
diff --git a/Engine/Source/Editor/SceneOutliner/Public/SceneOutlinerPublicTypes.h b/Engine/Source/Editor/SceneOutliner/Public/SceneOutlinerPublicTypes.h
--- a/Engine/Source/Editor/SceneOutliner/Public/SceneOutlinerPublicTypes.h
+++ b/Engine/Source/Editor/SceneOutliner/Public/SceneOutlinerPublicTypes.h
@@ -90,6 +90,9 @@ namespace SceneOutliner
 		/** True to only show folders in this outliner */
 		bool bOnlyShowFolders : 1;
 
+		/** Show transient objects */
+		bool bShowTransient : 1;
+
 	public:
 
 		/** Constructor */
@@ -97,6 +100,7 @@ namespace SceneOutliner
 			: Mode( ESceneOutlinerMode::ActorPicker )
 			, bShowParentTree( true )
 			, bOnlyShowFolders( false )
+			, bShowTransient( false )
 		{}
 
 		/** Set up a default array of columns for this outliner */
diff --git a/Engine/Source/Editor/Sequencer/Private/DisplayNodes/SequencerDisplayNode.cpp b/Engine/Source/Editor/Sequencer/Private/DisplayNodes/SequencerDisplayNode.cpp
--- a/Engine/Source/Editor/Sequencer/Private/DisplayNodes/SequencerDisplayNode.cpp
+++ b/Engine/Source/Editor/Sequencer/Private/DisplayNodes/SequencerDisplayNode.cpp
@@ -16,17 +16,22 @@
 #include "GroupedKeyArea.h"
 #include "ObjectEditorUtils.h"
 
+
 #define LOCTEXT_NAMESPACE "SequencerDisplayNode"
 
+
 namespace SequencerNodeConstants
 {
 	extern const float CommonPadding;
 	const float CommonPadding = 4.f;
 }
 
-class SSequencerObjectTrack : public SLeafWidget
+
+class SSequencerObjectTrack
+	: public SLeafWidget
 {
 public:
+
 	SLATE_BEGIN_ARGS(SSequencerObjectTrack) {}
 		/** The view range of the section area */
 		SLATE_ATTRIBUTE( TRange<float>, ViewRange )
@@ -46,6 +51,7 @@ public:
 	}
 
 private:
+
 	/** Collects all key times from the root node */
 	void CollectAllKeyTimes(TArray<float>& OutKeyTimes) const;
 
@@ -53,12 +59,15 @@ private:
 	void AddKeyTime(float NewTime, TArray<float>& OutKeyTimes) const;
 
 private:
+
 	/** Root node of this track view panel */
 	TSharedPtr<FSequencerDisplayNode> RootNode;
+
 	/** The current view range */
 	TAttribute< TRange<float> > ViewRange;
 };
 
+
 int32 SSequencerObjectTrack::OnPaint(const FPaintArgs& Args, const FGeometry& AllottedGeometry, const FSlateRect& MyClippingRect, FSlateWindowElementList& OutDrawElements, int32 LayerId, const FWidgetStyle& InWidgetStyle, bool bParentEnabled) const
 {
 	TArray<float> OutKeyTimes;
@@ -69,8 +78,8 @@ int32 SSequencerObjectTrack::OnPaint(const FPaintArgs& Args, const FGeometry& Al
 	for (int32 i = 0; i < OutKeyTimes.Num(); ++i)
 	{
 		float KeyPosition = TimeToPixelConverter.TimeToPixel(OutKeyTimes[i]);
-
 		static const FVector2D KeyMarkSize = FVector2D(3.f, 21.f);
+
 		FSlateDrawElement::MakeBox(
 			OutDrawElements,
 			LayerId+1,
@@ -85,15 +94,17 @@ int32 SSequencerObjectTrack::OnPaint(const FPaintArgs& Args, const FGeometry& Al
 	return LayerId+1;
 }
 
+
 FVector2D SSequencerObjectTrack::ComputeDesiredSize( float ) const
 {
 	// Note: X Size is not used
 	return FVector2D( 100.0f, RootNode->GetNodeHeight() );
 }
 
+
 void SSequencerObjectTrack::CollectAllKeyTimes(TArray<float>& OutKeyTimes) const
 {
-	TArray< TSharedRef<FSequencerSectionKeyAreaNode> > OutNodes;
+	TArray<TSharedRef<FSequencerSectionKeyAreaNode>> OutNodes;
 	RootNode->GetChildKeyAreaNodesRecursively(OutNodes);
 
 	for (int32 i = 0; i < OutNodes.Num(); ++i)
@@ -110,29 +121,29 @@ void SSequencerObjectTrack::CollectAllKeyTimes(TArray<float>& OutKeyTimes) const
 	}
 }
 
+
 void SSequencerObjectTrack::AddKeyTime(float NewTime, TArray<float>& OutKeyTimes) const
 {
 	// @todo Sequencer It might be more efficient to add each key and do the pruning at the end
-	for (int32 i = 0; i < OutKeyTimes.Num(); ++i)
+	for (float& KeyTime : OutKeyTimes)
 	{
-		if (FMath::IsNearlyEqual(OutKeyTimes[i], NewTime))
+		if (FMath::IsNearlyEqual(KeyTime, NewTime))
 		{
 			return;
 		}
 	}
+
 	OutKeyTimes.Add(NewTime);
 }
 
 
-
 FSequencerDisplayNode::FSequencerDisplayNode( FName InNodeName, TSharedPtr<FSequencerDisplayNode> InParentNode, FSequencerNodeTree& InParentTree )
 	: ParentNode( InParentNode )
 	, ParentTree( InParentTree )
 	, NodeName( InNodeName )
 	, bExpanded( false )
-{
-	
-}
+{ }
+
 
 void FSequencerDisplayNode::Initialize(float InVirtualTop, float InVirtualBottom)
 {
@@ -142,13 +153,14 @@ void FSequencerDisplayNode::Initialize(float InVirtualTop, float InVirtualBottom
 	VirtualBottom = InVirtualBottom;
 }
 
+
 void FSequencerDisplayNode::AddObjectBindingNode(TSharedRef<FSequencerObjectBindingNode> ObjectBindingNode)
 {
 	ChildNodes.Add(ObjectBindingNode);
 }
 
 
-bool FSequencerDisplayNode::Traverse_ChildFirst(TFunctionRef<bool(FSequencerDisplayNode&)> InPredicate, bool bIncludeThisNode)
+bool FSequencerDisplayNode::Traverse_ChildFirst(const TFunctionRef<bool(FSequencerDisplayNode&)>& InPredicate, bool bIncludeThisNode)
 {
 	for (auto& Child : GetChildNodes())
 	{
@@ -161,7 +173,8 @@ bool FSequencerDisplayNode::Traverse_ChildFirst(TFunctionRef<bool(FSequencerDisp
 	return bIncludeThisNode ? InPredicate(*this) : true;
 }
 
-bool FSequencerDisplayNode::Traverse_ParentFirst(TFunctionRef<bool(FSequencerDisplayNode&)> InPredicate, bool bIncludeThisNode)
+
+bool FSequencerDisplayNode::Traverse_ParentFirst(const TFunctionRef<bool(FSequencerDisplayNode&)>& InPredicate, bool bIncludeThisNode)
 {
 	if (bIncludeThisNode && !InPredicate(*this))
 	{
@@ -179,7 +192,8 @@ bool FSequencerDisplayNode::Traverse_ParentFirst(TFunctionRef<bool(FSequencerDis
 	return true;
 }
 
-bool FSequencerDisplayNode::TraverseVisible_ChildFirst(TFunctionRef<bool(FSequencerDisplayNode&)> InPredicate, bool bIncludeThisNode)
+
+bool FSequencerDisplayNode::TraverseVisible_ChildFirst(const TFunctionRef<bool(FSequencerDisplayNode&)>& InPredicate, bool bIncludeThisNode)
 {
 	// If the item is not expanded, its children ain't visible
 	if (IsExpanded())
@@ -203,7 +217,7 @@ bool FSequencerDisplayNode::TraverseVisible_ChildFirst(TFunctionRef<bool(FSequen
 }
 
 
-bool FSequencerDisplayNode::TraverseVisible_ParentFirst(TFunctionRef<bool(FSequencerDisplayNode&)> InPredicate, bool bIncludeThisNode)
+bool FSequencerDisplayNode::TraverseVisible_ParentFirst(const TFunctionRef<bool(FSequencerDisplayNode&)>& InPredicate, bool bIncludeThisNode)
 {
 	if (bIncludeThisNode && !IsHidden() && !InPredicate(*this))
 	{
@@ -225,111 +239,127 @@ bool FSequencerDisplayNode::TraverseVisible_ParentFirst(TFunctionRef<bool(FSeque
 	return true;
 }
 
+
 TSharedRef<FSequencerSectionCategoryNode> FSequencerDisplayNode::AddCategoryNode( FName CategoryName, const FText& DisplayLabel )
 {
 	TSharedPtr<FSequencerSectionCategoryNode> CategoryNode;
 
 	// See if there is an already existing category node to use
-	for( int32 ChildIndex = 0; ChildIndex < ChildNodes.Num(); ++ChildIndex )
+	for (const TSharedRef<FSequencerDisplayNode>& Node : ChildNodes)
 	{
-		TSharedRef<FSequencerDisplayNode>& ChildNode = ChildNodes[ChildIndex];
-		if( ChildNode->GetNodeName() == CategoryName && ChildNode->GetType() == ESequencerNode::Category )
+		if ((Node->GetNodeName() == CategoryName) && (Node->GetType() == ESequencerNode::Category))
 		{
-			CategoryNode = StaticCastSharedRef<FSequencerSectionCategoryNode>( ChildNode );
+			CategoryNode = StaticCastSharedRef<FSequencerSectionCategoryNode>(Node);
 		}
 	}
 
-	if( !CategoryNode.IsValid() )
+	if (!CategoryNode.IsValid())
 	{
 		// No existing category found, make a new one
-		CategoryNode = MakeShareable( new FSequencerSectionCategoryNode( CategoryName, DisplayLabel, SharedThis( this ), ParentTree ) );
-		ChildNodes.Add( CategoryNode.ToSharedRef() );
+		CategoryNode = MakeShareable(new FSequencerSectionCategoryNode(CategoryName, DisplayLabel, SharedThis(this), ParentTree));
+		ChildNodes.Add(CategoryNode.ToSharedRef());
 	}
 
 	return CategoryNode.ToSharedRef();
 }
 
-TSharedRef<FTrackNode> FSequencerDisplayNode::AddSectionAreaNode( FName SectionName, UMovieSceneTrack& AssociatedTrack, ISequencerTrackEditor& AssociatedEditor )
+
+TSharedRef<FSequencerTrackNode> FSequencerDisplayNode::AddSectionAreaNode(UMovieSceneTrack& AssociatedTrack, ISequencerTrackEditor& AssociatedEditor)
 {
-	TSharedPtr<FTrackNode> SectionNode;
+	TSharedPtr<FSequencerTrackNode> SectionNode;
 
-	// See if there is an already existing section node to use
-	for( int32 ChildIndex = 0; ChildIndex < ChildNodes.Num(); ++ChildIndex )
+	// see if there is an already existing section node to use
+	for (const TSharedRef<FSequencerDisplayNode>& Node : ChildNodes)
 	{
-		TSharedRef<FSequencerDisplayNode>& ChildNode = ChildNodes[ChildIndex];
-		if( ChildNode->GetNodeName() == SectionName && ChildNode->GetType() == ESequencerNode::Track )
+		if (Node->GetType() != ESequencerNode::Track)
+		{
+			continue;
+		}
+
+		SectionNode = StaticCastSharedRef<FSequencerTrackNode>(Node);
+
+		if (SectionNode->GetTrack() != &AssociatedTrack)
+		{
+			SectionNode.Reset();
+		}
+		else
 		{
-			SectionNode = StaticCastSharedRef<FTrackNode>( ChildNode );
+			break;
 		}
 	}
 
-	if( !SectionNode.IsValid() )
+	if (!SectionNode.IsValid())
 	{
 		// No existing node found make a new one
-		SectionNode = MakeShareable( new FTrackNode( SectionName, AssociatedTrack, AssociatedEditor, SharedThis( this ), ParentTree ) );
+		SectionNode = MakeShareable(new FSequencerTrackNode(AssociatedTrack, AssociatedEditor, SharedThis(this), ParentTree));
 		ChildNodes.Add( SectionNode.ToSharedRef() );
 	}
 
 	// The section node type has to match
-	check( SectionNode->GetTrack() == &AssociatedTrack );
+	check(SectionNode->GetTrack() == &AssociatedTrack);
 
 	return SectionNode.ToSharedRef();
 }
 
-void FSequencerDisplayNode::AddKeyAreaNode( FName KeyAreaName, const FText& DisplayName, TSharedRef<IKeyArea> KeyArea )
+
+void FSequencerDisplayNode::AddKeyAreaNode(FName KeyAreaName, const FText& DisplayName, TSharedRef<IKeyArea> KeyArea)
 {
 	TSharedPtr<FSequencerSectionKeyAreaNode> KeyAreaNode;
 
 	// See if there is an already existing key area node to use
-	for( int32 ChildIndex = 0; ChildIndex < ChildNodes.Num(); ++ChildIndex )
+	for (const TSharedRef<FSequencerDisplayNode>& Node : ChildNodes)
 	{
-		TSharedRef<FSequencerDisplayNode>& ChildNode = ChildNodes[ChildIndex];
-		if( ChildNode->GetNodeName() == KeyAreaName && ChildNode->GetType() == ESequencerNode::KeyArea )
+		if ((Node->GetNodeName() == KeyAreaName) && (Node->GetType() == ESequencerNode::KeyArea))
 		{
-			KeyAreaNode = StaticCastSharedRef<FSequencerSectionKeyAreaNode>( ChildNode );
+			KeyAreaNode = StaticCastSharedRef<FSequencerSectionKeyAreaNode>(Node);
 		}
 	}
 
-	if( !KeyAreaNode.IsValid() )
+	if (!KeyAreaNode.IsValid())
 	{
 		// No existing node found make a new one
 		KeyAreaNode = MakeShareable( new FSequencerSectionKeyAreaNode( KeyAreaName, DisplayName, SharedThis( this ), ParentTree ) );
-		ChildNodes.Add( KeyAreaNode.ToSharedRef() );
+		ChildNodes.Add(KeyAreaNode.ToSharedRef());
 	}
 
-	KeyAreaNode->AddKeyArea( KeyArea );
+	KeyAreaNode->AddKeyArea(KeyArea);
 }
 
+
 TSharedRef<SWidget> FSequencerDisplayNode::GenerateContainerWidgetForOutliner(const TSharedRef<SSequencerTreeViewRow>& InRow)
 {
-	return SNew( SAnimationOutlinerTreeNode, SharedThis( this ), InRow );
+	auto NewWidget = SNew(SAnimationOutlinerTreeNode, SharedThis(this), InRow);
+	TreeNodeWidgetPtr = NewWidget;
+
+	return NewWidget;
 }
 
+
 TSharedRef<SWidget> FSequencerDisplayNode::GenerateEditWidgetForOutliner()
 {
 	return SNew(SSpacer);
 }
 
-TSharedRef<SWidget> FSequencerDisplayNode::GenerateWidgetForSectionArea( const TAttribute< TRange<float> >& ViewRange )
+
+TSharedRef<SWidget> FSequencerDisplayNode::GenerateWidgetForSectionArea(const TAttribute< TRange<float> >& ViewRange)
 {
-	if( GetType() == ESequencerNode::Track )
+	if (GetType() == ESequencerNode::Track)
 	{
-		return 
-			SNew( SSequencerSectionAreaView, SharedThis( this ) )
-			.ViewRange( ViewRange );
+		return SNew(SSequencerSectionAreaView, SharedThis(this))
+			.ViewRange(ViewRange);
 	}
-	else if (GetType() == ESequencerNode::Object)
+	
+	if (GetType() == ESequencerNode::Object)
 	{
 		return SNew(SSequencerObjectTrack, SharedThis(this))
-			.ViewRange( ViewRange );
-	}
-	else
-	{
-		// Currently only section areas display widgets
-		return SNullWidget::NullWidget;
+			.ViewRange(ViewRange);
 	}
+
+	// currently only section areas display widgets
+	return SNullWidget::NullWidget;
 }
 
+
 TSharedPtr<FSequencerDisplayNode> FSequencerDisplayNode::GetSectionAreaAuthority() const
 {
 	TSharedPtr<FSequencerDisplayNode> Authority = SharedThis(const_cast<FSequencerDisplayNode*>(this));
@@ -345,14 +375,17 @@ TSharedPtr<FSequencerDisplayNode> FSequencerDisplayNode::GetSectionAreaAuthority
 			Authority = Authority->GetParent();
 		}
 	}
+
 	return Authority;
 }
 
+
 FString FSequencerDisplayNode::GetPathName() const
 {
 	// First get our parent's path
 	FString PathName;
-	if( ParentNode.IsValid() )
+
+	if (ParentNode.IsValid())
 	{
 		PathName = ParentNode.Pin()->GetPathName() + TEXT(".");
 	}
@@ -363,58 +396,77 @@ FString FSequencerDisplayNode::GetPathName() const
 	return PathName;
 }
 
+
 TSharedPtr<SWidget> FSequencerDisplayNode::OnSummonContextMenu(const FGeometry& MyGeometry, const FPointerEvent& MouseEvent)
 {
 	// @todo sequencer replace with UI Commands instead of faking it
 	const bool bShouldCloseWindowAfterMenuSelection = true;
-	FMenuBuilder MenuBuilder(bShouldCloseWindowAfterMenuSelection, NULL);
+	FMenuBuilder MenuBuilder(bShouldCloseWindowAfterMenuSelection, nullptr);
 
 	BuildContextMenu(MenuBuilder);
 
 	return MenuBuilder.MakeWidget();
 }
 
+
 void FSequencerDisplayNode::BuildContextMenu(FMenuBuilder& MenuBuilder)
 {
-	TSharedRef<FSequencerDisplayNode> NodeToBeDeleted = SharedThis(this);
+	TSharedRef<FSequencerDisplayNode> ThisNode = SharedThis(this);
 
 	MenuBuilder.AddMenuEntry(
 		LOCTEXT("DeleteNode", "Delete"),
 		LOCTEXT("DeleteNodeTooltip", "Delete this or selected nodes"),
 		FSlateIcon(),
-		FUIAction(FExecuteAction::CreateSP(&GetSequencer(), &FSequencer::DeleteNode, NodeToBeDeleted)));
+		FUIAction(FExecuteAction::CreateSP(&GetSequencer(), &FSequencer::DeleteNode, ThisNode))
+	);
+
+	MenuBuilder.AddMenuEntry(
+		LOCTEXT("RenameNode", "Rename"),
+		LOCTEXT("RenameNodeTooltip", "Rename this node"),
+		FSlateIcon(),
+		FUIAction(
+			FExecuteAction::CreateSP(this, &FSequencerDisplayNode::HandleContextMenuRenameNodeExecute, ThisNode),
+			FCanExecuteAction::CreateSP(this, &FSequencerDisplayNode::HandleContextMenuRenameNodeCanExecute, ThisNode)
+		)
+	);
 }
 
+
 void FSequencerDisplayNode::GetChildKeyAreaNodesRecursively(TArray< TSharedRef<FSequencerSectionKeyAreaNode> >& OutNodes) const
 {
-	for (int32 i = 0; i < ChildNodes.Num(); ++i)
+	for (const TSharedRef<FSequencerDisplayNode>& Node : ChildNodes)
 	{
-		if (ChildNodes[i]->GetType() == ESequencerNode::KeyArea)
+		if (Node->GetType() == ESequencerNode::KeyArea)
 		{
-			OutNodes.Add(StaticCastSharedRef<FSequencerSectionKeyAreaNode>(ChildNodes[i]));
+			OutNodes.Add(StaticCastSharedRef<FSequencerSectionKeyAreaNode>(Node));
 		}
-		ChildNodes[i]->GetChildKeyAreaNodesRecursively(OutNodes);
+
+		Node->GetChildKeyAreaNodesRecursively(OutNodes);
 	}
 }
 
+
 void FSequencerDisplayNode::SetExpansionState(bool bInExpanded)
 {
 	bExpanded = bInExpanded;
 
 	// Expansion state has changed, save it to the movie scene now
-	ParentTree.SaveExpansionState( *this, bExpanded );
+	ParentTree.SaveExpansionState(*this, bExpanded);
 }
 
+
 bool FSequencerDisplayNode::IsExpanded() const
 {
 	return ParentTree.HasActiveFilter() ? true : bExpanded;
 }
 
+
 bool FSequencerDisplayNode::IsHidden() const
 {
 	return ParentTree.HasActiveFilter() && !ParentTree.IsNodeFiltered(AsShared());
 }
 
+
 TSharedRef<FGroupedKeyArea> FSequencerDisplayNode::GetKeyGrouping(int32 InSectionIndex)
 {
 	if (!KeyGroupings.IsValidIndex(InSectionIndex))
@@ -422,14 +474,17 @@ TSharedRef<FGroupedKeyArea> FSequencerDisplayNode::GetKeyGrouping(int32 InSectio
 		KeyGroupings.SetNum(InSectionIndex + 1);
 	}
 
-	if (!KeyGroupings[InSectionIndex].IsValid())
+	auto& KeyGroup = KeyGroupings[InSectionIndex];
+
+	if (!KeyGroup.IsValid())
 	{
-		KeyGroupings[InSectionIndex] = MakeShareable(new FGroupedKeyArea(*this, InSectionIndex));
+		KeyGroup = MakeShareable(new FGroupedKeyArea(*this, InSectionIndex));
 	}
 
-	return KeyGroupings[InSectionIndex].ToSharedRef();
+	return KeyGroup.ToSharedRef();
 }
 
+
 TSharedRef<FGroupedKeyArea> FSequencerDisplayNode::UpdateKeyGrouping(int32 InSectionIndex)
 {
 	if (!KeyGroupings.IsValidIndex(InSectionIndex))
@@ -437,16 +492,35 @@ TSharedRef<FGroupedKeyArea> FSequencerDisplayNode::UpdateKeyGrouping(int32 InSec
 		KeyGroupings.SetNum(InSectionIndex + 1);
 	}
 
-	if (!KeyGroupings[InSectionIndex].IsValid())
+	auto& KeyGroup = KeyGroupings[InSectionIndex];
+
+	if (!KeyGroup.IsValid())
 	{
-		KeyGroupings[InSectionIndex] = MakeShareable(new FGroupedKeyArea(*this, InSectionIndex));
+		KeyGroup = MakeShareable(new FGroupedKeyArea(*this, InSectionIndex));
 	}
 	else
 	{
-		*KeyGroupings[InSectionIndex] = FGroupedKeyArea(*this, InSectionIndex);
+		*KeyGroup = FGroupedKeyArea(*this, InSectionIndex);
 	}
 
-	return KeyGroupings[InSectionIndex].ToSharedRef();
+	return KeyGroup.ToSharedRef();
+}
+
+
+void FSequencerDisplayNode::HandleContextMenuRenameNodeExecute(TSharedRef<FSequencerDisplayNode> NodeToBeRenamed)
+{
+	TSharedPtr<SAnimationOutlinerTreeNode> TreeNodeWidget = TreeNodeWidgetPtr.Pin();
+
+	if (TreeNodeWidget.IsValid())
+	{
+		TreeNodeWidget->EnterRenameMode();
+	}
+}
+
+
+bool FSequencerDisplayNode::HandleContextMenuRenameNodeCanExecute(TSharedRef<FSequencerDisplayNode> NodeToBeRenamed) const
+{
+	return NodeToBeRenamed->CanRenameNode();
 }
 
 
diff --git a/Engine/Source/Editor/Sequencer/Private/DisplayNodes/SequencerDisplayNode.h b/Engine/Source/Editor/Sequencer/Private/DisplayNodes/SequencerDisplayNode.h
--- a/Engine/Source/Editor/Sequencer/Private/DisplayNodes/SequencerDisplayNode.h
+++ b/Engine/Source/Editor/Sequencer/Private/DisplayNodes/SequencerDisplayNode.h
@@ -2,8 +2,10 @@
 
 #pragma once
 
+
 class IKeyArea;
 class ISequencerSection;
+class SAnimationOutlinerTreeNode;
 class SSequencerTreeViewRow;
 class FGroupedKeyArea;
 
@@ -13,14 +15,18 @@ class FGroupedKeyArea;
  */
 struct FNodePadding
 {
-	FNodePadding(float InUniform) : Top(InUniform), Bottom(InUniform) {}
-	FNodePadding(float InTop, float InBottom) : Top(InTop), Bottom(InBottom) {}
+	FNodePadding(float InUniform) : Top(InUniform), Bottom(InUniform) { }
+	FNodePadding(float InTop, float InBottom) : Top(InTop), Bottom(InBottom) { }
 	
 	/** @return The sum total of the separate padding values */
-	float Combined() const { return Top + Bottom; }
+	float Combined() const
+	{
+		return Top + Bottom;
+	}
 
 	/** Padding to be applied to the top of the node */
 	float Top;
+
 	/** Padding to be applied to the bottom of the node */
 	float Bottom;
 };
@@ -49,16 +55,21 @@ class FSequencerDisplayNode
 	: public TSharedFromThis<FSequencerDisplayNode>
 {
 public:
-	virtual ~FSequencerDisplayNode(){}
+
 	/**
-	 * Constructor
+	 * Create and initialize a new instance.
 	 * 
 	 * @param InNodeName	The name identifier of then node
-	 * @param InParentNode	The parent of this node or NULL if this is a root node
+	 * @param InParentNode	The parent of this node or nullptr if this is a root node
 	 * @param InParentTree	The tree this node is in
 	 */
 	FSequencerDisplayNode( FName InNodeName, TSharedPtr<FSequencerDisplayNode> InParentNode, FSequencerNodeTree& InParentTree );
 
+	/** Virtual destructor. */
+	virtual ~FSequencerDisplayNode(){}
+
+public:
+
 	/**
 	* Adds an object binding node to this node.
 	*
@@ -77,11 +88,10 @@ public:
 	/**
 	 * Adds a new section area for this node.
 	 * 
-	 * @param SectionName		Name of the section area
 	 * @param AssociatedTrack	The track associated with sections in this node
 	 * @param AssociatedEditor	The track editor for the associated track
 	 */
-	TSharedRef<class FTrackNode> AddSectionAreaNode( FName SectionName, UMovieSceneTrack& AssociatedTrack, ISequencerTrackEditor& AssociatedEditor );
+	TSharedRef<class FSequencerTrackNode> AddSectionAreaNode(UMovieSceneTrack& AssociatedTrack, ISequencerTrackEditor& AssociatedEditor );
 
 	/**
 	 * Adds a key area to this node
@@ -98,7 +108,10 @@ public:
 	virtual ESequencerNode::Type GetType() const = 0;
 
 	/** @return Whether or not this node can be selected */
-	virtual bool IsSelectable() const { return true; }
+	virtual bool IsSelectable() const
+	{
+		return true;
+	}
 
 	/**
 	 * @return The desired height of the node when displayed
@@ -110,11 +123,25 @@ public:
 	 */
 	virtual FNodePadding GetNodePadding() const = 0;
 
+	/**
+	 * Whether the node can be renamed.
+	 *
+	 * @return true if this node can be renamed, false otherwise.
+	 */
+	virtual bool CanRenameNode() const = 0;
+
 	/**
 	 * @return The localized display name of this node
 	 */
 	virtual FText GetDisplayName() const = 0;
 
+	/**
+	 * Set the node's display name.
+	 *
+	 * @param DisplayName the display name to set.
+	 */
+	virtual void SetDisplayName(const FText& DisplayName) = 0;
+
 	/**
 	 * Generates a container widget for tree display in the animation outliner portion of the track area
 	 * 
@@ -135,7 +162,7 @@ public:
 	 * @param ViewRange	The range of time in the sequencer that we are displaying
 	 * @return Generated outliner widget
 	 */
-	virtual TSharedRef<SWidget> GenerateWidgetForSectionArea( const TAttribute< TRange<float> >& ViewRange );
+	virtual TSharedRef<SWidget> GenerateWidgetForSectionArea( const TAttribute<TRange<float>>& ViewRange );
 
 	/**
 	 * Get the display node that is ultimately responsible for constructing a section area widget for this node.
@@ -157,49 +184,62 @@ public:
 	/**
 	 * @return The name of the node (for identification purposes)
 	 */
-	FName GetNodeName() const { return NodeName; }
+	FName GetNodeName() const
+	{
+		return NodeName;
+	}
 
 	/**
 	 * @return The number of child nodes belonging to this node
 	 */
-	uint32 GetNumChildren() const { return ChildNodes.Num(); }
+	uint32 GetNumChildren() const
+	{
+		return ChildNodes.Num();
+	}
 
 	/**
 	 * @return A List of all Child nodes belonging to this node
 	 */
-	const TArray< TSharedRef<FSequencerDisplayNode> >& GetChildNodes() const { return ChildNodes; }
+	const TArray<TSharedRef<FSequencerDisplayNode>>& GetChildNodes() const
+	{
+		return ChildNodes;
+	}
 
 	/**
 	 * Iterate this entire node tree, child first.
+	 *
 	 * @param 	InPredicate			Predicate to call for each node, returning whether to continue iteration or not
 	 * @param 	bIncludeThisNode	Whether to include this node in the iteration, or just children
 	 * @return  true where the client prematurely exited the iteration, false otherwise
 	 */
-	bool Traverse_ChildFirst(TFunctionRef<bool(FSequencerDisplayNode&)> InPredicate, bool bIncludeThisNode = true);
+	bool Traverse_ChildFirst(const TFunctionRef<bool(FSequencerDisplayNode&)>& InPredicate, bool bIncludeThisNode = true);
 
 	/**
 	 * Iterate this entire node tree, parent first.
+	 *
 	 * @param 	InPredicate			Predicate to call for each node, returning whether to continue iteration or not
 	 * @param 	bIncludeThisNode	Whether to include this node in the iteration, or just children
 	 * @return  true where the client prematurely exited the iteration, false otherwise
 	 */
-	bool Traverse_ParentFirst(TFunctionRef<bool(FSequencerDisplayNode&)> InPredicate, bool bIncludeThisNode = true);
+	bool Traverse_ParentFirst(const TFunctionRef<bool(FSequencerDisplayNode&)>& InPredicate, bool bIncludeThisNode = true);
 
 	/**
 	 * Iterate any visible portions of this node's sub-tree, child first.
+	 *
 	 * @param 	InPredicate			Predicate to call for each node, returning whether to continue iteration or not
 	 * @param 	bIncludeThisNode	Whether to include this node in the iteration, or just children
 	 * @return  true where the client prematurely exited the iteration, false otherwise
 	 */
-	bool TraverseVisible_ChildFirst(TFunctionRef<bool(FSequencerDisplayNode&)> InPredicate, bool bIncludeThisNode = true);
+	bool TraverseVisible_ChildFirst(const TFunctionRef<bool(FSequencerDisplayNode&)>& InPredicate, bool bIncludeThisNode = true);
 
 	/**
 	 * Iterate any visible portions of this node's sub-tree, parent first.
+	 *
 	 * @param 	InPredicate			Predicate to call for each node, returning whether to continue iteration or not
 	 * @param 	bIncludeThisNode	Whether to include this node in the iteration, or just children
 	 * @return  true where the client prematurely exited the iteration, false otherwise
 	 */
-	bool TraverseVisible_ParentFirst(TFunctionRef<bool(FSequencerDisplayNode&)> InPredicate, bool bIncludeThisNode = true);
+	bool TraverseVisible_ParentFirst(const TFunctionRef<bool(FSequencerDisplayNode&)>& InPredicate, bool bIncludeThisNode = true);
 
 	/**
 	 * Sorts the child nodes with the supplied predicate
@@ -208,6 +248,7 @@ public:
 	void SortChildNodes(const PREDICATE_CLASS& Predicate)
 	{
 		ChildNodes.StableSort(Predicate);
+
 		for (TSharedRef<FSequencerDisplayNode> ChildNode : ChildNodes)
 		{
 			ChildNode->SortChildNodes(Predicate);
@@ -217,16 +258,25 @@ public:
 	/**
 	 * @return The parent of this node                                                              
 	 */
-	TSharedPtr<FSequencerDisplayNode> GetParent() const { return ParentNode.Pin(); }
+	TSharedPtr<FSequencerDisplayNode> GetParent() const
+	{
+		return ParentNode.Pin();
+	}
 	
 	/** Gets the sequencer that owns this node */
-	FSequencer& GetSequencer() const { return ParentTree.GetSequencer(); }
+	FSequencer& GetSequencer() const
+	{
+		return ParentTree.GetSequencer();
+	}
 	
 	/** Gets the parent tree that this node is in */
-	FSequencerNodeTree& GetParentTree() const { return ParentTree; }
+	FSequencerNodeTree& GetParentTree() const
+	{
+		return ParentTree;
+	}
 
 	/** Gets all the key area nodes recursively, including this node if applicable */
-	virtual void GetChildKeyAreaNodesRecursively(TArray< TSharedRef<class FSequencerSectionKeyAreaNode> >& OutNodes) const;
+	virtual void GetChildKeyAreaNodesRecursively(TArray<TSharedRef<class FSequencerSectionKeyAreaNode>>& OutNodes) const;
 
 	/**
 	 * Set whether this node is expanded or not
@@ -247,10 +297,16 @@ public:
 	void Initialize(float InVirtualTop, float InVirtualBottom);
 
 	/** @return this node's virtual offset from the top of the tree, irrespective of expansion states */
-	float GetVirtualTop() const { return VirtualTop; }
+	float GetVirtualTop() const
+	{
+		return VirtualTop;
+	}
 	
 	/** @return this node's virtual offset plus its virtual height, irrespective of expansion states */
-	float GetVirtualBottom() const { return VirtualBottom; }
+	float GetVirtualBottom() const
+	{
+		return VirtualBottom;
+	}
 
 	/** Get the key grouping for the specified section index, ensuring it is fully up to date */
 	TSharedRef<FGroupedKeyArea> UpdateKeyGrouping(int32 InSectionIndex);
@@ -258,21 +314,29 @@ public:
 	/** Get the key grouping for the specified section index */
 	TSharedRef<FGroupedKeyArea> GetKeyGrouping(int32 InSectionIndex);
 
+private:
+
+	/** Callback for executing a "Rename Node" context menu action. */
+	void HandleContextMenuRenameNodeExecute(TSharedRef<FSequencerDisplayNode> NodeToBeRenamed);
+
+	/** Callback for determining whether a "Rename Node" context menu action can execute. */
+	bool HandleContextMenuRenameNodeCanExecute(TSharedRef<FSequencerDisplayNode> NodeToBeRenamed) const;
+
 protected:
 
-	/** The virtual offset of this item from the top of the tree, irrespective of expansion states */
+	/** The virtual offset of this item from the top of the tree, irrespective of expansion states. */
 	float VirtualTop;
 
-	/** The virtual offset + virtual heightof this item, irrespective of expansion states */
+	/** The virtual offset + virtual height of this item, irrespective of expansion states. */
 	float VirtualBottom;
 
 protected:
 
 	/** The parent of this node*/
-	TWeakPtr< FSequencerDisplayNode > ParentNode;
+	TWeakPtr<FSequencerDisplayNode > ParentNode;
 
 	/** List of children belonging to this node */
-	TArray< TSharedRef<FSequencerDisplayNode > > ChildNodes;
+	TArray<TSharedRef<FSequencerDisplayNode>> ChildNodes;
 
 	/** Parent tree that this node is in */
 	FSequencerNodeTree& ParentTree;
@@ -284,5 +348,8 @@ protected:
 	bool bExpanded;
 
 	/** Transient grouped keys for this node */
-	TArray< TSharedPtr<FGroupedKeyArea> > KeyGroupings;
+	TArray<TSharedPtr<FGroupedKeyArea>> KeyGroupings;
+
+	/** The outliner tree node widget that belongs to this display node. */
+	TWeakPtr<SAnimationOutlinerTreeNode> TreeNodeWidgetPtr;
 };
diff --git a/Engine/Source/Editor/Sequencer/Private/DisplayNodes/SequencerObjectBindingNode.cpp b/Engine/Source/Editor/Sequencer/Private/DisplayNodes/SequencerObjectBindingNode.cpp
--- a/Engine/Source/Editor/Sequencer/Private/DisplayNodes/SequencerObjectBindingNode.cpp
+++ b/Engine/Source/Editor/Sequencer/Private/DisplayNodes/SequencerObjectBindingNode.cpp
@@ -5,9 +5,11 @@
 #include "Sequencer.h"
 #include "MovieScene.h"
 #include "MovieSceneSequence.h"
+#include "MovieSceneTrack.h"
 #include "ObjectEditorUtils.h"
 #include "AssetEditorManager.h"
 
+
 #define LOCTEXT_NAMESPACE "FObjectBindingNode"
 
 
@@ -16,94 +18,50 @@
 	extern const float CommonPadding;
 }
 
-FText FSequencerObjectBindingNode::GetDisplayName() const
+
+void GetKeyablePropertyPaths(UClass* Class, UStruct* PropertySource, TArray<UProperty*>& PropertyPath, FSequencer& Sequencer, TArray<TArray<UProperty*>>& KeyablePropertyPaths)
 {
-	FText DisplayName;
-	return GetSequencer().GetFocusedMovieSceneSequence()->TryGetObjectDisplayName(ObjectBinding, DisplayName) ?
-		DisplayName : DefaultDisplayName;
+	//@todo need to resolve this between UMG and the level editor sequencer
+	const bool bRecurseAllProperties = Sequencer.IsLevelEditorSequencer();
+
+	for (TFieldIterator<UProperty> PropertyIterator(PropertySource); PropertyIterator; ++PropertyIterator)
+	{
+		UProperty* Property = *PropertyIterator;
+
+		if (Property && !Property->HasAnyPropertyFlags(CPF_Deprecated))
+		{
+			PropertyPath.Add(Property);
+
+			bool bIsPropertyKeyable = Sequencer.CanKeyProperty(FCanKeyPropertyParams(Class, PropertyPath));
+			if (bIsPropertyKeyable)
+			{
+				KeyablePropertyPaths.Add(PropertyPath);
+			}
+
+			if (!bIsPropertyKeyable || bRecurseAllProperties)
+			{
+				UStructProperty* StructProperty = Cast<UStructProperty>(Property);
+				if (StructProperty != nullptr)
+				{
+					GetKeyablePropertyPaths(Class, StructProperty->Struct, PropertyPath, Sequencer, KeyablePropertyPaths);
+				}
+			}
+
+			PropertyPath.RemoveAt(PropertyPath.Num() - 1);
+		}
+	}
 }
 
 
-float FSequencerObjectBindingNode::GetNodeHeight() const
+struct PropertyMenuData
 {
-	return SequencerLayoutConstants::ObjectNodeHeight;
-}
+	FString MenuName;
+	TArray<UProperty*> PropertyPath;
+};
 
 
-FNodePadding FSequencerObjectBindingNode::GetNodePadding() const
-{
-	return FNodePadding(SequencerNodeConstants::CommonPadding * 2, 0.f);
-}
+/* FSequencerDisplayNode interface
+ *****************************************************************************/
 
-
-ESequencerNode::Type FSequencerObjectBindingNode::GetType() const
-{
-	return ESequencerNode::Object;
-}
-
-
-const UClass* FSequencerObjectBindingNode::GetClassForObjectBinding()
-{
-	FSequencer& ParentSequencer = GetSequencer();
-
-	UMovieScene* MovieScene = GetSequencer().GetFocusedMovieSceneSequence()->GetMovieScene();
-
-	FMovieSceneSpawnable* Spawnable = MovieScene->FindSpawnable(ObjectBinding);
-	FMovieScenePossessable* Possessable = MovieScene->FindPossessable(ObjectBinding);
-	
-	// should exist, but also shouldn't be both a spawnable and a possessable
-	check((Spawnable != NULL) ^ (Possessable != NULL));
-	check((NULL == Spawnable) || (NULL != Spawnable->GetClass())  );
-	const UClass* ObjectClass = Spawnable ? Spawnable->GetClass()->GetSuperClass() : Possessable->GetPossessedObjectClass();
-
-	return ObjectClass;
-}
-
-
-TSharedRef<SWidget> FSequencerObjectBindingNode::GenerateEditWidgetForOutliner()
-{
-	// Create a container edit box
-	TSharedPtr<class SHorizontalBox> EditBox;
-	SAssignNew(EditBox, SHorizontalBox);	
-
-	// Add the property combo box
-	EditBox.Get()->AddSlot()
-	[
-		SNew(SComboButton)
-		.ButtonStyle(FEditorStyle::Get(), "FlatButton.Light")
-		.OnGetMenuContent(this, &FSequencerObjectBindingNode::OnGetAddTrackMenuContent)
-		.ContentPadding(FMargin(2, 0))
-		.ButtonContent()
-		[
-			SNew(SHorizontalBox)
-
-			+ SHorizontalBox::Slot()
-			.VAlign(VAlign_Center)
-			.AutoWidth()
-			[
-				SNew(STextBlock)
-				.Font(FEditorStyle::Get().GetFontStyle("FontAwesome.8"))
-				.Text(FText::FromString(FString(TEXT("\xf067"))) /*fa-plus*/)
-			]
-
-			+ SHorizontalBox::Slot()
-			.AutoWidth()
-			.Padding(4, 0, 0, 0)
-			[
-				SNew(STextBlock)
-				.Font(FEditorStyle::GetFontStyle("Sequencer.AnimationOutliner.RegularFont"))
-				.Text(LOCTEXT("AddTrackButton", "Track"))
-			]
-		]
-	];
-
-	const UClass* ObjectClass = GetClassForObjectBinding();
-
-	GetSequencer().BuildObjectBindingEditButtons(EditBox, ObjectBinding, ObjectClass);
-
-	return EditBox.ToSharedRef();
-}
-
-
 void FSequencerObjectBindingNode::BuildContextMenu(FMenuBuilder& MenuBuilder)
 {
@@ -150,48 +108,171 @@
 }
 
 
-void GetKeyablePropertyPaths(UClass* Class, UStruct* PropertySource, TArray<UProperty*>& PropertyPath, FSequencer& Sequencer, TArray<TArray<UProperty*>>& KeyablePropertyPaths)
+bool FSequencerObjectBindingNode::CanRenameNode() const
 {
-	//@todo need to resolve this between UMG and the level editor sequencer
-	const bool bRecurseAllProperties = Sequencer.IsLevelEditorSequencer();
+	return true;
+}
 
-	for (TFieldIterator<UProperty> PropertyIterator(PropertySource); PropertyIterator; ++PropertyIterator)
+
+TSharedRef<SWidget> FSequencerObjectBindingNode::GenerateEditWidgetForOutliner()
+{
+	// Create a container edit box
+	TSharedPtr<class SHorizontalBox> EditBox;
+	SAssignNew(EditBox, SHorizontalBox);	
+
+	// Add the property combo box
+	EditBox.Get()->AddSlot()
+	[
+		SNew(SComboButton)
+			.ButtonStyle(FEditorStyle::Get(), "FlatButton.Light")
+			.OnGetMenuContent(this, &FSequencerObjectBindingNode::HandleAddTrackComboButtonGetMenuContent)
+			.ContentPadding(FMargin(2, 0))
+			.ButtonContent()
+			[
+				SNew(SHorizontalBox)
+
+				+ SHorizontalBox::Slot()
+					.VAlign(VAlign_Center)
+					.AutoWidth()
+					[
+						SNew(STextBlock)
+							.Font(FEditorStyle::Get().GetFontStyle("FontAwesome.8"))
+							.Text(FText::FromString(FString(TEXT("\xf067"))) /*fa-plus*/)
+					]
+
+				+ SHorizontalBox::Slot()
+					.AutoWidth()
+					.Padding(4, 0, 0, 0)
+					[
+						SNew(STextBlock)
+							.Font(FEditorStyle::GetFontStyle("Sequencer.AnimationOutliner.RegularFont"))
+							.Text(LOCTEXT("AddTrackButton", "Track"))
+					]
+			]
+	];
+
+	const UClass* ObjectClass = GetClassForObjectBinding();
+
+	GetSequencer().BuildObjectBindingEditButtons(EditBox, ObjectBinding, ObjectClass);
+
+	return EditBox.ToSharedRef();
+}
+
+
+FText FSequencerObjectBindingNode::GetDisplayName() const
+{
+	UMovieScene* MovieScene = GetSequencer().GetFocusedMovieSceneSequence()->GetMovieScene();
+
+	if (MovieScene != nullptr)
 	{
-		UProperty* Property = *PropertyIterator;
+		return MovieScene->GetObjectDisplayName(ObjectBinding);
+	}
 
-		if (Property && !Property->HasAnyPropertyFlags(CPF_Deprecated))
-		{
-			PropertyPath.Add(Property);
+	return DefaultDisplayName;
+}
 
-			bool bIsPropertyKeyable = Sequencer.CanKeyProperty(FCanKeyPropertyParams(Class, PropertyPath));
-			if (bIsPropertyKeyable)
-			{
-				KeyablePropertyPaths.Add(PropertyPath);
-			}
 
-			if (!bIsPropertyKeyable || bRecurseAllProperties)
-			{
-				UStructProperty* StructProperty = Cast<UStructProperty>(Property);
-				if (StructProperty != nullptr)
-				{
-					GetKeyablePropertyPaths(Class, StructProperty->Struct, PropertyPath, Sequencer, KeyablePropertyPaths);
-				}
-			}
+float FSequencerObjectBindingNode::GetNodeHeight() const
+{
+	return SequencerLayoutConstants::ObjectNodeHeight;
+}
 
-			PropertyPath.RemoveAt(PropertyPath.Num() - 1);
-		}
+
+FNodePadding FSequencerObjectBindingNode::GetNodePadding() const
+{
+	return FNodePadding(SequencerNodeConstants::CommonPadding * 2, 0.f);
+}
+
+
+ESequencerNode::Type FSequencerObjectBindingNode::GetType() const
+{
+	return ESequencerNode::Object;
+}
+
+
+void FSequencerObjectBindingNode::SetDisplayName(const FText& DisplayName)
+{
+	UMovieScene* MovieScene = GetSequencer().GetFocusedMovieSceneSequence()->GetMovieScene();
+
+	if (MovieScene != nullptr)
+	{
+		MovieScene->SetObjectDisplayName(ObjectBinding, DisplayName);
 	}
 }
 
 
-struct PropertyMenuData
+/* FSequencerObjectBindingNode implementation
+ *****************************************************************************/
+
+void FSequencerObjectBindingNode::AddPropertyMenuItems(FMenuBuilder& AddTrackMenuBuilder, TArray<TArray<UProperty*> > KeyableProperties, int32 PropertyNameIndexStart, int32 PropertyNameIndexEnd)
 {
-	FString MenuName;
-	TArray<UProperty*> PropertyPath;
-};
+	TArray<PropertyMenuData> KeyablePropertyMenuData;
 
+	for (auto KeyableProperty : KeyableProperties)
+	{
+		TArray<FString> PropertyNames;
+		if (PropertyNameIndexEnd == -1)
+		{
+			PropertyNameIndexEnd = KeyableProperty.Num();
+		}
+
+		//@todo
+		if (PropertyNameIndexStart >= KeyableProperty.Num())
+		{
+			continue;
+		}
+
+		for (int32 PropertyNameIndex = PropertyNameIndexStart; PropertyNameIndex < PropertyNameIndexEnd; ++PropertyNameIndex)
+		{
+			PropertyNames.Add(KeyableProperty[PropertyNameIndex]->GetDisplayNameText().ToString());
+		}
+
+		PropertyMenuData KeyableMenuData;
+		{
+			KeyableMenuData.PropertyPath = KeyableProperty;
+			KeyableMenuData.MenuName = FString::Join( PropertyNames, TEXT( "." ) );
+		}
+
+		KeyablePropertyMenuData.Add(KeyableMenuData);
+	}
+
+	// Sort on the menu name
+	KeyablePropertyMenuData.Sort([](const PropertyMenuData& A, const PropertyMenuData& B)
+	{
+		int32 CompareResult = A.MenuName.Compare(B.MenuName);
+		return CompareResult < 0;
+	});
+
+	// Add menu items
+	for (int32 MenuDataIndex = 0; MenuDataIndex < KeyablePropertyMenuData.Num(); ++MenuDataIndex)
+	{
+		FUIAction AddTrackMenuAction(FExecuteAction::CreateSP(this, &FSequencerObjectBindingNode::HandlePropertyMenuItemExecute, KeyablePropertyMenuData[MenuDataIndex].PropertyPath));
+		AddTrackMenuBuilder.AddMenuEntry(FText::FromString(KeyablePropertyMenuData[MenuDataIndex].MenuName), FText(), FSlateIcon(), AddTrackMenuAction);
+	}
+}
+
+
+const UClass* FSequencerObjectBindingNode::GetClassForObjectBinding()
+{
+	FSequencer& ParentSequencer = GetSequencer();
+
+	UMovieScene* MovieScene = GetSequencer().GetFocusedMovieSceneSequence()->GetMovieScene();
+	FMovieSceneSpawnable* Spawnable = MovieScene->FindSpawnable(ObjectBinding);
+	FMovieScenePossessable* Possessable = MovieScene->FindPossessable(ObjectBinding);
 	
-TSharedRef<SWidget> FSequencerObjectBindingNode::OnGetAddTrackMenuContent()
+	// should exist, but also shouldn't be both a spawnable and a possessable
+	check((Spawnable != nullptr) ^ (Possessable != nullptr));
+	check((nullptr == Spawnable) || (nullptr != Spawnable->GetClass())  );
+	const UClass* ObjectClass = Spawnable ? Spawnable->GetClass()->GetSuperClass() : Possessable->GetPossessedObjectClass();
+
+	return ObjectClass;
+}
+
+
+/* FSequencerObjectBindingNode callbacks
+ *****************************************************************************/
+
+TSharedRef<SWidget> FSequencerObjectBindingNode::HandleAddTrackComboButtonGetMenuContent()
 {
 	FSequencer& Sequencer = GetSequencer();
 
@@ -274,7 +355,7 @@
 			AddTrackMenuBuilder.AddSubMenu(
 				FText::FromString(KeyablePropertyMenuData[MenuDataIndex].MenuName),
 				FText::GetEmpty(), 
-				FNewMenuDelegate::CreateSP(this, &FSequencerObjectBindingNode::AddPropertyMenuItemsWithCategories, KeyableSubMenuPropertyPaths));
+				FNewMenuDelegate::CreateSP(this, &FSequencerObjectBindingNode::HandleAddTrackSubMenuNew, KeyableSubMenuPropertyPaths));
 
 			++MenuDataIndex;
 		}
@@ -285,7 +366,7 @@
 }
 
 
-void FSequencerObjectBindingNode::AddPropertyMenuItemsWithCategories(FMenuBuilder& AddTrackMenuBuilder, TArray<TArray<UProperty*> > KeyablePropertyPaths)
+void FSequencerObjectBindingNode::HandleAddTrackSubMenuNew(FMenuBuilder& AddTrackMenuBuilder, TArray<TArray<UProperty*> > KeyablePropertyPaths)
 {
 	// [PostProcessSettings] [Bloom1Tint] [X]
 	// [PostProcessSettings] [Bloom1Tint] [Y]
@@ -357,55 +438,7 @@
 }
 
 
-void FSequencerObjectBindingNode::AddPropertyMenuItems(FMenuBuilder& AddTrackMenuBuilder, TArray<TArray<UProperty*> > KeyableProperties, int32 PropertyNameIndexStart, int32 PropertyNameIndexEnd)
+void FSequencerObjectBindingNode::HandlePropertyMenuItemExecute(TArray<UProperty*> PropertyPath)
 {
-	TArray<PropertyMenuData> KeyablePropertyMenuData;
-
-	for (auto KeyableProperty : KeyableProperties)
-	{
-		TArray<FString> PropertyNames;
-		if (PropertyNameIndexEnd == -1)
-		{
-			PropertyNameIndexEnd = KeyableProperty.Num();
-		}
-
-		//@todo
-		if (PropertyNameIndexStart >= KeyableProperty.Num())
-		{
-			continue;
-		}
-
-		for (int32 PropertyNameIndex = PropertyNameIndexStart; PropertyNameIndex < PropertyNameIndexEnd; ++PropertyNameIndex)
-		{
-			PropertyNames.Add(KeyableProperty[PropertyNameIndex]->GetDisplayNameText().ToString());
-		}
-
-		PropertyMenuData KeyableMenuData;
-		{
-			KeyableMenuData.PropertyPath = KeyableProperty;
-			KeyableMenuData.MenuName = FString::Join( PropertyNames, TEXT( "." ) );
-		}
-
-		KeyablePropertyMenuData.Add(KeyableMenuData);
-	}
-
-	// Sort on the menu name
-	KeyablePropertyMenuData.Sort([](const PropertyMenuData& A, const PropertyMenuData& B)
-	{
-		int32 CompareResult = A.MenuName.Compare(B.MenuName);
-		return CompareResult < 0;
-	});
-
-	// Add menu items
-	for (int32 MenuDataIndex = 0; MenuDataIndex < KeyablePropertyMenuData.Num(); ++MenuDataIndex)
-	{
-		FUIAction AddTrackMenuAction( FExecuteAction::CreateSP( this, &FSequencerObjectBindingNode::AddTrackForProperty, KeyablePropertyMenuData[MenuDataIndex].PropertyPath ) );
-		AddTrackMenuBuilder.AddMenuEntry( FText::FromString(KeyablePropertyMenuData[MenuDataIndex].MenuName), FText(), FSlateIcon(), AddTrackMenuAction );
-	}
-}
-
-
-void FSequencerObjectBindingNode::AddTrackForProperty(TArray<UProperty*> PropertyPath)
-{
 	FSequencer& Sequencer = GetSequencer();
 	UObject* BoundObject = Sequencer.GetFocusedMovieSceneSequence()->FindObject(ObjectBinding);
diff --git a/Engine/Source/Editor/Sequencer/Private/DisplayNodes/SequencerObjectBindingNode.h b/Engine/Source/Editor/Sequencer/Private/DisplayNodes/SequencerObjectBindingNode.h
--- a/Engine/Source/Editor/Sequencer/Private/DisplayNodes/SequencerObjectBindingNode.h
+++ b/Engine/Source/Editor/Sequencer/Private/DisplayNodes/SequencerObjectBindingNode.h
@@ -21,10 +21,10 @@ public:
 	 * @param InParentNode The parent of this node, or nullptr if this is a root node.
 	 * @param InParentTree The tree this node is in.
 	 */
-	FSequencerObjectBindingNode(FName NodeName, const FString& InObjectName, const FGuid& InObjectBinding, TSharedPtr<FSequencerDisplayNode> InParentNode, FSequencerNodeTree& InParentTree)
+	FSequencerObjectBindingNode(FName NodeName, const FText& InDisplayName, const FGuid& InObjectBinding, TSharedPtr<FSequencerDisplayNode> InParentNode, FSequencerNodeTree& InParentTree)
 		: FSequencerDisplayNode(NodeName, InParentNode, InParentTree)
 		, ObjectBinding(InObjectBinding)
-		, DefaultDisplayName(FText::FromString(InObjectName))
+		, DefaultDisplayName(InDisplayName)
 	{ }
 
 public:
@@ -41,25 +41,27 @@ public:
 	// FSequencerDisplayNode interface
 
 	virtual void BuildContextMenu(FMenuBuilder& MenuBuilder) override;
+	virtual bool CanRenameNode() const override;
 	virtual TSharedRef<SWidget> GenerateEditWidgetForOutliner() override;
 	virtual FText GetDisplayName() const override;
 	virtual float GetNodeHeight() const override;
 	virtual FNodePadding GetNodePadding() const override;
 	virtual ESequencerNode::Type GetType() const override;
+	virtual void SetDisplayName(const FText& DisplayName) override;
 
-private:
-
-	TSharedRef<SWidget> OnGetAddTrackMenuContent();
-
-	void AddPropertyMenuItemsWithCategories(FMenuBuilder& AddTrackMenuBuilder, TArray<TArray<UProperty*>> KeyablePropertyPath);
+protected:
 
 	void AddPropertyMenuItems(FMenuBuilder& AddTrackMenuBuilder, TArray<TArray<UProperty*>> KeyableProperties, int32 PropertyNameIndexStart = 0, int32 PropertyNameIndexEnd = -1);
 
-	void AddTrackForProperty(TArray<UProperty*> PropertyPath);
-
 	/** Get class for object binding */
 	const UClass* GetClassForObjectBinding();
 
+private:
+
+	TSharedRef<SWidget> HandleAddTrackComboButtonGetMenuContent();
+	void HandleAddTrackSubMenuNew(FMenuBuilder& AddTrackMenuBuilder, TArray<TArray<UProperty*>> KeyablePropertyPath);
+	void HandlePropertyMenuItemExecute(TArray<UProperty*> PropertyPath);
+
 private:
 
 	/** The binding to live objects */
diff --git a/Engine/Source/Editor/Sequencer/Private/DisplayNodes/SequencerSectionCategoryNode.cpp b/Engine/Source/Editor/Sequencer/Private/DisplayNodes/SequencerSectionCategoryNode.cpp
--- a/Engine/Source/Editor/Sequencer/Private/DisplayNodes/SequencerSectionCategoryNode.cpp
+++ b/Engine/Source/Editor/Sequencer/Private/DisplayNodes/SequencerSectionCategoryNode.cpp
@@ -13,6 +13,12 @@ namespace SequencerSectionCategoryNodeConstants
 /* FSequencerDisplayNode interface
  *****************************************************************************/
 
+bool FSequencerSectionCategoryNode::CanRenameNode() const
+{
+	return false;
+}
+
+
 FText FSequencerSectionCategoryNode::GetDisplayName() const
 {
 	return DisplayName;
@@ -35,3 +41,9 @@ ESequencerNode::Type FSequencerSectionCategoryNode::GetType() const
 {
 	return ESequencerNode::Category;
 }
+
+
+void FSequencerSectionCategoryNode::SetDisplayName(const FText& DisplayName)
+{
+	check(false);
+}
diff --git a/Engine/Source/Editor/Sequencer/Private/DisplayNodes/SequencerSectionCategoryNode.h b/Engine/Source/Editor/Sequencer/Private/DisplayNodes/SequencerSectionCategoryNode.h
--- a/Engine/Source/Editor/Sequencer/Private/DisplayNodes/SequencerSectionCategoryNode.h
+++ b/Engine/Source/Editor/Sequencer/Private/DisplayNodes/SequencerSectionCategoryNode.h
@@ -33,10 +33,12 @@ public:
 
 	// FSequencerDisplayNode interface
 
+	virtual bool CanRenameNode() const override;
 	virtual FText GetDisplayName() const override;
 	virtual float GetNodeHeight() const override;
 	virtual FNodePadding GetNodePadding() const override;
 	virtual ESequencerNode::Type GetType() const override;
+	virtual void SetDisplayName(const FText& DisplayName) override;
 
 private:
 
diff --git a/Engine/Source/Editor/Sequencer/Private/DisplayNodes/SequencerSectionKeyAreaNode.cpp b/Engine/Source/Editor/Sequencer/Private/DisplayNodes/SequencerSectionKeyAreaNode.cpp
--- a/Engine/Source/Editor/Sequencer/Private/DisplayNodes/SequencerSectionKeyAreaNode.cpp
+++ b/Engine/Source/Editor/Sequencer/Private/DisplayNodes/SequencerSectionKeyAreaNode.cpp
@@ -4,6 +4,7 @@
 #include "IKeyArea.h"
 #include "Sequencer.h"
 
+
 /* FSectionKeyAreaNode interface
  *****************************************************************************/
 
@@ -16,6 +17,12 @@ void FSequencerSectionKeyAreaNode::AddKeyArea(TSharedRef<IKeyArea> KeyArea)
 /* FSequencerDisplayNode interface
  *****************************************************************************/
 
+bool FSequencerSectionKeyAreaNode::CanRenameNode() const
+{
+	return false;
+}
+
+
 TSharedRef<SWidget> FSequencerSectionKeyAreaNode::GenerateEditWidgetForOutliner()
 {
 	// @todo sequencer: support multiple sections/key areas?
@@ -56,3 +63,9 @@ ESequencerNode::Type FSequencerSectionKeyAreaNode::GetType() const
 {
 	return ESequencerNode::KeyArea;
 }
+
+
+void FSequencerSectionKeyAreaNode::SetDisplayName(const FText& DisplayName)
+{
+	check(false);
+}
diff --git a/Engine/Source/Editor/Sequencer/Private/DisplayNodes/SequencerSectionKeyAreaNode.h b/Engine/Source/Editor/Sequencer/Private/DisplayNodes/SequencerSectionKeyAreaNode.h
--- a/Engine/Source/Editor/Sequencer/Private/DisplayNodes/SequencerSectionKeyAreaNode.h
+++ b/Engine/Source/Editor/Sequencer/Private/DisplayNodes/SequencerSectionKeyAreaNode.h
@@ -74,11 +74,13 @@ public:
 
 	// FSequencerDisplayNode interface
 
+	virtual bool CanRenameNode() const override;
 	virtual TSharedRef<SWidget> GenerateEditWidgetForOutliner() override;
 	virtual FText GetDisplayName() const override;
 	virtual float GetNodeHeight() const override;
 	virtual FNodePadding GetNodePadding() const override;
 	virtual ESequencerNode::Type GetType() const override;
+	virtual void SetDisplayName(const FText& DisplayName) override;
 
 private:
 
