Mergine changes from //depot/UE4 to //UE4/Main

diff --git a/Engine/Source/Editor/Sequencer/Private/DisplayNodes/SequencerTrackNode.cpp b/Engine/Source/Editor/Sequencer/Private/DisplayNodes/SequencerTrackNode.cpp
--- a/Engine/Source/Editor/Sequencer/Private/DisplayNodes/SequencerTrackNode.cpp
+++ b/Engine/Source/Editor/Sequencer/Private/DisplayNodes/SequencerTrackNode.cpp
@@ -5,28 +5,31 @@
 #include "ISequencerSection.h"
 #include "MovieSceneSection.h"
 #include "MovieSceneTrack.h"
+#include "MovieSceneNameableTrack.h"
 #include "Sequencer.h"
 #include "ISequencerTrackEditor.h"
 
+
 namespace SequencerNodeConstants
 {
 	extern const float CommonPadding;
 }
 
+
 /* FTrackNode structors
  *****************************************************************************/
 
-FTrackNode::FTrackNode(FName NodeName, UMovieSceneTrack& InAssociatedType, ISequencerTrackEditor& InAssociatedEditor, TSharedPtr<FSequencerDisplayNode> InParentNode, FSequencerNodeTree& InParentTree)
-	: FSequencerDisplayNode(NodeName, InParentNode, InParentTree)
+FSequencerTrackNode::FSequencerTrackNode(UMovieSceneTrack& InAssociatedTrack, ISequencerTrackEditor& InAssociatedEditor, TSharedPtr<FSequencerDisplayNode> InParentNode, FSequencerNodeTree& InParentTree)
+	: FSequencerDisplayNode(InAssociatedTrack.GetTrackName(), InParentNode, InParentTree)
 	, AssociatedEditor(InAssociatedEditor)
-	, AssociatedType(&InAssociatedType)
+	, AssociatedTrack(&InAssociatedTrack)
 { }
 
 
 /* FTrackNode interface
  *****************************************************************************/
 
-void FTrackNode::SetSectionAsKeyArea(TSharedRef<IKeyArea>& KeyArea)
+void FSequencerTrackNode::SetSectionAsKeyArea(TSharedRef<IKeyArea>& KeyArea)
 {
 	if( !TopLevelKeyNode.IsValid() )
 	{
@@ -38,7 +41,7 @@ void FTrackNode::SetSectionAsKeyArea(TSharedRef<IKeyArea>& KeyArea)
 }
 
 
-int32 FTrackNode::GetMaxRowIndex() const
+int32 FSequencerTrackNode::GetMaxRowIndex() const
 {
 	int32 MaxRowIndex = 0;
 
@@ -51,9 +54,9 @@ int32 FTrackNode::GetMaxRowIndex() const
 }
 
 
-void FTrackNode::FixRowIndices()
+void FSequencerTrackNode::FixRowIndices()
 {
-	if (AssociatedType->SupportsMultipleRows())
+	if (AssociatedTrack->SupportsMultipleRows())
 	{
 		// remove any empty track rows by waterfalling down sections to be as compact as possible
 		TArray< TArray< TSharedRef<ISequencerSection> > > TrackIndices;
@@ -93,7 +96,21 @@ void FTrackNode::FixRowIndices()
 /* FSequencerDisplayNode interface
  *****************************************************************************/
 
-TSharedRef<SWidget> FTrackNode::GenerateEditWidgetForOutliner()
+void FSequencerTrackNode::BuildContextMenu(FMenuBuilder& MenuBuilder)
+{
+	AssociatedEditor.BuildTrackContextMenu(MenuBuilder, AssociatedTrack.Get());
+	FSequencerDisplayNode::BuildContextMenu(MenuBuilder );
+}
+
+
+bool FSequencerTrackNode::CanRenameNode() const
+{
+	UMovieSceneTrack* Track = AssociatedTrack.Get();
+	return (Track != nullptr) && Track->IsA(UMovieSceneNameableTrack::StaticClass());
+}
+
+
+TSharedRef<SWidget> FSequencerTrackNode::GenerateEditWidgetForOutliner()
 {
 	TSharedPtr<FSequencerSectionKeyAreaNode> KeyAreaNode = GetTopLevelKeyNode();
 
@@ -120,7 +137,7 @@ TSharedRef<SWidget> FTrackNode::GenerateEditWidgetForOutliner()
 			ObjectBinding = StaticCastSharedPtr<FSequencerObjectBindingNode>(ParentNode)->GetObjectBinding();
 		}
 
-		TSharedPtr<SWidget> EditWidget = AssociatedEditor.BuildOutlinerEditWidget(ObjectBinding, AssociatedType.Get());
+		TSharedPtr<SWidget> EditWidget = AssociatedEditor.BuildOutlinerEditWidget(ObjectBinding, AssociatedTrack.Get());
 		if (EditWidget.IsValid())
 		{
 			return EditWidget.ToSharedRef();
@@ -131,7 +148,7 @@ TSharedRef<SWidget> FTrackNode::GenerateEditWidgetForOutliner()
 }
 
 
-void FTrackNode::GetChildKeyAreaNodesRecursively(TArray<TSharedRef<FSequencerSectionKeyAreaNode>>& OutNodes) const
+void FSequencerTrackNode::GetChildKeyAreaNodesRecursively(TArray<TSharedRef<FSequencerSectionKeyAreaNode>>& OutNodes) const
 {
 	FSequencerDisplayNode::GetChildKeyAreaNodesRecursively(OutNodes);
 
@@ -142,16 +159,13 @@ void FTrackNode::GetChildKeyAreaNodesRecursively(TArray<TSharedRef<FSequencerSec
 }
 
 
-FText FTrackNode::GetDisplayName() const
+FText FSequencerTrackNode::GetDisplayName() const
 {
-	// @todo sequencer: is there a better way to get the section interface name for the animation outliner?
-	return (TopLevelKeyNode.IsValid() && Sections.Num() > 0)
-		? Sections[0]->GetDisplayName()
-		: FText::FromName(NodeName);
+	return AssociatedTrack->GetDisplayName();
 }
 
 
-float FTrackNode::GetNodeHeight() const
+float FSequencerTrackNode::GetNodeHeight() const
 {
 	return (Sections.Num() > 0)
 		? Sections[0]->GetSectionHeight() * (GetMaxRowIndex() + 1)
@@ -159,19 +173,24 @@ float FTrackNode::GetNodeHeight() const
 }
 
 
-FNodePadding FTrackNode::GetNodePadding() const
+FNodePadding FSequencerTrackNode::GetNodePadding() const
 {
 	return FNodePadding(SequencerNodeConstants::CommonPadding);
 }
 
 
-ESequencerNode::Type FTrackNode::GetType() const
+ESequencerNode::Type FSequencerTrackNode::GetType() const
 {
 	return ESequencerNode::Track;
 }
 
-void FTrackNode::BuildContextMenu( FMenuBuilder& MenuBuilder )
+
+void FSequencerTrackNode::SetDisplayName(const FText& DisplayName)
 {
-	AssociatedEditor.BuildTrackContextMenu( MenuBuilder, AssociatedType.Get() );
-	FSequencerDisplayNode::BuildContextMenu( MenuBuilder );
+	auto NameableTrack = Cast<UMovieSceneNameableTrack>(AssociatedTrack.Get());
+
+	if (NameableTrack != nullptr)
+	{
+		NameableTrack->SetDisplayName(DisplayName);
+	}
 }
diff --git a/Engine/Source/Editor/Sequencer/Private/DisplayNodes/SequencerTrackNode.h b/Engine/Source/Editor/Sequencer/Private/DisplayNodes/SequencerTrackNode.h
--- a/Engine/Source/Editor/Sequencer/Private/DisplayNodes/SequencerTrackNode.h
+++ b/Engine/Source/Editor/Sequencer/Private/DisplayNodes/SequencerTrackNode.h
@@ -13,7 +13,7 @@ class ISequencerSection;
 /**
  * Represents an area to display Sequencer sections (possibly on multiple lines).
  */
-class FTrackNode
+class FSequencerTrackNode
 	: public FSequencerDisplayNode
 {
 public:
@@ -21,13 +21,12 @@ public:
 	/**
 	 * Create and initialize a new instance.
 	 * 
-	 * @param InNodeName The name identifier of then node.
 	 * @param InAssociatedType The track that this node represents.
 	 * @param InAssociatedEditor The track editor for the track that this node represents.
 	 * @param InParentNode The parent of this node, or nullptr if this is a root node.
 	 * @param InParentTree The tree this node is in.
 	 */
-	FTrackNode(FName NodeName, UMovieSceneTrack& InAssociatedType, ISequencerTrackEditor& InAssociatedEditor, TSharedPtr<FSequencerDisplayNode> InParentNode, FSequencerNodeTree& InParentTree);
+	FSequencerTrackNode(UMovieSceneTrack& InAssociatedTrack, ISequencerTrackEditor& InAssociatedEditor, TSharedPtr<FSequencerDisplayNode> InParentNode, FSequencerNodeTree& InParentTree);
 
 public:
 
@@ -70,7 +69,7 @@ public:
 	/** @return the track associated with this section */
 	UMovieSceneTrack* GetTrack() const
 	{
-		return AssociatedType.Get();
+		return AssociatedTrack.Get();
 	}
 
 	/** Gets the greatest row index of all the sections we have */
@@ -83,13 +82,15 @@ public:
 
 	// FSequencerDisplayNode interface
 
+	virtual void BuildContextMenu( FMenuBuilder& MenuBuilder );
+	virtual bool CanRenameNode() const override;
 	virtual TSharedRef<SWidget> GenerateEditWidgetForOutliner() override;
 	virtual void GetChildKeyAreaNodesRecursively(TArray<TSharedRef<FSequencerSectionKeyAreaNode>>& OutNodes) const override;
 	virtual FText GetDisplayName() const override;
 	virtual float GetNodeHeight() const override;
 	virtual FNodePadding GetNodePadding() const override;
 	virtual ESequencerNode::Type GetType() const override;
-	virtual void BuildContextMenu( FMenuBuilder& MenuBuilder );
+	virtual void SetDisplayName(const FText& DisplayName) override;
 
 private:
 
@@ -97,7 +98,7 @@ private:
 	ISequencerTrackEditor& AssociatedEditor;
 
 	/** The type associated with the sections in this node */
-	TWeakObjectPtr<UMovieSceneTrack> AssociatedType;
+	TWeakObjectPtr<UMovieSceneTrack> AssociatedTrack;
 
 	/** All of the sequencer sections in this node */
 	TArray<TSharedRef<ISequencerSection>> Sections;
diff --git a/Engine/Source/Editor/Sequencer/Private/KeyAreaLayout.cpp b/Engine/Source/Editor/Sequencer/Private/KeyAreaLayout.cpp
--- a/Engine/Source/Editor/Sequencer/Private/KeyAreaLayout.cpp
+++ b/Engine/Source/Editor/Sequencer/Private/KeyAreaLayout.cpp
@@ -71,7 +71,7 @@ FKeyAreaLayout::FKeyAreaLayout(FSequencerDisplayNode& InNode, int32 InSectionInd
 		FSequencerDisplayNode* LayoutNode = &InNode;
 		if (InNode.GetType() == ESequencerNode::Track)
 		{
-			TSharedPtr<FSequencerDisplayNode> KeyAreaNode = static_cast<FTrackNode&>(InNode).GetTopLevelKeyNode();
+			TSharedPtr<FSequencerDisplayNode> KeyAreaNode = static_cast<FSequencerTrackNode&>(InNode).GetTopLevelKeyNode();
 			if (KeyAreaNode.IsValid())
 			{
 				LayoutNode = KeyAreaNode.Get();
diff --git a/Engine/Source/Editor/Sequencer/Private/SAnimationOutlinerTreeNode.cpp b/Engine/Source/Editor/Sequencer/Private/SAnimationOutlinerTreeNode.cpp
--- a/Engine/Source/Editor/Sequencer/Private/SAnimationOutlinerTreeNode.cpp
+++ b/Engine/Source/Editor/Sequencer/Private/SAnimationOutlinerTreeNode.cpp
@@ -9,6 +9,7 @@
 #include "MovieSceneCommonHelpers.h"
 #include "Engine/Selection.h"
 #include "IKeyArea.h"
+#include "SEditableLabel.h"
 #include "SSequencerTreeView.h"
 
 
@@ -29,10 +30,11 @@ void SAnimationOutlinerTreeNode::Construct( const FArguments& InArgs, TSharedRef
 		FEditorStyle::GetFontStyle("Sequencer.AnimationOutliner.RegularFont") :
 		FEditorStyle::GetFontStyle("Sequencer.AnimationOutliner.BoldFont");
 
-	TSharedRef<SWidget> TextWidget = 
-		SNew( STextBlock )
-		.Text(this, &SAnimationOutlinerTreeNode::GetDisplayName )
-		.Font( NodeFont );
+	EditableLabel = SNew(SEditableLabel)
+		.CanEdit(this, &SAnimationOutlinerTreeNode::HandleNodeLabelCanEdit)
+		.Font(NodeFont)
+		.OnTextChanged(this, &SAnimationOutlinerTreeNode::HandleNodeLabelTextChanged)
+		.Text(this, &SAnimationOutlinerTreeNode::GetDisplayName);
 
 	auto NodeHeight = [=]() -> FOptionalSize { return DisplayNode->GetNodeHeight(); };
 
@@ -64,7 +66,7 @@ void SAnimationOutlinerTreeNode::Construct( const FArguments& InArgs, TSharedRef
 				.VAlign( VAlign_Center )
 				.FillWidth(3)
 				[
-					TextWidget
+					EditableLabel.ToSharedRef()
 				]
 
 				// Editor slot
@@ -160,6 +162,12 @@ void SAnimationOutlinerTreeNode::Construct( const FArguments& InArgs, TSharedRef
 }
 
 
+void SAnimationOutlinerTreeNode::EnterRenameMode()
+{
+	EditableLabel->EnterTextMode();
+}
+
+
 FLinearColor SAnimationOutlinerTreeNode::GetHoverTint() const
 {
 	return IsHovered() ? FLinearColor(1,1,1,0.9f) : FLinearColor(1,1,1,0.4f);
@@ -300,97 +308,98 @@ void SAnimationOutlinerTreeNode::GetAllDescendantNodes(TSharedPtr<FSequencerDisp
 
 FReply SAnimationOutlinerTreeNode::OnMouseButtonDown( const FGeometry& MyGeometry, const FPointerEvent& MouseEvent )
 {
-	if( MouseEvent.GetEffectingButton() == EKeys::LeftMouseButton && DisplayNode->IsSelectable() )
+	if (!DisplayNode->IsSelectable())
 	{
-		FSequencer& Sequencer = DisplayNode->GetSequencer();
-		bool bSelected = Sequencer.GetSelection().IsSelected(DisplayNode.ToSharedRef());
+		return FReply::Unhandled();
+	}
 
-		TArray<TSharedPtr<FSequencerDisplayNode> > AffectedNodes;
-		AffectedNodes.Add(DisplayNode.ToSharedRef());
+	if ((MouseEvent.GetEffectingButton() != EKeys::LeftMouseButton) &&
+		(MouseEvent.GetEffectingButton() != EKeys::RightMouseButton))
+	{
+		return FReply::Unhandled();
+	}
 
-		Sequencer.GetSelection().SuspendBroadcast();
+	FSequencer& Sequencer = DisplayNode->GetSequencer();
+	bool bSelected = Sequencer.GetSelection().IsSelected(DisplayNode.ToSharedRef());
 
-		if (MouseEvent.IsShiftDown())
-		{
-			FSequencerNodeTree& ParentTree = DisplayNode->GetParentTree();
+	TArray<TSharedPtr<FSequencerDisplayNode> > AffectedNodes;
+	AffectedNodes.Add(DisplayNode.ToSharedRef());
 
-			const TArray< TSharedRef<FSequencerDisplayNode> >RootNodes = ParentTree.GetRootNodes();
+	Sequencer.GetSelection().SuspendBroadcast();
 
-			// Get all nodes in order
-			TArray<TSharedRef<FSequencerDisplayNode> > AllNodes;
-			for (int32 i = 0; i < RootNodes.Num(); ++i)
-			{
-				GetAllDescendantNodes(RootNodes[i], AllNodes);
-			}
+	if (MouseEvent.IsShiftDown())
+	{
+		FSequencerNodeTree& ParentTree = DisplayNode->GetParentTree();
 
-			int32 FirstIndexToSelect = INT32_MAX;
-			int32 LastIndexToSelect = INT32_MIN;
+		const TArray< TSharedRef<FSequencerDisplayNode> >RootNodes = ParentTree.GetRootNodes();
 
-			for (int32 ChildIndex = 0; ChildIndex < AllNodes.Num(); ++ChildIndex)
-			{
-				TSharedRef<FSequencerDisplayNode> ChildNode = AllNodes[ChildIndex];
+		// Get all nodes in order
+		TArray<TSharedRef<FSequencerDisplayNode> > AllNodes;
+		for (int32 i = 0; i < RootNodes.Num(); ++i)
+		{
+			GetAllDescendantNodes(RootNodes[i], AllNodes);
+		}
 
-				if (ChildNode == DisplayNode.ToSharedRef() || Sequencer.GetSelection().IsSelected(ChildNode))
-				{
-					if (ChildIndex < FirstIndexToSelect)
-					{
-						FirstIndexToSelect = ChildIndex;
-					}
+		int32 FirstIndexToSelect = INT32_MAX;
+		int32 LastIndexToSelect = INT32_MIN;
 
-					if (ChildIndex > LastIndexToSelect)
-					{
-						LastIndexToSelect = ChildIndex;
-					}
-				}
-			}
+		for (int32 ChildIndex = 0; ChildIndex < AllNodes.Num(); ++ChildIndex)
+		{
+			TSharedRef<FSequencerDisplayNode> ChildNode = AllNodes[ChildIndex];
 
-			if (FirstIndexToSelect != INT32_MAX && LastIndexToSelect != INT32_MIN)
+			if (ChildNode == DisplayNode.ToSharedRef() || Sequencer.GetSelection().IsSelected(ChildNode))
 			{
-				for (int32 ChildIndex = FirstIndexToSelect; ChildIndex <= LastIndexToSelect; ++ChildIndex)
+				if (ChildIndex < FirstIndexToSelect)
 				{
-					TSharedRef<FSequencerDisplayNode> ChildNode = AllNodes[ChildIndex];
+					FirstIndexToSelect = ChildIndex;
+				}
 
-					if (!Sequencer.GetSelection().IsSelected(ChildNode))
-					{
-						Sequencer.GetSelection().AddToSelection(ChildNode);
-						AffectedNodes.Add(ChildNode);
-					}
+				if (ChildIndex > LastIndexToSelect)
+				{
+					LastIndexToSelect = ChildIndex;
 				}
 			}
 		}
-		else if( MouseEvent.IsControlDown() )
+
+		if (FirstIndexToSelect != INT32_MAX && LastIndexToSelect != INT32_MIN)
 		{
-			// Toggle selection when control is down
-			if (bSelected)
-			{
-				Sequencer.GetSelection().RemoveFromSelection(DisplayNode.ToSharedRef());
-			}
-			else
+			for (int32 ChildIndex = FirstIndexToSelect; ChildIndex <= LastIndexToSelect; ++ChildIndex)
 			{
-				Sequencer.GetSelection().AddToSelection(DisplayNode.ToSharedRef());
+				TSharedRef<FSequencerDisplayNode> ChildNode = AllNodes[ChildIndex];
+
+				if (!Sequencer.GetSelection().IsSelected(ChildNode))
+				{
+					Sequencer.GetSelection().AddToSelection(ChildNode);
+					AffectedNodes.Add(ChildNode);
+				}
 			}
 		}
+	}
+	else if( MouseEvent.IsControlDown() )
+	{
+		// Toggle selection when control is down
+		if (bSelected)
+		{
+			Sequencer.GetSelection().RemoveFromSelection(DisplayNode.ToSharedRef());
+		}
 		else
 		{
-			// Deselect the other nodes and select this node.
-			Sequencer.GetSelection().EmptySelectedOutlinerNodes();
 			Sequencer.GetSelection().AddToSelection(DisplayNode.ToSharedRef());
 		}
-
-		OnSelectionChanged( AffectedNodes );
-
-		Sequencer.GetSelection().ResumeBroadcast();
-		Sequencer.GetSelection().GetOnOutlinerNodeSelectionChanged().Broadcast();
-
-		return FReply::Handled();
 	}
-
-	if (MouseEvent.GetEffectingButton() == EKeys::RightMouseButton)
+	else
 	{
-		return FReply::Handled();
+		// Deselect the other nodes and select this node.
+		Sequencer.GetSelection().EmptySelectedOutlinerNodes();
+		Sequencer.GetSelection().AddToSelection(DisplayNode.ToSharedRef());
 	}
 
-	return FReply::Unhandled();
+	OnSelectionChanged( AffectedNodes );
+
+	Sequencer.GetSelection().ResumeBroadcast();
+	Sequencer.GetSelection().GetOnOutlinerNodeSelectionChanged().Broadcast();
+
+	return FReply::Handled();
 }
 
 
@@ -463,6 +472,18 @@ FText SAnimationOutlinerTreeNode::GetDisplayName() const
 }
 
 
+bool SAnimationOutlinerTreeNode::HandleNodeLabelCanEdit() const
+{
+	return DisplayNode->CanRenameNode();
+}
+
+
+void SAnimationOutlinerTreeNode::HandleNodeLabelTextChanged(const FText& NewLabel)
+{
+	DisplayNode->SetDisplayName(NewLabel);
+}
+
+
 void SAnimationOutlinerTreeNode::OnSelectionChanged( TArray<TSharedPtr<FSequencerDisplayNode> > AffectedNodes )
 {
 	FSequencer& Sequencer = DisplayNode->GetSequencer();
diff --git a/Engine/Source/Editor/Sequencer/Private/SAnimationOutlinerTreeNode.h b/Engine/Source/Editor/Sequencer/Private/SAnimationOutlinerTreeNode.h
--- a/Engine/Source/Editor/Sequencer/Private/SAnimationOutlinerTreeNode.h
+++ b/Engine/Source/Editor/Sequencer/Private/SAnimationOutlinerTreeNode.h
@@ -2,24 +2,34 @@
 
 #pragma once
 
+
+class SEditableLabel;
 class SSequencerTreeView;
 class SSequencerTreeViewRow;
 
+
 /**
- * A widget for displaying a sequencer tree node in the animation outliner
+ * A widget for displaying a sequencer tree node in the animation outliner.
  */
-class SAnimationOutlinerTreeNode : public SCompoundWidget
+class SAnimationOutlinerTreeNode
+	: public SCompoundWidget
 {
 public:
+
 	SLATE_BEGIN_ARGS(SAnimationOutlinerTreeNode){}
 	SLATE_END_ARGS()
 
 	void Construct( const FArguments& InArgs, TSharedRef<FSequencerDisplayNode> Node, const TSharedRef<SSequencerTreeViewRow>& InTableRow );
 
+public:
+
+	/** Change the node's label text to edit mode. */
+	void EnterRenameMode();
+
 	/**
 	 * @return The display node used by this widget                                                              
 	 */
-	const TSharedPtr< FSequencerDisplayNode> GetDisplayNode() const { return DisplayNode; }
+	const TSharedPtr<FSequencerDisplayNode> GetDisplayNode() const { return DisplayNode; }
 
 	/**
 	 * Gets a tint to apply buttons on hover
@@ -27,9 +37,11 @@ public:
 	FLinearColor GetHoverTint() const;
 
 private:
-	/** SWidget Interface */
-	virtual FReply OnMouseButtonDown( const FGeometry& MyGeometry, const FPointerEvent& MouseEvent ) override;
-	virtual FReply OnMouseButtonUp( const FGeometry& MyGeometry, const FPointerEvent& MouseEvent ) override;
+
+	// SWidget interface
+
+	virtual FReply OnMouseButtonDown(const FGeometry& MyGeometry, const FPointerEvent& MouseEvent) override;
+	virtual FReply OnMouseButtonUp(const FGeometry& MyGeometry, const FPointerEvent& MouseEvent) override;
 
 	FSlateColor GetForegroundBasedOnSelection() const;
 
@@ -48,6 +60,12 @@ private:
 	 */
 	FText GetDisplayName() const;
 
+	/** Callback for checking whether the node label can be edited. */
+	bool HandleNodeLabelCanEdit() const;
+
+	/** Callback for when the node label text has changed. */
+	void HandleNodeLabelTextChanged(const FText& NewLabel);
+
 	/** Handles the previous key button being clicked. */
 	FReply OnPreviousKeyClicked();
 
@@ -64,17 +82,25 @@ private:
 	TSharedPtr<FSequencerDisplayNode> GetRootNode(TSharedPtr<FSequencerDisplayNode> ObjectNode);
 
 	/** Called when nodes are selected */
-	void OnSelectionChanged( TArray<TSharedPtr<FSequencerDisplayNode> > AffectedNodes );
+	void OnSelectionChanged(TArray<TSharedPtr<FSequencerDisplayNode>> AffectedNodes);
 
 private:
+
 	/** Layout node the widget is visualizing */
 	TSharedPtr<FSequencerDisplayNode> DisplayNode;
+
+	/** Holds the editable text label widge.t */
+	TSharedPtr<SEditableLabel> EditableLabel;
+
 	/** Brush to display a border around the widget when it is selected */
 	const FSlateBrush* SelectedBrush;
+
 	/** Brush to display a border around the widget when it is selected but inactive */
 	const FSlateBrush* SelectedBrushInactive;
+
 	/** Brush to use if the node is not selected */
 	const FSlateBrush* NotSelectedBrush;
+
 	/** The table row style used for nodes in the tree. This is required as we don't actually use the tree for selection */
 	const FTableRowStyle* TableRowStyle;
 };
\ No newline at end of file
diff --git a/Engine/Source/Editor/Sequencer/Private/SSequencer.cpp b/Engine/Source/Editor/Sequencer/Private/SSequencer.cpp
--- a/Engine/Source/Editor/Sequencer/Private/SSequencer.cpp
+++ b/Engine/Source/Editor/Sequencer/Private/SSequencer.cpp
@@ -146,8 +146,7 @@ void SSequencer::Construct(const FArguments& InArgs, TSharedRef<FSequencer> InSe
 	SAssignNew(TrackOutliner, SSequencerTrackOutliner);
 
 	SAssignNew(TrackArea, SSequencerTrackArea, TimeSliderController, InSequencer)
-		.Visibility(this, &SSequencer::GetTrackAreaVisibility )
-		.LockInOutToStartEndRange(this, &SSequencer::GetLockInOutToStartEndRange);
+		.Visibility(this, &SSequencer::GetTrackAreaVisibility );
 	
 	SAssignNew(TreeView, SSequencerTreeView, SequencerNodeTree.ToSharedRef(), TrackArea.ToSharedRef())
 		.ExternalScrollbar(ScrollBar);
@@ -190,32 +189,31 @@ void SSequencer::Construct(const FArguments& InArgs, TSharedRef<FSequencer> InSe
 	ChildSlot
 	[
 		SNew(SBorder)
-		.Padding(2)
-		.BorderImage( FEditorStyle::GetBrush( "ToolPanel.GroupBorder" ) )
-		[
-			SNew( SVerticalBox )
-			+ SVerticalBox::Slot()
-			.AutoHeight()
+			.Padding(2.0f)
+			.BorderImage(FEditorStyle::GetBrush("ToolPanel.GroupBorder"))
 			[
-				SNew( SHorizontalBox )
-				+SHorizontalBox::Slot()
-				.AutoWidth()
-				[
-					MakeToolBar()
-				]
-				+SHorizontalBox::Slot()
-				.AutoWidth()
-				[
-					SNew( SSequencerCurveEditorToolBar, InSequencer, CurveEditor->GetCommands() )
-					.Visibility( this, &SSequencer::GetCurveEditorToolBarVisibility )
-				]
-			]
+				SNew(SVerticalBox)
 
-			+ SVerticalBox::Slot()
-				[
-					SNew(SOverlay)
+				+ SVerticalBox::Slot()
+					.AutoHeight()
+					[
+						SNew(SHorizontalBox)
+
+						+SHorizontalBox::Slot()
+							.AutoWidth()
+							[
+								MakeToolBar()
+							]
+
+						+SHorizontalBox::Slot()
+							.AutoWidth()
+							[
+								SNew(SSequencerCurveEditorToolBar, InSequencer, CurveEditor->GetCommands())
+									.Visibility(this, &SSequencer::GetCurveEditorToolBarVisibility)
+							]
+					]
 
-					+ SOverlay::Slot()
+				+ SVerticalBox::Slot()
 					[
 						SNew(SSplitter)
 							.Orientation(Orient_Horizontal)
@@ -223,175 +221,215 @@ void SSequencer::Construct(const FArguments& InArgs, TSharedRef<FSequencer> InSe
 						+ SSplitter::Slot()
 							.Value(0.80f)
 							[
-								SNew(SBox)
-									.Padding(FMargin(0.0f, 2.0f, 0.0f, 0.0f))
-									[
-										SNew( SGridPanel )
-											.FillRow( 1, 1.f )
-											.FillColumn( 0, FillCoefficient_0 )
-											.FillColumn( 1, FillCoefficient_1 )
-
-										+ SGridPanel::Slot( Column0, Row0 )
-											.VAlign( VAlign_Center )
-											[
-												// Search box for searching through the outliner
-												SNew( SSearchBox )
-													.OnTextChanged( this, &SSequencer::OnOutlinerSearchChanged )
-											]
+								SNew(SOverlay)
 
-										+ SGridPanel::Slot( Column0, Row1 )
-											.ColumnSpan(2)
+								+ SOverlay::Slot()
+									[
+										// track area grid panel
+										SNew(SBox)
+											.Padding(FMargin(0.0f, 3.0f, 0.0f, 0.0f))
 											[
-												SNew(SHorizontalBox)
+												SNew( SGridPanel )
+													.FillRow( 1, 1.f )
+													.FillColumn( 0, FillCoefficient_0 )
+													.FillColumn( 1, FillCoefficient_1 )
+
+												// outliner search box
+												+ SGridPanel::Slot( Column0, Row0 )
+													.Padding(FMargin(0.0f, 0.0f, 3.0f, 0.0f))
+													.VAlign( VAlign_Center )
+													[
+														SNew(SBorder)
+															.BorderImage(FEditorStyle::GetBrush("ToolPanel.GroupBorder"))
+															.BorderBackgroundColor(FLinearColor(.50f, .50f, .50f, 1.0f))
+															.Padding(FMargin(4.0f, 3.0f))
+															[
+																SNew( SSearchBox )
+																	.OnTextChanged( this, &SSequencer::OnOutlinerSearchChanged )
+															]
+													]
 
-												+ SHorizontalBox::Slot()
+												// main sequencer area
+												+ SGridPanel::Slot( Column0, Row1 )
+													.ColumnSpan(2)
 													[
-														SNew( SOverlay )
+														SNew(SHorizontalBox)
 
-														+ SOverlay::Slot()
+														+ SHorizontalBox::Slot()
 															[
-																SNew(SHorizontalBox)
-								
-																+ SHorizontalBox::Slot()
-																	.FillWidth( FillCoefficient_0 )
+																SNew( SOverlay )
+
+																+ SOverlay::Slot()
 																	[
-																		SNew(SSequencerTreeViewBox, InSequencer)
-																			.Padding(FMargin(0, 0, 10.f, 0)) // Padding to allow space for the scroll bar
+																		SNew(SHorizontalBox)
+								
+																		// outliner tree
+																		+ SHorizontalBox::Slot()
+																			.FillWidth( FillCoefficient_0 )
+																			[
+																				SNew(SSequencerTreeViewBox, InSequencer)
+																					.Padding(FMargin(0, 0, 10.f, 0)) // Padding to allow space for the scroll bar
+																					[
+																						TreeView.ToSharedRef()
+																					]
+																			]
+
+																		// track area
+																		+ SHorizontalBox::Slot()
+																			.FillWidth( FillCoefficient_1 )
 																			[
-																				TreeView.ToSharedRef()
+																				TrackArea.ToSharedRef()
 																			]
 																	]
 
-																+ SHorizontalBox::Slot()
-																	.FillWidth( FillCoefficient_1 )
+																+ SOverlay::Slot()
+																	.HAlign( HAlign_Right )
 																	[
-																		TrackArea.ToSharedRef()
+																		ScrollBar
 																	]
 															]
 
-														+ SOverlay::Slot()
-															.HAlign( HAlign_Right )
+														+ SHorizontalBox::Slot()
+															.FillWidth( TAttribute<float>( this, &SSequencer::GetOutlinerSpacerFill ) )
 															[
-																ScrollBar
+																SNew(SSpacer)
 															]
 													]
 
-												+ SHorizontalBox::Slot()
-													.FillWidth( TAttribute<float>( this, &SSequencer::GetOutlinerSpacerFill ) )
+												// playback buttons
+												+ SGridPanel::Slot( Column0, Row2 )
+													.Padding(FMargin(0.0f, 0.0f, 3.0f, 0.0f))
 													[
-														SNew(SSpacer)
+														SNew(SBorder)
+															.BorderImage(FEditorStyle::GetBrush("ToolPanel.GroupBorder"))
+															.BorderBackgroundColor(FLinearColor(.50f, .50f, .50f, 1.0f))
+															.HAlign(HAlign_Center)
+															[
+																MakeTransportControls()
+															]
 													]
-											]
-
-										+ SGridPanel::Slot( Column0, Row2 )
-											.HAlign( HAlign_Center )
-											[
-												MakeTransportControls()
-											]
 
-										// Second column
-										+ SGridPanel::Slot( Column1, Row0 )
-											[
-												SNew( SBorder )
-													.BorderImage( FEditorStyle::GetBrush("ToolPanel.GroupBorder") )
-													.BorderBackgroundColor( FLinearColor(.50f, .50f, .50f, 1.0f ) )
-													.Padding(0)
+												// Second column
+												+ SGridPanel::Slot( Column1, Row0 )
 													[
-														TopTimeSlider
+														SNew( SBorder )
+															.BorderImage( FEditorStyle::GetBrush("ToolPanel.GroupBorder") )
+															.BorderBackgroundColor( FLinearColor(.50f, .50f, .50f, 1.0f ) )
+															.Padding(0)
+															[
+																TopTimeSlider
+															]
 													]
-											]
-
-										// Overlay that draws the tick lines
-										+ SGridPanel::Slot( Column1, Row1, SGridPanel::Layer(0) )
-											[
-												SNew( SSequencerSectionOverlay, TimeSliderController )
-													.Visibility( EVisibility::HitTestInvisible )
-													.DisplayScrubPosition( false )
-													.DisplayTickLines( true )
-											]
 
-										// Curve editor
-										+ SGridPanel::Slot( Column1, Row1 )
-											[
-												CurveEditor.ToSharedRef()
-											]
+												// Overlay that draws the tick lines
+												+ SGridPanel::Slot( Column1, Row1, SGridPanel::Layer(0) )
+													[
+														SNew( SSequencerSectionOverlay, TimeSliderController )
+															.Visibility( EVisibility::HitTestInvisible )
+															.DisplayScrubPosition( false )
+															.DisplayTickLines( true )
+													]
 
-										// Overlay that draws the scrub position
-										+ SGridPanel::Slot( Column1, Row1, SGridPanel::Layer(30) )
-											[
-												SNew( SSequencerSectionOverlay, TimeSliderController )
-													.Visibility( EVisibility::HitTestInvisible )
-													.DisplayScrubPosition( true )
-													.DisplayTickLines( false )
-													.PaintPlaybackRangeArgs(this, &SSequencer::GetSectionPlaybackRangeArgs)
-											]
+												// Curve editor
+												+ SGridPanel::Slot( Column1, Row1 )
+													[
+														CurveEditor.ToSharedRef()
+													]
 
-										+ SGridPanel::Slot( Column1, Row2 )
-											[
-												SNew( SBorder )
-													.BorderImage( FEditorStyle::GetBrush("ToolPanel.GroupBorder") )
-													.BorderBackgroundColor( FLinearColor(.50f, .50f, .50f, 1.0f ) )
-													.Padding(0)
+												// Overlay that draws the scrub position
+												+ SGridPanel::Slot( Column1, Row1, SGridPanel::Layer(30) )
 													[
-														SNew(SVerticalBox)
+														SNew( SSequencerSectionOverlay, TimeSliderController )
+															.Visibility( EVisibility::HitTestInvisible )
+															.DisplayScrubPosition( true )
+															.DisplayTickLines( false )
+															.PaintPlaybackRangeArgs(this, &SSequencer::GetSectionPlaybackRangeArgs)
+													]
 
-														+ SVerticalBox::Slot()
-															.AutoHeight()
+												// play range sliders
+												+ SGridPanel::Slot( Column1, Row2 )
+													[
+														SNew( SBorder )
+															.BorderImage( FEditorStyle::GetBrush("ToolPanel.GroupBorder") )
+															.BorderBackgroundColor( FLinearColor(.50f, .50f, .50f, 1.0f ) )
 															[
-																SNew( SOverlay )
+																SNew(SVerticalBox)
 
-																+ SOverlay::Slot()
+																+ SVerticalBox::Slot()
+																	.AutoHeight()
 																	[
-																		BottomTimeSlider
+																		SNew( SOverlay )
+
+																		+ SOverlay::Slot()
+																			[
+																				BottomTimeSlider
+																			]
+
+																		+ SOverlay::Slot()
+																			[
+																				BottomTimeRange
+																			]
 																	]
 
-																+ SOverlay::Slot()
+																+ SVerticalBox::Slot()
+																	.AutoHeight()
 																	[
-																		BottomTimeRange
+																		SAssignNew(BreadcrumbTrail, SBreadcrumbTrail<FSequencerBreadcrumb>)
+																		.Visibility(this, &SSequencer::GetBreadcrumbTrailVisibility)
+																		.OnCrumbClicked(this, &SSequencer::OnCrumbClicked)
 																	]
 															]
-
-														+ SVerticalBox::Slot()
-															.AutoHeight()
-															[
-																SAssignNew(BreadcrumbTrail, SBreadcrumbTrail<FSequencerBreadcrumb>)
-																.Visibility(this, &SSequencer::GetBreadcrumbTrailVisibility)
-																.OnCrumbClicked(this, &SSequencer::OnCrumbClicked)
-															]
 													]
 											]
 									]
+
+								+ SOverlay::Slot()
+									[
+										// track area virtual splitter overlay
+										SNew(SSequencerSplitterOverlay)
+											.Style(FEditorStyle::Get(), "Sequencer.AnimationOutliner.Splitter")
+											.Visibility(EVisibility::SelfHitTestInvisible)
+
+										+ SSplitter::Slot()
+											.Value(FillCoefficient_0)
+											.OnSlotResized(SSplitter::FOnSlotResized::CreateSP(this, &SSequencer::OnColumnFillCoefficientChanged, 0))
+											[
+												SNew(SSpacer)
+											]
+
+										+ SSplitter::Slot()
+											.Value(FillCoefficient_1)
+											.OnSlotResized(SSplitter::FOnSlotResized::CreateSP(this, &SSequencer::OnColumnFillCoefficientChanged, 1))
+											[
+												SNew(SSpacer)
+											]
+									]
 							]
 
 						+ SSplitter::Slot()
 							.Value(0.2f)
 							[
-								DetailsView.ToSharedRef()
-							]
-					]
-
-				+ SOverlay::Slot()
-					[
-						SNew(SSequencerSplitterOverlay)
-							.Style(FEditorStyle::Get(), "Sequencer.AnimationOutliner.Splitter")
-							.Visibility(EVisibility::SelfHitTestInvisible)
+								// details view panel
+								SNew(SVerticalBox)
 
-						+ SSplitter::Slot()
-							.Value(FillCoefficient_0)
-							.OnSlotResized(SSplitter::FOnSlotResized::CreateSP(this, &SSequencer::OnColumnFillCoefficientChanged, 0))
-							[
-								SNew(SSpacer)
-							]
+								+ SVerticalBox::Slot()
+									.AutoHeight()
+									.HAlign(HAlign_Center)
+									.Padding(2.0f, 8.0f, 2.0f, 2.0f)
+									[
+										SNew(STextBlock)
+											.Text(LOCTEXT("NoObjectsSelected", "Select a key or section for details."))
+											.Visibility(this, &SSequencer::HandleDetailsViewTipVisibility)
+									]
 
-						+ SSplitter::Slot()
-							.Value(FillCoefficient_1)
-							.OnSlotResized(SSplitter::FOnSlotResized::CreateSP(this, &SSequencer::OnColumnFillCoefficientChanged, 1))
-							[
-								SNew(SSpacer)
+								+ SVerticalBox::Slot()
+									.FillHeight(1.0f)
+									[
+										DetailsView.ToSharedRef()
+									]
 							]
 					]
 			]
-		]
 	];
 
 	InSequencer->GetSelection().GetOnKeySelectionChanged().AddSP(this, &SSequencer::HandleKeySelectionChanged);
@@ -478,6 +516,14 @@ bool SSequencer::HandleDetailsViewEnabled() const
 }
 
 
+EVisibility SSequencer::HandleDetailsViewTipVisibility() const
+{
+	return (DetailsView->GetSelectedObjects().Num() > 0)
+		? EVisibility::Collapsed
+		: EVisibility::Visible;
+}
+
+
 EVisibility SSequencer::HandleDetailsViewVisibility() const
 {
 	if (Sequencer.Pin()->GetSettings()->GetDetailsViewVisible())
@@ -767,7 +813,6 @@ TSharedRef<SWidget> SSequencer::MakeSnapMenu()
 		MenuBuilder.AddMenuEntry( FSequencerCommands::Get().ToggleAutoScroll );
 		MenuBuilder.AddMenuEntry( FSequencerCommands::Get().ToggleShowFrameNumbers );
 		MenuBuilder.AddMenuEntry( FSequencerCommands::Get().ToggleShowRangeSlider );
-		MenuBuilder.AddMenuEntry( FSequencerCommands::Get().ToggleLockInOutToStartEndRange );
 	}
 	MenuBuilder.EndSection();
 
@@ -1284,14 +1329,14 @@ TArray<FSectionHandle> SSequencer::GetSectionHandles(const TSet<TWeakObjectPtr<U
 		Node->Traverse_ParentFirst([&](FSequencerDisplayNode& InNode) {
 			if (InNode.GetType() == ESequencerNode::Track)
 			{
-				FTrackNode& TrackNode = static_cast<FTrackNode&>(InNode);
+				FSequencerTrackNode& TrackNode = static_cast<FSequencerTrackNode&>(InNode);
 
 				const auto& AllSections = TrackNode.GetTrack()->GetAllSections();
 				for (int32 Index = 0; Index < AllSections.Num(); ++Index)
 				{
 					if (DesiredSections.Contains(TWeakObjectPtr<UMovieSceneSection>(AllSections[Index])))
 					{
-						SectionHandles.Emplace(StaticCastSharedRef<FTrackNode>(TrackNode.AsShared()), Index);
+						SectionHandles.Emplace(StaticCastSharedRef<FSequencerTrackNode>(TrackNode.AsShared()), Index);
 					}
 				}
 			}
@@ -1493,12 +1538,6 @@ EVisibility SSequencer::GetCurveEditorVisibility() const
 }
 
 
-bool SSequencer::GetLockInOutToStartEndRange() const
-{
-	return Settings->GetLockInOutToStartEndRange();
-}
-
-
 void SSequencer::OnCurveEditorVisibilityChanged()
 {
 	if (CurveEditor.IsValid())
diff --git a/Engine/Source/Editor/Sequencer/Private/SSequencer.h b/Engine/Source/Editor/Sequencer/Private/SSequencer.h
--- a/Engine/Source/Editor/Sequencer/Private/SSequencer.h
+++ b/Engine/Source/Editor/Sequencer/Private/SSequencer.h
@@ -215,6 +215,9 @@ private:
 	/** Handles checking whether the details view is enabled. */
 	bool HandleDetailsViewEnabled() const;
 
+	/** Handles determining the visibility of the details view selection tip. */
+	EVisibility HandleDetailsViewTipVisibility() const;
+
 	/** Handles determining the visibility of the details view. */
 	EVisibility HandleDetailsViewVisibility() const;
 
@@ -291,9 +294,6 @@ private:
 	/** Get the visibility of the track area */
 	EVisibility GetTrackAreaVisibility() const;
 
-	/** Get the lock in/out to start/end range value */
-	bool GetLockInOutToStartEndRange() const;
-
 	/**
 	 * Called when one or more assets are dropped into the widget
 	 *
diff --git a/Engine/Source/Editor/Sequencer/Private/SSequencerCurveEditor.cpp b/Engine/Source/Editor/Sequencer/Private/SSequencerCurveEditor.cpp
--- a/Engine/Source/Editor/Sequencer/Private/SSequencerCurveEditor.cpp
+++ b/Engine/Source/Editor/Sequencer/Private/SSequencerCurveEditor.cpp
@@ -23,7 +23,7 @@ void SSequencerCurveEditor::Construct( const FArguments& InArgs, TSharedRef<FSeq
 		SCurveEditor::FArguments()
 		.ViewMinInput_Lambda( [=]{ return ViewRange.Get().GetLowerBoundValue(); } )
 		.ViewMaxInput_Lambda( [=]{ return ViewRange.Get().GetUpperBoundValue(); } )
-		.OnSetInputViewRange_Lambda( [=](float InLowerBound, float InUpperBound){ OnViewRangeChanged.ExecuteIfBound(TRange<float>(InLowerBound, InUpperBound), EViewRangeInterpolation::Immediate, true); } )
+		.OnSetInputViewRange_Lambda( [=](float InLowerBound, float InUpperBound){ OnViewRangeChanged.ExecuteIfBound(TRange<float>(InLowerBound, InUpperBound), EViewRangeInterpolation::Immediate); } )
 		.HideUI( false )
 		.ZoomToFitHorizontal( true )
 		.ShowCurveSelector( false )
diff --git a/Engine/Source/Editor/Sequencer/Private/SSequencerSection.cpp b/Engine/Source/Editor/Sequencer/Private/SSequencerSection.cpp
--- a/Engine/Source/Editor/Sequencer/Private/SSequencerSection.cpp
+++ b/Engine/Source/Editor/Sequencer/Private/SSequencerSection.cpp
@@ -29,7 +29,7 @@ void SSequencerSection::EnableLayoutRegeneration()
 }
 
 
-void SSequencerSection::Construct( const FArguments& InArgs, TSharedRef<FTrackNode> SectionNode, int32 InSectionIndex )
+void SSequencerSection::Construct( const FArguments& InArgs, TSharedRef<FSequencerTrackNode> SectionNode, int32 InSectionIndex )
 {
 	SectionIndex = InSectionIndex;
 	ParentSectionArea = SectionNode;
diff --git a/Engine/Source/Editor/Sequencer/Private/SSequencerSection.h b/Engine/Source/Editor/Sequencer/Private/SSequencerSection.h
--- a/Engine/Source/Editor/Sequencer/Private/SSequencerSection.h
+++ b/Engine/Source/Editor/Sequencer/Private/SSequencerSection.h
@@ -12,7 +12,7 @@ public:
 	{}
 	SLATE_END_ARGS()
 
-	void Construct( const FArguments& InArgs, TSharedRef<FTrackNode> SectionNode, int32 SectionIndex );
+	void Construct( const FArguments& InArgs, TSharedRef<FSequencerTrackNode> SectionNode, int32 SectionIndex );
 
 	TSharedPtr<ISequencerSection> GetSectionInterface() const { return SectionInterface; }
 
@@ -206,7 +206,7 @@ private:
 	/** Interface to section data */
 	TSharedPtr<ISequencerSection> SectionInterface;
 	/** Section area where this section resides */
-	TSharedPtr<FTrackNode> ParentSectionArea;
+	TSharedPtr<FSequencerTrackNode> ParentSectionArea;
 	/** Cached layout generated each tick */
 	TOptional<FKeyAreaLayout> Layout;
 	/** The currently pressed key if any */
diff --git a/Engine/Source/Editor/Sequencer/Private/SSequencerSectionAreaView.cpp b/Engine/Source/Editor/Sequencer/Private/SSequencerSectionAreaView.cpp
--- a/Engine/Source/Editor/Sequencer/Private/SSequencerSectionAreaView.cpp
+++ b/Engine/Source/Editor/Sequencer/Private/SSequencerSectionAreaView.cpp
@@ -71,7 +71,7 @@ void SSequencerSectionAreaView::Construct( const FArguments& InArgs, TSharedRef<
 	ViewRange = InArgs._ViewRange;
 
 	check( Node->GetType() == ESequencerNode::Track );
-	SectionAreaNode = StaticCastSharedRef<FTrackNode>( Node );
+	SectionAreaNode = StaticCastSharedRef<FSequencerTrackNode>( Node );
 
 	BackgroundBrush = FEditorStyle::GetBrush("Sequencer.SectionArea.Background");
 
diff --git a/Engine/Source/Editor/Sequencer/Private/SSequencerSectionAreaView.h b/Engine/Source/Editor/Sequencer/Private/SSequencerSectionAreaView.h
--- a/Engine/Source/Editor/Sequencer/Private/SSequencerSectionAreaView.h
+++ b/Engine/Source/Editor/Sequencer/Private/SSequencerSectionAreaView.h
@@ -3,7 +3,7 @@
 #pragma once
 
 class FSequencerSectionKeyAreaNode;
-class FTrackNode;
+class FSequencerTrackNode;
 
 /**
  * Visualizes a section area and its children                                        
@@ -53,7 +53,7 @@ private:
 	FSequencer& GetSequencer() const;
 private:
 	/** The node containing the sections we are viewing/manipulating */
-	TSharedPtr<FTrackNode> SectionAreaNode;
+	TSharedPtr<FSequencerTrackNode> SectionAreaNode;
 	/** The current view range */
 	TAttribute< TRange<float> > ViewRange;
 	/** Background brush of the section area */
diff --git a/Engine/Source/Editor/Sequencer/Private/SSequencerTrackArea.h b/Engine/Source/Editor/Sequencer/Private/SSequencerTrackArea.h
--- a/Engine/Source/Editor/Sequencer/Private/SSequencerTrackArea.h
+++ b/Engine/Source/Editor/Sequencer/Private/SSequencerTrackArea.h
@@ -41,9 +41,7 @@ class SSequencerTrackArea
 public:
 
 	SLATE_BEGIN_ARGS( SSequencerTrackArea )
-	{ }		
-		/** Lock the in/out to the start/end range */
-		SLATE_ATTRIBUTE( bool, LockInOutToStartEndRange )
+	{ }
 	SLATE_END_ARGS()
 
 	/** Construct this widget */
@@ -108,7 +106,4 @@ private:
 
 	/** Keep a hold of the size of the area so we can maintain zoom levels. */
 	TOptional<FVector2D> SizeLastFrame;
-
-	/** Whether to lock the in/out to the start/end range. */
-	TAttribute<bool> bLockInOutToStartEndRange;
 };
diff --git a/Engine/Source/Editor/Sequencer/Private/Sequencer.cpp b/Engine/Source/Editor/Sequencer/Private/Sequencer.cpp
--- a/Engine/Source/Editor/Sequencer/Private/Sequencer.cpp
+++ b/Engine/Source/Editor/Sequencer/Private/Sequencer.cpp
@@ -119,7 +119,7 @@ void FSequencer::InitSequencer(const FSequencerInitParams& InitParams, const TSh
 			.OnBeginScrubbing( this, &FSequencer::OnBeginScrubbing )
 			.OnEndScrubbing( this, &FSequencer::OnEndScrubbing )
 			.OnScrubPositionChanged( this, &FSequencer::OnScrubPositionChanged )
-			.OnViewRangeChanged( this, &FSequencer::OnViewRangeChanged )
+			.OnViewRangeChanged( this, &FSequencer::SetViewRange )
 			.OnClampRangeChanged( this, &FSequencer::OnClampRangeChanged )
 			.OnGetAddMenuContent(InitParams.ViewParams.OnGetAddMenuContent)
 			.AddMenuExtender(InitParams.ViewParams.AddMenuExtender);
@@ -169,12 +169,10 @@ FSequencer::FSequencer()
 	, TargetViewRange(0.f, 5.f)
 	, LastViewRange(0.f, 5.f)
 	, ViewRangeBeforeZoom(TRange<float>::Empty())
-	, bAutoScrollEnabled(true)
 	, PlaybackState( EMovieScenePlayerStatus::Stopped )
 	, ScrubPosition( 0.0f )
-	, bLoopingEnabled( false )
 	, bPerspectiveViewportPossessionEnabled( true )
-	, bPerspectiveViewportCameraCutEnabled( true )
+	, bPerspectiveViewportCameraCutEnabled( false )
 	, bIsEditingWithinLevelEditor( false )
 	, bNeedTreeRefresh( false )
 {
@@ -251,34 +249,7 @@ void FSequencer::Tick(float InDeltaTime)
 	if (PlaybackState == EMovieScenePlayerStatus::Playing ||
 		PlaybackState == EMovieScenePlayerStatus::Recording)
 	{
-		TRange<float> TimeBounds = GetTimeBounds();
-		if (!TimeBounds.IsEmpty())
-		{
-			if (NewTime > TimeBounds.GetUpperBoundValue())
-			{
-				if (bLoopingEnabled)
-				{
-					NewTime -= TimeBounds.Size<float>();
-				}
-				else
-				{
-					NewTime = TimeBounds.GetUpperBoundValue();
-					PlaybackState = EMovieScenePlayerStatus::Stopped;
-				}
-			}
-
-			if (NewTime < TimeBounds.GetLowerBoundValue())
-			{
-				NewTime = TimeBounds.GetLowerBoundValue();
-			}
-
-			SetGlobalTime(NewTime);
-		}
-		else
-		{
-			// no bounds at all, stop playing
-			PlaybackState = EMovieScenePlayerStatus::Stopped;
-		}
+		SetGlobalTimeLooped(NewTime);
 	}
 
 	// Tick all the tools we own as well
@@ -677,14 +648,19 @@ float FSequencer::GetGlobalTime()
 
 void FSequencer::SetGlobalTime( float NewTime )
 {
-	if (bAutoScrollEnabled && !bLoopingEnabled)
+	if (IsAutoScrollEnabled())
 	{
 		float RangeOffset = CalculateAutoscrollEncroachment(NewTime).Get(0.f);
 			
 		// When not scrubbing, we auto scroll the view range immediately
 		if (RangeOffset != 0.f)
 		{
-			SetViewRange(TRange<float>(TargetViewRange.GetLowerBoundValue() + RangeOffset, TargetViewRange.GetUpperBoundValue() + RangeOffset), EViewRangeInterpolation::Immediate);
+			TRange<float> WorkingRange = GetClampRange();
+			if (TargetViewRange.GetLowerBoundValue() + RangeOffset > WorkingRange.GetLowerBoundValue() &&
+				TargetViewRange.GetUpperBoundValue() + RangeOffset < WorkingRange.GetUpperBoundValue())
+			{
+				SetViewRange(TRange<float>(TargetViewRange.GetLowerBoundValue() + RangeOffset, TargetViewRange.GetUpperBoundValue() + RangeOffset), EViewRangeInterpolation::Immediate);
+			}
 		}
 	}
 
@@ -695,16 +671,10 @@ void FSequencer::SetGlobalTimeDirectly( float NewTime )
 {
 	float LastTime = ScrubPosition;
 
-	// Clamp to the view range when not auto-scrolling
-	if (!Settings->GetAutoScrollEnabled())
-	{
-		NewTime = FMath::Clamp(NewTime, TargetViewRange.GetLowerBoundValue(), TargetViewRange.GetUpperBoundValue());
-	}
-
 	// Update the position
 	ScrubPosition = NewTime;
 
-	SequenceInstanceStack.Top()->Update( ScrubPosition, LastTime, *this );
+	SequenceInstanceStack.Top()->Update(NewTime, LastTime, *this);
 
 	// If realtime is off, this needs to be called to update the pivot location when scrubbing.
 	GUnrealEd->UpdatePivotLocationForSelection();
@@ -740,18 +710,14 @@ void FSequencer::UpdateAutoScroll(float NewTime)
 		}
 	}
 
-	// Set the autoscrub rate to the autoscroll rate
-	if (Settings->GetShowRangeSlider() && Settings->GetLockInOutToStartEndRange())
+	// Don't autoscrub if we're at the extremes of the movie scene range
+	const TRange<float>& WorkingRange = GetFocusedMovieSceneSequence()->GetMovieScene()->GetEditorData().WorkingRange;
+	if (NewTime < WorkingRange.GetLowerBoundValue() + Threshold ||
+		NewTime > WorkingRange.GetUpperBoundValue() - Threshold
+		)
 	{
-		// Don't autoscrub if we're at the extremes of the movie scene range
-		const TRange<float>& WorkingRange = GetFocusedMovieSceneSequence()->GetMovieScene()->GetEditorData().WorkingRange;
-		if (NewTime < WorkingRange.GetLowerBoundValue() + Threshold ||
-			NewTime > WorkingRange.GetUpperBoundValue() - Threshold
-			)
-		{
-			AutoscrubOffset.Reset();
-			return;
-		}
+		AutoscrubOffset.Reset();
+		return;
 	}
 
 	// Scrub at the same rate we scroll
@@ -837,7 +803,9 @@ FGuid FSequencer::GetHandleToObject( UObject* Object, bool bCreateHandleIfMissin
 			// Create a new possessable
 			FocusedMovieScene->Modify();
 
-			ObjectGuid = FocusedMovieScene->AddPossessable(Object->GetName(), Object->GetClass());
+			AActor* PossessedActor = Cast<AActor>(Object);
+
+			ObjectGuid = FocusedMovieScene->AddPossessable(PossessedActor != nullptr ? PossessedActor->GetActorLabel() : Object->GetName(), Object->GetClass());
 			FocusedMovieSceneSequence->BindPossessableObject(ObjectGuid, *Object);
 			bPossessableAdded = true;
 		}
@@ -876,19 +844,26 @@ void FSequencer::GetRuntimeObjects( TSharedRef<FMovieSceneSequenceInstance> Movi
 	}
 }
 
-void FSequencer::UpdateCameraCut(UObject* ObjectToViewThrough, bool bNewCameraCut) const
+void FSequencer::UpdateCameraCut(UObject* CameraObject, UObject* UnlockIfCameraObject) const
 {
-	if(!IsPerspectiveViewportCameraCutEnabled())
+	if (!IsPerspectiveViewportCameraCutEnabled())
 	{
 		return;
 	}
 
-	for(FLevelEditorViewportClient* LevelVC : GEditor->LevelViewportClients)
+	ACameraActor* UnlockIfCameraActor = Cast<ACameraActor>(UnlockIfCameraObject);
+
+	for (FLevelEditorViewportClient* LevelVC : GEditor->LevelViewportClients)
 	{
-		if(LevelVC && LevelVC->IsPerspective() && LevelVC->AllowsCinematicPreview())
+		if ((LevelVC == nullptr) || !LevelVC->IsPerspective() || !LevelVC->AllowsCinematicPreview())
 		{
-			UpdatePreviewLevelViewportClientFromCameraCut( *LevelVC, ObjectToViewThrough, bNewCameraCut );
+			continue;
 		}
+
+		if ((CameraObject != nullptr) || LevelVC->IsLockedToActor(UnlockIfCameraActor))
+		{
+			UpdatePreviewLevelViewportClientFromCameraCut(*LevelVC, CameraObject);
+		}			
 	}
 }
 
@@ -1006,11 +981,6 @@ FReply FSequencer::OnPlay(bool bTogglePlay)
 		TRange<float> TimeBounds = GetTimeBounds();
 		if (!TimeBounds.IsEmpty())
 		{
-			float CurrentTime = GetGlobalTime();
-			if (CurrentTime < TimeBounds.GetLowerBoundValue() || CurrentTime >= TimeBounds.GetUpperBoundValue())
-			{
-				SetGlobalTime(TimeBounds.GetLowerBoundValue());
-			}
 			PlaybackState = EMovieScenePlayerStatus::Playing;
 			
 			// Make sure Slate ticks during playback
@@ -1045,24 +1015,24 @@ FReply FSequencer::OnRecord()
 FReply FSequencer::OnStepForward()
 {
 	PlaybackState = EMovieScenePlayerStatus::Stopped;
-	ScrubPosition += Settings->GetTimeSnapInterval();
-	if ( Settings->GetIsSnapEnabled() )
+	float NewPosition = ScrubPosition + Settings->GetTimeSnapInterval();
+	if (Settings->GetIsSnapEnabled())
 	{
-		ScrubPosition = Settings->SnapTimeToInterval(ScrubPosition);
+		NewPosition = Settings->SnapTimeToInterval(NewPosition);
 	}
-	SetGlobalTime(ScrubPosition);
+	SetGlobalTime(NewPosition);
 	return FReply::Handled();
 }
 
 FReply FSequencer::OnStepBackward()
 {
 	PlaybackState = EMovieScenePlayerStatus::Stopped;
-	ScrubPosition -= Settings->GetTimeSnapInterval();
-	if ( Settings->GetIsSnapEnabled() )
+	float NewPosition = ScrubPosition - Settings->GetTimeSnapInterval();
+	if (Settings->GetIsSnapEnabled())
 	{
-		ScrubPosition = Settings->SnapTimeToInterval(ScrubPosition);
+		NewPosition = Settings->SnapTimeToInterval(NewPosition);
 	}
-	SetGlobalTime(ScrubPosition);
+	SetGlobalTime(NewPosition);
 	return FReply::Handled();
 }
 
@@ -1090,13 +1060,44 @@ FReply FSequencer::OnStepToBeginning()
 
 FReply FSequencer::OnToggleLooping()
 {
-	bLoopingEnabled = !bLoopingEnabled;
+	Settings->SetLooping(!Settings->IsLooping());
 	return FReply::Handled();
 }
 
 bool FSequencer::IsLooping() const
 {
-	return bLoopingEnabled;
+	return Settings->IsLooping();
+}
+
+void FSequencer::SetGlobalTimeLooped(float InTime)
+{
+	if (Settings->IsLooping())
+	{
+		const UMovieSceneSequence* FocusedSequence = GetFocusedMovieSceneSequence();
+		if (FocusedSequence)
+		{
+			if (InTime > FocusedSequence->GetMovieScene()->GetPlaybackRange().GetUpperBoundValue())
+			{
+				InTime -= FocusedSequence->GetMovieScene()->GetPlaybackRange().Size<float>();
+			}
+		}
+	}
+	else
+	{
+		TRange<float> TimeBounds = GetTimeBounds();
+		if (InTime > TimeBounds.GetUpperBoundValue())
+		{
+			InTime = TimeBounds.GetUpperBoundValue();
+			PlaybackState = EMovieScenePlayerStatus::Stopped;
+		}
+		else if (InTime < TimeBounds.GetLowerBoundValue())
+		{
+			InTime = TimeBounds.GetLowerBoundValue();
+			PlaybackState = EMovieScenePlayerStatus::Stopped;
+		}
+	}
+
+	SetGlobalTime(InTime);
 }
 
 bool FSequencer::CanShowFrameNumbers() const
@@ -1143,38 +1144,19 @@ void FSequencer::UpdateTimeBoundsToFocusedMovieScene()
 
 TRange<float> FSequencer::GetTimeBounds() const
 {
+	const UMovieSceneSequence* FocusedSequence = GetFocusedMovieSceneSequence();
+
 	// When recording, we never want to constrain the time bound range.  You might not even have any sections or keys yet
 	// but we need to be able to move the time cursor during playback so you can capture data in real-time
-	if( PlaybackState == EMovieScenePlayerStatus::Recording )
+	if( PlaybackState == EMovieScenePlayerStatus::Recording || !FocusedSequence)
 	{
 		return TRange<float>( -100000.0f, 100000.0f );
 	}
-
-	const UMovieSceneSequence* FocusedSequence = GetFocusedMovieSceneSequence();
-
-	if (FocusedSequence == nullptr)
-	{
-		return TRange<float>(0.0f);
-	}
-
-	const UMovieScene* MovieScene = FocusedSequence->GetMovieScene();
-	const UMovieSceneTrack* AnimatableShot = MovieScene->FindMasterTrack<UMovieSceneShotTrack>();
-
-	if (AnimatableShot)
-	{
-		// try getting the bounds of all shots
-			TRange<float> Bounds = AnimatableShot->GetSectionBoundaries();
-		
-		if (!Bounds.IsEmpty())
-		{
-			return Bounds;
-		}
-	}
 	
-	return MovieScene->GetPlaybackRange();
+	return FocusedSequence->GetMovieScene()->GetEditorData().WorkingRange;
 }
 
-void FSequencer::SetViewRange(const TRange<float>& NewViewRange, EViewRangeInterpolation Interpolation)
+void FSequencer::SetViewRange(TRange<float> NewViewRange, EViewRangeInterpolation Interpolation)
 {
 	if (!ensure(NewViewRange.HasUpperBound() && NewViewRange.HasLowerBound() && !NewViewRange.IsDegenerate()))
 	{
@@ -1203,33 +1185,16 @@ void FSequencer::SetViewRange(const TRange<float>& NewViewRange, EViewRangeInter
 		ZoomAnimation.JumpToEnd();
 	}
 
-	GetFocusedMovieSceneSequence()->GetMovieScene()->GetEditorData().ViewRange = TargetViewRange;
-}
-
-void FSequencer::OnViewRangeChanged( TRange<float> NewViewRange, EViewRangeInterpolation Interpolation, bool bExpandClampRange )
-{
-	if (!NewViewRange.IsEmpty())
-	{
-		UMovieScene* FocusedMovieScene = GetFocusedMovieSceneSequence()->GetMovieScene();
-		// If locked, clamp the new range to clamp range
-		if (!Settings->GetAutoScrollEnabled() && Settings->GetShowRangeSlider() && Settings->GetLockInOutToStartEndRange() && !bExpandClampRange)
-		{
-			const TRange<float>& WorkingRange = FocusedMovieScene->GetEditorData().WorkingRange;
-			NewViewRange = TRange<float>::Intersection(NewViewRange, WorkingRange);
-		}
-		// Otherwise, grow the clamp range to the new range
-		else
-		{
-			TRange<float>& WorkingRange = FocusedMovieScene->GetEditorData().WorkingRange;
+	UMovieScene* FocusedMovieScene = GetFocusedMovieSceneSequence()->GetMovieScene();
+	FocusedMovieScene->GetEditorData().ViewRange = TargetViewRange;
 
-			WorkingRange = TRange<float>(
-				FMath::Min(NewViewRange.GetLowerBoundValue(), WorkingRange.GetLowerBoundValue()),
-				FMath::Max(NewViewRange.GetUpperBoundValue(), WorkingRange.GetUpperBoundValue())
-				);
-		}
+	// Always ensure the working range is big enough to fit the view range
+	TRange<float>& WorkingRange = FocusedMovieScene->GetEditorData().WorkingRange;
 
-		SetViewRange(NewViewRange, Interpolation);
-	}
+	WorkingRange = TRange<float>(
+		FMath::Min(TargetViewRange.GetLowerBoundValue(), WorkingRange.GetLowerBoundValue()),
+		FMath::Max(TargetViewRange.GetUpperBoundValue(), WorkingRange.GetUpperBoundValue())
+		);
 }
 
 void FSequencer::OnClampRangeChanged( TRange<float> NewClampRange )
@@ -1248,7 +1213,7 @@ void FSequencer::OnScrubPositionChanged( float NewScrubPosition, bool bScrubbing
 		{
 			OnEndScrubbing();
 		}
-		else if (bAutoScrollEnabled)
+		else if (IsAutoScrollEnabled())
 		{
 			UpdateAutoScroll(NewScrubPosition);
 			
@@ -1258,6 +1223,11 @@ void FSequencer::OnScrubPositionChanged( float NewScrubPosition, bool bScrubbing
 				return;
 			}
 		}
+		else
+		{
+			// Clamp to the view range when not auto-scrolling
+			NewScrubPosition = FMath::Clamp(NewScrubPosition, TargetViewRange.GetLowerBoundValue(), TargetViewRange.GetUpperBoundValue());
+		}
 	}
 
 	SetGlobalTimeDirectly( NewScrubPosition );
@@ -1288,8 +1258,12 @@ void FSequencer::StopAutoscroll()
 
 void FSequencer::OnToggleAutoScroll()
 {
-	bAutoScrollEnabled = !bAutoScrollEnabled;
-	Settings->SetAutoScrollEnabled(bAutoScrollEnabled);
+	Settings->SetAutoScrollEnabled(!Settings->GetAutoScrollEnabled());
+}
+
+bool FSequencer::IsAutoScrollEnabled() const
+{
+	return Settings->GetAutoScrollEnabled();
 }
 
 void FSequencer::FindInContentBrowser()
@@ -1342,7 +1316,8 @@ FGuid FSequencer::AddSpawnableForAssetOrClass( UObject* Object )
 		const FString NewSpawnableName = AssetName.ToString();		// @todo sequencer: Need UI to allow user to rename these slots
 
 		// Create our new blueprint!
-		UBlueprint* NewBlueprint = nullptr;
+		UBlueprint* NewBlueprint = Cast<UBlueprint>(Object);
+		if (!NewBlueprint)
 		{
 			// @todo sequencer: Add support for forcing specific factories for an asset
 			UActorFactory* FactoryToUse = nullptr;
@@ -1402,7 +1377,7 @@ FGuid FSequencer::AddSpawnableForAssetOrClass( UObject* Object )
 					UMovieScene3DTransformTrack* TransformTrack = Cast<UMovieScene3DTransformTrack>(OwnerMovieScene->AddTrack(UMovieScene3DTransformTrack::StaticClass(), NewSpawnableGuid));
 					if (TransformTrack)
 					{
-						FName Transform("Transform");
+						static FName Transform("Transform");
 						TransformTrack->SetPropertyNameAndPath(Transform, Transform.ToString());
 						const bool bUnwindRotation = false;
 
@@ -1470,7 +1445,7 @@ bool FSequencer::OnHandleAssetDropped(UObject* DroppedAsset, const FGuid& Target
 // Takes a display node and traverses it's parents to find the nearest track node if any.  Also collects the names of the nodes which make
 // up the path from the track node to the display node being checked.  The name path includes the name of the node being checked, but not
 // the name of the track node.
-void GetParentTrackNodeAndNamePath(TSharedRef<const FSequencerDisplayNode> DisplayNode, TSharedPtr<FTrackNode>& OutParentTrack, TArray<FName>& OutNamePath )
+void GetParentTrackNodeAndNamePath(TSharedRef<const FSequencerDisplayNode> DisplayNode, TSharedPtr<FSequencerTrackNode>& OutParentTrack, TArray<FName>& OutNamePath )
 {
 	TArray<FName> PathToTrack;
 	PathToTrack.Add( DisplayNode->GetNodeName() );
@@ -1482,7 +1457,7 @@ void GetParentTrackNodeAndNamePath(TSharedRef<const FSequencerDisplayNode> Displ
 	}
 	if ( CurrentParent.IsValid() )
 	{
-		OutParentTrack = StaticCastSharedPtr<FTrackNode>( CurrentParent );
+		OutParentTrack = StaticCastSharedPtr<FSequencerTrackNode>( CurrentParent );
 		for ( int32 i = PathToTrack.Num() - 1; i >= 0; i-- )
 		{
 			OutNamePath.Add( PathToTrack[i] );
@@ -1531,7 +1506,7 @@ void FSequencer::OnRequestNodeDeleted( TSharedRef<const FSequencerDisplayNode> N
 	}
 	else if( NodeToBeDeleted->GetType() == ESequencerNode::Track  )
 	{
-		TSharedRef<const FTrackNode> SectionAreaNode = StaticCastSharedRef<const FTrackNode>( NodeToBeDeleted );
+		TSharedRef<const FSequencerTrackNode> SectionAreaNode = StaticCastSharedRef<const FSequencerTrackNode>( NodeToBeDeleted );
 		UMovieSceneTrack* Track = SectionAreaNode->GetTrack();
 
 		if (Track != nullptr)
@@ -1554,7 +1529,7 @@ void FSequencer::OnRequestNodeDeleted( TSharedRef<const FSequencerDisplayNode> N
 	}
 	else if ( NodeToBeDeleted->GetType() == ESequencerNode::Category )
 	{
-		TSharedPtr<FTrackNode> ParentTrackNode;
+		TSharedPtr<FSequencerTrackNode> ParentTrackNode;
 		TArray<FName> PathFromTrack;
 		GetParentTrackNodeAndNamePath(NodeToBeDeleted, ParentTrackNode, PathFromTrack);
 		if ( ParentTrackNode.IsValid() )
@@ -1567,7 +1542,7 @@ void FSequencer::OnRequestNodeDeleted( TSharedRef<const FSequencerDisplayNode> N
 	}
 	else if ( NodeToBeDeleted->GetType() == ESequencerNode::KeyArea )
 	{
-		TSharedPtr<FTrackNode> ParentTrackNode;
+		TSharedPtr<FSequencerTrackNode> ParentTrackNode;
 		TArray<FName> PathFromTrack;
 		GetParentTrackNodeAndNamePath( NodeToBeDeleted, ParentTrackNode, PathFromTrack );
 		if ( ParentTrackNode.IsValid() )
@@ -1726,42 +1701,41 @@ void FSequencer::DeactivateDetailKeyframeHandler()
 	}
 }
 
-void FSequencer::UpdatePreviewLevelViewportClientFromCameraCut( FLevelEditorViewportClient& InViewportClient, UObject* InCameraObject, bool bNewCameraCut ) const
+void FSequencer::UpdatePreviewLevelViewportClientFromCameraCut(FLevelEditorViewportClient& InViewportClient, UObject* InCameraObject) const
 {
-	ACameraActor* Cam = Cast<ACameraActor>(InCameraObject);
-	if (Cam)
-	{
-		InViewportClient.SetViewLocation(Cam->GetActorLocation());
-		InViewportClient.SetViewRotation(Cam->GetActorRotation());
+	ACameraActor* CameraActor = Cast<ACameraActor>(InCameraObject);
 
-		InViewportClient.bEditorCameraCut = bNewCameraCut;
+	if (CameraActor)
+	{
+		InViewportClient.SetViewLocation(CameraActor->GetActorLocation());
+		InViewportClient.SetViewRotation(CameraActor->GetActorRotation());
+		InViewportClient.bEditorCameraCut = !InViewportClient.IsLockedToActor(CameraActor);
 	}
 	else
 	{
 		InViewportClient.ViewFOV = InViewportClient.FOVAngle;
-
 		InViewportClient.bEditorCameraCut = false;
 	}
 
 	// Set the actor lock.
-	InViewportClient.SetMatineeActorLock(Cam);
-	InViewportClient.bLockedCameraView = Cam != nullptr;
+	InViewportClient.SetMatineeActorLock(CameraActor);
+	InViewportClient.bLockedCameraView = CameraActor != nullptr;
 	InViewportClient.RemoveCameraRoll();
 
 	// If viewing through a camera - enforce aspect ratio.
-	if (Cam)
+	if (CameraActor)
 	{
-		if (Cam->GetCameraComponent()->AspectRatio == 0)
+		if (CameraActor->GetCameraComponent()->AspectRatio == 0)
 		{
 			InViewportClient.AspectRatio = 1.7f;
 		}
 		else
 		{
-			InViewportClient.AspectRatio = Cam->GetCameraComponent()->AspectRatio;
+			InViewportClient.AspectRatio = CameraActor->GetCameraComponent()->AspectRatio;
 		}
 
 		//don't stop the camera from zooming when not playing back
-		InViewportClient.ViewFOV = Cam->GetCameraComponent()->FieldOfView;
+		InViewportClient.ViewFOV = CameraActor->GetCameraComponent()->FieldOfView;
 
 		// If there are selected actors, invalidate the viewports hit proxies, otherwise they won't be selectable afterwards
 		if (InViewportClient.Viewport && GEditor->GetSelectedActorCount() > 0)
@@ -1877,21 +1851,19 @@ void FSequencer::ZoomToSelectedSections()
 	}
 
 	if (!BoundsHull.IsEmpty() && !BoundsHull.IsDegenerate())
-	{	
-		bool bExpandClampRange = true;
-
+	{
 		// Zoom back to last view range if already expanded
 		if (!ViewRangeBeforeZoom.IsEmpty() &&
 			FMath::IsNearlyEqual(BoundsHull.GetLowerBoundValue(), GetViewRange().GetLowerBoundValue(), KINDA_SMALL_NUMBER) &&
 			FMath::IsNearlyEqual(BoundsHull.GetUpperBoundValue(), GetViewRange().GetUpperBoundValue(), KINDA_SMALL_NUMBER))
 		{
-			OnViewRangeChanged(ViewRangeBeforeZoom, EViewRangeInterpolation::Animated, bExpandClampRange);
+			SetViewRange(ViewRangeBeforeZoom, EViewRangeInterpolation::Animated);
 		}
 		else
 		{
 			ViewRangeBeforeZoom = GetViewRange();
 
-			OnViewRangeChanged(BoundsHull, EViewRangeInterpolation::Animated, bExpandClampRange);
+			SetViewRange(BoundsHull, EViewRangeInterpolation::Animated);
 		}
 	}
 }
@@ -2369,12 +2341,11 @@ void FSequencer::BindSequencerCommands()
 		FCanExecuteAction::CreateLambda( []{ return true; } ),
 		FIsActionChecked::CreateLambda( [this]{ return Settings->GetKeyAllEnabled(); } ) );
 
-	bAutoScrollEnabled = Settings->GetAutoScrollEnabled();
 	SequencerCommandBindings->MapAction(
 		Commands.ToggleAutoScroll,
-		FExecuteAction::CreateSP( this, &FSequencer::OnToggleAutoScroll ),
+		FExecuteAction::CreateLambda( [this]{ Settings->SetAutoScrollEnabled( !Settings->GetAutoScrollEnabled() ); } ),
 		FCanExecuteAction::CreateLambda( []{ return true; } ),
-		FIsActionChecked::CreateSP( this, &FSequencer::GetAutoScrollEnabled ) );
+		FIsActionChecked::CreateLambda( [this]{ return Settings->GetAutoScrollEnabled(); } ) );
 
 	SequencerCommandBindings->MapAction(
 		Commands.FindInContentBrowser,
@@ -2400,12 +2371,6 @@ void FSequencer::BindSequencerCommands()
 		FCanExecuteAction::CreateLambda( []{ return true; } ),
 		FIsActionChecked::CreateLambda( [this]{ return Settings->GetShowRangeSlider(); } ) );
 
-	SequencerCommandBindings->MapAction(
-		Commands.ToggleLockInOutToStartEndRange,
-		FExecuteAction::CreateLambda( [this]{ Settings->SetLockInOutToStartEndRange( !Settings->GetLockInOutToStartEndRange() ); } ),
-		FCanExecuteAction::CreateLambda( [this]{ return Settings->GetShowRangeSlider() && !Settings->GetAutoScrollEnabled(); } ),
-		FIsActionChecked::CreateLambda( [this]{ return !Settings->GetShowRangeSlider() || Settings->GetAutoScrollEnabled() ? false : Settings->GetLockInOutToStartEndRange(); } ) );
-
 	SequencerCommandBindings->MapAction(
 		Commands.ToggleIsSnapEnabled,
 		FExecuteAction::CreateLambda( [this]{ Settings->SetIsSnapEnabled( !Settings->GetIsSnapEnabled() ); } ),
@@ -2485,6 +2450,17 @@ void FSequencer::BindSequencerCommands()
 				MovieSceneCapture->SetLevelSequenceAsset(GetCurrentAsset()->GetPathName());
 			}
 
+			if (CanShowFrameNumbers())
+			{
+				MovieSceneCapture->Settings.FrameRate = FMath::RoundToInt(1.f / Settings->GetTimeSnapInterval());
+				// We always add 1 to the number of frames we want to capture, because we want to capture both the start and end frames (which if the play range is 0, would still yield a single frame)
+				MovieSceneCapture->Settings.FrameCount = FMath::RoundToInt((GetPlaybackRange().GetUpperBoundValue() - GetPlaybackRange().GetLowerBoundValue()) * MovieSceneCapture->Settings.FrameRate) + 1;
+			}
+			else
+			{
+				MovieSceneCapture->Settings.FrameCount = 0;
+			}
+
 			IMovieSceneCaptureDialogModule::Get().OpenDialog(LevelEditorModule.GetLevelEditorTabManager().ToSharedRef(), MovieSceneCapture);
 		})
 	);
diff --git a/Engine/Source/Editor/Sequencer/Private/Sequencer.h b/Engine/Source/Editor/Sequencer/Private/Sequencer.h
--- a/Engine/Source/Editor/Sequencer/Private/Sequencer.h
+++ b/Engine/Source/Editor/Sequencer/Private/Sequencer.h
@@ -6,11 +6,19 @@
 #include "EditorUndoClient.h"
 
 
-class UMovieScene;
+class FMenuBuilder;
+class FMovieSceneSequenceInstance;
+class IDetailKeyframeHandler;
+class IMenu;
+class ISequencerEditTool;
+class ISequencerObjectChangeListener;
 class IToolkitHost;
+class SSequencer;
+class UClass;
+class UMovieScene;
+class UMovieSceneSection;
 class UMovieSceneSequence;
-class ISequencerObjectChangeListener;
-class IDetailKeyframeHandler;
+class UWorld;
 
 
 /**
@@ -92,7 +100,7 @@ public:
 	virtual void KeyProperty(FKeyPropertyParams KeyPropertyParams) override;
 	virtual FSequencerSelection& GetSelection() override;
 	virtual FSequencerSelectionPreview& GetSelectionPreview() override;
-	virtual void NotifyMapChanged(class UWorld* NewWorld, EMapChangeType MapChangeType) override;
+	virtual void NotifyMapChanged(UWorld* NewWorld, EMapChangeType MapChangeType) override;
 	virtual FOnGlobalTimeChanged& OnGlobalTimeChanged() override { return OnGlobalTimeChangedDelegate; }
 
 	/** Set the global time directly, without performing any auto-scroll */
@@ -214,7 +222,7 @@ public:
 	 *
 	 * @param MenuBuilder The menu builder to add things to.
 	 */
-	void BuildAddTrackMenu(class FMenuBuilder& MenuBuilder);
+	void BuildAddTrackMenu(FMenuBuilder& MenuBuilder);
 
 	/**
 	 * Builds up the track menu for object binding nodes in the outliner
@@ -223,7 +231,7 @@ public:
 	 * @param ObjectBinding	The object binding of the selected node
 	 * @param ObjectClass	The class of the selected object
 	 */
-	void BuildObjectBindingTrackMenu(class FMenuBuilder& MenuBuilder, const FGuid& ObjectBinding, const class UClass* ObjectClass);
+	void BuildObjectBindingTrackMenu(FMenuBuilder& MenuBuilder, const FGuid& ObjectBinding, const UClass* ObjectClass);
 
 	/**
 	 * Builds up the edit buttons for object binding nodes in the outliner
@@ -232,16 +240,16 @@ public:
 	 * @param ObjectBinding	The object binding of the selected node
 	 * @param ObjectClass	The class of the selected object
 	 */
-	void BuildObjectBindingEditButtons(TSharedPtr<SHorizontalBox> EditBox, const FGuid& ObjectBinding, const class UClass* ObjectClass);
+	void BuildObjectBindingEditButtons(TSharedPtr<SHorizontalBox> EditBox, const FGuid& ObjectBinding, const UClass* ObjectClass);
 
 	/** IMovieScenePlayer interface */
 	virtual void GetRuntimeObjects( TSharedRef<FMovieSceneSequenceInstance> MovieSceneInstance, const FGuid& ObjectHandle, TArray< UObject* >& OutObjects ) const override;
-	virtual void UpdateCameraCut(UObject* ObjectToViewThrough, bool bNewCameraCut) const override;
+	virtual void UpdateCameraCut(UObject* CameraObject, UObject* UnlockIfCameraObject) const override;
 	virtual void SetViewportSettings(const TMap<FViewportClient*, EMovieSceneViewportParams>& ViewportParamsMap) override;
 	virtual void GetViewportSettings(TMap<FViewportClient*, EMovieSceneViewportParams>& ViewportParamsMap) const override;
 	virtual EMovieScenePlayerStatus::Type GetPlaybackStatus() const override;
-	virtual void AddOrUpdateMovieSceneInstance( class UMovieSceneSection& MovieSceneSection, TSharedRef<FMovieSceneSequenceInstance> InstanceToAdd ) override;
-	virtual void RemoveMovieSceneInstance( class UMovieSceneSection& MovieSceneSection, TSharedRef<FMovieSceneSequenceInstance> InstanceToRemove ) override;
+	virtual void AddOrUpdateMovieSceneInstance(UMovieSceneSection& MovieSceneSection, TSharedRef<FMovieSceneSequenceInstance> InstanceToAdd) override;
+	virtual void RemoveMovieSceneInstance(UMovieSceneSection& MovieSceneSection, TSharedRef<FMovieSceneSequenceInstance> InstanceToRemove) override;
 
 	/** Called when an actor is dropped into Sequencer */
 	void OnActorsDropped( const TArray<TWeakObjectPtr<AActor> >& Actors );
@@ -256,6 +264,10 @@ public:
 	FReply OnToggleLooping();
 
 	bool IsLooping() const;
+
+	/** Set the new global time, accounting for looping options */
+	void SetGlobalTimeLooped(float InTime);
+
 	EPlaybackMode::Type GetPlaybackMode() const;
 
 	/** Called to determine whether a frame number is set so that frame numbers can be shown */
@@ -319,18 +331,11 @@ protected:
 	virtual float OnGetScrubPosition() const { return ScrubPosition; }
 
 	/**
-	 * Set the view range
+	 * Set the view range, growing the working range to accomodate, if necessary
 	 * @param NewViewRange The new view range. Must be a finite range
 	 * @param Interpolation How to interpolate to the new view range
 	 */
-	void SetViewRange(const TRange<float>& NewViewRange, EViewRangeInterpolation Interpolation = EViewRangeInterpolation::Animated);
-
-	/**
-	 * Called when the view range is changed by the user
-	 *
-	 * @param	NewViewRange The new view range
-	 */
-	void OnViewRangeChanged( TRange<float> NewViewRange, EViewRangeInterpolation Interpolation = EViewRangeInterpolation::Animated, bool bExpandClampRange = false );
+	void SetViewRange(TRange<float> NewViewRange, EViewRangeInterpolation Interpolation = EViewRangeInterpolation::Animated);
 
 	/**
 	 * Called when the clamp range is changed by the user
@@ -383,10 +388,7 @@ protected:
 	 *
 	 * @return true if auto-scroll is enabled, false otherwise.
 	 */
-	bool GetAutoScrollEnabled() const
-	{
-		return bAutoScrollEnabled;
-	}
+	bool IsAutoScrollEnabled() const;
 
 	/** 
 	 * Find the viewed sequence asset in the content browser
@@ -459,10 +461,10 @@ protected:
 	void OnSelectedOutlinerNodesChanged();
 
 	/** Called before the world is going to be saved. The sequencer puts everything back to its initial state. */
-	void OnPreSaveWorld(uint32 SaveFlags, class UWorld* World);
+	void OnPreSaveWorld(uint32 SaveFlags, UWorld* World);
 
 	/** Called after the world has been saved. The sequencer updates to the animated state. */
-	void OnPostSaveWorld(uint32 SaveFlags, class UWorld* World, bool bSuccess);
+	void OnPostSaveWorld(uint32 SaveFlags, UWorld* World, bool bSuccess);
 
 	/** Called after a new level has been created. The sequencer editor mode needs to be enabled. */
 	void OnNewCurrentLevel();
@@ -477,7 +479,7 @@ protected:
 	void OnEndPIE(bool bIsSimulating);
 
 	/** Updates a viewport client from camera cut data */
-	void UpdatePreviewLevelViewportClientFromCameraCut( FLevelEditorViewportClient& InViewportClient, UObject* InCameraObject, bool bNewCameraCut ) const;
+	void UpdatePreviewLevelViewportClientFromCameraCut(FLevelEditorViewportClient& InViewportClient, UObject* InCameraObject) const;
 
 	/** Is the sequencer widget focused? */
 	bool IsSequencerWidgetFocused() const;
@@ -496,16 +498,16 @@ private:
 	TArray<TSharedPtr<ISequencerTrackEditor>> TrackEditors;
 
 	/** Listener for object changes being made while this sequencer is open*/
-	TSharedPtr< class ISequencerObjectChangeListener > ObjectChangeListener;
+	TSharedPtr<ISequencerObjectChangeListener> ObjectChangeListener;
 
 	/** Listener for object changes being made while this sequencer is open*/
-	TSharedPtr< class IDetailKeyframeHandler > DetailKeyframeHandler;
+	TSharedPtr<IDetailKeyframeHandler> DetailKeyframeHandler;
 
 	/** The runtime instance for the root movie scene */
-	TSharedPtr< class FMovieSceneSequenceInstance > RootMovieSceneSequenceInstance;
+	TSharedPtr<FMovieSceneSequenceInstance> RootMovieSceneSequenceInstance;
 
 	/** Main sequencer widget */
-	TSharedPtr< class SSequencer > SequencerWidget;
+	TSharedPtr<SSequencer> SequencerWidget;
 	
 	/** The asset editor that created this Sequencer if any */
 	TWeakPtr<IToolkitHost> ToolkitHost;
@@ -532,8 +534,6 @@ private:
 
 	/** The amount of autoscrub offset that is currently being applied */
 	TOptional<float> AutoscrubOffset;
-	/** Whether or not we are allowing autoscroll */
-	bool bAutoScrollEnabled;
 
 	/** Zoom smoothing curves */
 	FCurveSequence ZoomAnimation;
@@ -550,9 +550,6 @@ private:
 	// @todo sequencer: Should use FTimespan or "double" for Time Cursor Position! (cascades)
 	float ScrubPosition;
 
-	/** Whether looping while playing is enabled for this sequencer */
-	bool bLoopingEnabled;
-
 	bool bPerspectiveViewportPossessionEnabled;
 	bool bPerspectiveViewportCameraCutEnabled;
 
@@ -560,7 +557,7 @@ private:
 	bool bIsEditingWithinLevelEditor;
 
 	/** Generic Popup Entry */
-	TWeakPtr<class IMenu> EntryPopupMenu;
+	TWeakPtr<IMenu> EntryPopupMenu;
 
 	/** Stores a dirty bit for whether the sequencer tree (and other UI bits) may need to be refreshed.  We
 	    do this simply to avoid refreshing the UI more than once per frame. (e.g. during live recording where
diff --git a/Engine/Source/Editor/Sequencer/Private/SequencerCommands.cpp b/Engine/Source/Editor/Sequencer/Private/SequencerCommands.cpp
--- a/Engine/Source/Editor/Sequencer/Private/SequencerCommands.cpp
+++ b/Engine/Source/Editor/Sequencer/Private/SequencerCommands.cpp
@@ -42,7 +42,6 @@ void FSequencerCommands::RegisterCommands()
 
 	UI_COMMAND( ToggleShowFrameNumbers, "Show Frame Numbers", "Enables and disables showing frame numbers", EUserInterfaceActionType::ToggleButton, FInputChord() );
 	UI_COMMAND( ToggleShowRangeSlider, "Show Range Slider", "Enables and disables showing the time range slider", EUserInterfaceActionType::ToggleButton, FInputChord() );
-	UI_COMMAND( ToggleLockInOutToStartEndRange, "Lock In/Out to Start/End", "Enables and disables locking the in/out to the start/end range while scrolling the view", EUserInterfaceActionType::ToggleButton, FInputChord() );
 	UI_COMMAND( ToggleIsSnapEnabled, "Enable Snapping", "Enables and disables snapping", EUserInterfaceActionType::ToggleButton, FInputChord() );
 
 	UI_COMMAND( ToggleSnapKeyTimesToInterval, "Snap to the Interval", "Snap keys to the interval", EUserInterfaceActionType::ToggleButton, FInputChord() );
diff --git a/Engine/Source/Editor/Sequencer/Private/SequencerCommonHelpers.cpp b/Engine/Source/Editor/Sequencer/Private/SequencerCommonHelpers.cpp
--- a/Engine/Source/Editor/Sequencer/Private/SequencerCommonHelpers.cpp
+++ b/Engine/Source/Editor/Sequencer/Private/SequencerCommonHelpers.cpp
@@ -14,7 +14,7 @@ void SequencerHelpers::GetAllKeyAreas(TSharedPtr<FSequencerDisplayNode> DisplayN
 
 		if (NodeToCheck->GetType() == ESequencerNode::Track)
 		{
-			TSharedPtr<FTrackNode> TrackNode = StaticCastSharedPtr<FTrackNode>(NodeToCheck);
+			TSharedPtr<FSequencerTrackNode> TrackNode = StaticCastSharedPtr<FSequencerTrackNode>(NodeToCheck);
 			TArray<TSharedRef<FSequencerSectionKeyAreaNode>> KeyAreaNodes;
 			TrackNode->GetChildKeyAreaNodesRecursively(KeyAreaNodes);
 			for (TSharedRef<FSequencerSectionKeyAreaNode> KeyAreaNode : KeyAreaNodes)
diff --git a/Engine/Source/Editor/Sequencer/Private/SequencerCommonHelpers.h b/Engine/Source/Editor/Sequencer/Private/SequencerCommonHelpers.h
--- a/Engine/Source/Editor/Sequencer/Private/SequencerCommonHelpers.h
+++ b/Engine/Source/Editor/Sequencer/Private/SequencerCommonHelpers.h
@@ -6,6 +6,7 @@
 
 #define LOCTEXT_NAMESPACE "SequencerHelpers"
 
+
 class SequencerHelpers
 {
 public:
@@ -30,6 +31,7 @@ public:
 	static float FrameToTime(int32 Frame, float FrameRate);
 };
 
+
 class SequencerSnapValues
 {
 public:
@@ -38,44 +40,50 @@ public:
 	static const TArray<SNumericDropDown<float>::FNamedValue>& GetSnapValues()
 	{
 		static TArray<SNumericDropDown<float>::FNamedValue> SnapValues;
-		static bool bInitSnapValues = false;
-		if (!bInitSnapValues)
+		static bool SnapValuesInitialized = false;
+
+		if (!SnapValuesInitialized)
 		{
-			bInitSnapValues = true;
 			SnapValues.Add( SNumericDropDown<float>::FNamedValue( 0.001f, LOCTEXT( "Snap_OneThousandth", "0.001" ), LOCTEXT( "SnapDescription_OneThousandth", "Set snap to 1/1000th" ) ) );
 			SnapValues.Add( SNumericDropDown<float>::FNamedValue( 0.01f, LOCTEXT( "Snap_OneHundredth", "0.01" ), LOCTEXT( "SnapDescription_OneHundredth", "Set snap to 1/100th" ) ) );
 			SnapValues.Add( SNumericDropDown<float>::FNamedValue( 0.1f, LOCTEXT( "Snap_OneTenth", "0.1" ), LOCTEXT( "SnapDescription_OneTenth", "Set snap to 1/10th" ) ) );
 			SnapValues.Add( SNumericDropDown<float>::FNamedValue( 1.0f, LOCTEXT( "Snap_One", "1" ), LOCTEXT( "SnapDescription_One", "Set snap to 1" ) ) );
 			SnapValues.Add( SNumericDropDown<float>::FNamedValue( 10.0f, LOCTEXT( "Snap_Ten", "10" ), LOCTEXT( "SnapDescription_Ten", "Set snap to 10" ) ) );
 			SnapValues.Add( SNumericDropDown<float>::FNamedValue( 100.0f, LOCTEXT( "Snap_OneHundred", "100" ), LOCTEXT( "SnapDescription_OneHundred", "Set snap to 100" ) ) );
+
+			SnapValuesInitialized = true;
 		}
+
 		return SnapValues;
 	}
 
 	static const TArray<SNumericDropDown<float>::FNamedValue>& GetSecondsSnapValues()
 	{
 		static TArray<SNumericDropDown<float>::FNamedValue> SecondsSnapValues;
-		static bool bInitSecondsSnapValues = false;
-		if (!bInitSecondsSnapValues)
+		static bool SecondsSnapValuesInitialized = false;
+
+		if (!SecondsSnapValuesInitialized)
 		{
-			bInitSecondsSnapValues = true;
 			SecondsSnapValues.Add( SNumericDropDown<float>::FNamedValue( 0.001f, LOCTEXT( "Snap_OneThousandthSeconds", "0.001s" ), LOCTEXT( "SnapDescription_OneThousandthSeconds", "Set snap to 1/1000th of a second" ) ) );
 			SecondsSnapValues.Add( SNumericDropDown<float>::FNamedValue( 0.01f, LOCTEXT( "Snap_OneHundredthSeconds", "0.01s" ), LOCTEXT( "SnapDescription_OneHundredthSeconds", "Set snap to 1/100th of a second" ) ) );
 			SecondsSnapValues.Add( SNumericDropDown<float>::FNamedValue( 0.1f, LOCTEXT( "Snap_OneTenthSeconds", "0.1s" ), LOCTEXT( "SnapDescription_OneTenthSeconds", "Set snap to 1/10th of a second" ) ) );
 			SecondsSnapValues.Add( SNumericDropDown<float>::FNamedValue( 1.0f, LOCTEXT( "Snap_OneSeconds", "1s" ), LOCTEXT( "SnapDescription_OneSeconds", "Set snap to 1 second" ) ) );
 			SecondsSnapValues.Add( SNumericDropDown<float>::FNamedValue( 10.0f, LOCTEXT( "Snap_TenSeconds", "10s" ), LOCTEXT( "SnapDescription_TenSeconds", "Set snap to 10 seconds" ) ) );
 			SecondsSnapValues.Add( SNumericDropDown<float>::FNamedValue( 100.0f, LOCTEXT( "Snap_OneHundredSeconds", "100s" ), LOCTEXT( "SnapDescription_OneHundredSeconds", "Set snap to 100 seconds" ) ) );
+
+			SecondsSnapValuesInitialized = true;
 		}
+
 		return SecondsSnapValues;
 	}
 
 	static const TArray<SNumericDropDown<float>::FNamedValue>& GetFrameRateSnapValues()
 	{
 		static TArray<SNumericDropDown<float>::FNamedValue> FrameRateSnapValues;
-		static bool bInitFrameRateSnapValues = false;
-		if (!bInitFrameRateSnapValues)
+		static bool FrameRateSnapValuesInitialized = false;
+
+		if (!FrameRateSnapValuesInitialized)
 		{
-			bInitFrameRateSnapValues = true;
 			FrameRateSnapValues.Add( SNumericDropDown<float>::FNamedValue( 0.066667f, LOCTEXT( "Snap_15Fps", "15 fps" ), LOCTEXT( "SnapDescription_15Fps", "Set snap to 15 fps" ) ) );
 			FrameRateSnapValues.Add( SNumericDropDown<float>::FNamedValue( 0.041667f, LOCTEXT( "Snap_24Fps", "24 fps (film)" ), LOCTEXT( "SnapDescription_24Fps", "Set snap to 24 fps" ) ) );
 			FrameRateSnapValues.Add( SNumericDropDown<float>::FNamedValue( 0.04f, LOCTEXT( "Snap_25Fps", "25 fps (PAL/25)" ), LOCTEXT( "SnapDescription_25Fps", "Set snap to 25 fps" ) ) );
@@ -86,19 +94,23 @@ public:
 			FrameRateSnapValues.Add( SNumericDropDown<float>::FNamedValue( 0.016683f, LOCTEXT( "Snap_50.94Fps", "50.94 fps (NTSC/60)" ), LOCTEXT( "SnapDescription_50.94Fps", "Set snap to 50.94 fps" ) ) );
 			FrameRateSnapValues.Add( SNumericDropDown<float>::FNamedValue( 0.016667f, LOCTEXT( "Snap_60Fps", "60 fps" ), LOCTEXT( "SnapDescription_60Fps", "Set snap to 60 fps" ) ) );
 			FrameRateSnapValues.Add( SNumericDropDown<float>::FNamedValue( 0.008334f, LOCTEXT( "Snap_120Fps", "120 fps" ), LOCTEXT( "SnapDescription_120Fps", "Set snap to 120 fps" ) ) );
+
+			FrameRateSnapValuesInitialized = true;
 		}
+
 		return FrameRateSnapValues;
 	}
 
 	static const TArray<SNumericDropDown<float>::FNamedValue>& GetTimeSnapValues()
 	{
 		static TArray<SNumericDropDown<float>::FNamedValue> TimeSnapValues;
-		static bool bInitTimeSnapValues = false;
-		if (!bInitTimeSnapValues)
+		static bool TimeSnapValuesInitialized = false;
+		if (!TimeSnapValuesInitialized)
 		{
-			bInitTimeSnapValues = true;
 			TimeSnapValues.Append(GetSecondsSnapValues());
 			TimeSnapValues.Append(GetFrameRateSnapValues());
+
+			TimeSnapValuesInitialized = true;
 		}
 		return TimeSnapValues;
 	}
@@ -116,4 +128,5 @@ public:
 	}
 };
 
-#undef LOCTEXT_NAMESPACE
\ No newline at end of file
+
+#undef LOCTEXT_NAMESPACE
diff --git a/Engine/Source/Editor/Sequencer/Private/SequencerCurveOwner.cpp b/Engine/Source/Editor/Sequencer/Private/SequencerCurveOwner.cpp
--- a/Engine/Source/Editor/Sequencer/Private/SequencerCurveOwner.cpp
+++ b/Engine/Source/Editor/Sequencer/Private/SequencerCurveOwner.cpp
@@ -31,7 +31,7 @@ void GetAllKeyAreas( TSharedPtr<FSequencerNodeTree> InSequencerNodeTree, TArray<
 		NodesToProcess.RemoveNode( Node );
 		if ( Node->GetType() == ESequencerNode::Track )
 		{
-			TSharedRef<FTrackNode> TrackNode = StaticCastSharedRef<FTrackNode>( Node );
+			TSharedRef<FSequencerTrackNode> TrackNode = StaticCastSharedRef<FSequencerTrackNode>( Node );
 			TSharedPtr<FSequencerSectionKeyAreaNode> TopLevelKeyNode = TrackNode->GetTopLevelKeyNode();
 			if ( TopLevelKeyNode.IsValid() )
 			{
diff --git a/Engine/Source/Editor/Sequencer/Private/SequencerHotspots.h b/Engine/Source/Editor/Sequencer/Private/SequencerHotspots.h
--- a/Engine/Source/Editor/Sequencer/Private/SequencerHotspots.h
+++ b/Engine/Source/Editor/Sequencer/Private/SequencerHotspots.h
@@ -11,7 +11,7 @@
 struct FKeyHotspot
 	: ISequencerHotspot
 {
-	FKeyHotspot(FSequencerSelectedKey InKey, TSharedPtr<FTrackNode> InTrackNode)
+	FKeyHotspot(FSequencerSelectedKey InKey, TSharedPtr<FSequencerTrackNode> InTrackNode)
 		: Key(InKey)
 		, TrackNode(MoveTemp(InTrackNode))
 	{ }
@@ -23,21 +23,21 @@ struct FKeyHotspot
 	FSequencerSelectedKey Key;
 
 	/** The track node that holds the section the key is contained in */
-	TSharedPtr<FTrackNode> TrackNode;
+	TSharedPtr<FSequencerTrackNode> TrackNode;
 };
 
 
 /** Structure used to encapsulate a section and its track node */
 struct FSectionHandle
 {
-	FSectionHandle(TSharedPtr<FTrackNode> InTrackNode, int32 InSectionIndex)
+	FSectionHandle(TSharedPtr<FSequencerTrackNode> InTrackNode, int32 InSectionIndex)
 		: SectionIndex(InSectionIndex), TrackNode(MoveTemp(InTrackNode))
 	{ }
 
 	UMovieSceneSection* GetSectionObject() const { return TrackNode->GetSections()[SectionIndex]->GetSectionObject(); }
 
 	int32 SectionIndex;
-	TSharedPtr<FTrackNode> TrackNode;
+	TSharedPtr<FSequencerTrackNode> TrackNode;
 };
 
 
diff --git a/Engine/Source/Editor/Sequencer/Private/SequencerNodeTree.cpp b/Engine/Source/Editor/Sequencer/Private/SequencerNodeTree.cpp
--- a/Engine/Source/Editor/Sequencer/Private/SequencerNodeTree.cpp
+++ b/Engine/Source/Editor/Sequencer/Private/SequencerNodeTree.cpp
@@ -23,39 +23,38 @@ void FSequencerNodeTree::Empty()
 	FilteredNodes.Empty();
 }
 
+
 void FSequencerNodeTree::Update()
 {
 	// @todo Sequencer - This update pass is too aggressive.  Some nodes may still be valid
 	Empty();
 
 	UMovieScene* MovieScene = Sequencer.GetFocusedMovieSceneSequence()->GetMovieScene();
-
-	TArray< TSharedRef<FSequencerDisplayNode> > NewRootNodes;
+	TArray<TSharedRef<FSequencerDisplayNode>> NewRootNodes;
 
 	// Get the master tracks  so we can get sections from them
 	const TArray<UMovieSceneTrack*>& MasterTracks = MovieScene->GetMasterTracks();
 
-	for( UMovieSceneTrack* Track : MasterTracks )
+	for (UMovieSceneTrack* Track : MasterTracks)
 	{
 		UMovieSceneTrack& TrackRef = *Track;
 
-		TSharedRef<FTrackNode> SectionNode = MakeShareable( new FTrackNode( TrackRef.GetTrackName(), TrackRef, *FindOrAddTypeEditor(TrackRef), NULL, *this ) );
-		NewRootNodes.Add( SectionNode );
+		TSharedRef<FSequencerTrackNode> SectionNode = MakeShareable(new FSequencerTrackNode(TrackRef, *FindOrAddTypeEditor(TrackRef), nullptr, *this));
+		NewRootNodes.Add(SectionNode);
 	
-		MakeSectionInterfaces( TrackRef, SectionNode );
+		MakeSectionInterfaces(TrackRef, SectionNode);
 	}
 
-
 	const TArray<FMovieSceneBinding>& Bindings = MovieScene->GetBindings();
-
 	TMap<FGuid, const FMovieSceneBinding*> GuidToBindingMap;
+
 	for (const FMovieSceneBinding& Binding : Bindings)
 	{
 		GuidToBindingMap.Add(Binding.GetObjectGuid(), &Binding);
 	}
 
 	// Make nodes for all object bindings
-	TArray< TSharedRef<FSequencerDisplayNode> > NewObjectNodes;
+	TArray<TSharedRef<FSequencerDisplayNode>> NewObjectNodes;
 	for( const FMovieSceneBinding& Binding : Bindings )
 	{
 		TSharedRef<FSequencerObjectBindingNode> ObjectBindingNode = AddObjectBinding( Binding.GetName(), Binding.GetObjectGuid(), GuidToBindingMap, NewObjectNodes );
@@ -65,15 +64,12 @@ void FSequencerNodeTree::Update()
 		for( UMovieSceneTrack* Track : Tracks )
 		{
 			UMovieSceneTrack& TrackRef = *Track;
-
-			FName SectionName = TrackRef.GetTrackName();
-			check( SectionName != NAME_None );
-
-			TSharedRef<FTrackNode> SectionAreaNode = ObjectBindingNode->AddSectionAreaNode( SectionName, TrackRef, *FindOrAddTypeEditor( TrackRef ) );
+			TSharedRef<FSequencerTrackNode> SectionAreaNode = ObjectBindingNode->AddSectionAreaNode(TrackRef, *FindOrAddTypeEditor(TrackRef));
 			MakeSectionInterfaces( TrackRef, SectionAreaNode );
 		}
 	}
 
+
 	struct FObjectNodeSorter
 	{
 		bool operator()( const TSharedRef<FSequencerDisplayNode>& A, const TSharedRef<FSequencerDisplayNode>& B ) const
@@ -82,18 +78,22 @@ void FSequencerNodeTree::Update()
 			{
 				return true;
 			}
+
 			if (A->GetType() != ESequencerNode::Object && B->GetType() == ESequencerNode::Object)
 			{
 				return false;
 			}
+
 			if ( A->GetType() == ESequencerNode::Object && B->GetType() == ESequencerNode::Object )
 			{
 				return A->GetDisplayName().ToString() < B->GetDisplayName().ToString();
 			}
+
 			return 0;
 		}
 	};
 
+
 	NewObjectNodes.Sort( FObjectNodeSorter() );
 	for (TSharedRef<FSequencerDisplayNode> NewObjectNode : NewObjectNodes)
 	{
@@ -104,25 +104,25 @@ void FSequencerNodeTree::Update()
 
 	// Look for a shot track.  It will always come first if it exists
 	UMovieSceneTrack* ShotTrack = MovieScene->GetShotTrack();
-	if( ShotTrack )
+	
+	if(ShotTrack)
 	{
-		TSharedRef<FTrackNode> SectionNode = MakeShareable( new FTrackNode( ShotTrack->GetTrackName(), *ShotTrack, *FindOrAddTypeEditor( *ShotTrack ), NULL, *this ) );
+		TSharedRef<FSequencerTrackNode> SectionNode = MakeShareable(new FSequencerTrackNode(*ShotTrack, *FindOrAddTypeEditor(*ShotTrack), nullptr, *this));
 
 		// Shot track always comes first
-		RootNodes.Add( SectionNode );
-
-		MakeSectionInterfaces( *ShotTrack, SectionNode );
+		RootNodes.Add(SectionNode);
+		MakeSectionInterfaces(*ShotTrack, SectionNode);
 	}
 
 	// Add all other nodes after the shot track
-	RootNodes.Append( NewRootNodes );
+	RootNodes.Append(NewRootNodes);
 
 	// Set up virtual offsets, and expansion states
 	float VerticalOffset = 0.f;
+
 	for (auto& Node : RootNodes)
 	{
-		Node->Traverse_ParentFirst([&](FSequencerDisplayNode& InNode){
-
+		Node->Traverse_ParentFirst([&](FSequencerDisplayNode& InNode) {
 			float VerticalTop = VerticalOffset;
 			VerticalOffset += InNode.GetNodeHeight() + InNode.GetNodePadding().Combined();
 			InNode.Initialize(VerticalTop, VerticalOffset);
@@ -132,27 +132,30 @@ void FSequencerNodeTree::Update()
 	}
 
 	// Re-filter the tree after updating 
-	// @todo Sequencer - Newly added sections may need to be visible even when there is a filter
+	// @todo sequencer: Newly added sections may need to be visible even when there is a filter
 	FilterNodes( FilterString );
 }
 
+
 TSharedRef<ISequencerTrackEditor> FSequencerNodeTree::FindOrAddTypeEditor( UMovieSceneTrack& InTrack )
 {
 	TSharedPtr<ISequencerTrackEditor> Editor = EditorMap.FindRef( &InTrack );
 
 	if( !Editor.IsValid() )
 	{
-		const TArray< TSharedPtr<ISequencerTrackEditor> >& TrackEditors = Sequencer.GetTrackEditors();
+		const TArray<TSharedPtr<ISequencerTrackEditor>>& TrackEditors = Sequencer.GetTrackEditors();
 
 		// Get a tool for each track
-		// @todo Sequencer - Should probably only need to get this once and it shouldnt be done here. It depends on when movie scene tool modules are loaded
+		// @todo sequencer: Should probably only need to get this once and it shouldn't be done here. It depends on when movie scene tool modules are loaded
 		TSharedPtr<ISequencerTrackEditor> SupportedTool;
-		for( int32 EditorIndex = 0; EditorIndex < TrackEditors.Num(); ++EditorIndex )
+
+		for (const auto& TrackEditor : TrackEditors)
 		{
-			if( TrackEditors[EditorIndex]->SupportsType( InTrack.GetClass() ) )
+			if (TrackEditor->SupportsType(InTrack.GetClass()))
 			{
-				Editor = TrackEditors[EditorIndex];
-				EditorMap.Add( &InTrack, Editor );
+				EditorMap.Add(&InTrack, TrackEditor);
+				Editor = TrackEditor;
+
 				break;
 			}
 		}
@@ -161,7 +164,8 @@ TSharedRef<ISequencerTrackEditor> FSequencerNodeTree::FindOrAddTypeEditor( UMovi
 	return Editor.ToSharedRef();
 }
 
-void FSequencerNodeTree::MakeSectionInterfaces( UMovieSceneTrack& Track, TSharedRef<FTrackNode>& SectionAreaNode )
+
+void FSequencerNodeTree::MakeSectionInterfaces( UMovieSceneTrack& Track, TSharedRef<FSequencerTrackNode>& SectionAreaNode )
 {
 	const TArray<UMovieSceneSection*>& MovieSceneSections = Track.GetAllSections();
 
@@ -175,20 +179,20 @@ void FSequencerNodeTree::MakeSectionInterfaces( UMovieSceneTrack& Track, TShared
 		// Ask the section to generate it's inner layout
 		FSequencerSectionLayoutBuilder Builder( SectionAreaNode );
 		Section->GenerateSectionLayout( Builder );
-
 		SectionAreaNode->AddSection( Section );
 	}
 
 	SectionAreaNode->FixRowIndices();
 }
 
-const TArray< TSharedRef<FSequencerDisplayNode> >& FSequencerNodeTree::GetRootNodes() const
+
+const TArray<TSharedRef<FSequencerDisplayNode>>& FSequencerNodeTree::GetRootNodes() const
 {
 	return RootNodes;
 }
 
 
-TSharedRef<FSequencerObjectBindingNode> FSequencerNodeTree::AddObjectBinding(const FString& ObjectName, const FGuid& ObjectBinding, TMap<FGuid, const FMovieSceneBinding*>& GuidToBindingMap, TArray< TSharedRef<FSequencerDisplayNode> >& OutNodeList)
+TSharedRef<FSequencerObjectBindingNode> FSequencerNodeTree::AddObjectBinding(const FString& ObjectName, const FGuid& ObjectBinding, TMap<FGuid, const FMovieSceneBinding*>& GuidToBindingMap, TArray<TSharedRef<FSequencerDisplayNode>>& OutNodeList)
 {
 	TSharedPtr<FSequencerObjectBindingNode> ObjectNode;
 	TSharedPtr<FSequencerObjectBindingNode>* FoundObjectNode = ObjectBindingMap.Find(ObjectBinding);
@@ -203,15 +207,30 @@ TSharedRef<FSequencerObjectBindingNode> FSequencerNodeTree::AddObjectBinding(con
 
 		// Try to get the parent object node if there is one.
 		TSharedPtr<FSequencerObjectBindingNode> ParentNode;
-		TArray<UObject*> RuntimeObjects;
-		UMovieSceneSequence* Animation = Sequencer.GetFocusedMovieSceneSequence();
-		UObject* RuntimeObject = Animation->FindObject(ObjectBinding);
-		if ( RuntimeObject != nullptr)
+
+		UMovieSceneSequence* Sequence = Sequencer.GetFocusedMovieSceneSequence();
+
+		// Prefer to use the parent spawnable if possible, rather than relying on runtime object presence
+		FMovieScenePossessable* Possessable = Sequence->GetMovieScene()->FindPossessable(ObjectBinding);
+		if (Possessable && Possessable->GetParentSpawnable().IsValid())
 		{
-			UObject* ParentObject = Animation->GetParentObject(RuntimeObject);
+			const FMovieSceneBinding* ParentBinding = GuidToBindingMap.FindRef(Possessable->GetParentSpawnable());
+			if (ParentBinding)
+			{
+				ParentNode = AddObjectBinding( ParentBinding->GetName(), Possessable->GetParentSpawnable(), GuidToBindingMap, OutNodeList );
+			}
+		}
+		
+		UObject* RuntimeObject = Sequence->FindObject(ObjectBinding);
+
+		// fallback to using the parent runtime object
+		if (!ParentNode.IsValid() && RuntimeObject)
+		{
+			UObject* ParentObject = Sequence->GetParentObject(RuntimeObject);
+
 			if (ParentObject != nullptr)
 			{
-				FGuid ParentBinding = Animation->FindObjectId(*ParentObject);
+				FGuid ParentBinding = Sequence->FindObjectId(*ParentObject);
 				TSharedPtr<FSequencerObjectBindingNode>* FoundParentNode = ObjectBindingMap.Find( ParentBinding );
 				if ( FoundParentNode != nullptr )
 				{
@@ -228,8 +247,15 @@ TSharedRef<FSequencerObjectBindingNode> FSequencerNodeTree::AddObjectBinding(con
 			}
 		}
 
+		// get human readable name of the object
+		AActor* RuntimeActor = Cast<AActor>(RuntimeObject);
+		const FString& DisplayString = (RuntimeActor != nullptr)
+			? RuntimeActor->GetActorLabel()
+			: ObjectName;
+
 		// Create the node.
-		ObjectNode = MakeShareable( new FSequencerObjectBindingNode( ObjectNodeName, ObjectName, ObjectBinding, ParentNode, *this ) );
+		ObjectNode = MakeShareable(new FSequencerObjectBindingNode(ObjectNodeName, FText::FromString(DisplayString), ObjectBinding, ParentNode, *this));
+
 		if (ParentNode.IsValid())
 		{
 			ParentNode->AddObjectBindingNode(ObjectNode.ToSharedRef());
@@ -246,38 +272,41 @@ TSharedRef<FSequencerObjectBindingNode> FSequencerNodeTree::AddObjectBinding(con
 	return ObjectNode.ToSharedRef();
 }
 
-void FSequencerNodeTree::SaveExpansionState( const FSequencerDisplayNode& Node, bool bExpanded )
+
+void FSequencerNodeTree::SaveExpansionState(const FSequencerDisplayNode& Node, bool bExpanded)
 {	
 	// @todo Sequencer - This should be moved to the sequence level
 	UMovieScene* MovieScene = Sequencer.GetFocusedMovieSceneSequence()->GetMovieScene();
-
 	FMovieSceneEditorData& EditorData = MovieScene->GetEditorData();
 
-	EditorData.ExpansionStates.Add( Node.GetPathName(), FMovieSceneExpansionState(bExpanded) );
+	EditorData.ExpansionStates.Add(Node.GetPathName(), FMovieSceneExpansionState(bExpanded));
 }
 
-bool FSequencerNodeTree::GetSavedExpansionState( const FSequencerDisplayNode& Node ) const
+
+bool FSequencerNodeTree::GetSavedExpansionState(const FSequencerDisplayNode& Node) const
 {
 	// @todo Sequencer - This should be moved to the sequence level
 	UMovieScene* MovieScene = Sequencer.GetFocusedMovieSceneSequence()->GetMovieScene();
-
 	FMovieSceneEditorData& EditorData = MovieScene->GetEditorData();
 	FMovieSceneExpansionState* ExpansionState = EditorData.ExpansionStates.Find( Node.GetPathName() );
 
 	return ExpansionState ? ExpansionState->bExpanded : GetDefaultExpansionState(Node);
 }
 
+
 bool FSequencerNodeTree::GetDefaultExpansionState( const FSequencerDisplayNode& Node ) const
 {
 	// For now, only object nodes are expanded by default
 	return Node.GetType() == ESequencerNode::Object;
 }
 
+
 bool FSequencerNodeTree::IsNodeFiltered( const TSharedRef<const FSequencerDisplayNode> Node ) const
 {
 	return FilteredNodes.Contains( Node );
 }
 
+
 /**
  * Recursively filters nodes
  *
@@ -285,42 +314,42 @@ bool FSequencerNodeTree::IsNodeFiltered( const TSharedRef<const FSequencerDispla
  * @param FilterStrings		The filter strings which need to be matched
  * @param OutFilteredNodes	The list of all filtered nodes
  */
-static bool FilterNodesRecursive( const TSharedRef<FSequencerDisplayNode>& StartNode, const TArray<FString>& FilterStrings, TSet< TSharedRef< const FSequencerDisplayNode> >& OutFilteredNodes )
+static bool FilterNodesRecursive( const TSharedRef<FSequencerDisplayNode>& StartNode, const TArray<FString>& FilterStrings, TSet<TSharedRef<const FSequencerDisplayNode>>& OutFilteredNodes )
 {
-	// Assume the filter is acceptable
+	// assume the filter is acceptable
 	bool bFilterAcceptable = true;
 
-	// Check each string in the filter strings list against 
-	for (int32 TestNameIndex = 0; TestNameIndex < FilterStrings.Num(); ++TestNameIndex)
+	// check each string in the filter strings list against 
+	for (const FString& String : FilterStrings)
 	{
-		const FString& TestName = FilterStrings[TestNameIndex];
-
-		if ( !StartNode->GetDisplayName().ToString().Contains( TestName ) ) 
+		if (!StartNode->GetDisplayName().ToString().Contains(String)) 
 		{
 			bFilterAcceptable = false;
 			break;
 		}
 	}
 
-	// Whether or the start node is in the filter
+	// whether or the start node is in the filter
 	bool bInFilter = false;
 
-	if( bFilterAcceptable )
+	if (bFilterAcceptable)
 	{
 		// This node is now filtered
-		OutFilteredNodes.Add( StartNode );
+		OutFilteredNodes.Add(StartNode);
 		bInFilter = true;
 	}
 
-	// Check each child node to determine if it is filtered
-	const TArray< TSharedRef<FSequencerDisplayNode> >& ChildNodes = StartNode->GetChildNodes();
-	for( int32 ChildIndex = 0; ChildIndex < ChildNodes.Num(); ++ChildIndex )
+	// check each child node to determine if it is filtered
+	const TArray<TSharedRef<FSequencerDisplayNode>>& ChildNodes = StartNode->GetChildNodes();
+
+	for (const auto& Node : ChildNodes)
 	{
 		// Mark the parent as filtered if any child node was filtered
-		bFilterAcceptable |= FilterNodesRecursive( ChildNodes[ChildIndex], FilterStrings, OutFilteredNodes );
-		if( bFilterAcceptable && !bInFilter )
+		bFilterAcceptable |= FilterNodesRecursive(Node, FilterStrings, OutFilteredNodes);
+
+		if (bFilterAcceptable && !bInFilter)
 		{
-			OutFilteredNodes.Add( StartNode );
+			OutFilteredNodes.Add(StartNode);
 			bInFilter = true;
 		}
 	}
@@ -328,11 +357,12 @@ static bool FilterNodesRecursive( const TSharedRef<FSequencerDisplayNode>& Start
 	return bFilterAcceptable;
 }
 
-void FSequencerNodeTree::FilterNodes( const FString& InFilter )
+
+void FSequencerNodeTree::FilterNodes(const FString& InFilter)
 {
 	FilteredNodes.Empty();
 
-	if( InFilter.IsEmpty() )
+	if (InFilter.IsEmpty())
 	{
 		// No filter
 		FilterString.Empty();
@@ -346,13 +376,12 @@ void FSequencerNodeTree::FilterNodes( const FString& InFilter )
 		// Remove whitespace from the front and back of the string
 		FilterString.Trim();
 		FilterString.TrimTrailing();
-		const bool bCullEmpty = true;
-		FilterString.ParseIntoArray( FilterStrings, TEXT(" "), bCullEmpty );
+		FilterString.ParseIntoArray(FilterStrings, TEXT(" "), true /*bCullEmpty*/);
 
-		for( auto It = ObjectBindingMap.CreateIterator(); It; ++It )
+		for (auto It = ObjectBindingMap.CreateIterator(); It; ++It)
 		{
 			// Recursively filter all nodes, matching them against the list of filter strings.  All filter strings must be matched
-			FilterNodesRecursive( It.Value().ToSharedRef(), FilterStrings, FilteredNodes );
+			FilterNodesRecursive(It.Value().ToSharedRef(), FilterStrings, FilteredNodes);
 		}
 	}
 }
