Copying //UE4/Orion-Staging to Dev-Main (Originates from //Orion/Main @ CL-2759376)

diff --git a/Engine/Source/Editor/MovieSceneTools/Private/TrackEditors/TransformTrackEditor.cpp b/Engine/Source/Editor/MovieSceneTools/Private/TrackEditors/TransformTrackEditor.cpp
--- a/Engine/Source/Editor/MovieSceneTools/Private/TrackEditors/TransformTrackEditor.cpp
+++ b/Engine/Source/Editor/MovieSceneTools/Private/TrackEditors/TransformTrackEditor.cpp
@@ -11,10 +11,11 @@
 #include "ISequencerSection.h"
 #include "ISectionLayoutBuilder.h"
 #include "IKeyArea.h"
-#include "MovieSceneToolHelpers.h"
 #include "MovieSceneTrackEditor.h"
 #include "TransformTrackEditor.h"
+#include "MatineeImportTools.h"
 #include "Matinee/InterpTrackMove.h"
+#include "FloatCurveKeyArea.h"
 
 
 #define LOCTEXT_NAMESPACE "MovieScene_TransformTrack"
@@ -245,42 +246,6 @@ bool F3DTransformTrackEditor::SupportsType( TSubclassOf<UMovieSceneTrack> Type )
 }
 
 
-void PasteInterpMoveTrack( UInterpTrackMove* MoveTrack, UMovieScene3DTransformTrack* TransformTrack )
-{
-	float KeyTime = MoveTrack->GetKeyframeTime( 0 );
-	UMovieScene3DTransformSection* Section = Cast<UMovieScene3DTransformSection>(TransformTrack->FindOrAddSection( KeyTime ));
-	float SectionMin = Section->GetStartTime();
-	float SectionMax = Section->GetEndTime();
-
-	FRichCurve& TranslationXCurve = Section->GetTranslationCurve( EAxis::X );
-	FRichCurve& TranslationYCurve = Section->GetTranslationCurve( EAxis::Y );
-	FRichCurve& TranslationZCurve = Section->GetTranslationCurve( EAxis::Z );
-	for ( const auto& Point : MoveTrack->PosTrack.Points )
-	{
-		MatineeImportTools::SetOrAddKey( TranslationXCurve, Point.InVal, Point.OutVal.X, Point.ArriveTangent.X, Point.LeaveTangent.X, Point.InterpMode );
-		MatineeImportTools::SetOrAddKey( TranslationYCurve, Point.InVal, Point.OutVal.Y, Point.ArriveTangent.Y, Point.LeaveTangent.Y, Point.InterpMode );
-		MatineeImportTools::SetOrAddKey( TranslationZCurve, Point.InVal, Point.OutVal.Z, Point.ArriveTangent.Z, Point.LeaveTangent.Z, Point.InterpMode );
-		SectionMin = FMath::Min( SectionMin, Point.InVal );
-		SectionMax = FMath::Max( SectionMax, Point.InVal );
-	}
-
-	FRichCurve& RotationXCurve = Section->GetRotationCurve( EAxis::X );
-	FRichCurve& RotationYCurve = Section->GetRotationCurve( EAxis::Y );
-	FRichCurve& RotationZCurve = Section->GetRotationCurve( EAxis::Z );
-	for ( const auto& Point : MoveTrack->EulerTrack.Points )
-	{
-		MatineeImportTools::SetOrAddKey( RotationXCurve, Point.InVal, Point.OutVal.X, Point.ArriveTangent.X, Point.LeaveTangent.X, Point.InterpMode );
-		MatineeImportTools::SetOrAddKey( RotationYCurve, Point.InVal, Point.OutVal.Y, Point.ArriveTangent.Y, Point.LeaveTangent.Y, Point.InterpMode );
-		MatineeImportTools::SetOrAddKey( RotationZCurve, Point.InVal, Point.OutVal.Z, Point.ArriveTangent.Z, Point.LeaveTangent.Z, Point.InterpMode );
-		SectionMin = FMath::Min( SectionMin, Point.InVal );
-		SectionMax = FMath::Max( SectionMax, Point.InVal );
-	}
-
-	Section->SetStartTime( SectionMin );
-	Section->SetEndTime( SectionMax );
-}
-
-
 void F3DTransformTrackEditor::BuildTrackContextMenu( FMenuBuilder& MenuBuilder, UMovieSceneTrack* Track )
 {
 	UInterpTrackMove* MoveTrack = nullptr;
@@ -298,7 +263,7 @@ void F3DTransformTrackEditor::BuildTrackContextMenu( FMenuBuilder& MenuBuilder,
 		NSLOCTEXT( "Sequencer", "PasteMatineeTrackTooltip", "Pastes keys from a Matinee move track into this track." ),
 		FSlateIcon(),
 		FUIAction(
-			FExecuteAction::CreateStatic( &PasteInterpMoveTrack, MoveTrack, TransformTrack ),
+			FExecuteAction::CreateStatic( &FMatineeImportTools::CopyInterpMoveTrack, GetSequencer().ToSharedRef(), MoveTrack, TransformTrack ),
 			FCanExecuteAction::CreateLambda( [=]()->bool { return MoveTrack != nullptr && MoveTrack->GetNumKeys() > 0 && TransformTrack != nullptr; } ) ) );
 }
 
@@ -456,26 +421,6 @@ void F3DTransformTrackEditor::BindCommands(TSharedRef<FUICommandList> SequencerC
 
 void F3DTransformTrackEditor::BuildObjectBindingEditButtons(TSharedPtr<SHorizontalBox> EditBox, const FGuid& ObjectGuid, const UClass* ObjectClass)
 {
-	TArray<UObject*> OutObjects;
-	GetSequencer()->GetRuntimeObjects( GetSequencer()->GetFocusedMovieSceneSequenceInstance(), ObjectGuid, OutObjects);
-
-	TWeakObjectPtr<ACameraActor> CameraActor;
-
-	for (UObject* Object : OutObjects)
-	{
-		ACameraActor* Actor = Cast<ACameraActor>( Object );
-		if (Actor)
-		{
-			CameraActor = Actor;
-			break;
-		}
-	}
-
-	if (!CameraActor.IsValid())
-	{
-		return;
-	}
-
 	// If this is a camera track, add a button to lock the viewport to the camera
 	EditBox.Get()->AddSlot()
 		.VAlign(VAlign_Center)
@@ -483,7 +428,9 @@ void F3DTransformTrackEditor::BuildObjectBindingEditButtons(TSharedPtr<SHorizont
 		.AutoWidth()
 		.Padding(4, 0, 0, 0)
 		[
-			SNew(SCheckBox)
+			SNew(SCheckBox)		
+				.IsFocusable(false)
+				.Visibility(this, &F3DTransformTrackEditor::IsCameraVisible, ObjectGuid)
 				.IsChecked(this, &F3DTransformTrackEditor::IsCameraLocked, ObjectGuid)
 				.OnCheckStateChanged(this, &F3DTransformTrackEditor::OnLockCameraClicked, ObjectGuid)
 				.ToolTipText(this, &F3DTransformTrackEditor::GetLockCameraToolTip, ObjectGuid)
@@ -527,6 +474,23 @@ bool F3DTransformTrackEditor::CanAddTransformTrackForActorHandle( FGuid ObjectBi
 	return true;
 }
 
+EVisibility
+F3DTransformTrackEditor::IsCameraVisible(FGuid ObjectGuid) const
+{
+	TArray<UObject*> OutObjects;
+	GetSequencer()->GetRuntimeObjects( GetSequencer()->GetFocusedMovieSceneSequenceInstance(), ObjectGuid, OutObjects);
+
+	for (UObject* Object : OutObjects)
+	{
+		ACameraActor* Actor = Cast<ACameraActor>( Object );
+		if (Actor)
+		{
+			return EVisibility::Visible;
+		}
+	}
+
+	return EVisibility::Hidden;
+}
 
 ECheckBoxState F3DTransformTrackEditor::IsCameraLocked(FGuid ObjectGuid) const
 {
@@ -747,15 +711,20 @@ bool F3DTransformTrackEditor::OnAddTransformKeys( float Time, AActor* ActorToKey
 	TArray<UObject*> ObjectsToKey;
 	ObjectsToKey.Add(ActorToKey);
 	
-	FOnSetIntermediateValue OnSetIntermediateValue;
-	OnSetIntermediateValue.BindRaw( this, &F3DTransformTrackEditor::SetIntermediateValueFromTransformChange, CurrentTransform );
-
-	FOnInitializeNewTrack OnInitializeNewTrack;
-	OnInitializeNewTrack.BindLambda([this](UMovieScene3DTransformTrack* NewTrack) {
-		NewTrack->SetPropertyNameAndPath(TransformPropertyName, TransformPropertyName.ToString());
-	});
-	
-	return AddKeysToObjects(ObjectsToKey, Time, *Keys, KeyParams, UMovieScene3DTransformTrack::StaticClass(), TransformPropertyName, OnInitializeNewTrack, OnSetIntermediateValue);
+	return AddKeysToObjects(
+		ObjectsToKey,
+		Time,
+		*Keys,
+		KeyParams,
+		UMovieScene3DTransformTrack::StaticClass(),
+		TransformPropertyName,
+		[this](UMovieScene3DTransformTrack* NewTrack) {
+			NewTrack->SetPropertyNameAndPath(TransformPropertyName, TransformPropertyName.ToString());
+		},
+		[&](UMovieScene3DTransformTrack* NewTrack) {
+			SetIntermediateValueFromTransformChange(NewTrack, CurrentTransform);
+		}
+	);
 }
 
 
diff --git a/Engine/Source/Editor/MovieSceneTools/Private/TrackEditors/TransformTrackEditor.h b/Engine/Source/Editor/MovieSceneTools/Private/TrackEditors/TransformTrackEditor.h
--- a/Engine/Source/Editor/MovieSceneTools/Private/TrackEditors/TransformTrackEditor.h
+++ b/Engine/Source/Editor/MovieSceneTools/Private/TrackEditors/TransformTrackEditor.h
@@ -65,6 +65,9 @@ private:
 	 */
 	void OnTransformChanged( UObject& InObject );
 
+	/** Delegate for camera button visible state */
+	EVisibility IsCameraVisible(FGuid ObjectGuid) const;
+
 	/** Delegate for camera button lock state */
 	ECheckBoxState IsCameraLocked(FGuid ObjectGuid) const; 
 
diff --git a/Engine/Source/Editor/MovieSceneTools/Public/BoolKeyArea.h b/Engine/Source/Editor/MovieSceneTools/Public/BoolKeyArea.h
new file mode 100644
--- /dev/null
+++ b/Engine/Source/Editor/MovieSceneTools/Public/BoolKeyArea.h
@@ -0,0 +1,27 @@
+// Copyright 1998-2015 Epic Games, Inc. All Rights Reserved.
+
+#pragma once
+
+#include "IntegralKeyArea.h"
+
+
+/**
+ * A key area for displaying and editing integral curves representing bools.
+ */
+class FBoolKeyArea
+	: public FIntegralKeyArea<bool>
+{
+public:
+
+	/* Create and initialize a new instance. */
+	FBoolKeyArea(FIntegralCurve& InCurve, UMovieSceneSection* InOwningSection)
+		: FIntegralKeyArea(InCurve, InOwningSection)
+	{ }
+
+public:
+
+	// FIntegralKeyArea overrides
+
+	virtual bool CanCreateKeyEditor() override;
+	virtual TSharedRef<SWidget> CreateKeyEditor(ISequencer* Sequencer) override;
+};
diff --git a/Engine/Source/Editor/MovieSceneTools/Public/ByteKeyArea.h b/Engine/Source/Editor/MovieSceneTools/Public/ByteKeyArea.h
new file mode 100644
--- /dev/null
+++ b/Engine/Source/Editor/MovieSceneTools/Public/ByteKeyArea.h
@@ -0,0 +1,25 @@
+// Copyright 1998-2015 Epic Games, Inc. All Rights Reserved.
+
+#pragma once
+
+#include "IntegralKeyArea.h"
+
+
+/** A key area for displaying and editing byte curves. */
+class FByteKeyArea
+	: public FIntegralKeyArea<uint8>
+{
+public:
+
+	/* Create and initialize a new instance. */
+	FByteKeyArea(FIntegralCurve& InCurve, UMovieSceneSection* InOwningSection)
+		: FIntegralKeyArea(InCurve, InOwningSection)
+	{ }
+
+public:
+
+	// FIntegralKeyArea overrides
+
+	virtual bool CanCreateKeyEditor() override;
+	virtual TSharedRef<SWidget> CreateKeyEditor(ISequencer* Sequencer) override;
+};
diff --git a/Engine/Source/Editor/MovieSceneTools/Public/CommonMovieSceneTools.h b/Engine/Source/Editor/MovieSceneTools/Public/CommonMovieSceneTools.h
--- a/Engine/Source/Editor/MovieSceneTools/Public/CommonMovieSceneTools.h
+++ b/Engine/Source/Editor/MovieSceneTools/Public/CommonMovieSceneTools.h
@@ -53,18 +53,3 @@ private:
 	/** The number of pixels in the view range */
 	float PixelsPerInput;
 };
-
-
-class MatineeImportTools
-{
-public:
-
-	/** Converts a matinee interpolation mode to its equivalent rich curve interpolation mode. */
-	static ERichCurveInterpMode MatineeInterpolationToRichCurveInterpolation( EInterpCurveMode CurveMode );
-
-	/** Converts a matinee interpolation mode to its equivalent rich curve tangent mode. */
-	static ERichCurveTangentMode MatineeInterpolationToRichCurveTangent( EInterpCurveMode CurveMode );
-
-	/** Adds a key to a rich curve based on matinee curve key data. */
-	static void SetOrAddKey( FRichCurve& Curve, float Time, float Value, float ArriveTangent, float LeaveTangent, EInterpCurveMode MatineeInterpMode );
-};
diff --git a/Engine/Source/Editor/MovieSceneTools/Public/EnumKeyArea.h b/Engine/Source/Editor/MovieSceneTools/Public/EnumKeyArea.h
new file mode 100644
--- /dev/null
+++ b/Engine/Source/Editor/MovieSceneTools/Public/EnumKeyArea.h
@@ -0,0 +1,33 @@
+// Copyright 1998-2015 Epic Games, Inc. All Rights Reserved.
+
+#pragma once
+
+#include "ByteKeyArea.h"
+
+
+/**
+ * A key area for displaying and editing integral curves representing enums.
+ */
+class FEnumKeyArea
+	: public FByteKeyArea
+{
+public:
+
+	/* Create and initialize a new instance. */
+	FEnumKeyArea(FIntegralCurve& InCurve, UMovieSceneSection* InOwningSection, const UEnum* InEnum)
+		: FByteKeyArea(InCurve, InOwningSection)
+		, Enum(InEnum)
+	{ }
+
+public:
+
+	// FByteKeyArea overrides
+
+	virtual bool CanCreateKeyEditor() override;
+	virtual TSharedRef<SWidget> CreateKeyEditor(ISequencer* Sequencer) override;
+
+private:
+
+	/** The enum which provides available integral values for this key area. */
+	const UEnum* Enum;
+};
diff --git a/Engine/Source/Editor/MovieSceneTools/Public/FloatCurveKeyArea.h b/Engine/Source/Editor/MovieSceneTools/Public/FloatCurveKeyArea.h
new file mode 100644
--- /dev/null
+++ b/Engine/Source/Editor/MovieSceneTools/Public/FloatCurveKeyArea.h
@@ -0,0 +1,63 @@
+// Copyright 1998-2015 Epic Games, Inc. All Rights Reserved.
+
+#pragma once
+
+#include "IKeyArea.h"
+
+
+/**
+ * A key area for float keys
+ */
+class MOVIESCENETOOLS_API FFloatCurveKeyArea
+	: public IKeyArea
+{
+public:
+
+	FFloatCurveKeyArea(FRichCurve* InCurve, UMovieSceneSection* InOwningSection)
+		: Curve(InCurve)
+		, OwningSection(InOwningSection)
+	{ }
+
+public:
+
+	void ClearIntermediateValue()
+	{
+		IntermediateValue.Reset();
+	}
+
+	void SetIntermediateValue(float InIntermediateValue)
+	{
+		IntermediateValue = InIntermediateValue;
+	}
+
+public:
+
+	// IKeyArea interface
+
+	virtual TArray<FKeyHandle> AddKeyUnique(float Time, EMovieSceneKeyInterpolation InKeyInterpolation, float TimeToCopyFrom = FLT_MAX) override;
+	virtual bool CanCreateKeyEditor() override;
+	virtual TSharedRef<SWidget> CreateKeyEditor(ISequencer* Sequencer) override;
+	virtual void DeleteKey(FKeyHandle KeyHandle) override;
+	virtual ERichCurveExtrapolation GetExtrapolationMode(bool bPreInfinity) const override;
+	virtual ERichCurveInterpMode GetKeyInterpMode(FKeyHandle KeyHandle) const override;
+	virtual ERichCurveTangentMode GetKeyTangentMode(FKeyHandle KeyHandle) const override;
+	virtual float GetKeyTime(FKeyHandle KeyHandle) const override;
+	virtual UMovieSceneSection* GetOwningSection() override;
+	virtual FRichCurve* GetRichCurve() override;
+	virtual TArray<FKeyHandle> GetUnsortedKeyHandles() const override;
+	virtual FKeyHandle MoveKey(FKeyHandle KeyHandle, float DeltaPosition) override;
+	virtual void SetExtrapolationMode(ERichCurveExtrapolation ExtrapMode, bool bPreInfinity) override;
+	virtual void SetKeyInterpMode(FKeyHandle KeyHandle, ERichCurveInterpMode InterpMode) override;
+	virtual void SetKeyTangentMode(FKeyHandle KeyHandle, ERichCurveTangentMode TangentMode) override;
+	virtual void SetKeyTime(FKeyHandle KeyHandle, float NewKeyTime) const override;
+
+private:
+
+	/** The curve which provides the keys for this key area. */
+	FRichCurve* Curve;
+
+	/** The section that owns this key area. */
+	UMovieSceneSection* OwningSection;
+
+	TOptional<float> IntermediateValue;
+};
diff --git a/Engine/Source/Editor/MovieSceneTools/Public/IntegralKeyArea.h b/Engine/Source/Editor/MovieSceneTools/Public/IntegralKeyArea.h
new file mode 100644
--- /dev/null
+++ b/Engine/Source/Editor/MovieSceneTools/Public/IntegralKeyArea.h
@@ -0,0 +1,101 @@
+// Copyright 1998-2015 Epic Games, Inc. All Rights Reserved.
+
+#pragma once
+
+#include "IKeyArea.h"
+
+
+/**
+ * Abstract base class for integral curve key areas.
+ */
+class MOVIESCENETOOLS_API FIntegralCurveKeyAreaBase
+	: public IKeyArea
+{
+public:
+
+	/** Create and initialize a new instance. */
+	FIntegralCurveKeyAreaBase(FIntegralCurve& InCurve, UMovieSceneSection* InOwningSection)
+		: Curve(InCurve)
+		, OwningSection(InOwningSection)
+	{ }
+
+public:
+
+	// IKeyArea interface
+
+	virtual TArray<FKeyHandle> AddKeyUnique( float Time, EMovieSceneKeyInterpolation InKeyInterpolation, float TimeToCopyFrom = FLT_MAX ) override;
+	virtual void DeleteKey(FKeyHandle KeyHandle) override;
+	virtual ERichCurveExtrapolation GetExtrapolationMode(bool bPreInfinity) const override;
+	virtual ERichCurveInterpMode GetKeyInterpMode(FKeyHandle KeyHandle) const override;
+	virtual ERichCurveTangentMode GetKeyTangentMode(FKeyHandle KeyHandle) const override;
+	virtual float GetKeyTime(FKeyHandle KeyHandle) const override;
+	virtual UMovieSceneSection* GetOwningSection() override;
+	virtual FRichCurve* GetRichCurve() override;
+	virtual TArray<FKeyHandle> GetUnsortedKeyHandles() const override;
+	virtual FKeyHandle MoveKey(FKeyHandle KeyHandle, float DeltaPosition) override;
+	virtual void SetExtrapolationMode(ERichCurveExtrapolation ExtrapMode, bool bPreInfinity) override;
+	virtual void SetKeyInterpMode(FKeyHandle KeyHandle, ERichCurveInterpMode InterpMode) override;
+	virtual void SetKeyTangentMode(FKeyHandle KeyHandle, ERichCurveTangentMode TangentMode) override;
+	virtual void SetKeyTime(FKeyHandle KeyHandle, float NewKeyTime) const override;
+
+protected:
+
+	virtual void EvaluateAndAddKey(float Time, float TimeToCopyFrom, FKeyHandle& CurrentKey) = 0;
+
+protected:
+
+	/** Curve with keys in this area */
+	FIntegralCurve& Curve;
+
+	/** The section that owns this key area. */
+	UMovieSceneSection* OwningSection;
+};
+
+
+/**
+ * Template for integral curve key areas.
+ */
+template<typename IntegralType>
+class MOVIESCENETOOLS_API FIntegralKeyArea
+	: public FIntegralCurveKeyAreaBase
+{
+public:
+
+	FIntegralKeyArea(FIntegralCurve& InCurve, UMovieSceneSection* InOwningSection)
+		: FIntegralCurveKeyAreaBase(InCurve, InOwningSection)
+	{ }
+
+public:
+
+	void ClearIntermediateValue()
+	{
+		IntermediateValue.Reset();
+	}
+
+	void SetIntermediateValue(IntegralType InIntermediateValue)
+	{
+		IntermediateValue = InIntermediateValue;
+	}
+
+protected:
+
+	virtual void EvaluateAndAddKey(float Time, float TimeToCopyFrom, FKeyHandle& CurrentKey) override
+	{
+		int32 Value = Curve.Evaluate(Time);
+
+		if (TimeToCopyFrom != FLT_MAX)
+		{
+			Value = Curve.Evaluate(TimeToCopyFrom);
+		}
+		else if (IntermediateValue.IsSet())
+		{
+			Value = IntermediateValue.GetValue();
+		}
+
+		Curve.AddKey(Time, Value, CurrentKey);
+	}
+
+protected:
+
+	TOptional<IntegralType> IntermediateValue;
+};
diff --git a/Engine/Source/Editor/MovieSceneTools/Public/KeyframeTrackEditor.h b/Engine/Source/Editor/MovieSceneTools/Public/KeyframeTrackEditor.h
--- a/Engine/Source/Editor/MovieSceneTools/Public/KeyframeTrackEditor.h
+++ b/Engine/Source/Editor/MovieSceneTools/Public/KeyframeTrackEditor.h
@@ -27,10 +27,8 @@ public:
 		: FMovieSceneTrackEditor( InSequencer )
 	{ }
 
-
-	~FKeyframeTrackEditor()
-	{
-	}
+	/** Virtual destructor. */
+	~FKeyframeTrackEditor() { }
 
 public:
 
@@ -44,9 +42,6 @@ public:
 
 protected:
 
-	DECLARE_DELEGATE_OneParam( FOnInitializeNewTrack, TrackType* );
-	DECLARE_DELEGATE_OneParam( FOnSetIntermediateValue, UMovieSceneTrack* );
-	
 	/*
 	 * Adds keys to the specified object.  This may also add tracks and sections depending on the options specified. 
 	 *
@@ -61,8 +56,11 @@ protected:
 	 * @param OnSetIntermediateValue A delegate which is called when the key params specify that keying should only happen when
 	 * 		  auto-keying is on, but auto-keying is off.  This allows handling the changed but not keyed intermediate value.
 	 */
-	bool AddKeysToObjects(TArray<UObject*> ObjectsToKey, float KeyTime, const TArray<KeyDataType> NewKeys, FKeyParams KeyParams,
-		TSubclassOf<UMovieSceneTrack> TrackClass, FName PropertyName, FOnInitializeNewTrack OnIntializeNewTrack, FOnSetIntermediateValue OnSetIntermediateValue )
+	bool AddKeysToObjects(
+		TArray<UObject*> ObjectsToKey, float KeyTime, const TArray<KeyDataType> NewKeys, FKeyParams KeyParams,
+		TSubclassOf<UMovieSceneTrack> TrackClass, FName PropertyName,
+		TFunction<void(TrackType*)> OnIntializeNewTrack,
+		TFunction<void(TrackType*)> OnSetIntermediateValue)
 	{
 		bool bHandleCreated = false;
 		bool bTrackCreatedOrModified = false;
@@ -84,7 +82,11 @@ protected:
 
 private:
 
-	bool AddKeysToHandle( FGuid ObjectHandle, float KeyTime, const TArray<KeyDataType> NewKeys, FKeyParams KeyParams, TSubclassOf<UMovieSceneTrack> TrackClass, FName PropertyName, FOnInitializeNewTrack OnIntializeNewTrack, FOnSetIntermediateValue OnSetIntermediateValue )
+	bool AddKeysToHandle(
+		FGuid ObjectHandle, float KeyTime, const TArray<KeyDataType> NewKeys, FKeyParams KeyParams,
+		TSubclassOf<UMovieSceneTrack> TrackClass, FName PropertyName,
+		TFunction<void(TrackType*)> OnIntializeNewTrack,
+		TFunction<void(TrackType*)> OnSetIntermediateValue)
 	{
 		bool bTrackCreated = false;
 		bool bTrackModified = false;
@@ -96,7 +98,10 @@ private:
 
 		if ( bTrackCreated )
 		{
-			OnIntializeNewTrack.ExecuteIfBound( Track );
+			if (OnIntializeNewTrack)
+			{
+				OnIntializeNewTrack(Track);
+			}
 		}
 
 		if ( Track )
@@ -107,15 +112,17 @@ private:
 		return bTrackCreated || bTrackModified;
 	}
 
-	bool AddKeysToTrack( TrackType* Track, float KeyTime, const TArray<KeyDataType> NewKeys, FKeyParams KeyParams, FOnSetIntermediateValue OnSetIntermediateValue )
+	bool AddKeysToTrack(
+		TrackType* Track, float KeyTime, const TArray<KeyDataType> NewKeys, FKeyParams KeyParams,
+		TFunction<void(TrackType*)> OnSetIntermediateValue)
 	{
 		bool bTrackModified = false;
 
 		bool bSettingIntermediateValue = KeyParams.bCreateKeyOnlyWhenAutoKeying && GetSequencer()->GetAutoKeyEnabled() == false;
 		if ( bSettingIntermediateValue )
 		{
-			checkf( OnSetIntermediateValue.IsBound(), TEXT( "A valid OnSetIntermediateValue delegate must be provided for consistent keyframing behavior." ) );
-			OnSetIntermediateValue.Execute(Track);
+			checkf(OnSetIntermediateValue, TEXT("A valid OnSetIntermediateValue delegate must be provided for consistent keyframing behavior."));
+			OnSetIntermediateValue(Track);
 		}
 		else
 		{
diff --git a/Engine/Source/Editor/MovieSceneTools/Public/MovieSceneToolHelpers.h b/Engine/Source/Editor/MovieSceneTools/Public/MovieSceneToolHelpers.h
--- a/Engine/Source/Editor/MovieSceneTools/Public/MovieSceneToolHelpers.h
+++ b/Engine/Source/Editor/MovieSceneTools/Public/MovieSceneToolHelpers.h
@@ -2,443 +2,8 @@
 
 #pragma once
 
-#include "ISequencer.h"
-#include "ISequencerSection.h"
-#include "IKeyArea.h"
-#include "MovieSceneSection.h"
 
-/**
- * A key area for float keys
- */
-class MOVIESCENETOOLS_API FFloatCurveKeyArea : public IKeyArea
-{
-public:
-	FFloatCurveKeyArea( FRichCurve* InCurve, UMovieSceneSection* InOwningSection )
-		: Curve( InCurve )
-		, OwningSection( InOwningSection )
-	{}
-
-	/** IKeyArea interface */
-	virtual TArray<FKeyHandle> GetUnsortedKeyHandles() const override
-	{
-		TArray<FKeyHandle> OutKeyHandles;
-		for (auto It(Curve->GetKeyHandleIterator()); It; ++It)
-		{
-			OutKeyHandles.Add(It.Key());
-		}
-		return OutKeyHandles;
-	}
-
-	virtual void SetKeyTime(FKeyHandle KeyHandle, float NewKeyTime) const override
-	{
-		Curve->SetKeyTime(KeyHandle, NewKeyTime);
-	}
-
-	virtual float GetKeyTime(FKeyHandle KeyHandle) const override
-	{
-		return Curve->GetKeyTime(KeyHandle);
-	}
-
-	virtual FKeyHandle MoveKey(FKeyHandle KeyHandle, float DeltaPosition) override
-	{
-		return Curve->SetKeyTime( KeyHandle, Curve->GetKeyTime( KeyHandle ) + DeltaPosition );
-	}
-
-	virtual void DeleteKey(FKeyHandle KeyHandle) override
-	{
-		if ( Curve->IsKeyHandleValid( KeyHandle ) )
-		{
-			Curve->DeleteKey( KeyHandle );
-		}
-	}
-
-	virtual void SetKeyInterpMode(FKeyHandle KeyHandle, ERichCurveInterpMode InterpMode) override
-	{
-		if (Curve->IsKeyHandleValid(KeyHandle))
-		{
-			Curve->SetKeyInterpMode(KeyHandle, InterpMode);
-		}
-	}
-
-	virtual ERichCurveInterpMode GetKeyInterpMode(FKeyHandle KeyHandle) const override
-	{
-		if (Curve->IsKeyHandleValid(KeyHandle))
-		{
-			return Curve->GetKeyInterpMode(KeyHandle);
-		}
-		return RCIM_None;
-	}
-
-	virtual void SetKeyTangentMode(FKeyHandle KeyHandle, ERichCurveTangentMode TangentMode) override
-	{
-		if (Curve->IsKeyHandleValid(KeyHandle))
-		{
-			Curve->SetKeyTangentMode(KeyHandle, TangentMode);
-		}
-	}
-
-	virtual ERichCurveTangentMode GetKeyTangentMode(FKeyHandle KeyHandle) const override
-	{
-		if (Curve->IsKeyHandleValid(KeyHandle))
-		{
-			return Curve->GetKeyTangentMode(KeyHandle);
-		}
-		return RCTM_None;
-	}
-
-	virtual void SetExtrapolationMode(ERichCurveExtrapolation ExtrapMode, bool bPreInfinity) override
-	{
-		if (bPreInfinity)
-		{
-			Curve->PreInfinityExtrap = ExtrapMode;
-		}
-		else
-		{
-			Curve->PostInfinityExtrap = ExtrapMode;
-		}
-	}
-
-	virtual ERichCurveExtrapolation GetExtrapolationMode(bool bPreInfinity) const override
-	{
-		if (bPreInfinity)
-		{
-			return Curve->PreInfinityExtrap;
-		}
-		else
-		{
-			return Curve->PostInfinityExtrap;
-		}
-		return RCCE_None;
-	}
-
-	virtual TArray<FKeyHandle> AddKeyUnique(float Time, EMovieSceneKeyInterpolation InKeyInterpolation, float TimeToCopyFrom = FLT_MAX) override;
-
-	virtual FRichCurve* GetRichCurve() override { return Curve; }
-
-	virtual UMovieSceneSection* GetOwningSection() override { return OwningSection; }
-
-	virtual bool CanCreateKeyEditor() override
-	{
-		return true;
-	}
-
-	virtual TSharedRef<SWidget> CreateKeyEditor(ISequencer* Sequencer) override;
-
-	void SetIntermediateValue( float InIntermediateValue )
-	{
-		IntermediateValue = InIntermediateValue;
-	}
-
-	void ClearIntermediateValue()
-	{
-		IntermediateValue.Reset();
-	}
-
-private:
-	/** The curve which provides the keys for this key area. */
-	FRichCurve* Curve;
-	/** The section that owns this key area. */
-	UMovieSceneSection* OwningSection;
-
-	TOptional<float> IntermediateValue;
-};
-
-
-/** A key area for integral keys */
-template<typename IntegralType>
-class MOVIESCENETOOLS_API FIntegralKeyArea : public IKeyArea
-{
-public:
-	FIntegralKeyArea(FIntegralCurve& InCurve, UMovieSceneSection* InOwningSection)
-		: Curve(InCurve)
-		, OwningSection(InOwningSection)
-	{}
-
-	/** IKeyArea interface */
-	virtual TArray<FKeyHandle> GetUnsortedKeyHandles() const override
-	{
-		TArray<FKeyHandle> OutKeyHandles;
-		for (auto It(Curve.GetKeyHandleIterator()); It; ++It)
-		{
-			OutKeyHandles.Add(It.Key());
-		}
-		return OutKeyHandles;
-	}
-
-	virtual void SetKeyTime(FKeyHandle KeyHandle, float NewKeyTime) const override
-	{
-		Curve.SetKeyTime(KeyHandle, NewKeyTime);
-	}
-
-	virtual float GetKeyTime(FKeyHandle KeyHandle) const override
-	{
-		return Curve.GetKeyTime(KeyHandle);
-	}
-
-	virtual FKeyHandle MoveKey(FKeyHandle KeyHandle, float DeltaPosition) override
-	{
-		return Curve.SetKeyTime( KeyHandle, Curve.GetKeyTime( KeyHandle ) + DeltaPosition );
-	}
-
-	virtual void DeleteKey(FKeyHandle KeyHandle) override
-	{
-		Curve.DeleteKey(KeyHandle);
-	}
-
-	virtual void SetKeyInterpMode(FKeyHandle KeyHandle, ERichCurveInterpMode InterpMode) override
-	{
-	}
-
-	virtual ERichCurveInterpMode GetKeyInterpMode(FKeyHandle KeyHandle) const override
-	{
-		return RCIM_None;
-	}
-
-	virtual void SetKeyTangentMode(FKeyHandle KeyHandle, ERichCurveTangentMode TangentMode) override
-	{
-	}
-
-	virtual ERichCurveTangentMode GetKeyTangentMode(FKeyHandle KeyHandle) const override
-	{
-		return RCTM_None;
-	}
-
-	virtual void SetExtrapolationMode(ERichCurveExtrapolation ExtrapMode, bool bPreInfinity) override
-	{
-	}
-
-	virtual ERichCurveExtrapolation GetExtrapolationMode(bool bPreInfinity) const override
-	{
-		return RCCE_None;
-	}
-
-	virtual TArray<FKeyHandle> AddKeyUnique( float Time, EMovieSceneKeyInterpolation InKeyInterpolation, float TimeToCopyFrom = FLT_MAX ) override
-	{
-		TArray<FKeyHandle> AddedKeyHandles;
-
-		FKeyHandle CurrentKey = Curve.FindKey( Time );
-		if ( Curve.IsKeyHandleValid( CurrentKey ) == false )
-		{
-			if ( OwningSection->GetStartTime() > Time )
-			{
-				OwningSection->SetStartTime( Time );
-			}
-			if ( OwningSection->GetEndTime() < Time )
-			{
-				OwningSection->SetEndTime( Time );
-			}
-
-			int32 Value = Curve.Evaluate( Time );
-			if ( TimeToCopyFrom != FLT_MAX )
-			{
-				Value = Curve.Evaluate( TimeToCopyFrom );
-			}
-			else if ( IntermediateValue.IsSet() )
-			{
-				Value = IntermediateValue.GetValue();
-			}
-
-			Curve.AddKey( Time, Value, CurrentKey );
-			AddedKeyHandles.Add( CurrentKey );
-		}
-
-		return AddedKeyHandles;
-	}
-
-	virtual FRichCurve* GetRichCurve() override { return nullptr; };
-
-	virtual UMovieSceneSection* GetOwningSection() override { return OwningSection; }
-
-	void SetIntermediateValue( IntegralType InIntermediateValue )
-	{
-		IntermediateValue = InIntermediateValue;
-	}
-
-	void ClearIntermediateValue()
-	{
-		IntermediateValue.Reset();
-	}
-
-protected:
-	/** Curve with keys in this area */
-	FIntegralCurve& Curve;
-	/** The section that owns this key area. */
-	UMovieSceneSection* OwningSection;
-
-	TOptional<IntegralType> IntermediateValue;
-};
-
-/** A key area for displaying and editing byte curves. */
-class FByteKeyArea : public FIntegralKeyArea < uint8 >
-{
-public:
-	FByteKeyArea( FIntegralCurve& InCurve, UMovieSceneSection* InOwningSection )
-		: FIntegralKeyArea( InCurve, InOwningSection )
-	{}
-
-	virtual bool CanCreateKeyEditor() override
-	{
-		return true;
-	}
-
-	virtual TSharedRef<SWidget> CreateKeyEditor( ISequencer* Sequencer ) override;
-};
-
-/** A key area for displaying and editing integral curves representing enums. */
-class FEnumKeyArea : public FByteKeyArea
-{
-public:
-	FEnumKeyArea(FIntegralCurve& InCurve, UMovieSceneSection* InOwningSection, const UEnum* InEnum)
-		: FByteKeyArea(InCurve, InOwningSection)
-		, Enum(InEnum)
-	{}
-
-	virtual bool CanCreateKeyEditor() override 
-	{
-		return true;
-	}
-
-	virtual TSharedRef<SWidget> CreateKeyEditor(ISequencer* Sequencer) override;
-
-private:
-	/** The enum which provides available integral values for this key area. */
-	const UEnum* Enum;
-
-};
-
-
-/** A key area for displaying and editing integral curves representing bools. */
-class FBoolKeyArea : public FIntegralKeyArea<bool>
-{
-public:
-	FBoolKeyArea(FIntegralCurve& InCurve, UMovieSceneSection* InOwningSection)
-		: FIntegralKeyArea(InCurve, InOwningSection)
-	{}
-
-	virtual bool CanCreateKeyEditor() override 
-	{
-		return true;
-	}
-
-	virtual TSharedRef<SWidget> CreateKeyEditor(ISequencer* Sequencer) override;
-};
-
-
-/**
- * A key area for FName curves.
- */
-class MOVIESCENETOOLS_API FNameCurveKeyArea
-	: public IKeyArea
-{
-public:
-
-	/**
-	 * Create and initialize a new instance.
-	 *
-	 * @param InCurve The curve to assign to this key area.
-	 * @param InOwningSection The section that owns this key area.
-	 */
-	FNameCurveKeyArea(FNameCurve& InCurve, UMovieSceneSection* InOwningSection)
-		: Curve(InCurve)
-		, OwningSection(InOwningSection)
-	{ }
-
-public:
-
-	// IKeyArea interface
-
-	virtual TArray<FKeyHandle> GetUnsortedKeyHandles() const override
-	{
-		TArray<FKeyHandle> OutKeyHandles;
-
-		for (auto It(Curve.GetKeyHandleIterator()); It; ++It)
-		{
-			OutKeyHandles.Add(It.Key());
-		}
-
-		return OutKeyHandles;
-	}
-
-	virtual void SetKeyTime(FKeyHandle KeyHandle, float NewKeyTime) const override
-	{
-		Curve.SetKeyTime(KeyHandle, NewKeyTime);
-	}
-
-	virtual float GetKeyTime(FKeyHandle KeyHandle) const override
-	{
-		return Curve.GetKeyTime(KeyHandle);
-	}
-
-	virtual FKeyHandle MoveKey( FKeyHandle KeyHandle, float DeltaPosition ) override
-	{
-		return Curve.SetKeyTime(KeyHandle, Curve.GetKeyTime(KeyHandle) + DeltaPosition);
-	}
-	
-	virtual void DeleteKey(FKeyHandle KeyHandle) override
-	{
-		Curve.DeleteKey(KeyHandle);
-	}
-
-	virtual void SetKeyInterpMode(FKeyHandle KeyHandle, ERichCurveInterpMode InterpMode) override
-	{
-		// do nothing
-	}
-
-	virtual ERichCurveInterpMode GetKeyInterpMode(FKeyHandle Keyhandle) const override
-	{
-		return RCIM_None;
-	}
-
-	virtual void SetKeyTangentMode(FKeyHandle KeyHandle, ERichCurveTangentMode TangentMode) override
-	{
-		// do nothing
-	}
-
-	virtual ERichCurveTangentMode GetKeyTangentMode(FKeyHandle KeyHandle) const override
-	{
-		return RCTM_None;
-	}
-
-	virtual void SetExtrapolationMode(ERichCurveExtrapolation ExtrapMode, bool bPreInfinity) override
-	{
-		// do nothing
-	}
-
-	virtual ERichCurveExtrapolation GetExtrapolationMode(bool bPreInfinity) const override
-	{
-		return RCCE_None;
-	}
-
-	virtual TArray<FKeyHandle> AddKeyUnique(float Time, EMovieSceneKeyInterpolation InKeyInterpolation, float TimeToCopyFrom = FLT_MAX) override;
-
-	virtual FRichCurve* GetRichCurve() override
-	{
-		return nullptr;
-	}
-
-	virtual UMovieSceneSection* GetOwningSection() override
-	{
-		return OwningSection.Get();
-	}
-
-	virtual bool CanCreateKeyEditor() override
-	{
-		return false;
-	}
-
-	virtual TSharedRef<SWidget> CreateKeyEditor(ISequencer* Sequencer) override
-	{
-		return SNullWidget::NullWidget;
-	}
-
-protected:
-
-	/** The curve managed by this area. */
-	FNameCurve& Curve;
-
-	/** The section that owns this area. */
-	TWeakObjectPtr<UMovieSceneSection> OwningSection;
-};
+class UMovieSceneSection;
 
 
 class MOVIESCENETOOLS_API MovieSceneToolHelpers
diff --git a/Engine/Source/Editor/MovieSceneTools/Public/NameCurveKeyArea.h b/Engine/Source/Editor/MovieSceneTools/Public/NameCurveKeyArea.h
new file mode 100644
--- /dev/null
+++ b/Engine/Source/Editor/MovieSceneTools/Public/NameCurveKeyArea.h
@@ -0,0 +1,55 @@
+// Copyright 1998-2015 Epic Games, Inc. All Rights Reserved.
+
+#pragma once
+
+#include "IKeyArea.h"
+
+
+/**
+ * A key area for FName curves.
+ */
+class MOVIESCENETOOLS_API FNameCurveKeyArea
+	: public IKeyArea
+{
+public:
+
+	/**
+	 * Create and initialize a new instance.
+	 *
+	 * @param InCurve The curve to assign to this key area.
+	 * @param InOwningSection The section that owns this key area.
+	 */
+	FNameCurveKeyArea(FNameCurve& InCurve, UMovieSceneSection* InOwningSection)
+		: Curve(InCurve)
+		, OwningSection(InOwningSection)
+	{ }
+
+public:
+
+	// IKeyArea interface
+
+	virtual TArray<FKeyHandle> AddKeyUnique(float Time, EMovieSceneKeyInterpolation InKeyInterpolation, float TimeToCopyFrom = FLT_MAX) override;
+	virtual bool CanCreateKeyEditor() override;
+	virtual TSharedRef<SWidget> CreateKeyEditor(ISequencer* Sequencer) override;
+	virtual void DeleteKey(FKeyHandle KeyHandle) override;
+	virtual ERichCurveExtrapolation GetExtrapolationMode(bool bPreInfinity) const override;
+	virtual ERichCurveInterpMode GetKeyInterpMode(FKeyHandle Keyhandle) const override;
+	virtual ERichCurveTangentMode GetKeyTangentMode(FKeyHandle KeyHandle) const override;
+	virtual float GetKeyTime(FKeyHandle KeyHandle) const override;
+	virtual UMovieSceneSection* GetOwningSection() override;
+	virtual FRichCurve* GetRichCurve() override;
+	virtual TArray<FKeyHandle> GetUnsortedKeyHandles() const override;
+	virtual FKeyHandle MoveKey(FKeyHandle KeyHandle, float DeltaPosition) override;
+	virtual void SetExtrapolationMode(ERichCurveExtrapolation ExtrapMode, bool bPreInfinity) override;
+	virtual void SetKeyInterpMode(FKeyHandle KeyHandle, ERichCurveInterpMode InterpMode) override;
+	virtual void SetKeyTangentMode(FKeyHandle KeyHandle, ERichCurveTangentMode TangentMode) override;
+	virtual void SetKeyTime(FKeyHandle KeyHandle, float NewKeyTime) const override;
+
+protected:
+
+	/** The curve managed by this area. */
+	FNameCurve& Curve;
+
+	/** The section that owns this area. */
+	TWeakObjectPtr<UMovieSceneSection> OwningSection;
+};
diff --git a/Engine/Source/Editor/MovieSceneTools/Public/PropertyTrackEditor.h b/Engine/Source/Editor/MovieSceneTools/Public/PropertyTrackEditor.h
--- a/Engine/Source/Editor/MovieSceneTools/Public/PropertyTrackEditor.h
+++ b/Engine/Source/Editor/MovieSceneTools/Public/PropertyTrackEditor.h
@@ -6,9 +6,14 @@
 #include "PropertySection.h"
 #include "ISequencerObjectChangeListener.h"
 
+
 class IPropertyHandle;
 class FPropertyChangedParams;
 
+
+DECLARE_MULTICAST_DELEGATE_TwoParams(FOnSetIntermediateValueFromPropertyChange, UMovieSceneTrack*, FPropertyChangedParams)
+
+
 /**
 * Tools for animatable property types such as floats ands vectors
 */
@@ -16,9 +21,6 @@ template<typename TrackType, typename SectionType, typename KeyDataType>
 class FPropertyTrackEditor
 	: public FKeyframeTrackEditor<TrackType, SectionType, KeyDataType>
 {
-public:
-	DECLARE_MULTICAST_DELEGATE_TwoParams( FOnSetIntermediateValueFromPropertyChange, UMovieSceneTrack*, FPropertyChangedParams )
-
 public:
 	/**
 	* Constructor
@@ -88,7 +90,7 @@ public:
 
 	~FPropertyTrackEditor()
 	{
-		TSharedPtr<ISequencer> Sequencer = this->GetSequencer();
+		TSharedPtr<ISequencer> Sequencer = FMovieSceneTrackEditor::GetSequencer();
 		if ( Sequencer.IsValid() )
 		{
 			ISequencerObjectChangeListener& ObjectChangeListener = Sequencer->GetObjectChangeListener();
@@ -110,11 +112,11 @@ public:
 
 	virtual void AddKey( const FGuid& ObjectGuid, UObject* AdditionalAsset = NULL ) override
 	{
-		TSharedPtr<ISequencer> Sequencer = this->GetSequencer();
-		ISequencerObjectChangeListener& ObjectChangeListener = Sequencer->GetObjectChangeListener();
+		ISequencerObjectChangeListener& ObjectChangeListener = FMovieSceneTrackEditor::GetSequencer()->GetObjectChangeListener();
+		TSharedPtr<ISequencer> Sequencer = FMovieSceneTrackEditor::GetSequencer();
 
 		TArray<UObject*> OutObjects;
-		Sequencer->GetRuntimeObjects( Sequencer->GetFocusedMovieSceneSequenceInstance(), ObjectGuid, OutObjects );
+		Sequencer->GetRuntimeObjects(Sequencer->GetFocusedMovieSceneSequenceInstance(), ObjectGuid, OutObjects);
 		for ( int32 i = 0; i < OutObjects.Num(); ++i )
 		{
 			ObjectChangeListener.TriggerAllPropertiesChanged( OutObjects[i] );
@@ -125,7 +127,7 @@ public:
 	{
 		TSharedRef<FPropertySection> PropertySection = MakePropertySectionInterface( SectionObject, Track );
 		OnSetIntermediateValueFromPropertyChange.AddSP( PropertySection, &FPropertySection::SetIntermediateValueForTrack );
-		this->GetSequencer()->OnGlobalTimeChanged().AddSP( PropertySection, &FPropertySection::ClearIntermediateValue );
+		FMovieSceneTrackEditor::GetSequencer()->OnGlobalTimeChanged().AddSP( PropertySection, &FPropertySection::ClearIntermediateValue );
 		return PropertySection;
 	}
 
@@ -160,7 +162,7 @@ private:
 	/** Adds a callback for property changes for the supplied property type name. */
 	void AddWatchedPropertyType( FName WatchedPropertyTypeName )
 	{
-		this->GetSequencer()->GetObjectChangeListener().GetOnAnimatablePropertyChanged( WatchedPropertyTypeName ).AddRaw( this, &FPropertyTrackEditor<TrackType, SectionType, KeyDataType>::OnAnimatedPropertyChanged );
+		FMovieSceneTrackEditor::GetSequencer()->GetObjectChangeListener().GetOnAnimatablePropertyChanged( WatchedPropertyTypeName ).AddRaw( this, &FPropertyTrackEditor<TrackType, SectionType, KeyDataType>::OnAnimatedPropertyChanged );
 		WatchedPropertyTypeNames.Add( WatchedPropertyTypeName );
 	}
 
@@ -189,7 +191,7 @@ private:
 	*/
 	virtual void OnAnimatedPropertyChanged( const FPropertyChangedParams& PropertyChangedParams )
 	{
-		this->AnimatablePropertyChanged( FOnKeyProperty::CreateRaw( this, &FPropertyTrackEditor::OnKeyProperty, PropertyChangedParams ) );
+		FMovieSceneTrackEditor::AnimatablePropertyChanged( FOnKeyProperty::CreateRaw( this, &FPropertyTrackEditor::OnKeyProperty, PropertyChangedParams ) );
 	}
 
 	/** Adds a key based on a property change. */
@@ -199,20 +201,31 @@ private:
 		GenerateKeysFromPropertyChanged( PropertyChangedParams, KeysForPropertyChange );
 
 		TSubclassOf<UMovieSceneTrack> CustomizedClass = GetCustomizedTrackClass( PropertyChangedParams.PropertyPath.Last() );
-		TSubclassOf<UMovieSceneTrack> TrackClass = (*CustomizedClass != nullptr ? *CustomizedClass : TrackType::StaticClass());
+		TSubclassOf<UMovieSceneTrack> TrackClass;
+		
+		if (CustomizedClass != nullptr)
+		{
+			TrackClass = CustomizedClass;
+		}
+		else
+		{
+			TrackClass = TrackType::StaticClass();
+		}
 
 		// If the track class has been customized for this property then it's possible this track editor doesn't support it, 
 		// also check for track editors which should only be used for customization.
-		if ( SupportsType( TrackClass ) && ( ForCustomizedUseOnly() == false || *CustomizedClass != nullptr) )
+		if ( SupportsType( TrackClass ) && ( ForCustomizedUseOnly() == false || CustomizedClass != nullptr) )
 		{
-			typename FKeyframeTrackEditor<TrackType, SectionType, KeyDataType>::FOnInitializeNewTrack OnInitializeNewTrack;
-			OnInitializeNewTrack.BindLambda( [&](TrackType* NewTrack) { InitializeNewTrack( NewTrack, PropertyChangedParams ); } );
-
-			typename FKeyframeTrackEditor<TrackType, SectionType, KeyDataType>::FOnSetIntermediateValue OnSetIntermediateValue;
-			OnSetIntermediateValue.BindSP( this, &FPropertyTrackEditor::SetIntermediateValueFromPropertyChange, PropertyChangedParams );
-
-			return this->AddKeysToObjects( PropertyChangedParams.ObjectsThatChanged, KeyTime, KeysForPropertyChange, PropertyChangedParams.KeyParams,
-				TrackClass, PropertyChangedParams.PropertyPath.Last()->GetFName(), OnInitializeNewTrack, OnSetIntermediateValue );
+			return FKeyframeTrackEditor<TrackType, SectionType, KeyDataType>::AddKeysToObjects(
+				PropertyChangedParams.ObjectsThatChanged,
+				KeyTime,
+				KeysForPropertyChange,
+				PropertyChangedParams.KeyParams,
+				TrackClass,
+				PropertyChangedParams.PropertyPath.Last()->GetFName(),
+				[&](TrackType* NewTrack) { InitializeNewTrack(NewTrack, PropertyChangedParams); },
+				[&](TrackType* NewTrack) { SetIntermediateValueFromPropertyChange(NewTrack, PropertyChangedParams); }
+			);
 		}
 		else
 		{
@@ -222,7 +235,7 @@ private:
 
 private:
 
-	void SetIntermediateValueFromPropertyChange(  UMovieSceneTrack* Track, FPropertyChangedParams PropertyChangedParams )
+	void SetIntermediateValueFromPropertyChange(TrackType* Track, FPropertyChangedParams PropertyChangedParams)
 	{
 		OnSetIntermediateValueFromPropertyChange.Broadcast( Track, PropertyChangedParams );
 	}
diff --git a/Engine/Source/Editor/Sequencer/Private/DisplayNodes/SequencerDisplayNode.h b/Engine/Source/Editor/Sequencer/Private/DisplayNodes/SequencerDisplayNode.h
--- a/Engine/Source/Editor/Sequencer/Private/DisplayNodes/SequencerDisplayNode.h
+++ b/Engine/Source/Editor/Sequencer/Private/DisplayNodes/SequencerDisplayNode.h
@@ -138,9 +138,9 @@ public:
 	/**
 	 * Set the node's display name.
 	 *
-	 * @param DisplayName the display name to set.
+	 * @param NewDisplayName the display name to set.
 	 */
-	virtual void SetDisplayName(const FText& DisplayName) = 0;
+	virtual void SetDisplayName(const FText& NewDisplayName) = 0;
 
 	/**
 	 * Generates a container widget for tree display in the animation outliner portion of the track area
diff --git a/Engine/Source/Editor/Sequencer/Private/DisplayNodes/SequencerObjectBindingNode.cpp b/Engine/Source/Editor/Sequencer/Private/DisplayNodes/SequencerObjectBindingNode.cpp
--- a/Engine/Source/Editor/Sequencer/Private/DisplayNodes/SequencerObjectBindingNode.cpp
+++ b/Engine/Source/Editor/Sequencer/Private/DisplayNodes/SequencerObjectBindingNode.cpp
@@ -94,12 +94,11 @@ void FSequencerObjectBindingNode::BuildContextMenu(FMenuBuilder& MenuBuilder)
 			if (ObjectClass->IsChildOf(AActor::StaticClass()))
 			{
 				FFormatNamedArguments Args;
-				MenuBuilder.AddMenuEntry(
+
+				MenuBuilder.AddSubMenu(
 					FText::Format( LOCTEXT("Assign Actor ", "Assign Actor"), Args),
-					FText::Format( LOCTEXT("AssignActorTooltip", "Assign the selected actor to this track"), Args ),
-					FSlateIcon(),
-					FUIAction(FExecuteAction::CreateSP(&GetSequencer(), &FSequencer::AssignActor, ObjectBinding, this),
-								FCanExecuteAction::CreateSP(&GetSequencer(), &FSequencer::CanAssignActor, ObjectBinding)) );
+					FText::Format( LOCTEXT("AssignActorTooltip", "Assign an actor to this track"), Args ),
+					FNewMenuDelegate::CreateRaw(&GetSequencer(), &FSequencer::AssignActor, ObjectBinding));
 			}
 		}
 	}
@@ -190,13 +189,13 @@ ESequencerNode::Type FSequencerObjectBindingNode::GetType() const
 }
 
 
-void FSequencerObjectBindingNode::SetDisplayName(const FText& DisplayName)
+void FSequencerObjectBindingNode::SetDisplayName(const FText& NewDisplayName)
 {
 	UMovieScene* MovieScene = GetSequencer().GetFocusedMovieSceneSequence()->GetMovieScene();
 
 	if (MovieScene != nullptr)
 	{
-		MovieScene->SetObjectDisplayName(ObjectBinding, DisplayName);
+		MovieScene->SetObjectDisplayName(ObjectBinding, NewDisplayName);
 	}
 }
 
diff --git a/Engine/Source/Editor/Sequencer/Private/DisplayNodes/SequencerObjectBindingNode.h b/Engine/Source/Editor/Sequencer/Private/DisplayNodes/SequencerObjectBindingNode.h
--- a/Engine/Source/Editor/Sequencer/Private/DisplayNodes/SequencerObjectBindingNode.h
+++ b/Engine/Source/Editor/Sequencer/Private/DisplayNodes/SequencerObjectBindingNode.h
@@ -47,7 +47,7 @@ public:
 	virtual float GetNodeHeight() const override;
 	virtual FNodePadding GetNodePadding() const override;
 	virtual ESequencerNode::Type GetType() const override;
-	virtual void SetDisplayName(const FText& DisplayName) override;
+	virtual void SetDisplayName(const FText& NewDisplayName) override;
 
 protected:
 
diff --git a/Engine/Source/Editor/Sequencer/Private/DisplayNodes/SequencerSectionCategoryNode.cpp b/Engine/Source/Editor/Sequencer/Private/DisplayNodes/SequencerSectionCategoryNode.cpp
--- a/Engine/Source/Editor/Sequencer/Private/DisplayNodes/SequencerSectionCategoryNode.cpp
+++ b/Engine/Source/Editor/Sequencer/Private/DisplayNodes/SequencerSectionCategoryNode.cpp
@@ -43,7 +43,7 @@ ESequencerNode::Type FSequencerSectionCategoryNode::GetType() const
 }
 
 
-void FSequencerSectionCategoryNode::SetDisplayName(const FText& InDisplayName)
+void FSequencerSectionCategoryNode::SetDisplayName(const FText& NewDisplayName)
 {
 	check(false);
 }
diff --git a/Engine/Source/Editor/Sequencer/Private/DisplayNodes/SequencerSectionCategoryNode.h b/Engine/Source/Editor/Sequencer/Private/DisplayNodes/SequencerSectionCategoryNode.h
--- a/Engine/Source/Editor/Sequencer/Private/DisplayNodes/SequencerSectionCategoryNode.h
+++ b/Engine/Source/Editor/Sequencer/Private/DisplayNodes/SequencerSectionCategoryNode.h
@@ -38,7 +38,7 @@ public:
 	virtual float GetNodeHeight() const override;
 	virtual FNodePadding GetNodePadding() const override;
 	virtual ESequencerNode::Type GetType() const override;
-	virtual void SetDisplayName(const FText& DisplayName) override;
+	virtual void SetDisplayName(const FText& NewDisplayName) override;
 
 private:
 
diff --git a/Engine/Source/Editor/Sequencer/Private/DisplayNodes/SequencerSectionKeyAreaNode.cpp b/Engine/Source/Editor/Sequencer/Private/DisplayNodes/SequencerSectionKeyAreaNode.cpp
--- a/Engine/Source/Editor/Sequencer/Private/DisplayNodes/SequencerSectionKeyAreaNode.cpp
+++ b/Engine/Source/Editor/Sequencer/Private/DisplayNodes/SequencerSectionKeyAreaNode.cpp
@@ -65,7 +65,7 @@ ESequencerNode::Type FSequencerSectionKeyAreaNode::GetType() const
 }
 
 
-void FSequencerSectionKeyAreaNode::SetDisplayName(const FText& InDisplayName)
+void FSequencerSectionKeyAreaNode::SetDisplayName(const FText& NewDisplayName)
 {
 	check(false);
 }
diff --git a/Engine/Source/Editor/Sequencer/Private/DisplayNodes/SequencerTrackNode.cpp b/Engine/Source/Editor/Sequencer/Private/DisplayNodes/SequencerTrackNode.cpp
--- a/Engine/Source/Editor/Sequencer/Private/DisplayNodes/SequencerTrackNode.cpp
+++ b/Engine/Source/Editor/Sequencer/Private/DisplayNodes/SequencerTrackNode.cpp
@@ -185,12 +185,12 @@ ESequencerNode::Type FSequencerTrackNode::GetType() const
 }
 
 
-void FSequencerTrackNode::SetDisplayName(const FText& DisplayName)
+void FSequencerTrackNode::SetDisplayName(const FText& NewDisplayName)
 {
 	auto NameableTrack = Cast<UMovieSceneNameableTrack>(AssociatedTrack.Get());
 
 	if (NameableTrack != nullptr)
 	{
-		NameableTrack->SetDisplayName(DisplayName);
+		NameableTrack->SetDisplayName(NewDisplayName);
 	}
 }
diff --git a/Engine/Source/Editor/Sequencer/Private/GroupedKeyArea.cpp b/Engine/Source/Editor/Sequencer/Private/GroupedKeyArea.cpp
--- a/Engine/Source/Editor/Sequencer/Private/GroupedKeyArea.cpp
+++ b/Engine/Source/Editor/Sequencer/Private/GroupedKeyArea.cpp
@@ -163,7 +163,7 @@ FLinearColor FGroupedKeyArea::GetKeyTint(FKeyHandle InHandle) const
 
 const FSlateBrush* FGroupedKeyArea::GetBrush(FKeyHandle InHandle) const
 {
-	static const FSlateBrush* KeyBrush = FEditorStyle::GetBrush("Sequencer.Key");
+	static const FSlateBrush* KeyBrush = FEditorStyle::GetBrush("Sequencer.KeyDiamond");
 	static const FSlateBrush* PartialKeyBrush = FEditorStyle::GetBrush("Sequencer.PartialKey");
 
 	const auto* Group = FindGroup(InHandle);
diff --git a/Engine/Source/Editor/Sequencer/Private/SSequencer.cpp b/Engine/Source/Editor/Sequencer/Private/SSequencer.cpp
--- a/Engine/Source/Editor/Sequencer/Private/SSequencer.cpp
+++ b/Engine/Source/Editor/Sequencer/Private/SSequencer.cpp
@@ -100,12 +100,17 @@ void SSequencer::Construct(const FArguments& InArgs, TSharedRef<FSequencer> InSe
 
 	FSequencerWidgetsModule& SequencerWidgets = FModuleManager::Get().LoadModuleChecked<FSequencerWidgetsModule>( "SequencerWidgets" );
 
+	OnBeginPlaybackRangeDrag = InArgs._OnBeginPlaybackRangeDrag;
+	OnEndPlaybackRangeDrag = InArgs._OnEndPlaybackRangeDrag;
+
 	FTimeSliderArgs TimeSliderArgs;
 	{
 		TimeSliderArgs.ViewRange = InArgs._ViewRange;
 		TimeSliderArgs.ClampRange = InArgs._ClampRange;
 		TimeSliderArgs.PlaybackRange = InArgs._PlaybackRange;
 		TimeSliderArgs.OnPlaybackRangeChanged = InArgs._OnPlaybackRangeChanged;
+		TimeSliderArgs.OnBeginPlaybackRangeDrag = OnBeginPlaybackRangeDrag;
+		TimeSliderArgs.OnEndPlaybackRangeDrag = OnEndPlaybackRangeDrag;
 		TimeSliderArgs.OnViewRangeChanged = InArgs._OnViewRangeChanged;
 		TimeSliderArgs.OnClampRangeChanged = InArgs._OnClampRangeChanged;
 		TimeSliderArgs.ScrubPosition = InArgs._ScrubPosition;
@@ -411,6 +416,7 @@ void SSequencer::Construct(const FArguments& InArgs, TSharedRef<FSequencer> InSe
 							[
 								// details view panel
 								SNew(SVerticalBox)
+								.Visibility(this, &SSequencer::HandleDetailsViewVisibility)
 
 								+ SVerticalBox::Slot()
 									.AutoHeight()
@@ -744,10 +750,9 @@ TSharedRef<SWidget> SSequencer::MakeGeneralMenu()
 		TSharedPtr<FSequencer> PinnedSequencer = Sequencer.Pin();
 
 		// Menu entry for the start position
-		auto OnStartComitted = [=](float NewValue, ETextCommit::Type){
-			TRange<float> Range = PinnedSequencer->GetPlaybackRange();
-			NewValue = FMath::Min(NewValue, Range.GetUpperBoundValue());
-			PinnedSequencer->GetFocusedMovieSceneSequence()->GetMovieScene()->SetPlaybackRange(NewValue, Range.GetUpperBoundValue());
+		auto OnStartChanged = [=](float NewValue){
+			float Upper = PinnedSequencer->GetPlaybackRange().GetUpperBoundValue();
+			PinnedSequencer->SetPlaybackRange(TRange<float>(FMath::Min(NewValue, Upper), Upper));
 		};
 		MenuBuilder.AddWidget(
 			SNew(SHorizontalBox)
@@ -761,9 +766,15 @@ TSharedRef<SWidget> SSequencer::MakeGeneralMenu()
 				SNew(SSpinBox<float>)
 				.TypeInterface(NumericTypeInterface)
 				.Style(&FEditorStyle::GetWidgetStyle<FSpinBoxStyle>("Sequencer.HyperlinkSpinBox"))
-				.OnValueCommitted_Lambda(OnStartComitted)
-				.OnValueChanged_Lambda([=](float NewValue){
-					OnStartComitted(NewValue, ETextCommit::Default);
+				.OnValueCommitted_Lambda([=](float Value, ETextCommit::Type){ OnStartChanged(Value); })
+				.OnValueChanged_Lambda(OnStartChanged)
+				.OnBeginSliderMovement(OnBeginPlaybackRangeDrag)
+				.OnEndSliderMovement_Lambda([=](float Value){ OnStartChanged(Value); OnEndPlaybackRangeDrag.ExecuteIfBound(); })
+				.MinValue_Lambda([=]() -> float {
+					return PinnedSequencer->GetViewRange().GetLowerBoundValue(); 
+				})
+				.MaxValue_Lambda([=]() -> float {
+					return PinnedSequencer->GetPlaybackRange().GetUpperBoundValue(); 
 				})
 				.Value_Lambda([=]() -> float {
 					return PinnedSequencer->GetPlaybackRange().GetLowerBoundValue();
@@ -772,10 +783,9 @@ TSharedRef<SWidget> SSequencer::MakeGeneralMenu()
 			LOCTEXT("PlaybackStartLabel", "Start"));
 
 		// Menu entry for the end position
-		auto OnEndComitted = [=](float NewValue, ETextCommit::Type){
-			TRange<float> Range = PinnedSequencer->GetPlaybackRange();
-			NewValue = FMath::Max(NewValue, Range.GetLowerBoundValue());
-			PinnedSequencer->GetFocusedMovieSceneSequence()->GetMovieScene()->SetPlaybackRange(Range.GetLowerBoundValue(), NewValue);
+		auto OnEndChanged = [=](float NewValue){
+			float Lower = PinnedSequencer->GetPlaybackRange().GetLowerBoundValue();
+			PinnedSequencer->SetPlaybackRange(TRange<float>(Lower, FMath::Max(NewValue, Lower)));
 		};
 		MenuBuilder.AddWidget(
 			SNew(SHorizontalBox)
@@ -789,9 +799,15 @@ TSharedRef<SWidget> SSequencer::MakeGeneralMenu()
 				SNew(SSpinBox<float>)
 				.TypeInterface(NumericTypeInterface)
 				.Style(&FEditorStyle::GetWidgetStyle<FSpinBoxStyle>("Sequencer.HyperlinkSpinBox"))
-				.OnValueCommitted_Lambda(OnEndComitted)
-				.OnValueChanged_Lambda([=](float NewValue){
-					OnEndComitted(NewValue, ETextCommit::Default);
+				.OnValueCommitted_Lambda([=](float Value, ETextCommit::Type){ OnEndChanged(Value); })
+				.OnValueChanged_Lambda(OnEndChanged)
+				.OnBeginSliderMovement(OnBeginPlaybackRangeDrag)
+				.OnEndSliderMovement_Lambda([=](float Value){ OnEndChanged(Value); OnEndPlaybackRangeDrag.ExecuteIfBound(); })
+				.MinValue_Lambda([=]() -> float {
+					return PinnedSequencer->GetPlaybackRange().GetLowerBoundValue(); 
+				})
+				.MaxValue_Lambda([=]() -> float {
+					return PinnedSequencer->GetViewRange().GetUpperBoundValue(); 
 				})
 				.Value_Lambda([=]() -> float {
 					return PinnedSequencer->GetPlaybackRange().GetUpperBoundValue();
@@ -1324,6 +1340,7 @@ TArray<FSectionHandle> SSequencer::GetSectionHandles(const TSet<TWeakObjectPtr<U
 {
 	TArray<FSectionHandle> SectionHandles;
 
+	// @todo sequencer: this is potentially slow as it traverses the entire tree - there's scope for optimization here
 	for (auto& Node : SequencerNodeTree->GetRootNodes())
 	{
 		Node->Traverse_ParentFirst([&](FSequencerDisplayNode& InNode) {
diff --git a/Engine/Source/Editor/Sequencer/Private/SSequencer.h b/Engine/Source/Editor/Sequencer/Private/SSequencer.h
--- a/Engine/Source/Editor/Sequencer/Private/SSequencer.h
+++ b/Engine/Source/Editor/Sequencer/Private/SSequencer.h
@@ -97,6 +97,12 @@ public:
 		/** Called when the user changes the playback range */
 		SLATE_EVENT( FOnRangeChanged, OnPlaybackRangeChanged )
 
+		/** Called when the user has begun dragging the playback range */
+		SLATE_EVENT( FSimpleDelegate, OnBeginPlaybackRangeDrag )
+
+		/** Called when the user has finished dragging the playback range */
+		SLATE_EVENT( FSimpleDelegate, OnEndPlaybackRangeDrag )
+
 		/** The current scrub position in (seconds) */
 		SLATE_ATTRIBUTE( float, ScrubPosition )
 
@@ -409,4 +415,10 @@ private:
 	TSharedPtr<INumericTypeInterface<float>> NumericTypeInterface;
 
 	FOnGetAddMenuContent OnGetAddMenuContent;
+	/** Called when the user has begun dragging the playback range */
+	FSimpleDelegate OnBeginPlaybackRangeDrag;
+
+	/** Called when the user has finished dragging the playback range */
+	FSimpleDelegate OnEndPlaybackRangeDrag;
+
 };
diff --git a/Engine/Source/Editor/Sequencer/Private/SSequencerCurveEditor.cpp b/Engine/Source/Editor/Sequencer/Private/SSequencerCurveEditor.cpp
--- a/Engine/Source/Editor/Sequencer/Private/SSequencerCurveEditor.cpp
+++ b/Engine/Source/Editor/Sequencer/Private/SSequencerCurveEditor.cpp
@@ -48,7 +48,7 @@ FReply SSequencerCurveEditor::OnMouseButtonDown( const FGeometry& MyGeometry, co
 		return Reply;
 	}
 
-	return TimeSliderController->OnMouseButtonDown( AsShared(), MyGeometry, MouseEvent );
+	return TimeSliderController->OnMouseButtonDown( *this, MyGeometry, MouseEvent );
 }
 
 
@@ -61,7 +61,7 @@ FReply SSequencerCurveEditor::OnMouseButtonUp( const FGeometry& MyGeometry, cons
 		return Reply;
 	}
 
-	return TimeSliderController->OnMouseButtonUp( AsShared(), MyGeometry, MouseEvent );
+	return TimeSliderController->OnMouseButtonUp( *this, MyGeometry, MouseEvent );
 }
 
 
@@ -73,12 +73,12 @@ FReply SSequencerCurveEditor::OnMouseMove( const FGeometry& MyGeometry, const FP
 		return Reply;
 	}
 
-	return TimeSliderController->OnMouseMove( AsShared(), MyGeometry, MouseEvent );
+	return TimeSliderController->OnMouseMove( *this, MyGeometry, MouseEvent );
 }
 
 FReply SSequencerCurveEditor::OnMouseWheel( const FGeometry& MyGeometry, const FPointerEvent& MouseEvent )
 {
-	FReply Reply = TimeSliderController->OnMouseWheel( AsShared(), MyGeometry, MouseEvent );
+	FReply Reply = TimeSliderController->OnMouseWheel( *this, MyGeometry, MouseEvent );
 	if ( Reply.IsEventHandled() )
 	{
 		return Reply;
diff --git a/Engine/Source/Editor/Sequencer/Private/SSequencerSection.cpp b/Engine/Source/Editor/Sequencer/Private/SSequencerSection.cpp
--- a/Engine/Source/Editor/Sequencer/Private/SSequencerSection.cpp
+++ b/Engine/Source/Editor/Sequencer/Private/SSequencerSection.cpp
@@ -36,8 +36,6 @@ void SSequencerSection::Construct( const FArguments& InArgs, TSharedRef<FSequenc
 	SectionInterface = SectionNode->GetSections()[InSectionIndex];
 	Layout = FKeyAreaLayout(*SectionNode, InSectionIndex);
 
-	ResetState();
-
 	ChildSlot
 	[
 		SectionInterface->GenerateSectionWidget()
@@ -190,11 +188,6 @@ FSequencerSelectedKey SSequencerSection::CreateKeyUnderMouse( const FVector2D& M
 
 void SSequencerSection::CheckForEdgeInteraction( const FPointerEvent& MouseEvent, const FGeometry& SectionGeometry )
 {
-	bLeftEdgeHovered = false;
-	bRightEdgeHovered = false;
-	bLeftEdgePressed = false;
-	bRightEdgePressed = false;
-
 	if (!SectionInterface->SectionIsResizable())
 	{
 		return;
@@ -213,28 +206,10 @@ void SSequencerSection::CheckForEdgeInteraction( const FPointerEvent& MouseEvent
 
 	if( SectionRectLeft.IsUnderLocation( MouseEvent.GetScreenSpacePosition() ) )
 	{
-		if( MouseEvent.GetEffectingButton() == EKeys::LeftMouseButton )
-		{
-			bLeftEdgePressed = true;
-		}
-		else
-		{
-			bLeftEdgeHovered = true;
-		}
-
 		GetSequencer().GetEditTool().SetHotspot(MakeShareable( new FSectionResizeHotspot(FSectionResizeHotspot::Left, FSectionHandle(ParentSectionArea, SectionIndex))) );
 	}
 	else if( SectionRectRight.IsUnderLocation( MouseEvent.GetScreenSpacePosition() ) )
 	{
-		if( MouseEvent.GetEffectingButton() == EKeys::LeftMouseButton )
-		{
-			bRightEdgePressed = true;
-		}
-		else
-		{
-			bRightEdgeHovered = true;
-		}
-
 		GetSequencer().GetEditTool().SetHotspot(MakeShareable( new FSectionResizeHotspot(FSectionResizeHotspot::Right, FSectionHandle(ParentSectionArea, SectionIndex))) );
 	}
 }
@@ -298,26 +273,31 @@ int32 SSequencerSection::OnPaint( const FPaintArgs& Args, const FGeometry& Allot
 
 void SSequencerSection::PaintKeys( const FGeometry& AllottedGeometry, const FSlateRect& MyClippingRect, FSlateWindowElementList& OutDrawElements, int32 LayerId, const FWidgetStyle& InWidgetStyle, bool bParentEnabled ) const
 {
-	const ESlateDrawEffect::Type DrawEffects = bParentEnabled ? ESlateDrawEffect::None : ESlateDrawEffect::DisabledEffect;
-
-	UMovieSceneSection& SectionObject = *SectionInterface->GetSectionObject();
-	FSequencer& Sequencer = ParentSectionArea->GetSequencer();
-
 	static const FName BackgroundBrushName("Sequencer.SectionArea.Background");
-	static const FName KeyBrushName("Sequencer.Key");
-
-	const FSlateBrush* BackgroundBrush = FEditorStyle::GetBrush(BackgroundBrushName);
-	const FSlateBrush* DefaultKeyBrush = FEditorStyle::GetBrush(KeyBrushName);
+	static const FName CircleKeyBrushName("Sequencer.KeyCircle");
+	static const FName DiamondKeyBrushName("Sequencer.KeyDiamond");
+	static const FName SquareKeyBrushName("Sequencer.KeySquare");
+	static const FName TriangleKeyBrushName("Sequencer.KeyTriangle");
 
 	static const FName SelectionColorName("SelectionColor");
 	static const FName SelectionInactiveColorName("SelectionColorInactive");
 	static const FName SelectionColorPressedName("SelectionColor_Pressed");
 
+	static const float BrushBorderWidth = 2.0f;
+
 	const FLinearColor PressedKeyColor = FEditorStyle::GetSlateColor(SelectionColorPressedName).GetColor( InWidgetStyle );
 	FLinearColor SelectionColor = FEditorStyle::GetSlateColor(SelectionColorName).GetColor( InWidgetStyle );
 	FLinearColor SelectedKeyColor = SelectionColor;
-	// @todo Sequencer temp color, make hovered brighter than selected.
-	FLinearColor HoveredKeyColor = SelectedKeyColor * FLinearColor(1.5,1.5,1.5,1.0f);
+	FSequencer& Sequencer = ParentSectionArea->GetSequencer();
+	TSharedPtr<ISequencerHotspot> Hotspot = Sequencer.GetEditTool().GetHotspot();
+
+	// get hovered key
+	FSequencerSelectedKey HoveredKey;
+
+	if (Hotspot.IsValid() && Hotspot->GetType() == ESequencerHotspot::Key)
+	{
+		HoveredKey = static_cast<FKeyHotspot*>(Hotspot.Get())->Key;
+	}
 
 	auto& Selection = Sequencer.GetSelection();
 	auto& SelectionPreview = Sequencer.GetSelectionPreview();
@@ -325,15 +305,27 @@ void SSequencerSection::PaintKeys( const FGeometry& AllottedGeometry, const FSla
 	if (Selection.GetActiveSelection() != FSequencerSelection::EActiveSelection::KeyAndSection)
 	{
 		// Selected key color is different when the active selection is not keys
-		SelectedKeyColor = FEditorStyle::GetSlateColor(SelectionInactiveColorName).GetColor( InWidgetStyle )
-			* FLinearColor(.25, .25, .25, 1);  // Make the color a little darker since it's not very visible next to white keyframes.
+		SelectedKeyColor = FEditorStyle::GetSlateColor(SelectionInactiveColorName).GetColor(InWidgetStyle)
+			* FLinearColor(.25, .25, .25, 1);  // Make the color a little darker since it's not very visible next to white key frames.
 	}
 
-	// Draw all keys in each key area
-	for (const FKeyAreaLayoutElement& Element : Layout->GetElements())
+	// draw all keys in each key area
+	UMovieSceneSection& SectionObject = *SectionInterface->GetSectionObject();
+
+	const FSlateBrush* BackgroundBrush = FEditorStyle::GetBrush(BackgroundBrushName);
+	const FSlateBrush* CircleKeyBrush = FEditorStyle::GetBrush(CircleKeyBrushName);
+	const FSlateBrush* DiamondKeyBrush = FEditorStyle::GetBrush(DiamondKeyBrushName);
+	const FSlateBrush* SquareKeyBrush = FEditorStyle::GetBrush(SquareKeyBrushName);
+	const FSlateBrush* TriangleKeyBrush = FEditorStyle::GetBrush(TriangleKeyBrushName);
+
+	const ESlateDrawEffect::Type DrawEffects = bParentEnabled
+		? ESlateDrawEffect::None
+		: ESlateDrawEffect::DisabledEffect;
+
+	for (const FKeyAreaLayoutElement& LayoutElement : Layout->GetElements())
 	{
-		// Get the key area at the same index of the section.  Each section in this widget has the same layout and the same number of key areas
-		TSharedPtr<IKeyArea> KeyArea = Element.GetKeyArea();
+		// get key handles
+		TSharedPtr<IKeyArea> KeyArea = LayoutElement.GetKeyArea();
 		TArray<FKeyHandle> KeyHandles = KeyArea->GetUnsortedKeyHandles();
 
 		if (!KeyHandles.Num())
@@ -341,14 +333,13 @@ void SSequencerSection::PaintKeys( const FGeometry& AllottedGeometry, const FSla
 			continue;
 		}
 
-		FGeometry KeyAreaGeometry = GetKeyAreaGeometry( Element, AllottedGeometry );
+		FGeometry KeyAreaGeometry = GetKeyAreaGeometry(LayoutElement, AllottedGeometry);
 
 		FTimeToPixel TimeToPixelConverter = SectionObject.IsInfinite()
-			? FTimeToPixel( ParentGeometry, GetSequencer().GetViewRange())
-			: FTimeToPixel( KeyAreaGeometry, TRange<float>( SectionObject.GetStartTime(), SectionObject.GetEndTime() ) );
+			? FTimeToPixel(ParentGeometry, Sequencer.GetViewRange())
+			: FTimeToPixel(KeyAreaGeometry, TRange<float>(SectionObject.GetStartTime(), SectionObject.GetEndTime()));
 
-		// Draw a box for the key area 
-		// @todo Sequencer - Allow the IKeyArea to do this
+		// draw a box for the key area 
 		FSlateDrawElement::MakeBox( 
 			OutDrawElements,
 			LayerId,
@@ -356,7 +347,7 @@ void SSequencerSection::PaintKeys( const FGeometry& AllottedGeometry, const FSla
 			BackgroundBrush,
 			MyClippingRect,
 			DrawEffects,
-			FLinearColor( .1f, .1f, .1f, 0.7f )
+			FLinearColor(0.1f, 0.1f, 0.1f, 0.7f)
 		); 
 
 		const int32 KeyLayer = LayerId + 1;
@@ -365,72 +356,156 @@ void SSequencerSection::PaintKeys( const FGeometry& AllottedGeometry, const FSla
 		{
 			float KeyTime = KeyArea->GetKeyTime(KeyHandle);
 
-			// Omit keys which would not be visible
-			if( !SectionObject.IsTimeWithinSection( KeyTime ) )
+			// omit keys which would not be visible
+			if( !SectionObject.IsTimeWithinSection(KeyTime))
 			{
 				continue;
 			}
 
-			const FSlateBrush* BrushToUse = SectionInterface->GetKeyBrush(KeyHandle);
-			if(BrushToUse == nullptr)
+			// determine the key's brush & color
+			const FSlateBrush* KeyBrush;
+			FLinearColor KeyColor;
+			
+			switch (KeyArea->GetKeyInterpMode(KeyHandle))
 			{
-				BrushToUse = DefaultKeyBrush;
+			case RCIM_Linear:
+				KeyBrush = TriangleKeyBrush;
+				KeyColor = FLinearColor(0.0f, 0.617f, 0.449f, 1.0f); // blueish green
+				break;
+
+			case RCIM_Constant:
+				KeyBrush = SquareKeyBrush;
+				KeyColor = FLinearColor(0.0f, 0.445f, 0.695f, 1.0f); // blue
+				break;
+
+			case RCIM_Cubic:
+				KeyBrush = CircleKeyBrush;
+
+				switch (KeyArea->GetKeyTangentMode(KeyHandle))
+				{
+				case RCTM_Auto:
+					KeyColor = FLinearColor(0.972f, 0.2f, 0.2f, 1.0f); // vermillion
+					break;
+
+				case RCTM_Break:
+					KeyColor = FLinearColor(0.336f, 0.703f, 0.5f, 0.91f); // sky blue
+					break;
+
+				case RCTM_User:
+					KeyColor = FLinearColor(0.797f, 0.473f, 0.5f, 0.652f); // reddish purple
+					break;
+
+				default:
+					KeyColor = FLinearColor(0.75f, 0.75f, 0.75f, 1.0f); // light gray
+				}
+				break;
+
+			default:
+				KeyBrush = DiamondKeyBrush;
+				KeyColor = FLinearColor(1.0f, 1.0f, 1.0f, 1.0f); // white
+				break;
 			}
-			FLinearColor KeyColor( 1.0f, 1.0f, 1.0f, 1.0f );
-			FLinearColor KeyTint(1.f, 1.f, 1.f, 1.f);
 
-			if (Element.GetType() == FKeyAreaLayoutElement::Group)
+			// allow group & section overrides
+			const FSlateBrush* OverrideBrush = nullptr;
+
+			if (LayoutElement.GetType() == FKeyAreaLayoutElement::Group)
 			{
 				auto Group = StaticCastSharedPtr<FGroupedKeyArea>(KeyArea);
-				KeyTint = Group->GetKeyTint(KeyHandle);
-				BrushToUse = Group->GetBrush(KeyHandle);
+				OverrideBrush = Group->GetBrush(KeyHandle);
+			}
+			else
+			{
+				OverrideBrush = SectionInterface->GetKeyBrush(KeyHandle);
 			}
 
-			// Where to start drawing the key (relative to the section)
-			float KeyPosition =  TimeToPixelConverter.TimeToPixel( KeyTime );
+			if (OverrideBrush != nullptr)
+			{
+				KeyBrush = OverrideBrush;
+			}
 
-			FSequencerSelectedKey TestKey( SectionObject, KeyArea, KeyHandle );
+			// determine draw colors based on hover, selection, etc.
+			FLinearColor BorderColor;
+			FLinearColor FillColor;
 
-			bool bSelected = Selection.IsSelected( TestKey );
-			ESelectionPreviewState SelectionPreviewState = SelectionPreview.GetSelectionState( TestKey );
+			FSequencerSelectedKey TestKey(SectionObject, KeyArea, KeyHandle);
+			ESelectionPreviewState SelectionPreviewState = SelectionPreview.GetSelectionState(TestKey);
+			bool bSelected = Selection.IsSelected(TestKey);
 
-			if( TestKey == PressedKey )
+			if (SelectionPreviewState == ESelectionPreviewState::Selected)
 			{
-				KeyColor = PressedKeyColor;
+				FLinearColor PreviewSelectionColor = SelectionColor.LinearRGBToHSV();
+				PreviewSelectionColor.R += 0.1f; // +10% hue
+				PreviewSelectionColor.G = 0.6f; // 60% saturation
+				BorderColor = PreviewSelectionColor.HSVToLinearRGB();
+				FillColor = BorderColor;
 			}
-			else if( TestKey == HoveredKey )
+			else if (SelectionPreviewState == ESelectionPreviewState::NotSelected)
 			{
-				KeyColor = HoveredKeyColor;
+				BorderColor = FLinearColor(0.05f, 0.05f, 0.05f, 1.0f);
+				FillColor = KeyColor;
 			}
-			else if( SelectionPreviewState == ESelectionPreviewState::Selected )
+			else if (bSelected)
 			{
-				FLinearColor PreviewSelectionColor = SelectionColor.LinearRGBToHSV();
-				PreviewSelectionColor.R += 0.1f; // +10% hue
-				PreviewSelectionColor.G = 0.6f; // 60% saturation
-				KeyColor = PreviewSelectionColor.HSVToLinearRGB();
+				BorderColor = SelectedKeyColor;
+				FillColor = FLinearColor(0.05f, 0.05f, 0.05f, 1.0f);
+			}
+			else if (TestKey == HoveredKey)
+			{
+				BorderColor = FLinearColor(1.0f, 1.0f, 1.0f, 1.0f);
+				FillColor = FLinearColor(1.0f, 1.0f, 1.0f, 1.0f);
 			}
-			else if( SelectionPreviewState == ESelectionPreviewState::NotSelected )
+			else
 			{
-				// Default white selection color
+				BorderColor = FLinearColor(0.05f, 0.05f, 0.05f, 1.0f);
+				FillColor = KeyColor;
 			}
-			else if( bSelected )
+
+			// allow group to tint the color
+			if (LayoutElement.GetType() == FKeyAreaLayoutElement::Group)
 			{
-				KeyColor = SelectedKeyColor;
+				auto Group = StaticCastSharedPtr<FGroupedKeyArea>(KeyArea);
+				KeyColor *= Group->GetKeyTint(KeyHandle);
 			}
 
-			KeyColor *= KeyTint;
+			// draw border
+			float KeyPosition = TimeToPixelConverter.TimeToPixel(KeyTime);
 
-			// Draw the key
 			FSlateDrawElement::MakeBox(
 				OutDrawElements,
 				// always draw selected keys on top of other keys
-				bSelected ? KeyLayer+1 : KeyLayer,
+				bSelected ? KeyLayer + 1 : KeyLayer,
 				// Center the key along Y.  Ensure the middle of the key is at the actual key time
-				KeyAreaGeometry.ToPaintGeometry( FVector2D( KeyPosition - FMath::CeilToFloat(SequencerSectionConstants::KeySize.X/2.0f), ((KeyAreaGeometry.Size.Y*.5f)-(SequencerSectionConstants::KeySize.Y*.5f)) ), SequencerSectionConstants::KeySize ),
-				BrushToUse,
+				KeyAreaGeometry.ToPaintGeometry(
+					FVector2D(
+						KeyPosition - FMath::CeilToFloat(SequencerSectionConstants::KeySize.X / 2.0f),
+						((KeyAreaGeometry.Size.Y / 2.0f) - (SequencerSectionConstants::KeySize.Y / 2.0f))
+					),
+					SequencerSectionConstants::KeySize
+				),
+				KeyBrush,
 				MyClippingRect,
 				DrawEffects,
-				KeyColor
+				BorderColor
+			);
+
+			// draw fill
+			FSlateDrawElement::MakeBox(
+				OutDrawElements,
+				// always draw selected keys on top of other keys
+				bSelected ? KeyLayer + 2 : KeyLayer + 1,
+				// Center the key along Y.  Ensure the middle of the key is at the actual key time
+				KeyAreaGeometry.ToPaintGeometry(
+					FVector2D(
+						KeyPosition - FMath::CeilToFloat(SequencerSectionConstants::KeySize.X / 2.0f - BrushBorderWidth),
+						((KeyAreaGeometry.Size.Y / 2.0f) - (SequencerSectionConstants::KeySize.Y / 2.0f - BrushBorderWidth))
+					),
+					SequencerSectionConstants::KeySize - 2.0f * BrushBorderWidth
+				),
+				KeyBrush,
+				MyClippingRect,
+				DrawEffects,
+				FillColor
 			);
 		}
 	}
@@ -442,6 +517,24 @@ void SSequencerSection::DrawSectionHandles( const FGeometry& AllottedGeometry, c
 	const FSlateBrush* LeftGripBrush = FEditorStyle::GetBrush(SectionInterface->GetSectionGripLeftBrushName());
 	const FSlateBrush* RightGripBrush = FEditorStyle::GetBrush(SectionInterface->GetSectionGripRightBrushName());
 
+	bool bLeftHandleActive = false;
+	bool bRightHandleActive = false;
+
+	// Get the hovered/selected state for the section handles from the hotspot
+	TSharedPtr<ISequencerHotspot> Hotspot = GetSequencer().GetEditTool().GetHotspot();
+
+	if (Hotspot.IsValid() && (
+		Hotspot->GetType() == ESequencerHotspot::SectionResize_L ||
+		Hotspot->GetType() == ESequencerHotspot::SectionResize_R))
+	{
+		FSectionResizeHotspot& ResizeHotspot = static_cast<FSectionResizeHotspot&>(*Hotspot.Get());
+		if (ResizeHotspot.Section == FSectionHandle(ParentSectionArea, SectionIndex))
+		{
+			bLeftHandleActive = Hotspot->GetType() == ESequencerHotspot::SectionResize_L;
+			bRightHandleActive = Hotspot->GetType() == ESequencerHotspot::SectionResize_R;
+		}
+	}
+
 	// Left Grip
 	FSlateDrawElement::MakeBox
 	(
@@ -452,7 +545,7 @@ void SSequencerSection::DrawSectionHandles( const FGeometry& AllottedGeometry, c
 		LeftGripBrush,
 		MyClippingRect,
 		DrawEffects,
-		(bLeftEdgePressed || bLeftEdgeHovered) ? SelectionColor : LeftGripBrush->GetTint(FWidgetStyle())
+		bLeftHandleActive ? SelectionColor : LeftGripBrush->GetTint(FWidgetStyle())
 	);
 	
 	// Right Grip
@@ -465,7 +558,7 @@ void SSequencerSection::DrawSectionHandles( const FGeometry& AllottedGeometry, c
 		RightGripBrush,
 		MyClippingRect,
 		DrawEffects,
-		(bRightEdgePressed || bRightEdgeHovered) ? SelectionColor : RightGripBrush->GetTint(FWidgetStyle())
+		bRightHandleActive ? SelectionColor : RightGripBrush->GetTint(FWidgetStyle())
 	);
 }
 
@@ -517,279 +610,6 @@ void SSequencerSection::DrawSelectionBorder( const FGeometry& AllottedGeometry,
 	);
 }
 
-
-TSharedPtr<SWidget> SSequencerSection::OnSummonContextMenu( const FGeometry& MyGeometry, const FPointerEvent& MouseEvent )
-{
-	FSequencerSelectedKey Key = GetKeyUnderMouse( MouseEvent.GetScreenSpacePosition(), MyGeometry );
-	FSequencer& Sequencer = GetSequencer();
-	UMovieSceneSection* SceneSection = SectionInterface->GetSectionObject();
-
-	// @todo sequencer replace with UI Commands instead of faking it
-	const bool bShouldCloseWindowAfterMenuSelection = true;
-	FMenuBuilder MenuBuilder(bShouldCloseWindowAfterMenuSelection, NULL);
-
-	bool bHasMenuItems = false;
-	if (Key.IsValid() || Sequencer.GetSelection().GetSelectedKeys().Num())
-	{
-		bHasMenuItems = true;
-		MenuBuilder.BeginSection("SequencerInterpolation", NSLOCTEXT("Sequencer", "KeyInterpolationMenu", "Key Interpolation"));
-		{
-			MenuBuilder.AddMenuEntry(
-				NSLOCTEXT("Sequencer", "SetKeyInterpolationAuto", "Auto"),
-				NSLOCTEXT("Sequencer", "SetKeyInterpolationAutoTooltip", "Set key interpolation to auto"),
-				FSlateIcon(),
-				FUIAction(
-					FExecuteAction::CreateSP(&Sequencer, &FSequencer::SetInterpTangentMode, RCIM_Cubic, RCTM_Auto),
-					FCanExecuteAction(),
-					FIsActionChecked::CreateSP(&Sequencer, &FSequencer::IsInterpTangentModeSelected, RCIM_Cubic, RCTM_Auto) ),
-				NAME_None,
-				EUserInterfaceActionType::ToggleButton
-			);
-
-			MenuBuilder.AddMenuEntry(
-				NSLOCTEXT("Sequencer", "SetKeyInterpolationUser", "User"),
-				NSLOCTEXT("Sequencer", "SetKeyInterpolationUserTooltip", "Set key interpolation to user"),
-				FSlateIcon(),
-				FUIAction(
-					FExecuteAction::CreateSP(&Sequencer, &FSequencer::SetInterpTangentMode, RCIM_Cubic, RCTM_User),
-					FCanExecuteAction(),
-					FIsActionChecked::CreateSP(&Sequencer, &FSequencer::IsInterpTangentModeSelected, RCIM_Cubic, RCTM_User) ),
-				NAME_None,
-				EUserInterfaceActionType::ToggleButton
-			);
-
-			MenuBuilder.AddMenuEntry(
-				NSLOCTEXT("Sequencer", "SetKeyInterpolationBreak", "Break"),
-				NSLOCTEXT("Sequencer", "SetKeyInterpolationBreakTooltip", "Set key interpolation to break"),
-				FSlateIcon(),
-				FUIAction(
-					FExecuteAction::CreateSP(&Sequencer, &FSequencer::SetInterpTangentMode, RCIM_Cubic, RCTM_Break),
-					FCanExecuteAction(),
-					FIsActionChecked::CreateSP(&Sequencer, &FSequencer::IsInterpTangentModeSelected, RCIM_Cubic, RCTM_Break) ),
-				NAME_None,
-				EUserInterfaceActionType::ToggleButton
-			);
-
-			MenuBuilder.AddMenuEntry(
-				NSLOCTEXT("Sequencer", "SetKeyInterpolationLinear", "Linear"),
-				NSLOCTEXT("Sequencer", "SetKeyInterpolationLinearTooltip", "Set key interpolation to linear"),
-				FSlateIcon(),
-				FUIAction(
-					FExecuteAction::CreateSP(&Sequencer, &FSequencer::SetInterpTangentMode, RCIM_Linear, RCTM_Auto),
-					FCanExecuteAction(),
-					FIsActionChecked::CreateSP(&Sequencer, &FSequencer::IsInterpTangentModeSelected, RCIM_Linear, RCTM_Auto) ),
-				NAME_None,
-				EUserInterfaceActionType::ToggleButton
-			);
-
-			MenuBuilder.AddMenuEntry(
-				NSLOCTEXT("Sequencer", "SetKeyInterpolationConstant", "Constant"),
-				NSLOCTEXT("Sequencer", "SetKeyInterpolationConstantTooltip", "Set key interpolation to constant"),
-				FSlateIcon(),
-				FUIAction(
-					FExecuteAction::CreateSP(&Sequencer, &FSequencer::SetInterpTangentMode, RCIM_Constant, RCTM_Auto),
-					FCanExecuteAction(),
-					FIsActionChecked::CreateSP(&Sequencer, &FSequencer::IsInterpTangentModeSelected, RCIM_Constant, RCTM_Auto) ),
-				NAME_None,
-				EUserInterfaceActionType::ToggleButton
-			);
-		}
-		MenuBuilder.EndSection(); // SequencerInterpolation
-
-		MenuBuilder.BeginSection("SequencerKeys", NSLOCTEXT("Sequencer", "KeysMenu", "Keys"));
-		{
-			const bool bUseFrames = SequencerSnapValues::IsTimeSnapIntervalFrameRate(Sequencer.GetSettings()->GetTimeSnapInterval());
-
-			MenuBuilder.AddMenuEntry(
-				bUseFrames ? NSLOCTEXT("Sequencer", "SetKeyFrame", "Set Key Frame") : NSLOCTEXT("Sequencer", "SetKeyTime", "Set Key Time"),
-				bUseFrames ? NSLOCTEXT("Sequencer", "SetKeyFrameTooltip", "Set key frame") : NSLOCTEXT("Sequencer", "SetKeyTimeTooltip", "Set key time"),
-				FSlateIcon(),
-				FUIAction(
-					FExecuteAction::CreateSP(&Sequencer, &FSequencer::SetKeyTime, bUseFrames),
-					FCanExecuteAction::CreateSP(&Sequencer, &FSequencer::CanSetKeyTime))
-			);
-
-			MenuBuilder.AddMenuEntry(
-				NSLOCTEXT("Sequencer", "SnapToFrame", "Snap to Frame"),
-				NSLOCTEXT("Sequencer", "SnapToFrameToolTip", "Snap selected keys to frame"),
-				FSlateIcon(),
-				FUIAction(
-					FExecuteAction::CreateSP(&Sequencer, &FSequencer::SnapToFrame),
-					FCanExecuteAction::CreateSP(&Sequencer, &FSequencer::CanSnapToFrame))
-			);
-
-			MenuBuilder.AddMenuEntry(
-				NSLOCTEXT("Sequencer", "DeleteKey", "Delete"),
-				NSLOCTEXT("Sequencer", "DeleteKeyToolTip", "Deletes the selected keys"),
-				FSlateIcon(),
-				FUIAction(FExecuteAction::CreateSP(&Sequencer, &FSequencer::DeleteSelectedKeys))
-			);
-		}
-		MenuBuilder.EndSection(); // SequencerKeys
-	}
-	
-	if (SceneSection != nullptr && Sequencer.GetSelection().IsSelected(SceneSection))
-	{
-		bHasMenuItems = true;
-		SectionInterface->BuildSectionContextMenu(MenuBuilder);
-
-		MenuBuilder.BeginSection("SequencerSections", NSLOCTEXT("Sequencer", "SectionsMenu", "Sections"));
-		{
-			MenuBuilder.AddMenuEntry(
-				NSLOCTEXT("Sequencer", "SelectAllKeys", "Select All Keys"),
-				NSLOCTEXT("Sequencer", "SelectAllKeysTooltip", "Select all keys in section"),
-				FSlateIcon(),
-				FUIAction(
-					FExecuteAction::CreateSP(this, &SSequencerSection::SelectAllKeys),
-					FCanExecuteAction::CreateSP(this, &SSequencerSection::CanSelectAllKeys))
-			);
-
-			MenuBuilder.AddSubMenu(
-				NSLOCTEXT("Sequencer", "EditSection", "Edit"),
-				NSLOCTEXT("Sequencer", "EditSectionTooltip", "Edit section"),
-				FNewMenuDelegate::CreateRaw(this, &SSequencerSection::AddEditMenu));
-
-			MenuBuilder.AddSubMenu(
-				NSLOCTEXT("Sequencer", "SetPreInfinityExtrap", "Pre-Infinity"), 
-				NSLOCTEXT("Sequencer", "SetPreInfinityExtrapTooltip", "Set pre-infinity extrapolation"),
-				FNewMenuDelegate::CreateRaw(this, &SSequencerSection::AddExtrapolationMenu, true));
-
-			MenuBuilder.AddSubMenu(
-				NSLOCTEXT("Sequencer", "SetPostInfinityExtrap", "Post-Infinity"), 
-				NSLOCTEXT("Sequencer", "SetPostInfinityExtrapTooltip", "Set post-infinity extrapolation"),
-				FNewMenuDelegate::CreateRaw(this, &SSequencerSection::AddExtrapolationMenu, false));
-
-			MenuBuilder.AddMenuEntry(
-				NSLOCTEXT("Sequencer", "ToggleSectionActive", "Active"),
-				NSLOCTEXT("Sequencer", "ToggleSectionActiveTooltip", "Toggle section active/inactive"),
-				FSlateIcon(),
-				FUIAction(
-					FExecuteAction::CreateSP(this, &SSequencerSection::ToggleSectionActive),
-					FCanExecuteAction(),
-					FIsActionChecked::CreateSP(this, &SSequencerSection::IsToggleSectionActive)),
-				NAME_None,
-				EUserInterfaceActionType::ToggleButton
-			);
-
-			// @todo Sequencer this should delete all selected sections
-			// delete/selection needs to be rethought in general
-			MenuBuilder.AddMenuEntry(
-				NSLOCTEXT("Sequencer", "DeleteSection", "Delete"),
-				NSLOCTEXT("Sequencer", "DeleteSectionToolTip", "Deletes this section"),
-				FSlateIcon(),
-				FUIAction(FExecuteAction::CreateSP(this, &SSequencerSection::DeleteSection))
-			);
-		}
-		MenuBuilder.EndSection(); // SequencerSections
-	}
-	
-	ResetState();
-
-	if (bHasMenuItems)
-	{
-		return MenuBuilder.MakeWidget();
-	}
-
-	return TSharedPtr<SWidget>();
-}
-
-
-void SSequencerSection::AddEditMenu(FMenuBuilder& MenuBuilder)
-{
-	MenuBuilder.AddMenuEntry(
-		NSLOCTEXT("Sequencer", "TrimSectionLeft", "Trim Left"),
-		NSLOCTEXT("Sequencer", "TrimSectionLeftTooltip", "Trim section at current time to the left"),
-		FSlateIcon(),
-		FUIAction(
-			FExecuteAction::CreateSP(this, &SSequencerSection::TrimSection, true),
-			FCanExecuteAction::CreateSP(this, &SSequencerSection::IsTrimmable))
-	);
-
-	MenuBuilder.AddMenuEntry(
-		NSLOCTEXT("Sequencer", "TrimSectionRight", "Trim Right"),
-		NSLOCTEXT("Sequencer", "TrimSectionRightTooltip", "Trim section at current time to the right"),
-		FSlateIcon(),
-		FUIAction(
-			FExecuteAction::CreateSP(this, &SSequencerSection::TrimSection, false),
-			FCanExecuteAction::CreateSP(this, &SSequencerSection::IsTrimmable))
-	);
-
-	MenuBuilder.AddMenuEntry(
-		NSLOCTEXT("Sequencer", "SplitSection", "Split"),
-		NSLOCTEXT("Sequencer", "SplitSectionTooltip", "Split section at current time"),
-		FSlateIcon(),
-		FUIAction(
-			FExecuteAction::CreateSP(this, &SSequencerSection::SplitSection),
-			FCanExecuteAction::CreateSP(this, &SSequencerSection::IsTrimmable))
-	);
-}
-
-
-void SSequencerSection::AddExtrapolationMenu(FMenuBuilder& MenuBuilder, bool bPreInfinity)
-{
-	FSequencer& Sequencer = GetSequencer();
-
-	MenuBuilder.AddMenuEntry(
-		NSLOCTEXT("Sequencer", "SetExtrapCycle", "Cycle"),
-		NSLOCTEXT("Sequencer", "SetExtrapCycleTooltip", "Set extrapolation cycle"),
-		FSlateIcon(),
-		FUIAction(
-			FExecuteAction::CreateSP(this, &SSequencerSection::SetExtrapolationMode, RCCE_Cycle, bPreInfinity),
-			FCanExecuteAction(),
-			FIsActionChecked::CreateSP(this, &SSequencerSection::IsExtrapolationModeSelected, RCCE_Cycle, bPreInfinity) ),
-		NAME_None,
-		EUserInterfaceActionType::RadioButton
-	);
-
-	MenuBuilder.AddMenuEntry(
-		NSLOCTEXT("Sequencer", "SetExtrapCycleWithOffset", "Cycle with Offset"),
-		NSLOCTEXT("Sequencer", "SetExtrapCycleWithOffsetTooltip", "Set extrapolation cycle with offset"),
-		FSlateIcon(),
-		FUIAction(
-			FExecuteAction::CreateSP(this, &SSequencerSection::SetExtrapolationMode, RCCE_CycleWithOffset, bPreInfinity),
-			FCanExecuteAction(),
-			FIsActionChecked::CreateSP(this, &SSequencerSection::IsExtrapolationModeSelected, RCCE_CycleWithOffset, bPreInfinity) ),
-		NAME_None,
-		EUserInterfaceActionType::RadioButton
-	);
-
-	MenuBuilder.AddMenuEntry(
-		NSLOCTEXT("Sequencer", "SetExtrapOscillate", "Oscillate"),
-		NSLOCTEXT("Sequencer", "SetExtrapOscillateTooltip", "Set extrapolation oscillate"),
-		FSlateIcon(),
-		FUIAction(
-			FExecuteAction::CreateSP(this, &SSequencerSection::SetExtrapolationMode, RCCE_Oscillate, bPreInfinity),
-			FCanExecuteAction(),
-			FIsActionChecked::CreateSP(this, &SSequencerSection::IsExtrapolationModeSelected, RCCE_Oscillate, bPreInfinity) ),
-		NAME_None,
-		EUserInterfaceActionType::RadioButton
-	);
-
-	MenuBuilder.AddMenuEntry(
-		NSLOCTEXT("Sequencer", "SetExtrapLinear", "Linear"),
-		NSLOCTEXT("Sequencer", "SetExtrapLinearTooltip", "Set extrapolation linear"),
-		FSlateIcon(),
-		FUIAction(
-			FExecuteAction::CreateSP(this, &SSequencerSection::SetExtrapolationMode, RCCE_Linear, bPreInfinity),
-			FCanExecuteAction(),
-			FIsActionChecked::CreateSP(this, &SSequencerSection::IsExtrapolationModeSelected, RCCE_Linear, bPreInfinity) ),
-		NAME_None,
-		EUserInterfaceActionType::RadioButton
-	);
-
-	MenuBuilder.AddMenuEntry(
-		NSLOCTEXT("Sequencer", "SetExtrapConstant", "Constant"),
-		NSLOCTEXT("Sequencer", "SetExtrapConstantTooltip", "Set extrapolation constant"),
-		FSlateIcon(),
-		FUIAction(
-			FExecuteAction::CreateSP(this, &SSequencerSection::SetExtrapolationMode, RCCE_Constant, bPreInfinity),
-			FCanExecuteAction(),
-			FIsActionChecked::CreateSP(this, &SSequencerSection::IsExtrapolationModeSelected, RCCE_Constant, bPreInfinity) ),
-		NAME_None,
-		EUserInterfaceActionType::RadioButton
-	);
-}
-
 void SSequencerSection::Tick( const FGeometry& AllottedGeometry, const double InCurrentTime, const float InDeltaTime )
 {	
 	if( GetVisibility() == EVisibility::Visible )
@@ -804,83 +624,40 @@ void SSequencerSection::Tick( const FGeometry& AllottedGeometry, const double In
 	}
 }
 
-
 FReply SSequencerSection::OnMouseButtonDown( const FGeometry& MyGeometry, const FPointerEvent& MouseEvent )
 {
 	FSequencer& Sequencer = GetSequencer();
 
-	// Check for clicking on a key and mark it as the pressed key for drag detection (if necessary) later
-	PressedKey = GetKeyUnderMouse( MouseEvent.GetScreenSpacePosition(), MyGeometry );
-
-	if (MouseEvent.GetEffectingButton() == EKeys::LeftMouseButton)
+	FSequencerSelectedKey HoveredKey;
+	
+	// The hovered key is defined from the sequencer hotspot
+	TSharedPtr<ISequencerHotspot> Hotspot = GetSequencer().GetEditTool().GetHotspot();
+	if (Hotspot.IsValid() && Hotspot->GetType() == ESequencerHotspot::Key)
 	{
-		if( !PressedKey.IsValid() )
-		{
-			CheckForEdgeInteraction( MouseEvent, MyGeometry );
-		}
+		HoveredKey = static_cast<FKeyHotspot*>(Hotspot.Get())->Key;
 	}
-	
+
 	if (MouseEvent.GetEffectingButton() == EKeys::MiddleMouseButton)
 	{
 		// Generate a key and set it as the PressedKey
-		PressedKey = CreateKeyUnderMouse(MouseEvent.GetScreenSpacePosition(), MyGeometry, PressedKey);
-		HoveredKey = PressedKey;
-
-		Sequencer.GetSelection().EmptySelectedKeys();
-		Sequencer.GetSelection().AddToSelection(PressedKey);
-
-		// Pass the event to the tool to copy the hovered key and move it
-		auto& EditTool = GetSequencer().GetEditTool();
-
-		EditTool.SetHotspot( MakeShareable( new FSectionHotspot(FSectionHandle(ParentSectionArea, SectionIndex))) );
+		FSequencerSelectedKey NewKey = CreateKeyUnderMouse(MouseEvent.GetScreenSpacePosition(), MyGeometry, HoveredKey);
 
-		if ( HoveredKey.IsValid() )
+		if (NewKey.IsValid())
 		{
-			EditTool.SetHotspot( MakeShareable( new FKeyHotspot(HoveredKey, ParentSectionArea) ) );
-			FReply Reply = EditTool.OnMouseButtonDown(*this, MyGeometry, MouseEvent);
-			if (Reply.IsEventHandled())
-			{
-				return Reply;
-			}
-		}
+			Sequencer.GetSelection().EmptySelectedKeys();
+			Sequencer.GetSelection().AddToSelection(NewKey);
 
-		return FReply::Unhandled();
-	}
-
-	if ((MouseEvent.GetEffectingButton() == EKeys::LeftMouseButton) || (MouseEvent.GetEffectingButton() == EKeys::RightMouseButton))
-	{
-		HandleSelection(MyGeometry, MouseEvent);
-	}
-	
-	if (MouseEvent.GetEffectingButton() == EKeys::RightMouseButton && (PressedKey.IsValid() || Sequencer.GetSelection().GetSelectedKeys().Num() || Sequencer.GetSelection().GetSelectedSections().Num()))
-	{
-		// Don't return handled if we didn't click on a key so that right-click panning gets a look in
-		return FReply::Handled();
+			// Pass the event to the tool to copy the hovered key and move it
+			auto& EditTool = GetSequencer().GetEditTool();
+			EditTool.SetHotspot( MakeShareable( new FKeyHotspot(NewKey, ParentSectionArea) ) );
+			return FReply::Handled();
+		}
 	}
 
 	return FReply::Unhandled();
 }
 
 
-void SSequencerSection::ResetState()
-{
-	ResetHoveredState();
-	PressedKey = FSequencerSelectedKey();
-}
-
-
-void SSequencerSection::ResetHoveredState()
-{
-	bLeftEdgeHovered = false;
-	bRightEdgeHovered = false;
-	bLeftEdgePressed = false;
-	bRightEdgePressed = false;
-	HoveredKey = FSequencerSelectedKey();
-
-	GetSequencer().GetEditTool().SetHotspot(nullptr);
-}
-
-
 FGeometry SSequencerSection::MakeSectionGeometryWithoutHandles( const FGeometry& AllottedGeometry, const TSharedPtr<ISequencerSection>& InSectionInterface ) const
 {
 	const bool bSectionsAreConnected = InSectionInterface->AreSectionsConnected();
@@ -890,51 +667,8 @@ FGeometry SSequencerSection::MakeSectionGeometryWithoutHandles( const FGeometry&
 }
 
 
-FReply SSequencerSection::OnMouseButtonUp( const FGeometry& MyGeometry, const FPointerEvent& MouseEvent )
-{
-	if( ( MouseEvent.GetEffectingButton() == EKeys::LeftMouseButton || MouseEvent.GetEffectingButton() == EKeys::RightMouseButton ) && MyGeometry.IsUnderLocation( MouseEvent.GetScreenSpacePosition() ) )
-	{
-		// Snap time to the key under the mouse if shift is down
-		if (MouseEvent.IsShiftDown() && MouseEvent.GetEffectingButton() == EKeys::LeftMouseButton)
-		{
-			FSequencerSelectedKey Key = GetKeyUnderMouse( MouseEvent.GetScreenSpacePosition(), MyGeometry );
-			if (Key.IsValid())
-			{
-				float KeyTime = Key.KeyArea->GetKeyTime(Key.KeyHandle.GetValue());
-				GetSequencer().SetGlobalTime(KeyTime);
-			}
-		}
-	}
-
-	if( MouseEvent.GetEffectingButton() == EKeys::RightMouseButton )
-	{
-		TSharedPtr<SWidget> MenuContent = OnSummonContextMenu( MyGeometry, MouseEvent );
-		if (MenuContent.IsValid())
-		{
-			FWidgetPath WidgetPath = MouseEvent.GetEventPath() != nullptr ? *MouseEvent.GetEventPath() : FWidgetPath();
-
-			FSlateApplication::Get().PushMenu(
-				AsShared(),
-				WidgetPath,
-				MenuContent.ToSharedRef(),
-				MouseEvent.GetScreenSpacePosition(),
-				FPopupTransitionEffect( FPopupTransitionEffect::ContextMenu )
-				);
-
-			return FReply::Handled().SetUserFocus(MenuContent.ToSharedRef(), EFocusCause::SetDirectly);
-		}
-	}
-
-	ResetState();
-
-	return FReply::Handled();
-}
-
-
 FReply SSequencerSection::OnMouseButtonDoubleClick( const FGeometry& MyGeometry, const FPointerEvent& MouseEvent )
 {
-	ResetState();
-
 	if( MouseEvent.GetEffectingButton() == EKeys::LeftMouseButton )
 	{
 		FReply Reply = SectionInterface->OnSectionDoubleClicked( MyGeometry, MouseEvent );
@@ -955,22 +689,18 @@ FReply SSequencerSection::OnMouseButtonDoubleClick( const FGeometry& MyGeometry,
 
 FReply SSequencerSection::OnMouseMove( const FGeometry& MyGeometry, const FPointerEvent& MouseEvent )
 {
-	ResetHoveredState();
-
-	// Checked for hovered key
-	// @todo Sequencer - Needs visual cue
-	HoveredKey = GetKeyUnderMouse( MouseEvent.GetScreenSpacePosition(), MyGeometry );
-
 	auto& EditTool = GetSequencer().GetEditTool();
 
-	EditTool.SetHotspot( MakeShareable( new FSectionHotspot(FSectionHandle(ParentSectionArea, SectionIndex))) );
-
-	if ( HoveredKey.IsValid() )
+	// Checked for hovered key
+	FSequencerSelectedKey KeyUnderMouse = GetKeyUnderMouse( MouseEvent.GetScreenSpacePosition(), MyGeometry );
+	if ( KeyUnderMouse.IsValid() )
 	{
-		EditTool.SetHotspot( MakeShareable( new FKeyHotspot(HoveredKey, ParentSectionArea) ) );
+		EditTool.SetHotspot( MakeShareable( new FKeyHotspot(KeyUnderMouse, ParentSectionArea) ) );
 	}
 	else
 	{
+		EditTool.SetHotspot( MakeShareable( new FSectionHotspot(FSectionHandle(ParentSectionArea, SectionIndex))) );
+
 		// Only check for edge interaction if not hovering over a key
 		CheckForEdgeInteraction( MouseEvent, MyGeometry );
 	}
@@ -982,298 +712,5 @@ FReply SSequencerSection::OnMouseMove( const FGeometry& MyGeometry, const FPoint
 void SSequencerSection::OnMouseLeave( const FPointerEvent& MouseEvent )
 {
 	SCompoundWidget::OnMouseLeave( MouseEvent );
-	ResetHoveredState();
-}
-
-
-FCursorReply SSequencerSection::OnCursorQuery( const FGeometry& MyGeometry, const FPointerEvent& CursorEvent ) const
-{
-	if( bLeftEdgeHovered || bRightEdgeHovered )
-	{
-		return FCursorReply::Cursor( EMouseCursor::ResizeLeftRight );
-	}
-
-	return FCursorReply::Unhandled();
-}
-
-
-void SSequencerSection::HandleSelection( const FGeometry& MyGeometry, const FPointerEvent& MouseEvent )
-{
-	FSequencerSelectedKey Key = GetKeyUnderMouse( MouseEvent.GetScreenSpacePosition(), MyGeometry );
-
-	// If we have a currently pressed down and the key under the mouse is the pressed key, select that key now
-	if( PressedKey.IsValid() && Key.IsValid() && PressedKey == Key )
-	{
-		bool bSelectDueToDrag = false;
-		HandleKeySelection( Key, MouseEvent, bSelectDueToDrag );
-	}
-	else if ((MouseEvent.GetEffectingButton() == EKeys::LeftMouseButton) || (MouseEvent.GetEffectingButton() == EKeys::RightMouseButton))
-	{
-		bool bSelectDueToDrag = false;
-		HandleSectionSelection( MouseEvent, bSelectDueToDrag );
-	}
-}
-
-
-void SSequencerSection::HandleKeySelection( const FSequencerSelectedKey& Key, const FPointerEvent& MouseEvent, bool bSelectDueToDrag )
-{
-	if( Key.IsValid() )
-	{
-		bool bKeyIsSelected = GetSequencer().GetSelection().IsSelected( Key );
-
-		// Clear previous key selection if:
-		// we are selecting due to drag and the key being dragged is not selected
-		bool bShouldClearSelectionDueToDrag =  bSelectDueToDrag ? !bKeyIsSelected : true;
-
-		// Toggle selection if control is down and key is selected
-		if (MouseEvent.IsControlDown() && bKeyIsSelected)
-		{
-			GetSequencer().GetSelection().RemoveFromSelection( Key );
-		}
-		else
-		{
-			if( (!MouseEvent.IsShiftDown() && !MouseEvent.IsControlDown() && bShouldClearSelectionDueToDrag) && !bKeyIsSelected )
-			{
-				GetSequencer().GetSelection().EmptySelectedKeys();
-			}
-
-			GetSequencer().GetSelection().AddToSelection( Key );
-		}
-	}
-}
-
-
-void SSequencerSection::HandleSectionSelection( const FPointerEvent& MouseEvent, bool bSelectDueToDrag )
-{
-	// handle selecting sections 
-	UMovieSceneSection* Section = SectionInterface->GetSectionObject();
-
-	bool bSectionIsSelected = GetSequencer().GetSelection().IsSelected(Section);
-
-	// Clear previous section selection if:
-	// we are selecting due to drag and the section being dragged is not selected
-	bool bShouldClearSelectionDueToDrag =  bSelectDueToDrag ? !bSectionIsSelected : true;
-
-	// Keep key selection if right clicking to bring up a menu and the current key is selected
-	bool bKeepSectionSelection = MouseEvent.GetEffectingButton() == EKeys::RightMouseButton && bSectionIsSelected;
-	
-	// Don't clear selection edge is being dragged
-	bool bEdgePressed = bLeftEdgePressed || bRightEdgePressed;
-
-	if (MouseEvent.IsControlDown() && bSectionIsSelected && !bEdgePressed)
-	{
-		GetSequencer().GetSelection().RemoveFromSelection(Section);
-	}
-	else
-	{
-		if ( (!MouseEvent.IsShiftDown() && !MouseEvent.IsControlDown() && bShouldClearSelectionDueToDrag) && !bKeepSectionSelection)
-		{
-			GetSequencer().GetSelection().EmptySelectedSections();
-		}
-	}
-	
-	if (!bSectionIsSelected)
-	{
-		GetSequencer().GetSelection().AddToSelection(Section);
-	}
-	else if (!bKeepSectionSelection)
-	{
-		GetSequencer().GetSelection().RemoveFromSelection(Section);
-	}
-}
-
-
-void SSequencerSection::SelectAllKeys()
-{
-	// @todo Sequencer should operate on selected sections
-	UMovieSceneSection* Section = SectionInterface->GetSectionObject();
-
-	for (const FKeyAreaLayoutElement& Element : Layout->GetElements())
-	{
-		TSharedPtr<IKeyArea> KeyArea = Element.GetKeyArea();
-
-		TArray<FKeyHandle> KeyHandles = KeyArea->GetUnsortedKeyHandles();
-		for( int32 KeyIndex = 0; KeyIndex < KeyHandles.Num(); ++KeyIndex )
-		{
-			FKeyHandle KeyHandle = KeyHandles[KeyIndex];
-			FSequencerSelectedKey SelectKey(*Section, KeyArea, KeyHandle);
-			GetSequencer().GetSelection().AddToSelection(SelectKey);
-		}
-	}
-}
-
-
-bool SSequencerSection::CanSelectAllKeys() const
-{
-	// @todo Sequencer should operate on selected sections
-	UMovieSceneSection* Section = SectionInterface->GetSectionObject();
-
-	for (const FKeyAreaLayoutElement& Element : Layout->GetElements())
-	{
-		TSharedPtr<IKeyArea> KeyArea = Element.GetKeyArea();
-
-		TArray<FKeyHandle> KeyHandles = KeyArea->GetUnsortedKeyHandles();
-		if (KeyHandles.Num() > 0)
-		{
-			return true;
-		}
-	}
-
-	return false;
-}
-
-
-void SSequencerSection::TrimSection(bool bTrimLeft)
-{
-	FScopedTransaction TrimSectionTransaction(NSLOCTEXT("Sequencer", "TrimSection_Transaction", "Trim Section"));
-
-	MovieSceneToolHelpers::TrimSection(GetSequencer().GetSelection().GetSelectedSections(), GetSequencer().GetGlobalTime(), bTrimLeft);
-	GetSequencer().NotifyMovieSceneDataChanged();
-}
-
-
-void SSequencerSection::SplitSection()
-{
-	FScopedTransaction SplitSectionTransaction(NSLOCTEXT("Sequencer", "SplitSection_Transaction", "Split Section"));
-
-	MovieSceneToolHelpers::SplitSection(GetSequencer().GetSelection().GetSelectedSections(), GetSequencer().GetGlobalTime());
-	GetSequencer().NotifyMovieSceneDataChanged();
-}
-
-bool SSequencerSection::IsTrimmable() const
-{
-	for (auto Section : GetSequencer().GetSelection().GetSelectedSections())
-	{
-		if (Section.IsValid() && Section->IsTimeWithinSection(GetSequencer().GetGlobalTime()))
-		{
-			return true;
-		}
-	}
-	return false;
-}
-
-
-void SSequencerSection::SetExtrapolationMode(ERichCurveExtrapolation ExtrapMode, bool bPreInfinity)
-{
-	// @todo Sequencer should operate on selected sections
-	UMovieSceneSection* Section = SectionInterface->GetSectionObject();
-	FScopedTransaction SetExtrapolationModeTransaction(NSLOCTEXT("Sequencer", "SetExtrapolationMode_Transaction", "Set Extrapolation Mode"));
-	bool bAnythingChanged = false;
-
-	Section->Modify();
-	for (const FKeyAreaLayoutElement& Element : Layout->GetElements())
-	{
-		TSharedPtr<IKeyArea> KeyArea = Element.GetKeyArea();
-		bAnythingChanged = true;
-		KeyArea->SetExtrapolationMode(ExtrapMode, bPreInfinity);
-	}
-
-	if (bAnythingChanged)
-	{
-		GetSequencer().UpdateRuntimeInstances();
-	}
-}
-
-
-bool SSequencerSection::IsExtrapolationModeSelected(ERichCurveExtrapolation ExtrapMode, bool bPreInfinity) const
-{
-	// @todo Sequencer should operate on selected sections
-	bool bAllSelected = false;
-
-	for (const FKeyAreaLayoutElement& Element : Layout->GetElements())
-	{
-		TSharedPtr<IKeyArea> KeyArea = Element.GetKeyArea();
-
-		bAllSelected = true;
-		if (KeyArea->GetExtrapolationMode(bPreInfinity) != ExtrapMode)
-		{
-			bAllSelected = false;
-			break;
-		}
-	}
-
-	return bAllSelected;
-}
-
-
-void SSequencerSection::ToggleSectionActive()
-{
-	FScopedTransaction ToggleSectionActiveTransaction( NSLOCTEXT("Sequencer", "ToggleSectionActive_Transaction", "Toggle Section Active") );
-
-	bool bIsActive = !IsToggleSectionActive();
-	bool bAnythingChanged = false;
-
-	if (GetSequencer().GetSelection().GetSelectedSections().Num() != 0)
-	{
-		for (auto Section : GetSequencer().GetSelection().GetSelectedSections())
-		{
-			bAnythingChanged = true;
-			Section->Modify();
-			Section->SetIsActive(bIsActive);
-		}
-	}
-	else
-	{
-		UMovieSceneSection* Section = SectionInterface->GetSectionObject();
-		if (Section)
-		{
-			bAnythingChanged = true;
-			Section->Modify();
-			Section->SetIsActive(bIsActive);
-		}
-	}
-
-	if (bAnythingChanged)
-	{
-		GetSequencer().UpdateRuntimeInstances();
-	}
-}
-
-
-bool SSequencerSection::IsToggleSectionActive() const
-{
-	// Active only if all are active
-	if (GetSequencer().GetSelection().GetSelectedSections().Num() != 0)
-	{
-		for (auto Section : GetSequencer().GetSelection().GetSelectedSections())
-		{
-			if (!Section->IsActive())
-			{
-				return false;
-			}
-		}
-	}
-	else
-	{
-		UMovieSceneSection* Section = SectionInterface->GetSectionObject();
-		if (Section)
-		{
-			if (!Section->IsActive())
-			{
-				return false;
-			}
-		}
-	}
-
-	return true;
-}
-
-
-void SSequencerSection::DeleteSection()
-{
-	if (GetSequencer().GetSelection().GetSelectedSections().Num() != 0)
-	{
-		GetSequencer().DeleteSections(GetSequencer().GetSelection().GetSelectedSections());
-	}
-	else
-	{
-		UMovieSceneSection* Section = SectionInterface->GetSectionObject();
-		if (Section)
-		{
-			TSet<TWeakObjectPtr<UMovieSceneSection> > Sections;
-			Sections.Add(Section);
-
-			GetSequencer().DeleteSections(Sections);
-		}
-	}
-}
+	GetSequencer().GetEditTool().SetHotspot(nullptr);
+}
\ No newline at end of file
