Copying //UE4/Orion-Staging to Dev-Main (Originates from //Orion/Main @ CL-2759376)

diff --git a/Engine/Source/Editor/Sequencer/Private/SSequencerSection.h b/Engine/Source/Editor/Sequencer/Private/SSequencerSection.h
--- a/Engine/Source/Editor/Sequencer/Private/SSequencerSection.h
+++ b/Engine/Source/Editor/Sequencer/Private/SSequencerSection.h
@@ -69,11 +69,9 @@ private:
 	virtual int32 OnPaint( const FPaintArgs& Args, const FGeometry& AllottedGeometry, const FSlateRect& MyClippingRect, FSlateWindowElementList& OutDrawElements, int32 LayerId, const FWidgetStyle& InWidgetStyle, bool bParentEnabled ) const override;
 	virtual void Tick( const FGeometry& AllottedGeometry, const double InCurrentTime, const float InDeltaTime ) override;
 	virtual FReply OnMouseButtonDown( const FGeometry& MyGeometry, const FPointerEvent& MouseEvent ) override;
-	virtual FReply OnMouseButtonUp( const FGeometry& MyGeometry, const FPointerEvent& MouseEvent ) override;
 	virtual FReply OnMouseButtonDoubleClick( const FGeometry& MyGeometry, const FPointerEvent& MouseEvent ) override;
 	virtual FReply OnMouseMove( const FGeometry& MyGeometry, const FPointerEvent& MouseEvent ) override;
 	virtual void OnMouseLeave( const FPointerEvent& MouseEvent ) override;
-	virtual FCursorReply OnCursorQuery( const FGeometry& MyGeometry, const FPointerEvent& CursorEvent ) const override;
 
 	/**
 	 * Paints keys visible inside the section
@@ -95,113 +93,14 @@ private:
 	 */
 	void DrawSelectionBorder( const FGeometry& AllottedGeometry, const FSlateRect& MyClippingRect, FSlateWindowElementList& OutDrawElements, int32 LayerId, ESlateDrawEffect::Type DrawEffects, FLinearColor SelectionColor ) const;
 
-	/** Summons a context menu over the associated section */
-	TSharedPtr<SWidget> OnSummonContextMenu( const FGeometry& MyGeometry, const FPointerEvent& MouseEvent );
-
-	/** Add edit menu for trim and split */
-	void AddEditMenu(FMenuBuilder& MenuBuilder);
-
-	/** Add extrapolation menu for pre and post infinity */
-	void AddExtrapolationMenu(FMenuBuilder& MenuBuilder, bool bPreInfinity);
-
-	/**
-	 * Selects the section or keys if they are found under the mouse
-	 *
-	 * @param MyGeometry Geometry of the section
-	 * @param MouseEvent	The mouse event that may cause selection
-	 */
-	void HandleSelection( const FGeometry& MyGeometry, const FPointerEvent& MouseEvent );
-
-	/**
-	 * Selects a key if needed
-	 *
-	 * @param Key				The key to select
-	 * @param MouseEvent		The mouse event that may have caused selection
-	 * @param bSelectDueToDrag	If we are selecting the key due to a drag operation rather than it just being clicked
-	 */
-	void HandleKeySelection( const FSequencerSelectedKey& Key, const FPointerEvent& MouseEvent, bool bSelectDueToDrag );
-
-	/**
-	 * Selects the section if needed
-	 *
-	 * @param MouseEvent		The mouse event that may have caused selection
-	 * @param bSelectDueToDrag	If we are selecting the key due to a drag operation rather than it just being clicked
-	 */
-	void HandleSectionSelection( const FPointerEvent& MouseEvent, bool bSelectDueToDrag );
-
 	/** @return the sequencer interface */
 	FSequencer& GetSequencer() const;
 
-	/**
-	 * Creates a drag operation based on the current mouse event and what is selected
-	 *
-	 * @param MyGeometry		Geometry of the section
-	 * @param MouseEvent		The mouse event that caused the drag
-	 * @param bKeysUnderMouse	true if there are any keys under the mouse
-	 */
-	void CreateDragOperation( const FGeometry& MyGeometry, const FPointerEvent& MouseEvent, bool bKeysUnderMouse );
-
-	/**
-	 * Resets entire internal state of the widget
-	 */
-	void ResetState();
-
-	/**
-	 * Resets the hovered state of elements in the widget (section edges, keys)
-	 */
-	void ResetHoveredState();
-
 	/** 
 	 * Creates geometry for a section without space for the handles
 	 */
 	FGeometry MakeSectionGeometryWithoutHandles( const FGeometry& AllottedGeometry, const TSharedPtr<ISequencerSection>& InSectionInterface ) const;
 
-	/** 
-	 * Select all keys within the section
-	 */
-	void SelectAllKeys();
-
-	/**
-	 * Are there keys to select? 
-	 */
-	bool CanSelectAllKeys() const;
-
-	/**
-	 * Trim section 
-	 */
-	void TrimSection(bool bTrimLeft);
-
-	/**
-	 * Split section
-	 */
-	void SplitSection();
-
-	/**
-	 * Is there a trimmable section?
-	 */
-	bool IsTrimmable() const;
-
-	/**
-	 * Set extrapolation mode
-	 */
-	void SetExtrapolationMode(ERichCurveExtrapolation ExtrapMode, bool bPreInfinity);
-
-	/**
-	 * Is extrapolation mode selected
-	 */
-	bool IsExtrapolationModeSelected(ERichCurveExtrapolation ExtrapMode, bool bPreInfinity) const;
-
-	/**
-	 * Section active/inactive toggle
-	 */
-	void ToggleSectionActive();
-	bool IsToggleSectionActive() const;
-
-	/*
-	 * Delete section
-	 */ 
-	void DeleteSection();
-
 private:
 	/** Interface to section data */
 	TSharedPtr<ISequencerSection> SectionInterface;
@@ -209,20 +108,8 @@ private:
 	TSharedPtr<FSequencerTrackNode> ParentSectionArea;
 	/** Cached layout generated each tick */
 	TOptional<FKeyAreaLayout> Layout;
-	/** The currently pressed key if any */
-	FSequencerSelectedKey PressedKey;
-	/** The current hovered key if any */
-	FSequencerSelectedKey HoveredKey;
 	/** The index of this section in the parent section area */
 	int32 SectionIndex;
-	/** Whether or not the left edge of the section is hovered */
-	bool bLeftEdgeHovered;
-	/** Whether or not the left edge of the section is pressed */
-	bool bLeftEdgePressed;
-	/** Whether or not the right edge of the section is hovered */
-	bool bRightEdgeHovered;
-	/** Whether or not the right edge of the section is pressed */
-	bool bRightEdgePressed;
 	/** Cached parent geometry to pass down to any section interfaces that need it during tick */
 	FGeometry ParentGeometry;
 };
\ No newline at end of file
diff --git a/Engine/Source/Editor/Sequencer/Private/SSequencerTrackArea.cpp b/Engine/Source/Editor/Sequencer/Private/SSequencerTrackArea.cpp
--- a/Engine/Source/Editor/Sequencer/Private/SSequencerTrackArea.cpp
+++ b/Engine/Source/Editor/Sequencer/Private/SSequencerTrackArea.cpp
@@ -12,7 +12,6 @@
 #include "SSequencerSection.h"
 #include "ISequencerTrackEditor.h"
 
-
 FTrackAreaSlot::FTrackAreaSlot(const TSharedPtr<SSequencerTrackLane>& InSlotContent)
 {
 	TrackLane = InSlotContent;
@@ -38,6 +37,13 @@ void SSequencerTrackArea::Construct(const FArguments& InArgs, TSharedRef<FSequen
 {
 	Sequencer = InSequencer;
 	TimeSliderController = InTimeSliderController;
+
+	// Input stack in order or priority
+
+	// Space for the edit tool
+	InputStack.AddHandler(nullptr);
+	// The time slider controller
+	InputStack.AddHandler(TimeSliderController.Get());
 }
 
 
@@ -143,70 +149,45 @@ int32 SSequencerTrackArea::OnPaint(const FPaintArgs& Args, const FGeometry& Allo
 
 FReply SSequencerTrackArea::OnMouseButtonDown( const FGeometry& MyGeometry, const FPointerEvent& MouseEvent )
 {
-	FReply Reply = Sequencer->GetEditTool().OnMouseButtonDown(*this, MyGeometry, MouseEvent);
-
-	if (Reply.IsEventHandled())
-	{
-		return Reply;
-	}
-
-	return TimeSliderController->OnMouseButtonDown( SharedThis(this), MyGeometry, MouseEvent );
+	// Always ensure the edit tool is set up
+	InputStack.SetHandlerAt(0, &Sequencer->GetEditTool());
+	return InputStack.HandleMouseButtonDown(*this, MyGeometry, MouseEvent);
 }
 
 
 FReply SSequencerTrackArea::OnMouseButtonUp( const FGeometry& MyGeometry, const FPointerEvent& MouseEvent )
 {
-	FReply Reply = Sequencer->GetEditTool().OnMouseButtonUp(*this, MyGeometry, MouseEvent);
-
-	if (Reply.IsEventHandled())
-	{
-		return Reply;
-	}
-
-	return TimeSliderController->OnMouseButtonUp( SharedThis(this), MyGeometry, MouseEvent );
+	// Always ensure the edit tool is set up
+	InputStack.SetHandlerAt(0, &Sequencer->GetEditTool());
+	return InputStack.HandleMouseButtonUp(*this, MyGeometry, MouseEvent);
 }
 
 
 FReply SSequencerTrackArea::OnMouseMove( const FGeometry& MyGeometry, const FPointerEvent& MouseEvent )
 {
-	FReply Reply = Sequencer->GetEditTool().OnMouseMove(*this, MyGeometry, MouseEvent);
+	// Always ensure the edit tool is set up
+	InputStack.SetHandlerAt(0, &Sequencer->GetEditTool());
 
-	if (Reply.IsEventHandled())
-	{
-		return Reply;
-	}
+	FReply Reply = InputStack.HandleMouseMove(*this, MyGeometry, MouseEvent);
 
-	if (MouseEvent.IsMouseButtonDown(EKeys::RightMouseButton) && HasMouseCapture())
+	// Handle right click scrolling on the track area, if the captured index is that of the time slider
+	if (Reply.IsEventHandled() && InputStack.GetCapturedIndex() == 1)
 	{
-		TreeView.Pin()->ScrollByDelta(-MouseEvent.GetCursorDelta().Y);
+		if (MouseEvent.IsMouseButtonDown(EKeys::RightMouseButton) && HasMouseCapture())
+		{
+			TreeView.Pin()->ScrollByDelta(-MouseEvent.GetCursorDelta().Y);
+		}
 	}
 
-	return TimeSliderController->OnMouseMove( SharedThis(this), MyGeometry, MouseEvent );
+	return Reply;
 }
 
 
 FReply SSequencerTrackArea::OnMouseWheel( const FGeometry& MyGeometry, const FPointerEvent& MouseEvent )
 {
-	// First try the edit tool
-	FReply Reply = Sequencer->GetEditTool().OnMouseWheel(*this, MyGeometry, MouseEvent);
-
-	if (Reply.IsEventHandled())
-	{
-		return Reply;
-	}
-
-	// Then the time slider
-	Reply = TimeSliderController->OnMouseWheel(SharedThis(this), MyGeometry, MouseEvent);
-
-	if (Reply.IsEventHandled())
-	{
-		return Reply;
-	}
-
-	// Failing that, we'll just scroll vertically
-	TreeView.Pin()->ScrollByDelta(WheelScrollAmount * -MouseEvent.GetWheelDelta());
-
-	return FReply::Handled();
+	// Always ensure the edit tool is set up
+	InputStack.SetHandlerAt(0, &Sequencer->GetEditTool());
+	return InputStack.HandleMouseWheel(*this, MyGeometry, MouseEvent);
 }
 
 
diff --git a/Engine/Source/Editor/Sequencer/Private/SSequencerTrackArea.h b/Engine/Source/Editor/Sequencer/Private/SSequencerTrackArea.h
--- a/Engine/Source/Editor/Sequencer/Private/SSequencerTrackArea.h
+++ b/Engine/Source/Editor/Sequencer/Private/SSequencerTrackArea.h
@@ -8,6 +8,7 @@ class SSequencerTrackLane;
 class FSequencerTimeSliderController;
 class FSequencer;
 
+#include "SequencerInputHandlerStack.h"
 
 /**
  * Structure representing a slot in the track area.
@@ -106,4 +107,9 @@ private:
 
 	/** Keep a hold of the size of the area so we can maintain zoom levels. */
 	TOptional<FVector2D> SizeLastFrame;
+
+private:
+
+	/** Input handler stack responsible for routing input to the different handlers */
+	FSequencerInputHandlerStack InputStack;
 };
diff --git a/Engine/Source/Editor/Sequencer/Private/Sequencer.cpp b/Engine/Source/Editor/Sequencer/Private/Sequencer.cpp
--- a/Engine/Source/Editor/Sequencer/Private/Sequencer.cpp
+++ b/Engine/Source/Editor/Sequencer/Private/Sequencer.cpp
@@ -49,6 +49,9 @@
 #include "MovieSceneCaptureDialogModule.h"
 #include "AutomatedLevelSequenceCapture.h"
 
+#include "SceneOutlinerModule.h"
+#include "SceneOutlinerPublicTypes.h"
+
 #define LOCTEXT_NAMESPACE "Sequencer"
 
 DEFINE_LOG_CATEGORY(LogSequencer);
@@ -115,6 +118,8 @@ void FSequencer::InitSequencer(const FSequencerInitParams& InitParams, const TSh
 			.ClampRange( this, &FSequencer::GetClampRange )
 			.PlaybackRange( this, &FSequencer::GetPlaybackRange )
 			.OnPlaybackRangeChanged( this, &FSequencer::SetPlaybackRange )
+			.OnBeginPlaybackRangeDrag( this, &FSequencer::OnBeginPlaybackRangeDrag )
+			.OnEndPlaybackRangeDrag( this, &FSequencer::OnEndPlaybackRangeDrag )
 			.ScrubPosition( this, &FSequencer::OnGetScrubPosition )
 			.OnBeginScrubbing( this, &FSequencer::OnBeginScrubbing )
 			.OnEndScrubbing( this, &FSequencer::OnEndScrubbing )
@@ -574,6 +579,8 @@ void FSequencer::SetPlaybackRange(TRange<float> InNewPlaybackRange)
 {
 	if (ensure(InNewPlaybackRange.HasLowerBound() && InNewPlaybackRange.HasUpperBound() && !InNewPlaybackRange.IsDegenerate()))
 	{
+		const FScopedTransaction Transaction(LOCTEXT("SetPlaybackRange_Transaction", "Set playback range"));
+
 		UMovieScene* FocussedMovieScene = GetFocusedMovieSceneSequence()->GetMovieScene();
 		FocussedMovieScene->SetPlaybackRange(InNewPlaybackRange.GetLowerBoundValue(), InNewPlaybackRange.GetUpperBoundValue());
 	}
@@ -978,14 +985,10 @@ FReply FSequencer::OnPlay(bool bTogglePlay)
 	}
 	else
 	{
-		TRange<float> TimeBounds = GetTimeBounds();
-		if (!TimeBounds.IsEmpty())
-		{
-			PlaybackState = EMovieScenePlayerStatus::Playing;
-			
-			// Make sure Slate ticks during playback
-			SequencerWidget->RegisterActiveTimerForPlayback();
-		}
+		PlaybackState = EMovieScenePlayerStatus::Playing;
+		
+		// Make sure Slate ticks during playback
+		SequencerWidget->RegisterActiveTimerForPlayback();
 	}
 
 	return FReply::Handled();
@@ -1039,22 +1042,14 @@ FReply FSequencer::OnStepBackward()
 FReply FSequencer::OnStepToEnd()
 {
 	PlaybackState = EMovieScenePlayerStatus::Stopped;
-	TRange<float> TimeBounds = GetTimeBounds();
-	if (!TimeBounds.IsEmpty())
-	{
-		SetGlobalTime(TimeBounds.GetUpperBoundValue());
-	}
+	SetGlobalTime(GetPlaybackRange().GetUpperBoundValue());
 	return FReply::Handled();
 }
 
 FReply FSequencer::OnStepToBeginning()
 {
 	PlaybackState = EMovieScenePlayerStatus::Stopped;
-	TRange<float> TimeBounds = GetTimeBounds();
-	if (!TimeBounds.IsEmpty())
-	{
-		SetGlobalTime(TimeBounds.GetLowerBoundValue());
-	}
+	SetGlobalTime(GetPlaybackRange().GetLowerBoundValue());
 	return FReply::Handled();
 }
 
@@ -1084,16 +1079,20 @@ void FSequencer::SetGlobalTimeLooped(float InTime)
 	}
 	else
 	{
+
 		TRange<float> TimeBounds = GetTimeBounds();
-		if (InTime > TimeBounds.GetUpperBoundValue())
+		TRange<float> WorkingRange = GetClampRange();
+
+		// Stop if we hit the playback range end
+		if (GetGlobalTime() < TimeBounds.GetUpperBoundValue() && InTime >= TimeBounds.GetUpperBoundValue())
 		{
 			InTime = TimeBounds.GetUpperBoundValue();
 			PlaybackState = EMovieScenePlayerStatus::Stopped;
 		}
-		else if (InTime < TimeBounds.GetLowerBoundValue())
+		// Jump to the start of the working range if necessary
+		else if (InTime < WorkingRange.GetLowerBoundValue())
 		{
-			InTime = TimeBounds.GetLowerBoundValue();
-			PlaybackState = EMovieScenePlayerStatus::Stopped;
+			InTime = WorkingRange.GetLowerBoundValue();
 		}
 	}
 
@@ -1153,7 +1152,7 @@ TRange<float> FSequencer::GetTimeBounds() const
 		return TRange<float>( -100000.0f, 100000.0f );
 	}
 	
-	return FocusedSequence->GetMovieScene()->GetEditorData().WorkingRange;
+	return FocusedSequence->GetMovieScene()->GetPlaybackRange();
 }
 
 void FSequencer::SetViewRange(TRange<float> NewViewRange, EViewRangeInterpolation Interpolation)
@@ -1207,6 +1206,8 @@ void FSequencer::OnClampRangeChanged( TRange<float> NewClampRange )
 
 void FSequencer::OnScrubPositionChanged( float NewScrubPosition, bool bScrubbing )
 {
+	bool bClampToViewRange = true;
+
 	if (PlaybackState == EMovieScenePlayerStatus::Scrubbing)
 	{
 		if (!bScrubbing)
@@ -1215,6 +1216,9 @@ void FSequencer::OnScrubPositionChanged( float NewScrubPosition, bool bScrubbing
 		}
 		else if (IsAutoScrollEnabled())
 		{
+			// Clamp to the view range when not auto-scrolling
+			bClampToViewRange = false;
+	
 			UpdateAutoScroll(NewScrubPosition);
 			
 			// When scrubbing, we animate auto-scrolled scrub position in Tick()
@@ -1223,11 +1227,11 @@ void FSequencer::OnScrubPositionChanged( float NewScrubPosition, bool bScrubbing
 				return;
 			}
 		}
-		else
-		{
-			// Clamp to the view range when not auto-scrolling
-			NewScrubPosition = FMath::Clamp(NewScrubPosition, TargetViewRange.GetLowerBoundValue(), TargetViewRange.GetUpperBoundValue());
-		}
+	}
+
+	if (bClampToViewRange)
+	{
+		NewScrubPosition = FMath::Clamp(NewScrubPosition, TargetViewRange.GetLowerBoundValue(), TargetViewRange.GetUpperBoundValue());
 	}
 
 	SetGlobalTimeDirectly( NewScrubPosition );
@@ -1246,6 +1250,16 @@ void FSequencer::OnEndScrubbing()
 	StopAutoscroll();
 }
 
+void FSequencer::OnBeginPlaybackRangeDrag()
+{
+	GEditor->BeginTransaction(LOCTEXT("SetPlaybackRange_Transaction", "Set playback range"));
+}
+
+void FSequencer::OnEndPlaybackRangeDrag()
+{
+	GEditor->EndTransaction();
+}
+
 void FSequencer::StartAutoscroll(float UnitsPerS)
 {
 	AutoscrollOffset = UnitsPerS;
@@ -1909,12 +1923,59 @@ void FSequencer::DeleteSelectedItems()
 	}
 }
 
-void FSequencer::AssignActor(FGuid InObjectBinding, FSequencerObjectBindingNode* ObjectBindingNode)
+void FSequencer::AssignActor(FMenuBuilder& MenuBuilder, FGuid InObjectBinding)
 {
-	if (ObjectBindingNode != nullptr)
+	UMovieSceneSequence* OwnerSequence = GetFocusedMovieSceneSequence();
+	UObject* RuntimeObject = OwnerSequence->FindObject(InObjectBinding);
+
+	auto IsActorValidForAssignment = [=](const AActor* InActor, UObject* CurrentObject){
+		return CurrentObject != InActor;
+	};
+
+	using namespace SceneOutliner;
+
+	// Set up a menu entry to assign an actor to the object binding node
+	FInitializationOptions InitOptions;
 	{
-		// Get the first selected actor as the actor to replace with
-		AActor* Actor = GEditor->GetSelectedActors()->GetTop<AActor>();
+		InitOptions.Mode = ESceneOutlinerMode::ActorPicker;
+
+		// We hide the header row to keep the UI compact.
+		InitOptions.bShowHeaderRow = false;
+		InitOptions.bShowSearchBox = true;
+		InitOptions.bShowCreateNewFolder = false;
+		// Only want the actor label column
+		InitOptions.ColumnMap.Add(FBuiltInColumnTypes::Label(), FColumnInfo(EColumnVisibility::Visible, 0));
+
+		// Only display actors that are not possessed already
+		InitOptions.Filters->AddFilterPredicate( FActorFilterPredicate::CreateLambda( IsActorValidForAssignment, RuntimeObject ) );
+	}
+
+	// actor selector to allow the user to choose an actor
+	FSceneOutlinerModule& SceneOutlinerModule = FModuleManager::LoadModuleChecked<FSceneOutlinerModule>("SceneOutliner");
+	TSharedRef< SWidget > MiniSceneOutliner =
+		SNew( SVerticalBox )
+		+SVerticalBox::Slot()
+		.AutoHeight()
+		.MaxHeight(400.0f)
+		[
+			SceneOutlinerModule.CreateSceneOutliner(
+				InitOptions,
+				FOnActorPicked::CreateLambda([=](AActor* Actor){
+					// Create a new binding for this actor
+					FSlateApplication::Get().DismissAllMenus();
+					DoAssignActor(&Actor, 1, InObjectBinding);
+				})
+			)
+		];
+	MenuBuilder.AddWidget(MiniSceneOutliner, FText::GetEmpty(), true);
+	MenuBuilder.EndSection();
+}
+
+void FSequencer::DoAssignActor(AActor*const* InActors, int32 NumActors, FGuid InObjectBinding)
+{
+	if (NumActors > 0)
+	{
+		AActor* Actor = InActors[0];
 		if (Actor != nullptr)
 		{
 			FScopedTransaction AssignActor( NSLOCTEXT("Sequencer", "AssignActor", "Assign Actor") );
@@ -1984,11 +2045,6 @@ void FSequencer::AssignActor(FGuid InObjectBinding, FSequencerObjectBindingNode*
 	}
 }
 
-bool FSequencer::CanAssignActor(FGuid ObjectBinding) const
-{
-	return GEditor->GetSelectedActors()->Num() > 0;
-}
-
 void FSequencer::DeleteNode(TSharedRef<FSequencerDisplayNode> NodeToBeDeleted)
 {
 	// If this node is selected, delete all selected nodes
@@ -2433,32 +2489,28 @@ void FSequencer::BindSequencerCommands()
 			UAutomatedLevelSequenceCapture* MovieSceneCapture = NewObject<UAutomatedLevelSequenceCapture>(GetTransientPackage(), UAutomatedLevelSequenceCapture::StaticClass(), NAME_None, RF_Transient);
 			MovieSceneCapture->LoadConfig();
 
-			bool bFoundActor = false;
-
 			// Attempt to find a level sequence actor in the world that references this asset
 			for (auto It = TActorIterator<ALevelSequenceActor>(GWorld); It; ++It)
 			{
 				if (It->LevelSequence == GetCurrentAsset())
 				{
 					MovieSceneCapture->SetLevelSequenceActor(*It);
-					bFoundActor = true;
+					break;
 				}
 			}
 
-			if (!bFoundActor)
-			{
-				MovieSceneCapture->SetLevelSequenceAsset(GetCurrentAsset()->GetPathName());
-			}
+			MovieSceneCapture->SetLevelSequenceAsset(GetCurrentAsset()->GetPathName());
 
 			if (CanShowFrameNumbers())
 			{
 				MovieSceneCapture->Settings.FrameRate = FMath::RoundToInt(1.f / Settings->GetTimeSnapInterval());
-				// We always add 1 to the number of frames we want to capture, because we want to capture both the start and end frames (which if the play range is 0, would still yield a single frame)
-				MovieSceneCapture->Settings.FrameCount = FMath::RoundToInt((GetPlaybackRange().GetUpperBoundValue() - GetPlaybackRange().GetLowerBoundValue()) * MovieSceneCapture->Settings.FrameRate) + 1;
-			}
-			else
-			{
-				MovieSceneCapture->Settings.FrameCount = 0;
+				MovieSceneCapture->Settings.bUseRelativeFrameNumbers = false;
+
+				const int32 SequenceStartFrame = FMath::RoundToInt( GetPlaybackRange().GetLowerBoundValue() * MovieSceneCapture->Settings.FrameRate );
+				const int32 SequenceEndFrame = FMath::Max( SequenceStartFrame, FMath::RoundToInt( GetPlaybackRange().GetUpperBoundValue() * MovieSceneCapture->Settings.FrameRate ) );
+
+				MovieSceneCapture->Settings.StartFrame = SequenceStartFrame;
+				MovieSceneCapture->Settings.EndFrame = SequenceEndFrame;
 			}
 
 			IMovieSceneCaptureDialogModule::Get().OpenDialog(LevelEditorModule.GetLevelEditorTabManager().ToSharedRef(), MovieSceneCapture);
diff --git a/Engine/Source/Editor/Sequencer/Private/Sequencer.h b/Engine/Source/Editor/Sequencer/Private/Sequencer.h
--- a/Engine/Source/Editor/Sequencer/Private/Sequencer.h
+++ b/Engine/Source/Editor/Sequencer/Private/Sequencer.h
@@ -286,8 +286,8 @@ public:
 	void AddSelectedObjects();
 
 	/** Called when a user executes the assign actor to track menu item */
-	void AssignActor(FGuid ObjectBinding, FSequencerObjectBindingNode* ObjectBindingNode);
-	bool CanAssignActor(FGuid ObjectBinding) const;
+	void AssignActor(FMenuBuilder& MenuBuilder, FGuid ObjectBinding);
+	void DoAssignActor(AActor*const* InActors, int32 NumActors, FGuid ObjectBinding);
 
 	/** Called when a user executes the delete node menu item */
 	void DeleteNode(TSharedRef<FSequencerDisplayNode> NodeToBeDeleted);
@@ -358,6 +358,12 @@ protected:
 	/** Called when the user has finished scrubbing */
 	void OnEndScrubbing();
 
+	/** Called when the user has begun dragging the play range */
+	void OnBeginPlaybackRangeDrag();
+
+	/** Called when the user has finished dragging the play range */
+	void OnEndPlaybackRangeDrag();
+
 public:
 
 	/** Put the sequencer in a horizontally autoscrolling state with the given rate */
diff --git a/Engine/Source/Editor/Sequencer/Private/SequencerContextMenus.cpp b/Engine/Source/Editor/Sequencer/Private/SequencerContextMenus.cpp
new file mode 100644
--- /dev/null
+++ b/Engine/Source/Editor/Sequencer/Private/SequencerContextMenus.cpp
@@ -0,0 +1,459 @@
+// Copyright 1998-2015 Epic Games, Inc. All Rights Reserved.
+
+#include "SequencerPrivatePCH.h"
+#include "SequencerContextMenus.h"
+#include "SequencerHotspots.h"
+#include "ScopedTransaction.h"
+#include "MovieSceneToolHelpers.h"
+#include "MovieSceneSection.h"
+
+void FKeyContextMenu::BuildMenu(FMenuBuilder& MenuBuilder, FSequencer& InSequencer)
+{
+	TSharedRef<FKeyContextMenu> Menu = MakeShareable(new FKeyContextMenu(InSequencer));
+	Menu->PopulateMenu(MenuBuilder);
+}
+
+void FKeyContextMenu::PopulateMenu(FMenuBuilder& MenuBuilder)
+{
+	FSequencer* SequencerPtr = &Sequencer.Get();
+	MenuBuilder.BeginSection("SequencerInterpolation", NSLOCTEXT("Sequencer", "KeyInterpolationMenu", "Key Interpolation"));
+	{
+		MenuBuilder.AddMenuEntry(
+			NSLOCTEXT("Sequencer", "SetKeyInterpolationAuto", "Auto"),
+			NSLOCTEXT("Sequencer", "SetKeyInterpolationAutoTooltip", "Set key interpolation to auto"),
+			FSlateIcon(),
+			FUIAction(
+				FExecuteAction::CreateSP(SequencerPtr, &FSequencer::SetInterpTangentMode, RCIM_Cubic, RCTM_Auto),
+				FCanExecuteAction(),
+				FIsActionChecked::CreateSP(SequencerPtr, &FSequencer::IsInterpTangentModeSelected, RCIM_Cubic, RCTM_Auto) ),
+			NAME_None,
+			EUserInterfaceActionType::ToggleButton
+		);
+
+		MenuBuilder.AddMenuEntry(
+			NSLOCTEXT("Sequencer", "SetKeyInterpolationUser", "User"),
+			NSLOCTEXT("Sequencer", "SetKeyInterpolationUserTooltip", "Set key interpolation to user"),
+			FSlateIcon(),
+			FUIAction(
+				FExecuteAction::CreateSP(SequencerPtr, &FSequencer::SetInterpTangentMode, RCIM_Cubic, RCTM_User),
+				FCanExecuteAction(),
+				FIsActionChecked::CreateSP(SequencerPtr, &FSequencer::IsInterpTangentModeSelected, RCIM_Cubic, RCTM_User) ),
+			NAME_None,
+			EUserInterfaceActionType::ToggleButton
+		);
+
+		MenuBuilder.AddMenuEntry(
+			NSLOCTEXT("Sequencer", "SetKeyInterpolationBreak", "Break"),
+			NSLOCTEXT("Sequencer", "SetKeyInterpolationBreakTooltip", "Set key interpolation to break"),
+			FSlateIcon(),
+			FUIAction(
+				FExecuteAction::CreateSP(SequencerPtr, &FSequencer::SetInterpTangentMode, RCIM_Cubic, RCTM_Break),
+				FCanExecuteAction(),
+				FIsActionChecked::CreateSP(SequencerPtr, &FSequencer::IsInterpTangentModeSelected, RCIM_Cubic, RCTM_Break) ),
+			NAME_None,
+			EUserInterfaceActionType::ToggleButton
+		);
+
+		MenuBuilder.AddMenuEntry(
+			NSLOCTEXT("Sequencer", "SetKeyInterpolationLinear", "Linear"),
+			NSLOCTEXT("Sequencer", "SetKeyInterpolationLinearTooltip", "Set key interpolation to linear"),
+			FSlateIcon(),
+			FUIAction(
+				FExecuteAction::CreateSP(SequencerPtr, &FSequencer::SetInterpTangentMode, RCIM_Linear, RCTM_Auto),
+				FCanExecuteAction(),
+				FIsActionChecked::CreateSP(SequencerPtr, &FSequencer::IsInterpTangentModeSelected, RCIM_Linear, RCTM_Auto) ),
+			NAME_None,
+			EUserInterfaceActionType::ToggleButton
+		);
+
+		MenuBuilder.AddMenuEntry(
+			NSLOCTEXT("Sequencer", "SetKeyInterpolationConstant", "Constant"),
+			NSLOCTEXT("Sequencer", "SetKeyInterpolationConstantTooltip", "Set key interpolation to constant"),
+			FSlateIcon(),
+			FUIAction(
+				FExecuteAction::CreateSP(SequencerPtr, &FSequencer::SetInterpTangentMode, RCIM_Constant, RCTM_Auto),
+				FCanExecuteAction(),
+				FIsActionChecked::CreateSP(SequencerPtr, &FSequencer::IsInterpTangentModeSelected, RCIM_Constant, RCTM_Auto) ),
+			NAME_None,
+			EUserInterfaceActionType::ToggleButton
+		);
+	}
+	MenuBuilder.EndSection(); // SequencerInterpolation
+
+	MenuBuilder.BeginSection("SequencerKeys", NSLOCTEXT("Sequencer", "KeysMenu", "Keys"));
+	{
+		const bool bUseFrames = SequencerSnapValues::IsTimeSnapIntervalFrameRate(Sequencer->GetSettings()->GetTimeSnapInterval());
+
+		MenuBuilder.AddMenuEntry(
+			bUseFrames ? NSLOCTEXT("Sequencer", "SetKeyFrame", "Set Key Frame") : NSLOCTEXT("Sequencer", "SetKeyTime", "Set Key Time"),
+			bUseFrames ? NSLOCTEXT("Sequencer", "SetKeyFrameTooltip", "Set key frame") : NSLOCTEXT("Sequencer", "SetKeyTimeTooltip", "Set key time"),
+			FSlateIcon(),
+			FUIAction(
+				FExecuteAction::CreateSP(SequencerPtr, &FSequencer::SetKeyTime, bUseFrames),
+				FCanExecuteAction::CreateSP(SequencerPtr, &FSequencer::CanSetKeyTime))
+		);
+
+		MenuBuilder.AddMenuEntry(
+			NSLOCTEXT("Sequencer", "SnapToFrame", "Snap to Frame"),
+			NSLOCTEXT("Sequencer", "SnapToFrameToolTip", "Snap selected keys to frame"),
+			FSlateIcon(),
+			FUIAction(
+				FExecuteAction::CreateSP(SequencerPtr, &FSequencer::SnapToFrame),
+				FCanExecuteAction::CreateSP(SequencerPtr, &FSequencer::CanSnapToFrame))
+		);
+
+		MenuBuilder.AddMenuEntry(
+			NSLOCTEXT("Sequencer", "DeleteKey", "Delete"),
+			NSLOCTEXT("Sequencer", "DeleteKeyToolTip", "Deletes the selected keys"),
+			FSlateIcon(),
+			FUIAction(FExecuteAction::CreateSP(SequencerPtr, &FSequencer::DeleteSelectedKeys))
+		);
+	}
+	MenuBuilder.EndSection(); // SequencerKeys
+}
+
+void FSectionContextMenu::BuildMenu(FMenuBuilder& MenuBuilder, FSequencer& InSequencer)
+{
+	TSharedRef<FSectionContextMenu> Menu = MakeShareable(new FSectionContextMenu(InSequencer));
+	Menu->PopulateMenu(MenuBuilder);
+}
+
+void FSectionContextMenu::PopulateMenu(FMenuBuilder& MenuBuilder)
+{
+	// Copy a reference to the context menu by value into each lambda handler to ensure the type stays alive until the menu is closed
+	TSharedRef<FSectionContextMenu> Shared = AsShared();
+
+	MenuBuilder.BeginSection("SequencerSections", NSLOCTEXT("Sequencer", "SectionsMenu", "Sections"));
+	{
+		bool bCanSelectAll = CanSelectAllKeys();
+		MenuBuilder.AddMenuEntry(
+			NSLOCTEXT("Sequencer", "SelectAllKeys", "Select All Keys"),
+			NSLOCTEXT("Sequencer", "SelectAllKeysTooltip", "Select all keys in section"),
+			FSlateIcon(),
+			FUIAction(
+				FExecuteAction::CreateLambda([=]{ return Shared->SelectAllKeys(); }),
+				FCanExecuteAction::CreateLambda([=]{ return bCanSelectAll; }))
+		);
+
+		MenuBuilder.AddSubMenu(
+			NSLOCTEXT("Sequencer", "EditSection", "Edit"),
+			NSLOCTEXT("Sequencer", "EditSectionTooltip", "Edit section"),
+			FNewMenuDelegate::CreateLambda([=](FMenuBuilder& MenuBuilder){ Shared->AddEditMenu(MenuBuilder); }));
+
+		MenuBuilder.AddSubMenu(
+			NSLOCTEXT("Sequencer", "SetPreInfinityExtrap", "Pre-Infinity"), 
+			NSLOCTEXT("Sequencer", "SetPreInfinityExtrapTooltip", "Set pre-infinity extrapolation"),
+			FNewMenuDelegate::CreateLambda([=](FMenuBuilder& MenuBuilder){ Shared->AddExtrapolationMenu(MenuBuilder, true); }));
+
+		MenuBuilder.AddSubMenu(
+			NSLOCTEXT("Sequencer", "SetPostInfinityExtrap", "Post-Infinity"), 
+			NSLOCTEXT("Sequencer", "SetPostInfinityExtrapTooltip", "Set post-infinity extrapolation"),
+			FNewMenuDelegate::CreateLambda([=](FMenuBuilder& MenuBuilder){ Shared->AddExtrapolationMenu(MenuBuilder, false); }));
+
+		MenuBuilder.AddMenuEntry(
+			NSLOCTEXT("Sequencer", "ToggleSectionActive", "Active"),
+			NSLOCTEXT("Sequencer", "ToggleSectionActiveTooltip", "Toggle section active/inactive"),
+			FSlateIcon(),
+			FUIAction(
+				FExecuteAction::CreateLambda([=]{ Shared->ToggleSectionActive(); }),
+				FCanExecuteAction(),
+				FIsActionChecked::CreateLambda([=]{ return Shared->IsToggleSectionActive(); })),
+			NAME_None,
+			EUserInterfaceActionType::ToggleButton
+		);
+
+		// @todo Sequencer this should delete all selected sections
+		// delete/selection needs to be rethought in general
+		MenuBuilder.AddMenuEntry(
+			NSLOCTEXT("Sequencer", "DeleteSection", "Delete"),
+			NSLOCTEXT("Sequencer", "DeleteSectionToolTip", "Deletes this section"),
+			FSlateIcon(),
+			FUIAction(FExecuteAction::CreateLambda([=]{ return Shared->DeleteSection(); }))
+		);
+	}
+	MenuBuilder.EndSection(); // SequencerSections
+}
+
+void FSectionContextMenu::AddEditMenu(FMenuBuilder& MenuBuilder)
+{
+	// Copy a reference to the context menu by value into each lambda handler to ensure the type stays alive until the menu is closed
+	TSharedRef<FSectionContextMenu> Shared = AsShared();
+
+	MenuBuilder.AddMenuEntry(
+		NSLOCTEXT("Sequencer", "TrimSectionLeft", "Trim Left"),
+		NSLOCTEXT("Sequencer", "TrimSectionLeftTooltip", "Trim section at current time to the left"),
+		FSlateIcon(),
+		FUIAction(
+			FExecuteAction::CreateLambda([=]{ Shared->TrimSection(true); }),
+			FCanExecuteAction::CreateLambda([=]{ return Shared->IsTrimmable(); }))
+	);
+
+	MenuBuilder.AddMenuEntry(
+		NSLOCTEXT("Sequencer", "TrimSectionRight", "Trim Right"),
+		NSLOCTEXT("Sequencer", "TrimSectionRightTooltip", "Trim section at current time to the right"),
+		FSlateIcon(),
+		FUIAction(
+			FExecuteAction::CreateLambda([=]{ Shared->TrimSection(false); }),
+			FCanExecuteAction::CreateLambda([=]{ return Shared->IsTrimmable(); }))
+	);
+
+	MenuBuilder.AddMenuEntry(
+		NSLOCTEXT("Sequencer", "SplitSection", "Split"),
+		NSLOCTEXT("Sequencer", "SplitSectionTooltip", "Split section at current time"),
+		FSlateIcon(),
+		FUIAction(
+			FExecuteAction::CreateLambda([=]{ Shared->SplitSection(); }),
+			FCanExecuteAction::CreateLambda([=]{ return Shared->IsTrimmable(); }))
+	);
+}
+
+void FSectionContextMenu::AddExtrapolationMenu(FMenuBuilder& MenuBuilder, bool bPreInfinity)
+{
+	// Copy a reference to the context menu by value into each lambda handler to ensure the type stays alive until the menu is closed
+	TSharedRef<FSectionContextMenu> Shared = AsShared();
+
+	MenuBuilder.AddMenuEntry(
+		NSLOCTEXT("Sequencer", "SetExtrapCycle", "Cycle"),
+		NSLOCTEXT("Sequencer", "SetExtrapCycleTooltip", "Set extrapolation cycle"),
+		FSlateIcon(),
+		FUIAction(
+			FExecuteAction::CreateLambda([=]{ Shared->SetExtrapolationMode(RCCE_Cycle, bPreInfinity); }),
+			FCanExecuteAction(),
+			FIsActionChecked::CreateLambda([=]{ return Shared->IsExtrapolationModeSelected(RCCE_Cycle, bPreInfinity); })
+		),
+		NAME_None,
+		EUserInterfaceActionType::RadioButton
+	);
+
+	MenuBuilder.AddMenuEntry(
+		NSLOCTEXT("Sequencer", "SetExtrapCycleWithOffset", "Cycle with Offset"),
+		NSLOCTEXT("Sequencer", "SetExtrapCycleWithOffsetTooltip", "Set extrapolation cycle with offset"),
+		FSlateIcon(),
+		FUIAction(
+			FExecuteAction::CreateLambda([=]{ Shared->SetExtrapolationMode(RCCE_CycleWithOffset, bPreInfinity); }),
+			FCanExecuteAction(),
+			FIsActionChecked::CreateLambda([=]{ return Shared->IsExtrapolationModeSelected(RCCE_CycleWithOffset, bPreInfinity); })
+		),
+		NAME_None,
+		EUserInterfaceActionType::RadioButton
+	);
+
+	MenuBuilder.AddMenuEntry(
+		NSLOCTEXT("Sequencer", "SetExtrapOscillate", "Oscillate"),
+		NSLOCTEXT("Sequencer", "SetExtrapOscillateTooltip", "Set extrapolation oscillate"),
+		FSlateIcon(),
+		FUIAction(
+			FExecuteAction::CreateLambda([=]{ Shared->SetExtrapolationMode(RCCE_Oscillate, bPreInfinity); }),
+			FCanExecuteAction(),
+			FIsActionChecked::CreateLambda([=]{ return Shared->IsExtrapolationModeSelected(RCCE_Oscillate, bPreInfinity); })
+		),
+		NAME_None,
+		EUserInterfaceActionType::RadioButton
+	);
+
+	MenuBuilder.AddMenuEntry(
+		NSLOCTEXT("Sequencer", "SetExtrapLinear", "Linear"),
+		NSLOCTEXT("Sequencer", "SetExtrapLinearTooltip", "Set extrapolation linear"),
+		FSlateIcon(),
+		FUIAction(
+			FExecuteAction::CreateLambda([=]{ Shared->SetExtrapolationMode(RCCE_Linear, bPreInfinity); }),
+			FCanExecuteAction(),
+			FIsActionChecked::CreateLambda([=]{ return Shared->IsExtrapolationModeSelected(RCCE_Linear, bPreInfinity); })
+		),
+		NAME_None,
+		EUserInterfaceActionType::RadioButton
+	);
+
+	MenuBuilder.AddMenuEntry(
+		NSLOCTEXT("Sequencer", "SetExtrapConstant", "Constant"),
+		NSLOCTEXT("Sequencer", "SetExtrapConstantTooltip", "Set extrapolation constant"),
+		FSlateIcon(),
+		FUIAction(
+			FExecuteAction::CreateLambda([=]{ Shared->SetExtrapolationMode(RCCE_Constant, bPreInfinity); }),
+			FCanExecuteAction(),
+			FIsActionChecked::CreateLambda([=]{ return Shared->IsExtrapolationModeSelected(RCCE_Constant, bPreInfinity); })
+		),
+		NAME_None,
+		EUserInterfaceActionType::RadioButton
+	);
+}
+
+void FSectionContextMenu::SelectAllKeys()
+{
+	TArray<FSectionHandle> SelectedSections = StaticCastSharedRef<SSequencer>(Sequencer->GetSequencerWidget())->GetSectionHandles(Sequencer->GetSelection().GetSelectedSections());
+	for (const FSectionHandle& Handle : SelectedSections)
+	{
+		UMovieSceneSection* Section = Handle.GetSectionObject();
+		if (!Section)
+		{
+			continue;
+		}
+
+		FKeyAreaLayout Layout(*Handle.TrackNode, Handle.SectionIndex);
+		for (const FKeyAreaLayoutElement& Element : Layout.GetElements())
+		{
+			TSharedPtr<IKeyArea> KeyArea = Element.GetKeyArea();
+			for (const FKeyHandle& KeyHandle : KeyArea->GetUnsortedKeyHandles())
+			{
+				FSequencerSelectedKey SelectKey(*Section, KeyArea, KeyHandle);
+				Sequencer->GetSelection().AddToSelection(SelectKey);
+			}
+		}
+	}
+}
+
+bool FSectionContextMenu::CanSelectAllKeys() const
+{
+	TArray<FSectionHandle> SelectedSections = StaticCastSharedRef<SSequencer>(Sequencer->GetSequencerWidget())->GetSectionHandles(Sequencer->GetSelection().GetSelectedSections());
+	for (const FSectionHandle& Handle : SelectedSections)
+	{
+		UMovieSceneSection* Section = Handle.GetSectionObject();
+		if (!Section)
+		{
+			continue;
+		}
+
+		FKeyAreaLayout Layout(*Handle.TrackNode, Handle.SectionIndex);
+		for (const FKeyAreaLayoutElement& Element : Layout.GetElements())
+		{
+			TArray<FKeyHandle> KeyHandles = Element.GetKeyArea()->GetUnsortedKeyHandles();
+			if (KeyHandles.Num() > 0)
+			{
+				return true;
+			}
+		}
+	}
+
+	return false;
+}
+
+void FSectionContextMenu::TrimSection(bool bTrimLeft)
+{
+	FScopedTransaction TrimSectionTransaction(NSLOCTEXT("Sequencer", "TrimSection_Transaction", "Trim Section"));
+
+	MovieSceneToolHelpers::TrimSection(Sequencer->GetSelection().GetSelectedSections(), Sequencer->GetGlobalTime(), bTrimLeft);
+	Sequencer->NotifyMovieSceneDataChanged();
+}
+
+void FSectionContextMenu::SplitSection()
+{
+	FScopedTransaction SplitSectionTransaction(NSLOCTEXT("Sequencer", "SplitSection_Transaction", "Split Section"));
+
+	MovieSceneToolHelpers::SplitSection(Sequencer->GetSelection().GetSelectedSections(), Sequencer->GetGlobalTime());
+	Sequencer->NotifyMovieSceneDataChanged();
+}
+
+bool FSectionContextMenu::IsTrimmable() const
+{
+	for (auto Section : Sequencer->GetSelection().GetSelectedSections())
+	{
+		if (Section.IsValid() && Section->IsTimeWithinSection(Sequencer->GetGlobalTime()))
+		{
+			return true;
+		}
+	}
+	return false;
+}
+
+void FSectionContextMenu::SetExtrapolationMode(ERichCurveExtrapolation ExtrapMode, bool bPreInfinity)
+{
+	FScopedTransaction SetExtrapolationModeTransaction(NSLOCTEXT("Sequencer", "SetExtrapolationMode_Transaction", "Set Extrapolation Mode"));
+
+	bool bAnythingChanged = false;
+
+	TArray<FSectionHandle> SelectedSections = StaticCastSharedRef<SSequencer>(Sequencer->GetSequencerWidget())->GetSectionHandles(Sequencer->GetSelection().GetSelectedSections());
+	for (const FSectionHandle& Handle : SelectedSections)
+	{
+		UMovieSceneSection* Section = Handle.GetSectionObject();
+		if (!Section)
+		{
+			continue;
+		}
+
+		Section->Modify();
+
+		FKeyAreaLayout Layout(*Handle.TrackNode, Handle.SectionIndex);
+		for (const FKeyAreaLayoutElement& Element : Layout.GetElements())
+		{
+			TSharedPtr<IKeyArea> KeyArea = Element.GetKeyArea();
+			bAnythingChanged = true;
+			KeyArea->SetExtrapolationMode(ExtrapMode, bPreInfinity);
+		}
+	}
+
+	if (bAnythingChanged)
+	{
+		Sequencer->UpdateRuntimeInstances();
+	}
+	else
+	{
+		SetExtrapolationModeTransaction.Cancel();
+	}
+}
+
+bool FSectionContextMenu::IsExtrapolationModeSelected(ERichCurveExtrapolation ExtrapMode, bool bPreInfinity) const
+{
+	// @todo Sequencer should operate on selected sections
+	bool bAllSelected = false;
+
+	TArray<FSectionHandle> SelectedSections = StaticCastSharedRef<SSequencer>(Sequencer->GetSequencerWidget())->GetSectionHandles(Sequencer->GetSelection().GetSelectedSections());
+	for (const FSectionHandle& Handle : SelectedSections)
+	{
+		FKeyAreaLayout Layout(*Handle.TrackNode, Handle.SectionIndex);
+		for (const FKeyAreaLayoutElement& Element : Layout.GetElements())
+		{
+			TSharedPtr<IKeyArea> KeyArea = Element.GetKeyArea();
+
+			bAllSelected = true;
+			if (KeyArea->GetExtrapolationMode(bPreInfinity) != ExtrapMode)
+			{
+				bAllSelected = false;
+				break;
+			}
+		}
+	}
+
+	return bAllSelected;
+}
+
+void FSectionContextMenu::ToggleSectionActive()
+{
+	FScopedTransaction ToggleSectionActiveTransaction( NSLOCTEXT("Sequencer", "ToggleSectionActive_Transaction", "Toggle Section Active") );
+	bool bIsActive = !IsToggleSectionActive();
+	bool bAnythingChanged = false;
+
+	for (auto Section : Sequencer->GetSelection().GetSelectedSections())
+	{
+		bAnythingChanged = true;
+		Section->Modify();
+		Section->SetIsActive(bIsActive);
+	}
+
+	if (bAnythingChanged)
+	{
+		Sequencer->UpdateRuntimeInstances();
+	}
+	else
+	{
+		ToggleSectionActiveTransaction.Cancel();
+	}
+}
+
+bool FSectionContextMenu::IsToggleSectionActive() const
+{
+	// Active only if all are active
+	for (auto Section : Sequencer->GetSelection().GetSelectedSections())
+	{
+		if (!Section->IsActive())
+		{
+			return false;
+		}
+	}
+
+	return true;
+}
+
+void FSectionContextMenu::DeleteSection()
+{
+	Sequencer->DeleteSections(Sequencer->GetSelection().GetSelectedSections());
+}
\ No newline at end of file
diff --git a/Engine/Source/Editor/Sequencer/Private/SequencerContextMenus.h b/Engine/Source/Editor/Sequencer/Private/SequencerContextMenus.h
new file mode 100644
--- /dev/null
+++ b/Engine/Source/Editor/Sequencer/Private/SequencerContextMenus.h
@@ -0,0 +1,83 @@
+// Copyright 1998-2015 Epic Games, Inc. All Rights Reserved.
+
+#pragma once
+
+#include "Sequencer.h"
+
+/**
+ * Class responsible for generating a menu for the currently selected sections.
+ * This is a shared class that's entirely owned by the context menu handlers.
+ * Once the menu is closed, all references to this class are removed, and the 
+ * instance is cleaned up. To ensure no weak references are held, AsShared is hidden
+ */
+struct FSectionContextMenu : TSharedFromThis<FSectionContextMenu>
+{
+	static void BuildMenu(FMenuBuilder& MenuBuilder, FSequencer& Sequencer);
+
+private:
+
+	/** Hidden AsShared() methods to discourage CreateSP delegate use. */
+	TSharedRef<FSectionContextMenu> AsShared() { return TSharedFromThis::AsShared(); }
+	TSharedRef<const FSectionContextMenu> AsShared() const { return TSharedFromThis::AsShared(); }
+
+	FSectionContextMenu(FSequencer& InSeqeuncer)
+		: Sequencer(StaticCastSharedRef<FSequencer>(InSeqeuncer.AsShared()))
+	{}
+
+	void PopulateMenu(FMenuBuilder& MenuBuilder);
+
+	/** Add edit menu for trim and split */
+	void AddEditMenu(FMenuBuilder& MenuBuilder);
+
+	/** Add extrapolation menu for pre and post infinity */
+	void AddExtrapolationMenu(FMenuBuilder& MenuBuilder, bool bPreInfinity);
+
+	void SelectAllKeys();
+
+	bool CanSelectAllKeys() const;
+
+	void TrimSection(bool bTrimLeft);
+
+	void SplitSection();
+
+	bool IsTrimmable() const;
+
+	void SetExtrapolationMode(ERichCurveExtrapolation ExtrapMode, bool bPreInfinity);
+
+	bool IsExtrapolationModeSelected(ERichCurveExtrapolation ExtrapMode, bool bPreInfinity) const;
+
+	void ToggleSectionActive();
+
+	bool IsToggleSectionActive() const;
+
+	void DeleteSection();
+
+	/** The key itself */
+	TSharedRef<FSequencer> Sequencer;
+};
+
+/**
+ * Class responsible for generating a menu for the currently selected keys.
+ * This is a shared class that's entirely owned by the context menu handlers.
+ * Once the menu is closed, all references to this class are removed, and the 
+ * instance is cleaned up. To ensure no weak references are held, AsShared is hidden
+ */
+struct FKeyContextMenu : TSharedFromThis<FKeyContextMenu>
+{
+	static void BuildMenu(FMenuBuilder& MenuBuilder, FSequencer& Sequencer);
+
+private:
+
+	FKeyContextMenu(FSequencer& InSeqeuncer)
+		: Sequencer(StaticCastSharedRef<FSequencer>(InSeqeuncer.AsShared()))
+	{}
+
+	/** Hidden AsShared() methods to discourage CreateSP delegate use. */
+	TSharedRef<FKeyContextMenu> AsShared() { return TSharedFromThis::AsShared(); }
+	TSharedRef<const FKeyContextMenu> AsShared() const { return TSharedFromThis::AsShared(); }
+
+	void PopulateMenu(FMenuBuilder& MenuBuilder);
+
+	/** The key itself */
+	TSharedRef<FSequencer> Sequencer;
+};
\ No newline at end of file
diff --git a/Engine/Source/Editor/Sequencer/Private/SequencerHotspots.cpp b/Engine/Source/Editor/Sequencer/Private/SequencerHotspots.cpp
--- a/Engine/Source/Editor/Sequencer/Private/SequencerHotspots.cpp
+++ b/Engine/Source/Editor/Sequencer/Private/SequencerHotspots.cpp
@@ -5,6 +5,23 @@
 #include "Tools/EditToolDragOperations.h"
 #include "Sequencer.h"
 #include "MovieSceneSection.h"
+#include "SequencerContextMenus.h"
+
+void FKeyHotspot::PopulateContextMenu(FMenuBuilder& MenuBuilder, ISequencer& InSequencer)
+{
+	FSequencer& Sequencer = static_cast<FSequencer&>(InSequencer);
+	FKeyContextMenu::BuildMenu(MenuBuilder, Sequencer);
+}
+
+void FSectionHotspot::PopulateContextMenu(FMenuBuilder& MenuBuilder, ISequencer& InSequencer)
+{
+	FSequencer& Sequencer = static_cast<FSequencer&>(InSequencer);
+
+	TSharedPtr<ISequencerSection> SectionInterface = Section.TrackNode->GetSections()[Section.SectionIndex];
+	SectionInterface->BuildSectionContextMenu(MenuBuilder);
+
+	FSectionContextMenu::BuildMenu(MenuBuilder, Sequencer);
+}
 
 TSharedPtr<ISequencerEditToolDragOperation> FSectionResizeHotspot::InitiateDrag(ISequencer& Sequencer)
 {
diff --git a/Engine/Source/Editor/Sequencer/Private/SequencerHotspots.h b/Engine/Source/Editor/Sequencer/Private/SequencerHotspots.h
--- a/Engine/Source/Editor/Sequencer/Private/SequencerHotspots.h
+++ b/Engine/Source/Editor/Sequencer/Private/SequencerHotspots.h
@@ -18,6 +18,7 @@ struct FKeyHotspot
 
 	virtual ESequencerHotspot GetType() const override { return ESequencerHotspot::Key; }
 	virtual TSharedPtr<ISequencerEditToolDragOperation> InitiateDrag(ISequencer&) override { return nullptr; }
+	virtual void PopulateContextMenu(FMenuBuilder& MenuBuilder, ISequencer& Sequencer) override;
 
 	/** The key itself */
 	FSequencerSelectedKey Key;
@@ -34,6 +35,11 @@ struct FSectionHandle
 		: SectionIndex(InSectionIndex), TrackNode(MoveTemp(InTrackNode))
 	{ }
 
+	friend bool operator==(const FSectionHandle& A, const FSectionHandle& B)
+	{
+		return A.SectionIndex == B.SectionIndex && A.TrackNode == B.TrackNode;
+	}
+	
 	UMovieSceneSection* GetSectionObject() const { return TrackNode->GetSections()[SectionIndex]->GetSectionObject(); }
 
 	int32 SectionIndex;
@@ -51,6 +57,7 @@ struct FSectionHotspot
 
 	virtual ESequencerHotspot GetType() const override { return ESequencerHotspot::Section; }
 	virtual TSharedPtr<ISequencerEditToolDragOperation> InitiateDrag(ISequencer&) override { return nullptr; }
+	virtual void PopulateContextMenu(FMenuBuilder& MenuBuilder, ISequencer& Sequencer) override;
 
 	/** Handle to the section */
 	FSectionHandle Section;
@@ -71,6 +78,7 @@ struct FSectionResizeHotspot
 
 	virtual ESequencerHotspot GetType() const override { return HandleType == Left ? ESequencerHotspot::SectionResize_L : ESequencerHotspot::SectionResize_R; }
 	virtual TSharedPtr<ISequencerEditToolDragOperation> InitiateDrag(ISequencer& Sequencer) override;
+	virtual FCursorReply GetCursor() const { return FCursorReply::Cursor( EMouseCursor::ResizeLeftRight ); }
 
 	/** Handle to the section */
 	FSectionHandle Section;
diff --git a/Engine/Source/Editor/Sequencer/Private/TimeSliderController.cpp b/Engine/Source/Editor/Sequencer/Private/TimeSliderController.cpp
--- a/Engine/Source/Editor/Sequencer/Private/TimeSliderController.cpp
+++ b/Engine/Source/Editor/Sequencer/Private/TimeSliderController.cpp
@@ -437,7 +437,7 @@ int32 FSequencerTimeSliderController::DrawPlaybackRange(const FGeometry& Allotte
 	return LayerId + 1;
 }
 
-FReply FSequencerTimeSliderController::OnMouseButtonDown( TSharedRef<SWidget> WidgetOwner, const FGeometry& MyGeometry, const FPointerEvent& MouseEvent )
+FReply FSequencerTimeSliderController::OnMouseButtonDown( SWidget& WidgetOwner, const FGeometry& MyGeometry, const FPointerEvent& MouseEvent )
 {
 	bool bHandleLeftMouseButton = MouseEvent.GetEffectingButton() == EKeys::LeftMouseButton;
 	bool bHandleRightMouseButton = MouseEvent.GetEffectingButton() == EKeys::RightMouseButton && TimeSliderArgs.AllowZoom;
@@ -450,21 +450,21 @@ FReply FSequencerTimeSliderController::OnMouseButtonDown( TSharedRef<SWidget> Wi
 
 	if ( bHandleLeftMouseButton )
 	{
-		return FReply::Handled().CaptureMouse( WidgetOwner ).PreventThrottling();
+		return FReply::Handled().CaptureMouse( WidgetOwner.AsShared() ).PreventThrottling();
 	}
 	else if ( bHandleRightMouseButton )
 	{
 		// Always capture mouse if we left or right click on the widget
-		return FReply::Handled().CaptureMouse( WidgetOwner );
+		return FReply::Handled().CaptureMouse( WidgetOwner.AsShared() );
 	}
 
 	return FReply::Unhandled();
 }
 
-FReply FSequencerTimeSliderController::OnMouseButtonUp( TSharedRef<SWidget> WidgetOwner, const FGeometry& MyGeometry, const FPointerEvent& MouseEvent )
+FReply FSequencerTimeSliderController::OnMouseButtonUp( SWidget& WidgetOwner, const FGeometry& MyGeometry, const FPointerEvent& MouseEvent )
 {
-	bool bHandleLeftMouseButton = MouseEvent.GetEffectingButton() == EKeys::LeftMouseButton && WidgetOwner->HasMouseCapture();
-	bool bHandleRightMouseButton = MouseEvent.GetEffectingButton() == EKeys::RightMouseButton && WidgetOwner->HasMouseCapture() && TimeSliderArgs.AllowZoom ;
+	bool bHandleLeftMouseButton = MouseEvent.GetEffectingButton() == EKeys::LeftMouseButton && WidgetOwner.HasMouseCapture();
+	bool bHandleRightMouseButton = MouseEvent.GetEffectingButton() == EKeys::RightMouseButton && WidgetOwner.HasMouseCapture() && TimeSliderArgs.AllowZoom ;
 	
 	if ( bHandleRightMouseButton )
 	{
@@ -480,54 +480,60 @@ FReply FSequencerTimeSliderController::OnMouseButtonUp( TSharedRef<SWidget> Widg
 	}
 	else if ( bHandleLeftMouseButton )
 	{
-		if( MouseDragType != DRAG_NONE )
+		if (MouseDragType == DRAG_START_RANGE)
 		{
-			TimeSliderArgs.OnEndScrubberMovement.ExecuteIfBound();
+			TimeSliderArgs.OnEndPlaybackRangeDrag.ExecuteIfBound();
+		}
+		// Set the end range time?
+		else if (MouseDragType == DRAG_END_RANGE)
+		{
+			TimeSliderArgs.OnEndPlaybackRangeDrag.ExecuteIfBound();
+		}
+		else if (MouseDragType == DRAG_SETTING_RANGE)
+		{
+			FScrubRangeToScreen RangeToScreen( TimeSliderArgs.ViewRange.Get(), MyGeometry.Size );
+			FVector2D CursorPos = MyGeometry.AbsoluteToLocal(MouseEvent.GetLastScreenSpacePosition());
+			float NewValue = RangeToScreen.LocalXToInput(CursorPos.X);
 
-			if (MouseDragType == DRAG_SETTING_RANGE)
+			if ( TimeSliderArgs.Settings->GetIsSnapEnabled() )
 			{
-				FScrubRangeToScreen RangeToScreen( TimeSliderArgs.ViewRange.Get(), MyGeometry.Size );
-				FVector2D CursorPos = MyGeometry.AbsoluteToLocal(MouseEvent.GetLastScreenSpacePosition());
-				float NewValue = RangeToScreen.LocalXToInput(CursorPos.X);
-
-				if ( TimeSliderArgs.Settings->GetIsSnapEnabled() )
-				{
-					NewValue = TimeSliderArgs.Settings->SnapTimeToInterval(NewValue);
-				}
+				NewValue = TimeSliderArgs.Settings->SnapTimeToInterval(NewValue);
+			}
 
-				float DownValue = MouseDownRange[0];
-					
-				if ( TimeSliderArgs.Settings->GetIsSnapEnabled() )
-				{
-					DownValue = TimeSliderArgs.Settings->SnapTimeToInterval(DownValue);
-				}
+			float DownValue = MouseDownRange[0];
+				
+			if ( TimeSliderArgs.Settings->GetIsSnapEnabled() )
+			{
+				DownValue = TimeSliderArgs.Settings->SnapTimeToInterval(DownValue);
+			}
 
-				// Zoom in
-				if (NewValue > DownValue)
-				{
-					// push the current value onto the stack
-					RangeStack.Add(FVector2D(TimeSliderArgs.ViewRange.Get().GetLowerBoundValue(), TimeSliderArgs.ViewRange.Get().GetUpperBoundValue()));
-				}
-				// Zoom out
-				else if (RangeStack.Num())
-				{
-					// pop the stack
-					FVector2D LastRange = RangeStack.Pop();
-					DownValue = LastRange[0];
-					NewValue = LastRange[1];
-				}
+			// Zoom in
+			if (NewValue > DownValue)
+			{
+				// push the current value onto the stack
+				RangeStack.Add(FVector2D(TimeSliderArgs.ViewRange.Get().GetLowerBoundValue(), TimeSliderArgs.ViewRange.Get().GetUpperBoundValue()));
+			}
+			// Zoom out
+			else if (RangeStack.Num())
+			{
+				// pop the stack
+				FVector2D LastRange = RangeStack.Pop();
+				DownValue = LastRange[0];
+				NewValue = LastRange[1];
+			}
 
-				TimeSliderArgs.OnViewRangeChanged.ExecuteIfBound(TRange<float>(DownValue, NewValue), EViewRangeInterpolation::Immediate);
-						
-				if( !TimeSliderArgs.ViewRange.IsBound() )
-				{	
-					// The output is not bound to a delegate so we'll manage the value ourselves
-					TimeSliderArgs.ViewRange.Set( TRange<float>( DownValue, NewValue ) );
-				}
+			TimeSliderArgs.OnViewRangeChanged.ExecuteIfBound(TRange<float>(DownValue, NewValue), EViewRangeInterpolation::Immediate);
+					
+			if( !TimeSliderArgs.ViewRange.IsBound() )
+			{	
+				// The output is not bound to a delegate so we'll manage the value ourselves
+				TimeSliderArgs.ViewRange.Set( TRange<float>( DownValue, NewValue ) );
 			}
 		}
 		else
 		{
+			TimeSliderArgs.OnEndScrubberMovement.ExecuteIfBound();
+
 			FScrubRangeToScreen RangeToScreen( TimeSliderArgs.ViewRange.Get(), MyGeometry.Size );
 			FVector2D CursorPos = MyGeometry.AbsoluteToLocal(MouseEvent.GetLastScreenSpacePosition());
 			float NewValue = RangeToScreen.LocalXToInput(CursorPos.X);
@@ -548,9 +554,9 @@ FReply FSequencerTimeSliderController::OnMouseButtonUp( TSharedRef<SWidget> Widg
 	return FReply::Unhandled();
 }
 
-FReply FSequencerTimeSliderController::OnMouseMove( TSharedRef<SWidget> WidgetOwner, const FGeometry& MyGeometry, const FPointerEvent& MouseEvent )
+FReply FSequencerTimeSliderController::OnMouseMove( SWidget& WidgetOwner, const FGeometry& MyGeometry, const FPointerEvent& MouseEvent )
 {
-	if ( WidgetOwner->HasMouseCapture() )
+	if ( WidgetOwner.HasMouseCapture() )
 	{
 		if (MouseEvent.IsMouseButtonDown( EKeys::RightMouseButton ))
 		{
@@ -583,7 +589,6 @@ FReply FSequencerTimeSliderController::OnMouseMove( TSharedRef<SWidget> WidgetOw
 		else if (MouseEvent.IsMouseButtonDown( EKeys::LeftMouseButton ))
 		{
 			TRange<float> LocalViewRange = TimeSliderArgs.ViewRange.Get();
-
 			DistanceDragged += FMath::Abs( MouseEvent.GetCursorDelta().X );
 
 			if ( MouseDragType == DRAG_NONE )
@@ -600,10 +605,12 @@ FReply FSequencerTimeSliderController::OnMouseMove( TSharedRef<SWidget> WidgetOw
 					if (HitTestPlaybackEnd(RangeToScreen, PlaybackRange, LocalMouseDownPos, ScrubPosition))
 					{
 						MouseDragType = DRAG_END_RANGE;
+						TimeSliderArgs.OnBeginPlaybackRangeDrag.ExecuteIfBound();
 					}
 					else if (HitTestPlaybackStart(RangeToScreen, PlaybackRange, LocalMouseDownPos, ScrubPosition))
 					{
 						MouseDragType = DRAG_START_RANGE;
+						TimeSliderArgs.OnBeginPlaybackRangeDrag.ExecuteIfBound();
 					}
 					else if (FSlateApplication::Get().GetModifierKeys().AreModifersDown(EModifierKey::Control))
 					{
@@ -612,8 +619,8 @@ FReply FSequencerTimeSliderController::OnMouseMove( TSharedRef<SWidget> WidgetOw
 					else
 					{
 						MouseDragType = DRAG_SCRUBBING_TIME;
+						TimeSliderArgs.OnBeginScrubberMovement.ExecuteIfBound();
 					}
-					TimeSliderArgs.OnBeginScrubberMovement.ExecuteIfBound();
 				}
 			}
 			else
@@ -676,7 +683,7 @@ void FSequencerTimeSliderController::CommitScrubPosition( float NewValue, bool b
 	TimeSliderArgs.OnScrubPositionChanged.ExecuteIfBound( NewValue, bIsScrubbing );
 }
 
-FReply FSequencerTimeSliderController::OnMouseWheel( TSharedRef<SWidget> WidgetOwner, const FGeometry& MyGeometry, const FPointerEvent& MouseEvent )
+FReply FSequencerTimeSliderController::OnMouseWheel( SWidget& WidgetOwner, const FGeometry& MyGeometry, const FPointerEvent& MouseEvent )
 {
 	TOptional<TRange<float>> NewTargetRange;
 
diff --git a/Engine/Source/Editor/Sequencer/Private/TimeSliderController.h b/Engine/Source/Editor/Sequencer/Private/TimeSliderController.h
--- a/Engine/Source/Editor/Sequencer/Private/TimeSliderController.h
+++ b/Engine/Source/Editor/Sequencer/Private/TimeSliderController.h
@@ -46,10 +46,10 @@ public:
 
 	/** ITimeSliderController Interface */
 	virtual int32 OnPaintTimeSlider( bool bMirrorLabels, const FGeometry& AllottedGeometry, const FSlateRect& MyClippingRect, FSlateWindowElementList& OutDrawElements, int32 LayerId, const FWidgetStyle& InWidgetStyle, bool bParentEnabled ) const override;
-	virtual FReply OnMouseButtonDown( TSharedRef<SWidget> WidgetOwner, const FGeometry& MyGeometry, const FPointerEvent& MouseEvent ) override;
-	virtual FReply OnMouseButtonUp( TSharedRef<SWidget> WidgetOwner, const FGeometry& MyGeometry, const FPointerEvent& MouseEvent ) override;
-	virtual FReply OnMouseMove( TSharedRef<SWidget> WidgetOwner, const FGeometry& MyGeometry, const FPointerEvent& MouseEvent ) override;
-	virtual FReply OnMouseWheel( TSharedRef<SWidget> WidgetOwner, const FGeometry& MyGeometry, const FPointerEvent& MouseEvent ) override;
+	virtual FReply OnMouseButtonDown( SWidget& WidgetOwner, const FGeometry& MyGeometry, const FPointerEvent& MouseEvent ) override;
+	virtual FReply OnMouseButtonUp( SWidget& WidgetOwner, const FGeometry& MyGeometry, const FPointerEvent& MouseEvent ) override;
+	virtual FReply OnMouseMove( SWidget& WidgetOwner, const FGeometry& MyGeometry, const FPointerEvent& MouseEvent ) override;
+	virtual FReply OnMouseWheel( SWidget& WidgetOwner, const FGeometry& MyGeometry, const FPointerEvent& MouseEvent ) override;
 	virtual FCursorReply OnCursorQuery( TSharedRef<const SWidget> WidgetOwner, const FGeometry& MyGeometry, const FPointerEvent& CursorEvent ) const override;
 	/** End ITimeSliderController Interface */
 
diff --git a/Engine/Source/Editor/Sequencer/Private/Tools/SequencerEditTool_Default.cpp b/Engine/Source/Editor/Sequencer/Private/Tools/SequencerEditTool_Default.cpp
--- a/Engine/Source/Editor/Sequencer/Private/Tools/SequencerEditTool_Default.cpp
+++ b/Engine/Source/Editor/Sequencer/Private/Tools/SequencerEditTool_Default.cpp
@@ -3,16 +3,110 @@
 #include "SequencerPrivatePCH.h"
 #include "SequencerEditTool_Default.h"
 #include "Sequencer.h"
+#include "SequencerHotspots.h"
 
 
 FReply FSequencerEditTool_Default::OnMouseButtonUp(SWidget& OwnerWidget, const FGeometry& MyGeometry, const FPointerEvent& MouseEvent)
 {
 	if (MouseEvent.GetEffectingButton() == EKeys::LeftMouseButton)
 	{
-		FSequencer& Sequencer = static_cast<FSequencer&>(GetSequencer());
-		Sequencer.GetSelection().EmptySelectedSections();
-		Sequencer.GetSelection().EmptySelectedKeys();
+		FSequencerSelection& Selection = GetSequencer().GetSelection();
+		Selection.EmptySelectedSections();
+		Selection.EmptySelectedKeys();
 	}
 
 	return FReply::Unhandled();
 }
+
+void FSequencerEditTool_Default::PerformHotspotSelection(const FPointerEvent& MouseEvent)
+{
+	FSequencerSelection& Selection = GetSequencer().GetSelection();
+
+	// @todo: selection in transactions
+	auto ConditionallyClearSelection = [&]{
+		if (!MouseEvent.IsShiftDown() && !MouseEvent.IsControlDown())
+		{
+			Selection.EmptySelectedSections();
+			Selection.EmptySelectedKeys();
+		}
+	};
+
+	// Handle right-click selection separately since we never deselect on right click (except for clearing on exclusive selection)
+	if (MouseEvent.GetEffectingButton() == EKeys::RightMouseButton)
+	{
+		if (Hotspot->GetType() == ESequencerHotspot::Key)
+		{
+			FSequencerSelectedKey Key = static_cast<FKeyHotspot*>(Hotspot.Get())->Key;
+			if (!Selection.IsSelected(Key))
+			{
+				ConditionallyClearSelection();
+				Selection.AddToSelection(Key);
+			}
+		}
+		else if (Hotspot->GetType() == ESequencerHotspot::Section)
+		{
+			UMovieSceneSection* Section = static_cast<FSectionHotspot*>(Hotspot.Get())->Section.GetSectionObject();
+			if (!Selection.IsSelected(Section))
+			{
+				ConditionallyClearSelection();
+				Selection.AddToSelection(Section);
+			}
+		}
+
+		return;
+	}
+
+	// Normal selection
+	ConditionallyClearSelection();
+
+	bool bForceSelect = !MouseEvent.IsControlDown();
+
+	if (Hotspot->GetType() == ESequencerHotspot::Key)
+	{
+		FSequencerSelectedKey Key = static_cast<FKeyHotspot*>(Hotspot.Get())->Key;
+		if (bForceSelect || !Selection.IsSelected(Key))
+		{
+			Selection.AddToSelection(Key);
+		}
+		else
+		{
+			Selection.RemoveFromSelection(Key);
+		}
+	}
+	else if (Hotspot->GetType() == ESequencerHotspot::Section)
+	{
+		UMovieSceneSection* Section = static_cast<FSectionHotspot*>(Hotspot.Get())->Section.GetSectionObject();
+		if (bForceSelect || !Selection.IsSelected(Section))
+		{
+			Selection.AddToSelection(Section);
+		}
+		else
+		{
+			Selection.RemoveFromSelection(Section);
+		}
+	}
+}
+
+TSharedPtr<SWidget> FSequencerEditTool_Default::OnSummonContextMenu( const FGeometry& MyGeometry, const FPointerEvent& MouseEvent )
+{
+	if (!Hotspot.IsValid())
+	{
+		return nullptr;
+	}
+
+	ISequencer& Sequencer = GetSequencer();
+
+	// @todo sequencer replace with UI Commands instead of faking it
+
+	if (Hotspot->GetType() == ESequencerHotspot::Section || Hotspot->GetType() == ESequencerHotspot::Key)
+	{
+		const bool bShouldCloseWindowAfterMenuSelection = true;
+		FMenuBuilder MenuBuilder(bShouldCloseWindowAfterMenuSelection, NULL);
+
+		Hotspot->PopulateContextMenu(MenuBuilder, Sequencer);
+
+		return MenuBuilder.MakeWidget();
+	}
+
+	return nullptr;
+}
\ No newline at end of file
diff --git a/Engine/Source/Editor/Sequencer/Private/Tools/SequencerEditTool_Default.h b/Engine/Source/Editor/Sequencer/Private/Tools/SequencerEditTool_Default.h
--- a/Engine/Source/Editor/Sequencer/Private/Tools/SequencerEditTool_Default.h
+++ b/Engine/Source/Editor/Sequencer/Private/Tools/SequencerEditTool_Default.h
@@ -13,4 +13,10 @@ public:
 	// ISequencerEditTool interface
 
 	virtual FReply OnMouseButtonUp(SWidget& OwnerWidget, const FGeometry& MyGeometry, const FPointerEvent& MouseEvent) override;
+	
+protected:
+
+	void PerformHotspotSelection(const FPointerEvent& MouseEvent);
+	
+	TSharedPtr<SWidget> OnSummonContextMenu( const FGeometry& MyGeometry, const FPointerEvent& MouseEvent );
 };
diff --git a/Engine/Source/Editor/Sequencer/Private/Tools/SequencerEditTool_Movement.cpp b/Engine/Source/Editor/Sequencer/Private/Tools/SequencerEditTool_Movement.cpp
--- a/Engine/Source/Editor/Sequencer/Private/Tools/SequencerEditTool_Movement.cpp
+++ b/Engine/Source/Editor/Sequencer/Private/Tools/SequencerEditTool_Movement.cpp
@@ -162,6 +162,31 @@ FReply FSequencerEditTool_Movement::OnMouseButtonUp(SWidget& OwnerWidget, const
 		// Only return handled if we actually started a drag
 		return FReply::Handled().ReleaseMouseCapture();
 	}
+	else if (Hotspot.IsValid())
+	{
+		PerformHotspotSelection(MouseEvent);
+
+		if (MouseEvent.GetEffectingButton() == EKeys::RightMouseButton)
+		{
+			TSharedPtr<SWidget> MenuContent = OnSummonContextMenu( MyGeometry, MouseEvent );
+			if (MenuContent.IsValid())
+			{
+				FWidgetPath WidgetPath = MouseEvent.GetEventPath() != nullptr ? *MouseEvent.GetEventPath() : FWidgetPath();
+
+				FSlateApplication::Get().PushMenu(
+					OwnerWidget.AsShared(),
+					WidgetPath,
+					MenuContent.ToSharedRef(),
+					MouseEvent.GetScreenSpacePosition(),
+					FPopupTransitionEffect( FPopupTransitionEffect::ContextMenu )
+					);
+
+				return FReply::Handled().SetUserFocus(MenuContent.ToSharedRef(), EFocusCause::SetDirectly).ReleaseMouseCapture();
+			}
+		}
+
+		return FReply::Handled().ReleaseMouseCapture();
+	}
 
 	return FSequencerEditTool_Default::OnMouseButtonUp(OwnerWidget, MyGeometry, MouseEvent);
 }
@@ -176,7 +201,20 @@ void FSequencerEditTool_Movement::OnMouseCaptureLost()
 
 FCursorReply FSequencerEditTool_Movement::OnCursorQuery(const FGeometry& MyGeometry, const FPointerEvent& CursorEvent) const
 {
-	return DragOperation.IsValid() ? FCursorReply::Cursor(EMouseCursor::CardinalCross) : FCursorReply::Cursor(EMouseCursor::Default);
+	if (DragOperation.IsValid())
+	{
+		return FCursorReply::Cursor(EMouseCursor::CardinalCross);
+	}
+
+	if (Hotspot.IsValid())
+	{
+		FCursorReply Reply = Hotspot->GetCursor();
+		if (Reply.IsEventHandled())
+		{
+			return Reply;
+		}
+	}
+	return FCursorReply::Cursor(EMouseCursor::Default);
 }
 
 
diff --git a/Engine/Source/Editor/Sequencer/Public/IKeyArea.h b/Engine/Source/Editor/Sequencer/Public/IKeyArea.h
--- a/Engine/Source/Editor/Sequencer/Public/IKeyArea.h
+++ b/Engine/Source/Editor/Sequencer/Public/IKeyArea.h
@@ -2,27 +2,40 @@
 
 #pragma once
 
-/** Information for how to draw each key */
+#include "Curves/CurveBase.h"
+#include "KeyParams.h"
+
+
+class UMovieSceneSection;
+
+
+/**
+ * Information for how to draw each key.
+ */
 struct FKeyDrawingInfo
 {
 	FKeyDrawingInfo()
-		: Brush( NULL )
-		, Size( FVector2D::ZeroVector )
-	{}
+		: Brush(nullptr)
+		, Size(FVector2D::ZeroVector)
+	{ }
 
 	/** The brush to use for each key */
 	const FSlateBrush* Brush;
+
 	/** The size of each key */
 	const FVector2D Size;
 };
 
+
 /**
  * Interface that should be implemented for the UI portion of a key area within a section
  */
 class IKeyArea
 {
 public:
-	virtual ~IKeyArea() {}
+
+	/** Virtual destructor. */
+	virtual ~IKeyArea() { }
 
 	/** @return The array of unsorted key handles in the key area */
 	virtual TArray<FKeyHandle> GetUnsortedKeyHandles() const = 0;
@@ -50,7 +63,7 @@ public:
 	 * @param DeltaPosition	The delta position of the key
 	 * @return The handle of the key
 	 */
-	virtual FKeyHandle MoveKey( FKeyHandle KeyHandle, float DeltaPosition ) = 0;
+	virtual FKeyHandle MoveKey(FKeyHandle KeyHandle, float DeltaPosition) = 0;
 	
 	/**
 	 * Deletes a key
@@ -136,5 +149,3 @@ public:
 	*/
 	virtual TSharedRef<SWidget> CreateKeyEditor(ISequencer* Sequencer) = 0;
 };
-
-
diff --git a/Engine/Source/Editor/Sequencer/Public/ISequencerEditTool.h b/Engine/Source/Editor/Sequencer/Public/ISequencerEditTool.h
--- a/Engine/Source/Editor/Sequencer/Public/ISequencerEditTool.h
+++ b/Engine/Source/Editor/Sequencer/Public/ISequencerEditTool.h
@@ -55,15 +55,11 @@ public:
 /**
  * Interface for edit tools in Sequencer.
  */
-class ISequencerEditTool
+class ISequencerEditTool : public ISequencerInputHandler
 {
 public:
 
 	// @todo sequencer: documentation needed
-	virtual FReply OnMouseButtonDown(SWidget& OwnerWidget, const FGeometry& MyGeometry, const FPointerEvent& MouseEvent) = 0;
-	virtual FReply OnMouseButtonUp(SWidget& OwnerWidget, const FGeometry& MyGeometry, const FPointerEvent& MouseEvent) = 0;
-	virtual FReply OnMouseMove(SWidget& OwnerWidget, const FGeometry& MyGeometry, const FPointerEvent& MouseEvent) = 0;
-	virtual FReply OnMouseWheel(SWidget& OwnerWidget, const FGeometry& MyGeometry, const FPointerEvent& MouseEvent) = 0;
 	virtual void OnMouseCaptureLost() = 0;
 	virtual void OnMouseEnter(SWidget& OwnerWidget, const FGeometry& MyGeometry, const FPointerEvent& MouseEvent) = 0;
 	virtual void OnMouseLeave(SWidget& OwnerWidget, const FPointerEvent& MouseEvent) = 0;
diff --git a/Engine/Source/Editor/Sequencer/Public/ISequencerHotspot.h b/Engine/Source/Editor/Sequencer/Public/ISequencerHotspot.h
--- a/Engine/Source/Editor/Sequencer/Public/ISequencerHotspot.h
+++ b/Engine/Source/Editor/Sequencer/Public/ISequencerHotspot.h
@@ -22,4 +22,6 @@ struct ISequencerHotspot
 	virtual ~ISequencerHotspot() { }
 	virtual ESequencerHotspot GetType() const = 0;
 	virtual TSharedPtr<ISequencerEditToolDragOperation> InitiateDrag(ISequencer&) = 0;
+	virtual void PopulateContextMenu(FMenuBuilder& MenuBuilder, ISequencer& Sequencer){}
+	virtual FCursorReply GetCursor() const { return FCursorReply::Unhandled(); }
 };
diff --git a/Engine/Source/Editor/Sequencer/Public/ISequencerInputHandler.h b/Engine/Source/Editor/Sequencer/Public/ISequencerInputHandler.h
new file mode 100644
--- /dev/null
+++ b/Engine/Source/Editor/Sequencer/Public/ISequencerInputHandler.h
@@ -0,0 +1,14 @@
+// Copyright 1998-2015 Epic Games, Inc. All Rights Reserved.
+
+#pragma once
+
+/**
+ * Common base-class for objects that handle input in the sequencer.
+ */
+struct ISequencerInputHandler
+{
+	virtual FReply OnMouseButtonDown(SWidget& OwnerWidget, const FGeometry& MyGeometry, const FPointerEvent& MouseEvent) = 0;
+	virtual FReply OnMouseButtonUp(SWidget& OwnerWidget, const FGeometry& MyGeometry, const FPointerEvent& MouseEvent) = 0;
+	virtual FReply OnMouseMove(SWidget& OwnerWidget, const FGeometry& MyGeometry, const FPointerEvent& MouseEvent) = 0;
+	virtual FReply OnMouseWheel(SWidget& OwnerWidget, const FGeometry& MyGeometry, const FPointerEvent& MouseEvent) = 0;
+};
\ No newline at end of file
diff --git a/Engine/Source/Editor/Sequencer/Public/ISequencerSection.h b/Engine/Source/Editor/Sequencer/Public/ISequencerSection.h
--- a/Engine/Source/Editor/Sequencer/Public/ISequencerSection.h
+++ b/Engine/Source/Editor/Sequencer/Public/ISequencerSection.h
@@ -10,7 +10,7 @@ namespace SequencerSectionConstants
 	const float SectionDragStartDistance = 5.0f;
 
 	/** The size of each key */
-	const FVector2D KeySize(11.0f, 11.0f);
+	const FVector2D KeySize(12.0f, 12.0f);
 
 	const float DefaultSectionGripSize = 7.0f;
 
diff --git a/Engine/Source/Editor/Sequencer/Sequencer.Build.cs b/Engine/Source/Editor/Sequencer/Sequencer.Build.cs
--- a/Engine/Source/Editor/Sequencer/Sequencer.Build.cs
+++ b/Engine/Source/Editor/Sequencer/Sequencer.Build.cs
@@ -39,6 +39,7 @@ public class Sequencer : ModuleRules
 			new string[] {
 				"PropertyEditor",
 				"Kismet",
+				"SceneOutliner"
 			}
 		);
 
@@ -46,7 +47,8 @@ public class Sequencer : ModuleRules
 			new string[] {
 				"PropertyEditor",
 				"LevelEditor",
-				"WorkspaceMenuStructure",
+				"SceneOutliner",
+				"WorkspaceMenuStructure"
 			}
 		);
 
diff --git a/Engine/Source/Editor/SequencerWidgets/Private/STimeSlider.cpp b/Engine/Source/Editor/SequencerWidgets/Private/STimeSlider.cpp
--- a/Engine/Source/Editor/SequencerWidgets/Private/STimeSlider.cpp
+++ b/Engine/Source/Editor/SequencerWidgets/Private/STimeSlider.cpp
@@ -21,17 +21,17 @@ int32 STimeSlider::OnPaint( const FPaintArgs& Args, const FGeometry& AllottedGeo
 
 FReply STimeSlider::OnMouseButtonDown( const FGeometry& MyGeometry, const FPointerEvent& MouseEvent )
 {
-	return TimeSliderController->OnMouseButtonDown( SharedThis(this), MyGeometry, MouseEvent );
+	return TimeSliderController->OnMouseButtonDown( *this, MyGeometry, MouseEvent );
 }
 
 FReply STimeSlider::OnMouseButtonUp( const FGeometry& MyGeometry, const FPointerEvent& MouseEvent )
 {
-	return TimeSliderController->OnMouseButtonUp( SharedThis(this),  MyGeometry, MouseEvent );
+	return TimeSliderController->OnMouseButtonUp( *this,  MyGeometry, MouseEvent );
 }
 
 FReply STimeSlider::OnMouseMove( const FGeometry& MyGeometry, const FPointerEvent& MouseEvent )
 {
-	return TimeSliderController->OnMouseMove( SharedThis(this), MyGeometry, MouseEvent );
+	return TimeSliderController->OnMouseMove( *this, MyGeometry, MouseEvent );
 }
 
 FVector2D STimeSlider::ComputeDesiredSize( float ) const
@@ -41,7 +41,7 @@ FVector2D STimeSlider::ComputeDesiredSize( float ) const
 
 FReply STimeSlider::OnMouseWheel( const FGeometry& MyGeometry, const FPointerEvent& MouseEvent )
 {
-	return TimeSliderController->OnMouseWheel( SharedThis(this), MyGeometry, MouseEvent );
+	return TimeSliderController->OnMouseWheel( *this, MyGeometry, MouseEvent );
 }
 
 FCursorReply STimeSlider::OnCursorQuery(const FGeometry& MyGeometry, const FPointerEvent& CursorEvent) const
diff --git a/Engine/Source/Editor/SequencerWidgets/Public/ITimeSlider.h b/Engine/Source/Editor/SequencerWidgets/Public/ITimeSlider.h
--- a/Engine/Source/Editor/SequencerWidgets/Public/ITimeSlider.h
+++ b/Engine/Source/Editor/SequencerWidgets/Public/ITimeSlider.h
@@ -2,6 +2,8 @@
 
 #pragma once
 
+#include "Editor/Sequencer/Public/ISequencerInputHandler.h"
+
 class USequencerSettings;
 
 /** Enum specifying how to interpolate to a new view range */
@@ -86,21 +88,21 @@ struct FTimeSliderArgs
 	TAttribute<TRange<float>> PlaybackRange;
 	/** Delegate that is called when the playback range wants to change */
 	FOnRangeChanged OnPlaybackRangeChanged;
+	/** Called right before the playback range starts to be dragged */
+	FSimpleDelegate OnBeginPlaybackRangeDrag;
+	/** Called right after the playback range has finished being dragged */
+	FSimpleDelegate OnEndPlaybackRangeDrag;
 	/** If we are allowed to zoom */
 	bool AllowZoom;
 	/** User-supplied settings object */
 	USequencerSettings* Settings;
 };
 
-class ITimeSliderController
+class ITimeSliderController : public ISequencerInputHandler
 {
 public:
 	virtual ~ITimeSliderController(){}
 	virtual int32 OnPaintTimeSlider( bool bMirrorLabels, const FGeometry& AllottedGeometry, const FSlateRect& MyClippingRect, FSlateWindowElementList& OutDrawElements, int32 LayerId, const FWidgetStyle& InWidgetStyle, bool bParentEnabled ) const = 0;
-	virtual FReply OnMouseButtonDown( TSharedRef<SWidget> WidgetOwner, const FGeometry& MyGeometry, const FPointerEvent& MouseEvent ) = 0;
-	virtual FReply OnMouseButtonUp( TSharedRef<SWidget> WidgetOwner, const FGeometry& MyGeometry, const FPointerEvent& MouseEvent ) = 0;
-	virtual FReply OnMouseMove( TSharedRef<SWidget> WidgetOwner, const FGeometry& MyGeometry, const FPointerEvent& MouseEvent ) = 0;
-	virtual FReply OnMouseWheel( TSharedRef<SWidget> WidgetOwner, const FGeometry& MyGeometry, const FPointerEvent& MouseEvent ) = 0;
 	virtual FCursorReply OnCursorQuery( TSharedRef<const SWidget> WidgetOwner, const FGeometry& MyGeometry, const FPointerEvent& CursorEvent ) const = 0;
 
 	/** Get the current view range for this controller */
diff --git a/Engine/Source/Editor/UMGEditor/Private/Animation/MarginTrackEditor.cpp b/Engine/Source/Editor/UMGEditor/Private/Animation/MarginTrackEditor.cpp
--- a/Engine/Source/Editor/UMGEditor/Private/Animation/MarginTrackEditor.cpp
+++ b/Engine/Source/Editor/UMGEditor/Private/Animation/MarginTrackEditor.cpp
@@ -6,7 +6,7 @@
 #include "MovieSceneMarginSection.h"
 #include "PropertySection.h"
 #include "ISectionLayoutBuilder.h"
-#include "MovieSceneToolHelpers.h"
+#include "FloatCurveKeyArea.h"
 
 
 FName FMarginTrackEditor::LeftName( "Left" );
diff --git a/Engine/Source/Editor/UMGEditor/Private/Animation/Sequencer2DTransformTrackEditor.cpp b/Engine/Source/Editor/UMGEditor/Private/Animation/Sequencer2DTransformTrackEditor.cpp
--- a/Engine/Source/Editor/UMGEditor/Private/Animation/Sequencer2DTransformTrackEditor.cpp
+++ b/Engine/Source/Editor/UMGEditor/Private/Animation/Sequencer2DTransformTrackEditor.cpp
@@ -6,13 +6,15 @@
 #include "MovieScene2DTransformTrack.h"
 #include "PropertySection.h"
 #include "ISectionLayoutBuilder.h"
-#include "MovieSceneToolHelpers.h"
+#include "FloatCurveKeyArea.h"
+
 
 FName F2DTransformTrackEditor::TranslationName( "Translation" );
 FName F2DTransformTrackEditor::ScaleName( "Scale" );
 FName F2DTransformTrackEditor::ShearName( "Shear" );
 FName F2DTransformTrackEditor::AngleName( "Angle" );
 
+
 class F2DTransformSection
 	: public FPropertySection
 {
diff --git a/Engine/Source/Editor/UnrealEd/Private/Commandlets/ContentCommandlets.cpp b/Engine/Source/Editor/UnrealEd/Private/Commandlets/ContentCommandlets.cpp
--- a/Engine/Source/Editor/UnrealEd/Private/Commandlets/ContentCommandlets.cpp
+++ b/Engine/Source/Editor/UnrealEd/Private/Commandlets/ContentCommandlets.cpp
@@ -30,7 +30,8 @@ DEFINE_LOG_CATEGORY_STATIC(LogContentCommandlet, Log, All);
 // for UResavePackagesCommandlet::PerformAdditionalOperations building lighting code
 #include "Engine/WorldComposition.h"
 #include "LightingBuildOptions.h"
-
+// For preloading FFindInBlueprintSearchManager
+#include "Editor/Kismet/Public/FindInBlueprintManager.h"
 
 /**-----------------------------------------------------------------------------
  *	UResavePackages commandlet.
@@ -63,6 +64,7 @@ int32 UResavePackagesCommandlet::InitializeResaveParameters( const TArray<FStrin
 	{
 		FString Package;
 		FString PackageFolder;
+		FString Maps;
 		const FString& CurrentSwitch = Switches[ SwitchIdx ];
 		if( FParse::Value( *CurrentSwitch, TEXT( "PACKAGE="), Package ) )
 		{
@@ -82,7 +84,29 @@ int32 UResavePackagesCommandlet::InitializeResaveParameters( const TArray<FStrin
 				PackageNames.Add( *PackageFile );
             }
 			bExplicitPackages = true;
-        }
+		}
+		else if (FParse::Value(*CurrentSwitch, TEXT("MAP="), Maps))
+		{
+			// Allow support for -MAP=Value1+Value2+Value3
+			for (int32 PlusIdx = Maps.Find(TEXT("+")); PlusIdx != INDEX_NONE; PlusIdx = Maps.Find(TEXT("+")))
+			{
+				const FString NextMap = Maps.Left(PlusIdx);
+				if (NextMap.Len() > 0)
+				{
+					FString MapFile;
+					FPackageName::SearchForPackageOnDisk(NextMap, NULL, &MapFile, false);
+					PackageNames.Add(*MapFile);
+					bExplicitPackages = true;
+				}
+
+				Maps = Maps.Right(Maps.Len() - (PlusIdx + 1));
+			}
+			FString MapFile;
+			FPackageName::SearchForPackageOnDisk(Maps, NULL, &MapFile, false);
+			PackageNames.Add(*MapFile);
+			bExplicitPackages = true;
+		}
+
 	}
 
 	// ... if not, load in all packages
@@ -98,6 +122,11 @@ int32 UResavePackagesCommandlet::InitializeResaveParameters( const TArray<FStrin
 			PackageFilter |= NORMALIZE_ExcludeContentPackages;
 		}
 
+		if (Switches.Contains(TEXT("PROJECTONLY")))
+		{
+			PackageFilter |= NORMALIZE_ExcludeEnginePackages;
+		}
+
 		if ( Switches.Contains(TEXT("SkipDeveloperFolders")) || Switches.Contains(TEXT("NODEV")) )
 		{
 			PackageFilter |= NORMALIZE_ExcludeDeveloperPackages;
@@ -443,7 +472,7 @@ void UResavePackagesCommandlet::LoadAndSaveOnePackage(const FString& Filename)
 
 									VerboseMessage(TEXT("Post CheckOut"));
 
-									FString PackageName(FPaths::GetBaseFilename(Filename));
+									FString PackageName(FPackageName::FilenameToLongPackageName(Filename));
 									FilesToSubmit.Add(*PackageName);
 								}
 							}
@@ -544,6 +573,9 @@ int32 UResavePackagesCommandlet::Main( const FString& Params )
 		UE_LOG(LogContentCommandlet, Display, TEXT( "Restricted to packages containing %s" ), *PackageSubstring );
 	}
 
+	// Avoid crash saving blueprint
+	FFindInBlueprintSearchManager::Get();
+
 	// Iterate over all packages.
 	for( int32 PackageIndex = 0; PackageIndex < PackageNames.Num(); PackageIndex++ )
 	{
@@ -602,7 +634,19 @@ int32 UResavePackagesCommandlet::Main( const FString& Params )
 
 FText UResavePackagesCommandlet::GetChangelistDescription() const
 {
-	return NSLOCTEXT("ContentCmdlets", "ChangelistDescription", "Resave Deprecated Packages");
+	FText ChangelistDescription;
+
+	static bool bBuildLighting = FParse::Param(FCommandLine::Get(), TEXT("buildlighting")) == true;
+	if (bBuildLighting)
+	{
+		ChangelistDescription = NSLOCTEXT("ContentCmdlets", "ChangelistDescriptionBuildLighting", "Rebuild lightmaps.");
+	}
+	else
+	{
+		ChangelistDescription = NSLOCTEXT("ContentCmdlets", "ChangelistDescription", "Resave Deprecated Packages");
+	}
+
+	return ChangelistDescription;
 }
 
 
@@ -660,97 +704,68 @@ void UResavePackagesCommandlet::PerformAdditionalOperations(class UWorld* World,
 
 	if (bBuildLighting)
 	{
-		TArray<FString> CmdLineMapEntries;
-		for (const FString& Switch : Switches)
+		static bool bHasLoadedStartupPackages = false;
+		if (bHasLoadedStartupPackages == false)
 		{
-			auto GetSwitchValueElements = [&Switch](const FString SwitchKey) -> TArray < FString >
-			{
-				TArray<FString> ValueElements;
-				if (Switch.StartsWith(SwitchKey + TEXT("=")) == true)
-				{
-					FString ValuesList = Switch.Right(Switch.Len() - (SwitchKey + TEXT("=")).Len());
-
-					// Allow support for -KEY=Value1+Value2+Value3 as well as -KEY=Value1 -KEY=Value2
-					for (int32 PlusIdx = ValuesList.Find(TEXT("+")); PlusIdx != INDEX_NONE; PlusIdx = ValuesList.Find(TEXT("+")))
-					{
-						const FString ValueElement = ValuesList.Left(PlusIdx);
-						ValueElements.Add(ValueElement);
-
-						ValuesList = ValuesList.Right(ValuesList.Len() - (PlusIdx + 1));
-					}
-					ValueElements.Add(ValuesList);
-				}
-				return ValueElements;
-			};
-
-			CmdLineMapEntries += GetSwitchValueElements(TEXT("MAP"));
+			// make sure all possible script/startup packages are loaded
+			bHasLoadedStartupPackages = FStartupPackages::LoadAll();
 		}
 
-		// If specific maps were specified, then only continue if we are processing one that was required.
-		bool bShouldBuildLightmapsForWorld = CmdLineMapEntries.Num() == 0 || CmdLineMapEntries.Contains(World->GetMapName());
-		if (bShouldBuildLightmapsForWorld)
+		World->AddToRoot();
+		if (!World->bIsWorldInitialized)
 		{
-			World->AddToRoot();
-			if (!World->bIsWorldInitialized)
-			{
-				UWorld::InitializationValues IVS;
-				IVS.RequiresHitProxies(false);
-				IVS.ShouldSimulatePhysics(false);
-				IVS.EnableTraceCollision(false);
-				IVS.CreateNavigation(false);
-				IVS.CreateAISystem(false);
-				IVS.AllowAudioPlayback(false);
-				IVS.CreatePhysicsScene(false);
-
-				World->InitWorld(IVS);
-				World->PersistentLevel->UpdateModelComponents();
-				World->UpdateWorldComponents(true, false);
-			}
-
-			// load all the sublevels
-
-			if (World->StreamingLevels.Num())
-			{
-				//World->LoadSecondaryLevels(true, &PreviouslyCookedPackages);
-				World->LoadSecondaryLevels(true, NULL);
-			}
+			UWorld::InitializationValues IVS;
+			IVS.RequiresHitProxies(false);
+			IVS.ShouldSimulatePhysics(false);
+			IVS.EnableTraceCollision(false);
+			IVS.CreateNavigation(false);
+			IVS.CreateAISystem(false);
+			IVS.AllowAudioPlayback(false);
+			IVS.CreatePhysicsScene(false);
+
+			World->InitWorld(IVS);
+			World->PersistentLevel->UpdateModelComponents();
+			World->UpdateWorldComponents(true, false);
+		}
 
-			GWorld = World;
+		// load all the sublevels
 
-			TArray<FString> SubPackages;
+		if (World->StreamingLevels.Num())
+		{
+			//World->LoadSecondaryLevels(true, &PreviouslyCookedPackages);
+			World->LoadSecondaryLevels(true, NULL);
+		}
 
-			// force load all of the map to be loaded so we can get good rebuilt lighting 
-			if (World->WorldComposition)
-			{
-				World->WorldComposition->CollectTilesToCook(SubPackages);
-			}
+		GWorld = World;
 
-			for (const auto& SubPackage : SubPackages)
-			{
-				UPackage* Package = LoadPackage(nullptr, *SubPackage, 0);
-				check(Package->IsFullyLoaded());
-			}
+		TArray<FString> SubPackages;
 
-			GRedirectCollector.ResolveStringAssetReference();
+		// force load all of the map to be loaded so we can get good rebuilt lighting 
+		if (World->WorldComposition)
+		{
+			World->WorldComposition->CollectTilesToCook(SubPackages);
+		}
 
+		for (const auto& SubPackage : SubPackages)
+		{
+			UPackage* Package = LoadPackage(nullptr, *SubPackage, 0);
+			check(Package->IsFullyLoaded());
+		}
 
-			FLightingBuildOptions LightingOptions;
+		GRedirectCollector.ResolveStringAssetReference();
 
-			int32 QualityLevel = Quality_Production;
-			GConfig->GetInt(TEXT("LightingBuildOptions"), TEXT("QualityLevel"), QualityLevel, GEditorPerProjectIni);
-			QualityLevel = FMath::Clamp<int32>(QualityLevel, Quality_Preview, Quality_Production);
-			
-			LightingOptions.QualityLevel = (ELightingBuildQuality)QualityLevel;
 
-			GEditor->BuildLighting(LightingOptions);
-			while (GEditor->IsLightingBuildCurrentlyRunning())
-			{
-				GEditor->UpdateBuildLighting();
-			}
+		FLightingBuildOptions LightingOptions;
+		// Always build on production
+		LightingOptions.QualityLevel = Quality_Production;
 
-			World->RemoveFromRoot();
+		GEditor->BuildLighting(LightingOptions);
+		while (GEditor->IsLightingBuildCurrentlyRunning())
+		{
+			GEditor->UpdateBuildLighting();
 		}
 
+		World->RemoveFromRoot();
 	}
 }
 
diff --git a/Engine/Source/Editor/UnrealEd/Private/Fbx/FbxMainImport.cpp b/Engine/Source/Editor/UnrealEd/Private/Fbx/FbxMainImport.cpp
--- a/Engine/Source/Editor/UnrealEd/Private/Fbx/FbxMainImport.cpp
+++ b/Engine/Source/Editor/UnrealEd/Private/Fbx/FbxMainImport.cpp
@@ -491,30 +491,33 @@ bool FFbxImporter::GetSceneInfo(FString Filename, FbxSceneInfo& SceneInfo)
 				FbxNode* GeoNode = Geometry->GetNode();
 				SceneInfo.TotalGeometryNum++;
 				FbxMesh* Mesh = (FbxMesh*)Geometry;
-				SceneInfo.MeshInfo.AddUninitialized();
+				SceneInfo.MeshInfo.AddZeroed(1);
 				FbxMeshInfo& MeshInfo = SceneInfo.MeshInfo.Last();
 				if(Geometry->GetName()[0] != '\0')
 					MeshInfo.Name = MakeName(Geometry->GetName());
 				else
-					MeshInfo.Name = MakeName(GeoNode->GetName());
+					MeshInfo.Name = MakeString(GeoNode ? GeoNode->GetName() : "None");
 				MeshInfo.bTriangulated = Mesh->IsTriangleMesh();
-				MeshInfo.MaterialNum = GeoNode->GetMaterialCount();
+				MeshInfo.MaterialNum = GeoNode? GeoNode->GetMaterialCount() : 0;
 				MeshInfo.FaceNum = Mesh->GetPolygonCount();
 				MeshInfo.VertexNum = Mesh->GetControlPointsCount();
 				
 				// LOD info
 				MeshInfo.LODGroup = NULL;
-				FbxNode* ParentNode = GeoNode->GetParent();
-				if ( ParentNode->GetNodeAttribute() && ParentNode->GetNodeAttribute()->GetAttributeType() == FbxNodeAttribute::eLODGroup)
+				if (GeoNode)
 				{
-					FbxNodeAttribute* LODGroup = ParentNode->GetNodeAttribute();
-					MeshInfo.LODGroup = MakeName(ParentNode->GetName());
-					for (int32 LODIndex = 0; LODIndex < ParentNode->GetChildCount(); LODIndex++)
+					FbxNode* ParentNode = GeoNode->GetParent();
+					if (ParentNode->GetNodeAttribute() && ParentNode->GetNodeAttribute()->GetAttributeType() == FbxNodeAttribute::eLODGroup)
 					{
-						if ( GeoNode == ParentNode->GetChild(LODIndex))
+						FbxNodeAttribute* LODGroup = ParentNode->GetNodeAttribute();
+						MeshInfo.LODGroup = MakeString(ParentNode->GetName());
+						for (int32 LODIndex = 0; LODIndex < ParentNode->GetChildCount(); LODIndex++)
 						{
-							MeshInfo.LODLevel = LODIndex;
-							break;
+							if (GeoNode == ParentNode->GetChild(LODIndex))
+							{
+								MeshInfo.LODLevel = LODIndex;
+								break;
+							}
 						}
 					}
 				}
@@ -544,7 +547,7 @@ bool FFbxImporter::GetSceneInfo(FString Filename, FbxSceneInfo& SceneInfo)
 						}
 					}
 
-					MeshInfo.SkeletonRoot = Link ? MakeName(Link->GetName()) : MakeName("None");
+					MeshInfo.SkeletonRoot = MakeString(Link ? Link->GetName() : ("None"));
 					MeshInfo.SkeletonElemNum = Link ? Link->GetChildCount(true) : 0;
 
 					if (Link)
@@ -990,6 +993,11 @@ ANSICHAR* FFbxImporter::MakeName(const ANSICHAR* Name)
 	return TmpName;
 }
 
+FString FFbxImporter::MakeString(const ANSICHAR* Name)
+{
+	return FString(ANSI_TO_TCHAR(Name));
+}
+
 FName FFbxImporter::MakeNameForMesh(FString InName, FbxObject* FbxObject)
 {
 	FName OutputName;
diff --git a/Engine/Source/Editor/UnrealEd/Private/FbxImporter.h b/Engine/Source/Editor/UnrealEd/Private/FbxImporter.h
--- a/Engine/Source/Editor/UnrealEd/Private/FbxImporter.h
+++ b/Engine/Source/Editor/UnrealEd/Private/FbxImporter.h
@@ -182,16 +182,16 @@ struct FBXImportOptions
 
 struct FbxMeshInfo
 {
-	const char* Name;
+	FString Name;
 	uint64 UniqueId;
 	int32 FaceNum;
 	int32 VertexNum;
 	bool bTriangulated;
 	int32 MaterialNum;
 	bool bIsSkelMesh;
-	const char* SkeletonRoot;
+	FString SkeletonRoot;
 	int32 SkeletonElemNum;
-	const char* LODGroup;
+	FString LODGroup;
 	int32 LODLevel;
 	int32 MorphNum;
 };
@@ -514,6 +514,7 @@ public:
 
 	//help
 	ANSICHAR* MakeName(const ANSICHAR* name);
+	FString MakeString(const ANSICHAR* Name);
 	FName MakeNameForMesh(FString InName, FbxObject* FbxObject);
 
 	// meshes
diff --git a/Engine/Source/Editor/UnrealEd/Private/StaticLightingSystem/StaticLightingSystem.cpp b/Engine/Source/Editor/UnrealEd/Private/StaticLightingSystem/StaticLightingSystem.cpp
--- a/Engine/Source/Editor/UnrealEd/Private/StaticLightingSystem/StaticLightingSystem.cpp
+++ b/Engine/Source/Editor/UnrealEd/Private/StaticLightingSystem/StaticLightingSystem.cpp
@@ -407,7 +407,7 @@ bool FStaticLightingSystem::BeginLightmassProcess()
 			}
 		}
 
-		if (SkippedLevels.Len() > 0)
+		if (SkippedLevels.Len() > 0 && !IsRunningCommandlet())
 		{
 			// Warn when some levels are not visible and therefore will not be built, because that indicates that only a partial build will be done,
 			// Lighting will still be unbuilt for some areas when playing through the level.
@@ -827,7 +827,7 @@ void FStaticLightingSystem::GatherStaticLightingInfo(bool bRebuildDirtyGeometryF
 			if (bRebuildDirtyGeometryForLighting)
 			{
 				// This will go ahead and clean up lighting on all dirty levels (not just this one)
-				UE_LOG(LogStaticLightingSystem, Warning, TEXT("WARNING: Lighting build automatically rebuilding geometry."));
+				UE_LOG(LogStaticLightingSystem, Warning, TEXT("WARNING: Lighting build automatically rebuilding geometry.") );
 				GEditor->Exec(World, TEXT("MAP REBUILD ALLDIRTYFORLIGHTING"));
 			}
 		}
@@ -2181,7 +2181,7 @@ void FStaticLightingSystem::UpdateLightingBuild()
 	}
 	else if ( CurrentBuildStage == FStaticLightingSystem::AutoApplyingImport )
 	{
-		if ( CanAutoApplyLighting() )
+		if ( CanAutoApplyLighting() || IsRunningCommandlet() )
 		{
 			bool bAutpApplyFailed = false;
 			FStaticLightingManager::Get()->SendBuildDoneNotification( bAutpApplyFailed );
diff --git a/Engine/Source/Editor/UnrealEd/Private/Tests/EditorAutomationTests.cpp b/Engine/Source/Editor/UnrealEd/Private/Tests/EditorAutomationTests.cpp
--- a/Engine/Source/Editor/UnrealEd/Private/Tests/EditorAutomationTests.cpp
+++ b/Engine/Source/Editor/UnrealEd/Private/Tests/EditorAutomationTests.cpp
@@ -276,7 +276,7 @@ bool FLoadAllMapsInEditorTest::RunTest(const FString& Parameters)
 	if (FEngineAnalytics::IsAvailable())
 	{
 		FEngineAnalytics::GetProvider().RecordEvent( TEXT("Editor.Usage.TestEvent"));
-		UE_LOG(LogAnalytics, Log, TEXT("AnayticsTest: Load All Maps automation triggered and Editor.Usage.TestEvent analytic event has been fired."));
+		UE_LOG(LogEditorAutomationTests, Log, TEXT("AnayticsTest: Load All Maps automation triggered and Editor.Usage.TestEvent analytic event has been fired."));
 	}
 	
 
diff --git a/Engine/Source/Programs/AutomationTool/AutomationUtils/CommandletUtils.cs b/Engine/Source/Programs/AutomationTool/AutomationUtils/CommandletUtils.cs
--- a/Engine/Source/Programs/AutomationTool/AutomationUtils/CommandletUtils.cs
+++ b/Engine/Source/Programs/AutomationTool/AutomationUtils/CommandletUtils.cs
@@ -106,7 +106,7 @@ namespace AutomationTool
 				MapsToRebuildLighting.Trim();
 			}
 
-			RunCommandlet(ProjectName, UE4Exe, "ResavePackages", String.Format("-buildlighting -AllowCommandletRendering {0} {1}", MapsToRebuildLighting, Parameters));
+			RunCommandlet(ProjectName, UE4Exe, "ResavePackages", String.Format("-buildlighting -MapsOnly -ProjectOnly -AllowCommandletRendering {0} {1}", MapsToRebuildLighting, Parameters));
 		}
 
         /// <summary>
diff --git a/Engine/Source/Programs/AutomationTool/AutomationUtils/Utils.cs b/Engine/Source/Programs/AutomationTool/AutomationUtils/Utils.cs
--- a/Engine/Source/Programs/AutomationTool/AutomationUtils/Utils.cs
+++ b/Engine/Source/Programs/AutomationTool/AutomationUtils/Utils.cs
@@ -967,7 +967,7 @@ namespace AutomationTool
         {
             try
             {
-				if (bAllowNoVersion && versionString.StartsWith("++depot"))
+				if (bAllowNoVersion && versionString.StartsWith("++"))
 				{
 					// This form of version is used when a product has no major.minor.patch version
 					// E.g. ++depot+UE4-ProdName-CL-12345678
diff --git a/Engine/Source/Programs/AutomationTool/Scripts/RebuildLightMapsCommand.Automation.cs b/Engine/Source/Programs/AutomationTool/Scripts/RebuildLightMapsCommand.Automation.cs
--- a/Engine/Source/Programs/AutomationTool/Scripts/RebuildLightMapsCommand.Automation.cs
+++ b/Engine/Source/Programs/AutomationTool/Scripts/RebuildLightMapsCommand.Automation.cs
@@ -1,4 +1,4 @@
-﻿// Copyright 1998-2015 Epic Games, Inc. All Rights Reserved.
+// Copyright 1998-2015 Epic Games, Inc. All Rights Reserved.
 using System;
 using System.Collections.Generic;
 using System.IO;
@@ -13,7 +13,7 @@ using UnrealBuildTool;
 /// </summary>
 /// <remarks>
 /// Command line parameters used by this command:
-/// -
+/// -project	Absolute path to a .uproject file
 /// </remarks>
 namespace AutomationScripts.Automation
 {
@@ -21,8 +21,46 @@ namespace AutomationScripts.Automation
 	{
 		#region RebuildLightMaps Command
 
+		/**
+		 * Parse the P4 output for any errors that we really care about.
+		 * e.g. umaps and assets are exclusive checkout files, if we cant check out a map for this reason
+		 *		then we need to stop.
+		 */
+		private bool FoundCheckOutErrorInP4Output(string Output)
+		{
+			bool bHadAnError = false;
+
+			var Lines = Output.Split(new string[] { Environment.NewLine }, StringSplitOptions.RemoveEmptyEntries);
+			foreach (string Line in Lines)
+			{
+				// Check for log spew that matches exclusive checkout failure
+				// http://answers.perforce.com/articles/KB/3114
+				if (Line.Contains("can't edit exclusive file already opened"))
+				{
+					bHadAnError = true;
+					break;
+				}
+			}
+
+			return bHadAnError;
+		}
+
+		/**
+		 * Cleanup anything this build may leave behind
+		 */
+		private void HandleFailure()
+		{
+			if (WorkingCL != -1)
+			{
+				P4.RevertAll(WorkingCL);
+				P4.DeleteChange(WorkingCL);
+			}
+		}
+
 		public override void ExecuteBuild()
 		{
+			Log("********** REBUILD LIGHT MAPS COMMAND STARTED **********");
+
 			var Params = new ProjectParams
 			(
 				Command: this,
@@ -30,8 +68,34 @@ namespace AutomationScripts.Automation
 				RawProjectPath: ProjectPath
 			);
 
+			// Setup our P4 changelist so we can submit the new lightmaps after the commandlet has completed.
+			if (P4Enabled)
+			{
+				WorkingCL = P4.CreateChange(P4Env.Client, String.Format("{0} rebuilding lightmaps from changelist {1}", Params.ShortProjectName, P4Env.Changelist));
+				Log("Working in {0}", WorkingCL);
+				try
+				{
+					string AllMapsWildcardCmdline = String.Format("{0}\\...\\*.umap", Params.RawProjectPath.Directory);
+					string Output;
+					P4.P4Output(out Output, String.Format("edit -c {0} {1}", WorkingCL, AllMapsWildcardCmdline));
 
-			Log("********** REBUILD LIGHT MAPS COMMAND STARTED **********");
+					// We need to ensure that any error in the output log is observed.
+					// P4 is still successful if it manages to run the operation.
+					if (FoundCheckOutErrorInP4Output(Output) == true)
+					{
+						throw new Exception();
+					}
+				}
+				catch (Exception)
+				{
+					HandleFailure();
+					throw new AutomationException("Failed to check out every one of the project maps.");
+				}
+			}
+			else
+			{
+				throw new AutomationException("This command needs to run with P4.");
+			}
 
 			string UE4EditorExe = HostPlatform.Current.GetUE4ExePath(Params.UE4Exe);
 			if (!FileExists(UE4EditorExe))
@@ -39,6 +103,7 @@ namespace AutomationScripts.Automation
 				throw new AutomationException("Missing " + UE4EditorExe + " executable. Needs to be built first.");
 			}
 
+			// Now let's rebuild lightmaps for the project
 			try
 			{
 				var CommandletParams = IsBuildMachine ? "-buildmachine -fileopenlog" : "-fileopenlog";
@@ -46,23 +111,30 @@ namespace AutomationScripts.Automation
 			}
 			catch (Exception Ex)
 			{
-				if (Params.IgnoreLightMapErrors)
-				{
-					LogWarning("Ignoring Rebuild Light Map failure.");
-				}
-				else
-				{
-					// Abandon this run, dont check in any updated files, etc.
-					Log("Rebuild Light Maps has failed.");
-					//CleanupCookedData(PlatformsToCook.ToList(), Params);
-					throw new AutomationException(ExitCode.Error_UnknownCookFailure, Ex, "RebuildLightMaps failed.");
-				}
+				// Something went wrong with the commandlet. Abandon this run, don't check in any updated files, etc.
+				HandleFailure();
+
+				Log("Rebuild Light Maps has failed.");
+				throw new AutomationException(ExitCode.Error_Unknown, Ex, "RebuildLightMaps failed.");
 			}
-			Log("********** COOK COMMAND COMPLETED **********");
+
+			// Check everything in!
+			if (WorkingCL != -1)
+			{
+				int SubmittedCL;
+				P4.Submit(WorkingCL, out SubmittedCL, true, true);
+				Log("New lightmaps have been submitted in changelist {0}", SubmittedCL);
+			}
+
+			Log("********** REBUILD LIGHT MAPS COMMAND COMPLETED **********");
 		}
 
 		#endregion
 
+		// The Chanelist used when doing the work.
+		private int WorkingCL = -1;
+
+		// Process command-line and find a project file. This is necessary for the commandlet to run successfully
 
 		private FileReference ProjectFullPath;
 		public virtual FileReference ProjectPath
diff --git a/Engine/Source/Programs/AutomationTool/Win/Win.Automation.csproj b/Engine/Source/Programs/AutomationTool/Win/Win.Automation.csproj
--- a/Engine/Source/Programs/AutomationTool/Win/Win.Automation.csproj
+++ b/Engine/Source/Programs/AutomationTool/Win/Win.Automation.csproj
@@ -51,7 +51,6 @@
     </Compile>
     <Compile Include="Properties\AssemblyInfo.cs" />
     <Compile Include="WinPlatform.Automation.cs" />
-    <Compile Include="WinPlatformResources.Automation.cs" />
   </ItemGroup>
   <ItemGroup>
     <ProjectReference Include="..\..\UnrealBuildTool\UnrealBuildTool.csproj">
diff --git a/Engine/Source/Programs/AutomationTool/Win/WinPlatform.Automation.cs b/Engine/Source/Programs/AutomationTool/Win/WinPlatform.Automation.cs
--- a/Engine/Source/Programs/AutomationTool/Win/WinPlatform.Automation.cs
+++ b/Engine/Source/Programs/AutomationTool/Win/WinPlatform.Automation.cs
@@ -61,27 +61,6 @@ public abstract class BaseWinPlatform : Platform
 							BootstrapArguments = String.Format("..\\..\\..\\{0}\\{0}.uproject", SC.ShortProjectName);
 						}
 
-
-                        foreach (string StageExePath in GetExecutableNames(SC))
-                        {
-                            // set the icon on the original exe this will be used in the task bar when the bootstrap exe runs
-                            if (InternalUtils.SafeFileExists(CombinePaths(SC.ProjectRoot, "Build/Windows/Application.ico")))
-                            {
-                                GroupIconResource GroupIcon = null;
-                                GroupIcon = GroupIconResource.FromIco(CombinePaths(SC.ProjectRoot, "Build/Windows/Application.ico"));
-
-                                // Update the icon on the original exe because this will be used when the game is running in the task bar
-                                using (ModuleResourceUpdate Update = new ModuleResourceUpdate(StageExePath, false))
-                                {
-                                    const int IconResourceId = 101;
-                                    if (GroupIcon != null)
-                                    {
-                                        Update.SetIcons(IconResourceId, GroupIcon);
-                                    }
-                                }
-                            }
-                        }
-
 						string BootstrapExeName;
 						if(SC.StageTargetConfigurations.Count > 1)
 						{
diff --git a/Engine/Source/Programs/UnrealBuildTool/UnrealBuildTool.csproj b/Engine/Source/Programs/UnrealBuildTool/UnrealBuildTool.csproj
--- a/Engine/Source/Programs/UnrealBuildTool/UnrealBuildTool.csproj
+++ b/Engine/Source/Programs/UnrealBuildTool/UnrealBuildTool.csproj
@@ -134,6 +134,8 @@
     <Compile Include="UWP\UWPDeploy.cs" />
     <Compile Include="UWP\UWPProjectGenerator.cs" />
     <Compile Include="UWP\UWPToolChain.cs" />
+    <Compile Include="Windows\WindowsDeploy.cs" />
+    <Compile Include="Windows\WinResources.cs" />
     <Compile Include="WinRT\UEBuildWinRT.cs" />
     <Compile Include="WinRT\WinRTToolChain.cs" />
     <Compile Include="WinRT\WinRTDeploy.cs" />
diff --git a/Engine/Source/Programs/UnrealBuildTool/Windows/UEBuildWindows.cs b/Engine/Source/Programs/UnrealBuildTool/Windows/UEBuildWindows.cs
--- a/Engine/Source/Programs/UnrealBuildTool/Windows/UEBuildWindows.cs
+++ b/Engine/Source/Programs/UnrealBuildTool/Windows/UEBuildWindows.cs
@@ -428,7 +428,7 @@ namespace UnrealBuildTool
 		/// <returns>New deployment handler, or null if not required</returns>
 		public override UEBuildDeploy CreateDeploymentHandler()
 		{
-			return null;
+			return new BaseWindowsDeploy();
 		}
 	}
 
