Fix shadow variables

diff --git a/Engine/Source/Editor/Sequencer/Private/SequencerContextMenus.cpp b/Engine/Source/Editor/Sequencer/Private/SequencerContextMenus.cpp
--- a/Engine/Source/Editor/Sequencer/Private/SequencerContextMenus.cpp
+++ b/Engine/Source/Editor/Sequencer/Private/SequencerContextMenus.cpp
@@ -138,17 +138,17 @@ void FSectionContextMenu::PopulateMenu(FMenuBuilder& MenuBuilder)
 		MenuBuilder.AddSubMenu(
 			NSLOCTEXT("Sequencer", "EditSection", "Edit"),
 			NSLOCTEXT("Sequencer", "EditSectionTooltip", "Edit section"),
-			FNewMenuDelegate::CreateLambda([=](FMenuBuilder& MenuBuilder){ Shared->AddEditMenu(MenuBuilder); }));
+			FNewMenuDelegate::CreateLambda([=](FMenuBuilder& InMenuBuilder){ Shared->AddEditMenu(InMenuBuilder); }));
 
 		MenuBuilder.AddSubMenu(
 			NSLOCTEXT("Sequencer", "SetPreInfinityExtrap", "Pre-Infinity"), 
 			NSLOCTEXT("Sequencer", "SetPreInfinityExtrapTooltip", "Set pre-infinity extrapolation"),
-			FNewMenuDelegate::CreateLambda([=](FMenuBuilder& MenuBuilder){ Shared->AddExtrapolationMenu(MenuBuilder, true); }));
+			FNewMenuDelegate::CreateLambda([=](FMenuBuilder& InMenuBuilder){ Shared->AddExtrapolationMenu(InMenuBuilder, true); }));
 
 		MenuBuilder.AddSubMenu(
 			NSLOCTEXT("Sequencer", "SetPostInfinityExtrap", "Post-Infinity"), 
 			NSLOCTEXT("Sequencer", "SetPostInfinityExtrapTooltip", "Set post-infinity extrapolation"),
-			FNewMenuDelegate::CreateLambda([=](FMenuBuilder& MenuBuilder){ Shared->AddExtrapolationMenu(MenuBuilder, false); }));
+			FNewMenuDelegate::CreateLambda([=](FMenuBuilder& InMenuBuilder){ Shared->AddExtrapolationMenu(InMenuBuilder, false); }));
 
 		MenuBuilder.AddMenuEntry(
 			NSLOCTEXT("Sequencer", "ToggleSectionActive", "Active"),
