Copying //UE4/Orion-Staging to Dev-Main (//UE4/Dev-Main)

Orion changes originate from CL 2763274

diff --git a/Engine/Plugins/MovieScene/LevelSequenceEditor/Source/LevelSequenceEditor/Private/LevelSequenceEditorToolkit.cpp b/Engine/Plugins/MovieScene/LevelSequenceEditor/Source/LevelSequenceEditor/Private/LevelSequenceEditorToolkit.cpp
--- a/Engine/Plugins/MovieScene/LevelSequenceEditor/Source/LevelSequenceEditor/Private/LevelSequenceEditorToolkit.cpp
+++ b/Engine/Plugins/MovieScene/LevelSequenceEditor/Source/LevelSequenceEditor/Private/LevelSequenceEditorToolkit.cpp
@@ -217,7 +217,7 @@ void FLevelSequenceEditorToolkit::HandleAddComponentMaterialActionExecute( UPrim
 			for ( UMovieSceneTrack* Track : Binding.GetTracks() )
 			{
 				UMovieSceneComponentMaterialTrack* MaterialTrack = Cast<UMovieSceneComponentMaterialTrack>( Track );
-				if ( MaterialTrack->GetMaterialIndex() == MaterialIndex )
+				if ( MaterialTrack != nullptr && MaterialTrack->GetMaterialIndex() == MaterialIndex )
 				{
 					bHasMaterialTrack = true;
 					break;
diff --git a/Engine/Shaders/PostProcessCircleDOF.usf b/Engine/Shaders/PostProcessCircleDOF.usf
--- a/Engine/Shaders/PostProcessCircleDOF.usf
+++ b/Engine/Shaders/PostProcessCircleDOF.usf
@@ -378,6 +378,10 @@ void Circle4Samples(float4 UVAndScreenPos : TEXCOORD0, out float4 OutColor0 : SV
 	// Set current result as possible background.
 	float3 Background = OutColor0.rgb * (1.0/Weight);
 
+	// former method
+//	 #define FadeOutOutsideCoC(INDEX, COC) C##INDEX.rgb = lerp(C##INDEX.rgb, Background.rgb, saturate(abs(C##INDEX.a) - COC));
+	// new method avoids having the center leaking with large CoC
+#define FadeOutOutsideCoC(INDEX, COC) W##INDEX = lerp(W##INDEX, 0, saturate(abs(C##INDEX.a) - COC));
 
 	// Uncomment to see intermediate debug output
 //	OutColor0.rgb *= (1.0/Weight);return; 
@@ -418,10 +422,10 @@ void Circle4Samples(float4 UVAndScreenPos : TEXCOORD0, out float4 OutColor0 : SV
 	C4 = PostprocessInput0.SampleLevel(PostprocessInput0Sampler, UV4, 0);
 
 	// Lerp to background if outside possibly smaller CoC.
-	C1.rgb = lerp(C1.rgb, Background.rgb, saturate(abs(C1.a) - Coc2));
-	C2.rgb = lerp(C2.rgb, Background.rgb, saturate(abs(C2.a) - Coc2));
-	C3.rgb = lerp(C3.rgb, Background.rgb, saturate(abs(C3.a) - Coc2));
-	C4.rgb = lerp(C4.rgb, Background.rgb, saturate(abs(C4.a) - Coc2));
+	FadeOutOutsideCoC(1, Coc2)
+	FadeOutOutsideCoC(2, Coc2)
+	FadeOutOutsideCoC(3, Coc2)
+	FadeOutOutsideCoC(4, Coc2)
 
 	W1 = R1;
 	W2 = R2;
@@ -490,10 +494,10 @@ void Circle4Samples(float4 UVAndScreenPos : TEXCOORD0, out float4 OutColor0 : SV
 	C3 = PostprocessInput0.SampleLevel(PostprocessInput0Sampler, UV3, 0);
 	C4 = PostprocessInput0.SampleLevel(PostprocessInput0Sampler, UV4, 0);
 
-	C1.rgb = lerp(C1.rgb, Background.rgb, saturate(abs(C1.a) - Coc3));
-	C2.rgb = lerp(C2.rgb, Background.rgb, saturate(abs(C2.a) - Coc3));
-	C3.rgb = lerp(C3.rgb, Background.rgb, saturate(abs(C3.a) - Coc3));
-	C4.rgb = lerp(C4.rgb, Background.rgb, saturate(abs(C4.a) - Coc3));
+	FadeOutOutsideCoC(1, Coc3)
+	FadeOutOutsideCoC(2, Coc3)
+	FadeOutOutsideCoC(3, Coc3)
+	FadeOutOutsideCoC(4, Coc3)
 
 	W1 = R1;
 	W2 = R2;
diff --git a/Engine/Shaders/ShadingModels.usf b/Engine/Shaders/ShadingModels.usf
--- a/Engine/Shaders/ShadingModels.usf
+++ b/Engine/Shaders/ShadingModels.usf
@@ -667,32 +667,24 @@ float3 SubsurfaceShadingTwoSided( float3 SubsurfaceColor, float3 L, float3 V, ha
 float3 EyeShading( FGBufferData GBuffer, float3 LobeRoughness, float3 LobeEnergy, float3 L, float3 V, half3 N )
 {
 	float3 H = normalize(V + L);
+	float NoL = saturate( dot(N, L) );
+	float NoV = max( dot(N, V), 1e-5 );	// constant to prevent NaN
+	float NoH = saturate( dot(N, H) );
 	float VoH = saturate( dot(V, H) );
-
-	float3 F = F_Schlick( GBuffer.SpecularColor, VoH );
-
-	//
-	// Cornea Reflectance - Use geo normal
-	//
-	float CNoL = saturate( dot(N, L) );
-	float CNoV = max( dot(N, V), 1e-5 );	// constant to prevent NaN
-	float CNoH = saturate( dot(N, H) );
 	
 	// Generalized microfacet specular
-	float D = D_GGX( LobeRoughness[1], CNoH ) * LobeEnergy[1];
-	float Vis = Vis_SmithJointApprox( LobeRoughness[1], CNoV, CNoL );
-
-	// Cornea Specular - Covers the whole surface
-	float3 CorneaSpecular;
-	CorneaSpecular = D * Vis * F;
+	float D = D_GGX( LobeRoughness[1], NoH ) * LobeEnergy[1];
+	float Vis = Vis_SmithJointApprox( LobeRoughness[1], NoV, NoL );
+	float3 F = F_Schlick( GBuffer.SpecularColor, VoH );
 
-	return CorneaSpecular;
+	return D * Vis * F;
 }
 
-float3 EyeSubsurfaceShading(  FGBufferData GBuffer, float3 L, float3 V, half3 N )
+float3 EyeSubsurfaceShading( FGBufferData GBuffer, float3 L, float3 V, half3 N )
 {
 	float3 H = normalize(V + L);
 	float VoH = saturate( dot(V, H) );
+	float NoV = max( dot(N, V), 1e-5 );
 
 	// F_Schlick
 	float F0 = GBuffer.SpecularColor;
@@ -709,20 +701,18 @@ float3 EyeSubsurfaceShading(  FGBufferData GBuffer, float3 L, float3 V, half3 N
 	// Not great as it ties the concavity to the convexity of the cornea surface
 	// No good justification for that. On the other hand, if we're just looking to
 	// introduce some concavity, this does the job.
-	IrisNormal = normalize(lerp(IrisNormal, -N, IrisMask*IrisDistance));
+	float3 CausticNormal = normalize(lerp(IrisNormal, -N, IrisMask*IrisDistance));
 
-	// Use IrisNormal only within the Iris
-	IrisNormal = lerp(N, IrisNormal, IrisMask);
-	IrisNormal = normalize( IrisNormal );
+	float NoL = saturate( dot( IrisNormal, L ) );
+	float Power = lerp( 12, 1, NoL );
+	float Caustic = 0.5 + 0.5 * ( Power + 1 ) * pow( saturate( dot( CausticNormal, L ) ), Power );
+	float Iris = NoL * Caustic;
 
 	// http://blog.stevemcauley.com/2011/12/03/energy-conserving-wrapped-diffuse/
-	float Wrap = 0.1;
-	float INoL = saturate( ( dot(IrisNormal, L) + Wrap ) / Square( 1 + Wrap ) );
+	float Wrap = 0.15;
+	float Sclera = saturate( ( dot(N, L) + Wrap ) / Square( 1 + Wrap ) );
 	
-	float3 rawResult = (1 - F) * INoL * GBuffer.DiffuseColor / PI;
-
-	// Covers whole object
-	return rawResult;
+	return (1 - F) * lerp( Sclera, Iris, IrisMask ) * GBuffer.DiffuseColor / PI;
 }
 
 Texture2D		PreIntegratedBRDF;
diff --git a/Engine/Shaders/SlateElementPixelShader.usf b/Engine/Shaders/SlateElementPixelShader.usf
--- a/Engine/Shaders/SlateElementPixelShader.usf
+++ b/Engine/Shaders/SlateElementPixelShader.usf
@@ -310,7 +310,7 @@ float4 Main( VertexToPixelInterpolants VIn ) : SV_Target0
 #endif
 
 
-#if !PRE_MULTIPLIED_ALPHA
+#if !PRE_MULTIPLIED_ALPHA && !MATERIALBLENDING_ADDITIVE && !MATERIALBLENDING_TRANSLUCENT
 	OutColor.rgb = OutColor.rgb * OutColor.a;
 #endif
 
diff --git a/Engine/Source/Developer/FunctionalTesting/Classes/FunctionalTest.h b/Engine/Source/Developer/FunctionalTesting/Classes/FunctionalTest.h
--- a/Engine/Source/Developer/FunctionalTesting/Classes/FunctionalTest.h
+++ b/Engine/Source/Developer/FunctionalTesting/Classes/FunctionalTest.h
@@ -92,16 +92,26 @@ struct FUNCTIONALTESTING_API FPerfStatsRecord
 	FStatsData Record;
 	/** Stats data for the baseline. */
 	FStatsData Baseline;
+	
+	float GPUBudget;
+	float RenderThreadBudget;
+	float GameThreadBudget;
 
+	void SetBudgets(float InGPUBudget, float InRenderThreadBudget, float InGameThreadBudget);
 	void Sample(UWorld* Owner, float DeltaSeconds, bool bBaseline);
 
 	FString GetReportString()const;
 	FString GetBaselineString()const;
 	FString GetRecordString()const;
+	FString GetOverBudgetString()const;
 
 	void GetGPUTimes(double& OutMin, double& OutMax, double& OutAvg)const;
 	void GetGameThreadTimes(double& OutMin, double& OutMax, double& OutAvg)const;
 	void GetRenderThreadTimes(double& OutMin, double& OutMax, double& OutAvg)const;
+
+	bool IsWithinGPUBudget()const;
+	bool IsWithinGameThreadBudget()const;
+	bool IsWithinRenderThreadBudget()const;
 };
 
 
@@ -141,7 +151,7 @@ public:
 	void EndRecordingBaseline();
 	/** Begins recording a new named performance stats record. We start by recording the baseline. */
 	UFUNCTION(BlueprintCallable, Category = Perf)
-	void BeginRecording(FString RecordName);
+	void BeginRecording(FString RecordName,	float InGPUBudget, float InRenderThreadBudget, float InGameThreadBudget);
 	/** Stops recording performance stats. */
 	UFUNCTION(BlueprintCallable, Category = Perf)
 	void EndRecording();
@@ -160,13 +170,14 @@ public:
 	void OnAllTestsComplete();
 
 	const FPerfStatsRecord* GetCurrentRecord()const;
+	FPerfStatsRecord* GetCurrentRecord();
 
 	UFUNCTION(BlueprintCallable, Category = Perf)
-	bool IsCurrentRecordWithinGPUBudget(float InMaxTimeBudget, float InAvgTimeBudget = -1.0f)const;
+	bool IsCurrentRecordWithinGPUBudget()const;
 	UFUNCTION(BlueprintCallable, Category = Perf)
-	bool IsCurrentRecordWithinGameThreadBudget(float InMaxTimeBudget, float InAvgTimeBudget = -1.0f)const;
+	bool IsCurrentRecordWithinGameThreadBudget()const;
 	UFUNCTION(BlueprintCallable, Category = Perf)
-	bool IsCurrentRecordWithinRenderThreadBudget(float InMaxTimeBudget, float InAvgTimeBudget = -1.0f)const;
+	bool IsCurrentRecordWithinRenderThreadBudget()const;
 	//End basic stats recording.
 
 	// Automatic traces capturing 
diff --git a/Engine/Source/Developer/FunctionalTesting/Private/FunctionalTest.cpp b/Engine/Source/Developer/FunctionalTesting/Private/FunctionalTest.cpp
--- a/Engine/Source/Developer/FunctionalTesting/Private/FunctionalTest.cpp
+++ b/Engine/Source/Developer/FunctionalTesting/Private/FunctionalTest.cpp
@@ -330,9 +330,19 @@ UBillboardComponent* AFunctionalTest::GetSpriteComponent() { return SpriteCompon
 
 FPerfStatsRecord::FPerfStatsRecord(FString InName)
 : Name(InName)
+, GPUBudget(0.0f)
+, RenderThreadBudget(0.0f)
+, GameThreadBudget(0.0f)
 {
 }
 
+void FPerfStatsRecord::SetBudgets(float InGPUBudget, float InRenderThreadBudget, float InGameThreadBudget)
+{
+	GPUBudget = InGPUBudget;
+	RenderThreadBudget = InRenderThreadBudget;
+	GameThreadBudget = InGameThreadBudget;
+}
+
 FString FPerfStatsRecord::GetReportString() const
 {
 	return FString::Printf(TEXT("%s,%.2f,%.2f,%.2f,%.2f,%.2f,%.2f,%.2f,%.2f,%.2f,%.2f,%.2f,%.2f"),
@@ -387,6 +397,51 @@ FString FPerfStatsRecord::GetRecordString() const
 		Record.GPUTimeTracker.GetMaxValue());
 }
 
+FString FPerfStatsRecord::GetOverBudgetString() const
+{
+	double Min, Max, Avg;
+	GetRenderThreadTimes(Min, Max, Avg);
+	float RTBudgetFrac = Max / RenderThreadBudget;
+	GetGameThreadTimes(Min, Max, Avg);
+	float GTBudgetFrac = Max / GameThreadBudget;
+	GetGPUTimes(Min, Max, Avg);
+	float GPUBudgetFrac = Max / GPUBudget;
+
+	return FString::Printf(TEXT("%s,%.2f,%.2f,%.2f,%.2f,%.2f,%.2f,%.2f,%.2f,%.2f"),
+		*Name,
+		Record.RenderThreadTimeTracker.GetMaxValue(),
+		RenderThreadBudget,
+		RTBudgetFrac,
+		Record.GameThreadTimeTracker.GetMaxValue(),
+		GameThreadBudget,
+		GTBudgetFrac,
+		Record.GPUTimeTracker.GetMaxValue(),
+		GPUBudget,
+		GPUBudgetFrac
+		);
+}
+
+bool FPerfStatsRecord::IsWithinGPUBudget()const
+{
+	double Min, Max, Avg;
+	GetGPUTimes(Min, Max, Avg);
+	return Max <= GPUBudget;
+}
+
+bool FPerfStatsRecord::IsWithinGameThreadBudget()const
+{
+	double Min, Max, Avg;
+	GetGameThreadTimes(Min, Max, Avg);
+	return Max <= GameThreadBudget;
+}
+
+bool FPerfStatsRecord::IsWithinRenderThreadBudget()const
+{
+	double Min, Max, Avg;
+	GetRenderThreadTimes(Min, Max, Avg);
+	return Max <= RenderThreadBudget;
+}
+
 void FPerfStatsRecord::GetGPUTimes(double& OutMin, double& OutMax, double& OutAvg)const
 {
 	OutMin = Record.GPUTimeTracker.GetMinValue() - Baseline.GPUTimeTracker.GetMinValue();
@@ -457,20 +512,22 @@ void UAutomationPerformaceHelper::EndRecordingBaseline()
 	bRecordingBasicStats = false;
 }
 
-void UAutomationPerformaceHelper::BeginRecording(FString RecordName)
+void UAutomationPerformaceHelper::BeginRecording(FString RecordName, float InGPUBudget, float InRenderThreadBudget, float InGameThreadBudget)
 {
 	//Ensure we're recording engine stats.
 	GEngine->SetEngineStat(GetOuter()->GetWorld(), GetOuter()->GetWorld()->GetGameViewport(), TEXT("Unit"), true);
 	bRecordingBasicStats = true;
 	bRecordingBaselineBasicStats = false;
-	//We recorded a baseline for this so continue adding real stats to this one.
-	if (GetCurrentRecord()->Name != RecordName)
+
+	FPerfStatsRecord* CurrRecord = GetCurrentRecord();
+	if (!CurrRecord || CurrRecord->Name != RecordName)
 	{
 		Records.Add(FPerfStatsRecord(RecordName));
+		CurrRecord = GetCurrentRecord();
 	}
-	
-	//We should have a previously baselined record or have just added one.
-	check(Records.Num() > 0);
+
+	check(CurrRecord);
+	CurrRecord->SetBudgets(InGPUBudget, InRenderThreadBudget, InGameThreadBudget);
 }
 
 void UAutomationPerformaceHelper::EndRecording()
@@ -518,26 +575,27 @@ void UAutomationPerformaceHelper::WriteLogFile(const FString& CaptureDir, const
 
 	const FString Filename = OutputFileBase + Extension;
 	const FString FilenameFull = PathName + Filename;
-
-	const FString LogHeader = TEXT("TestName,MinFrameTime,AvgFrameTime,MaxFrameTime,MinRT,AvgRT,MaxRT,MinGT,AvgGT,MaxGT,MinGPU,AvgGPU,MaxGPU\n");
-
+	
+	const FString OverBudgetTableHeader = TEXT("TestName, MaxRT, RT Budget, RT Frac, MaxGT, GT Budget, GT Frac, MaxGPU, GPU Budget, GPU Frac\n");
+	FString OverbudgetTable;
+	const FString DataTableHeader = TEXT("TestName,MinFrameTime,AvgFrameTime,MaxFrameTime,MinRT,AvgRT,MaxRT,MinGT,AvgGT,MaxGT,MinGPU,AvgGPU,MaxGPU\n");
 	FString AdjustedTable;
-	for (FPerfStatsRecord& Record : Records)
-	{
-		AdjustedTable += Record.GetReportString() + FString(TEXT("\n"));
-	}
 	FString RecordTable;
-	for (FPerfStatsRecord& Record : Records)
-	{
-		RecordTable += Record.GetRecordString() + FString(TEXT("\n"));
-	}
 	FString BaselineTable;
 	for (FPerfStatsRecord& Record : Records)
 	{
+		AdjustedTable += Record.GetReportString() + FString(TEXT("\n"));
+		RecordTable += Record.GetRecordString() + FString(TEXT("\n"));
 		BaselineTable += Record.GetBaselineString() + FString(TEXT("\n"));
+
+		if (!Record.IsWithinGPUBudget() || !Record.IsWithinRenderThreadBudget() || !Record.IsWithinGameThreadBudget())
+		{
+			OverbudgetTable += Record.GetOverBudgetString() + FString(TEXT("\n"));
+		}
 	}
 
-	FString FileContents = FString::Printf(TEXT("Adjusted Results\n%s%s\nRaw Results\n%s%s\nBaseline Results\n%s%s\n"), *LogHeader, *AdjustedTable, *LogHeader, *RecordTable, *LogHeader, *BaselineTable);
+	FString FileContents = FString::Printf(TEXT("Over Budget Tests\n%s%s\nAdjusted Results\n%s%s\nRaw Results\n%s%s\nBaseline Results\n%s%s\n"), 
+		*OverBudgetTableHeader, *OverbudgetTable, *DataTableHeader, *AdjustedTable, *DataTableHeader, *RecordTable, *DataTableHeader, *BaselineTable);
 
 	FFileHelper::SaveStringToFile(FileContents, *FilenameFull);
 
@@ -587,40 +645,31 @@ void UAutomationPerformaceHelper::OnAllTestsComplete()
 	}
 }
 
-bool UAutomationPerformaceHelper::IsCurrentRecordWithinGPUBudget(float InMaxTimeBudget, float InAvgTimeBudget)const
+bool UAutomationPerformaceHelper::IsCurrentRecordWithinGPUBudget()const
 {
-	double Min, Max, Avg;
-	bool bInBudget = true;
 	if (const FPerfStatsRecord* Curr = GetCurrentRecord())
 	{
-		Curr->GetGPUTimes(Min, Max, Avg);
-		bInBudget = Max < InMaxTimeBudget && (InAvgTimeBudget < 0.0f || Avg < InAvgTimeBudget);
+		return Curr->IsWithinGPUBudget();
 	}
-	return bInBudget;
+	return true;
 }
 
-bool UAutomationPerformaceHelper::IsCurrentRecordWithinGameThreadBudget(float InMaxTimeBudget, float InAvgTimeBudget)const
+bool UAutomationPerformaceHelper::IsCurrentRecordWithinGameThreadBudget()const
 {
-	double Min, Max, Avg;
-	bool bInBudget = true;
 	if (const FPerfStatsRecord* Curr = GetCurrentRecord())
 	{
-		Curr->GetGameThreadTimes(Min, Max, Avg);
-		bInBudget = Max < InMaxTimeBudget && (InAvgTimeBudget < 0.0f || Avg < InAvgTimeBudget);
+		return Curr->IsWithinGameThreadBudget();
 	}
-	return bInBudget;
+	return true;
 }
 
-bool UAutomationPerformaceHelper::IsCurrentRecordWithinRenderThreadBudget(float InMaxTimeBudget, float InAvgTimeBudget)const
+bool UAutomationPerformaceHelper::IsCurrentRecordWithinRenderThreadBudget()const
 {
-	double Min, Max, Avg;
-	bool bInBudget = true;
 	if (const FPerfStatsRecord* Curr = GetCurrentRecord())
 	{
-		Curr->GetRenderThreadTimes(Min, Max, Avg);
-		bInBudget = Max < InMaxTimeBudget && (InAvgTimeBudget < 0.0f || Avg < InAvgTimeBudget);
+		return Curr->IsWithinRenderThreadBudget();
 	}
-	return bInBudget;
+	return true;
 }
 
 const FPerfStatsRecord* UAutomationPerformaceHelper::GetCurrentRecord()const
@@ -632,6 +681,15 @@ const FPerfStatsRecord* UAutomationPerformaceHelper::GetCurrentRecord()const
 	}
 	return nullptr;
 }
+FPerfStatsRecord* UAutomationPerformaceHelper::GetCurrentRecord()
+{
+	int32 Index = Records.Num() - 1;
+	if (Index >= 0)
+	{
+		return &Records[Index];
+	}
+	return nullptr;
+}
 
 void UAutomationPerformaceHelper::StartCPUProfiling()
 {
diff --git a/Engine/Source/Editor/EditorStyle/Private/SlateEditorStyle.cpp b/Engine/Source/Editor/EditorStyle/Private/SlateEditorStyle.cpp
--- a/Engine/Source/Editor/EditorStyle/Private/SlateEditorStyle.cpp
+++ b/Engine/Source/Editor/EditorStyle/Private/SlateEditorStyle.cpp
@@ -1595,6 +1595,11 @@ void FSlateEditorStyle::FStyle::SetupGeneralStyles()
 		Set( "Sequencer.NotificationImage_AddedPlayMovieSceneEvent", new IMAGE_BRUSH( "Old/Checkbox_checked", Icon16x16 ) );
 		Set( "Sequencer.Save", new IMAGE_BRUSH( "Icons/SequencerIcons/icon_Sequencer_Save_24x", Icon40x40 ) );
 		Set( "Sequencer.Save.Small", new IMAGE_BRUSH( "Icons/SequencerIcons/icon_Sequencer_Save_24x", Icon20x20 ) );
+		Set( "Sequencer.IconKeyAuto", new IMAGE_BRUSH( "Sequencer/IconKeyAuto", Icon12x12 ) );
+		Set( "Sequencer.IconKeyBreak", new IMAGE_BRUSH( "Sequencer/IconKeyBreak", Icon12x12 ) );
+		Set( "Sequencer.IconKeyConstant", new IMAGE_BRUSH( "Sequencer/IconKeyConstant", Icon12x12 ) );
+		Set( "Sequencer.IconKeyLinear", new IMAGE_BRUSH( "Sequencer/IconKeyLinear", Icon12x12 ) );
+		Set( "Sequencer.IconKeyUser", new IMAGE_BRUSH( "Sequencer/IconKeyUser", Icon12x12 ) );
 		Set( "Sequencer.KeyCircle", new IMAGE_BRUSH( "Sequencer/KeyCircle", Icon12x12 ) );
 		Set( "Sequencer.KeyDiamond", new IMAGE_BRUSH( "Sequencer/KeyDiamond", Icon12x12 ) );
 		Set( "Sequencer.KeySquare", new IMAGE_BRUSH( "Sequencer/KeySquare", Icon12x12 ) );
@@ -1630,6 +1635,7 @@ void FSlateEditorStyle::FStyle::SetupGeneralStyles()
 		Set( "Sequencer.LockCamera", new IMAGE_BRUSH( "Icons/icon_ToggleActorPilotCameraView_16x", Icon16x16 ) );
 		Set( "Sequencer.UnlockCamera", new IMAGE_BRUSH( "Icons/icon_ToggleActorPilotCameraView_16x", Icon16x16, FLinearColor(1.f, 1.f, 1.f, 0.5f) ) );
 		Set( "Sequencer.Section.SelectionBorder", new BORDER_BRUSH( TEXT("Sequencer/SectionHighlight"), FMargin(7.f/16.f) ) );
+		Set( "Sequencer.Section.LockedBorder", new BORDER_BRUSH( TEXT("Sequencer/SectionLocked"), FMargin(7.f/16.f) ) );
 		Set( "Sequencer.ShotTrack.SectionHandle", new IMAGE_BRUSH( "Old/White", Icon16x16, FLinearColor::Black ) );
 		Set( "Sequencer.General", new IMAGE_BRUSH( "Icons/SequencerIcons/icon_Sequencer_View_24x", Icon24x24 ) );
 		Set( "Sequencer.General.Small", new IMAGE_BRUSH( "Icons/SequencerIcons/icon_Sequencer_View_24x", Icon20x20 ) );
@@ -1643,7 +1649,6 @@ void FSlateEditorStyle::FStyle::SetupGeneralStyles()
 		Set( "Sequencer.Tracks.Shot", new IMAGE_BRUSH("Icons/SequencerIcons/icon_Sequencer_Shot_16x", Icon16x16));
 		Set( "Sequencer.Tracks.Slomo", new IMAGE_BRUSH("Icons/SequencerIcons/icon_Sequencer_Slomo_16x", Icon16x16));
 		Set( "Sequencer.Tracks.Sub", new IMAGE_BRUSH("Icons/SequencerIcons/icon_Sequencer_Sub_16x", Icon16x16));
-		Set( "Sequencer.Tracks.Fade", new IMAGE_BRUSH("Icons/SequencerIcons/icon_Sequencer_Fade_16x", Icon16x16));
 
 		Set( "Sequencer.CursorDecorator_Marquee", new IMAGE_BRUSH( "Sequencer/CursorDecorator_Marquee", Icon16x16));
 		Set( "Sequencer.CursorDecorator_MarqueeAdd", new IMAGE_BRUSH( "Sequencer/CursorDecorator_MarqueeAdd", Icon16x16));
diff --git a/Engine/Source/Editor/LandscapeEditor/Private/LandscapeEdMode.cpp b/Engine/Source/Editor/LandscapeEditor/Private/LandscapeEdMode.cpp
--- a/Engine/Source/Editor/LandscapeEditor/Private/LandscapeEdMode.cpp
+++ b/Engine/Source/Editor/LandscapeEditor/Private/LandscapeEdMode.cpp
@@ -1709,7 +1709,7 @@ int32 FEdModeLandscape::UpdateLandscapeList()
 	{
 		if (LandscapeList.Num() > 0)
 		{
-			if (CurrentTool != NULL)
+			if (CurrentTool != nullptr)
 			{
 				CurrentBrush->LeaveBrush();
 				CurrentTool->ExitTool();
@@ -1717,7 +1717,7 @@ int32 FEdModeLandscape::UpdateLandscapeList()
 			CurrentToolTarget.LandscapeInfo = LandscapeList[0].Info;
 			CurrentIndex = 0;
 			UpdateTargetList();
-			if (CurrentTool != NULL)
+			if (CurrentTool != nullptr)
 			{
 				CurrentTool->EnterTool();
 				CurrentBrush->EnterBrush();
@@ -1725,7 +1725,11 @@ int32 FEdModeLandscape::UpdateLandscapeList()
 		}
 		else
 		{
-			CurrentToolTarget.LandscapeInfo = NULL;
+			// no landscape, switch to "new landscape" tool
+			CurrentToolTarget.LandscapeInfo = nullptr;
+			UpdateTargetList();
+			SetCurrentToolMode("ToolMode_Manage", false);
+			SetCurrentTool("NewLandscape");
 		}
 	}
 
@@ -2586,7 +2590,7 @@ ALandscape* FEdModeLandscape::ChangeComponentSetting(int32 NumComponentsX, int32
 				}
 
 				// offset landscape to component boundary
-				LandscapeOffset = FVector(NewMinX - OldMinX, NewMinY - OldMinY, 0) * OldLandscapeProxy->GetActorScale();
+				LandscapeOffset = FVector(NewMinX, NewMinY, 0) * OldLandscapeProxy->GetActorScale();
 				NewMinX = 0;
 				NewMinY = 0;
 				NewMaxX = NewVertsX - 1;
@@ -2594,7 +2598,9 @@ ALandscape* FEdModeLandscape::ChangeComponentSetting(int32 NumComponentsX, int32
 			}
 
 			const FVector Location = OldLandscapeProxy->GetActorLocation() + LandscapeOffset;
-			Landscape = OldLandscapeProxy->GetWorld()->SpawnActor<ALandscape>(Location, OldLandscapeProxy->GetActorRotation());
+			FActorSpawnParameters SpawnParams;
+			SpawnParams.OverrideLevel = OldLandscapeProxy->GetLevel();
+			Landscape = OldLandscapeProxy->GetWorld()->SpawnActor<ALandscape>(Location, OldLandscapeProxy->GetActorRotation(), SpawnParams);
 
 			const FVector OldScale = OldLandscapeProxy->GetActorScale();
 			Landscape->SetActorRelativeScale3D(FVector(OldScale.X * LandscapeScaleFactor, OldScale.Y * LandscapeScaleFactor, OldScale.Z));
diff --git a/Engine/Source/Editor/LandscapeEditor/Private/LandscapeEditorDetailCustomization_ResizeLandscape.cpp b/Engine/Source/Editor/LandscapeEditor/Private/LandscapeEditorDetailCustomization_ResizeLandscape.cpp
--- a/Engine/Source/Editor/LandscapeEditor/Private/LandscapeEditorDetailCustomization_ResizeLandscape.cpp
+++ b/Engine/Source/Editor/LandscapeEditor/Private/LandscapeEditorDetailCustomization_ResizeLandscape.cpp
@@ -517,9 +517,6 @@ FReply FLandscapeEditorDetailCustomization_ResizeLandscape::OnApplyButtonClicked
 		LandscapeEdMode->ChangeComponentSetting(ComponentCount.X, ComponentCount.Y, SectionsPerComponent, QuadsPerSection, bResample);
 
 		LandscapeEdMode->UpdateLandscapeList();
-
-		LandscapeEdMode->SetCurrentTool("Select"); // change tool so switching back to the manage mode doesn't give "Resize Landscape" again
-		LandscapeEdMode->SetCurrentTool("Sculpt"); // change to sculpting mode and tool
 	}
 
 	return FReply::Handled();
diff --git a/Engine/Source/Editor/Localization/Private/LocalizationConfigurationScript.cpp b/Engine/Source/Editor/Localization/Private/LocalizationConfigurationScript.cpp
--- a/Engine/Source/Editor/Localization/Private/LocalizationConfigurationScript.cpp
+++ b/Engine/Source/Editor/Localization/Private/LocalizationConfigurationScript.cpp
@@ -279,14 +279,9 @@
 
 		const FString ContentDirRelativeToGameDir = MakePathRelativeForCommandletProcess(GetContentDir(Target), !Target->IsMemberOfEngineTargetSet());
 
-		// GatherTextStep0 - InternationalizationExport
+		// CommonSettings
 		{
-			FConfigSection& ConfigSection = Script.GatherTextStep(0);
+			FConfigSection& ConfigSection = Script.CommonSettings();
 
-			// CommandletClass
-			ConfigSection.Add( TEXT("CommandletClass"), TEXT("InternationalizationExport") );
-
-			ConfigSection.Add( TEXT("bImportLoc"), TEXT("true") );
-
 			FString SourcePath;
 			// Overriding output path changes the source directory for the PO file.
@@ -360,6 +355,16 @@
 			ConfigSection.Add( TEXT("PortableObjectName"), POFileName );
 		}
 
+		// GatherTextStep0 - InternationalizationExport
+		{
+			FConfigSection& ConfigSection = Script.GatherTextStep(0);
+
+			// CommandletClass
+			ConfigSection.Add( TEXT("CommandletClass"), TEXT("InternationalizationExport") );
+
+			ConfigSection.Add( TEXT("bImportLoc"), TEXT("true") );
+		}
+
 		Script.Dirty = true;
 
 		return Script;
@@ -386,20 +391,9 @@
 
 		const FString ContentDirRelativeToGameDir = MakePathRelativeForCommandletProcess(GetContentDir(Target), !Target->IsMemberOfEngineTargetSet());
 
-		// GatherTextStep0 - InternationalizationExport
+		// CommonSettings
 		{
-			FConfigSection& ConfigSection = Script.GatherTextStep(0);
+			FConfigSection& ConfigSection = Script.CommonSettings();
 
-			// CommandletClass
-			ConfigSection.Add( TEXT("CommandletClass"), TEXT("InternationalizationExport") );
-
-			ConfigSection.Add(TEXT("bExportLoc"), TEXT("true"));
-
-			// Export-specific settings.
-			{
-				ConfigSection.Add(TEXT("ShouldPersistCommentsOnExport"), Target->Settings.ExportSettings.ShouldPersistCommentsOnExport ? TEXT("true") : TEXT("false"));
-				ConfigSection.Add(TEXT("ShouldAddSourceLocationsAsComments"), Target->Settings.ExportSettings.ShouldAddSourceLocationsAsComments ? TEXT("true") : TEXT("false"));
-			}
-
 			const FString SourcePath = ContentDirRelativeToGameDir / TEXT("Localization") / Target->Settings.Name;
 			ConfigSection.Add( TEXT("SourcePath"), SourcePath );
@@ -473,6 +467,22 @@
 			ConfigSection.Add( TEXT("PortableObjectName"), POFileName );
 		}
 
+		// GatherTextStep0 - InternationalizationExport
+		{
+			FConfigSection& ConfigSection = Script.GatherTextStep(0);
+
+			// CommandletClass
+			ConfigSection.Add( TEXT("CommandletClass"), TEXT("InternationalizationExport") );
+
+			ConfigSection.Add(TEXT("bExportLoc"), TEXT("true"));
+
+			// Export-specific settings.
+			{
+				ConfigSection.Add(TEXT("ShouldPersistCommentsOnExport"), Target->Settings.ExportSettings.ShouldPersistCommentsOnExport ? TEXT("true") : TEXT("false"));
+				ConfigSection.Add(TEXT("ShouldAddSourceLocationsAsComments"), Target->Settings.ExportSettings.ShouldAddSourceLocationsAsComments ? TEXT("true") : TEXT("false"));
+			}
+		}
+
 		Script.Dirty = true;
 
 		return Script;
@@ -499,12 +509,9 @@
 
 		const FString ContentDirRelativeToGameDir = MakePathRelativeForCommandletProcess(GetContentDir(Target), !Target->IsMemberOfEngineTargetSet());
 
-		// GatherTextStep0 - GenerateTextLocalizationReport
+		// CommonSettings
 		{
-			FConfigSection& ConfigSection = Script.GatherTextStep(0);
+			FConfigSection& ConfigSection = Script.CommonSettings();
 
-			// CommandletClass
-			ConfigSection.Add( TEXT("CommandletClass"), TEXT("GenerateTextLocalizationReport") );
-
 			const FString SourcePath = ContentDirRelativeToGameDir / TEXT("Localization") / Target->Settings.Name;
 			ConfigSection.Add( TEXT("SourcePath"), SourcePath );
@@ -517,6 +524,14 @@
 			{
 				ConfigSection.Add( TEXT("CulturesToGenerate"), CultureStatistics.CultureName );
 			}
+		}
 
+		// GatherTextStep0 - GenerateTextLocalizationReport
+		{
+			FConfigSection& ConfigSection = Script.GatherTextStep(0);
+
+			// CommandletClass
+			ConfigSection.Add( TEXT("CommandletClass"), TEXT("GenerateTextLocalizationReport") );
+
 			ConfigSection.Add( TEXT("bWordCountReport"), TEXT("true") );
 
@@ -539,12 +554,9 @@
 
 		const FString ContentDirRelativeToGameDir = MakePathRelativeForCommandletProcess(GetContentDir(Target), !Target->IsMemberOfEngineTargetSet());
 
-		// GatherTextStep0 - GenerateTextLocalizationResource
+		// CommonSettings
 		{
-			FConfigSection& ConfigSection = Script.GatherTextStep(0);
+			FConfigSection& ConfigSection = Script.CommonSettings();
 
-			// CommandletClass
-			ConfigSection.Add( TEXT("CommandletClass"), TEXT("GenerateTextLocalizationResource") );
-
 			const FString SourcePath = ContentDirRelativeToGameDir / TEXT("Localization") / Target->Settings.Name;
 			ConfigSection.Add( TEXT("SourcePath"), SourcePath );
@@ -580,6 +592,14 @@
 			}
 		}
 
+		// GatherTextStep0 - GenerateTextLocalizationResource
+		{
+			FConfigSection& ConfigSection = Script.GatherTextStep(0);
+
+			// CommandletClass
+			ConfigSection.Add( TEXT("CommandletClass"), TEXT("GenerateTextLocalizationResource") );
+		}
+
 		Script.Dirty = true;
 
 		return Script;
diff --git a/Engine/Source/Editor/MovieSceneCaptureDialog/Private/MovieSceneCaptureDialogModule.cpp b/Engine/Source/Editor/MovieSceneCaptureDialog/Private/MovieSceneCaptureDialogModule.cpp
--- a/Engine/Source/Editor/MovieSceneCaptureDialog/Private/MovieSceneCaptureDialogModule.cpp
+++ b/Engine/Source/Editor/MovieSceneCaptureDialog/Private/MovieSceneCaptureDialogModule.cpp
@@ -302,7 +302,7 @@ class FMovieSceneCaptureDialogModule : public IMovieSceneCaptureDialogModule
 				.HasCloseButton(true)
 				.SupportsMaximize(false)
 				.SupportsMinimize(false)
-				.ClientSize(FVector2D(350, 650));
+				.ClientSize(FVector2D(500, 700));
 
 			TSharedPtr<SDockTab> OwnerTab = TabManager->GetOwnerTab();
 			TSharedPtr<SWindow> RootWindow = OwnerTab.IsValid() ? OwnerTab->GetParentWindow() : TSharedPtr<SWindow>();
diff --git a/Engine/Source/Editor/MovieSceneTools/Private/CurveKeyAreas/FloatCurveKeyArea.cpp b/Engine/Source/Editor/MovieSceneTools/Private/CurveKeyAreas/FloatCurveKeyArea.cpp
--- a/Engine/Source/Editor/MovieSceneTools/Private/CurveKeyAreas/FloatCurveKeyArea.cpp
+++ b/Engine/Source/Editor/MovieSceneTools/Private/CurveKeyAreas/FloatCurveKeyArea.cpp
@@ -3,7 +3,8 @@
 #include "MovieSceneToolsPrivatePCH.h"
 #include "FloatCurveKeyArea.h"
 #include "SFloatCurveKeyEditor.h"
-
+#include "MovieSceneClipboard.h"
+#include "SequencerClipboardReconciler.h"
 
 /* IKeyArea interface
  *****************************************************************************/
@@ -89,6 +90,12 @@ void FFloatCurveKeyArea::DeleteKey(FKeyHandle KeyHandle)
 }
 
 
+FLinearColor FFloatCurveKeyArea::GetColor()
+{
+	return Color;
+}
+
+
 ERichCurveExtrapolation FFloatCurveKeyArea::GetExtrapolationMode(bool bPreInfinity) const
 {
 	if (bPreInfinity)
@@ -194,3 +201,69 @@ void FFloatCurveKeyArea::SetKeyTime(FKeyHandle KeyHandle, float NewKeyTime) cons
 {
 	Curve->SetKeyTime(KeyHandle, NewKeyTime);
 }
+
+void FFloatCurveKeyArea::CopyKeys(FMovieSceneClipboardBuilder& ClipboardBuilder, const TFunctionRef<bool(FKeyHandle, const IKeyArea&)>& KeyMask) const
+{
+	const UMovieSceneSection* Section = const_cast<FFloatCurveKeyArea*>(this)->GetOwningSection();
+	UMovieSceneTrack* Track = Section ? Section->GetTypedOuter<UMovieSceneTrack>() : nullptr;
+	if (!Track)
+	{
+		return;
+	}
+
+	FMovieSceneClipboardKeyTrack* KeyTrack = nullptr;
+
+	for (auto It(Curve->GetKeyHandleIterator()); It; ++It)
+	{
+		FKeyHandle Handle = It.Key();
+		if (KeyMask(Handle, *this))
+		{
+			if (!KeyTrack)
+			{
+				KeyTrack = &ClipboardBuilder.FindOrAddKeyTrack<FRichCurveKey>(GetName(), *Track);
+			}
+
+			FRichCurveKey Key = Curve->GetKey(Handle);
+			KeyTrack->AddKey(Key.Time, Key);
+		}
+	}
+}
+
+void FFloatCurveKeyArea::PasteKeys(const FMovieSceneClipboardKeyTrack& KeyTrack, const FMovieSceneClipboardEnvironment& SrcEnvironment, const FSequencerPasteEnvironment& DstEnvironment)
+{
+	float PasteAt = DstEnvironment.CardinalTime;
+
+	KeyTrack.IterateKeys([&](const FMovieSceneClipboardKey& Key){
+		UMovieSceneSection* Section = GetOwningSection();
+		if (!Section)
+		{
+			return true;
+		}
+
+		if (Section->TryModify())
+		{
+			float Time = PasteAt + Key.GetTime();
+			if (Section->GetStartTime() > Time)
+			{
+				Section->SetStartTime(Time);
+			}
+			if (Section->GetEndTime() < Time)
+			{
+				Section->SetEndTime(Time);
+			}
+
+			FRichCurveKey NewKey = Key.GetValue<FRichCurveKey>();
+
+			// Rich curve keys store their time internally, so we need to ensure they are set up correctly here (if there's been some conversion, for instance)
+			NewKey.Time = Time;
+
+			FKeyHandle KeyHandle = Curve->UpdateOrAddKey(Time, NewKey.Value);
+			// Ensure the rest of the key properties are added
+			Curve->GetKey(KeyHandle) = NewKey;
+
+			DstEnvironment.ReportPastedKey(KeyHandle, *this);
+		}
+
+		return true;
+	});
+}
diff --git a/Engine/Source/Editor/MovieSceneTools/Private/CurveKeyAreas/IntegralKeyArea.cpp b/Engine/Source/Editor/MovieSceneTools/Private/CurveKeyAreas/IntegralKeyArea.cpp
--- a/Engine/Source/Editor/MovieSceneTools/Private/CurveKeyAreas/IntegralKeyArea.cpp
+++ b/Engine/Source/Editor/MovieSceneTools/Private/CurveKeyAreas/IntegralKeyArea.cpp
@@ -2,7 +2,7 @@
 
 #include "MovieSceneToolsPrivatePCH.h"
 #include "IntegralKeyArea.h"
-
+#include "MovieSceneClipboard.h"
 
 /* IKeyArea interface
  *****************************************************************************/
@@ -38,6 +38,12 @@ void FIntegralCurveKeyAreaBase::DeleteKey(FKeyHandle KeyHandle)
 }
 
 
+FLinearColor FIntegralCurveKeyAreaBase::GetColor()
+{
+	return FLinearColor(0.1f, 0.1f, 0.1f, 0.7f);
+}
+
+
 ERichCurveExtrapolation FIntegralCurveKeyAreaBase::GetExtrapolationMode(bool bPreInfinity) const
 {
 	return RCCE_None;
diff --git a/Engine/Source/Editor/MovieSceneTools/Private/CurveKeyAreas/NameCurveKeyArea.cpp b/Engine/Source/Editor/MovieSceneTools/Private/CurveKeyAreas/NameCurveKeyArea.cpp
--- a/Engine/Source/Editor/MovieSceneTools/Private/CurveKeyAreas/NameCurveKeyArea.cpp
+++ b/Engine/Source/Editor/MovieSceneTools/Private/CurveKeyAreas/NameCurveKeyArea.cpp
@@ -2,7 +2,7 @@
 
 #include "MovieSceneToolsPrivatePCH.h"
 #include "NameCurveKeyArea.h"
-
+#include "SequencerClipboardReconciler.h"
 
 /* IKeyArea interface
  *****************************************************************************/
@@ -50,6 +50,12 @@ void FNameCurveKeyArea::DeleteKey(FKeyHandle KeyHandle)
 }
 
 
+FLinearColor FNameCurveKeyArea::GetColor()
+{
+	return FLinearColor(0.1f, 0.1f, 0.1f, 0.7f);
+}
+
+
 ERichCurveExtrapolation FNameCurveKeyArea::GetExtrapolationMode(bool bPreInfinity) const
 {
 	return RCCE_None;
@@ -128,3 +134,61 @@ void FNameCurveKeyArea::SetKeyTime(FKeyHandle KeyHandle, float NewKeyTime) const
 {
 	Curve.SetKeyTime(KeyHandle, NewKeyTime);
 }
+
+void FNameCurveKeyArea::CopyKeys(FMovieSceneClipboardBuilder& ClipboardBuilder, const TFunctionRef<bool(FKeyHandle, const IKeyArea&)>& KeyMask) const
+{
+	const UMovieSceneSection* Section = const_cast<FNameCurveKeyArea*>(this)->GetOwningSection();
+	UMovieSceneTrack* Track = Section ? Section->GetTypedOuter<UMovieSceneTrack>() : nullptr;
+	if (!Track)
+	{
+		return;
+	}
+
+	FMovieSceneClipboardKeyTrack* KeyTrack = nullptr;
+
+	for (auto It(Curve.GetKeyHandleIterator()); It; ++It)
+	{
+		FKeyHandle Handle = It.Key();
+		if (KeyMask(Handle, *this))
+		{
+			FNameCurveKey Key = Curve.GetKey(Handle);
+			if (!KeyTrack)
+			{
+				KeyTrack = &ClipboardBuilder.FindOrAddKeyTrack<FName>(GetName(), *Track);
+			}
+
+			KeyTrack->AddKey(Key.Time, Key.Value);
+		}
+	}
+}
+
+void FNameCurveKeyArea::PasteKeys(const FMovieSceneClipboardKeyTrack& KeyTrack, const FMovieSceneClipboardEnvironment& SrcEnvironment, const FSequencerPasteEnvironment& DstEnvironment)
+{
+	float PasteAt = DstEnvironment.CardinalTime;
+
+	KeyTrack.IterateKeys([&](const FMovieSceneClipboardKey& Key){
+		UMovieSceneSection* Section = GetOwningSection();
+		if (!Section)
+		{
+			return true;
+		}
+
+		if (Section->TryModify())
+		{		
+			float Time = PasteAt + Key.GetTime();
+			if (Section->GetStartTime() > Time)
+			{
+				Section->SetStartTime(Time);
+			}
+			if (Section->GetEndTime() < Time)
+			{
+				Section->SetEndTime(Time);
+			}
+
+			FKeyHandle KeyHandle = Curve.UpdateOrAddKey(Time, Key.GetValue<FName>());
+			DstEnvironment.ReportPastedKey(KeyHandle, *this);
+		}
+			
+		return true;
+	});
+}
diff --git a/Engine/Source/Editor/MovieSceneTools/Private/CurveKeyEditors/SBoolCurveKeyEditor.cpp b/Engine/Source/Editor/MovieSceneTools/Private/CurveKeyEditors/SBoolCurveKeyEditor.cpp
--- a/Engine/Source/Editor/MovieSceneTools/Private/CurveKeyEditors/SBoolCurveKeyEditor.cpp
+++ b/Engine/Source/Editor/MovieSceneTools/Private/CurveKeyEditors/SBoolCurveKeyEditor.cpp
@@ -40,32 +40,33 @@ void SBoolCurveKeyEditor::OnCheckStateChanged(ECheckBoxState NewCheckboxState)
 {
 	FScopedTransaction Transaction(LOCTEXT("SetBoolKey", "Set Bool Key Value"));
 	OwningSection->SetFlags(RF_Transactional);
-	OwningSection->Modify();
+	if (OwningSection->TryModify())
+	{
+		float CurrentTime = Sequencer->GetCurrentLocalTime(*Sequencer->GetFocusedMovieSceneSequence());
 
-	float CurrentTime = Sequencer->GetCurrentLocalTime(*Sequencer->GetFocusedMovieSceneSequence());
+		bool bKeyWillBeAdded = Curve->IsKeyHandleValid(Curve->FindKey(CurrentTime)) == false;
+		if (bKeyWillBeAdded)
+		{
+			if (OwningSection->GetStartTime() > CurrentTime)
+			{
+				OwningSection->SetStartTime(CurrentTime);
+			}
+			if (OwningSection->GetEndTime() < CurrentTime)
+			{
+				OwningSection->SetEndTime(CurrentTime);
+			}
+		}
 
-	bool bKeyWillBeAdded = Curve->IsKeyHandleValid(Curve->FindKey(CurrentTime)) == false;
-	if (bKeyWillBeAdded)
-	{
-		if (OwningSection->GetStartTime() > CurrentTime)
+		if (Curve->GetNumKeys() == 0)
 		{
-			OwningSection->SetStartTime(CurrentTime);
+			Curve->SetDefaultValue(NewCheckboxState == ECheckBoxState::Checked ? 1 : 0);
 		}
-		if (OwningSection->GetEndTime() < CurrentTime)
+		else
 		{
-			OwningSection->SetEndTime(CurrentTime);
+			Curve->UpdateOrAddKey(CurrentTime, NewCheckboxState == ECheckBoxState::Checked ? 1 : 0);
 		}
+		Sequencer->UpdateRuntimeInstances();
 	}
-
-	if (Curve->GetNumKeys() == 0)
-	{
-		Curve->SetDefaultValue(NewCheckboxState == ECheckBoxState::Checked ? 1 : 0);
-	}
-	else
-	{
-		Curve->UpdateOrAddKey(CurrentTime, NewCheckboxState == ECheckBoxState::Checked ? 1 : 0);
-	}
-	Sequencer->UpdateRuntimeInstances();
 }
 
 #undef LOCTEXT_NAMESPACE
\ No newline at end of file
diff --git a/Engine/Source/Editor/MovieSceneTools/Private/CurveKeyEditors/SFloatCurveKeyEditor.cpp b/Engine/Source/Editor/MovieSceneTools/Private/CurveKeyEditors/SFloatCurveKeyEditor.cpp
--- a/Engine/Source/Editor/MovieSceneTools/Private/CurveKeyEditors/SFloatCurveKeyEditor.cpp
+++ b/Engine/Source/Editor/MovieSceneTools/Private/CurveKeyEditors/SFloatCurveKeyEditor.cpp
@@ -34,7 +34,7 @@ void SFloatCurveKeyEditor::OnBeginSliderMovement()
 {
 	GEditor->BeginTransaction(LOCTEXT("SetFloatKey", "Set Float Key Value"));
 	OwningSection->SetFlags(RF_Transactional);
-	OwningSection->Modify();
+	OwningSection->TryModify();
 }
 
 void SFloatCurveKeyEditor::OnEndSliderMovement(float Value)
@@ -58,37 +58,38 @@ float SFloatCurveKeyEditor::OnGetKeyValue() const
 
 void SFloatCurveKeyEditor::OnValueChanged(float Value)
 {
-	OwningSection->Modify();
-
-	float CurrentTime = Sequencer->GetCurrentLocalTime(*Sequencer->GetFocusedMovieSceneSequence());
-	FKeyHandle CurrentKeyHandle = Curve->FindKey(CurrentTime);
-	if (Curve->IsKeyHandleValid(CurrentKeyHandle))
-	{
-		Curve->SetKeyValue(CurrentKeyHandle, Value);
-	}
-	else
+	if (OwningSection->TryModify())
 	{
-		if (Curve->GetNumKeys() == 0)
+		float CurrentTime = Sequencer->GetCurrentLocalTime(*Sequencer->GetFocusedMovieSceneSequence());
+		FKeyHandle CurrentKeyHandle = Curve->FindKey(CurrentTime);
+		if (Curve->IsKeyHandleValid(CurrentKeyHandle))
 		{
-			Curve->SetDefaultValue(Value);
+			Curve->SetKeyValue(CurrentKeyHandle, Value);
 		}
 		else
 		{
-			Curve->AddKey(CurrentTime, Value, false, CurrentKeyHandle);
+			if (Curve->GetNumKeys() == 0)
+			{
+				Curve->SetDefaultValue(Value);
+			}
+			else
+			{
+				Curve->AddKey(CurrentTime, Value, false, CurrentKeyHandle);
 
-			MovieSceneHelpers::SetKeyInterpolation(*Curve, CurrentKeyHandle, Sequencer->GetKeyInterpolation());
-		}
+				MovieSceneHelpers::SetKeyInterpolation(*Curve, CurrentKeyHandle, Sequencer->GetKeyInterpolation());
+			}
 
-		if (OwningSection->GetStartTime() > CurrentTime)
-		{
-			OwningSection->SetStartTime(CurrentTime);
-		}
-		if (OwningSection->GetEndTime() < CurrentTime)
-		{
-			OwningSection->SetEndTime(CurrentTime);
+			if (OwningSection->GetStartTime() > CurrentTime)
+			{
+				OwningSection->SetStartTime(CurrentTime);
+			}
+			if (OwningSection->GetEndTime() < CurrentTime)
+			{
+				OwningSection->SetEndTime(CurrentTime);
+			}
 		}
+		Sequencer->UpdateRuntimeInstances();
 	}
-	Sequencer->UpdateRuntimeInstances();
 }
 
 void SFloatCurveKeyEditor::OnValueCommitted(float Value, ETextCommit::Type CommitInfo)
diff --git a/Engine/Source/Editor/MovieSceneTools/Private/CurveKeyEditors/SIntegralCurveKeyEditor.cpp b/Engine/Source/Editor/MovieSceneTools/Private/CurveKeyEditors/SIntegralCurveKeyEditor.cpp
--- a/Engine/Source/Editor/MovieSceneTools/Private/CurveKeyEditors/SIntegralCurveKeyEditor.cpp
+++ b/Engine/Source/Editor/MovieSceneTools/Private/CurveKeyEditors/SIntegralCurveKeyEditor.cpp
@@ -34,7 +34,7 @@ void SIntegralCurveKeyEditor::OnBeginSliderMovement()
 {
 	GEditor->BeginTransaction(LOCTEXT("SetIntegralKey", "Set Integral Key Value"));
 	OwningSection->SetFlags(RF_Transactional);
-	OwningSection->Modify();
+	OwningSection->TryModify();
 }
 
 void SIntegralCurveKeyEditor::OnEndSliderMovement(int32 Value)
@@ -58,32 +58,33 @@ int32 SIntegralCurveKeyEditor::OnGetKeyValue() const
 
 void SIntegralCurveKeyEditor::OnValueChanged(int32 Value)
 {
-	OwningSection->Modify();
+	if (OwningSection->TryModify())
+	{
+		float CurrentTime = Sequencer->GetCurrentLocalTime(*Sequencer->GetFocusedMovieSceneSequence());
 
-	float CurrentTime = Sequencer->GetCurrentLocalTime(*Sequencer->GetFocusedMovieSceneSequence());
+		bool bKeyWillBeAdded = Curve->IsKeyHandleValid(Curve->FindKey(CurrentTime)) == false;
+		if (bKeyWillBeAdded)
+		{
+			if (OwningSection->GetStartTime() > CurrentTime)
+			{
+				OwningSection->SetStartTime(CurrentTime);
+			}
+			if (OwningSection->GetEndTime() < CurrentTime)
+			{
+				OwningSection->SetEndTime(CurrentTime);
+			}
+		}
 
-	bool bKeyWillBeAdded = Curve->IsKeyHandleValid(Curve->FindKey(CurrentTime)) == false;
-	if (bKeyWillBeAdded)
-	{
-		if (OwningSection->GetStartTime() > CurrentTime)
+		if (Curve->GetNumKeys() == 0)
 		{
-			OwningSection->SetStartTime(CurrentTime);
+			Curve->SetDefaultValue(Value);
 		}
-		if (OwningSection->GetEndTime() < CurrentTime)
+		else
 		{
-			OwningSection->SetEndTime(CurrentTime);
+			Curve->UpdateOrAddKey(CurrentTime, Value);
 		}
+		Sequencer->UpdateRuntimeInstances();
 	}
-
-	if (Curve->GetNumKeys() == 0)
-	{
-		Curve->SetDefaultValue(Value);
-	}
-	else
-	{
-		Curve->UpdateOrAddKey(CurrentTime, Value);
-	}
-	Sequencer->UpdateRuntimeInstances();
 }
 
 void SIntegralCurveKeyEditor::OnValueCommitted(int32 Value, ETextCommit::Type CommitInfo)
diff --git a/Engine/Source/Editor/MovieSceneTools/Private/MatineeImportTools.cpp b/Engine/Source/Editor/MovieSceneTools/Private/MatineeImportTools.cpp
--- a/Engine/Source/Editor/MovieSceneTools/Private/MatineeImportTools.cpp
+++ b/Engine/Source/Editor/MovieSceneTools/Private/MatineeImportTools.cpp
@@ -94,25 +94,26 @@ void FMatineeImportTools::CopyInterpFloatTrack( TSharedRef<ISequencer> Sequencer
 		FloatTrack->AddSection( *Section );
 		bSectionCreated = true;
 	}
-	Section->Modify();
-
-	float SectionMin = Section->GetStartTime();
-	float SectionMax = Section->GetEndTime();
-
-	FRichCurve& FloatCurve = Section->GetFloatCurve();
-	for ( const auto& Point : MatineeFloatTrack->FloatTrack.Points )
+	if (Section->TryModify())
 	{
-		FMatineeImportTools::SetOrAddKey( FloatCurve, Point.InVal, Point.OutVal, Point.ArriveTangent, Point.LeaveTangent, Point.InterpMode );
-		SectionMin = FMath::Min( SectionMin, Point.InVal );
-		SectionMax = FMath::Max( SectionMax, Point.InVal );
-	}
+		float SectionMin = Section->GetStartTime();
+		float SectionMax = Section->GetEndTime();
+
+		FRichCurve& FloatCurve = Section->GetFloatCurve();
+		for ( const auto& Point : MatineeFloatTrack->FloatTrack.Points )
+		{
+			FMatineeImportTools::SetOrAddKey( FloatCurve, Point.InVal, Point.OutVal, Point.ArriveTangent, Point.LeaveTangent, Point.InterpMode );
+			SectionMin = FMath::Min( SectionMin, Point.InVal );
+			SectionMax = FMath::Max( SectionMax, Point.InVal );
+		}
 
-	Section->SetStartTime( SectionMin );
-	Section->SetEndTime( SectionMax );
+		Section->SetStartTime( SectionMin );
+		Section->SetEndTime( SectionMax );
 
-	if ( bSectionCreated )
-	{
-		Sequencer->NotifyMovieSceneDataChanged();
+		if ( bSectionCreated )
+		{
+			Sequencer->NotifyMovieSceneDataChanged();
+		}
 	}
 }
 
@@ -135,43 +136,45 @@ void FMatineeImportTools::CopyInterpMoveTrack( TSharedRef<ISequencer> Sequencer,
 		TransformTrack->AddSection( *Section );
 		bSectionCreated = true;
 	}
-	Section->Modify();
 
-	float SectionMin = Section->GetStartTime();
-	float SectionMax = Section->GetEndTime();
+	if (Section->TryModify())
+	{
+		float SectionMin = Section->GetStartTime();
+		float SectionMax = Section->GetEndTime();
 
-	FRichCurve& TranslationXCurve = Section->GetTranslationCurve( EAxis::X );
-	FRichCurve& TranslationYCurve = Section->GetTranslationCurve( EAxis::Y );
-	FRichCurve& TranslationZCurve = Section->GetTranslationCurve( EAxis::Z );
+		FRichCurve& TranslationXCurve = Section->GetTranslationCurve( EAxis::X );
+		FRichCurve& TranslationYCurve = Section->GetTranslationCurve( EAxis::Y );
+		FRichCurve& TranslationZCurve = Section->GetTranslationCurve( EAxis::Z );
 
-	for ( const auto& Point : MoveTrack->PosTrack.Points )
-	{
-		FMatineeImportTools::SetOrAddKey( TranslationXCurve, Point.InVal, Point.OutVal.X, Point.ArriveTangent.X, Point.LeaveTangent.X, Point.InterpMode );
-		FMatineeImportTools::SetOrAddKey( TranslationYCurve, Point.InVal, Point.OutVal.Y, Point.ArriveTangent.Y, Point.LeaveTangent.Y, Point.InterpMode );
-		FMatineeImportTools::SetOrAddKey( TranslationZCurve, Point.InVal, Point.OutVal.Z, Point.ArriveTangent.Z, Point.LeaveTangent.Z, Point.InterpMode );
-		SectionMin = FMath::Min( SectionMin, Point.InVal );
-		SectionMax = FMath::Max( SectionMax, Point.InVal );
-	}
+		for ( const auto& Point : MoveTrack->PosTrack.Points )
+		{
+			FMatineeImportTools::SetOrAddKey( TranslationXCurve, Point.InVal, Point.OutVal.X, Point.ArriveTangent.X, Point.LeaveTangent.X, Point.InterpMode );
+			FMatineeImportTools::SetOrAddKey( TranslationYCurve, Point.InVal, Point.OutVal.Y, Point.ArriveTangent.Y, Point.LeaveTangent.Y, Point.InterpMode );
+			FMatineeImportTools::SetOrAddKey( TranslationZCurve, Point.InVal, Point.OutVal.Z, Point.ArriveTangent.Z, Point.LeaveTangent.Z, Point.InterpMode );
+			SectionMin = FMath::Min( SectionMin, Point.InVal );
+			SectionMax = FMath::Max( SectionMax, Point.InVal );
+		}
 
-	FRichCurve& RotationXCurve = Section->GetRotationCurve( EAxis::X );
-	FRichCurve& RotationYCurve = Section->GetRotationCurve( EAxis::Y );
-	FRichCurve& RotationZCurve = Section->GetRotationCurve( EAxis::Z );
+		FRichCurve& RotationXCurve = Section->GetRotationCurve( EAxis::X );
+		FRichCurve& RotationYCurve = Section->GetRotationCurve( EAxis::Y );
+		FRichCurve& RotationZCurve = Section->GetRotationCurve( EAxis::Z );
 
-	for ( const auto& Point : MoveTrack->EulerTrack.Points )
-	{
-		FMatineeImportTools::SetOrAddKey( RotationXCurve, Point.InVal, Point.OutVal.X, Point.ArriveTangent.X, Point.LeaveTangent.X, Point.InterpMode );
-		FMatineeImportTools::SetOrAddKey( RotationYCurve, Point.InVal, Point.OutVal.Y, Point.ArriveTangent.Y, Point.LeaveTangent.Y, Point.InterpMode );
-		FMatineeImportTools::SetOrAddKey( RotationZCurve, Point.InVal, Point.OutVal.Z, Point.ArriveTangent.Z, Point.LeaveTangent.Z, Point.InterpMode );
-		SectionMin = FMath::Min( SectionMin, Point.InVal );
-		SectionMax = FMath::Max( SectionMax, Point.InVal );
-	}
+		for ( const auto& Point : MoveTrack->EulerTrack.Points )
+		{
+			FMatineeImportTools::SetOrAddKey( RotationXCurve, Point.InVal, Point.OutVal.X, Point.ArriveTangent.X, Point.LeaveTangent.X, Point.InterpMode );
+			FMatineeImportTools::SetOrAddKey( RotationYCurve, Point.InVal, Point.OutVal.Y, Point.ArriveTangent.Y, Point.LeaveTangent.Y, Point.InterpMode );
+			FMatineeImportTools::SetOrAddKey( RotationZCurve, Point.InVal, Point.OutVal.Z, Point.ArriveTangent.Z, Point.LeaveTangent.Z, Point.InterpMode );
+			SectionMin = FMath::Min( SectionMin, Point.InVal );
+			SectionMax = FMath::Max( SectionMax, Point.InVal );
+		}
 
-	Section->SetStartTime( SectionMin );
-	Section->SetEndTime( SectionMax );
+		Section->SetStartTime( SectionMin );
+		Section->SetEndTime( SectionMax );
 
-	if ( bSectionCreated )
-	{
-		Sequencer->NotifyMovieSceneDataChanged();
+		if ( bSectionCreated )
+		{
+			Sequencer->NotifyMovieSceneDataChanged();
+		}
 	}
 }
 
@@ -191,29 +194,31 @@ void FMatineeImportTools::CopyInterpParticleTrack( TSharedRef<ISequencer> Sequen
 		ParticleTrack->AddSection( *Section );
 		bSectionCreated = true;
 	}
-	Section->Modify();
-
-	float SectionMin = Section->GetStartTime();
-	float SectionMax = Section->GetEndTime();
 
-	FIntegralCurve& ParticleCurve = Section->GetParticleCurve();
-	for ( const auto& Key : MatineeToggleTrack->ToggleTrack )
+	if (Section->TryModify())
 	{
-		EParticleKey::Type ParticleKey;
-		if ( TryConvertMatineeToggleToOutParticleKey( Key.ToggleAction, ParticleKey ) )
+		float SectionMin = Section->GetStartTime();
+		float SectionMax = Section->GetEndTime();
+
+		FIntegralCurve& ParticleCurve = Section->GetParticleCurve();
+		for ( const auto& Key : MatineeToggleTrack->ToggleTrack )
 		{
-			ParticleCurve.AddKey( Key.Time, (int32)ParticleKey, ParticleCurve.FindKey( Key.Time ) );
+			EParticleKey::Type ParticleKey;
+			if ( TryConvertMatineeToggleToOutParticleKey( Key.ToggleAction, ParticleKey ) )
+			{
+				ParticleCurve.AddKey( Key.Time, (int32)ParticleKey, ParticleCurve.FindKey( Key.Time ) );
+			}
+			SectionMin = FMath::Min( SectionMin, Key.Time );
+			SectionMax = FMath::Max( SectionMax, Key.Time );
 		}
-		SectionMin = FMath::Min( SectionMin, Key.Time );
-		SectionMax = FMath::Max( SectionMax, Key.Time );
-	}
 
-	Section->SetStartTime( SectionMin );
-	Section->SetEndTime( SectionMax );
+		Section->SetStartTime( SectionMin );
+		Section->SetEndTime( SectionMax );
 
-	if ( bSectionCreated )
-	{
-		Sequencer->NotifyMovieSceneDataChanged();
+		if ( bSectionCreated )
+		{
+			Sequencer->NotifyMovieSceneDataChanged();
+		}
 	}
 }
 
diff --git a/Engine/Source/Editor/MovieSceneTools/Private/MovieSceneToolsModule.cpp b/Engine/Source/Editor/MovieSceneTools/Private/MovieSceneToolsModule.cpp
--- a/Engine/Source/Editor/MovieSceneTools/Private/MovieSceneToolsModule.cpp
+++ b/Engine/Source/Editor/MovieSceneTools/Private/MovieSceneToolsModule.cpp
@@ -31,6 +31,11 @@
 #include "FadeTrackEditor.h"
 #include "SpawnTrackEditor.h"
 
+#include "MovieSceneClipboard.h"
+#include "SequencerClipboardReconciler.h"
+#include "ClipboardTypes.h"
+#include "Curves/CurveBase.h"
+
 
 /**
  * Implements the MovieSceneTools module.
@@ -68,6 +73,8 @@ public:
 		ComponentMaterialTrackCreateEditorHandle = SequencerModule.RegisterTrackEditor_Handle( FOnCreateTrackEditor::CreateStatic( &FComponentMaterialTrackEditor::CreateTrackEditor ) );
 		FadeTrackCreateEditorHandle = SequencerModule.RegisterTrackEditor_Handle( FOnCreateTrackEditor::CreateStatic( &FFadeTrackEditor::CreateTrackEditor ) );
 		SpawnTrackCreateEditorHandle = SequencerModule.RegisterTrackEditor_Handle( FOnCreateTrackEditor::CreateStatic( &FSpawnTrackEditor::CreateTrackEditor ) );
+
+		RegisterClipboardConversions();
 	}
 
 	virtual void ShutdownModule() override
@@ -103,6 +110,40 @@ public:
 		SequencerModule.UnRegisterTrackEditor_Handle( SpawnTrackCreateEditorHandle );
 	}
 
+	void RegisterClipboardConversions()
+	{
+		using namespace MovieSceneClipboard;
+
+		DefineImplicitConversion<int32, uint8>();
+		DefineImplicitConversion<int32, bool>();
+
+		DefineImplicitConversion<uint8, int32>();
+		DefineImplicitConversion<uint8, bool>();
+
+		DefineExplicitConversion<int32, FRichCurveKey>([](const int32& In) -> FRichCurveKey { return FRichCurveKey(0.f, In);	});
+		DefineExplicitConversion<uint8, FRichCurveKey>([](const uint8& In) -> FRichCurveKey { return FRichCurveKey(0.f, In);	});
+		DefineExplicitConversion<FRichCurveKey, int32>([](const FRichCurveKey& In) -> int32 { return In.Value; 					});
+		DefineExplicitConversion<FRichCurveKey, uint8>([](const FRichCurveKey& In) -> uint8 { return In.Value; 					});
+		DefineExplicitConversion<FRichCurveKey, bool>([](const FRichCurveKey& In) -> bool	{ return !!In.Value; 				});
+
+		FSequencerClipboardReconciler::AddTrackAlias("Location.X", "R");
+		FSequencerClipboardReconciler::AddTrackAlias("Location.Y", "G");
+		FSequencerClipboardReconciler::AddTrackAlias("Location.Z", "B");
+
+		FSequencerClipboardReconciler::AddTrackAlias("Rotation.X", "R");
+		FSequencerClipboardReconciler::AddTrackAlias("Rotation.Y", "G");
+		FSequencerClipboardReconciler::AddTrackAlias("Rotation.Z", "B");
+
+		FSequencerClipboardReconciler::AddTrackAlias("Scale.X", "R");
+		FSequencerClipboardReconciler::AddTrackAlias("Scale.Y", "G");
+		FSequencerClipboardReconciler::AddTrackAlias("Scale.Z", "B");
+
+		FSequencerClipboardReconciler::AddTrackAlias("X", "R");
+		FSequencerClipboardReconciler::AddTrackAlias("Y", "G");
+		FSequencerClipboardReconciler::AddTrackAlias("Z", "B");
+		FSequencerClipboardReconciler::AddTrackAlias("W", "A");
+	}
+
 private:
 
 	/** Registered delegate handles */
diff --git a/Engine/Source/Editor/MovieSceneTools/Private/Sections/MaterialParameterSection.cpp b/Engine/Source/Editor/MovieSceneTools/Private/Sections/MaterialParameterSection.cpp
--- a/Engine/Source/Editor/MovieSceneTools/Private/Sections/MaterialParameterSection.cpp
+++ b/Engine/Source/Editor/MovieSceneTools/Private/Sections/MaterialParameterSection.cpp
@@ -88,8 +88,11 @@ bool FMaterialParameterSection::RequestDeleteCategory( const TArray<FName>& Cate
 {
 	const FScopedTransaction Transaction( LOCTEXT( "DeleteVectorParameter", "Delete vector parameter" ) );
 	UMovieSceneMaterialParameterSection* ParameterSection = Cast<UMovieSceneMaterialParameterSection>( &SectionObject );
-	ParameterSection->Modify();
-	return ParameterSection->RemoveVectorParameter( CategoryNamePath[0] );
+	if (ParameterSection->TryModify())
+	{
+		return ParameterSection->RemoveVectorParameter( CategoryNamePath[0] );
+	}
+	return false;
 }
 
 
@@ -100,8 +103,10 @@ bool FMaterialParameterSection::RequestDeleteKeyArea( const TArray<FName>& KeyAr
 	{
 		const FScopedTransaction Transaction( LOCTEXT( "DeleteScalarParameter", "Delete scalar parameter" ) );
 		UMovieSceneMaterialParameterSection* ParameterSection = Cast<UMovieSceneMaterialParameterSection>( &SectionObject );
-		ParameterSection->Modify();
-		return ParameterSection->RemoveScalarParameter( KeyAreaNamePath[0] );
+		if (ParameterSection->TryModify())
+		{
+			return ParameterSection->RemoveScalarParameter( KeyAreaNamePath[0] );
+		}
 	}
 	return false;
 }
diff --git a/Engine/Source/Editor/MovieSceneTools/Private/TrackEditors/MaterialTrackEditor.cpp b/Engine/Source/Editor/MovieSceneTools/Private/TrackEditors/MaterialTrackEditor.cpp
--- a/Engine/Source/Editor/MovieSceneTools/Private/TrackEditors/MaterialTrackEditor.cpp
+++ b/Engine/Source/Editor/MovieSceneTools/Private/TrackEditors/MaterialTrackEditor.cpp
@@ -154,10 +154,6 @@ void FMaterialTrackEditor::AddScalarParameter( FGuid ObjectBinding, UMovieSceneM
 		float ParameterValue;
 		Material->GetScalarParameterValue(ParameterName, ParameterValue);
 		MaterialTrack->Modify();
-		for ( UMovieSceneSection* Section : MaterialTrack->GetAllSections() )
-		{
-			Section->Modify();
-		}
 		MaterialTrack->AddScalarParameterKey(ParameterName, KeyTime, ParameterValue);
 	}
 	NotifyMovieSceneDataChanged();
@@ -176,10 +172,6 @@ void FMaterialTrackEditor::AddVectorParameter( FGuid ObjectBinding, UMovieSceneM
 		FLinearColor ParameterValue;
 		Material->GetVectorParameterValue( ParameterName, ParameterValue );
 		MaterialTrack->Modify();
-		for ( UMovieSceneSection* Section : MaterialTrack->GetAllSections() )
-		{
-			Section->Modify();
-		}
 		MaterialTrack->AddVectorParameterKey( ParameterName, KeyTime, ParameterValue );
 	}
 	NotifyMovieSceneDataChanged();
diff --git a/Engine/Source/Editor/MovieSceneTools/Private/TrackEditors/ShotTrackEditor/ShotTrackEditor.cpp b/Engine/Source/Editor/MovieSceneTools/Private/TrackEditors/ShotTrackEditor/ShotTrackEditor.cpp
--- a/Engine/Source/Editor/MovieSceneTools/Private/TrackEditors/ShotTrackEditor/ShotTrackEditor.cpp
+++ b/Engine/Source/Editor/MovieSceneTools/Private/TrackEditors/ShotTrackEditor/ShotTrackEditor.cpp
@@ -17,6 +17,10 @@
 #include "AssetToolsModule.h"
 #include "ShotSequencerSection.h"
 
+#include "ActorEditorUtils.h"
+#include "Editor/SceneOutliner/Public/SceneOutliner.h"
+#include "Editor/LevelEditor/Public/LevelEditor.h"
+#include "SceneOutlinerPublicTypes.h"
 
 #define LOCTEXT_NAMESPACE "FShotTrackEditor"
 
@@ -303,83 +307,60 @@ void FShotTrackEditor::HandleAddShotTrackMenuEntryExecute()
 	GetSequencer()->NotifyMovieSceneDataChanged();
 }
 
+bool FShotTrackEditor::IsCameraPickable(const AActor* const PickableActor)
+{
+	if (PickableActor->IsListedInSceneOutliner() &&
+		!FActorEditorUtils::IsABuilderBrush(PickableActor) &&
+		!PickableActor->IsA( AWorldSettings::StaticClass() ) &&
+		!PickableActor->IsPendingKill() &&
+		PickableActor->IsA( ACameraActor::StaticClass()) )
+	{			
+		return true;
+	}
+	return false;
+}
 
 TSharedRef<SWidget> FShotTrackEditor::HandleAddShotComboButtonGetMenuContent()
 {
 	FMenuBuilder MenuBuilder(true, nullptr);
 
-	// get camera actors
-	TArray<AActor*> AllCameras;
-	TArray<AActor*> SelectedCameras;
-
-	UWorld* World = GEditor->GetEditorWorldContext().World();
+	using namespace SceneOutliner;
 
-	if (World != nullptr)
+	SceneOutliner::FInitializationOptions InitOptions;
 	{
-		USelection* SelectedActors = GEditor->GetSelectedActors();
-
-		if (SelectedActors != nullptr)
-		{
-			for (FActorIterator ActorIt(World); ActorIt; ++ActorIt)
-			{
-				AActor* Actor = *ActorIt;
-
-				if (!Actor->IsListedInSceneOutliner() || !Actor->IsA<ACameraActor>())
-				{
-					continue;
-				}
+		InitOptions.Mode = ESceneOutlinerMode::ActorPicker;			
+		InitOptions.bShowHeaderRow = false;
+		InitOptions.bFocusSearchBoxWhenOpened = true;
+		InitOptions.bShowTransient = true;
+		InitOptions.bShowCreateNewFolder = false;
+		// Only want the actor label column
+		InitOptions.ColumnMap.Add(FBuiltInColumnTypes::Label(), FColumnInfo(EColumnVisibility::Visible, 0));
 
-				AllCameras.Add(Actor);
+		// Only display Actors that we can attach too
+		InitOptions.Filters->AddFilterPredicate( SceneOutliner::FActorFilterPredicate::CreateSP(this, &FShotTrackEditor::IsCameraPickable) );
+	}		
 
-				if (SelectedActors->IsSelected(Actor))
-				{
-					SelectedCameras.Add(Actor);
-				}
-			}
-		}
-	}
+	// Actor selector to allow the user to choose a parent actor
+	FSceneOutlinerModule& SceneOutlinerModule = FModuleManager::LoadModuleChecked<FSceneOutlinerModule>( "SceneOutliner" );
 
-	if (SelectedCameras.Num() == 1)
-	{
-		AActor*& SelectedCamera = SelectedCameras[0];
-		FText ActorLabel = FText::FromString(SelectedCamera->GetActorLabel());
+	TSharedRef< SWidget > MenuWidget = 
+		SNew(SHorizontalBox)
 
-		MenuBuilder.BeginSection(NAME_None, FText::FromString("Selected Camera"));
-		{
-			MenuBuilder.AddMenuEntry(
-				FText::Format(LOCTEXT("AddShotEntryTextFormat", "{0}"), ActorLabel),
-				FText::Format(LOCTEXT("AddShotEntryTooltipFormat", "Create a shot track for camera {0}"), ActorLabel),
-				FSlateIcon(),
-				FUIAction(FExecuteAction::CreateRaw(this, &FShotTrackEditor::HandleAddShotComboButtonMenuEntryExecute, SelectedCamera))
-			);
-		}
-		MenuBuilder.EndSection();
-	}
-
-	if (AllCameras.Num() > 0)
-	{
-		AllCameras.Sort([](const AActor& A, const AActor& B)
-		{
-			return A.GetActorLabel() < B.GetActorLabel();
-		});
-
-		MenuBuilder.BeginSection(NAME_None, FText::FromString("All Cameras"));
-		{
-			for (auto Camera : AllCameras)
-			{
-				FText ActorLabel = FText::FromString(Camera->GetActorLabel());
-
-				MenuBuilder.AddMenuEntry(
-					FText::Format(LOCTEXT("AddShotEntryTextFormat", "{0}"), ActorLabel),
-					FText::Format(LOCTEXT("AddShotEntryTooltipFormat", "Create a shot track for camera {0}"), ActorLabel),
-					FSlateIcon(),
-					FUIAction(FExecuteAction::CreateRaw(this, &FShotTrackEditor::HandleAddShotComboButtonMenuEntryExecute, Camera))
-				);
-			}
-		}
-		MenuBuilder.EndSection();
-	}
+		+SHorizontalBox::Slot()
+		.AutoWidth()
+		[
+			SNew(SVerticalBox)
+			+SVerticalBox::Slot()
+			.MaxHeight(400.0f)
+			[
+				SceneOutlinerModule.CreateSceneOutliner(
+					InitOptions,
+					FOnActorPicked::CreateSP(this, &FShotTrackEditor::HandleAddShotComboButtonMenuEntryExecute )
+					)
+			]
+		];
 
+	MenuBuilder.AddWidget(MenuWidget, FText::GetEmpty(), false);
 	return MenuBuilder.MakeWidget();
 }
 
diff --git a/Engine/Source/Editor/MovieSceneTools/Private/TrackEditors/ShotTrackEditor/ShotTrackEditor.h b/Engine/Source/Editor/MovieSceneTools/Private/TrackEditors/ShotTrackEditor/ShotTrackEditor.h
--- a/Engine/Source/Editor/MovieSceneTools/Private/TrackEditors/ShotTrackEditor/ShotTrackEditor.h
+++ b/Engine/Source/Editor/MovieSceneTools/Private/TrackEditors/ShotTrackEditor/ShotTrackEditor.h
@@ -72,6 +72,9 @@ private:
 	/** Callback for generating the menu of the "Add Shot" combo button. */
 	TSharedRef<SWidget> HandleAddShotComboButtonGetMenuContent();
 
+	/** Callback for whether a camera is pickable. */
+	bool IsCameraPickable( const AActor* const PickableActor );
+
 	/** Callback for executing a menu entry in the "Add Shot" combo button. */
 	void HandleAddShotComboButtonMenuEntryExecute(AActor* Camera);
 
diff --git a/Engine/Source/Editor/MovieSceneTools/Private/TrackEditors/TransformTrackEditor.cpp b/Engine/Source/Editor/MovieSceneTools/Private/TrackEditors/TransformTrackEditor.cpp
--- a/Engine/Source/Editor/MovieSceneTools/Private/TrackEditors/TransformTrackEditor.cpp
+++ b/Engine/Source/Editor/MovieSceneTools/Private/TrackEditors/TransformTrackEditor.cpp
@@ -54,19 +54,23 @@ public:
 
 	virtual void GenerateSectionLayout( class ISectionLayoutBuilder& LayoutBuilder ) const override
 	{
+		static const FLinearColor BlueKeyAreaColor(0.0f, 0.0f, 0.7f, 0.2f);
+		static const FLinearColor GreenKeyAreaColor(0.0f, 0.7f, 0.0f, 0.2f);
+		static const FLinearColor RedKeyAreaColor(0.7f, 0.0f, 0.0f, 0.2f);
+
 		UMovieScene3DTransformSection* TransformSection = Cast<UMovieScene3DTransformSection>( &Section );
 
-		TranslationXKeyArea = MakeShareable( new FFloatCurveKeyArea( &TransformSection->GetTranslationCurve( EAxis::X ), TransformSection ) );
-		TranslationYKeyArea = MakeShareable( new FFloatCurveKeyArea( &TransformSection->GetTranslationCurve( EAxis::Y ), TransformSection ) );
-		TranslationZKeyArea = MakeShareable( new FFloatCurveKeyArea( &TransformSection->GetTranslationCurve( EAxis::Z ), TransformSection ) );
+		TranslationXKeyArea = MakeShareable(new FFloatCurveKeyArea(&TransformSection->GetTranslationCurve(EAxis::X), TransformSection, RedKeyAreaColor));
+		TranslationYKeyArea = MakeShareable(new FFloatCurveKeyArea(&TransformSection->GetTranslationCurve(EAxis::Y), TransformSection, GreenKeyAreaColor));
+		TranslationZKeyArea = MakeShareable(new FFloatCurveKeyArea(&TransformSection->GetTranslationCurve(EAxis::Z), TransformSection, BlueKeyAreaColor));
 
-		RotationXKeyArea = MakeShareable( new FFloatCurveKeyArea( &TransformSection->GetRotationCurve( EAxis::X ), TransformSection ) );
-		RotationYKeyArea = MakeShareable( new FFloatCurveKeyArea( &TransformSection->GetRotationCurve( EAxis::Y ), TransformSection ) );
-		RotationZKeyArea = MakeShareable( new FFloatCurveKeyArea( &TransformSection->GetRotationCurve( EAxis::Z ), TransformSection ) );
+		RotationXKeyArea = MakeShareable(new FFloatCurveKeyArea(&TransformSection->GetRotationCurve(EAxis::X), TransformSection, RedKeyAreaColor));
+		RotationYKeyArea = MakeShareable(new FFloatCurveKeyArea(&TransformSection->GetRotationCurve(EAxis::Y), TransformSection, GreenKeyAreaColor));
+		RotationZKeyArea = MakeShareable(new FFloatCurveKeyArea(&TransformSection->GetRotationCurve(EAxis::Z), TransformSection, BlueKeyAreaColor));
 
-		ScaleXKeyArea = MakeShareable( new FFloatCurveKeyArea( &TransformSection->GetScaleCurve( EAxis::X ), TransformSection ) );
-		ScaleYKeyArea = MakeShareable( new FFloatCurveKeyArea( &TransformSection->GetScaleCurve( EAxis::Y ), TransformSection ) );
-		ScaleZKeyArea = MakeShareable( new FFloatCurveKeyArea( &TransformSection->GetScaleCurve( EAxis::Z ), TransformSection ) );
+		ScaleXKeyArea = MakeShareable(new FFloatCurveKeyArea(&TransformSection->GetScaleCurve(EAxis::X), TransformSection, RedKeyAreaColor));
+		ScaleYKeyArea = MakeShareable(new FFloatCurveKeyArea(&TransformSection->GetScaleCurve(EAxis::Y), TransformSection, GreenKeyAreaColor));
+		ScaleZKeyArea = MakeShareable(new FFloatCurveKeyArea(&TransformSection->GetScaleCurve(EAxis::Z), TransformSection, BlueKeyAreaColor));
 
 		// This generates the tree structure for the transform section
 		LayoutBuilder.PushCategory( "Location", NSLOCTEXT("FTransformSection", "LocationArea", "Location") );
diff --git a/Engine/Source/Editor/MovieSceneTools/Public/BoolKeyArea.h b/Engine/Source/Editor/MovieSceneTools/Public/BoolKeyArea.h
--- a/Engine/Source/Editor/MovieSceneTools/Public/BoolKeyArea.h
+++ b/Engine/Source/Editor/MovieSceneTools/Public/BoolKeyArea.h
@@ -3,7 +3,7 @@
 #pragma once
 
 #include "IntegralKeyArea.h"
-
+#include "ClipboardTypes.h"
 
 /**
  * A key area for displaying and editing integral curves representing bools.
diff --git a/Engine/Source/Editor/MovieSceneTools/Public/ByteKeyArea.h b/Engine/Source/Editor/MovieSceneTools/Public/ByteKeyArea.h
--- a/Engine/Source/Editor/MovieSceneTools/Public/ByteKeyArea.h
+++ b/Engine/Source/Editor/MovieSceneTools/Public/ByteKeyArea.h
@@ -3,7 +3,7 @@
 #pragma once
 
 #include "IntegralKeyArea.h"
-
+#include "ClipboardTypes.h"
 
 /** A key area for displaying and editing byte curves. */
 class FByteKeyArea
diff --git a/Engine/Source/Editor/MovieSceneTools/Public/FloatCurveKeyArea.h b/Engine/Source/Editor/MovieSceneTools/Public/FloatCurveKeyArea.h
--- a/Engine/Source/Editor/MovieSceneTools/Public/FloatCurveKeyArea.h
+++ b/Engine/Source/Editor/MovieSceneTools/Public/FloatCurveKeyArea.h
@@ -2,19 +2,21 @@
 
 #pragma once
 
-#include "IKeyArea.h"
+#include "NamedKeyArea.h"
+#include "ClipboardTypes.h"
 
 
 /**
  * A key area for float keys
  */
 class MOVIESCENETOOLS_API FFloatCurveKeyArea
-	: public IKeyArea
+	: public FNamedKeyArea
 {
 public:
 
-	FFloatCurveKeyArea(FRichCurve* InCurve, UMovieSceneSection* InOwningSection)
-		: Curve(InCurve)
+	FFloatCurveKeyArea(FRichCurve* InCurve, UMovieSceneSection* InOwningSection, const FLinearColor& InColor = FLinearColor(0.1f, 0.1f, 0.1f, 0.7f))
+		: Color(InColor)
+		, Curve(InCurve)
 		, OwningSection(InOwningSection)
 	{ }
 
@@ -38,6 +40,7 @@ public:
 	virtual bool CanCreateKeyEditor() override;
 	virtual TSharedRef<SWidget> CreateKeyEditor(ISequencer* Sequencer) override;
 	virtual void DeleteKey(FKeyHandle KeyHandle) override;
+	virtual FLinearColor GetColor() override;
 	virtual ERichCurveExtrapolation GetExtrapolationMode(bool bPreInfinity) const override;
 	virtual ERichCurveInterpMode GetKeyInterpMode(FKeyHandle KeyHandle) const override;
 	virtual ERichCurveTangentMode GetKeyTangentMode(FKeyHandle KeyHandle) const override;
@@ -50,9 +53,14 @@ public:
 	virtual void SetKeyInterpMode(FKeyHandle KeyHandle, ERichCurveInterpMode InterpMode) override;
 	virtual void SetKeyTangentMode(FKeyHandle KeyHandle, ERichCurveTangentMode TangentMode) override;
 	virtual void SetKeyTime(FKeyHandle KeyHandle, float NewKeyTime) const override;
+	virtual void CopyKeys(FMovieSceneClipboardBuilder& ClipboardBuilder, const TFunctionRef<bool(FKeyHandle, const IKeyArea&)>& KeyMask) const override;
+	virtual void PasteKeys(const FMovieSceneClipboardKeyTrack& KeyTrack, const FMovieSceneClipboardEnvironment& SrcEnvironment, const FSequencerPasteEnvironment& DstEnvironment) override;
 
 private:
 
+	/** The key area's color. */
+	FLinearColor Color;
+
 	/** The curve which provides the keys for this key area. */
 	FRichCurve* Curve;
 
diff --git a/Engine/Source/Editor/MovieSceneTools/Public/IntegralKeyArea.h b/Engine/Source/Editor/MovieSceneTools/Public/IntegralKeyArea.h
--- a/Engine/Source/Editor/MovieSceneTools/Public/IntegralKeyArea.h
+++ b/Engine/Source/Editor/MovieSceneTools/Public/IntegralKeyArea.h
@@ -2,14 +2,15 @@
 
 #pragma once
 
-#include "IKeyArea.h"
-
+#include "NamedKeyArea.h"
+#include "MovieSceneClipboard.h"
+#include "SequencerClipboardReconciler.h"
 
 /**
  * Abstract base class for integral curve key areas.
  */
 class MOVIESCENETOOLS_API FIntegralCurveKeyAreaBase
-	: public IKeyArea
+	: public FNamedKeyArea
 {
 public:
 
@@ -25,6 +26,7 @@ public:
 
 	virtual TArray<FKeyHandle> AddKeyUnique( float Time, EMovieSceneKeyInterpolation InKeyInterpolation, float TimeToCopyFrom = FLT_MAX ) override;
 	virtual void DeleteKey(FKeyHandle KeyHandle) override;
+	virtual FLinearColor GetColor() override;
 	virtual ERichCurveExtrapolation GetExtrapolationMode(bool bPreInfinity) const override;
 	virtual ERichCurveInterpMode GetKeyInterpMode(FKeyHandle KeyHandle) const override;
 	virtual ERichCurveTangentMode GetKeyTangentMode(FKeyHandle KeyHandle) const override;
@@ -65,6 +67,65 @@ public:
 		: FIntegralCurveKeyAreaBase(InCurve, InOwningSection)
 	{ }
 
+	virtual void CopyKeys(FMovieSceneClipboardBuilder& ClipboardBuilder, const TFunctionRef<bool(FKeyHandle, const IKeyArea&)>& KeyMask) const override
+	{
+		const UMovieSceneSection* Section = const_cast<FIntegralKeyArea*>(this)->GetOwningSection();
+		UMovieSceneTrack* Track = Section ? Section->GetTypedOuter<UMovieSceneTrack>() : nullptr;
+		if (!Track)
+		{
+			return;
+		}
+
+		FMovieSceneClipboardKeyTrack* KeyTrack = nullptr;
+
+		for (auto It(Curve.GetKeyHandleIterator()); It; ++It)
+		{
+			FKeyHandle Handle = It.Key();
+			if (KeyMask(Handle, *this))
+			{
+				if (!KeyTrack)
+				{
+					KeyTrack = &ClipboardBuilder.FindOrAddKeyTrack<IntegralType>(GetName(), *Track);
+				}
+
+				FIntegralKey Key = Curve.GetKey(Handle);
+				IntegralType Value = MovieSceneClipboard::TImplicitConversionFacade<decltype(Key.Value), IntegralType>::Cast(Key.Value);
+				KeyTrack->AddKey(Key.Time, Value);
+			}
+		}
+	}
+
+	virtual void PasteKeys(const FMovieSceneClipboardKeyTrack& KeyTrack, const FMovieSceneClipboardEnvironment& SrcEnvironment, const FSequencerPasteEnvironment& DstEnvironment) override
+	{
+		float PasteAt = DstEnvironment.CardinalTime;
+
+		KeyTrack.IterateKeys([&](const FMovieSceneClipboardKey& Key){
+			UMovieSceneSection* Section = GetOwningSection();
+			if (!Section)
+			{
+				return true;
+			}
+			
+			if (Section->TryModify())
+			{
+				float Time = PasteAt + Key.GetTime();
+				if (Section->GetStartTime() > Time)
+				{
+					Section->SetStartTime(Time);
+				}
+				if (Section->GetEndTime() < Time)
+				{
+					Section->SetEndTime(Time);
+				}
+
+				FKeyHandle KeyHandle = Curve.UpdateOrAddKey(Time, Key.GetValue<IntegralType>());
+				DstEnvironment.ReportPastedKey(KeyHandle, *this);
+			}
+				
+			return true;
+		});
+	}
+
 public:
 
 	void ClearIntermediateValue()
diff --git a/Engine/Source/Editor/MovieSceneTools/Public/NameCurveKeyArea.h b/Engine/Source/Editor/MovieSceneTools/Public/NameCurveKeyArea.h
--- a/Engine/Source/Editor/MovieSceneTools/Public/NameCurveKeyArea.h
+++ b/Engine/Source/Editor/MovieSceneTools/Public/NameCurveKeyArea.h
@@ -2,14 +2,14 @@
 
 #pragma once
 
-#include "IKeyArea.h"
-
+#include "NamedKeyArea.h"
+#include "ClipboardTypes.h"
 
 /**
  * A key area for FName curves.
  */
 class MOVIESCENETOOLS_API FNameCurveKeyArea
-	: public IKeyArea
+	: public FNamedKeyArea
 {
 public:
 
@@ -32,6 +32,7 @@ public:
 	virtual bool CanCreateKeyEditor() override;
 	virtual TSharedRef<SWidget> CreateKeyEditor(ISequencer* Sequencer) override;
 	virtual void DeleteKey(FKeyHandle KeyHandle) override;
+	virtual FLinearColor GetColor() override;
 	virtual ERichCurveExtrapolation GetExtrapolationMode(bool bPreInfinity) const override;
 	virtual ERichCurveInterpMode GetKeyInterpMode(FKeyHandle Keyhandle) const override;
 	virtual ERichCurveTangentMode GetKeyTangentMode(FKeyHandle KeyHandle) const override;
@@ -44,6 +45,8 @@ public:
 	virtual void SetKeyInterpMode(FKeyHandle KeyHandle, ERichCurveInterpMode InterpMode) override;
 	virtual void SetKeyTangentMode(FKeyHandle KeyHandle, ERichCurveTangentMode TangentMode) override;
 	virtual void SetKeyTime(FKeyHandle KeyHandle, float NewKeyTime) const override;
+	virtual void CopyKeys(FMovieSceneClipboardBuilder& ClipboardBuilder, const TFunctionRef<bool(FKeyHandle, const IKeyArea&)>& KeyMask) const override;
+	virtual void PasteKeys(const FMovieSceneClipboardKeyTrack& KeyTrack, const FMovieSceneClipboardEnvironment& SrcEnvironment, const FSequencerPasteEnvironment& DstEnvironment) override;
 
 protected:
 
diff --git a/Engine/Source/Editor/MovieSceneTools/Public/NamedKeyArea.h b/Engine/Source/Editor/MovieSceneTools/Public/NamedKeyArea.h
new file mode 100644
--- /dev/null
+++ b/Engine/Source/Editor/MovieSceneTools/Public/NamedKeyArea.h
@@ -0,0 +1,31 @@
+// Copyright 1998-2015 Epic Games, Inc. All Rights Reserved.
+
+#pragma once
+
+#include "IKeyArea.h"
+
+/**
+ * Abstract base class for named key areas.
+ */
+class MOVIESCENETOOLS_API FNamedKeyArea
+	: public IKeyArea
+{
+public:
+
+	/** Set the name of this key area */
+	virtual void SetName(FName Name) override
+	{
+		KeyAreaName = Name;
+	}
+
+	/** Get the name of this key area */
+	virtual FName GetName() const override
+	{
+		return KeyAreaName;
+	}
+
+protected:
+
+	/** The name of this key area */
+	FName KeyAreaName;
+};
\ No newline at end of file
diff --git a/Engine/Source/Editor/PropertyEditor/Private/PropertyEditorModule.cpp b/Engine/Source/Editor/PropertyEditor/Private/PropertyEditorModule.cpp
--- a/Engine/Source/Editor/PropertyEditor/Private/PropertyEditorModule.cpp
+++ b/Engine/Source/Editor/PropertyEditor/Private/PropertyEditorModule.cpp
@@ -241,6 +241,8 @@ TSharedRef<IDetailsView> FPropertyEditorModule::CreateDetailView( const FDetails
 		.DetailsViewArgs( DetailsViewArgs );
 
 	AllDetailViews.Add( DetailView );
+
+	PropertyEditorOpened.Broadcast();
 	return DetailView;
 }
 
diff --git a/Engine/Source/Editor/PropertyEditor/Public/PropertyEditorModule.h b/Engine/Source/Editor/PropertyEditor/Public/PropertyEditorModule.h
--- a/Engine/Source/Editor/PropertyEditor/Public/PropertyEditorModule.h
+++ b/Engine/Source/Editor/PropertyEditor/Public/PropertyEditorModule.h
@@ -287,6 +287,9 @@ public:
 	FPropertyTypeLayoutCallback GetPropertyTypeCustomization(const UProperty* InProperty,const IPropertyHandle& PropertyHandle, TSharedPtr<IDetailsView> DetailsViewInstance );
 	bool IsCustomizedStruct(const UStruct* Struct, TSharedPtr<IDetailsView> DetailsViewInstance ) const;
 
+	DECLARE_EVENT(PropertyEditorModule, FPropertyEditorOpenedEvent);
+	virtual FPropertyEditorOpenedEvent& OnPropertyEditorOpened() { return PropertyEditorOpened; }
+
 private:
 
 	/**
@@ -315,4 +318,7 @@ private:
 	FCustomPropertyTypeLayoutMap GlobalPropertyTypeToLayoutMap;
 	/** Instanced property type customization mapped to a specific details view */
 	TMap< TWeakPtr<IDetailsView>, FCustomPropertyTypeLayoutMap > InstancePropertyTypeLayoutMap;
+	/** Event to be called when a property editor is opened */
+	FPropertyEditorOpenedEvent PropertyEditorOpened;
+
 };
diff --git a/Engine/Source/Editor/Sequencer/Private/DisplayNodes/SequencerDisplayNode.cpp b/Engine/Source/Editor/Sequencer/Private/DisplayNodes/SequencerDisplayNode.cpp
--- a/Engine/Source/Editor/Sequencer/Private/DisplayNodes/SequencerDisplayNode.cpp
+++ b/Engine/Source/Editor/Sequencer/Private/DisplayNodes/SequencerDisplayNode.cpp
@@ -413,6 +413,32 @@ void FSequencerDisplayNode::BuildContextMenu(FMenuBuilder& MenuBuilder)
 {
 	TSharedRef<FSequencerDisplayNode> ThisNode = SharedThis(this);
 
+	MenuBuilder.AddMenuEntry(
+		LOCTEXT("ToggleNodeActive", "Active"),
+		LOCTEXT("ToggleNodeActiveTooltip", "Set this node or selected nodes active/inactive"),
+		FSlateIcon(),
+		FUIAction(
+			FExecuteAction::CreateSP(&GetSequencer(), &FSequencer::ToggleNodeActive),
+			FCanExecuteAction(),
+			FIsActionChecked::CreateSP(&GetSequencer(), &FSequencer::IsNodeActive)
+		),
+		NAME_None,
+		EUserInterfaceActionType::ToggleButton
+	);
+
+	MenuBuilder.AddMenuEntry(
+		LOCTEXT("ToggleNodeLock", "Locked"),
+		LOCTEXT("ToggleNodeLockTooltip", "Lock or unlock this node or selected nodes"),
+		FSlateIcon(),
+		FUIAction(
+			FExecuteAction::CreateSP(&GetSequencer(), &FSequencer::ToggleNodeLocked),
+			FCanExecuteAction(),
+			FIsActionChecked::CreateSP(&GetSequencer(), &FSequencer::IsNodeLocked)
+		),
+		NAME_None,
+		EUserInterfaceActionType::ToggleButton
+	);
+
 	MenuBuilder.AddMenuEntry(
 		LOCTEXT("DeleteNode", "Delete"),
 		LOCTEXT("DeleteNodeTooltip", "Delete this or selected nodes"),
diff --git a/Engine/Source/Editor/Sequencer/Private/GroupedKeyArea.cpp b/Engine/Source/Editor/Sequencer/Private/GroupedKeyArea.cpp
--- a/Engine/Source/Editor/Sequencer/Private/GroupedKeyArea.cpp
+++ b/Engine/Source/Editor/Sequencer/Private/GroupedKeyArea.cpp
@@ -388,7 +388,57 @@ bool FGroupedKeyArea::CanCreateKeyEditor()
 	return false;
 }
 
+
+void FGroupedKeyArea::CopyKeys(FMovieSceneClipboardBuilder& ClipboardBuilder, const TFunctionRef<bool(FKeyHandle, const IKeyArea&)>& KeyMask) const
+{
+	const FIndexEntry* IndexEntry = GlobalIndex.Find(IndexKey);
+	if (!IndexEntry)
+	{
+		return;
+	}
+
+	// Since we are a group of nested key areas, we test the key mask for our key handles, and forward on the results to each key area
+
+	// Using ptr as map key is fine here as we know they will not change
+	TMap<const IKeyArea*, TSet<FKeyHandle>> AllValidHandles;
+
+	for (auto& Pair : IndexEntry->HandleToGroup)
+	{
+		if (!KeyMask(Pair.Key, *this) || !Groups.IsValidIndex(Pair.Value))
+		{
+			continue;
+		}
+
+		const FKeyGrouping& Group = Groups[Pair.Value];
+		for (const FKeyGrouping::FKeyIndex& KeyIndex : Group.Keys)
+		{
+			const IKeyArea& KeyArea = KeyAreas[KeyIndex.AreaIndex].Get();
+			AllValidHandles.FindOrAdd(&KeyArea).Add(KeyIndex.KeyHandle);
+		}
+	}
+
+	for (auto& Pair : AllValidHandles)
+	{
+		Pair.Key->CopyKeys(ClipboardBuilder, [&](FKeyHandle Handle, const IKeyArea&){
+			return Pair.Value.Contains(Handle);
+		});
+	}
+}
+
+
 TSharedRef<SWidget> FGroupedKeyArea::CreateKeyEditor(ISequencer* Sequencer)
 {
 	return SNullWidget::NullWidget;
-}
\ No newline at end of file
+}
+
+
+FLinearColor FGroupedKeyArea::GetColor()
+{
+	return FLinearColor(0.1f, 0.1f, 0.1f, 0.7f);
+}
+
+
+void FGroupedKeyArea::PasteKeys(const FMovieSceneClipboardKeyTrack& KeyTrack, const FMovieSceneClipboardEnvironment& SrcEnvironment, const FSequencerPasteEnvironment& DstEnvironment)
+{
+	checkf(false, TEXT("Pasting into grouped key areas is not supported, and should not be used. Iterate child tracks instead."));
+}
diff --git a/Engine/Source/Editor/Sequencer/Private/GroupedKeyArea.h b/Engine/Source/Editor/Sequencer/Private/GroupedKeyArea.h
--- a/Engine/Source/Editor/Sequencer/Private/GroupedKeyArea.h
+++ b/Engine/Source/Editor/Sequencer/Private/GroupedKeyArea.h
@@ -86,6 +86,11 @@ public:
 	virtual UMovieSceneSection*		GetOwningSection() override;
 	virtual bool					CanCreateKeyEditor() override;
 	virtual TSharedRef<SWidget>		CreateKeyEditor(ISequencer* Sequencer) override;
+	virtual void					SetName(FName Name) override {}
+	virtual FName					GetName() const override { return NAME_None; }
+	virtual void					CopyKeys(FMovieSceneClipboardBuilder& ClipboardBuilder, const TFunctionRef<bool(FKeyHandle, const IKeyArea&)>& KeyMask) const override;
+	virtual void					PasteKeys(const FMovieSceneClipboardKeyTrack& KeyTrack, const FMovieSceneClipboardEnvironment& SrcEnvironment, const FSequencerPasteEnvironment& DstEnvironment) override;
+	virtual FLinearColor GetColor() override;
 	/* End IKeyArea interface */
 
 	/** Get the desired tint for the specified key handle */
diff --git a/Engine/Source/Editor/Sequencer/Private/SAnimationOutlinerTreeNode.cpp b/Engine/Source/Editor/Sequencer/Private/SAnimationOutlinerTreeNode.cpp
--- a/Engine/Source/Editor/Sequencer/Private/SAnimationOutlinerTreeNode.cpp
+++ b/Engine/Source/Editor/Sequencer/Private/SAnimationOutlinerTreeNode.cpp
@@ -267,8 +267,10 @@ FReply SAnimationOutlinerTreeNode::OnAddKeyClicked()
 		if (OwningSection == NearestSection)
 		{
 			OwningSection->SetFlags(RF_Transactional);
-			OwningSection->Modify();
-			KeyArea->AddKeyUnique(CurrentTime, Sequencer.GetKeyInterpolation());
+			if (OwningSection->TryModify())
+			{
+				KeyArea->AddKeyUnique(CurrentTime, Sequencer.GetKeyInterpolation());
+			}
 		}
 	}
 
diff --git a/Engine/Source/Editor/Sequencer/Private/SSequencer.cpp b/Engine/Source/Editor/Sequencer/Private/SSequencer.cpp
--- a/Engine/Source/Editor/Sequencer/Private/SSequencer.cpp
+++ b/Engine/Source/Editor/Sequencer/Private/SSequencer.cpp
@@ -35,6 +35,8 @@
 #include "IDetailsView.h"
 #include "PropertyEditorModule.h"
 #include "SSequencerShotFilterOverlay.h"
+#include "GenericCommands.h"
+#include "SequencerContextMenus.h"
 #include "SSequencerTreeViewBox.h"
 #include "NumericTypeInterface.h"
 #include "NumericUnitTypeInterface.inl"
@@ -468,6 +470,22 @@ void SSequencer::BindCommands(TSharedRef<FUICommandList> SequencerCommandBinding
 		}),
 		FIsActionChecked::CreateSP(this, &SSequencer::IsEditToolEnabled, FName("Selection"))
 	);
+
+	SequencerCommandBindings->MapAction(
+		FGenericCommands::Get().Paste,
+		FExecuteAction::CreateSP(this, &SSequencer::Paste),
+		FCanExecuteAction::CreateLambda([&]{
+			return Sequencer.Pin()->GetClipboardStack().Num() != 0;
+		})
+	);
+
+	SequencerCommandBindings->MapAction(
+		FSequencerCommands::Get().PasteFromHistory,
+		FExecuteAction::CreateSP(this, &SSequencer::PasteFromHistory),
+		FCanExecuteAction::CreateLambda([&]{
+			return Sequencer.Pin()->GetClipboardStack().Num() != 0;
+		})
+	);
 }
 
 
@@ -1584,5 +1602,114 @@ FVirtualTrackArea SSequencer::GetVirtualTrackArea() const
 	return FVirtualTrackArea(*Sequencer.Pin(), *TreeView.Get(), TrackArea->GetCachedGeometry());
 }
 
+FPasteContextMenuArgs SSequencer::GeneratePasteArgs(float PasteAtTime, TSharedPtr<FMovieSceneClipboard> Clipboard)
+{
+	TSharedPtr<FSequencer> SequencerPinned = Sequencer.Pin();
+	if (Settings->GetIsSnapEnabled())
+	{
+		PasteAtTime = Settings->SnapTimeToInterval(PasteAtTime);
+	}
+
+	// Open a paste menu at the current mouse position
+	FSlateApplication& Application = FSlateApplication::Get();
+	FVector2D LocalMousePosition = TrackArea->GetCachedGeometry().AbsoluteToLocal(Application.GetCursorPos());
+
+	FVirtualTrackArea VirtualTrackArea = GetVirtualTrackArea();
+
+	// Paste into the currently selected sections, or hit test the mouse position as a last resort
+	TArray<TSharedRef<FSequencerDisplayNode>> PasteIntoNodes;
+	{
+		TSet<TWeakObjectPtr<UMovieSceneSection>> Sections = SequencerPinned->GetSelection().GetSelectedSections();
+		for (const FSequencerSelectedKey& Key : SequencerPinned->GetSelection().GetSelectedKeys())
+		{
+			Sections.Add(Key.Section);
+		}
+
+		for (const FSectionHandle& Handle : GetSectionHandles(Sections))
+		{
+			PasteIntoNodes.Add(Handle.TrackNode.ToSharedRef());
+		}
+	}
+
+	if (PasteIntoNodes.Num() == 0)
+	{
+		TSharedPtr<FSequencerDisplayNode> Node = VirtualTrackArea.HitTestNode(LocalMousePosition.Y);
+		if (Node.IsValid())
+		{
+			PasteIntoNodes.Add(Node.ToSharedRef());
+		}
+	}
+
+	return FPasteContextMenuArgs::PasteInto(MoveTemp(PasteIntoNodes), PasteAtTime, Clipboard);
+}
+
+void SSequencer::Paste()
+{
+	TSharedPtr<FPasteContextMenu> ContextMenu;
+
+	TSharedPtr<FSequencer> SequencerPinned = Sequencer.Pin();
+	if (SequencerPinned->GetClipboardStack().Num() != 0)
+	{
+		FPasteContextMenuArgs Args = GeneratePasteArgs(SequencerPinned->GetGlobalTime(), SequencerPinned->GetClipboardStack().Last());
+		ContextMenu = FPasteContextMenu::CreateMenu(*SequencerPinned, Args);
+	}
+
+	if (!ContextMenu.IsValid() || !ContextMenu->IsValidPaste())
+	{
+		return;
+	}
+	else if (ContextMenu->AutoPaste())
+	{
+		return;
+	}
+
+	const bool bShouldCloseWindowAfterMenuSelection = true;
+	FMenuBuilder MenuBuilder(bShouldCloseWindowAfterMenuSelection, Sequencer.Pin()->GetCommandBindings());
+
+	ContextMenu->PopulateMenu(MenuBuilder);
+
+	FWidgetPath Path;
+	FSlateApplication::Get().FindPathToWidget(AsShared(), Path);
+	
+	FSlateApplication::Get().PushMenu(
+		AsShared(),
+		Path,
+		MenuBuilder.MakeWidget(),
+		FSlateApplication::Get().GetCursorPos(),
+		FPopupTransitionEffect(FPopupTransitionEffect::ContextMenu)
+		);
+}
+
+void SSequencer::PasteFromHistory()
+{
+	TSharedPtr<FSequencer> SequencerPinned = Sequencer.Pin();
+	if (SequencerPinned->GetClipboardStack().Num() == 0)
+	{
+		return;
+	}
+
+	FPasteContextMenuArgs Args = GeneratePasteArgs(SequencerPinned->GetGlobalTime());
+	TSharedPtr<FPasteFromHistoryContextMenu> ContextMenu = FPasteFromHistoryContextMenu::CreateMenu(*SequencerPinned, Args);
+
+	if (ContextMenu.IsValid())
+	{
+		const bool bShouldCloseWindowAfterMenuSelection = true;
+		FMenuBuilder MenuBuilder(bShouldCloseWindowAfterMenuSelection, SequencerPinned->GetCommandBindings());
+
+		ContextMenu->PopulateMenu(MenuBuilder);
+
+		FWidgetPath Path;
+		FSlateApplication::Get().FindPathToWidget(AsShared(), Path);
+		
+		FSlateApplication::Get().PushMenu(
+			AsShared(),
+			Path,
+			MenuBuilder.MakeWidget(),
+			FSlateApplication::Get().GetCursorPos(),
+			FPopupTransitionEffect(FPopupTransitionEffect::ContextMenu)
+			);
+	}
+}
 
 #undef LOCTEXT_NAMESPACE
+
diff --git a/Engine/Source/Editor/Sequencer/Private/SSequencer.h b/Engine/Source/Editor/Sequencer/Private/SSequencer.h
--- a/Engine/Source/Editor/Sequencer/Private/SSequencer.h
+++ b/Engine/Source/Editor/Sequencer/Private/SSequencer.h
@@ -367,6 +367,18 @@ private:
 	/** Gets paint options for painting the playback range on sequencer */
 	FPaintPlaybackRangeArgs GetSectionPlaybackRangeArgs() const;
 
+public:
+
+	/** Open the paste menu */
+	void Paste();
+	
+	/** Open the paste from history menu */
+	void PasteFromHistory();
+
+	/** Generate a paste menu args structure */
+	struct FPasteContextMenuArgs GeneratePasteArgs(float PasteAtTime, TSharedPtr<FMovieSceneClipboard> Clipboard = nullptr);
+
+
 private:
 
 	/** Holds the details view. */
@@ -420,5 +432,4 @@ private:
 
 	/** Called when the user has finished dragging the playback range */
 	FSimpleDelegate OnEndPlaybackRangeDrag;
-
 };
diff --git a/Engine/Source/Editor/Sequencer/Private/SSequencerSection.cpp b/Engine/Source/Editor/Sequencer/Private/SSequencerSection.cpp
--- a/Engine/Source/Editor/Sequencer/Private/SSequencerSection.cpp
+++ b/Engine/Source/Editor/Sequencer/Private/SSequencerSection.cpp
@@ -12,6 +12,7 @@
 #include "SequencerHotspots.h"
 #include "ScopedTransaction.h"
 
+double SSequencerSection::SelectionThrobEndTime = 0;
 
 /** When 0, regeneration of dynamic key layouts is enabled, when non-zero, such behaviour is disabled */
 FThreadSafeCounter LayoutRegenerationLock;
@@ -143,41 +144,42 @@ FSequencerSelectedKey SSequencerSection::CreateKeyUnderMouse( const FVector2D& M
 			}
 
 			FScopedTransaction CreateKeyTransaction(NSLOCTEXT("Sequencer", "CreateKey_Transaction", "Create Key"));
-			Section.Modify();
-
-			// If the pressed key exists, offset the new key and look for it in the newly laid out key areas
-			if (InPressedKey.IsValid())
+			if (Section.TryModify())
 			{
-				// Offset by 1 pixel worth of time
-				const float TimeFuzz = (GetSequencer().GetViewRange().GetUpperBoundValue() - GetSequencer().GetViewRange().GetLowerBoundValue()) / ParentGeometry.GetLocalSize().X;
-
-				TArray<FKeyHandle> KeyHandles = KeyArea->AddKeyUnique(KeyTime+TimeFuzz, GetSequencer().GetKeyInterpolation(), KeyTime);
+				// If the pressed key exists, offset the new key and look for it in the newly laid out key areas
+				if (InPressedKey.IsValid())
+				{
+					// Offset by 1 pixel worth of time
+					const float TimeFuzz = (GetSequencer().GetViewRange().GetUpperBoundValue() - GetSequencer().GetViewRange().GetLowerBoundValue()) / ParentGeometry.GetLocalSize().X;
 
-				Layout = FKeyAreaLayout(*ParentSectionArea, SectionIndex);
+					TArray<FKeyHandle> KeyHandles = KeyArea->AddKeyUnique(KeyTime+TimeFuzz, GetSequencer().GetKeyInterpolation(), KeyTime);
 
-				// Look specifically for the key with the offset key time
-				for (const FKeyAreaLayoutElement& NewElement : Layout->GetElements())
-				{
-					TSharedPtr<IKeyArea> NewKeyArea = NewElement.GetKeyArea();
+					Layout = FKeyAreaLayout(*ParentSectionArea, SectionIndex);
 
-					for (auto KeyHandle : KeyHandles)
+					// Look specifically for the key with the offset key time
+					for (const FKeyAreaLayoutElement& NewElement : Layout->GetElements())
 					{
-						for (auto UnsortedKeyHandle : NewKeyArea->GetUnsortedKeyHandles())
+						TSharedPtr<IKeyArea> NewKeyArea = NewElement.GetKeyArea();
+
+						for (auto KeyHandle : KeyHandles)
 						{
-							if (FMath::IsNearlyEqual(KeyTime+TimeFuzz, NewKeyArea->GetKeyTime(UnsortedKeyHandle), KINDA_SMALL_NUMBER))
+							for (auto UnsortedKeyHandle : NewKeyArea->GetUnsortedKeyHandles())
 							{
-								return FSequencerSelectedKey(Section, NewKeyArea, UnsortedKeyHandle);
+								if (FMath::IsNearlyEqual(KeyTime+TimeFuzz, NewKeyArea->GetKeyTime(UnsortedKeyHandle), KINDA_SMALL_NUMBER))
+								{
+									return FSequencerSelectedKey(Section, NewKeyArea, UnsortedKeyHandle);
+								}
 							}
 						}
 					}
 				}
-			}
-			else
-			{
-				KeyArea->AddKeyUnique(KeyTime, GetSequencer().GetKeyInterpolation());
-				Layout = FKeyAreaLayout(*ParentSectionArea, SectionIndex);
+				else
+				{
+					KeyArea->AddKeyUnique(KeyTime, GetSequencer().GetKeyInterpolation());
+					Layout = FKeyAreaLayout(*ParentSectionArea, SectionIndex);
 						
-				return GetKeyUnderMouse(MousePosition, AllottedGeometry);
+					return GetKeyUnderMouse(MousePosition, AllottedGeometry);
+				}
 			}
 		}
 	}
@@ -224,7 +226,8 @@ FSequencer& SSequencerSection::GetSequencer() const
 int32 SSequencerSection::OnPaint( const FPaintArgs& Args, const FGeometry& AllottedGeometry, const FSlateRect& MyClippingRect, FSlateWindowElementList& OutDrawElements, int32 LayerId, const FWidgetStyle& InWidgetStyle, bool bParentEnabled ) const
 {
 	UMovieSceneSection& SectionObject = *SectionInterface->GetSectionObject();
-	bool bEnabled = bParentEnabled && SectionObject.IsActive();
+	const bool bEnabled = bParentEnabled && SectionObject.IsActive();
+	const bool bLocked = SectionObject.IsLocked();
 	const ESlateDrawEffect::Type DrawEffects = bEnabled ? ESlateDrawEffect::None : ESlateDrawEffect::DisabledEffect;
 
 	int32 StartLayer = SCompoundWidget::OnPaint( Args, AllottedGeometry, MyClippingRect, OutDrawElements, LayerId, InWidgetStyle, bEnabled );
@@ -239,6 +242,21 @@ int32 SSequencerSection::OnPaint( const FPaintArgs& Args, const FGeometry& Allot
 	DrawSelectionBorder(AllottedGeometry, MyClippingRect, OutDrawElements, PostSectionLayer, DrawEffects, SelectionColor);
 	PaintKeys( SectionGeometry, MyClippingRect, OutDrawElements, PostSectionLayer, InWidgetStyle, bEnabled );
 
+	if (bLocked)
+	{
+		static const FName SelectionBorder("Sequencer.Section.LockedBorder");
+
+		FSlateDrawElement::MakeBox(
+			OutDrawElements,
+			PostSectionLayer+1,
+			AllottedGeometry.ToPaintGeometry(),
+			FEditorStyle::GetBrush(SelectionBorder),
+			MyClippingRect,
+			DrawEffects,
+			FLinearColor::Red
+		);
+	}
+
 	// Section name with drop shadow
 	FText SectionTitle = SectionInterface->GetSectionTitle();
 
@@ -309,6 +327,8 @@ void SSequencerSection::PaintKeys( const FGeometry& AllottedGeometry, const FSla
 			* FLinearColor(.25, .25, .25, 1);  // Make the color a little darker since it's not very visible next to white key frames.
 	}
 
+	const float ThrobScaleValue = GetSelectionThrobValue();
+
 	// draw all keys in each key area
 	UMovieSceneSection& SectionObject = *SectionInterface->GetSectionObject();
 
@@ -347,7 +367,7 @@ void SSequencerSection::PaintKeys( const FGeometry& AllottedGeometry, const FSla
 			BackgroundBrush,
 			MyClippingRect,
 			DrawEffects,
-			FLinearColor(0.1f, 0.1f, 0.1f, 0.7f)
+			KeyArea->GetColor()
 		); 
 
 		const int32 KeyLayer = LayerId + 1;
@@ -471,6 +491,9 @@ void SSequencerSection::PaintKeys( const FGeometry& AllottedGeometry, const FSla
 			// draw border
 			float KeyPosition = TimeToPixelConverter.TimeToPixel(KeyTime);
 
+			static FVector2D ThrobAmount(12.f, 12.f);
+			const FVector2D KeySize = bSelected ? SequencerSectionConstants::KeySize + ThrobAmount * ThrobScaleValue : SequencerSectionConstants::KeySize;
+
 			FSlateDrawElement::MakeBox(
 				OutDrawElements,
 				// always draw selected keys on top of other keys
@@ -478,10 +501,10 @@ void SSequencerSection::PaintKeys( const FGeometry& AllottedGeometry, const FSla
 				// Center the key along Y.  Ensure the middle of the key is at the actual key time
 				KeyAreaGeometry.ToPaintGeometry(
 					FVector2D(
-						KeyPosition - FMath::CeilToFloat(SequencerSectionConstants::KeySize.X / 2.0f),
-						((KeyAreaGeometry.Size.Y / 2.0f) - (SequencerSectionConstants::KeySize.Y / 2.0f))
+						KeyPosition - FMath::CeilToFloat(KeySize.X / 2.0f),
+						((KeyAreaGeometry.Size.Y / 2.0f) - (KeySize.Y / 2.0f))
 					),
-					SequencerSectionConstants::KeySize
+					KeySize
 				),
 				KeyBrush,
 				MyClippingRect,
@@ -497,10 +520,10 @@ void SSequencerSection::PaintKeys( const FGeometry& AllottedGeometry, const FSla
 				// Center the key along Y.  Ensure the middle of the key is at the actual key time
 				KeyAreaGeometry.ToPaintGeometry(
 					FVector2D(
-						KeyPosition - FMath::CeilToFloat(SequencerSectionConstants::KeySize.X / 2.0f - BrushBorderWidth),
-						((KeyAreaGeometry.Size.Y / 2.0f) - (SequencerSectionConstants::KeySize.Y / 2.0f - BrushBorderWidth))
+						KeyPosition - FMath::CeilToFloat(KeySize.X / 2.0f - BrushBorderWidth),
+						((KeyAreaGeometry.Size.Y / 2.0f) - (KeySize.Y / 2.0f - BrushBorderWidth))
 					),
-					SequencerSectionConstants::KeySize - 2.0f * BrushBorderWidth
+					KeySize - 2.0f * BrushBorderWidth
 				),
 				KeyBrush,
 				MyClippingRect,
@@ -644,13 +667,17 @@ FReply SSequencerSection::OnMouseButtonDown( const FGeometry& MyGeometry, const
 
 		if (NewKey.IsValid())
 		{
+			HoveredKey = NewKey;
+
 			Sequencer.GetSelection().EmptySelectedKeys();
 			Sequencer.GetSelection().AddToSelection(NewKey);
 
 			// Pass the event to the tool to copy the hovered key and move it
 			auto& EditTool = GetSequencer().GetEditTool();
 			EditTool.SetHotspot( MakeShareable( new FKeyHotspot(NewKey, ParentSectionArea) ) );
-			return FReply::Handled();
+
+			// Return unhandled so that the EditTool can handle the mouse down based on the newly created keyframe and prepare to move it
+			return FReply::Unhandled();
 		}
 	}
 
@@ -713,4 +740,28 @@ void SSequencerSection::OnMouseLeave( const FPointerEvent& MouseEvent )
 {
 	SCompoundWidget::OnMouseLeave( MouseEvent );
 	GetSequencer().GetEditTool().SetHotspot(nullptr);
+}
+
+static float ThrobDurationSeconds = .5f;
+void SSequencerSection::ThrobSelection(int32 ThrobCount)
+{
+	SelectionThrobEndTime = FPlatformTime::Seconds() + ThrobCount*ThrobDurationSeconds;
+}
+
+float EvaluateThrob(float Alpha)
+{
+	return .5f - FMath::Cos(FMath::Pow(Alpha, 0.5f) * 2 * PI) * .5f;
+}
+
+float SSequencerSection::GetSelectionThrobValue()
+{
+	double CurrentTime = FPlatformTime::Seconds();
+
+	if (SelectionThrobEndTime > CurrentTime)
+	{
+		float Difference = SelectionThrobEndTime - CurrentTime;
+		return EvaluateThrob(1.f - FMath::Fmod(Difference, ThrobDurationSeconds));
+	}
+
+	return 0.f;
 }
\ No newline at end of file
diff --git a/Engine/Source/Editor/Sequencer/Private/SSequencerSection.h b/Engine/Source/Editor/Sequencer/Private/SSequencerSection.h
--- a/Engine/Source/Editor/Sequencer/Private/SSequencerSection.h
+++ b/Engine/Source/Editor/Sequencer/Private/SSequencerSection.h
@@ -101,6 +101,15 @@ private:
 	 */
 	FGeometry MakeSectionGeometryWithoutHandles( const FGeometry& AllottedGeometry, const TSharedPtr<ISequencerSection>& InSectionInterface ) const;
 
+
+public:
+
+	/** Indicate that the current key selection should throb the specified number of times. A single throb takes 0.2s. */
+	static void ThrobSelection(int32 ThrobCount = 1);
+
+	/** Get a value between 0 and 1 that indicicates the amount of throb-scale to apply to the currently selected keys */
+	static float GetSelectionThrobValue();
+
 private:
 	/** Interface to section data */
 	TSharedPtr<ISequencerSection> SectionInterface;
@@ -112,4 +121,6 @@ private:
 	int32 SectionIndex;
 	/** Cached parent geometry to pass down to any section interfaces that need it during tick */
 	FGeometry ParentGeometry;
+	/** The end time for a throbbing animation for selected keys */
+	static double SelectionThrobEndTime;
 };
\ No newline at end of file
