Copying //UE4/Orion-Staging to Dev-Main (//UE4/Dev-Main)

diff --git a/Engine/Source/Editor/Sequencer/Private/Sequencer.cpp b/Engine/Source/Editor/Sequencer/Private/Sequencer.cpp
--- a/Engine/Source/Editor/Sequencer/Private/Sequencer.cpp
+++ b/Engine/Source/Editor/Sequencer/Private/Sequencer.cpp
@@ -44,8 +44,10 @@
 #include "Tracks/MovieScene3DTransformTrack.h"
 #include "LevelEditor.h"
 #include "IMenu.h"
+#include "MovieSceneClipboard.h"
+#include "SequencerClipboardReconciler.h"
 #include "STextEntryPopup.h"
-
+#include "SequencerHotspots.h"
 #include "MovieSceneCaptureDialogModule.h"
 #include "AutomatedLevelSequenceCapture.h"
 
@@ -139,6 +141,9 @@ void FSequencer::InitSequencer(const FSequencerInitParams& InitParams, const TSh
 			GEditor->GetActorRecordingState().AddSP(this, &FSequencer::GetActorRecordingState);
 
 			ActivateDetailKeyframeHandler();
+				
+			FPropertyEditorModule& EditModule = FModuleManager::Get().GetModuleChecked<FPropertyEditorModule>("PropertyEditor");
+			EditModule.OnPropertyEditorOpened().AddSP(this, &FSequencer::OnPropertyEditorOpened);
 		}
 
 		// Create tools and bind them to this sequencer
@@ -218,7 +223,15 @@ void FSequencer::Close()
 			GLevelEditorModeTools().DeactivateMode( FSequencerEdMode::EM_SequencerMode );
 		}
 	}
-	DeactivateDetailKeyframeHandler();
+
+	if( bIsEditingWithinLevelEditor )
+	{
+		DeactivateDetailKeyframeHandler();
+			
+		FPropertyEditorModule& EditModule = FModuleManager::Get().GetModuleChecked<FPropertyEditorModule>("PropertyEditor");			
+		EditModule.OnPropertyEditorOpened().RemoveAll(this);
+	}
+
 	SequencerWidget.Reset();
 }
 
@@ -365,7 +378,7 @@ void FSequencer::DeleteSections(const TSet<TWeakObjectPtr<UMovieSceneSection>>&
 
 	for (const auto Section : Sections)
 	{
-		if (!Section.IsValid())
+		if (!Section.IsValid() || Section->IsLocked())
 		{
 			continue;
 		}
@@ -400,9 +413,11 @@ void FSequencer::DeleteSelectedKeys()
 	{
 		if (Key.IsValid())
 		{
-			Key.Section->Modify();
-			Key.KeyArea->DeleteKey(Key.KeyHandle.GetValue());
-			bAnythingRemoved = true;
+			if (Key.Section->TryModify())
+			{
+				Key.KeyArea->DeleteKey(Key.KeyHandle.GetValue());
+				bAnythingRemoved = true;
+			}
 		}
 	}
 
@@ -430,10 +445,12 @@ void FSequencer::SetInterpTangentMode(ERichCurveInterpMode InterpMode, ERichCurv
 	{
 		if (Key.IsValid())
 		{
-			Key.Section->Modify();
-			Key.KeyArea->SetKeyInterpMode(Key.KeyHandle.GetValue(), InterpMode);
-			Key.KeyArea->SetKeyTangentMode(Key.KeyHandle.GetValue(), TangentMode);
-			bAnythingChanged = true;
+			if (Key.Section->TryModify())
+			{
+				Key.KeyArea->SetKeyInterpMode(Key.KeyHandle.GetValue(), InterpMode);
+				Key.KeyArea->SetKeyTangentMode(Key.KeyHandle.GetValue(), TangentMode);
+				bAnythingChanged = true;
+			}
 		}
 	}
 
@@ -476,18 +493,20 @@ void FSequencer::SnapToFrame()
 	{
 		if (Key.IsValid())
 		{
-			Key.Section->Modify();
-			float NewKeyTime = Key.KeyArea->GetKeyTime(Key.KeyHandle.GetValue());
+			if (Key.Section->TryModify())
+			{
+				float NewKeyTime = Key.KeyArea->GetKeyTime(Key.KeyHandle.GetValue());
 
-			// Convert to frame
-			float FrameRate = 1.0f / Settings->GetTimeSnapInterval();
-			int32 NewFrame = SequencerHelpers::TimeToFrame(NewKeyTime, FrameRate);
+				// Convert to frame
+				float FrameRate = 1.0f / Settings->GetTimeSnapInterval();
+				int32 NewFrame = SequencerHelpers::TimeToFrame(NewKeyTime, FrameRate);
 
-			// Convert back to time
-			NewKeyTime = SequencerHelpers::FrameToTime(NewFrame, FrameRate);
+				// Convert back to time
+				NewKeyTime = SequencerHelpers::FrameToTime(NewFrame, FrameRate);
 
-			Key.KeyArea->SetKeyTime(Key.KeyHandle.GetValue(), NewKeyTime);
-			bAnythingChanged = true;
+				Key.KeyArea->SetKeyTime(Key.KeyHandle.GetValue(), NewKeyTime);
+				bAnythingChanged = true;
+			}
 		}
 	}
 
@@ -1715,6 +1734,11 @@ void FSequencer::DeactivateDetailKeyframeHandler()
 	}
 }
 
+void FSequencer::OnPropertyEditorOpened()
+{
+	ActivateDetailKeyframeHandler();
+}
+
 void FSequencer::UpdatePreviewLevelViewportClientFromCameraCut(FLevelEditorViewportClient& InViewportClient, UObject* InCameraObject) const
 {
 	ACameraActor* CameraActor = Cast<ACameraActor>(InCameraObject);
@@ -2080,6 +2104,82 @@ void FSequencer::DeleteSelectedNodes()
 	}
 }
 
+void FSequencer::ToggleNodeActive()
+{
+	bool bIsActive = !IsNodeActive();
+
+	const FScopedTransaction Transaction( NSLOCTEXT("Sequencer", "ToggleNodeActive", "Toggle Node Active") );
+
+	for (auto OutlinerNode : Selection.GetSelectedOutlinerNodes())
+	{
+		TSet<TWeakObjectPtr<UMovieSceneSection> > Sections;
+		SequencerHelpers::GetAllSections(OutlinerNode, Sections);
+
+		for (auto Section : Sections)
+		{
+			Section->SetIsActive(bIsActive);
+		}
+	}
+}
+
+bool FSequencer::IsNodeActive() const
+{
+	// Active only if all are active
+	for (auto OutlinerNode : Selection.GetSelectedOutlinerNodes())
+	{
+		TSet<TWeakObjectPtr<UMovieSceneSection> > Sections;
+		SequencerHelpers::GetAllSections(OutlinerNode, Sections);
+
+		for (auto Section : Sections)
+		{
+			if (!Section->IsActive())
+			{
+				return false;
+			}
+		}
+	}
+	return true;
+}
+
+void FSequencer::ToggleNodeLocked()
+{
+	bool bIsLocked = !IsNodeLocked();
+
+	const FScopedTransaction Transaction( NSLOCTEXT("Sequencer", "ToggleNodeLocked", "Toggle Node Locked") );
+
+	for (auto OutlinerNode : Selection.GetSelectedOutlinerNodes())
+	{
+		TSet<TWeakObjectPtr<UMovieSceneSection> > Sections;
+		SequencerHelpers::GetAllSections(OutlinerNode, Sections);
+
+		for (auto Section : Sections)
+		{
+			Section->SetIsLocked(bIsLocked);
+		}
+	}
+}
+
+bool FSequencer::IsNodeLocked() const
+{
+	// Locked only if all are locked
+	int NumSections = 0;
+	for (auto OutlinerNode : Selection.GetSelectedOutlinerNodes())
+	{
+		TSet<TWeakObjectPtr<UMovieSceneSection> > Sections;
+		SequencerHelpers::GetAllSections(OutlinerNode, Sections);
+
+		for (auto Section : Sections)
+		{
+			if (!Section->IsLocked())
+			{
+				return false;
+			}
+			++NumSections;
+		}
+	}
+	return NumSections > 0;
+}
+
 void FSequencer::TogglePlay()
 {
 	OnPlay();
@@ -2215,13 +2315,85 @@ void FSequencer::OnSetKeyTimeTextCommitted(const FText& InText, ETextCommit::Typ
 		{
 			if (Key.IsValid())
 			{
-				Key.Section->Modify();
-				Key.KeyArea->SetKeyTime(Key.KeyHandle.GetValue(), NewKeyTime);
+				if (Key.Section->TryModify())
+				{
+					Key.KeyArea->SetKeyTime(Key.KeyHandle.GetValue(), NewKeyTime);
+				}
 			}
 		}
 	}
 }
 
+TArray<TSharedPtr<FMovieSceneClipboard>> GClipboardStack;
+
+void FSequencer::CopySelectedKeys()
+{
+	TOptional<float> CopyRelativeTo;
+	
+	// Copy relative to the current key hotspot, if applicable
+	TSharedPtr<ISequencerHotspot> Hotspot = SequencerWidget->GetEditTool().GetHotspot();
+	if (Hotspot.IsValid() && Hotspot->GetType() == ESequencerHotspot::Key)
+	{
+		FKeyHotspot* KeyHotspot = StaticCastSharedPtr<FKeyHotspot>(Hotspot).Get();
+		if (KeyHotspot->Key.KeyArea.IsValid() && KeyHotspot->Key.KeyHandle.IsSet())
+		{
+			CopyRelativeTo = KeyHotspot->Key.KeyArea->GetKeyTime(KeyHotspot->Key.KeyHandle.GetValue());
+		}
+	}
+
+	FMovieSceneClipboardBuilder Builder;
+
+	// Map selected keys to their key areas
+	TMap<const IKeyArea*, TArray<FKeyHandle>> KeyAreaMap;
+	for (const FSequencerSelectedKey& Key : Selection.GetSelectedKeys())
+	{
+		if (Key.KeyHandle.IsSet())
+		{
+			KeyAreaMap.FindOrAdd(Key.KeyArea.Get()).Add(Key.KeyHandle.GetValue());
+		}
+	}
+
+	// Serialize each key area to the clipboard
+	for (auto& Pair : KeyAreaMap)
+	{
+		Pair.Key->CopyKeys(Builder, [&](FKeyHandle Handle, const IKeyArea&){
+			return Pair.Value.Contains(Handle);
+		});
+	}
+
+
+	TSharedRef<FMovieSceneClipboard> Clipboard = MakeShareable( new FMovieSceneClipboard(Builder.Commit(CopyRelativeTo)) );
+	
+	GClipboardStack.Push(Clipboard);
+
+	if (GClipboardStack.Num() > 10)
+	{
+		GClipboardStack.RemoveAt(0, 1);
+	}
+}
+
+void FSequencer::CutSelectedKeys()
+{
+	FScopedTransaction CutKeysTransaction( LOCTEXT("CutSelectedKeys_Transaction", "Cut Selected Keys") );
+	CopySelectedKeys();
+	DeleteSelectedKeys();
+}
+
+const TArray<TSharedPtr<FMovieSceneClipboard>>& FSequencer::GetClipboardStack() const
+{
+	return GClipboardStack;
+}
+
+void FSequencer::OnClipboardUsed(TSharedPtr<FMovieSceneClipboard> Clipboard)
+{
+	Clipboard->GetEnvironment().DateTime = FDateTime::Now();
+
+	// Last entry in the stack should be the most up-to-date
+	GClipboardStack.Sort([](const TSharedPtr<FMovieSceneClipboard>& A, const TSharedPtr<FMovieSceneClipboard>& B){
+		return A->GetEnvironment().DateTime < B->GetEnvironment().DateTime;
+	});
+}
+
 void FSequencer::GenericTextEntryModeless(const FText& DialogText, const FText& DefaultText, FOnTextCommitted OnTextComitted)
 {
 	TSharedRef<STextEntryPopup> TextEntryPopup = 
@@ -2481,6 +2653,34 @@ void FSequencer::BindSequencerCommands()
 		FCanExecuteAction::CreateLambda( []{ return true; } ),
 		FIsActionChecked::CreateLambda( [this]{ return Settings->GetShowCurveEditor(); } ) );
 
+	auto CanCutOrCopy = [this]{
+		UMovieSceneTrack* Track = nullptr;
+		for (FSequencerSelectedKey Key : Selection.GetSelectedKeys())
+		{
+			if (!Track)
+			{
+				Track = Key.Section->GetTypedOuter<UMovieSceneTrack>();
+			}
+			if (!Track || Track != Key.Section->GetTypedOuter<UMovieSceneTrack>())
+			{
+				return false;
+			}
+		}
+		return true;
+	};
+
+	SequencerCommandBindings->MapAction(
+		FGenericCommands::Get().Cut,
+		FExecuteAction::CreateSP(this, &FSequencer::CutSelectedKeys),
+		FCanExecuteAction::CreateLambda(CanCutOrCopy)
+	);
+
+	SequencerCommandBindings->MapAction(
+		FGenericCommands::Get().Copy,
+		FExecuteAction::CreateSP(this, &FSequencer::CopySelectedKeys),
+		FCanExecuteAction::CreateLambda(CanCutOrCopy)
+	);
+
 	SequencerCommandBindings->MapAction( Commands.RenderMovie,
 		FExecuteAction::CreateLambda([this]{
 			FLevelEditorModule& LevelEditorModule = FModuleManager::GetModuleChecked<FLevelEditorModule>(TEXT("LevelEditor"));
@@ -2509,8 +2709,8 @@ void FSequencer::BindSequencerCommands()
 				const int32 SequenceStartFrame = FMath::RoundToInt( GetPlaybackRange().GetLowerBoundValue() * MovieSceneCapture->Settings.FrameRate );
 				const int32 SequenceEndFrame = FMath::Max( SequenceStartFrame, FMath::RoundToInt( GetPlaybackRange().GetUpperBoundValue() * MovieSceneCapture->Settings.FrameRate ) );
 
-				MovieSceneCapture->Settings.StartFrame = SequenceStartFrame;
-				MovieSceneCapture->Settings.EndFrame = SequenceEndFrame;
+				MovieSceneCapture->StartFrame = SequenceStartFrame;
+				MovieSceneCapture->EndFrame = SequenceEndFrame;
 			}
 
 			IMovieSceneCaptureDialogModule::Get().OpenDialog(LevelEditorModule.GetLevelEditorTabManager().ToSharedRef(), MovieSceneCapture);
diff --git a/Engine/Source/Editor/Sequencer/Private/Sequencer.h b/Engine/Source/Editor/Sequencer/Private/Sequencer.h
--- a/Engine/Source/Editor/Sequencer/Private/Sequencer.h
+++ b/Engine/Source/Editor/Sequencer/Private/Sequencer.h
@@ -4,6 +4,7 @@
 
 #include "Editor/EditorWidgets/Public/ITransportControl.h"
 #include "EditorUndoClient.h"
+#include "MovieSceneClipboard.h"
 
 
 class FMenuBuilder;
@@ -293,11 +294,33 @@ public:
 	void DeleteNode(TSharedRef<FSequencerDisplayNode> NodeToBeDeleted);
 	void DeleteSelectedNodes();
 
+	/** Called when a user executes the active node menu item */
+	void ToggleNodeActive();
+	bool IsNodeActive() const;
+
+	/** Called when a user executes the locked node menu item */
+	void ToggleNodeLocked();
+	bool IsNodeLocked() const;
+
 	/** Called when a user executes the set key time for selected keys */
 	bool CanSetKeyTime() const;
 	void SetKeyTime(const bool bUseFrames);
 	void OnSetKeyTimeTextCommitted(const FText& InText, ETextCommit::Type CommitInfo, const bool bUseFrames);
 
+public:
+
+	/** Copy the selected keys to the clipboard */
+	void CopySelectedKeys();
+
+	/** Copy the selected keys to the clipboard, then delete them as part of an undoable transaction */
+	void CutSelectedKeys();
+
+	/** Get the in-memory clipboard stack */
+	const TArray<TSharedPtr<FMovieSceneClipboard>>& GetClipboardStack() const;
+
+	/** Promote a clipboard to the top of the clipboard stack, and update its timestamp */
+	void OnClipboardUsed(TSharedPtr<FMovieSceneClipboard> Clipboard);
+
 public:
 
 	/** Access the currently enabled edit tool */
@@ -456,6 +479,7 @@ protected:
 
 	void ActivateDetailKeyframeHandler();
 	void DeactivateDetailKeyframeHandler();
+	void OnPropertyEditorOpened();
 
 	//~ Begin FEditorUndoClient Interface
 	virtual void PostUndo(bool bSuccess) override;
diff --git a/Engine/Source/Editor/Sequencer/Private/SequencerCommands.cpp b/Engine/Source/Editor/Sequencer/Private/SequencerCommands.cpp
--- a/Engine/Source/Editor/Sequencer/Private/SequencerCommands.cpp
+++ b/Engine/Source/Editor/Sequencer/Private/SequencerCommands.cpp
@@ -64,6 +64,8 @@ void FSequencerCommands::RegisterCommands()
 	UI_COMMAND( MarqueeTool, "Marquee", "Activates the marquee selection tool", EUserInterfaceActionType::ToggleButton, FInputChord( EKeys::M ) );
 
 	UI_COMMAND( RenderMovie, "Render Movie", "Render this movie to a video, or image frame sequence", EUserInterfaceActionType::Button, FInputChord() );
+
+	UI_COMMAND( PasteFromHistory, "Paste From History", "Paste from the sequencer clipboard history", EUserInterfaceActionType::Button, FInputChord(EModifierKey::Control | EModifierKey::Shift, EKeys::V) );
 }
 
 #undef LOCTEXT_NAMESPACE
\ No newline at end of file
diff --git a/Engine/Source/Editor/Sequencer/Private/SequencerCommands.h b/Engine/Source/Editor/Sequencer/Private/SequencerCommands.h
--- a/Engine/Source/Editor/Sequencer/Private/SequencerCommands.h
+++ b/Engine/Source/Editor/Sequencer/Private/SequencerCommands.h
@@ -147,6 +147,9 @@ public:
 	/** Open a panel that enables exporting the sequence to a movie */
 	TSharedPtr< FUICommandInfo > RenderMovie;
 
+	/** Paste from the sequencer clipboard history */
+	TSharedPtr< FUICommandInfo > PasteFromHistory;
+
 	/**
 	 * Initialize commands
 	 */
diff --git a/Engine/Source/Editor/Sequencer/Private/SequencerCommonHelpers.cpp b/Engine/Source/Editor/Sequencer/Private/SequencerCommonHelpers.cpp
--- a/Engine/Source/Editor/Sequencer/Private/SequencerCommonHelpers.cpp
+++ b/Engine/Source/Editor/Sequencer/Private/SequencerCommonHelpers.cpp
@@ -2,6 +2,8 @@
 
 #include "SequencerPrivatePCH.h"
 #include "SequencerCommonHelpers.h"
+#include "MovieSceneSection.h"
+#include "MovieSceneTrack.h"
 
 void SequencerHelpers::GetAllKeyAreas(TSharedPtr<FSequencerDisplayNode> DisplayNode, TSet<TSharedPtr<IKeyArea>>& KeyAreas)
 {
@@ -53,6 +55,41 @@ void SequencerHelpers::GetDescendantNodes(TSharedRef<FSequencerDisplayNode> Disp
 	}
 }
 
+void SequencerHelpers::GetAllSections(TSharedRef<FSequencerDisplayNode> DisplayNode, TSet<TWeakObjectPtr<UMovieSceneSection>>& Sections)
+{
+	TSet<TSharedRef<FSequencerDisplayNode> > AllNodes;
+	AllNodes.Add(DisplayNode);
+	GetDescendantNodes(DisplayNode, AllNodes);
+
+	for (auto NodeToCheck : AllNodes)
+	{
+		TSet<TSharedPtr<IKeyArea> > KeyAreas;
+		GetAllKeyAreas(NodeToCheck, KeyAreas);
+		
+		for (auto KeyArea : KeyAreas)
+		{
+			UMovieSceneSection* OwningSection = KeyArea->GetOwningSection();
+			if (OwningSection != nullptr)
+			{
+				Sections.Add(OwningSection);	
+			}
+		}
+
+		if (NodeToCheck->GetType() == ESequencerNode::Track)
+		{
+			TSharedRef<const FSequencerTrackNode> TrackNode = StaticCastSharedRef<const FSequencerTrackNode>( NodeToCheck );
+			UMovieSceneTrack* Track = TrackNode->GetTrack();
+			if (Track != nullptr)
+			{
+				for (auto Section : Track->GetAllSections())
+				{
+					Sections.Add(Section);
+				}
+			}
+		}
+	}
+}
+
 int32 SequencerHelpers::TimeToFrame(float Time, float FrameRate)
 {
 	float Frame = Time * FrameRate;
diff --git a/Engine/Source/Editor/Sequencer/Private/SequencerCommonHelpers.h b/Engine/Source/Editor/Sequencer/Private/SequencerCommonHelpers.h
--- a/Engine/Source/Editor/Sequencer/Private/SequencerCommonHelpers.h
+++ b/Engine/Source/Editor/Sequencer/Private/SequencerCommonHelpers.h
@@ -20,6 +20,11 @@ public:
 	*/
 	static void GetDescendantNodes(TSharedRef<FSequencerDisplayNode> DisplayNode, TSet<TSharedRef<FSequencerDisplayNode> >& Nodes);
 
+	/**
+	* Gets all sections from the requested node
+	*/
+	static void GetAllSections(TSharedRef<FSequencerDisplayNode> DisplayNode, TSet<TWeakObjectPtr<UMovieSceneSection>>& Sections);
+
 	/**
 	* Convert time to frame
 	*/
diff --git a/Engine/Source/Editor/Sequencer/Private/SequencerContextMenus.cpp b/Engine/Source/Editor/Sequencer/Private/SequencerContextMenus.cpp
--- a/Engine/Source/Editor/Sequencer/Private/SequencerContextMenus.cpp
+++ b/Engine/Source/Editor/Sequencer/Private/SequencerContextMenus.cpp
@@ -6,6 +6,13 @@
 #include "ScopedTransaction.h"
 #include "MovieSceneToolHelpers.h"
 #include "MovieSceneSection.h"
+#include "SSequencerSection.h"
+#include "MovieSceneTrack.h"
+#include "MovieSceneCommonHelpers.h"
+#include "GenericCommands.h"
+#include "SequencerCommands.h"
+
+#define LOCTEXT_NAMESPACE "SequencerContextMenus"
 
 void FKeyContextMenu::BuildMenu(FMenuBuilder& MenuBuilder, FSequencer& InSequencer)
 {
@@ -16,12 +23,24 @@ void FKeyContextMenu::BuildMenu(FMenuBuilder& MenuBuilder, FSequencer& InSequenc
 void FKeyContextMenu::PopulateMenu(FMenuBuilder& MenuBuilder)
 {
 	FSequencer* SequencerPtr = &Sequencer.Get();
-	MenuBuilder.BeginSection("SequencerInterpolation", NSLOCTEXT("Sequencer", "KeyInterpolationMenu", "Key Interpolation"));
+
+	MenuBuilder.BeginSection("SequencerKeyEdit", LOCTEXT("EditMenu", "Edit"));
+	{
+		TSharedPtr<ISequencerHotspot> Hotspot = SequencerPtr->GetEditTool().GetHotspot();
+		if (Hotspot.IsValid() && Hotspot->GetType() == ESequencerHotspot::Key)
+		{
+			MenuBuilder.AddMenuEntry(FGenericCommands::Get().Cut);
+			MenuBuilder.AddMenuEntry(FGenericCommands::Get().Copy);
+		}
+	}
+	MenuBuilder.EndSection(); // SequencerKeyEdit
+
+	MenuBuilder.BeginSection("SequencerInterpolation", LOCTEXT("KeyInterpolationMenu", "Key Interpolation"));
 	{
 		MenuBuilder.AddMenuEntry(
-			NSLOCTEXT("Sequencer", "SetKeyInterpolationAuto", "Auto"),
-			NSLOCTEXT("Sequencer", "SetKeyInterpolationAutoTooltip", "Set key interpolation to auto"),
-			FSlateIcon(),
+			LOCTEXT("SetKeyInterpolationAuto", "Cubic (Auto)"),
+			LOCTEXT("SetKeyInterpolationAutoTooltip", "Set key interpolation to auto"),
+			FSlateIcon(FEditorStyle::GetStyleSetName(), "Sequencer.IconKeyAuto"),
 			FUIAction(
 				FExecuteAction::CreateSP(SequencerPtr, &FSequencer::SetInterpTangentMode, RCIM_Cubic, RCTM_Auto),
 				FCanExecuteAction(),
@@ -31,9 +50,9 @@ void FKeyContextMenu::PopulateMenu(FMenuBuilder& MenuBuilder)
 		);
 
 		MenuBuilder.AddMenuEntry(
-			NSLOCTEXT("Sequencer", "SetKeyInterpolationUser", "User"),
-			NSLOCTEXT("Sequencer", "SetKeyInterpolationUserTooltip", "Set key interpolation to user"),
-			FSlateIcon(),
+			LOCTEXT("SetKeyInterpolationUser", "Cubic (User)"),
+			LOCTEXT("SetKeyInterpolationUserTooltip", "Set key interpolation to user"),
+			FSlateIcon(FEditorStyle::GetStyleSetName(), "Sequencer.IconKeyUser"),
 			FUIAction(
 				FExecuteAction::CreateSP(SequencerPtr, &FSequencer::SetInterpTangentMode, RCIM_Cubic, RCTM_User),
 				FCanExecuteAction(),
@@ -43,9 +62,9 @@ void FKeyContextMenu::PopulateMenu(FMenuBuilder& MenuBuilder)
 		);
 
 		MenuBuilder.AddMenuEntry(
-			NSLOCTEXT("Sequencer", "SetKeyInterpolationBreak", "Break"),
-			NSLOCTEXT("Sequencer", "SetKeyInterpolationBreakTooltip", "Set key interpolation to break"),
-			FSlateIcon(),
+			LOCTEXT("SetKeyInterpolationBreak", "Cubic (Break)"),
+			LOCTEXT("SetKeyInterpolationBreakTooltip", "Set key interpolation to break"),
+			FSlateIcon(FEditorStyle::GetStyleSetName(), "Sequencer.IconKeyBreak"),
 			FUIAction(
 				FExecuteAction::CreateSP(SequencerPtr, &FSequencer::SetInterpTangentMode, RCIM_Cubic, RCTM_Break),
 				FCanExecuteAction(),
@@ -55,9 +74,9 @@ void FKeyContextMenu::PopulateMenu(FMenuBuilder& MenuBuilder)
 		);
 
 		MenuBuilder.AddMenuEntry(
-			NSLOCTEXT("Sequencer", "SetKeyInterpolationLinear", "Linear"),
-			NSLOCTEXT("Sequencer", "SetKeyInterpolationLinearTooltip", "Set key interpolation to linear"),
-			FSlateIcon(),
+			LOCTEXT("SetKeyInterpolationLinear", "Linear"),
+			LOCTEXT("SetKeyInterpolationLinearTooltip", "Set key interpolation to linear"),
+			FSlateIcon(FEditorStyle::GetStyleSetName(), "Sequencer.IconKeyLinear"),
 			FUIAction(
 				FExecuteAction::CreateSP(SequencerPtr, &FSequencer::SetInterpTangentMode, RCIM_Linear, RCTM_Auto),
 				FCanExecuteAction(),
@@ -67,9 +86,9 @@ void FKeyContextMenu::PopulateMenu(FMenuBuilder& MenuBuilder)
 		);
 
 		MenuBuilder.AddMenuEntry(
-			NSLOCTEXT("Sequencer", "SetKeyInterpolationConstant", "Constant"),
-			NSLOCTEXT("Sequencer", "SetKeyInterpolationConstantTooltip", "Set key interpolation to constant"),
-			FSlateIcon(),
+			LOCTEXT("SetKeyInterpolationConstant", "Constant"),
+			LOCTEXT("SetKeyInterpolationConstantTooltip", "Set key interpolation to constant"),
+			FSlateIcon(FEditorStyle::GetStyleSetName(), "Sequencer.IconKeyConstant"),
 			FUIAction(
 				FExecuteAction::CreateSP(SequencerPtr, &FSequencer::SetInterpTangentMode, RCIM_Constant, RCTM_Auto),
 				FCanExecuteAction(),
@@ -80,13 +99,13 @@ void FKeyContextMenu::PopulateMenu(FMenuBuilder& MenuBuilder)
 	}
 	MenuBuilder.EndSection(); // SequencerInterpolation
 
-	MenuBuilder.BeginSection("SequencerKeys", NSLOCTEXT("Sequencer", "KeysMenu", "Keys"));
+	MenuBuilder.BeginSection("SequencerKeys", LOCTEXT("KeysMenu", "Keys"));
 	{
 		const bool bUseFrames = SequencerSnapValues::IsTimeSnapIntervalFrameRate(Sequencer->GetSettings()->GetTimeSnapInterval());
 
 		MenuBuilder.AddMenuEntry(
-			bUseFrames ? NSLOCTEXT("Sequencer", "SetKeyFrame", "Set Key Frame") : NSLOCTEXT("Sequencer", "SetKeyTime", "Set Key Time"),
-			bUseFrames ? NSLOCTEXT("Sequencer", "SetKeyFrameTooltip", "Set key frame") : NSLOCTEXT("Sequencer", "SetKeyTimeTooltip", "Set key time"),
+			bUseFrames ? LOCTEXT("SetKeyFrame", "Set Key Frame") : LOCTEXT("SetKeyTime", "Set Key MouseDownTime"),
+			bUseFrames ? LOCTEXT("SetKeyFrameTooltip", "Set key frame") : LOCTEXT("SetKeyTimeTooltip", "Set key MouseDownTime"),
 			FSlateIcon(),
 			FUIAction(
 				FExecuteAction::CreateSP(SequencerPtr, &FSequencer::SetKeyTime, bUseFrames),
@@ -94,8 +113,8 @@ void FKeyContextMenu::PopulateMenu(FMenuBuilder& MenuBuilder)
 		);
 
 		MenuBuilder.AddMenuEntry(
-			NSLOCTEXT("Sequencer", "SnapToFrame", "Snap to Frame"),
-			NSLOCTEXT("Sequencer", "SnapToFrameToolTip", "Snap selected keys to frame"),
+			LOCTEXT("SnapToFrame", "Snap to Frame"),
+			LOCTEXT("SnapToFrameToolTip", "Snap selected keys to frame"),
 			FSlateIcon(),
 			FUIAction(
 				FExecuteAction::CreateSP(SequencerPtr, &FSequencer::SnapToFrame),
@@ -103,8 +122,8 @@ void FKeyContextMenu::PopulateMenu(FMenuBuilder& MenuBuilder)
 		);
 
 		MenuBuilder.AddMenuEntry(
-			NSLOCTEXT("Sequencer", "DeleteKey", "Delete"),
-			NSLOCTEXT("Sequencer", "DeleteKeyToolTip", "Deletes the selected keys"),
+			LOCTEXT("DeleteKey", "Delete"),
+			LOCTEXT("DeleteKeyToolTip", "Deletes the selected keys"),
 			FSlateIcon(),
 			FUIAction(FExecuteAction::CreateSP(SequencerPtr, &FSequencer::DeleteSelectedKeys))
 		);
@@ -112,9 +131,9 @@ void FKeyContextMenu::PopulateMenu(FMenuBuilder& MenuBuilder)
 	MenuBuilder.EndSection(); // SequencerKeys
 }
 
-void FSectionContextMenu::BuildMenu(FMenuBuilder& MenuBuilder, FSequencer& InSequencer)
+void FSectionContextMenu::BuildMenu(FMenuBuilder& MenuBuilder, FSequencer& InSequencer, float InMouseDownTime)
 {
-	TSharedRef<FSectionContextMenu> Menu = MakeShareable(new FSectionContextMenu(InSequencer));
+	TSharedRef<FSectionContextMenu> Menu = MakeShareable(new FSectionContextMenu(InSequencer, InMouseDownTime));
 	Menu->PopulateMenu(MenuBuilder);
 }
 
@@ -123,12 +142,50 @@ void FSectionContextMenu::PopulateMenu(FMenuBuilder& MenuBuilder)
 	// Copy a reference to the context menu by value into each lambda handler to ensure the type stays alive until the menu is closed
 	TSharedRef<FSectionContextMenu> Shared = AsShared();
 
-	MenuBuilder.BeginSection("SequencerSections", NSLOCTEXT("Sequencer", "SectionsMenu", "Sections"));
+	MenuBuilder.BeginSection("SequencerKeyEdit", LOCTEXT("EditMenu", "Edit"));
+	{
+		TSharedPtr<FPasteFromHistoryContextMenu> PasteFromHistoryMenu;
+		TSharedPtr<FPasteContextMenu> PasteMenu;
+
+		if (Sequencer->GetClipboardStack().Num() != 0)
+		{
+			FPasteContextMenuArgs PasteArgs = FPasteContextMenuArgs::PasteAt(MouseDownTime);
+			PasteMenu = FPasteContextMenu::CreateMenu(*Sequencer, PasteArgs);
+			PasteFromHistoryMenu = FPasteFromHistoryContextMenu::CreateMenu(*Sequencer, PasteArgs);
+		}
+
+		MenuBuilder.AddSubMenu(
+			LOCTEXT("Paste", "Paste"),
+			FText(),
+			FNewMenuDelegate::CreateLambda([=](FMenuBuilder& SubMenuBuilder){ if (PasteMenu.IsValid()) { PasteMenu->PopulateMenu(SubMenuBuilder); } }),
+			FUIAction (
+				FExecuteAction(),
+				FCanExecuteAction::CreateLambda([=]{ return PasteMenu.IsValid() && PasteMenu->IsValidPaste(); })
+			),
+			NAME_None,
+			EUserInterfaceActionType::Button
+		);
+
+		MenuBuilder.AddSubMenu(
+			LOCTEXT("PasteFromHistory", "Paste From History"),
+			FText(),
+			FNewMenuDelegate::CreateLambda([=](FMenuBuilder& SubMenuBuilder){ if (PasteFromHistoryMenu.IsValid()) { PasteFromHistoryMenu->PopulateMenu(SubMenuBuilder); } }),
+			FUIAction (
+				FExecuteAction(),
+				FCanExecuteAction::CreateLambda([=]{ return PasteFromHistoryMenu.IsValid(); })
+			),
+			NAME_None,
+			EUserInterfaceActionType::Button
+		);
+	}
+	MenuBuilder.EndSection(); // SequencerKeyEdit
+
+	MenuBuilder.BeginSection("SequencerSections", LOCTEXT("SectionsMenu", "Sections"));
 	{
 		bool bCanSelectAll = CanSelectAllKeys();
 		MenuBuilder.AddMenuEntry(
-			NSLOCTEXT("Sequencer", "SelectAllKeys", "Select All Keys"),
-			NSLOCTEXT("Sequencer", "SelectAllKeysTooltip", "Select all keys in section"),
+			LOCTEXT("SelectAllKeys", "Select All Keys"),
+			LOCTEXT("SelectAllKeysTooltip", "Select all keys in section"),
 			FSlateIcon(),
 			FUIAction(
 				FExecuteAction::CreateLambda([=]{ return Shared->SelectAllKeys(); }),
@@ -136,28 +193,40 @@ void FSectionContextMenu::PopulateMenu(FMenuBuilder& MenuBuilder)
 		);
 
 		MenuBuilder.AddSubMenu(
-			NSLOCTEXT("Sequencer", "EditSection", "Edit"),
-			NSLOCTEXT("Sequencer", "EditSectionTooltip", "Edit section"),
+			LOCTEXT("EditSection", "Edit"),
+			LOCTEXT("EditSectionTooltip", "Edit section"),
 			FNewMenuDelegate::CreateLambda([=](FMenuBuilder& InMenuBuilder){ Shared->AddEditMenu(InMenuBuilder); }));
 
 		MenuBuilder.AddSubMenu(
-			NSLOCTEXT("Sequencer", "SetPreInfinityExtrap", "Pre-Infinity"), 
-			NSLOCTEXT("Sequencer", "SetPreInfinityExtrapTooltip", "Set pre-infinity extrapolation"),
-			FNewMenuDelegate::CreateLambda([=](FMenuBuilder& InMenuBuilder){ Shared->AddExtrapolationMenu(InMenuBuilder, true); }));
+			LOCTEXT("SetPreInfinityExtrap", "Pre-Infinity"),
+			LOCTEXT("SetPreInfinityExtrapTooltip", "Set pre-infinity extrapolation"),
+			FNewMenuDelegate::CreateLambda([=](FMenuBuilder& SubMenuBuilder){ Shared->AddExtrapolationMenu(SubMenuBuilder, true); }));
 
 		MenuBuilder.AddSubMenu(
-			NSLOCTEXT("Sequencer", "SetPostInfinityExtrap", "Post-Infinity"), 
-			NSLOCTEXT("Sequencer", "SetPostInfinityExtrapTooltip", "Set post-infinity extrapolation"),
-			FNewMenuDelegate::CreateLambda([=](FMenuBuilder& InMenuBuilder){ Shared->AddExtrapolationMenu(InMenuBuilder, false); }));
+			LOCTEXT("SetPostInfinityExtrap", "Post-Infinity"),
+			LOCTEXT("SetPostInfinityExtrapTooltip", "Set post-infinity extrapolation"),
+			FNewMenuDelegate::CreateLambda([=](FMenuBuilder& SubMenuBuilder){ Shared->AddExtrapolationMenu(SubMenuBuilder, false); }));
 
 		MenuBuilder.AddMenuEntry(
-			NSLOCTEXT("Sequencer", "ToggleSectionActive", "Active"),
-			NSLOCTEXT("Sequencer", "ToggleSectionActiveTooltip", "Toggle section active/inactive"),
+			LOCTEXT("ToggleSectionActive", "Active"),
+			LOCTEXT("ToggleSectionActiveTooltip", "Toggle section active/inactive"),
 			FSlateIcon(),
 			FUIAction(
 				FExecuteAction::CreateLambda([=]{ Shared->ToggleSectionActive(); }),
 				FCanExecuteAction(),
-				FIsActionChecked::CreateLambda([=]{ return Shared->IsToggleSectionActive(); })),
+				FIsActionChecked::CreateLambda([=]{ return Shared->IsSectionActive(); })),
+			NAME_None,
+			EUserInterfaceActionType::ToggleButton
+		);
+
+		MenuBuilder.AddMenuEntry(
+			NSLOCTEXT("Sequencer", "ToggleSectionLocked", "Locked"),
+			NSLOCTEXT("Sequencer", "ToggleSectionLockedTooltip", "Toggle section locked/unlocked"),
+			FSlateIcon(),
+			FUIAction(
+				FExecuteAction::CreateLambda([=]{ Shared->ToggleSectionLocked(); }),
+				FCanExecuteAction(),
+				FIsActionChecked::CreateLambda([=]{ return Shared->IsSectionLocked(); })),
 			NAME_None,
 			EUserInterfaceActionType::ToggleButton
 		);
@@ -165,8 +234,8 @@ void FSectionContextMenu::PopulateMenu(FMenuBuilder& MenuBuilder)
 		// @todo Sequencer this should delete all selected sections
 		// delete/selection needs to be rethought in general
 		MenuBuilder.AddMenuEntry(
-			NSLOCTEXT("Sequencer", "DeleteSection", "Delete"),
-			NSLOCTEXT("Sequencer", "DeleteSectionToolTip", "Deletes this section"),
+			LOCTEXT("DeleteSection", "Delete"),
+			LOCTEXT("DeleteSectionToolTip", "Deletes this section"),
 			FSlateIcon(),
 			FUIAction(FExecuteAction::CreateLambda([=]{ return Shared->DeleteSection(); }))
 		);
@@ -180,8 +249,8 @@ void FSectionContextMenu::AddEditMenu(FMenuBuilder& MenuBuilder)
 	TSharedRef<FSectionContextMenu> Shared = AsShared();
 
 	MenuBuilder.AddMenuEntry(
-		NSLOCTEXT("Sequencer", "TrimSectionLeft", "Trim Left"),
-		NSLOCTEXT("Sequencer", "TrimSectionLeftTooltip", "Trim section at current time to the left"),
+		LOCTEXT("TrimSectionLeft", "Trim Left"),
+		LOCTEXT("TrimSectionLeftTooltip", "Trim section at current MouseDownTime to the left"),
 		FSlateIcon(),
 		FUIAction(
 			FExecuteAction::CreateLambda([=]{ Shared->TrimSection(true); }),
@@ -189,8 +258,8 @@ void FSectionContextMenu::AddEditMenu(FMenuBuilder& MenuBuilder)
 	);
 
 	MenuBuilder.AddMenuEntry(
-		NSLOCTEXT("Sequencer", "TrimSectionRight", "Trim Right"),
-		NSLOCTEXT("Sequencer", "TrimSectionRightTooltip", "Trim section at current time to the right"),
+		LOCTEXT("TrimSectionRight", "Trim Right"),
+		LOCTEXT("TrimSectionRightTooltip", "Trim section at current MouseDownTime to the right"),
 		FSlateIcon(),
 		FUIAction(
 			FExecuteAction::CreateLambda([=]{ Shared->TrimSection(false); }),
@@ -198,8 +267,8 @@ void FSectionContextMenu::AddEditMenu(FMenuBuilder& MenuBuilder)
 	);
 
 	MenuBuilder.AddMenuEntry(
-		NSLOCTEXT("Sequencer", "SplitSection", "Split"),
-		NSLOCTEXT("Sequencer", "SplitSectionTooltip", "Split section at current time"),
+		LOCTEXT("SplitSection", "Split"),
+		LOCTEXT("SplitSectionTooltip", "Split section at current MouseDownTime"),
 		FSlateIcon(),
 		FUIAction(
 			FExecuteAction::CreateLambda([=]{ Shared->SplitSection(); }),
@@ -213,8 +282,8 @@ void FSectionContextMenu::AddExtrapolationMenu(FMenuBuilder& MenuBuilder, bool b
 	TSharedRef<FSectionContextMenu> Shared = AsShared();
 
 	MenuBuilder.AddMenuEntry(
-		NSLOCTEXT("Sequencer", "SetExtrapCycle", "Cycle"),
-		NSLOCTEXT("Sequencer", "SetExtrapCycleTooltip", "Set extrapolation cycle"),
+		LOCTEXT("SetExtrapCycle", "Cycle"),
+		LOCTEXT("SetExtrapCycleTooltip", "Set extrapolation cycle"),
 		FSlateIcon(),
 		FUIAction(
 			FExecuteAction::CreateLambda([=]{ Shared->SetExtrapolationMode(RCCE_Cycle, bPreInfinity); }),
@@ -226,8 +295,8 @@ void FSectionContextMenu::AddExtrapolationMenu(FMenuBuilder& MenuBuilder, bool b
 	);
 
 	MenuBuilder.AddMenuEntry(
-		NSLOCTEXT("Sequencer", "SetExtrapCycleWithOffset", "Cycle with Offset"),
-		NSLOCTEXT("Sequencer", "SetExtrapCycleWithOffsetTooltip", "Set extrapolation cycle with offset"),
+		LOCTEXT("SetExtrapCycleWithOffset", "Cycle with Offset"),
+		LOCTEXT("SetExtrapCycleWithOffsetTooltip", "Set extrapolation cycle with offset"),
 		FSlateIcon(),
 		FUIAction(
 			FExecuteAction::CreateLambda([=]{ Shared->SetExtrapolationMode(RCCE_CycleWithOffset, bPreInfinity); }),
@@ -239,8 +308,8 @@ void FSectionContextMenu::AddExtrapolationMenu(FMenuBuilder& MenuBuilder, bool b
 	);
 
 	MenuBuilder.AddMenuEntry(
-		NSLOCTEXT("Sequencer", "SetExtrapOscillate", "Oscillate"),
-		NSLOCTEXT("Sequencer", "SetExtrapOscillateTooltip", "Set extrapolation oscillate"),
+		LOCTEXT("SetExtrapOscillate", "Oscillate"),
+		LOCTEXT("SetExtrapOscillateTooltip", "Set extrapolation oscillate"),
 		FSlateIcon(),
 		FUIAction(
 			FExecuteAction::CreateLambda([=]{ Shared->SetExtrapolationMode(RCCE_Oscillate, bPreInfinity); }),
@@ -252,8 +321,8 @@ void FSectionContextMenu::AddExtrapolationMenu(FMenuBuilder& MenuBuilder, bool b
 	);
 
 	MenuBuilder.AddMenuEntry(
-		NSLOCTEXT("Sequencer", "SetExtrapLinear", "Linear"),
-		NSLOCTEXT("Sequencer", "SetExtrapLinearTooltip", "Set extrapolation linear"),
+		LOCTEXT("SetExtrapLinear", "Linear"),
+		LOCTEXT("SetExtrapLinearTooltip", "Set extrapolation linear"),
 		FSlateIcon(),
 		FUIAction(
 			FExecuteAction::CreateLambda([=]{ Shared->SetExtrapolationMode(RCCE_Linear, bPreInfinity); }),
@@ -265,8 +334,8 @@ void FSectionContextMenu::AddExtrapolationMenu(FMenuBuilder& MenuBuilder, bool b
 	);
 
 	MenuBuilder.AddMenuEntry(
-		NSLOCTEXT("Sequencer", "SetExtrapConstant", "Constant"),
-		NSLOCTEXT("Sequencer", "SetExtrapConstantTooltip", "Set extrapolation constant"),
+		LOCTEXT("SetExtrapConstant", "Constant"),
+		LOCTEXT("SetExtrapConstantTooltip", "Set extrapolation constant"),
 		FSlateIcon(),
 		FUIAction(
 			FExecuteAction::CreateLambda([=]{ Shared->SetExtrapolationMode(RCCE_Constant, bPreInfinity); }),
@@ -329,7 +398,7 @@ bool FSectionContextMenu::CanSelectAllKeys() const
 
 void FSectionContextMenu::TrimSection(bool bTrimLeft)
 {
-	FScopedTransaction TrimSectionTransaction(NSLOCTEXT("Sequencer", "TrimSection_Transaction", "Trim Section"));
+	FScopedTransaction TrimSectionTransaction(LOCTEXT("TrimSection_Transaction", "Trim Section"));
 
 	MovieSceneToolHelpers::TrimSection(Sequencer->GetSelection().GetSelectedSections(), Sequencer->GetGlobalTime(), bTrimLeft);
 	Sequencer->NotifyMovieSceneDataChanged();
@@ -337,7 +406,7 @@ void FSectionContextMenu::TrimSection(bool bTrimLeft)
 
 void FSectionContextMenu::SplitSection()
 {
-	FScopedTransaction SplitSectionTransaction(NSLOCTEXT("Sequencer", "SplitSection_Transaction", "Split Section"));
+	FScopedTransaction SplitSectionTransaction(LOCTEXT("SplitSection_Transaction", "Split Section"));
 
 	MovieSceneToolHelpers::SplitSection(Sequencer->GetSelection().GetSelectedSections(), Sequencer->GetGlobalTime());
 	Sequencer->NotifyMovieSceneDataChanged();
@@ -357,7 +426,7 @@ bool FSectionContextMenu::IsTrimmable() const
 
 void FSectionContextMenu::SetExtrapolationMode(ERichCurveExtrapolation ExtrapMode, bool bPreInfinity)
 {
-	FScopedTransaction SetExtrapolationModeTransaction(NSLOCTEXT("Sequencer", "SetExtrapolationMode_Transaction", "Set Extrapolation Mode"));
+	FScopedTransaction SetExtrapolationModeTransaction(LOCTEXT("SetExtrapolationMode_Transaction", "Set Extrapolation Mode"));
 
 	bool bAnythingChanged = false;
 
@@ -370,14 +439,15 @@ void FSectionContextMenu::SetExtrapolationMode(ERichCurveExtrapolation ExtrapMod
 			continue;
 		}
 
-		Section->Modify();
-
-		FKeyAreaLayout Layout(*Handle.TrackNode, Handle.SectionIndex);
-		for (const FKeyAreaLayoutElement& Element : Layout.GetElements())
+		if (Section->TryModify())
 		{
-			TSharedPtr<IKeyArea> KeyArea = Element.GetKeyArea();
-			bAnythingChanged = true;
-			KeyArea->SetExtrapolationMode(ExtrapMode, bPreInfinity);
+			FKeyAreaLayout Layout(*Handle.TrackNode, Handle.SectionIndex);
+			for (const FKeyAreaLayoutElement& Element : Layout.GetElements())
+			{
+				TSharedPtr<IKeyArea> KeyArea = Element.GetKeyArea();
+				bAnythingChanged = true;
+				KeyArea->SetExtrapolationMode(ExtrapMode, bPreInfinity);
+			}
 		}
 	}
 
@@ -418,8 +488,8 @@ bool FSectionContextMenu::IsExtrapolationModeSelected(ERichCurveExtrapolation Ex
 
 void FSectionContextMenu::ToggleSectionActive()
 {
-	FScopedTransaction ToggleSectionActiveTransaction( NSLOCTEXT("Sequencer", "ToggleSectionActive_Transaction", "Toggle Section Active") );
-	bool bIsActive = !IsToggleSectionActive();
+	FScopedTransaction ToggleSectionActiveTransaction( LOCTEXT("ToggleSectionActive_Transaction", "Toggle Section Active") );
+	bool bIsActive = !IsSectionActive();
 	bool bAnythingChanged = false;
 
 	for (auto Section : Sequencer->GetSelection().GetSelectedSections())
@@ -439,7 +509,7 @@ void FSectionContextMenu::ToggleSectionActive()
 	}
 }
 
-bool FSectionContextMenu::IsToggleSectionActive() const
+bool FSectionContextMenu::IsSectionActive() const
 {
 	// Active only if all are active
 	for (auto Section : Sequencer->GetSelection().GetSelectedSections())
@@ -453,7 +523,411 @@ bool FSectionContextMenu::IsToggleSectionActive() const
 	return true;
 }
 
+void FSectionContextMenu::ToggleSectionLocked()
+{
+	FScopedTransaction ToggleSectionLockedTransaction( NSLOCTEXT("Sequencer", "ToggleSectionLocked_Transaction", "Toggle Section Locked") );
+	bool bIsLocked = !IsSectionLocked();
+	bool bAnythingChanged = false;
+
+	for (auto Section : Sequencer->GetSelection().GetSelectedSections())
+	{
+		bAnythingChanged = true;
+		Section->Modify();
+		Section->SetIsLocked(bIsLocked);
+	}
+
+	if (bAnythingChanged)
+	{
+		Sequencer->UpdateRuntimeInstances();
+	}
+	else
+	{
+		ToggleSectionLockedTransaction.Cancel();
+	}
+}
+
+bool FSectionContextMenu::IsSectionLocked() const
+{
+	// Locked only if all are locked
+	for (auto Section : Sequencer->GetSelection().GetSelectedSections())
+	{
+		if (!Section->IsLocked())
+		{
+			return false;
+		}
+	}
+
+	return true;
+}
+
 void FSectionContextMenu::DeleteSection()
 {
 	Sequencer->DeleteSections(Sequencer->GetSelection().GetSelectedSections());
-}
\ No newline at end of file
+}
+
+bool FPasteContextMenu::BuildMenu(FMenuBuilder& MenuBuilder, FSequencer& InSequencer, const FPasteContextMenuArgs& Args)
+{
+	TSharedRef<FPasteContextMenu> Menu = MakeShareable(new FPasteContextMenu(InSequencer, Args));
+	Menu->Setup();
+	if (!Menu->IsValidPaste())
+	{
+		return false;
+	}
+
+	Menu->PopulateMenu(MenuBuilder);
+	return true;
+}
+
+TSharedRef<FPasteContextMenu> FPasteContextMenu::CreateMenu(FSequencer& InSequencer, const FPasteContextMenuArgs& Args)
+{
+	TSharedRef<FPasteContextMenu> Menu = MakeShareable(new FPasteContextMenu(InSequencer, Args));
+	Menu->Setup();
+	return Menu;
+}
+
+TArray<TSharedRef<FSequencerSectionKeyAreaNode>> KeyAreaNodesBuffer;
+
+void FPasteContextMenu::GatherPasteDestinationsForNode(FSequencerDisplayNode& InNode, int32 SectionIndex, const FName& CurrentScope, TMap<FName, FSequencerClipboardReconciler>& Map)
+{
+	KeyAreaNodesBuffer.Reset();
+	if (InNode.GetType() == ESequencerNode::KeyArea)
+	{
+		KeyAreaNodesBuffer.Add(StaticCastSharedRef<FSequencerSectionKeyAreaNode>(InNode.AsShared()));
+	}
+	else
+	{
+		InNode.GetChildKeyAreaNodesRecursively(KeyAreaNodesBuffer);
+	}
+
+	if (!KeyAreaNodesBuffer.Num())
+	{
+		return;
+	}
+
+	FName ThisScope;
+	{
+		FString ThisScopeString;
+		if (!CurrentScope.IsNone())
+		{
+			ThisScopeString.Append(CurrentScope.ToString());
+			ThisScopeString.AppendChar('.');
+		}
+		ThisScopeString.Append(InNode.GetDisplayName().ToString());
+		ThisScope = *ThisScopeString;
+	}
+
+	FSequencerClipboardReconciler* Reconciler = Map.Find(ThisScope);
+	if (!Reconciler)
+	{
+		Reconciler = &Map.Add(ThisScope, FSequencerClipboardReconciler(Args.Clipboard.ToSharedRef()));
+	}
+
+	FSequencerClipboardPasteGroup Group = Reconciler->AddDestinationGroup();
+	for (const TSharedRef<FSequencerSectionKeyAreaNode>& KeyAreaNode : KeyAreaNodesBuffer)
+	{
+		Group.Add(KeyAreaNode->GetKeyArea(SectionIndex).Get());
+	}
+
+	// Add children
+	for (const TSharedPtr<FSequencerDisplayNode>& Child : InNode.GetChildNodes())
+	{
+		GatherPasteDestinationsForNode(*Child, SectionIndex, ThisScope, Map);
+	}
+}
+
+void GetFullNodePath(FSequencerDisplayNode& InNode, FString& Path)
+{
+	TSharedPtr<FSequencerDisplayNode> Parent = InNode.GetParent();
+	if (Parent.IsValid())
+	{
+		GetFullNodePath(*Parent, Path);
+	}
+
+	if (!Path.IsEmpty())
+	{
+		Path.AppendChar('.');
+	}
+
+	Path.Append(InNode.GetDisplayName().ToString());
+}
+
+
+TSharedPtr<FSequencerTrackNode> GetTrackFromNode(FSequencerDisplayNode& InNode, FString& Scope)
+{
+	if (InNode.GetType() == ESequencerNode::Track)
+	{
+		return StaticCastSharedRef<FSequencerTrackNode>(InNode.AsShared());
+	}
+	else if (InNode.GetType() == ESequencerNode::Object)
+	{
+		return nullptr;
+	}
+
+	TSharedPtr<FSequencerDisplayNode> Parent = InNode.GetParent();
+	if (Parent.IsValid())
+	{
+		TSharedPtr<FSequencerTrackNode> Track = GetTrackFromNode(*Parent, Scope);
+		if (Track.IsValid())
+		{
+			FString ThisScope = InNode.GetDisplayName().ToString();
+			if (!Scope.IsEmpty())
+			{
+				ThisScope.AppendChar('.');
+				ThisScope.Append(Scope);
+				Scope = MoveTemp(ThisScope);
+			}
+			return Track;
+		}
+	}
+
+	return nullptr;
+}
+
+void FPasteContextMenu::Setup()
+{
+	if (!Args.Clipboard.IsValid())
+	{
+		if (Sequencer->GetClipboardStack().Num() != 0)
+		{
+			Args.Clipboard = Sequencer->GetClipboardStack().Last();
+		}
+		else
+		{
+			return;
+		}
+	}
+
+	// Gather a list of sections we want to paste into
+	TArray<FSectionHandle> SectionHandles;
+
+	if (Args.DestinationNodes.Num())
+	{
+		// Paste into only these nodes
+		for (const TSharedRef<FSequencerDisplayNode>& Node : Args.DestinationNodes)
+		{
+			FString Scope;
+			TSharedPtr<FSequencerTrackNode> TrackNode = GetTrackFromNode(*Node, Scope);
+			if (!TrackNode.IsValid())
+			{
+				continue;
+			}
+
+			const TArray<UMovieSceneSection*>& Sections = TrackNode->GetTrack()->GetAllSections();
+			UMovieSceneSection* Section = MovieSceneHelpers::FindNearestSectionAtTime(Sections, Args.PasteAtTime);
+			int32 SectionIndex = INDEX_NONE;
+			if (Section)
+			{
+				SectionIndex = Sections.IndexOfByKey(Section);
+			}
+
+			if (SectionIndex != INDEX_NONE)
+			{
+				SectionHandles.Add(FSectionHandle(TrackNode, SectionIndex));
+			}
+		}
+	}
+	else
+	{
+		// Use the selected sections
+		TSharedRef<SSequencer> SequencerWidget = StaticCastSharedRef<SSequencer>(Sequencer->GetSequencerWidget());
+		SectionHandles = SequencerWidget->GetSectionHandles(Sequencer->GetSelection().GetSelectedSections());
+	}
+
+	TMap<FName, TArray<FSectionHandle>> SectionsByType;
+	for (const FSectionHandle& Section : SectionHandles)
+	{
+		UMovieSceneTrack* Track = Section.TrackNode->GetTrack();
+		if (Track)
+		{
+			SectionsByType.FindOrAdd(Track->GetClass()->GetFName()).Add(Section);
+		}
+	}
+
+	for (auto& Pair : SectionsByType)
+	{
+		FPasteDestination& Destination = PasteDestinations[PasteDestinations.AddDefaulted()];
+		if (Pair.Value.Num() == 1)
+		{
+			FString Path;
+			GetFullNodePath(*Pair.Value[0].TrackNode, Path);
+			Destination.Name = FText::FromString(Path);
+		}
+		else
+		{
+			Destination.Name = FText::Format(LOCTEXT("PasteMenuHeaderFormat", "{0} ({1} tracks)"), FText::FromName(Pair.Key), FText::AsNumber(Pair.Value.Num()));
+		}
+
+		for (const FSectionHandle& Section : Pair.Value)
+		{
+			GatherPasteDestinationsForNode(*Section.TrackNode, Section.SectionIndex, NAME_None, Destination.Reconcilers);
+		}
+
+		// Reconcile and remove invalid pastes
+		for (auto It = Destination.Reconcilers.CreateIterator(); It; ++It)
+		{
+			if (!It.Value().Reconcile())
+			{
+				It.RemoveCurrent();
+			}
+		}
+		if (!Destination.Reconcilers.Num())
+		{
+			PasteDestinations.RemoveAt(PasteDestinations.Num() - 1, 1, false);
+		}
+	}
+}
+
+bool FPasteContextMenu::IsValidPaste() const
+{
+	return Args.Clipboard.IsValid() && PasteDestinations.Num() != 0;
+}
+
+void FPasteContextMenu::PopulateMenu(FMenuBuilder& MenuBuilder)
+{
+	// Copy a reference to the context menu by value into each lambda handler to ensure the type stays alive until the menu is closed
+	TSharedRef<FPasteContextMenu> Shared = AsShared();
+
+	bool bElevateMenu = PasteDestinations.Num() == 1;
+	for (int32 Index = 0; Index < PasteDestinations.Num(); ++Index)
+	{
+		if (bElevateMenu)
+		{
+			MenuBuilder.BeginSection("PasteInto", FText::Format(LOCTEXT("PasteIntoTitle", "Paste Into {0}"), PasteDestinations[Index].Name));
+			AddPasteMenuForTrackType(MenuBuilder, Index);
+			MenuBuilder.EndSection();
+			break;
+		}
+
+		MenuBuilder.AddSubMenu(
+			PasteDestinations[Index].Name,
+			FText(),
+			FNewMenuDelegate::CreateLambda([=](FMenuBuilder& SubMenuBuilder){ Shared->AddPasteMenuForTrackType(SubMenuBuilder, Index); })
+		);
+	}
+}
+
+void FPasteContextMenu::AddPasteMenuForTrackType(FMenuBuilder& MenuBuilder, int32 DestinationIndex)
+{
+	// Copy a reference to the context menu by value into each lambda handler to ensure the type stays alive until the menu is closed
+	TSharedRef<FPasteContextMenu> Shared = AsShared();
+
+	for (auto& Pair : PasteDestinations[DestinationIndex].Reconcilers)
+	{
+		MenuBuilder.AddMenuEntry(
+			FText::FromName(Pair.Key),
+			FText(),
+			FSlateIcon(),
+			FUIAction(
+				FExecuteAction::CreateLambda([=](){ Shared->PasteInto(DestinationIndex, Pair.Key); })
+			)
+		);
+	}
+}
+
+bool FPasteContextMenu::AutoPaste()
+{
+	if (PasteDestinations.Num() == 1)
+	{
+		for (auto& Pair : PasteDestinations[0].Reconcilers)
+		{
+			if (Pair.Value.CanAutoPaste())
+			{
+				PasteInto(0, Pair.Key);
+				return true;
+			}
+		}
+	}
+
+	return false;
+}
+
+void FPasteContextMenu::PasteInto(int32 DestinationIndex, FName KeyAreaName)
+{
+	FSequencerClipboardReconciler& Reconciler = PasteDestinations[DestinationIndex].Reconcilers[KeyAreaName];
+
+	TSet<FSequencerSelectedKey> NewSelection;
+
+	FSequencerPasteEnvironment PasteEnvironment;
+	PasteEnvironment.CardinalTime = Args.PasteAtTime;
+	PasteEnvironment.OnKeyPasted = [&](FKeyHandle Handle, IKeyArea& KeyArea){
+		NewSelection.Add(FSequencerSelectedKey(*KeyArea.GetOwningSection(), KeyArea.AsShared(), Handle));
+	};
+
+	FScopedTransaction Transaction(LOCTEXT("PasteKeysTransaction", "Paste Keys"));
+	if (!Reconciler.Paste(PasteEnvironment))
+	{
+		Transaction.Cancel();
+	}
+	else
+	{
+		SSequencerSection::ThrobSelection();
+
+		// @todo sequencer: selection in transactions
+		FSequencerSelection& Selection = Sequencer->GetSelection();
+		Selection.SuspendBroadcast();
+		Selection.EmptySelectedKeys();
+
+		for (const FSequencerSelectedKey& Key : NewSelection)
+		{
+			Selection.AddToSelection(Key);
+		}
+		Selection.ResumeBroadcast();
+		Selection.GetOnKeySelectionChanged().Broadcast();
+
+		Sequencer->OnClipboardUsed(Args.Clipboard);
+	}
+}
+
+bool FPasteFromHistoryContextMenu::BuildMenu(FMenuBuilder& MenuBuilder, FSequencer& InSequencer, const FPasteContextMenuArgs& Args)
+{
+	if (InSequencer.GetClipboardStack().Num() == 0)
+	{
+		return false;
+	}
+
+	TSharedRef<FPasteFromHistoryContextMenu> Menu = MakeShareable(new FPasteFromHistoryContextMenu(InSequencer, Args));
+	Menu->PopulateMenu(MenuBuilder);
+	return true;
+}
+
+TSharedPtr<FPasteFromHistoryContextMenu> FPasteFromHistoryContextMenu::CreateMenu(FSequencer& InSequencer, const FPasteContextMenuArgs& Args)
+{
+	if (InSequencer.GetClipboardStack().Num() == 0)
+	{
+		return nullptr;
+	}
+
+	return MakeShareable(new FPasteFromHistoryContextMenu(InSequencer, Args));
+}
+
+void FPasteFromHistoryContextMenu::PopulateMenu(FMenuBuilder& MenuBuilder)
+{
+	// Copy a reference to the context menu by value into each lambda handler to ensure the type stays alive until the menu is closed
+	TSharedRef<FPasteFromHistoryContextMenu> Shared = AsShared();
+
+	MenuBuilder.BeginSection("SequencerPasteHistory", LOCTEXT("PasteFromHistory", "Paste History"));
+
+	for (int32 Index = Sequencer->GetClipboardStack().Num() - 1; Index >= 0; --Index)
+	{
+		FPasteContextMenuArgs ThisPasteArgs = Args;
+		ThisPasteArgs.Clipboard = Sequencer->GetClipboardStack()[Index];
+
+		TSharedRef<FPasteContextMenu> PasteMenu = FPasteContextMenu::CreateMenu(*Sequencer, ThisPasteArgs);
+
+		MenuBuilder.AddSubMenu(
+			ThisPasteArgs.Clipboard->GetDisplayText(),
+			FText(),
+			FNewMenuDelegate::CreateLambda([=](FMenuBuilder& SubMenuBuilder){ PasteMenu->PopulateMenu(SubMenuBuilder); }),
+			FUIAction (
+				FExecuteAction(),
+				FCanExecuteAction::CreateLambda([=]{ return PasteMenu->IsValidPaste(); })
+			),
+			NAME_None,
+			EUserInterfaceActionType::Button
+		);
+	}
+
+	MenuBuilder.EndSection();
+}
+
+#undef LOCTEXT_NAMESPACE
\ No newline at end of file
diff --git a/Engine/Source/Editor/Sequencer/Private/SequencerContextMenus.h b/Engine/Source/Editor/Sequencer/Private/SequencerContextMenus.h
--- a/Engine/Source/Editor/Sequencer/Private/SequencerContextMenus.h
+++ b/Engine/Source/Editor/Sequencer/Private/SequencerContextMenus.h
@@ -3,6 +3,7 @@
 #pragma once
 
 #include "Sequencer.h"
+#include "SequencerClipboardReconciler.h"
 
 /**
  * Class responsible for generating a menu for the currently selected sections.
@@ -12,16 +13,16 @@
  */
 struct FSectionContextMenu : TSharedFromThis<FSectionContextMenu>
 {
-	static void BuildMenu(FMenuBuilder& MenuBuilder, FSequencer& Sequencer);
+	static void BuildMenu(FMenuBuilder& MenuBuilder, FSequencer& Sequencer, float InMouseDownTime);
 
 private:
 
 	/** Hidden AsShared() methods to discourage CreateSP delegate use. */
-	TSharedRef<FSectionContextMenu> AsShared() { return TSharedFromThis::AsShared(); }
-	TSharedRef<const FSectionContextMenu> AsShared() const { return TSharedFromThis::AsShared(); }
+	using TSharedFromThis::AsShared;
 
-	FSectionContextMenu(FSequencer& InSeqeuncer)
+	FSectionContextMenu(FSequencer& InSeqeuncer, float InMouseDownTime)
 		: Sequencer(StaticCastSharedRef<FSequencer>(InSeqeuncer.AsShared()))
+		, MouseDownTime(InMouseDownTime)
 	{}
 
 	void PopulateMenu(FMenuBuilder& MenuBuilder);
@@ -48,12 +49,119 @@ private:
 
 	void ToggleSectionActive();
 
-	bool IsToggleSectionActive() const;
+	bool IsSectionActive() const;
+
+	void ToggleSectionLocked();
+
+	bool IsSectionLocked() const;
 
 	void DeleteSection();
 
-	/** The key itself */
+	/** The sequencer */
+	TSharedRef<FSequencer> Sequencer;
+
+	/** The time that we clicked on to summon this menu */
+	float MouseDownTime;
+};
+
+/** Arguments required for a paste operation */
+struct FPasteContextMenuArgs
+{
+	/** Paste the clipboard into the specified array of sequencer nodes, at the given time */
+	static FPasteContextMenuArgs PasteInto(TArray<TSharedRef<FSequencerDisplayNode>> InNodes, float InTime, TSharedPtr<FMovieSceneClipboard> InClipboard = nullptr)
+	{
+		FPasteContextMenuArgs Args;
+		Args.Clipboard = InClipboard;
+		Args.DestinationNodes = MoveTemp(InNodes);
+		Args.PasteAtTime = InTime;
+		return Args;
+	}
+
+	/** Paste the clipboard at the given time, using the sequencer selection states to determine paste destinations */
+	static FPasteContextMenuArgs PasteAt(float InTime, TSharedPtr<FMovieSceneClipboard> InClipboard = nullptr)
+	{
+		FPasteContextMenuArgs Args;
+		Args.Clipboard = InClipboard;
+		Args.PasteAtTime = InTime;
+		return Args;
+	}
+
+	/** The clipboard to paste */
+	TSharedPtr<FMovieSceneClipboard> Clipboard;
+
+	/** The Time to paste at */
+	float PasteAtTime;
+
+	/** Optional user-supplied nodes to paste into */
+	TArray<TSharedRef<FSequencerDisplayNode>> DestinationNodes;
+};
+
+struct FPasteContextMenu : TSharedFromThis<FPasteContextMenu>
+{
+	static bool BuildMenu(FMenuBuilder& MenuBuilder, FSequencer& Sequencer, const FPasteContextMenuArgs& Args);
+
+	static TSharedRef<FPasteContextMenu> CreateMenu(FSequencer& Sequencer, const FPasteContextMenuArgs& Args);
+
+	void PopulateMenu(FMenuBuilder& MenuBuilder);
+
+	bool IsValidPaste() const;
+
+	bool AutoPaste();
+
+private:
+
+	using TSharedFromThis::AsShared;
+
+	FPasteContextMenu(FSequencer& InSequencer, const FPasteContextMenuArgs& InArgs)
+		: Sequencer(StaticCastSharedRef<FSequencer>(InSequencer.AsShared()))
+		, Args(InArgs)
+	{}
+
+	void Setup();
+
+	void AddPasteMenuForTrackType(FMenuBuilder& MenuBuilder, int32 DestinationIndex);
+
+	void PasteInto(int32 DestinationIndex, FName KeyAreaName);
+
+	void GatherPasteDestinationsForNode(FSequencerDisplayNode& InNode, int32 SectionIndex, const FName& CurrentScope, TMap<FName, FSequencerClipboardReconciler>& Map);
+
+	/** The sequencer */
 	TSharedRef<FSequencer> Sequencer;
+
+	/** Paste destinations are organized by track type primarily, then by key area name  */
+	struct FPasteDestination
+	{
+		FText Name;
+		TMap<FName, FSequencerClipboardReconciler> Reconcilers;
+	};
+	TArray<FPasteDestination> PasteDestinations;
+
+	/** Paste arguments */
+	FPasteContextMenuArgs Args;
+};
+
+struct FPasteFromHistoryContextMenu : TSharedFromThis<FPasteFromHistoryContextMenu>
+{
+	static bool BuildMenu(FMenuBuilder& MenuBuilder, FSequencer& Sequencer, const FPasteContextMenuArgs& Args);
+
+	static TSharedPtr<FPasteFromHistoryContextMenu> CreateMenu(FSequencer& Sequencer, const FPasteContextMenuArgs& Args);
+
+	void PopulateMenu(FMenuBuilder& MenuBuilder);
+
+private:
+
+	using TSharedFromThis::AsShared;
+
+	FPasteFromHistoryContextMenu(FSequencer& InSequencer, const FPasteContextMenuArgs& InArgs)
+		: Sequencer(StaticCastSharedRef<FSequencer>(InSequencer.AsShared()))
+		, Args(InArgs)
+	{}
+
+	/** The sequencer */
+	TSharedRef<FSequencer> Sequencer;
+
+	/** Paste arguments */
+	FPasteContextMenuArgs Args;
 };
 
 /**
@@ -73,11 +181,10 @@ private:
 	{}
 
 	/** Hidden AsShared() methods to discourage CreateSP delegate use. */
-	TSharedRef<FKeyContextMenu> AsShared() { return TSharedFromThis::AsShared(); }
-	TSharedRef<const FKeyContextMenu> AsShared() const { return TSharedFromThis::AsShared(); }
+	using TSharedFromThis::AsShared;
 
 	void PopulateMenu(FMenuBuilder& MenuBuilder);
 
-	/** The key itself */
+	/** The sequencer */
 	TSharedRef<FSequencer> Sequencer;
 };
\ No newline at end of file
diff --git a/Engine/Source/Editor/Sequencer/Private/SequencerHotspots.cpp b/Engine/Source/Editor/Sequencer/Private/SequencerHotspots.cpp
--- a/Engine/Source/Editor/Sequencer/Private/SequencerHotspots.cpp
+++ b/Engine/Source/Editor/Sequencer/Private/SequencerHotspots.cpp
@@ -7,20 +7,20 @@
 #include "MovieSceneSection.h"
 #include "SequencerContextMenus.h"
 
-void FKeyHotspot::PopulateContextMenu(FMenuBuilder& MenuBuilder, ISequencer& InSequencer)
+void FKeyHotspot::PopulateContextMenu(FMenuBuilder& MenuBuilder, ISequencer& InSequencer, float MouseDownTime)
 {
 	FSequencer& Sequencer = static_cast<FSequencer&>(InSequencer);
 	FKeyContextMenu::BuildMenu(MenuBuilder, Sequencer);
 }
 
-void FSectionHotspot::PopulateContextMenu(FMenuBuilder& MenuBuilder, ISequencer& InSequencer)
+void FSectionHotspot::PopulateContextMenu(FMenuBuilder& MenuBuilder, ISequencer& InSequencer, float MouseDownTime)
 {
 	FSequencer& Sequencer = static_cast<FSequencer&>(InSequencer);
 
 	TSharedPtr<ISequencerSection> SectionInterface = Section.TrackNode->GetSections()[Section.SectionIndex];
 	SectionInterface->BuildSectionContextMenu(MenuBuilder);
 
-	FSectionContextMenu::BuildMenu(MenuBuilder, Sequencer);
+	FSectionContextMenu::BuildMenu(MenuBuilder, Sequencer, MouseDownTime);
 }
 
 TSharedPtr<ISequencerEditToolDragOperation> FSectionResizeHotspot::InitiateDrag(ISequencer& Sequencer)
diff --git a/Engine/Source/Editor/Sequencer/Private/SequencerHotspots.h b/Engine/Source/Editor/Sequencer/Private/SequencerHotspots.h
--- a/Engine/Source/Editor/Sequencer/Private/SequencerHotspots.h
+++ b/Engine/Source/Editor/Sequencer/Private/SequencerHotspots.h
@@ -18,7 +18,7 @@ struct FKeyHotspot
 
 	virtual ESequencerHotspot GetType() const override { return ESequencerHotspot::Key; }
 	virtual TSharedPtr<ISequencerEditToolDragOperation> InitiateDrag(ISequencer&) override { return nullptr; }
-	virtual void PopulateContextMenu(FMenuBuilder& MenuBuilder, ISequencer& Sequencer) override;
+	virtual void PopulateContextMenu(FMenuBuilder& MenuBuilder, ISequencer& Sequencer, float MouseDownTime) override;
 
 	/** The key itself */
 	FSequencerSelectedKey Key;
@@ -57,7 +57,7 @@ struct FSectionHotspot
 
 	virtual ESequencerHotspot GetType() const override { return ESequencerHotspot::Section; }
 	virtual TSharedPtr<ISequencerEditToolDragOperation> InitiateDrag(ISequencer&) override { return nullptr; }
-	virtual void PopulateContextMenu(FMenuBuilder& MenuBuilder, ISequencer& Sequencer) override;
+	virtual void PopulateContextMenu(FMenuBuilder& MenuBuilder, ISequencer& Sequencer, float MouseDownTime) override;
 
 	/** Handle to the section */
 	FSectionHandle Section;
diff --git a/Engine/Source/Editor/Sequencer/Private/SequencerInputHandlerStack.h b/Engine/Source/Editor/Sequencer/Private/SequencerInputHandlerStack.h
--- a/Engine/Source/Editor/Sequencer/Private/SequencerInputHandlerStack.h
+++ b/Engine/Source/Editor/Sequencer/Private/SequencerInputHandlerStack.h
@@ -86,7 +86,7 @@ private:
 
 		for (int32 Index = 0; Index < Handlers.Num(); ++Index)
 		{
-			if (!Handlers[Index])
+			if (!Handlers[Index] || CapturedIndex == Index)
 			{
 				continue;
 			}
diff --git a/Engine/Source/Editor/Sequencer/Private/SequencerSectionLayoutBuilder.cpp b/Engine/Source/Editor/Sequencer/Private/SequencerSectionLayoutBuilder.cpp
--- a/Engine/Source/Editor/Sequencer/Private/SequencerSectionLayoutBuilder.cpp
+++ b/Engine/Source/Editor/Sequencer/Private/SequencerSectionLayoutBuilder.cpp
@@ -25,11 +25,13 @@ void FSequencerSectionLayoutBuilder::PopCategory()
 
 void FSequencerSectionLayoutBuilder::SetSectionAsKeyArea( TSharedRef<IKeyArea> KeyArea )
 {
+	KeyArea->SetName(RootNode->GetNodeName());
 	RootNode->SetSectionAsKeyArea( KeyArea ); 
 }
 
 void FSequencerSectionLayoutBuilder::AddKeyArea( FName KeyAreaName, const FText& DisplayName, TSharedRef<IKeyArea> KeyArea )
 {
+	KeyArea->SetName(KeyAreaName);
 	CurrentNode->AddKeyAreaNode( KeyAreaName, DisplayName, KeyArea );
 }
 
diff --git a/Engine/Source/Editor/Sequencer/Private/Tools/EditToolDragOperations.cpp b/Engine/Source/Editor/Sequencer/Private/Tools/EditToolDragOperations.cpp
--- a/Engine/Source/Editor/Sequencer/Private/Tools/EditToolDragOperations.cpp
+++ b/Engine/Source/Editor/Sequencer/Private/Tools/EditToolDragOperations.cpp
@@ -111,18 +111,25 @@ int32 FEditToolDragOperation::OnPaint(const FGeometry& AllottedGeometry, const F
 	return LayerId;
 }
 
-void FEditToolDragOperation::BeginTransaction( const TArray< FSectionHandle >& Sections, const FText& TransactionDesc )
+void FEditToolDragOperation::BeginTransaction( TArray< FSectionHandle >& Sections, const FText& TransactionDesc )
 {
 	// Begin an editor transaction and mark the section as transactional so it's state will be saved
 	Transaction.Reset( new FScopedTransaction(TransactionDesc) );
 
-	for (auto& Handle : Sections)
+	for (int32 SectionIndex = 0; SectionIndex < Sections.Num(); )
 	{
-		UMovieSceneSection* SectionObj = Handle.GetSectionObject();
+		UMovieSceneSection* SectionObj = Sections[SectionIndex].GetSectionObject();
 
 		SectionObj->SetFlags( RF_Transactional );
 		// Save the current state of the section
-		SectionObj->Modify();
+		if (SectionObj->TryModify())
+		{
+			++SectionIndex;
+		}
+		else
+		{
+			Sections.RemoveAt(SectionIndex);
+		}
 	}
 }
 
@@ -450,11 +457,12 @@ void FMoveKeys::OnBeginDrag(const FPointerEvent& MouseEvent, FVector2D LocalMous
 	SnapField = FSequencerSnapField(Sequencer, SnapCandidates);
 
 	// Begin an editor transaction and mark the section as transactional so it's state will be saved
-	BeginTransaction( TArray< FSectionHandle >(), NSLOCTEXT("Sequencer", "MoveKeysTransaction", "Move Keys") );
+	TArray<FSectionHandle> DummySections;
+	BeginTransaction( DummySections, NSLOCTEXT("Sequencer", "MoveKeysTransaction", "Move Keys") );
 
 	const float MouseTime = VirtualTrackArea.PixelToTime(LocalMousePos.X);
-
-	TSet<UMovieSceneSection*> ModifiedSections;
+	
+	ModifiedSections.Empty();
 	for( FSequencerSelectedKey SelectedKey : SelectedKeys )
 	{
 		UMovieSceneSection* OwningSection = SelectedKey.Section;
@@ -467,10 +475,11 @@ void FMoveKeys::OnBeginDrag(const FPointerEvent& MouseEvent, FVector2D LocalMous
 			OwningSection->SetFlags( RF_Transactional );
 
 			// Save the current state of the section
-			OwningSection->Modify();
-
-			// Section has been modified
-			ModifiedSections.Add( OwningSection );
+			if (OwningSection->TryModify())
+			{
+				// Section has been modified
+				ModifiedSections.Add( OwningSection );
+			}
 		}
 	}
 }
@@ -520,6 +529,11 @@ void FMoveKeys::OnDrag(const FPointerEvent& MouseEvent, FVector2D LocalMousePos,
 	{
 		UMovieSceneSection* Section = SelectedKey.Section;
 
+		if (!ModifiedSections.Contains(Section))
+		{
+			continue;
+		}
+
 		TSharedPtr<IKeyArea>& KeyArea = SelectedKey.KeyArea;
 
 		float NewKeyTime = MouseTime + RelativeOffsets.FindRef(SelectedKey);
@@ -560,6 +574,7 @@ void FMoveKeys::OnDrag(const FPointerEvent& MouseEvent, FVector2D LocalMousePos,
 
 void FMoveKeys::OnEndDrag(const FPointerEvent& MouseEvent, FVector2D LocalMousePos, const FVirtualTrackArea& VirtualTrackArea)
 {
+	ModifiedSections.Empty();
 	EndTransaction();
 	SSequencerSection::EnableLayoutRegeneration();
 }
\ No newline at end of file
diff --git a/Engine/Source/Editor/Sequencer/Private/Tools/EditToolDragOperations.h b/Engine/Source/Editor/Sequencer/Private/Tools/EditToolDragOperations.h
--- a/Engine/Source/Editor/Sequencer/Private/Tools/EditToolDragOperations.h
+++ b/Engine/Source/Editor/Sequencer/Private/Tools/EditToolDragOperations.h
@@ -32,7 +32,7 @@ public:
 protected:
 
 	/** begin a new scoped transaction for this drag */
-	void BeginTransaction( const TArray< FSectionHandle >& Sections, const FText& TransactionDesc );
+	void BeginTransaction( TArray< FSectionHandle >& Sections, const FText& TransactionDesc );
 
 	/** End an existing scoped transaction if one exists */
 	void EndTransaction();
@@ -155,4 +155,7 @@ private:
 
 	/** Snap field used to assist in snapping calculations */
 	TOptional<FSequencerSnapField> SnapField;
+
+	/** The set of sections being modified */
+	TSet<UMovieSceneSection*> ModifiedSections;
 };
diff --git a/Engine/Source/Editor/Sequencer/Private/Tools/SequencerEditTool_Default.cpp b/Engine/Source/Editor/Sequencer/Private/Tools/SequencerEditTool_Default.cpp
--- a/Engine/Source/Editor/Sequencer/Private/Tools/SequencerEditTool_Default.cpp
+++ b/Engine/Source/Editor/Sequencer/Private/Tools/SequencerEditTool_Default.cpp
@@ -4,6 +4,9 @@
 #include "SequencerEditTool_Default.h"
 #include "Sequencer.h"
 #include "SequencerHotspots.h"
+#include "SequencerContextMenus.h"
+#include "CommonMovieSceneTools.h"
+#include "VirtualTrackArea.h"
 
 
 FReply FSequencerEditTool_Default::OnMouseButtonUp(SWidget& OwnerWidget, const FGeometry& MyGeometry, const FPointerEvent& MouseEvent)
@@ -89,23 +92,36 @@ void FSequencerEditTool_Default::PerformHotspotSelection(const FPointerEvent& Mo
 
 TSharedPtr<SWidget> FSequencerEditTool_Default::OnSummonContextMenu( const FGeometry& MyGeometry, const FPointerEvent& MouseEvent )
 {
-	if (!Hotspot.IsValid())
-	{
-		return nullptr;
-	}
+	// @todo sequencer replace with UI Commands instead of faking it
 
-	ISequencer& Sequencer = GetSequencer();
+	FSequencer& Sequencer = static_cast<FSequencer&>(GetSequencer());
+	FVector2D MouseDownPos = MyGeometry.AbsoluteToLocal(MouseEvent.GetScreenSpacePosition());
 
-	// @todo sequencer replace with UI Commands instead of faking it
+	// Attempt to paste into either the current node selection, or the clicked on track
+	TSharedRef<SSequencer> SequencerWidget = StaticCastSharedRef<SSequencer>(Sequencer.GetSequencerWidget());
+	const float PasteAtTime = SequencerWidget->GetVirtualTrackArea().PixelToTime(MouseDownPos.X);
+
+	const bool bShouldCloseWindowAfterMenuSelection = true;
+	FMenuBuilder MenuBuilder(bShouldCloseWindowAfterMenuSelection, Sequencer.GetCommandBindings());
 
-	if (Hotspot->GetType() == ESequencerHotspot::Section || Hotspot->GetType() == ESequencerHotspot::Key)
+	if (Hotspot.IsValid())
 	{
-		const bool bShouldCloseWindowAfterMenuSelection = true;
-		FMenuBuilder MenuBuilder(bShouldCloseWindowAfterMenuSelection, NULL);
+		if (Hotspot->GetType() == ESequencerHotspot::Section || Hotspot->GetType() == ESequencerHotspot::Key)
+		{
+			Hotspot->PopulateContextMenu(MenuBuilder, Sequencer, PasteAtTime);
 
-		Hotspot->PopulateContextMenu(MenuBuilder, Sequencer);
+			return MenuBuilder.MakeWidget();
+		}
+	}
+	else if (Sequencer.GetClipboardStack().Num() != 0)
+	{
+		TSharedPtr<FPasteContextMenu> PasteMenu = FPasteContextMenu::CreateMenu(Sequencer, SequencerWidget->GeneratePasteArgs(PasteAtTime));
+		if (PasteMenu.IsValid() && PasteMenu->IsValidPaste())
+		{
+			PasteMenu->PopulateMenu(MenuBuilder);
 
-		return MenuBuilder.MakeWidget();
+			return MenuBuilder.MakeWidget();
+		}
 	}
 
 	return nullptr;
diff --git a/Engine/Source/Editor/Sequencer/Private/Tools/SequencerEditTool_Movement.cpp b/Engine/Source/Editor/Sequencer/Private/Tools/SequencerEditTool_Movement.cpp
--- a/Engine/Source/Editor/Sequencer/Private/Tools/SequencerEditTool_Movement.cpp
+++ b/Engine/Source/Editor/Sequencer/Private/Tools/SequencerEditTool_Movement.cpp
@@ -162,33 +162,40 @@ FReply FSequencerEditTool_Movement::OnMouseButtonUp(SWidget& OwnerWidget, const
 		// Only return handled if we actually started a drag
 		return FReply::Handled().ReleaseMouseCapture();
 	}
-	else if (Hotspot.IsValid())
+
+	FReply Reply = FReply::Unhandled();
+
+	if (Hotspot.IsValid())
 	{
 		PerformHotspotSelection(MouseEvent);
+		Reply = FReply::Handled();
+	}
 
-		if (MouseEvent.GetEffectingButton() == EKeys::RightMouseButton)
+	if (MouseEvent.GetEffectingButton() == EKeys::RightMouseButton)
+	{
+		TSharedPtr<SWidget> MenuContent = OnSummonContextMenu( MyGeometry, MouseEvent );
+		if (MenuContent.IsValid())
 		{
-			TSharedPtr<SWidget> MenuContent = OnSummonContextMenu( MyGeometry, MouseEvent );
-			if (MenuContent.IsValid())
-			{
-				FWidgetPath WidgetPath = MouseEvent.GetEventPath() != nullptr ? *MouseEvent.GetEventPath() : FWidgetPath();
+			FWidgetPath WidgetPath = MouseEvent.GetEventPath() != nullptr ? *MouseEvent.GetEventPath() : FWidgetPath();
 
-				FSlateApplication::Get().PushMenu(
-					OwnerWidget.AsShared(),
-					WidgetPath,
-					MenuContent.ToSharedRef(),
-					MouseEvent.GetScreenSpacePosition(),
-					FPopupTransitionEffect( FPopupTransitionEffect::ContextMenu )
-					);
+			FSlateApplication::Get().PushMenu(
+				OwnerWidget.AsShared(),
+				WidgetPath,
+				MenuContent.ToSharedRef(),
+				MouseEvent.GetScreenSpacePosition(),
+				FPopupTransitionEffect( FPopupTransitionEffect::ContextMenu )
+				);
 
-				return FReply::Handled().SetUserFocus(MenuContent.ToSharedRef(), EFocusCause::SetDirectly).ReleaseMouseCapture();
-			}
+			return FReply::Handled().SetUserFocus(MenuContent.ToSharedRef(), EFocusCause::SetDirectly).ReleaseMouseCapture();
 		}
+	}
 
-		return FReply::Handled().ReleaseMouseCapture();
+	if (!Reply.IsEventHandled())
+	{
+		Reply = FSequencerEditTool_Default::OnMouseButtonUp(OwnerWidget, MyGeometry, MouseEvent);
 	}
 
-	return FSequencerEditTool_Default::OnMouseButtonUp(OwnerWidget, MyGeometry, MouseEvent);
+	return Reply;
 }
 
 
diff --git a/Engine/Source/Editor/Sequencer/Public/IKeyArea.h b/Engine/Source/Editor/Sequencer/Public/IKeyArea.h
--- a/Engine/Source/Editor/Sequencer/Public/IKeyArea.h
+++ b/Engine/Source/Editor/Sequencer/Public/IKeyArea.h
@@ -7,7 +7,10 @@
 
 
 class UMovieSceneSection;
-
+class FMovieSceneClipboardBuilder;
+class FMovieSceneClipboardKeyTrack;
+struct FMovieSceneClipboardEnvironment;
+struct FSequencerPasteEnvironment;
 
 /**
  * Information for how to draw each key.
@@ -30,7 +33,7 @@ struct FKeyDrawingInfo
 /**
  * Interface that should be implemented for the UI portion of a key area within a section
  */
-class IKeyArea
+class IKeyArea : public TSharedFromThis<IKeyArea>
 {
 public:
 
@@ -145,7 +148,41 @@ public:
 	virtual bool CanCreateKeyEditor() = 0;
 
 	/**
-	* Creates a key editor for this key area for use in the animation outliner.
-	*/
+	 * Creates a key editor for this key area for use in the animation outliner.
+	 */
 	virtual TSharedRef<SWidget> CreateKeyEditor(ISequencer* Sequencer) = 0;
+
+	/**
+	 * Copy keys from this key area
+	 *
+	 * @param ClipboardBuilder	Clipboard builder to add keys to
+	 * @param KeyMask			Predicate to be called to check whether a key should be copied
+	 */
+	virtual void CopyKeys(FMovieSceneClipboardBuilder& ClipboardBuilder, const TFunctionRef<bool(FKeyHandle, const IKeyArea&)>& KeyMask) const = 0;
+
+	/**
+	 * Paste keys into this key area
+	 *
+	 * @param KeyTrack			Container of keys to paste into this area
+	 * @param SrcEnvironment		The source environment the cliboard originated from
+	 * @param DstEnvironment		The new environment to paste in
+	 */
+	virtual void PasteKeys(const FMovieSceneClipboardKeyTrack& KeyTrack, const FMovieSceneClipboardEnvironment& SrcEnvironment, const FSequencerPasteEnvironment& DstEnvironment) = 0;
+
+	/**
+	 * Give this key area a specific name
+	 */
+	virtual void SetName(FName Name) = 0;
+
+	/**
+	 * Get this area's name, or None if not set
+	 */
+	virtual FName GetName() const = 0;
+
+	/**
+	 * Get the key area's color.
+	 *
+	 * @return Key area color.
+	 */
+	virtual FLinearColor GetColor() = 0;
 };
diff --git a/Engine/Source/Editor/Sequencer/Public/ISequencerHotspot.h b/Engine/Source/Editor/Sequencer/Public/ISequencerHotspot.h
--- a/Engine/Source/Editor/Sequencer/Public/ISequencerHotspot.h
+++ b/Engine/Source/Editor/Sequencer/Public/ISequencerHotspot.h
@@ -22,6 +22,6 @@ struct ISequencerHotspot
 	virtual ~ISequencerHotspot() { }
 	virtual ESequencerHotspot GetType() const = 0;
 	virtual TSharedPtr<ISequencerEditToolDragOperation> InitiateDrag(ISequencer&) = 0;
-	virtual void PopulateContextMenu(FMenuBuilder& MenuBuilder, ISequencer& Sequencer){}
+	virtual void PopulateContextMenu(FMenuBuilder& MenuBuilder, ISequencer& Sequencer, float MouseDownTime){}
 	virtual FCursorReply GetCursor() const { return FCursorReply::Unhandled(); }
 };
diff --git a/Engine/Source/Editor/UMGEditor/Private/Animation/Sequencer2DTransformTrackEditor.cpp b/Engine/Source/Editor/UMGEditor/Private/Animation/Sequencer2DTransformTrackEditor.cpp
--- a/Engine/Source/Editor/UMGEditor/Private/Animation/Sequencer2DTransformTrackEditor.cpp
+++ b/Engine/Source/Editor/UMGEditor/Private/Animation/Sequencer2DTransformTrackEditor.cpp
@@ -25,18 +25,21 @@ public:
 
 	virtual void GenerateSectionLayout( class ISectionLayoutBuilder& LayoutBuilder ) const override
 	{
+		static const FLinearColor GreenKeyAreaColor(0.0f, 0.7f, 0.0f, 0.2f);
+		static const FLinearColor RedKeyAreaColor(0.7f, 0.0f, 0.0f, 0.2f);
+
 		UMovieScene2DTransformSection* TransformSection = Cast<UMovieScene2DTransformSection>(&SectionObject);
 
-		TranslationXKeyArea = MakeShareable( new FFloatCurveKeyArea( &TransformSection->GetTranslationCurve( EAxis::X ), TransformSection ) );
-		TranslationYKeyArea = MakeShareable( new FFloatCurveKeyArea( &TransformSection->GetTranslationCurve( EAxis::Y ), TransformSection ) );
+		TranslationXKeyArea = MakeShareable(new FFloatCurveKeyArea(&TransformSection->GetTranslationCurve(EAxis::X), TransformSection, RedKeyAreaColor));
+		TranslationYKeyArea = MakeShareable(new FFloatCurveKeyArea(&TransformSection->GetTranslationCurve(EAxis::Y), TransformSection, GreenKeyAreaColor));
 
-		RotationKeyArea = MakeShareable( new FFloatCurveKeyArea( &TransformSection->GetRotationCurve(), TransformSection ) );
+		RotationKeyArea = MakeShareable(new FFloatCurveKeyArea(&TransformSection->GetRotationCurve(), TransformSection));
 
-		ScaleXKeyArea = MakeShareable( new FFloatCurveKeyArea( &TransformSection->GetScaleCurve( EAxis::X ), TransformSection ) );
-		ScaleYKeyArea = MakeShareable( new FFloatCurveKeyArea( &TransformSection->GetScaleCurve( EAxis::Y ), TransformSection ) );
+		ScaleXKeyArea = MakeShareable(new FFloatCurveKeyArea(&TransformSection->GetScaleCurve(EAxis::X), TransformSection, RedKeyAreaColor));
+		ScaleYKeyArea = MakeShareable(new FFloatCurveKeyArea(&TransformSection->GetScaleCurve(EAxis::Y), TransformSection, GreenKeyAreaColor));
 
-		ShearXKeyArea = MakeShareable( new FFloatCurveKeyArea( &TransformSection->GetSheerCurve( EAxis::X ), TransformSection ) );
-		ShearYKeyArea = MakeShareable( new FFloatCurveKeyArea( &TransformSection->GetSheerCurve( EAxis::Y ), TransformSection ) );
+		ShearXKeyArea = MakeShareable(new FFloatCurveKeyArea(&TransformSection->GetSheerCurve(EAxis::X), TransformSection, RedKeyAreaColor));
+		ShearYKeyArea = MakeShareable(new FFloatCurveKeyArea(&TransformSection->GetSheerCurve(EAxis::Y), TransformSection, GreenKeyAreaColor));
 
 		// This generates the tree structure for the transform section
 		LayoutBuilder.PushCategory("Location", NSLOCTEXT("F2DTransformSection", "LocationArea", "Location"));
diff --git a/Engine/Source/Editor/UMGEditor/Private/Details/SWidgetDetailsView.cpp b/Engine/Source/Editor/UMGEditor/Private/Details/SWidgetDetailsView.cpp
--- a/Engine/Source/Editor/UMGEditor/Private/Details/SWidgetDetailsView.cpp
+++ b/Engine/Source/Editor/UMGEditor/Private/Details/SWidgetDetailsView.cpp
@@ -360,7 +360,7 @@ FText SWidgetDetailsView::GetNameText() const
 		UWidget* Widget = Cast<UWidget>(SelectedObjects[0].Get());
 		if ( Widget )
 		{
-			return FText::FromName(Widget->GetFName());
+			return Widget->IsGeneratedName() ? FText::FromName(Widget->GetFName()) : Widget->GetLabelText();
 		}
 	}
 	
@@ -411,7 +411,7 @@ void SWidgetDetailsView::HandleNameTextCommitted(const FText& Text, ETextCommit:
 			if ( HandleVerifyNameTextChanged(Text, DummyText) )
 			{
 				UWidget* Widget = Cast<UWidget>(SelectedObjects[0].Get());
-				FWidgetBlueprintEditorUtils::RenameWidget(BlueprintEditor.Pin().ToSharedRef(), Widget->GetFName(), FName(*Text.ToString()));
+				FWidgetBlueprintEditorUtils::RenameWidget(BlueprintEditor.Pin().ToSharedRef(), Widget->GetFName(), Text.ToString());
 			}
 		}
 		IsReentrant = false;
diff --git a/Engine/Source/Editor/UMGEditor/Private/WidgetBlueprintCompiler.cpp b/Engine/Source/Editor/UMGEditor/Private/WidgetBlueprintCompiler.cpp
--- a/Engine/Source/Editor/UMGEditor/Private/WidgetBlueprintCompiler.cpp
+++ b/Engine/Source/Editor/UMGEditor/Private/WidgetBlueprintCompiler.cpp
@@ -223,9 +223,13 @@ void FWidgetBlueprintCompiler::CreateClassVariablesFromBlueprint()
 		}
 
 		FEdGraphPinType WidgetPinType(Schema->PC_Object, TEXT(""), Widget->GetClass(), false, false);
+		
+		// Always name the variable according to the underlying FName of the widget object
 		UProperty* WidgetProperty = CreateVariable(Widget->GetFName(), WidgetPinType);
 		if ( WidgetProperty != nullptr )
 		{
+			const FString VariableName = Widget->IsGeneratedName() ? Widget->GetName() : Widget->GetLabelText().ToString();
+			WidgetProperty->SetMetaData(TEXT("DisplayName"), *VariableName);
 			WidgetProperty->SetMetaData(TEXT("Category"), *Blueprint->GetName());
 			
 			WidgetProperty->SetPropertyFlags(CPF_BlueprintVisible);
diff --git a/Engine/Source/Editor/UMGEditor/Private/WidgetBlueprintEditorUtils.cpp b/Engine/Source/Editor/UMGEditor/Private/WidgetBlueprintEditorUtils.cpp
--- a/Engine/Source/Editor/UMGEditor/Private/WidgetBlueprintEditorUtils.cpp
+++ b/Engine/Source/Editor/UMGEditor/Private/WidgetBlueprintEditorUtils.cpp
@@ -83,10 +83,11 @@ bool FWidgetBlueprintEditorUtils::VerifyWidgetRename(TSharedRef<class FWidgetBlu
 		return false;
 	}
 
-	const FName NewFName(*NewNameString);
-	
+	// Slug the new name down to a valid object name
+	const FName NewNameSlug = MakeObjectNameFromDisplayLabel(NewNameString, Widget.GetTemplate()->GetFName());
+
 	UWidgetBlueprint* Blueprint = BlueprintEditor->GetWidgetBlueprintObj();
-	UWidget* ExistingTemplate = Blueprint->WidgetTree->FindWidget(NewFName);
+	UWidget* ExistingTemplate = Blueprint->WidgetTree->FindWidget(NewNameSlug);
 
 	bool bIsSameWidget = false;
 	if (ExistingTemplate != nullptr)
@@ -103,17 +104,12 @@ bool FWidgetBlueprintEditorUtils::VerifyWidgetRename(TSharedRef<class FWidgetBlu
 	}
 	else
 	{
-		if (!FName(*NewNameString).IsValidObjectName(OutErrorMessage))
-		{
-			return false;
-		}
-
 		// Not an existing widget in the tree BUT it still mustn't create a UObject name clash
 		UWidget* WidgetPreview = Widget.GetPreview();
 		if (WidgetPreview)
 		{
 			// Dummy rename with flag REN_Test returns if rename is possible
-			if (!WidgetPreview->Rename(*NewNameString, nullptr, REN_Test))
+			if (!WidgetPreview->Rename(*NewNameSlug.ToString(), nullptr, REN_Test))
 			{
 				OutErrorMessage = LOCTEXT("ExistingObjectName", "Existing Object Name");
 				return false;
@@ -123,7 +119,7 @@ bool FWidgetBlueprintEditorUtils::VerifyWidgetRename(TSharedRef<class FWidgetBlu
 		if (WidgetTemplate)
 		{
 			// Dummy rename with flag REN_Test returns if rename is possible
-			if (!WidgetTemplate->Rename(*NewNameString, nullptr, REN_Test))
+			if (!WidgetTemplate->Rename(*NewNameSlug.ToString(), nullptr, REN_Test))
 			{
 				OutErrorMessage = LOCTEXT("ExistingObjectName", "Existing Object Name");
 				return false;
@@ -133,7 +129,8 @@ bool FWidgetBlueprintEditorUtils::VerifyWidgetRename(TSharedRef<class FWidgetBlu
 
 	FKismetNameValidator Validator(Blueprint);
 
-	const bool bUniqueNameForVariable = (EValidatorResult::Ok == Validator.IsValid(NewFName));
+	// For variable comparison, use the display label, not the slug
+	const bool bUniqueNameForVariable = (EValidatorResult::Ok == Validator.IsValid(NewNameString));
 
 	if (!bUniqueNameForVariable && !bIsSameWidget)
 	{
@@ -144,47 +141,50 @@ bool FWidgetBlueprintEditorUtils::VerifyWidgetRename(TSharedRef<class FWidgetBlu
 	return true;
 }
 
-bool FWidgetBlueprintEditorUtils::RenameWidget(TSharedRef<FWidgetBlueprintEditor> BlueprintEditor, const FName& OldName, const FName& NewName)
+bool FWidgetBlueprintEditorUtils::RenameWidget(TSharedRef<FWidgetBlueprintEditor> BlueprintEditor, const FName& OldObjectName, const FString& NewDisplayName)
 {
 	UWidgetBlueprint* Blueprint = BlueprintEditor->GetWidgetBlueprintObj();
 	check(Blueprint);
 
+	UWidget* Widget = Blueprint->WidgetTree->FindWidget(OldObjectName);
+	check(Widget);
+
 	bool bRenamed = false;
 
 	TSharedPtr<INameValidatorInterface> NameValidator = MakeShareable(new FKismetNameValidator(Blueprint));
 
 	// NewName should be already validated. But one must make sure that NewTemplateName is also unique.
-	const bool bUniqueNameForTemplate = ( EValidatorResult::Ok == NameValidator->IsValid(NewName) );
-
-	const FString NewNameStr = NewName.ToString();
-	const FString OldNameStr = OldName.ToString();
-
-	UWidget* Widget = Blueprint->WidgetTree->FindWidget(OldName);
-	check(Widget);
-
+	const bool bUniqueNameForTemplate = ( EValidatorResult::Ok == NameValidator->IsValid(NewDisplayName) );
 	if ( Widget )
 	{
+		// Get the new FName slug from the given display name
+		const FName NewFName = MakeObjectNameFromDisplayLabel(NewDisplayName, Widget->GetFName());
+
+		// Stringify the FNames
+		const FString NewNameStr = NewFName.ToString();
+		const FString OldNameStr = OldObjectName.ToString();
+
 		const FScopedTransaction Transaction(LOCTEXT("RenameWidget", "Rename Widget"));
 
 		// Rename Template
 		Blueprint->Modify();
 		Widget->Modify();
 
-
 		// Rename Preview before renaming the template widget so the preview widget can be found
 		UWidget* WidgetPreview = BlueprintEditor->GetReferenceFromTemplate(Widget).GetPreview();
 		if(WidgetPreview)
 		{
+			WidgetPreview->SetDisplayLabel(NewDisplayName);
 			WidgetPreview->Rename(*NewNameStr);
 		}
 
 		// Find and update all variable references in the graph
-
+		Widget->SetDisplayLabel(NewDisplayName);
 		Widget->Rename(*NewNameStr);
 
 		// Update Variable References and
 		// Update Event References to member variables
-		FBlueprintEditorUtils::ReplaceVariableReferences(Blueprint, OldName, NewName);
+		FBlueprintEditorUtils::ReplaceVariableReferences(Blueprint, OldObjectName, NewFName);
 		
 		// Find and update all binding references in the widget blueprint
 		for ( FDelegateEditorBinding& Binding : Blueprint->Bindings )
@@ -200,9 +200,9 @@ bool FWidgetBlueprintEditorUtils::RenameWidget(TSharedRef<FWidgetBlueprintEditor
 		{
 			for( FWidgetAnimationBinding& AnimBinding : WidgetAnimation->AnimationBindings )
 			{
-				if( AnimBinding.WidgetName == OldName )
+				if( AnimBinding.WidgetName == OldObjectName )
 				{
-					AnimBinding.WidgetName = NewName;
+					AnimBinding.WidgetName = NewFName;
 				}
 			}
 		}
@@ -211,7 +211,7 @@ bool FWidgetBlueprintEditorUtils::RenameWidget(TSharedRef<FWidgetBlueprintEditor
 		// TODO...
 
 		// Validate child blueprints and adjust variable names to avoid a potential name collision
-		FBlueprintEditorUtils::ValidateBlueprintChildVariables(Blueprint, NewName);
+		FBlueprintEditorUtils::ValidateBlueprintChildVariables(Blueprint, NewFName);
 
 		// Refresh references and flush editors
 		FBlueprintEditorUtils::MarkBlueprintAsStructurallyModified(Blueprint);
diff --git a/Engine/Source/Editor/UMGEditor/Private/WidgetBlueprintEditorUtils.h b/Engine/Source/Editor/UMGEditor/Private/WidgetBlueprintEditorUtils.h
--- a/Engine/Source/Editor/UMGEditor/Private/WidgetBlueprintEditorUtils.h
+++ b/Engine/Source/Editor/UMGEditor/Private/WidgetBlueprintEditorUtils.h
@@ -15,7 +15,7 @@ class FWidgetBlueprintEditorUtils
 public:
 	static bool VerifyWidgetRename(TSharedRef<class FWidgetBlueprintEditor> BlueprintEditor, FWidgetReference Widget, const FText& NewName, FText& OutErrorMessage);
 
-	static bool RenameWidget(TSharedRef<class FWidgetBlueprintEditor> BlueprintEditor, const FName& OldName, const FName& NewName);
+	static bool RenameWidget(TSharedRef<class FWidgetBlueprintEditor> BlueprintEditor, const FName& OldObjectName, const FString& NewDisplayName);
 
 	static void CreateWidgetContextMenu(FMenuBuilder& MenuBuilder, TSharedRef<FWidgetBlueprintEditor> BlueprintEditor, FVector2D TargetLocation);
 
diff --git a/Engine/Source/Editor/UnrealEd/Private/Commandlets/ContentCommandlets.cpp b/Engine/Source/Editor/UnrealEd/Private/Commandlets/ContentCommandlets.cpp
--- a/Engine/Source/Editor/UnrealEd/Private/Commandlets/ContentCommandlets.cpp
+++ b/Engine/Source/Editor/UnrealEd/Private/Commandlets/ContentCommandlets.cpp
@@ -65,6 +65,7 @@ int32 UResavePackagesCommandlet::InitializeResaveParameters( const TArray<FStrin
 		FString Package;
 		FString PackageFolder;
 		FString Maps;
+		FString File;
 		const FString& CurrentSwitch = Switches[ SwitchIdx ];
 		if( FParse::Value( *CurrentSwitch, TEXT( "PACKAGE="), Package ) )
 		{
@@ -106,6 +107,39 @@ int32 UResavePackagesCommandlet::InitializeResaveParameters( const TArray<FStrin
 			PackageNames.Add(*MapFile);
 			bExplicitPackages = true;
 		}
+		else if (FParse::Value(*CurrentSwitch, TEXT("FILE="), File))
+		{
+			FString Text;
+			if (FFileHelper::LoadFileToString(Text, *File))
+			{
+				TArray<FString> Lines;
+
+				// Remove all carriage return characters.
+				Text.ReplaceInline(TEXT("\r"), TEXT(""));
+				// Read all lines
+				Text.ParseIntoArray(Lines, TEXT("\n"), true);
+
+				for (auto Line : Lines)
+				{
+					FString PackageFile;
+					if (FPackageName::SearchForPackageOnDisk(Line, NULL, &PackageFile, false))
+					{
+						PackageNames.Add(*PackageFile);
+					}
+					else
+					{
+						UE_LOG(LogContentCommandlet, Error, TEXT("Failed to find package %s"), *Line);
+					}
+				}
+
+				bExplicitPackages = true;
+				UE_LOG(LogContentCommandlet, Display, TEXT("Loaded %d Packages from %s"), PackageNames.Num(), *File);
+			}
+			else
+			{
+				UE_LOG(LogContentCommandlet, Error, TEXT("Failed to load file %s"), *File);
+			}
+		}
 
 	}
 
diff --git a/Engine/Source/Editor/UnrealEd/Private/Lightmass/LightmassRender.cpp b/Engine/Source/Editor/UnrealEd/Private/Lightmass/LightmassRender.cpp
--- a/Engine/Source/Editor/UnrealEd/Private/Lightmass/LightmassRender.cpp
+++ b/Engine/Source/Editor/UnrealEd/Private/Lightmass/LightmassRender.cpp
@@ -17,6 +17,7 @@
 #include "LightmassLandscapeRender.h"
 #include "LandscapeMaterialInstanceConstant.h"
 #include "LandscapeLight.h"
+#include "EngineModule.h"
 
 DEFINE_LOG_CATEGORY_STATIC(LogLightmassRender, Error, All);
 
@@ -1014,6 +1015,13 @@ bool FLightmassMaterialRenderer::GenerateMaterialPropertyData(
 			}
 			else
 			{
+				ENQUEUE_UNIQUE_RENDER_COMMAND(
+					InitializeSystemTextures,
+					{
+						GetRendererModule().InitializeSystemTextures(RHICmdList);
+					});
+				
+
 				if (bIsLandscapeMaterial)
 				{
 					// Landscape needs special handling because it uses multiple UVs, which isn't yet supported by lightmass's regular pipeline
diff --git a/Engine/Source/Programs/AutomationTool/AutomationUtils/MCPPublic.cs b/Engine/Source/Programs/AutomationTool/AutomationUtils/MCPPublic.cs
--- a/Engine/Source/Programs/AutomationTool/AutomationUtils/MCPPublic.cs
+++ b/Engine/Source/Programs/AutomationTool/AutomationUtils/MCPPublic.cs
@@ -662,15 +662,45 @@ namespace EpicGames.MCP.Automation
 	public abstract class CloudStorageBase
 	{
 		private static readonly object LockObj = new object();
-		private static CloudStorageBase Handler = null;
+		private static Dictionary<string, CloudStorageBase> Handlers = new Dictionary<string, CloudStorageBase>();
+		private const string DEFAULT_INSTANCE_NAME = "DefaultInstance";
 
+		/// <summary>
+		/// Gets the default instance of CloudStorageBase
+		/// </summary>
+		/// <returns>A default instance of CloudStorageBase. The first time each instance is returned, it will require initialization with its Init() method.</returns>
 		public static CloudStorageBase Get()
 		{
-			if (Handler == null)
+			return GetByNameImpl(DEFAULT_INSTANCE_NAME); // Identifier for the default cloud storage
+		}
+
+		/// <summary>
+		/// Gets an instance of CloudStorageBase.
+		/// Multiple calls with the same instance name will return the same object.
+		/// </summary>
+		/// <param name="InstanceName">The name of the object to return</param>
+		/// <returns>An instance of CloudStorageBase. The first time each instance is returned, it will require initialization with its Init() method.</returns>
+		public static CloudStorageBase GetByName(string InstanceName)
+		{
+			if (InstanceName == DEFAULT_INSTANCE_NAME)
+			{
+				CommandUtils.LogWarning("CloudStorageBase.GetByName called with {0}. This will return the same instance as Get().", DEFAULT_INSTANCE_NAME);
+			}
+			return GetByNameImpl(InstanceName);
+		}
+
+		private  static CloudStorageBase GetByNameImpl(string InstanceName)
+		{
+			CloudStorageBase Result = null;
+			if (!Handlers.TryGetValue(InstanceName, out Result))
 			{
 				lock (LockObj)
 				{
-					if (Handler == null)
+					if (Handlers.ContainsKey(InstanceName))
+					{
+						Result = Handlers[InstanceName];
+					}
+					else
 					{
 						Assembly[] LoadedAssemblies = AppDomain.CurrentDomain.GetAssemblies();
 						foreach (var Dll in LoadedAssemblies)
@@ -680,19 +710,20 @@ namespace EpicGames.MCP.Automation
 							{
 								if (PotentialConfigType != typeof(CloudStorageBase) && typeof(CloudStorageBase).IsAssignableFrom(PotentialConfigType))
 								{
-									Handler = Activator.CreateInstance(PotentialConfigType) as CloudStorageBase;
+									Result = Activator.CreateInstance(PotentialConfigType) as CloudStorageBase;
+									Handlers.Add(InstanceName, Result);
 									break;
 								}
 							}
 						}
 					}
 				}
-				if (Handler == null)
+				if (Result == null)
 				{
-					throw new AutomationException("Attempt to use CloudStorageBase.Get() and it doesn't appear that there are any modules that implement this class.");
+					throw new AutomationException("Could not find any modules which provide an implementation of CloudStorageBase.");
 				}
 			}
-			return Handler;
+			return Result;
 		}
 
 		/// <summary>
diff --git a/Engine/Source/Programs/AutomationTool/Scripts/CookCommand.Automation.cs b/Engine/Source/Programs/AutomationTool/Scripts/CookCommand.Automation.cs
--- a/Engine/Source/Programs/AutomationTool/Scripts/CookCommand.Automation.cs
+++ b/Engine/Source/Programs/AutomationTool/Scripts/CookCommand.Automation.cs
@@ -298,12 +298,20 @@ public partial class Project : CommandUtils
             try
             {
                 Directory.Delete(TemporaryPakPath, true);
+            }
+            catch(Exception Ex)
+            {
+                Log("Failed deleting temporary directories "+TemporaryPakPath+" continuing. "+ Ex.GetType().ToString());
+            }
+            try
+            {
                 Directory.Delete(TemporaryFilesPath, true);
             }
-            catch(Exception )
+            catch (Exception Ex)
             {
-                Log("Failed deleting temporary directories "+TemporaryPakPath+" "+TemporaryFilesPath+" continuing.");
+                Log("Failed deleting temporary directories " + TemporaryFilesPath + " continuing. " + Ex.GetType().ToString());
             }
+
             try
             {
 
@@ -322,7 +330,8 @@ public partial class Project : CommandUtils
                     CurrentPlatform.ExtractPackage(Params, Params.DiffCookedContentPath, TemporaryPakPath);
 
                     // find the pak file
-                    PakFiles = Directory.EnumerateFiles(TemporaryPakPath, "*.pak", SearchOption.AllDirectories).ToList();
+                    PakFiles.AddRange( Directory.EnumerateFiles(TemporaryPakPath, Params.ShortProjectName+"*.pak", SearchOption.AllDirectories));
+                    PakFiles.AddRange( Directory.EnumerateFiles(TemporaryPakPath, "pakchunk*.pak", SearchOption.AllDirectories));
                 }
 
                 string CookPlatformString = CurrentPlatform.GetCookPlatform(false, Params.HasDedicatedServerAndClient, Params.CookFlavor);
@@ -361,8 +370,14 @@ public partial class Project : CommandUtils
                     Log("Extracting pak " + Name + " for comparision to location " + TemporaryFilesPath);
 
                     string UnrealPakParams = Name + " -Extract " + " " + TemporaryFilesPath;
-
-                    RunAndLog(CmdEnv, UnrealPakExe, UnrealPakParams, Options: ERunOptions.Default | ERunOptions.UTF8Output | ERunOptions.LoggingOfRunDuration);
+                    try
+                    {
+                        RunAndLog(CmdEnv, UnrealPakExe, UnrealPakParams, Options: ERunOptions.Default | ERunOptions.UTF8Output | ERunOptions.LoggingOfRunDuration);
+                    }
+                    catch(Exception Ex)
+                    {
+                        Log("Pak failed to extract because of " + Ex.GetType().ToString());
+                    }
                 }
 
                 const string RootFailedContentDirectory = "\\\\epicgames.net\\root\\Developers\\Daniel.Lamb";
diff --git a/Engine/Source/Runtime/AnimGraphRuntime/Private/AnimNodes/AnimNode_SaveCachedPose.cpp b/Engine/Source/Runtime/AnimGraphRuntime/Private/AnimNodes/AnimNode_SaveCachedPose.cpp
--- a/Engine/Source/Runtime/AnimGraphRuntime/Private/AnimNodes/AnimNode_SaveCachedPose.cpp
+++ b/Engine/Source/Runtime/AnimGraphRuntime/Private/AnimNodes/AnimNode_SaveCachedPose.cpp
@@ -78,7 +78,7 @@ void FAnimNode_SaveCachedPose::GatherDebugData(FNodeDebugData& DebugData)
 
 	// we need to fix SaveCachePose to collect the highest weight branch, not the first one,
 	// for it to be useful in debug.
-	if(false)
+	if(true)
 // 	if (!DebugDataCounter.IsSynchronizedWith(DebugData.AnimInstance->DebugDataCounter))
 	{
 		DebugDataCounter.SynchronizeWith(DebugData.AnimInstance->DebugDataCounter);
diff --git a/Engine/Source/Runtime/AnimGraphRuntime/Private/BoneControllers/AnimNode_Trail.cpp b/Engine/Source/Runtime/AnimGraphRuntime/Private/BoneControllers/AnimNode_Trail.cpp
--- a/Engine/Source/Runtime/AnimGraphRuntime/Private/BoneControllers/AnimNode_Trail.cpp
+++ b/Engine/Source/Runtime/AnimGraphRuntime/Private/BoneControllers/AnimNode_Trail.cpp
@@ -6,6 +6,8 @@
 /////////////////////////////////////////////////////
 // FAnimNode_Trail
 
+DECLARE_CYCLE_STAT(TEXT("Trail Eval"), STAT_Trail_Eval, STATGROUP_Anim);
+
 FAnimNode_Trail::FAnimNode_Trail()
 	: ChainLength(2)
 	, ChainBoneAxis(EAxis::X)
@@ -43,6 +45,8 @@ void FAnimNode_Trail::GatherDebugData(FNodeDebugData& DebugData)
 
 void FAnimNode_Trail::EvaluateBoneTransforms(USkeletalMeshComponent* SkelComp, FCSPose<FCompactPose>& MeshBases, TArray<FBoneTransform>& OutBoneTransforms)
 {
+	SCOPE_CYCLE_COUNTER(STAT_Trail_Eval);
+
 	check(OutBoneTransforms.Num() == 0);
 
 	if( ChainBoneIndices.Num() <= 0 )
diff --git a/Engine/Source/Runtime/AssetRegistry/Private/AssetDataGatherer.cpp b/Engine/Source/Runtime/AssetRegistry/Private/AssetDataGatherer.cpp
--- a/Engine/Source/Runtime/AssetRegistry/Private/AssetDataGatherer.cpp
+++ b/Engine/Source/Runtime/AssetRegistry/Private/AssetDataGatherer.cpp
@@ -63,7 +63,7 @@ FAssetDataDiscovery::FAssetDataDiscovery(const TArray<FString>& InPaths, bool bI
 	for (const FString& Path : InPaths)
 	{
 		// Convert the package path to a filename with no extension (directory)
-		DirectoriesToSearch.Add(FPackageName::LongPackageNameToFilename(Path));
+		DirectoriesToSearch.Add(FPackageName::LongPackageNameToFilename(Path / TEXT("")));
 	}
 
 	if (bIsSynchronous)
@@ -321,13 +321,13 @@ void FAssetDataDiscovery::AddPathToSearch(const FString& Path)
 	FScopeLock CritSectionLock(&WorkerThreadCriticalSection);
 
 	// Convert the package path to a filename with no extension (directory)
-	DirectoriesToSearch.Add(FPackageName::LongPackageNameToFilename(Path));
+	DirectoriesToSearch.Add(FPackageName::LongPackageNameToFilename(Path / TEXT("")));
 }
 
 void FAssetDataDiscovery::PrioritizeSearchPath(const FString& PathToPrioritize)
 {
 	FString LocalFilenamePathToPrioritize;
-	if (FPackageName::TryConvertLongPackageNameToFilename(PathToPrioritize, LocalFilenamePathToPrioritize))
+	if (FPackageName::TryConvertLongPackageNameToFilename(PathToPrioritize / TEXT(""), LocalFilenamePathToPrioritize))
 	{
 		FScopeLock CritSectionLock(&WorkerThreadCriticalSection);
 
@@ -735,7 +735,7 @@ void FAssetDataGatherer::PrioritizeSearchPath(const FString& PathToPrioritize)
 	}
 
 	FString LocalFilenamePathToPrioritize;
-	if (FPackageName::TryConvertLongPackageNameToFilename(PathToPrioritize, LocalFilenamePathToPrioritize))
+	if (FPackageName::TryConvertLongPackageNameToFilename(PathToPrioritize / TEXT(""), LocalFilenamePathToPrioritize))
 	{
 		FScopeLock CritSectionLock(&WorkerThreadCriticalSection);
 
diff --git a/Engine/Source/Runtime/AssetRegistry/Private/PathTree.cpp b/Engine/Source/Runtime/AssetRegistry/Private/PathTree.cpp
--- a/Engine/Source/Runtime/AssetRegistry/Private/PathTree.cpp
+++ b/Engine/Source/Runtime/AssetRegistry/Private/PathTree.cpp
@@ -2,7 +2,7 @@
 
 #include "AssetRegistryPCH.h"
 
-bool FPathTree::CachePath(const FName& Path)
+bool FPathTree::CachePath(FName Path)
 {
 	if (Path.IsNone())
 	{
@@ -19,6 +19,19 @@ bool FPathTree::CachePath(const FName& Path)
 	check(PathStr.Len() >= 2);	// Must be at least "/A"
 	check(PathStr[0] == '/');	// Must start with a "/"
 
+	// Paths are cached without their trailing slash, so if the given path has a trailing slash, test it again now as it may already be cached
+	if (PathStr[PathStr.Len() - 1] == '/')
+	{
+		PathStr.RemoveAt(PathStr.Len() - 1, 1, /*bAllowShrinking*/false);
+		Path = *PathStr;
+
+		if (ParentPathToChildPaths.Contains(Path))
+		{
+			// Already cached - nothing more to do
+			return false;
+		}
+	}
+
 	// Walk each part of the path, adding known path entries if required
 	// This manipulates PathStr in-place to avoid making any string copies
 	FName LastPath;
@@ -66,7 +79,7 @@ bool FPathTree::CachePath(const FName& Path)
 	return true;
 }
 
-bool FPathTree::RemovePath(const FName& Path)
+bool FPathTree::RemovePath(FName Path)
 {
 	if (Path.IsNone())
 	{
@@ -75,8 +88,25 @@ bool FPathTree::RemovePath(const FName& Path)
 
 	if (!ParentPathToChildPaths.Contains(Path))
 	{
-		// Doesn't exist - nothing more to do
-		return false;
+		// Paths are cached without their trailing slash, so if the given path has a trailing slash, test it again now as it may already be cached
+		// We do this after the initial map test as: a) Most paths are well formed, b) This avoids string allocations until we know we need them
+		FString PathStr = Path.ToString();
+		if (PathStr[PathStr.Len() - 1] == '/')
+		{
+			PathStr.RemoveAt(PathStr.Len() - 1, 1, /*bAllowShrinking*/false);
+			Path = *PathStr;
+
+			if (!ParentPathToChildPaths.Contains(Path))
+			{
+				// Doesn't exist - nothing more to do
+				return false;
+			}
+		}
+		else
+		{
+			// Doesn't exist - nothing more to do
+			return false;
+		}
 	}
 
 	// We also need to gather up and remove any children of this path
@@ -119,7 +149,7 @@ bool FPathTree::GetAllPaths(TSet<FName>& OutPaths) const
 	return OutPaths.Num() > 0;
 }
 
-bool FPathTree::GetSubPaths(const FName& BasePath, TSet<FName>& OutPaths, bool bRecurse) const
+bool FPathTree::GetSubPaths(FName BasePath, TSet<FName>& OutPaths, bool bRecurse) const
 {
 	if (BasePath.IsNone())
 	{
@@ -129,7 +159,24 @@ bool FPathTree::GetSubPaths(const FName& BasePath, TSet<FName>& OutPaths, bool b
 	const TSet<FName>* ChildPathsPtr = ParentPathToChildPaths.Find(BasePath);
 	if (!ChildPathsPtr)
 	{
-		return false;
+		// Paths are cached without their trailing slash, so if the given path has a trailing slash, test it again now as it may already be cached
+		// We do this after the initial map test as: a) Most paths are well formed, b) This avoids string allocations until we know we need them
+		FString BasePathStr = BasePath.ToString();
+		if (BasePathStr[BasePathStr.Len() - 1] == '/')
+		{
+			BasePathStr.RemoveAt(BasePathStr.Len() - 1, 1, /*bAllowShrinking*/false);
+			BasePath = *BasePathStr;
+
+			ChildPathsPtr = ParentPathToChildPaths.Find(BasePath);
+			if (!ChildPathsPtr)
+			{
+				return false;
+			}
+		}
+		else
+		{
+			return false;
+		}
 	}
 
 	const int32 OutPathsOriginalNum = OutPaths.Num();
diff --git a/Engine/Source/Runtime/AssetRegistry/Private/PathTree.h b/Engine/Source/Runtime/AssetRegistry/Private/PathTree.h
--- a/Engine/Source/Runtime/AssetRegistry/Private/PathTree.h
+++ b/Engine/Source/Runtime/AssetRegistry/Private/PathTree.h
@@ -6,16 +6,16 @@ class FPathTree
 {
 public:
 	/** Adds the path to the tree relative to this node, creating nodes as needed. Returns true if the specified path was actually added (as opposed to already existed) */
-	bool CachePath(const FName& Path);
+	bool CachePath(FName Path);
 
 	/** Removes the specified path in the tree relative to this node. Returns true if the path was found and removed. */
-	bool RemovePath(const FName& Path);
+	bool RemovePath(FName Path);
 
 	/** Get all of the paths we know about */
 	bool GetAllPaths(TSet<FName>& OutPaths) const;
 
 	/** Recursively gathers all child paths from the specified base path relative to this node */
-	bool GetSubPaths(const FName& BasePath, TSet<FName>& OutPaths, bool bRecurse = true) const;
+	bool GetSubPaths(FName BasePath, TSet<FName>& OutPaths, bool bRecurse = true) const;
 
 private:
 	/** A one-to-many mapping between a parent path and its child paths. */
diff --git a/Engine/Source/Runtime/Core/Private/Misc/AutomationTest.cpp b/Engine/Source/Runtime/Core/Private/Misc/AutomationTest.cpp
--- a/Engine/Source/Runtime/Core/Private/Misc/AutomationTest.cpp
+++ b/Engine/Source/Runtime/Core/Private/Misc/AutomationTest.cpp
@@ -156,7 +156,7 @@ bool FAutomationTestFramework::RunSmokeTests()
 			for ( int TestIndex = 0; TestIndex < TestInfo.Num(); ++TestIndex )
 			{
 				SlowTask.EnterProgressFrame(1);
-				if (TestInfo[TestIndex].GetTestFlags() == EAutomationTestFlags::SmokeFilter )
+				if (TestInfo[TestIndex].GetTestFlags() & EAutomationTestFlags::SmokeFilter )
 				{
 					FString TestCommand = TestInfo[TestIndex].GetTestName();
 					FAutomationTestExecutionInfo& CurExecutionInfo = OutExecutionInfoMap.Add( TestCommand, FAutomationTestExecutionInfo() );
