Copying //UE4/Orion-Staging to Dev-Main (Originating from //Orion/Main at CL-2777663)

diff --git a/Engine/Source/Editor/Sequencer/Private/SSequencer.cpp b/Engine/Source/Editor/Sequencer/Private/SSequencer.cpp
--- a/Engine/Source/Editor/Sequencer/Private/SSequencer.cpp
+++ b/Engine/Source/Editor/Sequencer/Private/SSequencer.cpp
@@ -138,8 +138,8 @@ void SSequencer::Construct(const FArguments& InArgs, TSharedRef<FSequencer> InSe
 	OnGetAddMenuContent = InArgs._OnGetAddMenuContent;
 	AddMenuExtender = InArgs._AddMenuExtender;
 
-	ColumnFillCoefficients[0] = .25f;
-	ColumnFillCoefficients[1] = .75f;
+	ColumnFillCoefficients[0] = 0.3f;
+	ColumnFillCoefficients[1] = 0.7f;
 
 	TAttribute<float> FillCoefficient_0, FillCoefficient_1;
 	{
@@ -226,7 +226,16 @@ void SSequencer::Construct(const FArguments& InArgs, TSharedRef<FSequencer> InSe
 							.Orientation(Orient_Horizontal)
 						
 						+ SSplitter::Slot()
-							.Value(0.80f)
+							.Value(0.10f)
+							[
+								// track label browser
+								SAssignNew(LabelBrowser, SSequencerLabelBrowser)
+									.OnSelectionChanged(this, &SSequencer::HandleLabelBrowserSelectionChanged)
+									.Visibility(this, &SSequencer::HandleLabelBrowserVisibility)
+							]
+
+						+ SSplitter::Slot()
+							.Value(0.7f)
 							[
 								SNew(SOverlay)
 
@@ -251,7 +260,7 @@ void SSequencer::Construct(const FArguments& InArgs, TSharedRef<FSequencer> InSe
 															.BorderBackgroundColor(FLinearColor(.50f, .50f, .50f, 1.0f))
 															.Padding(FMargin(4.0f, 3.0f))
 															[
-																SNew( SSearchBox )
+																SAssignNew(SearchBox, SSearchBox)
 																	.OnTextChanged( this, &SSequencer::OnOutlinerSearchChanged )
 															]
 													]
@@ -274,7 +283,7 @@ void SSequencer::Construct(const FArguments& InArgs, TSharedRef<FSequencer> InSe
 																		+ SHorizontalBox::Slot()
 																			.FillWidth( FillCoefficient_0 )
 																			[
-																				SNew(SSequencerTreeViewBox, InSequencer)
+																				SNew(SSequencerTreeViewBox, InSequencer, SharedThis(this))
 																					.Padding(FMargin(0, 0, 10.f, 0)) // Padding to allow space for the scroll bar
 																					[
 																						TreeView.ToSharedRef()
@@ -359,6 +368,7 @@ void SSequencer::Construct(const FArguments& InArgs, TSharedRef<FSequencer> InSe
 														SNew( SBorder )
 															.BorderImage( FEditorStyle::GetBrush("ToolPanel.GroupBorder") )
 															.BorderBackgroundColor( FLinearColor(.50f, .50f, .50f, 1.0f ) )
+															.Padding(0)
 															[
 																SNew(SVerticalBox)
 
@@ -418,7 +428,7 @@ void SSequencer::Construct(const FArguments& InArgs, TSharedRef<FSequencer> InSe
 							[
 								// details view panel
 								SNew(SVerticalBox)
-								.Visibility(this, &SSequencer::HandleDetailsViewVisibility)
+									.Visibility(this, &SSequencer::HandleDetailsViewVisibility)
 
 								+ SVerticalBox::Slot()
 									.AutoHeight()
@@ -451,7 +461,7 @@ void SSequencer::BindCommands(TSharedRef<FUICommandList> SequencerCommandBinding
 {
 	SequencerCommandBindings->MapAction(
 		FSequencerCommands::Get().MoveTool,
-		FExecuteAction::CreateLambda([&] {
+		FExecuteAction::CreateLambda([this] {
 			EditTool.Reset(new FSequencerEditTool_Movement(Sequencer.Pin(), SharedThis(this)));
 		}),
 		FCanExecuteAction::CreateLambda([] {
@@ -462,7 +472,7 @@ void SSequencer::BindCommands(TSharedRef<FUICommandList> SequencerCommandBinding
 
 	SequencerCommandBindings->MapAction(
 		FSequencerCommands::Get().MarqueeTool,
-		FExecuteAction::CreateLambda( [&] {
+		FExecuteAction::CreateLambda( [this] {
 			EditTool.Reset(new FSequencerEditTool_Selection(Sequencer.Pin(), SharedThis(this)));
 		}),
 		FCanExecuteAction::CreateLambda([] {
@@ -471,20 +481,25 @@ void SSequencer::BindCommands(TSharedRef<FUICommandList> SequencerCommandBinding
 		FIsActionChecked::CreateSP(this, &SSequencer::IsEditToolEnabled, FName("Selection"))
 	);
 
+	auto CanPaste = [this]{
+		if (!HasFocusedDescendants() && !HasKeyboardFocus())
+		{
+			return false;
+		}
+
+		return Sequencer.Pin()->GetClipboardStack().Num() != 0;
+	};
+
 	SequencerCommandBindings->MapAction(
 		FGenericCommands::Get().Paste,
 		FExecuteAction::CreateSP(this, &SSequencer::Paste),
-		FCanExecuteAction::CreateLambda([&]{
-			return Sequencer.Pin()->GetClipboardStack().Num() != 0;
-		})
+		FCanExecuteAction::CreateLambda(CanPaste)
 	);
 
 	SequencerCommandBindings->MapAction(
 		FSequencerCommands::Get().PasteFromHistory,
 		FExecuteAction::CreateSP(this, &SSequencer::PasteFromHistory),
-		FCanExecuteAction::CreateLambda([&]{
-			return Sequencer.Pin()->GetClipboardStack().Num() != 0;
-		})
+		FCanExecuteAction::CreateLambda(CanPaste)
 	);
 }
 
@@ -565,6 +580,23 @@ void SSequencer::HandleKeySelectionChanged()
 }
 
 
+void SSequencer::HandleLabelBrowserSelectionChanged(FString NewLabel, ESelectInfo::Type SelectInfo)
+{
+	SearchBox->SetText(FText::FromString(NewLabel));
+}
+
+
+EVisibility SSequencer::HandleLabelBrowserVisibility() const
+{
+	if (Sequencer.Pin()->GetSettings()->GetLabelBrowserVisible())
+	{
+		return EVisibility::Visible;
+	}
+
+	return EVisibility::Collapsed;
+}
+
+
 void SSequencer::HandleSectionSelectionChanged()
 {
 	UpdateDetailsView();
@@ -751,6 +783,7 @@ TSharedRef<SWidget> SSequencer::MakeGeneralMenu()
 		if (Sequencer.Pin()->IsLevelEditorSequencer())
 		{
 			MenuBuilder.AddMenuEntry( FSequencerCommands::Get().ToggleDetailsView );
+			MenuBuilder.AddMenuEntry( FSequencerCommands::Get().ToggleLabelBrowser );
 		}
 
 		if (Sequencer.Pin()->IsLevelEditorSequencer())
diff --git a/Engine/Source/Editor/Sequencer/Private/SSequencer.h b/Engine/Source/Editor/Sequencer/Private/SSequencer.h
--- a/Engine/Source/Editor/Sequencer/Private/SSequencer.h
+++ b/Engine/Source/Editor/Sequencer/Private/SSequencer.h
@@ -18,6 +18,7 @@ class FUnloadedClassDragDropOp;
 class IDetailsView;
 class UMovieSceneSection;
 class SSequencerCurveEditor;
+class SSequencerLabelBrowser;
 class SSequencerTrackArea;
 class SSequencerTrackOutliner;
 class SSequencerTreeView;
@@ -230,6 +231,12 @@ private:
 	/** Handles key selection changes. */
 	void HandleKeySelectionChanged();
 
+	/** Handles selection changes in the label browser. */
+	void HandleLabelBrowserSelectionChanged(FString NewLabel, ESelectInfo::Type SelectInfo);
+
+	/** Handles determining the visibility of the label browser. */
+	EVisibility HandleLabelBrowserVisibility() const;
+
 	/** Handles section selection changes. */
 	void HandleSectionSelectionChanged();
 
@@ -399,6 +406,12 @@ private:
 	/** The breadcrumb trail widget for this sequencer */
 	TSharedPtr<SBreadcrumbTrail<FSequencerBreadcrumb>> BreadcrumbTrail;
 
+	/** The label browser for filtering tracks. */
+	TSharedPtr<SSequencerLabelBrowser> LabelBrowser;
+
+	/** The search box for filtering tracks. */
+	TSharedPtr<SSearchBox> SearchBox;
+
 	/** The sequencer tree view responsible for the outliner and track areas */
 	TSharedPtr<SSequencerTreeView> TreeView;
 
diff --git a/Engine/Source/Editor/Sequencer/Private/SSequencerLabelBrowser.cpp b/Engine/Source/Editor/Sequencer/Private/SSequencerLabelBrowser.cpp
new file mode 100644
--- /dev/null
+++ b/Engine/Source/Editor/Sequencer/Private/SSequencerLabelBrowser.cpp
@@ -0,0 +1,167 @@
+// Copyright 1998-2015 Epic Games, Inc. All Rights Reserved.
+
+#include "SequencerPrivatePCH.h"
+#include "SSequencerLabelListRow.h"
+
+
+#define LOCTEXT_NAMESPACE "SSequencerLabelBrowser"
+
+
+/* SSequencerLabelBrowser structors
+ *****************************************************************************/
+
+SSequencerLabelBrowser::~SSequencerLabelBrowser()
+{
+}
+
+
+/* SSequencerLabelBrowser interface
+ *****************************************************************************/
+
+BEGIN_SLATE_FUNCTION_BUILD_OPTIMIZATION
+void SSequencerLabelBrowser::Construct(const FArguments& InArgs)
+{
+	OnSelectionChanged = InArgs._OnSelectionChanged;
+
+	ChildSlot
+	[
+		SAssignNew(LabelTreeView, STreeView<TSharedPtr<FSequencerLabelTreeNode>>)
+			.ItemHeight(20.0f)
+			.OnContextMenuOpening(this, &SSequencerLabelBrowser::HandleLabelTreeViewContextMenuOpening)
+			.OnGenerateRow(this, &SSequencerLabelBrowser::HandleLabelTreeViewGenerateRow)
+			.OnGetChildren(this, &SSequencerLabelBrowser::HandleLabelTreeViewGetChildren)
+			.OnSelectionChanged(this, &SSequencerLabelBrowser::HandleLabelTreeViewSelectionChanged)
+			.SelectionMode(ESelectionMode::Single)
+			.TreeItemsSource(&LabelList)
+	];
+
+	ReloadLabelList(true);
+}
+END_SLATE_FUNCTION_BUILD_OPTIMIZATION
+
+
+/* SSequencerLabelBrowser implementation
+ *****************************************************************************/
+
+void SSequencerLabelBrowser::ReloadLabelList(bool FullyReload)
+{
+	LabelList.Reset();
+
+	// test
+	LabelList.Add(MakeShareable(new FSequencerLabelTreeNode(TEXT(""))));
+
+	LabelTreeView->RequestTreeRefresh();
+}
+
+
+/* SSequencerLabelBrowser callbacks
+ *****************************************************************************/
+
+TSharedPtr<SWidget> SSequencerLabelBrowser::HandleLabelTreeViewContextMenuOpening()
+{
+	FMenuBuilder MenuBuilder(true /*bInShouldCloseWindowAfterMenuSelection*/, nullptr);
+	{
+		MenuBuilder.BeginSection("Customize", LOCTEXT("CustomizeContextMenuSectionName", "Customize"));
+		{
+			MenuBuilder.AddMenuEntry(
+				LOCTEXT("SetColorMenuEntryLabel", "Set Color"),
+				LOCTEXT("SetColorMenuEntryTip", "Set the background color of this label"),
+				FSlateIcon(),
+				FUIAction(
+					FExecuteAction::CreateSP(this, &SSequencerLabelBrowser::HandleSetColorMenuEntryExecute),
+					FCanExecuteAction::CreateSP(this, &SSequencerLabelBrowser::HandleSetColorMenuEntryCanExecute)
+				)
+			);
+		}
+		MenuBuilder.EndSection();
+
+		MenuBuilder.BeginSection("Edit", LOCTEXT("EditContextMenuSectionName", "Edit"));
+		{
+			MenuBuilder.AddMenuEntry(
+				LOCTEXT("RemoveLabelMenuEntryLabel", "Remove"),
+				LOCTEXT("RemoveLabelMenuEntryTip", "Remove this label from this list and all tracks"),
+				FSlateIcon(),
+				FUIAction(
+					FExecuteAction::CreateSP(this, &SSequencerLabelBrowser::HandleRemoveLabelMenuEntryExecute),
+					FCanExecuteAction::CreateSP(this, &SSequencerLabelBrowser::HandleRemoveLabelMenuEntryCanExecute)
+				)
+			);
+
+			MenuBuilder.AddMenuEntry(
+				LOCTEXT("RenameLabelMenuEntryLabel", "Rename"),
+				LOCTEXT("RenameLabelMenuEntryTip", "Change the name of this label"),
+				FSlateIcon(),
+				FUIAction(
+					FExecuteAction::CreateSP(this, &SSequencerLabelBrowser::HandleRenameLabelMenuEntryExecute),
+					FCanExecuteAction::CreateSP(this, &SSequencerLabelBrowser::HandleRenameLabelMenuEntryCanExecute)
+				)
+			);
+		}
+		MenuBuilder.EndSection();
+	}
+
+	return MenuBuilder.MakeWidget();
+}
+
+
+ TSharedRef<ITableRow> SSequencerLabelBrowser::HandleLabelTreeViewGenerateRow(TSharedPtr<FSequencerLabelTreeNode> Item, const TSharedRef<STableViewBase>& OwnerTable)
+{
+	return SNew(SSequencerLabelListRow, OwnerTable)
+		.Node(Item);
+}
+
+
+void SSequencerLabelBrowser::HandleLabelTreeViewGetChildren(TSharedPtr<FSequencerLabelTreeNode> Item, TArray<TSharedPtr<FSequencerLabelTreeNode>>& OutChildren)
+{
+	if (Item.IsValid())
+	{
+		OutChildren = Item->GetChildren();
+	}
+}
+
+
+void SSequencerLabelBrowser::HandleLabelTreeViewSelectionChanged(TSharedPtr<FSequencerLabelTreeNode> InItem, ESelectInfo::Type SelectInfo)
+{
+	OnSelectionChanged.ExecuteIfBound(
+		InItem.IsValid()
+			? InItem->GetLabel()
+			: FString(),
+		SelectInfo
+	);
+}
+
+
+void SSequencerLabelBrowser::HandleRemoveLabelMenuEntryExecute()
+{
+}
+
+
+bool SSequencerLabelBrowser::HandleRemoveLabelMenuEntryCanExecute() const
+{
+	return false;
+}
+
+
+void SSequencerLabelBrowser::HandleRenameLabelMenuEntryExecute()
+{
+}
+
+
+bool SSequencerLabelBrowser::HandleRenameLabelMenuEntryCanExecute() const
+{
+	return false;
+}
+
+
+void SSequencerLabelBrowser::HandleSetColorMenuEntryExecute()
+{
+}
+
+
+bool SSequencerLabelBrowser::HandleSetColorMenuEntryCanExecute() const
+{
+	return false;
+}
+
+
+#undef LOCTEXT_NAMESPACE
diff --git a/Engine/Source/Editor/Sequencer/Private/SSequencerLabelBrowser.h b/Engine/Source/Editor/Sequencer/Private/SSequencerLabelBrowser.h
new file mode 100644
--- /dev/null
+++ b/Engine/Source/Editor/Sequencer/Private/SSequencerLabelBrowser.h
@@ -0,0 +1,90 @@
+// Copyright 1998-2015 Epic Games, Inc. All Rights Reserved.
+
+#pragma once
+
+
+class FSequencerLabelTreeNode;
+
+
+/**
+ * Implements a widget for browsing sequencer track labels.
+ */
+class SSequencerLabelBrowser
+	: public SCompoundWidget
+{
+public:
+
+	typedef TSlateDelegates<FString>::FOnSelectionChanged FOnSelectionChanged;
+
+	SLATE_BEGIN_ARGS(SSequencerLabelBrowser)
+		: _OnSelectionChanged()
+	{ }
+
+		SLATE_EVENT(FOnSelectionChanged, OnSelectionChanged)
+	SLATE_END_ARGS()
+
+public:
+
+	/** Destructor. */
+	~SSequencerLabelBrowser();
+
+public:
+
+	/**
+	 * Construct this widget.
+	 *
+	 * @param InArgs The construction arguments.
+	 * @param InModel The view model to use.
+	 * @param InStyle The visual style to use for this widget.
+	 * @param InTracer The message tracer.
+	 */
+	void Construct(const FArguments& InArgs/*, const FMessagingDebuggerModelRef& InModel*/);
+
+protected:
+
+	/** Reloads the list of processes. */
+	void ReloadLabelList(bool FullyReload);
+
+private:
+
+	/** Callback for generating the label tree view's context menu. */
+	TSharedPtr<SWidget> HandleLabelTreeViewContextMenuOpening();
+
+	/** Callback for generating a row widget in the label tree view. */
+	TSharedRef<ITableRow> HandleLabelTreeViewGenerateRow(TSharedPtr<FSequencerLabelTreeNode> Item, const TSharedRef<STableViewBase>& OwnerTable);
+
+	/** Callback for getting the children of a node in the label tree view. */
+	void HandleLabelTreeViewGetChildren(TSharedPtr<FSequencerLabelTreeNode> Item, TArray<TSharedPtr<FSequencerLabelTreeNode>>& OutChildren);
+
+	/** Callback for when the label tree view selection changed. */
+	void HandleLabelTreeViewSelectionChanged(TSharedPtr<FSequencerLabelTreeNode> InItem, ESelectInfo::Type SelectInfo);
+
+	/** Callback for executing the 'Set Remove label' context menu entry. */
+	void HandleRemoveLabelMenuEntryExecute();
+
+	/** Callback for checking whether 'Remove label' context menu entry can execute. */
+	bool HandleRemoveLabelMenuEntryCanExecute() const;
+
+	/** Callback for executing the 'Rename label' context menu entry. */
+	void HandleRenameLabelMenuEntryExecute();
+
+	/** Callback for checking whether 'Rename' context menu entry can execute. */
+	bool HandleRenameLabelMenuEntryCanExecute() const;
+
+	/** Callback for executing the 'Set Color' context menu entry. */
+	void HandleSetColorMenuEntryExecute();
+
+	/** Callback for checking whether 'Set Color' context menu entry can execute. */
+	bool HandleSetColorMenuEntryCanExecute() const;
+
+private:
+
+	/** Holds the filtered list of processes running on the device. */
+	TArray<TSharedPtr<FSequencerLabelTreeNode>> LabelList;
+
+	/** Holds the label tree view. */
+	TSharedPtr<STreeView<TSharedPtr<FSequencerLabelTreeNode>>> LabelTreeView;
+
+	/** Delegate to invoke when the selected label changed. */
+	FOnSelectionChanged OnSelectionChanged;
+};
diff --git a/Engine/Source/Editor/Sequencer/Private/SSequencerLabelEditor.cpp b/Engine/Source/Editor/Sequencer/Private/SSequencerLabelEditor.cpp
new file mode 100644
--- /dev/null
+++ b/Engine/Source/Editor/Sequencer/Private/SSequencerLabelEditor.cpp
@@ -0,0 +1,184 @@
+// Copyright 1998-2015 Epic Games, Inc. All Rights Reserved.
+
+#include "SequencerPrivatePCH.h"
+#include "MovieScene.h"
+#include "SSequencerLabelEditorListRow.h"
+
+
+#define LOCTEXT_NAMESPACE "SSequencerLabelEditor"
+
+
+/* SSequencerLabelEditor interface
+ *****************************************************************************/
+
+void SSequencerLabelEditor::Construct(const FArguments& InArgs, UMovieScene* InMovieScene, const FGuid& InObjectId)
+{
+	MovieScene = InMovieScene;
+	ObjectId = InObjectId;
+
+	ChildSlot
+	[
+		SNew(SVerticalBox)
+
+		+ SVerticalBox::Slot()
+			.AutoHeight()
+			[
+				SNew(STextBlock)
+					.Text(LOCTEXT("LabelAs", "Label as:"))
+			]
+
+		+ SVerticalBox::Slot()
+			.AutoHeight()
+			.Padding(0.0f, 2.0f)
+			[
+				SNew(SHorizontalBox)
+
+				+ SHorizontalBox::Slot()
+					.FillWidth(1.0f)
+					[
+						SAssignNew(FilterBox, SEditableTextBox)
+							.MinDesiredWidth(144.0f)
+							.OnTextChanged(this, &SSequencerLabelEditor::HandleFilterBoxTextChanged)
+					]
+
+				+ SHorizontalBox::Slot()
+					.AutoWidth()
+					.Padding(4.0f, 0.0f, 0.0f, 0.0f)
+					[
+						SNew(SButton)
+							.IsEnabled(this, &SSequencerLabelEditor::HandleCreateNewLabelButtonIsEnabled)
+							.OnClicked(this, &SSequencerLabelEditor::HandleCreateNewLabelButtonClicked)
+							.Content()
+							[
+								SNew(STextBlock)
+									.Text(LOCTEXT("CreateNewLabelButton", "Create New"))
+							]
+					]
+			]
+
+		+ SVerticalBox::Slot()
+			.AutoHeight()
+			.Padding(0.0f, 4.0f, 0.0f, 0.0f)
+			[
+				SAssignNew(LabelListView, SListView<TSharedPtr<FString>>)
+					.ItemHeight(20.0f)
+					.ListItemsSource(&LabelList)
+					.OnGenerateRow(this, &SSequencerLabelEditor::HandleLabelListViewGenerateRow)
+					.SelectionMode(ESelectionMode::None)
+			]
+	];
+
+	ReloadLabelList(true);
+}
+
+
+/* SSequencerLabelEditor implementation
+ *****************************************************************************/
+
+void SSequencerLabelEditor::ReloadLabelList(bool FullyReload)
+{
+	// reload label list
+	if (FullyReload)
+	{
+		AvailableLabels.Reset();
+		MovieScene->GetAllLabels(AvailableLabels);
+	}
+
+	// filter label list
+	TArray<FString> Filters;
+	FilterBox->GetText().ToString().ParseIntoArray(Filters, TEXT(" "));
+
+	LabelList.Reset();
+
+	for (const auto& Label : AvailableLabels)
+	{
+		bool Matched = true;
+
+		for (const auto& Filter : Filters)
+		{
+			if (!Label.Contains(Filter))
+			{
+				Matched = false;
+				break;
+			}
+		}
+
+		if (Matched)
+		{
+			LabelList.Add(MakeShareable(new FString(Label)));
+		}
+	}
+
+	// refresh list view
+	LabelListView->RequestListRefresh();
+}
+
+
+/* SSequencerLabelEditor callbacks
+ *****************************************************************************/
+
+FReply SSequencerLabelEditor::HandleCreateNewLabelButtonClicked()
+{
+	MovieScene->GetObjectLabels(ObjectId).Strings.AddUnique(FilterBox->GetText().ToString());
+	FilterBox->SetText(FText::GetEmpty());
+
+	return FReply::Handled();
+}
+
+
+bool SSequencerLabelEditor::HandleCreateNewLabelButtonIsEnabled() const
+{
+	FString FilterString = FilterBox->GetText().ToString();
+	FilterString.Trim();
+
+	return !FilterString.IsEmpty() && !MovieScene->LabelExists(FilterString);
+}
+
+
+void SSequencerLabelEditor::HandleFilterBoxTextChanged(const FText& NewText)
+{
+	ReloadLabelList(false);
+}
+
+
+TSharedRef<ITableRow> SSequencerLabelEditor::HandleLabelListViewGenerateRow(TSharedPtr<FString> Label, const TSharedRef<STableViewBase>& OwnerTable)
+{
+	return SNew(SSequencerLabelEditorListRow, OwnerTable)
+		.HighlightText(this, &SSequencerLabelEditor::HandleLabelListViewRowHighlightText)
+		.IsChecked(this, &SSequencerLabelEditor::HandleLabelListViewRowIsChecked, Label)
+		.Label(Label)
+		.OnCheckStateChanged(this, &SSequencerLabelEditor::HandleLabelListViewRowCheckedStateChanged, Label);
+}
+
+
+void SSequencerLabelEditor::HandleLabelListViewRowCheckedStateChanged(ECheckBoxState State, TSharedPtr<FString> Label)
+{
+	FMovieSceneTrackLabels& Labels = MovieScene->GetObjectLabels(ObjectId);
+
+	if (State == ECheckBoxState::Checked)
+	{
+		Labels.Strings.AddUnique(*Label);
+	}
+	else
+	{
+		Labels.Strings.Remove(*Label);
+	}
+}
+
+
+FText SSequencerLabelEditor::HandleLabelListViewRowHighlightText() const
+{
+	return FilterBox->GetText();
+}
+
+
+ECheckBoxState SSequencerLabelEditor::HandleLabelListViewRowIsChecked(TSharedPtr<FString> Label) const
+{
+	const FMovieSceneTrackLabels& Labels = MovieScene->GetObjectLabels(ObjectId);
+	return Labels.Strings.Contains(*Label)
+		? ECheckBoxState::Checked
+		: ECheckBoxState::Unchecked;
+}
+
+
+#undef LOCTEXT_NAMESPACE
diff --git a/Engine/Source/Editor/Sequencer/Private/SSequencerLabelEditor.h b/Engine/Source/Editor/Sequencer/Private/SSequencerLabelEditor.h
new file mode 100644
--- /dev/null
+++ b/Engine/Source/Editor/Sequencer/Private/SSequencerLabelEditor.h
@@ -0,0 +1,76 @@
+// Copyright 1998-2015 Epic Games, Inc. All Rights Reserved.
+
+#pragma once
+
+
+class SSequencerLabelEditor
+	: public SCompoundWidget
+{
+public:
+
+	SLATE_BEGIN_ARGS(SSequencerLabelEditor)
+	{ }
+	SLATE_END_ARGS()
+
+public:
+
+	/**
+	 * Construct this widget.
+	 *
+	 * @param InArgs The construction arguments.
+	 * @param InObjectId The unique identifier of the object whose labels are being edited.
+	 */
+	void Construct(const FArguments& InArgs, UMovieScene* InMovieScene, const FGuid& InObjectId);
+
+protected:
+
+	/**
+	 * Reload the list of track labels.
+	 *
+	 * @param FullyReload Whether to fully reload entries, or only re-apply filtering.
+	 */
+	void ReloadLabelList(bool FullyReload);
+
+private:
+
+	/** Callback for clicking the 'Create new label' button. */
+	FReply HandleCreateNewLabelButtonClicked();
+
+	/** Callback for getting the enabled state of the 'Create new label' button. */
+	bool HandleCreateNewLabelButtonIsEnabled() const;
+
+	/** Callback for text changes in the filter text box. */
+	void HandleFilterBoxTextChanged(const FText& NewText);
+
+	/** Callback for generating a row widget in the device service list view. */
+	TSharedRef<ITableRow> HandleLabelListViewGenerateRow(TSharedPtr<FString> Label, const TSharedRef<STableViewBase>& OwnerTable);
+
+	/** Callback for changing the state of a list view row's check box. */
+	void HandleLabelListViewRowCheckedStateChanged(ECheckBoxState State, TSharedPtr<FString> Label);
+
+	/** Callback for getting the highlight text in session rows. */
+	FText HandleLabelListViewRowHighlightText() const;
+
+	/** Callback for determining whether the check box of a list view row is checked. */
+	ECheckBoxState HandleLabelListViewRowIsChecked(TSharedPtr<FString> Label) const;
+
+private:
+
+	/** The list of available track labels. */
+	TArray<FString> AvailableLabels;
+
+	/** The label filter text box. */
+	TSharedPtr<SEditableTextBox> FilterBox;
+
+	/** The filtered list of track labels. */
+	TArray<TSharedPtr<FString>> LabelList;
+
+	/** Holds the list view for filtered track labels. */
+	TSharedPtr<SListView<TSharedPtr<FString>>> LabelListView;
+
+	/** The movie scene being edited. */
+	UMovieScene* MovieScene;
+
+	/** The identifier of the object being edited. */
+	FGuid ObjectId;
+};
diff --git a/Engine/Source/Editor/Sequencer/Private/SSequencerLabelEditorListRow.h b/Engine/Source/Editor/Sequencer/Private/SSequencerLabelEditorListRow.h
new file mode 100644
--- /dev/null
+++ b/Engine/Source/Editor/Sequencer/Private/SSequencerLabelEditorListRow.h
@@ -0,0 +1,65 @@
+// Copyright 1998-2015 Epic Games, Inc. All Rights Reserved.
+
+#pragma once
+
+
+#define LOCTEXT_NAMESPACE "SSequencerLabelListRow"
+
+
+/**
+ * Implements a row widget for the label browser tree view.
+ */
+class SSequencerLabelEditorListRow
+	: public STableRow<TSharedPtr<FString>>
+{
+public:
+
+	SLATE_BEGIN_ARGS(SSequencerLabelEditorListRow) { }
+		SLATE_ATTRIBUTE(FText, HighlightText)
+		SLATE_ATTRIBUTE(ECheckBoxState, IsChecked)
+		SLATE_ARGUMENT(TSharedPtr<FString>, Label)
+		SLATE_EVENT(FOnCheckStateChanged, OnCheckStateChanged)
+	SLATE_END_ARGS()
+
+public:
+
+	/**
+	 * Constructs the widget.
+	 *
+	 * @param InArgs The arguments.
+	 */
+	void Construct(const FArguments& InArgs, const TSharedRef<STableViewBase>& InOwnerTableView)
+	{
+		ChildSlot
+		[
+			SNew(SHorizontalBox)
+
+			+ SHorizontalBox::Slot()
+				.AutoWidth()
+				.Padding(2.0f, 0.0f)
+				[
+					SNew(SCheckBox)
+						.IsChecked(InArgs._IsChecked)
+						.OnCheckStateChanged(InArgs._OnCheckStateChanged)
+				]
+
+			+ SHorizontalBox::Slot()
+				.FillWidth(1.0f)
+				[
+					SNew(STextBlock)
+						.HighlightText(InArgs._HighlightText)
+						.Text(FText::FromString(*InArgs._Label))
+				]
+		];
+
+		STableRow<TSharedPtr<FString>>::ConstructInternal(
+			STableRow<TSharedPtr<FString>>::FArguments()
+				.ShowSelection(false)
+				.Style(FEditorStyle::Get(), "DetailsView.TreeView.TableRow"),
+			InOwnerTableView
+		);
+	}
+};
+
+
+#undef LOCTEXT_NAMESPACE
diff --git a/Engine/Source/Editor/Sequencer/Private/SSequencerLabelListRow.h b/Engine/Source/Editor/Sequencer/Private/SSequencerLabelListRow.h
new file mode 100644
--- /dev/null
+++ b/Engine/Source/Editor/Sequencer/Private/SSequencerLabelListRow.h
@@ -0,0 +1,173 @@
+// Copyright 1998-2015 Epic Games, Inc. All Rights Reserved.
+
+#pragma once
+
+
+class FSequencerLabelTreeNode
+{
+public:
+
+	/**
+	 * Creates and initializes a new instance.
+	 *
+	 * @param InLabel The node's label.
+	 */
+	FSequencerLabelTreeNode(const FString& InLabel)
+		: Label(InLabel)
+	{ }
+
+public:
+
+	/**
+	 * Adds a child label node to this node.
+	 *
+	 * @param The child node to add.
+	 */
+	void AddChild(const TSharedPtr<FSequencerLabelTreeNode>& Child)
+	{
+		Children.Add(Child);
+	}
+
+	/** Clears the collection of child nodes. */
+	void ClearChildren()
+	{
+		Children.Reset();
+	}
+
+	/**
+	 * Gets the child nodes.
+	 *
+	 * @return Child nodes.
+	 */
+	const TArray<TSharedPtr<FSequencerLabelTreeNode>>& GetChildren()
+	{
+		return Children;
+	}
+
+	/**
+	 * Gets the node's process information.
+	 *
+	 * @return The process information.
+	 */
+	const FString& GetLabel() const
+	{
+		return Label;
+	}
+
+private:
+
+	/** Holds the child label nodes. */
+	TArray<TSharedPtr<FSequencerLabelTreeNode>> Children;
+
+	/** Holds the label. */
+	FString Label;
+};
+
+
+#define LOCTEXT_NAMESPACE "SSequencerLabelListRow"
+
+/**
+ * Implements a row widget for the label browser tree view.
+ */
+class SSequencerLabelListRow
+	: public STableRow<TSharedPtr<FSequencerLabelTreeNode>>
+{
+public:
+
+	SLATE_BEGIN_ARGS(SSequencerLabelListRow) { }
+		SLATE_ARGUMENT(TSharedPtr<FSequencerLabelTreeNode>, Node)
+	SLATE_END_ARGS()
+
+public:
+
+	/**
+	 * Constructs the widget.
+	 *
+	 * @param InArgs The arguments.
+	 */
+	void Construct( const FArguments& InArgs, const TSharedRef<STableViewBase>& InOwnerTableView )
+	{
+		Node = InArgs._Node;
+
+		ChildSlot
+			.Padding(0.0f, 2.0f, 0.0f, 0.0f)
+			[
+				SNew(SHorizontalBox)
+
+				// folder icon
+				+ SHorizontalBox::Slot()
+					.AutoWidth()
+					.Padding(2.0f, 2.0f)
+					.VAlign(VAlign_Center)
+					[
+						SNew(SImage) 
+							.Image(this, &SSequencerLabelListRow::HandleFolderIconImage)
+							.ColorAndOpacity(this, &SSequencerLabelListRow::HandleFolderIconColor)
+					]
+
+				// folder name
+				+ SHorizontalBox::Slot()
+					.FillWidth(1.0f)
+					.Padding(0.0f, 2.0f)
+					.VAlign(VAlign_Center)
+					[
+						SNew(STextBlock)
+							.Text(
+								Node->GetLabel().IsEmpty()
+									? LOCTEXT("AllTracksLabel", "All Tracks")
+									: FText::FromString(Node->GetLabel())
+							)
+					]
+
+				// edit icon
+				+ SHorizontalBox::Slot()
+					.AutoWidth()
+					.VAlign(VAlign_Center)
+					[
+						SNew(SImage)
+							.Image(FCoreStyle::Get().GetBrush(TEXT("EditableLabel.EditIcon")))
+							.Visibility(this, &SSequencerLabelListRow::HandleEditIconVisibility)
+					]
+			];
+
+		STableRow<TSharedPtr<FSequencerLabelTreeNode>>::ConstructInternal(
+			STableRow<TSharedPtr<FSequencerLabelTreeNode>>::FArguments()
+				.ShowSelection(false)
+				.Style(FEditorStyle::Get(), "DetailsView.TreeView.TableRow"),
+			InOwnerTableView
+		);
+	}
+
+private:
+
+	EVisibility HandleEditIconVisibility() const
+	{
+		return IsHovered()
+			? EVisibility::Visible
+			: EVisibility::Collapsed;
+	}
+
+	const FSlateBrush* HandleFolderIconImage() const
+	{
+		static const FSlateBrush* FolderOpenBrush = FEditorStyle::GetBrush("ContentBrowser.AssetTreeFolderOpen");
+		static const FSlateBrush* FolderClosedBrush = FEditorStyle::GetBrush("ContentBrowser.AssetTreeFolderClosed");
+
+		return IsItemExpanded()
+			? FolderOpenBrush
+			: FolderClosedBrush;
+	}
+
+	FSlateColor HandleFolderIconColor() const
+	{
+		// TODO sequencer: gmp: allow folder color customization
+		return FLinearColor::Gray;
+	}
+
+private:
+
+	/** Holds the label node. */
+	TSharedPtr<FSequencerLabelTreeNode> Node;
+};
+
+
+#undef LOCTEXT_NAMESPACE
diff --git a/Engine/Source/Editor/Sequencer/Private/SSequencerSection.cpp b/Engine/Source/Editor/Sequencer/Private/SSequencerSection.cpp
--- a/Engine/Source/Editor/Sequencer/Private/SSequencerSection.cpp
+++ b/Engine/Source/Editor/Sequencer/Private/SSequencerSection.cpp
@@ -211,6 +211,9 @@ FSequencer& SSequencerSection::GetSequencer() const
 int32 SSequencerSection::OnPaint( const FPaintArgs& Args, const FGeometry& AllottedGeometry, const FSlateRect& MyClippingRect, FSlateWindowElementList& OutDrawElements, int32 LayerId, const FWidgetStyle& InWidgetStyle, bool bParentEnabled ) const
 {
 	UMovieSceneSection& SectionObject = *SectionInterface->GetSectionObject();
+
+	LayerId += SectionObject.GetOverlapPriority();
+	
 	const bool bEnabled = bParentEnabled && SectionObject.IsActive();
 	const bool bLocked = SectionObject.IsLocked();
 	const ESlateDrawEffect::Type DrawEffects = bEnabled ? ESlateDrawEffect::None : ESlateDrawEffect::DisabledEffect;
diff --git a/Engine/Source/Editor/Sequencer/Private/SSequencerSectionAreaView.cpp b/Engine/Source/Editor/Sequencer/Private/SSequencerSectionAreaView.cpp
--- a/Engine/Source/Editor/Sequencer/Private/SSequencerSectionAreaView.cpp
+++ b/Engine/Source/Editor/Sequencer/Private/SSequencerSectionAreaView.cpp
@@ -120,14 +120,29 @@ int32 SSequencerSectionAreaView::OnPaint( const FPaintArgs& Args, const FGeometr
 	FArrangedChildren ArrangedChildren(EVisibility::Visible);
 	ArrangeChildren(AllottedGeometry, ArrangedChildren);
 
+	int32 MaxLayerId = LayerId;
 	for (int32 ChildIndex = 0; ChildIndex < ArrangedChildren.Num(); ++ChildIndex)
 	{
 		FArrangedWidget& CurWidget = ArrangedChildren[ChildIndex];
 		FSlateRect ChildClipRect = MyClippingRect.IntersectionWith( CurWidget.Geometry.GetClippingRect() );
-		const int32 CurWidgetsMaxLayerId = CurWidget.Widget->Paint( Args.WithNewParent(this), CurWidget.Geometry, ChildClipRect, OutDrawElements, LayerId+1, InWidgetStyle, ShouldBeEnabled( bParentEnabled ) );
+		const int32 ThisLayerId = CurWidget.Widget->Paint( Args.WithNewParent(this), CurWidget.Geometry, ChildClipRect, OutDrawElements, LayerId, InWidgetStyle, ShouldBeEnabled( bParentEnabled ) );
+		MaxLayerId = FMath::Max(ThisLayerId, MaxLayerId);
 	}
 
-	return LayerId+1;
+	return MaxLayerId + 1;
+}
+
+void SSequencerSectionAreaView::Tick(const FGeometry& AllottedGeometry, const double InCurrentTime, const float InDeltaTime)
+{
+	if (Children.Num())
+	{
+		StableSort(&Children[0], Children.Num(), [](const TSharedRef<SSequencerSection>& A, const TSharedRef<SSequencerSection>& B){
+			UMovieSceneSection* SectionA = A->GetSectionInterface()->GetSectionObject();
+			UMovieSceneSection* SectionB = B->GetSectionInterface()->GetSectionObject();
+
+			return SectionA && SectionB && SectionA->GetOverlapPriority() < SectionB->GetOverlapPriority();
+		});
+	}
 }
 
 void SSequencerSectionAreaView::OnArrangeChildren( const FGeometry& AllottedGeometry, FArrangedChildren& ArrangedChildren ) const
diff --git a/Engine/Source/Editor/Sequencer/Private/SSequencerSectionAreaView.h b/Engine/Source/Editor/Sequencer/Private/SSequencerSectionAreaView.h
--- a/Engine/Source/Editor/Sequencer/Private/SSequencerSectionAreaView.h
+++ b/Engine/Source/Editor/Sequencer/Private/SSequencerSectionAreaView.h
@@ -22,6 +22,7 @@ public:
 private:
 	/** SWidget Interface */
 	virtual int32 OnPaint( const FPaintArgs& Args, const FGeometry& AllottedGeometry, const FSlateRect& MyClippingRect, FSlateWindowElementList& OutDrawElements, int32 LayerId, const FWidgetStyle& InWidgetStyle, bool bParentEnabled ) const override;
+	virtual void Tick( const FGeometry& AllottedGeometry, const double InCurrentTime, const float InDeltaTime ) override;
 
 	/** SPanel Interface */
 	virtual void OnArrangeChildren( const FGeometry& AllottedGeometry, FArrangedChildren& ArrangedChildren ) const override;
diff --git a/Engine/Source/Editor/Sequencer/Private/SSequencerTreeViewBox.h b/Engine/Source/Editor/Sequencer/Private/SSequencerTreeViewBox.h
--- a/Engine/Source/Editor/Sequencer/Private/SSequencerTreeViewBox.h
+++ b/Engine/Source/Editor/Sequencer/Private/SSequencerTreeViewBox.h
@@ -14,9 +14,10 @@ class SSequencerTreeViewBox
 {
 public:
 
-	void Construct(const FArguments& InArgs, TSharedRef<FSequencer> InSequencer)
+	void Construct(const FArguments& InArgs, TSharedRef<FSequencer> InSequencer, TSharedPtr<SSequencer> InSequencerWidget)
 	{
 		Sequencer = InSequencer;
+		SequencerWidget = InSequencerWidget;
 		SBox::Construct(InArgs);
 	}
 	
@@ -28,7 +29,17 @@ public:
 
 			if (Selection.GetSelectedOutlinerNodes().Num())
 			{
+				if (SequencerWidget.Pin().IsValid())
+				{
+					SequencerWidget.Pin()->SetUserIsSelecting(true);
+				}
+
 				Selection.EmptySelectedOutlinerNodes();
+
+				if (SequencerWidget.Pin().IsValid())
+				{
+					SequencerWidget.Pin()->SetUserIsSelecting(false);
+				}
 				return FReply::Handled();
 			}
 		}
@@ -40,4 +51,7 @@ private:
 
 	/** Weak pointer to the sequencer */
 	TWeakPtr<FSequencer> Sequencer;
+
+	/** Weak pointer to the sequencer widget */
+	TWeakPtr<SSequencer> SequencerWidget;
 };
diff --git a/Engine/Source/Editor/Sequencer/Private/Sequencer.cpp b/Engine/Source/Editor/Sequencer/Private/Sequencer.cpp
--- a/Engine/Source/Editor/Sequencer/Private/Sequencer.cpp
+++ b/Engine/Source/Editor/Sequencer/Private/Sequencer.cpp
@@ -44,12 +44,15 @@
 #include "Tracks/MovieScene3DTransformTrack.h"
 #include "LevelEditor.h"
 #include "IMenu.h"
+#include "NotificationManager.h"
+#include "SNotificationList.h"
 #include "MovieSceneClipboard.h"
 #include "SequencerClipboardReconciler.h"
 #include "STextEntryPopup.h"
 #include "SequencerHotspots.h"
 #include "MovieSceneCaptureDialogModule.h"
 #include "AutomatedLevelSequenceCapture.h"
+#include "MovieSceneCommonHelpers.h"
 
 #include "SceneOutlinerModule.h"
 #include "SceneOutlinerPublicTypes.h"
@@ -75,6 +78,16 @@ void FSequencer::InitSequencer(const FSequencerInitParams& InitParams, const TSh
 	{
 		bIsEditingWithinLevelEditor = InitParams.bEditWithinLevelEditor;
 		
+		if (InitParams.SpawnRegister.IsValid())
+		{
+			SpawnRegister = InitParams.SpawnRegister;
+		}
+		else
+		{
+			// Spawnables not supported
+			SpawnRegister = MakeShareable(new FNullMovieSceneSpawnRegister);
+		}
+
 		// If this sequencer edits the level, close out the existing active sequencer and mark this as the active sequencer.
 		if (bIsEditingWithinLevelEditor)
 		{
@@ -280,13 +293,23 @@ void FSequencer::Tick(float InDeltaTime)
 
 UMovieSceneSequence* FSequencer::GetRootMovieSceneSequence() const
 {
-	return SequenceInstanceStack[0]->GetSequence();
+	if (SequenceInstanceStack.Num())
+	{
+		return SequenceInstanceStack[0]->GetSequence();
+	}
+
+	return nullptr;
 }
 
 UMovieSceneSequence* FSequencer::GetFocusedMovieSceneSequence() const
 {
 	// the last item is the focused movie scene
-	return SequenceInstanceStack.Top()->GetSequence();
+	if (SequenceInstanceStack.Num())
+	{
+		return SequenceInstanceStack.Top()->GetSequence();
+	}
+
+	return nullptr;
 }
 
 void FSequencer::ResetToNewRootSequence(UMovieSceneSequence& NewSequence)
@@ -343,6 +366,116 @@ void FSequencer::FocusSequenceInstance(TSharedRef<FMovieSceneSequenceInstance> S
 	SequencerWidget->UpdateBreadcrumbs();
 }
 
+FGuid FSequencer::CreateBinding(UObject& InObject, const FString& InName)
+{
+	UMovieSceneSequence* OwnerSequence = GetFocusedMovieSceneSequence();
+	UMovieScene* OwnerMovieScene = OwnerSequence->GetMovieScene();
+		
+	const FGuid PossessableGuid = OwnerMovieScene->AddPossessable(InName, InObject.GetClass());
+
+	UObject* BindingContext = OwnerSequence->GetParentObject(&InObject);
+	if (BindingContext)
+	{
+		// Ensure we have possessed the outer object, if necessary
+		FGuid ParentGuid = GetHandleToObject(BindingContext);
+		
+		// Set up parent/child guids for possessables within spawnables
+		if (ParentGuid.IsValid())
+		{
+			FMovieScenePossessable* ChildPossessable = OwnerMovieScene->FindPossessable(PossessableGuid);
+			if (ensure(ChildPossessable))
+			{
+				ChildPossessable->SetParent(ParentGuid);
+			}
+
+			FMovieSceneSpawnable* ParentSpawnable = OwnerMovieScene->FindSpawnable(ParentGuid);
+			if (ParentSpawnable)
+			{
+				ParentSpawnable->AddChildPossessable(PossessableGuid);
+			}
+		}
+	}
+	else
+	{
+		BindingContext = GetPlaybackContext();
+	}
+
+	OwnerSequence->BindPossessableObject(PossessableGuid, InObject, BindingContext);
+
+	return PossessableGuid;
+}
+
+UObject* FSequencer::GetPlaybackContext() const
+{
+	if (bIsEditingWithinLevelEditor)
+	{
+		return GWorld;
+	}
+	return nullptr;
+}
+
+void FSequencer::GetAllKeyedProperties(UObject& Object, TSet<UProperty*>& OutProperties)
+{
+	GetAllKeyedPropertiesForInstance(Object, *GetRootMovieSceneSequenceInstance(), OutProperties);
+}
+
+void FSequencer::GetAllKeyedPropertiesForInstance(UObject& Object, FMovieSceneSequenceInstance& Instance, TSet<UProperty*>& OutProperties)
+{
+	GetKeyedProperties(Object, Instance, OutProperties);
+
+	UMovieSceneSubTrack* SubTrack = Instance.GetSequence()->GetMovieScene()->FindMasterTrack<UMovieSceneSubTrack>();
+	if (SubTrack)
+	{
+		// Iterate sub movie scenes
+		for (UMovieSceneSection* SubSection : SubTrack->GetAllSections())
+		{
+			// Find the instance for the section
+			TSharedRef<FMovieSceneSequenceInstance>* SubInstance = SequenceInstanceBySection.Find(SubSection);
+			GetAllKeyedPropertiesForInstance(Object, SubInstance->Get(), OutProperties);
+		}
+	}
+}
+
+void FSequencer::GetKeyedProperties(UObject& Object, FMovieSceneSequenceInstance& Instance, TSet<UProperty*>& OutProperties)
+{
+	FGuid BindingId = Instance.FindObjectId(Object);
+
+	const FMovieSceneBinding* Binding = BindingId.IsValid() ? Instance.GetSequence()->GetMovieScene()->GetBindings().FindByPredicate([&](const FMovieSceneBinding& InBinding){
+		return BindingId == InBinding.GetObjectGuid();
+	}) : nullptr;
+
+	if (!Binding)
+	{
+		return;
+	}
+
+	// Add any properties for this object
+	for (UMovieSceneTrack* Track : Binding->GetTracks())
+	{
+		if (UMovieScene3DTransformTrack* TransformTrack = Cast<UMovieScene3DTransformTrack>(Track))
+		{
+			// Add transform properties manually. This is technically not correct, since it's the scene component that's being keyed, not the actor,
+			// But since the transform track is always stored on an actor binding, we have to add them this way
+			AActor* Actor = Cast<AActor>(&Object);
+			if (Actor)
+			{
+				UClass* Class = USceneComponent::StaticClass();
+				OutProperties.Add(Class->FindPropertyByName(GET_MEMBER_NAME_CHECKED(USceneComponent, RelativeLocation)));
+				OutProperties.Add(Class->FindPropertyByName(GET_MEMBER_NAME_CHECKED(USceneComponent, RelativeRotation)));
+				OutProperties.Add(Class->FindPropertyByName(GET_MEMBER_NAME_CHECKED(USceneComponent, RelativeScale3D)));
+			}
+		}
+		else if (UMovieScenePropertyTrack* PropertyTrack = Cast<UMovieScenePropertyTrack>(Track))
+		{
+			FTrackInstancePropertyBindings InstanceBinding(PropertyTrack->GetPropertyName(), PropertyTrack->GetPropertyPath());
+			UProperty* Property = InstanceBinding.GetProperty(&Object);
+			if (Property)
+			{
+				OutProperties.Add(Property);
+			}
+		}
+	}
+}
 
 TSharedRef<FMovieSceneSequenceInstance> FSequencer::GetSequenceInstanceForSection(UMovieSceneSection& Section) const
 {
@@ -535,24 +668,20 @@ void FSequencer::NotifyMapChanged( class UWorld* NewWorld, EMapChangeType MapCha
 
 void FSequencer::OnActorsDropped( const TArray<TWeakObjectPtr<AActor> >& Actors )
 {
+	// @todo sequencer: Undo doesn't seem to be working at all
+	const FScopedTransaction Transaction(LOCTEXT("UndoPossessingObject", "Possess Object in Sequencer"));
+	GetFocusedMovieSceneSequence()->Modify();
+
 	bool bPossessableAdded = false;
-	UMovieSceneSequence* OwnerSequence = GetFocusedMovieSceneSequence();
-	UMovieScene* OwnerMovieScene = OwnerSequence->GetMovieScene();
 	for (TWeakObjectPtr<AActor> WeakActor : Actors)
 	{
-		AActor* Actor = WeakActor.Get();
-
-		if (Actor != nullptr)
-		{		
-			// @todo sequencer: Undo doesn't seem to be working at all
-			const FScopedTransaction Transaction(LOCTEXT("UndoPossessingObject", "Possess Object with MovieScene"));
-			
-			// possess the object!
-			OwnerSequence->Modify();
-				
-			const FGuid PossessableGuid = OwnerMovieScene->AddPossessable(Actor->GetActorLabel(), Actor->GetClass());
-
-			OwnerSequence->BindPossessableObject(PossessableGuid, *Actor);
+		if (AActor* Actor = WeakActor.Get())
+		{
+			FGuid ExistingGuid = SequenceInstanceStack.Top()->FindObjectId(*Actor);
+			if (!ExistingGuid.IsValid())
+			{
+				CreateBinding(*Actor, Actor->GetActorLabel());
+			}
 			bPossessableAdded = true;
 		}
 	}
@@ -702,6 +831,15 @@ void FSequencer::SetGlobalTimeDirectly( float NewTime )
 
 	SequenceInstanceStack.Top()->Update(NewTime, LastTime, *this);
 
+	// Ensure any relevant spawned objects are cleaned up if we're playing back the master sequence
+	if (SequenceInstanceStack.Num() == 1)
+	{
+		if (!GetPlaybackRange().Contains(NewTime))
+		{
+			SpawnRegister->DestroyAllOwnedObjects(*this);
+		}
+	}
+
 	// If realtime is off, this needs to be called to update the pivot location when scrubbing.
 	GUnrealEd->UpdatePivotLocationForSelection();
 	GEditor->RedrawLevelEditingViewports();
@@ -796,10 +934,11 @@ FGuid FSequencer::GetHandleToObject( UObject* Object, bool bCreateHandleIfMissin
 	UMovieSceneSequence* FocusedMovieSceneSequence = FocusedMovieSceneSequenceInstance->GetSequence();
 
 	UMovieScene* FocusedMovieScene = FocusedMovieSceneSequence->GetMovieScene();
-	// The movie scene can resolve persistent, possessable objects
-	FGuid ObjectGuid = FocusedMovieSceneSequence->FindObjectId(*Object);
+	
+	// Attempt to resolve the object through the movie scene instance first, 
+	FGuid ObjectGuid = FocusedMovieSceneSequenceInstance->FindObjectId(*Object);
 
-	// Check here for spawnable otherwise spawnables get recreated as possessables, which doesn't make sens
+	// Check here for spawnable otherwise spawnables get recreated as possessables, which doesn't make sense
 	FMovieSceneSpawnable* Spawnable = FocusedMovieScene->FindSpawnable(ObjectGuid);
 	if (Spawnable)
 	{
@@ -821,27 +960,10 @@ FGuid FSequencer::GetHandleToObject( UObject* Object, bool bCreateHandleIfMissin
 	// Note: Only possessed actors can be added like this
 	if (!ObjectGuid.IsValid() && FocusedMovieSceneSequence->CanPossessObject(*Object) && bCreateHandleIfMissing)
 	{
-		// @todo sequencer: Undo doesn't seem to be working at all
-		const FScopedTransaction Transaction(LOCTEXT("UndoPossessingObject", "Possess Object with MovieScene"));
-
-		// Possess the object!
-		{
-			// Create a new possessable
-			FocusedMovieScene->Modify();
-
-			AActor* PossessedActor = Cast<AActor>(Object);
-
-			ObjectGuid = FocusedMovieScene->AddPossessable(PossessedActor != nullptr ? PossessedActor->GetActorLabel() : Object->GetName(), Object->GetClass());
-			FocusedMovieSceneSequence->BindPossessableObject(ObjectGuid, *Object);
-			bPossessableAdded = true;
-		}
+		AActor* PossessedActor = Cast<AActor>(Object);
 
-		// If we're adding a possessable, generate handles to it's parent objects too.
-		UObject* ParentObject = FocusedMovieSceneSequence->GetParentObject(Object);
-		if ( ParentObject != nullptr )
-		{
-			GetHandleToObject(ParentObject);
-		}
+		ObjectGuid = CreateBinding(*Object, PossessedActor != nullptr ? PossessedActor->GetActorLabel() : Object->GetName());
+		bPossessableAdded = true;
 	}
 
 	if (bPossessableAdded)
@@ -859,14 +981,10 @@ ISequencerObjectChangeListener& FSequencer::GetObjectChangeListener()
 
 void FSequencer::GetRuntimeObjects( TSharedRef<FMovieSceneSequenceInstance> MovieSceneInstance, const FGuid& ObjectHandle, TArray< UObject*>& OutObjects ) const
 {
-	if (MovieSceneInstance->GetSequence() != nullptr)
+	UObject* FoundObject = MovieSceneInstance->FindObject(ObjectHandle, *this);
+	if (FoundObject)
 	{
-		UObject* FoundObject = MovieSceneInstance->GetSequence()->FindObject(ObjectHandle);
-
-		if (FoundObject != nullptr)
-		{
-			OutObjects.Add(FoundObject);
-		}
+		OutObjects.Add(FoundObject);
 	}
 }
 
@@ -959,7 +1077,8 @@ void FSequencer::AddOrUpdateMovieSceneInstance( UMovieSceneSection& MovieSceneSe
 void FSequencer::RemoveMovieSceneInstance( UMovieSceneSection& MovieSceneSection, TSharedRef<FMovieSceneSequenceInstance> InstanceToRemove )
 {
 	UMovieSceneSequence* Sequence = InstanceToRemove->GetSequence();
-	Sequence->DestroyAllSpawnedObjects();
+
+	SpawnRegister->DestroyObjectsSpawnedByInstance(*InstanceToRemove, *this);
 
 	SequenceInstanceBySection.Remove( &MovieSceneSection );
 }
@@ -987,8 +1106,17 @@ void FSequencer::ResetPerMovieSceneData()
 void FSequencer::UpdateRuntimeInstances()
 {
 	// Refresh the current root instance
-	RootMovieSceneSequenceInstance->RefreshInstance( *this );
-	RootMovieSceneSequenceInstance->Update(ScrubPosition, ScrubPosition, *this);
+	SequenceInstanceStack.Top()->RefreshInstance( *this );
+	SequenceInstanceStack.Top()->Update(ScrubPosition, ScrubPosition, *this);
+
+	// Ensure any relevant spawned objects are cleaned up if we're playing back the master sequence
+	if (SequenceInstanceStack.Num() == 1)
+	{
+		if (!GetPlaybackRange().Contains(ScrubPosition))
+		{
+			SpawnRegister->DestroyAllOwnedObjects(*this);
+		}
+	}
 
 	GEditor->RedrawLevelEditingViewports();
 }
@@ -1000,7 +1128,7 @@ FReply FSequencer::OnPlay(bool bTogglePlay)
 	{
 		PlaybackState = EMovieScenePlayerStatus::Stopped;
 		// Update on stop (cleans up things like sounds that are playing)
-		RootMovieSceneSequenceInstance->Update( ScrubPosition, ScrubPosition, *this );
+		SequenceInstanceStack.Top()->Update( ScrubPosition, ScrubPosition, *this );
 	}
 	else
 	{
@@ -1217,16 +1345,24 @@ void FSequencer::SetViewRange(TRange<float> NewViewRange, EViewRangeInterpolatio
 		ZoomAnimation.JumpToEnd();
 	}
 
-	UMovieScene* FocusedMovieScene = GetFocusedMovieSceneSequence()->GetMovieScene();
-	FocusedMovieScene->GetEditorData().ViewRange = TargetViewRange;
 
-	// Always ensure the working range is big enough to fit the view range
-	TRange<float>& WorkingRange = FocusedMovieScene->GetEditorData().WorkingRange;
+	UMovieSceneSequence* FocusedMovieSequence = GetFocusedMovieSceneSequence();
+	if (FocusedMovieSequence != nullptr)
+	{
+		UMovieScene* FocusedMovieScene = FocusedMovieSequence->GetMovieScene();
+		if (FocusedMovieScene != nullptr)
+		{
+			FocusedMovieScene->GetEditorData().ViewRange = TargetViewRange;
+
+			// Always ensure the working range is big enough to fit the view range
+			TRange<float>& WorkingRange = FocusedMovieScene->GetEditorData().WorkingRange;
 
-	WorkingRange = TRange<float>(
-		FMath::Min(TargetViewRange.GetLowerBoundValue(), WorkingRange.GetLowerBoundValue()),
-		FMath::Max(TargetViewRange.GetUpperBoundValue(), WorkingRange.GetUpperBoundValue())
-		);
+			WorkingRange = TRange<float>(
+				FMath::Min(TargetViewRange.GetLowerBoundValue(), WorkingRange.GetLowerBoundValue()),
+				FMath::Max(TargetViewRange.GetUpperBoundValue(), WorkingRange.GetUpperBoundValue())
+				);
+		}
+	}
 }
 
 void FSequencer::OnClampRangeChanged( TRange<float> NewClampRange )
@@ -1338,118 +1474,148 @@ void FSequencer::VerticalScroll(float ScrollAmountUnits)
 
 FGuid FSequencer::AddSpawnableForAssetOrClass( UObject* Object )
 {
-	FGuid NewSpawnableGuid;
-
-
 	UMovieSceneSequence* Sequence = GetFocusedMovieSceneSequence();
-	if (Sequence->AllowsSpawnableObjects())
+	if (!Sequence->AllowsSpawnableObjects())
 	{
-		// Grab the MovieScene that is currently focused.  We'll add our Blueprint as an inner of the
-		// MovieScene asset.
-		UMovieScene* OwnerMovieScene = Sequence->GetMovieScene();
+		return FGuid();
+	}
+
+	// Grab the MovieScene that is currently focused.  We'll add our Blueprint as an inner of the
+	// MovieScene asset.
+	UMovieScene* OwnerMovieScene = Sequence->GetMovieScene();
+
+	// @todo sequencer: Undo doesn't seem to be working at all
+	const FScopedTransaction Transaction( LOCTEXT("UndoAddingObject", "Add Object to MovieScene") );
 
-		// @todo sequencer: Undo doesn't seem to be working at all
-		const FScopedTransaction Transaction( LOCTEXT("UndoAddingObject", "Add Object to MovieScene") );
+	// Use the class as the spawnable's name if this is an actor class, otherwise just use the object name (asset)
+	const bool bIsActorClass = Object->IsA( AActor::StaticClass() ) && !Object->HasAnyFlags( RF_ArchetypeObject );
+	const FName AssetName = bIsActorClass ? Object->GetClass()->GetFName() : Object->GetFName();
 
-		// Use the class as the spawnable's name if this is an actor class, otherwise just use the object name (asset)
-		const bool bIsActorClass = Object->IsA( AActor::StaticClass() ) && !Object->HasAnyFlags( RF_ArchetypeObject );
-		const FName AssetName = bIsActorClass ? Object->GetClass()->GetFName() : Object->GetFName();
+	// Inner objects don't need a name (it will be auto-generated by the UObject system), but we want one in this case
+	// because the class of any actors that are created from this Blueprint may end up being user-facing.
+	const FName BlueprintName = MakeUniqueObjectName( OwnerMovieScene, UBlueprint::StaticClass(), AssetName );
 
-		// Inner objects don't need a name (it will be auto-generated by the UObject system), but we want one in this case
-		// because the class of any actors that are created from this Blueprint may end up being user-facing.
-		const FName BlueprintName = MakeUniqueObjectName( OwnerMovieScene, UBlueprint::StaticClass(), AssetName );
+	// Use the asset name as the initial spawnable name
+	const FString NewSpawnableName = AssetName.ToString();		// @todo sequencer: Need UI to allow user to rename these slots
 
-		// Use the asset name as the initial spawnable name
-		const FString NewSpawnableName = AssetName.ToString();		// @todo sequencer: Need UI to allow user to rename these slots
+	FText ErrorText;
 
-		// Create our new blueprint!
-		UBlueprint* NewBlueprint = Cast<UBlueprint>(Object);
-		if (!NewBlueprint)
+	// Create our new blueprint!
+	UBlueprint* NewBlueprint = nullptr;
+
+	// If we're dropping a blueprint, create a new blueprint out of this one, so we can propagate changes from the world to it without affecting the dragged BP.
+	if (UBlueprint* SourceBlueprint = Cast<UBlueprint>(Object))
+	{
+		UClass* SourceParentClass = SourceBlueprint->GeneratedClass;
+
+		if (!FKismetEditorUtilities::CanCreateBlueprintOfClass(SourceParentClass))
 		{
-			// @todo sequencer: Add support for forcing specific factories for an asset
-			UActorFactory* FactoryToUse = nullptr;
-			if( bIsActorClass )
-			{
-				// Placing an actor class directly::
-				FactoryToUse = GEditor->FindActorFactoryForActorClass( Object->GetClass() );
-			}
-			else
-			{
-				// Placing an asset
-				FactoryToUse = FActorFactoryAssetProxy::GetFactoryForAssetObject( Object );
-			}
+			ErrorText = FText::Format(LOCTEXT("UnableToAddSpawnableBlueprint", "Unable to add spawnable for class of type '{0}' since it is not a valid blueprint parent class."),
+				FText::FromString(SourceParentClass->GetName())
+				);
+		}
+		else
+		{
+			NewBlueprint = FKismetEditorUtilities::CreateBlueprint(SourceParentClass, OwnerMovieScene, BlueprintName, EBlueprintType::BPTYPE_Normal, UBlueprint::StaticClass(), UBlueprintGeneratedClass::StaticClass());
+		}
+	}
+	else
+	{
+		// @todo sequencer: Add support for forcing specific factories for an asset
+		UActorFactory* FactoryToUse = nullptr;
+		if( bIsActorClass )
+		{
+			// Placing an actor class directly::
+			FactoryToUse = GEditor->FindActorFactoryForActorClass( Object->GetClass() );
+		}
+		else
+		{
+			// Placing an asset
+			FactoryToUse = FActorFactoryAssetProxy::GetFactoryForAssetObject( Object );
+		}
 
-			if( FactoryToUse != nullptr )
-			{
-				// Create the blueprint
-				NewBlueprint = FactoryToUse->CreateBlueprint( Object, OwnerMovieScene, BlueprintName );
-			}
-			else if( bIsActorClass )
-			{
-				// We don't have a factory, but we can still try to create a blueprint for this actor class
-				NewBlueprint = FKismetEditorUtilities::CreateBlueprint( Object->GetClass(), OwnerMovieScene, BlueprintName, EBlueprintType::BPTYPE_Normal, UBlueprint::StaticClass(), UBlueprintGeneratedClass::StaticClass() );
-			}
+		if( FactoryToUse != nullptr )
+		{
+			// Create the blueprint
+			NewBlueprint = FactoryToUse->CreateBlueprint( Object, OwnerMovieScene, BlueprintName );
+		}
+		else if( bIsActorClass )
+		{
+			// We don't have a factory, but we can still try to create a blueprint for this actor class
+			NewBlueprint = FKismetEditorUtilities::CreateBlueprint( Object->GetClass(), OwnerMovieScene, BlueprintName, EBlueprintType::BPTYPE_Normal, UBlueprint::StaticClass(), UBlueprintGeneratedClass::StaticClass() );
 		}
+	}
 
-		if (NewBlueprint != nullptr)
+	if (!NewBlueprint)
+	{
+		if (ErrorText.IsEmpty())
 		{
-			FTransformData NewTransform;
+			ErrorText = FText::Format(LOCTEXT("UnableToAddSpawnable", "Unable to add spawnable for '{0}'."),
+				FText::FromString(Object->GetName())
+				);
+		}
 
-			if ((NewBlueprint->GeneratedClass != nullptr) && FBlueprintEditorUtils::IsActorBased(NewBlueprint))
-			{
-				AActor* ActorCDO = CastChecked< AActor >( NewBlueprint->GeneratedClass->ClassDefaultObject );
-				// Place the new spawnable in front of the camera (unless we were automatically created from a PIE actor)
-				if (Settings->GetSpawnPosition() == SSP_PlaceInFrontOfCamera)
-				{
-					PlaceActorInFrontOfCamera( ActorCDO );
-				}
-				NewTransform.Translation = ActorCDO->GetActorLocation();
-				NewTransform.Rotation = ActorCDO->GetActorRotation();
-				NewTransform.Scale = FVector(1.0f, 1.0f, 1.0f);
-				NewTransform.bValid = true;
-			}
+		FNotificationInfo Info(ErrorText);
+		Info.ExpireDuration = 3.0f;
+		FSlateNotificationManager::Get().AddNotification(Info);
+		return FGuid();
+	}
+
+	FTransformData NewTransform;
 
-			NewSpawnableGuid = OwnerMovieScene->AddSpawnable(NewSpawnableName, NewBlueprint);
-			if (NewSpawnableGuid.IsValid())
+	if ((NewBlueprint->GeneratedClass != nullptr) && FBlueprintEditorUtils::IsActorBased(NewBlueprint))
+	{
+		AActor* ActorCDO = CastChecked< AActor >( NewBlueprint->GeneratedClass->ClassDefaultObject );
+		// Place the new spawnable in front of the camera (unless we were automatically created from a PIE actor)
+		if (Settings->GetSpawnPosition() == SSP_PlaceInFrontOfCamera)
+		{
+			PlaceActorInFrontOfCamera( ActorCDO );
+		}
+		NewTransform.Translation = ActorCDO->GetActorLocation();
+		NewTransform.Rotation = ActorCDO->GetActorRotation();
+		NewTransform.Scale = FVector(1.0f, 1.0f, 1.0f);
+		NewTransform.bValid = true;
+	}
+
+	FGuid NewSpawnableGuid = OwnerMovieScene->AddSpawnable(NewSpawnableName, NewBlueprint);
+	if (NewSpawnableGuid.IsValid())
+	{
+		UMovieSceneSpawnTrack* SpawnTrack = Cast<UMovieSceneSpawnTrack>(OwnerMovieScene->AddTrack(UMovieSceneSpawnTrack::StaticClass(), NewSpawnableGuid));
+		if (SpawnTrack)
+		{
+			SpawnTrack->AddSection(*SpawnTrack->CreateNewSection());
+			SpawnTrack->SetObjectId(NewSpawnableGuid);
+		}
+						
+		if (NewTransform.bValid)
+		{
+			UMovieScene3DTransformTrack* TransformTrack = Cast<UMovieScene3DTransformTrack>(OwnerMovieScene->AddTrack(UMovieScene3DTransformTrack::StaticClass(), NewSpawnableGuid));
+			if (TransformTrack)
 			{
-				UMovieSceneSpawnTrack* SpawnTrack = Cast<UMovieSceneSpawnTrack>(OwnerMovieScene->AddTrack(UMovieSceneSpawnTrack::StaticClass(), NewSpawnableGuid));
-				if (SpawnTrack)
-				{
-					SpawnTrack->AddSection(*SpawnTrack->CreateNewSection());
-					SpawnTrack->SetObject(NewSpawnableGuid);
-				}
-								
-				if (NewTransform.bValid)
-				{
-					UMovieScene3DTransformTrack* TransformTrack = Cast<UMovieScene3DTransformTrack>(OwnerMovieScene->AddTrack(UMovieScene3DTransformTrack::StaticClass(), NewSpawnableGuid));
-					if (TransformTrack)
-					{
-						static FName Transform("Transform");
-						TransformTrack->SetPropertyNameAndPath(Transform, Transform.ToString());
-						const bool bUnwindRotation = false;
+				static FName Transform("Transform");
+				TransformTrack->SetPropertyNameAndPath(Transform, Transform.ToString());
+				const bool bUnwindRotation = false;
 
-						float Time = GetCurrentLocalTime( *Sequence );
-						EMovieSceneKeyInterpolation Interpolation = GetKeyInterpolation();
+				float Time = GetCurrentLocalTime( *Sequence );
+				EMovieSceneKeyInterpolation Interpolation = GetKeyInterpolation();
 
-						UMovieScene3DTransformSection* NewSection = CastChecked<UMovieScene3DTransformSection>(TransformTrack->CreateNewSection());
+				UMovieScene3DTransformSection* NewSection = CastChecked<UMovieScene3DTransformSection>(TransformTrack->CreateNewSection());
 
-						NewSection->SetDefault( FTransformKey( EKey3DTransformChannel::Translation, EAxis::X, NewTransform.Translation.X, bUnwindRotation ) );
-						NewSection->SetDefault( FTransformKey( EKey3DTransformChannel::Translation, EAxis::Y, NewTransform.Translation.Y, bUnwindRotation ) );
-						NewSection->SetDefault( FTransformKey( EKey3DTransformChannel::Translation, EAxis::Z, NewTransform.Translation.Z, bUnwindRotation ) );
+				NewSection->SetDefault( FTransformKey( EKey3DTransformChannel::Translation, EKey3DTransformChannel::Default, EAxis::X, NewTransform.Translation.X, bUnwindRotation ) );
+				NewSection->SetDefault( FTransformKey( EKey3DTransformChannel::Translation, EKey3DTransformChannel::Default, EAxis::Y, NewTransform.Translation.Y, bUnwindRotation ) );
+				NewSection->SetDefault( FTransformKey( EKey3DTransformChannel::Translation, EKey3DTransformChannel::Default, EAxis::Z, NewTransform.Translation.Z, bUnwindRotation ) );
 
-						NewSection->SetDefault( FTransformKey( EKey3DTransformChannel::Rotation, EAxis::X, NewTransform.Rotation.Euler().X, bUnwindRotation ) );
-						NewSection->SetDefault( FTransformKey( EKey3DTransformChannel::Rotation, EAxis::Y, NewTransform.Rotation.Euler().Y, bUnwindRotation ) );
-						NewSection->SetDefault( FTransformKey( EKey3DTransformChannel::Rotation, EAxis::Z, NewTransform.Rotation.Euler().Z, bUnwindRotation ) );
+				NewSection->SetDefault( FTransformKey( EKey3DTransformChannel::Rotation, EKey3DTransformChannel::Default, EAxis::X, NewTransform.Rotation.Euler().X, bUnwindRotation ) );
+				NewSection->SetDefault( FTransformKey( EKey3DTransformChannel::Rotation, EKey3DTransformChannel::Default, EAxis::Y, NewTransform.Rotation.Euler().Y, bUnwindRotation ) );
+				NewSection->SetDefault( FTransformKey( EKey3DTransformChannel::Rotation, EKey3DTransformChannel::Default, EAxis::Z, NewTransform.Rotation.Euler().Z, bUnwindRotation ) );
 
-						NewSection->SetDefault( FTransformKey( EKey3DTransformChannel::Scale, EAxis::X, NewTransform.Scale.X, bUnwindRotation ) );
-						NewSection->SetDefault( FTransformKey( EKey3DTransformChannel::Scale, EAxis::Y, NewTransform.Scale.Y, bUnwindRotation ) );
-						NewSection->SetDefault( FTransformKey( EKey3DTransformChannel::Scale, EAxis::Z, NewTransform.Scale.Z, bUnwindRotation ) );
+				NewSection->SetDefault( FTransformKey( EKey3DTransformChannel::Scale, EKey3DTransformChannel::Default, EAxis::X, NewTransform.Scale.X, bUnwindRotation ) );
+				NewSection->SetDefault( FTransformKey( EKey3DTransformChannel::Scale, EKey3DTransformChannel::Default, EAxis::Y, NewTransform.Scale.Y, bUnwindRotation ) );
+				NewSection->SetDefault( FTransformKey( EKey3DTransformChannel::Scale, EKey3DTransformChannel::Default, EAxis::Z, NewTransform.Scale.Z, bUnwindRotation ) );
 
-						TransformTrack->AddSection(*NewSection);
-					}
-				}
+				TransformTrack->AddSection(*NewSection);
 			}
-		}	
+		}
 	}
 
 	return NewSpawnableGuid;
@@ -1514,7 +1680,6 @@ void GetParentTrackNodeAndNamePath(TSharedRef<const FSequencerDisplayNode> Displ
 
 void FSequencer::OnRequestNodeDeleted( TSharedRef<const FSequencerDisplayNode> NodeToBeDeleted )
 {
-	bool bAnySpawnablesRemoved = false;
 	bool bAnythingRemoved = false;
 	
 	TSharedRef<FMovieSceneSequenceInstance> MovieSceneInstance = GetFocusedMovieSceneSequenceInstance();
@@ -1522,31 +1687,34 @@ void FSequencer::OnRequestNodeDeleted( TSharedRef<const FSequencerDisplayNode> N
 	UMovieScene* OwnerMovieScene = Sequence->GetMovieScene();
 
 	// Only object nodes or section areas can be deleted
-	if( NodeToBeDeleted->GetType() == ESequencerNode::Object  )
+	if (NodeToBeDeleted->GetType() == ESequencerNode::Object)
 	{
-		const FGuid& BindingToRemove = StaticCastSharedRef<const FSequencerObjectBindingNode>( NodeToBeDeleted )->GetObjectBinding();
+		// Delete any child object bindings
+		for (const TSharedRef<FSequencerDisplayNode>& ChildNode : NodeToBeDeleted->GetChildNodes())
+		{
+			if (ChildNode->GetType() == ESequencerNode::Object)
+			{
+				OnRequestNodeDeleted(ChildNode);
+			}
+		}
 
-		//@todo Sequencer - add transaction
+		const FGuid& BindingToRemove = StaticCastSharedRef<const FSequencerObjectBindingNode>( NodeToBeDeleted )->GetObjectBinding();
 		
 		// Try to remove as a spawnable first
-		bool bRemoved = OwnerMovieScene->RemoveSpawnable( BindingToRemove );
-		if( bRemoved )
+		if (OwnerMovieScene->RemoveSpawnable(BindingToRemove))
 		{
-			Sequence->DestroySpawnedObject(BindingToRemove);
-			bAnySpawnablesRemoved = true;
+			SpawnRegister->DestroySpawnedObject(BindingToRemove, *MovieSceneInstance, *this);
 		}
-
-		if( !bRemoved )
+		// The guid should be associated with a possessable if it wasnt a spawnable
+		else if (OwnerMovieScene->RemovePossessable(BindingToRemove))
 		{
-			// The guid should be associated with a possessable if it wasnt a spawnable
-			bRemoved = OwnerMovieScene->RemovePossessable( BindingToRemove );
-			
 			Sequence->Modify();
-
 			Sequence->UnbindPossessableObjects( BindingToRemove );
-
-			// If this check fails the guid was not associated with a spawnable or possessable so there was an invalid guid being stored on a node
-			check( bRemoved );
+		}
+		else
+		{
+			// If this check fails there was an invalid guid being stored on a node
+			checkf(false, TEXT("Guid was not associated with a spawnable or possessable"));
 		}
 
 		bAnythingRemoved = true;
@@ -1826,38 +1994,25 @@ void FSequencer::AddSelectedObjects()
 		AActor* Actor = CastChecked<AActor>(*It);
 
 		if (Actor != nullptr)
-		{			
-			FGuid ObjectGuid = OwnerSequence->FindObjectId(*Actor);
-
-			// Add this object if it hasn't already been possessed
-			if (ObjectGuid.IsValid())
-			{
-				FMovieScenePossessable* Possessable = OwnerMovieScene->FindPossessable(ObjectGuid);
-				if (Possessable == nullptr)
-				{
-					SelectedActors.Add(Actor);
-				}
-			}
-			else
-			{
-				SelectedActors.Add(Actor);
-			}
+		{
+			SelectedActors.Add(Actor);
 		}
 	}
 
 	bool bPossessableAdded = false;
 	if (SelectedActors.Num() != 0)
 	{
-		const FScopedTransaction Transaction(LOCTEXT("UndoPossessingObject", "Possess Object with MovieScene"));
-			
-		// possess the object!
+		const FScopedTransaction Transaction(LOCTEXT("UndoPossessingObject", "Possess Object in Sequencer"));
 		OwnerSequence->Modify();
-				
+
 		for (auto Actor : SelectedActors)
 		{
-			const FGuid PossessableGuid = OwnerMovieScene->AddPossessable(Actor->GetActorLabel(), Actor->GetClass());
-			OwnerSequence->BindPossessableObject(PossessableGuid, *Actor);
-			bPossessableAdded = true;
+			FGuid ExistingGuid = SequenceInstanceStack.Top()->FindObjectId(*Actor);
+			if (!ExistingGuid.IsValid())
+			{
+				CreateBinding(*Actor, Actor->GetActorLabel());
+				bPossessableAdded = true;
+			}
 		}
 	}
 
@@ -1876,7 +2031,7 @@ void FSequencer::OnSelectedOutlinerNodesChanged()
 	if (IsLevelEditorSequencer())
 	{
 		// Clear the selection if no sequencer display nodes were clicked on and there are selected actors
-		if (Selection.GetSelectedOutlinerNodes().Num() == 0 && GEditor->GetSelectedActors()->Num() != 0)
+		if (Selection.GetSelectedOutlinerNodes().Num() == 0 && GEditor->GetSelectedActors()->Num() != 0 && SequencerWidget->UserIsSelecting())
 		{
 			const bool bNotifySelectionChanged = true;
 			const bool bDeselectBSP = true;
@@ -1963,8 +2118,7 @@ void FSequencer::DeleteSelectedItems()
 
 void FSequencer::AssignActor(FMenuBuilder& MenuBuilder, FGuid InObjectBinding)
 {
-	UMovieSceneSequence* OwnerSequence = GetFocusedMovieSceneSequence();
-	UObject* RuntimeObject = OwnerSequence->FindObject(InObjectBinding);
+	UObject* RuntimeObject = SequenceInstanceStack.Top()->FindObject(InObjectBinding, *this);
 
 	auto IsActorValidForAssignment = [=](const AActor* InActor, UObject* CurrentObject){
 		return CurrentObject != InActor;
@@ -2018,6 +2172,8 @@ void FSequencer::DoAssignActor(AActor*const* InActors, int32 NumActors, FGuid In
 		{
 			FScopedTransaction AssignActor( NSLOCTEXT("Sequencer", "AssignActor", "Assign Actor") );
 
+			RootMovieSceneSequenceInstance->SaveState(*this);
+
 			UMovieSceneSequence* OwnerSequence = GetFocusedMovieSceneSequence();
 			UMovieScene* OwnerMovieScene = OwnerSequence->GetMovieScene();
 
@@ -2025,22 +2181,44 @@ void FSequencer::DoAssignActor(AActor*const* InActors, int32 NumActors, FGuid In
 			OwnerSequence->Modify();
 			OwnerMovieScene->Modify();
 
+			UObject* RuntimeObject = SequenceInstanceStack.Top()->FindObject(InObjectBinding, *this);
+
+			// Replace the object itself
+			FGuid ParentGuid;
+			{
+				// Get the object guid to assign, remove the binding if it already exists
+				ParentGuid = SequenceInstanceStack.Top()->FindObjectId(*Actor);
+				FString NewActorLabel = Actor->GetActorLabel();
+				if (ParentGuid.IsValid())
+				{
+					OwnerMovieScene->RemovePossessable(ParentGuid);
+					OwnerSequence->UnbindPossessableObjects(ParentGuid);
+				}
+
+				// Add this object
+				FMovieScenePossessable NewPossessable( NewActorLabel, Actor->GetClass());
+				ParentGuid = NewPossessable.GetGuid();
+				OwnerSequence->BindPossessableObject(ParentGuid, *Actor, GetPlaybackContext());
+
+				// Replace
+				OwnerMovieScene->ReplacePossessable(InObjectBinding, ParentGuid, NewActorLabel);
+			}
+
 			// Handle components
-			UObject* RuntimeObject = OwnerSequence->FindObject(InObjectBinding);
 			if (AActor* ActorToReplace = Cast<AActor>(RuntimeObject))
 			{
 				for (UActorComponent* ComponentToReplace : ActorToReplace->GetComponents())
 				{
 					if (ComponentToReplace != nullptr)
 					{
-						FGuid ComponentGuid = OwnerSequence->FindObjectId(*ComponentToReplace);
+						FGuid ComponentGuid = SequenceInstanceStack.Top()->FindObjectId(*ComponentToReplace);
 						if (ComponentGuid.IsValid())
 						{
 							UActorComponent* NewComponent = Actor->GetComponentByClass(ComponentToReplace->GetClass());
 							if (NewComponent)
 							{
 								// Get the object guid to assign, remove the binding if it already exists
-								FGuid NewComponentGuid = OwnerSequence->FindObjectId(*NewComponent);
+								FGuid NewComponentGuid = SequenceInstanceStack.Top()->FindObjectId(*NewComponent);
 								FString NewComponentLabel = NewComponent->GetName();
 								if (NewComponentGuid.IsValid())
 								{
@@ -2051,33 +2229,23 @@ void FSequencer::DoAssignActor(AActor*const* InActors, int32 NumActors, FGuid In
 								// Add this object
 								FMovieScenePossessable NewPossessable( NewComponentLabel, NewComponent->GetClass());
 								NewComponentGuid = NewPossessable.GetGuid();
-								OwnerSequence->BindPossessableObject(NewComponentGuid, *NewComponent);
+								OwnerSequence->BindPossessableObject(NewComponentGuid, *NewComponent, Actor);
 
 								// Replace
 								OwnerMovieScene->ReplacePossessable(ComponentGuid, NewComponentGuid, NewComponentLabel);
+
+								FMovieScenePossessable* ThisPossessable = OwnerMovieScene->FindPossessable(NewComponentGuid);
+								if (ensure(ThisPossessable))
+								{
+									ThisPossessable->SetParent(ParentGuid);
+								}
 							}
 						}
 					}
 				}
 			}
-
-			// Get the object guid to assign, remove the binding if it already exists
-			FGuid NewObjectGuid = OwnerSequence->FindObjectId(*Actor);
-			FString NewActorLabel = Actor->GetActorLabel();
-			if (NewObjectGuid.IsValid())
-			{
-				OwnerMovieScene->RemovePossessable(NewObjectGuid);
-				OwnerSequence->UnbindPossessableObjects(NewObjectGuid);
-			}
-
-			// Add this object
-			FMovieScenePossessable NewPossessable( NewActorLabel, Actor->GetClass());
-			NewObjectGuid = NewPossessable.GetGuid();
-			OwnerSequence->BindPossessableObject(NewObjectGuid, *Actor);
-
-			// Replace
-			OwnerMovieScene->ReplacePossessable(InObjectBinding, NewObjectGuid, NewActorLabel);
-					
+			
+			RootMovieSceneSequenceInstance->RestoreState(*this);
 			NotifyMovieSceneDataChanged();
 		}
 	}
@@ -2612,6 +2780,14 @@ void FSequencer::BindSequencerCommands()
 		FCanExecuteAction::CreateLambda( []{ return true; } ),
 		FIsActionChecked::CreateLambda( [this]{ return Settings->GetDetailsViewVisible(); } ) );
 
+	SequencerCommandBindings->MapAction(
+		Commands.ToggleLabelBrowser,
+		FExecuteAction::CreateLambda( [this]{
+			Settings->SetLabelBrowserVisible( !Settings->GetLabelBrowserVisible() );
+		} ),
+		FCanExecuteAction::CreateLambda( []{ return true; } ),
+		FIsActionChecked::CreateLambda( [this]{ return Settings->GetLabelBrowserVisible(); } ) );
+
 	SequencerCommandBindings->MapAction(
 		Commands.ToggleShowFrameNumbers,
 		FExecuteAction::CreateLambda( [this]{ Settings->SetShowFrameNumbers( !Settings->GetShowFrameNumbers() ); } ),
@@ -2679,6 +2855,11 @@ void FSequencer::BindSequencerCommands()
 		FIsActionChecked::CreateLambda( [this]{ return Settings->GetShowCurveEditor(); } ) );
 
 	auto CanCutOrCopy = [this]{
+		if (!IsSequencerWidgetFocused())
+		{
+			return false;
+		}
+		
 		UMovieSceneTrack* Track = nullptr;
 		for (FSequencerSelectedKey Key : Selection.GetSelectedKeys())
 		{
diff --git a/Engine/Source/Editor/Sequencer/Private/Sequencer.h b/Engine/Source/Editor/Sequencer/Private/Sequencer.h
--- a/Engine/Source/Editor/Sequencer/Private/Sequencer.h
+++ b/Engine/Source/Editor/Sequencer/Private/Sequencer.h
@@ -6,7 +6,6 @@
 #include "EditorUndoClient.h"
 #include "MovieSceneClipboard.h"
 
-
 class FMenuBuilder;
 class FMovieSceneSequenceInstance;
 class IDetailKeyframeHandler;
@@ -20,6 +19,7 @@ class UMovieScene;
 class UMovieSceneSection;
 class UMovieSceneSequence;
 class UWorld;
+class IMovieSceneSpawnRegister;
 
 
 /**
@@ -103,6 +103,9 @@ public:
 	virtual FSequencerSelectionPreview& GetSelectionPreview() override;
 	virtual void NotifyMapChanged(UWorld* NewWorld, EMapChangeType MapChangeType) override;
 	virtual FOnGlobalTimeChanged& OnGlobalTimeChanged() override { return OnGlobalTimeChangedDelegate; }
+	virtual FGuid CreateBinding(UObject& InObject, const FString& InName) override;
+	virtual UObject* GetPlaybackContext() const override;
+	virtual void GetAllKeyedProperties(UObject& Object, TSet<UProperty*>& OutProperties) override;
 
 	/** Set the global time directly, without performing any auto-scroll */
 	void SetGlobalTimeDirectly(float Time);
@@ -251,6 +254,7 @@ public:
 	virtual EMovieScenePlayerStatus::Type GetPlaybackStatus() const override;
 	virtual void AddOrUpdateMovieSceneInstance(UMovieSceneSection& MovieSceneSection, TSharedRef<FMovieSceneSequenceInstance> InstanceToAdd) override;
 	virtual void RemoveMovieSceneInstance(UMovieSceneSection& MovieSceneSection, TSharedRef<FMovieSceneSequenceInstance> InstanceToRemove) override;
+	virtual IMovieSceneSpawnRegister& GetSpawnRegister() override { return *SpawnRegister; }
 
 	/** Called when an actor is dropped into Sequencer */
 	void OnActorsDropped( const TArray<TWeakObjectPtr<AActor> >& Actors );
@@ -429,6 +433,26 @@ protected:
 	 */
 	UObject* GetCurrentAsset() const;
 
+protected:
+
+	/**
+	 * Populate the specified set with all UProperties that are currently keyed for the specified object and sequence instance, including any sub sequences
+	 *
+	 * @param Object		The object to get keyed properties for
+	 * @param Instance		The sequence instance in which to look for the object bindings
+	 * @param OutProperties	set to populate with properties
+	 */
+	void GetAllKeyedPropertiesForInstance(UObject& Object, FMovieSceneSequenceInstance& Instance, TSet<UProperty*>& OutProperties);
+
+	/**
+	 * Populate the specified set with all UProperties that are currently keyed for the specified object and sequence instance, excluding any sub sequences
+	 *
+	 * @param Object		The object to get keyed properties for
+	 * @param Instance		The sequence instance in which to look for the object bindings
+	 * @param OutProperties	set to populate with properties
+	 */
+	void GetKeyedProperties(UObject& Object, FMovieSceneSequenceInstance& Instance, TSet<UProperty*>& OutProperties);
+
 protected:
 
 	/** Called via UEditorEngine::GetActorRecordingStateEvent to check to see whether we need to record actor state */
@@ -539,6 +563,9 @@ private:
 	/** Main sequencer widget */
 	TSharedPtr<SSequencer> SequencerWidget;
 	
+	/** Spawn register for keeping track of what is spawned */
+	TSharedPtr<IMovieSceneSpawnRegister> SpawnRegister;
+
 	/** The asset editor that created this Sequencer if any */
 	TWeakPtr<IToolkitHost> ToolkitHost;
 
diff --git a/Engine/Source/Editor/Sequencer/Private/SequencerCommands.cpp b/Engine/Source/Editor/Sequencer/Private/SequencerCommands.cpp
--- a/Engine/Source/Editor/Sequencer/Private/SequencerCommands.cpp
+++ b/Engine/Source/Editor/Sequencer/Private/SequencerCommands.cpp
@@ -58,6 +58,7 @@ void FSequencerCommands::RegisterCommands()
 
 	UI_COMMAND( FindInContentBrowser, "Find in Content Browser", "Find the viewed sequence asset in the content browser", EUserInterfaceActionType::Button, FInputChord() );
 	UI_COMMAND( ToggleDetailsView, "Details View", "Enable details view for selected sections and keys", EUserInterfaceActionType::ToggleButton, FInputChord() );
+	UI_COMMAND( ToggleLabelBrowser, "Label Browser", "Show/hide the track label browser", EUserInterfaceActionType::ToggleButton, FInputChord() );
 	
 	UI_COMMAND( ToggleShowCurveEditor, "Curve Editor", "Show the animation keys in a curve editor", EUserInterfaceActionType::ToggleButton, FInputChord() );
 
diff --git a/Engine/Source/Editor/Sequencer/Private/SequencerCommands.h b/Engine/Source/Editor/Sequencer/Private/SequencerCommands.h
--- a/Engine/Source/Editor/Sequencer/Private/SequencerCommands.h
+++ b/Engine/Source/Editor/Sequencer/Private/SequencerCommands.h
@@ -135,6 +135,9 @@ public:
 	/** Toggles whether the details view is enabled in the level editor. */
 	TSharedPtr< FUICommandInfo > ToggleDetailsView;
 
+	/** Toggles whether the label browser is enabled in the level editor. */
+	TSharedPtr< FUICommandInfo > ToggleLabelBrowser;
+
 	/** Turns auto scroll on and off. */
 	TSharedPtr< FUICommandInfo > ToggleAutoScroll;
 
diff --git a/Engine/Source/Editor/Sequencer/Private/SequencerContextMenus.cpp b/Engine/Source/Editor/Sequencer/Private/SequencerContextMenus.cpp
--- a/Engine/Source/Editor/Sequencer/Private/SequencerContextMenus.cpp
+++ b/Engine/Source/Editor/Sequencer/Private/SequencerContextMenus.cpp
@@ -241,6 +241,23 @@ void FSectionContextMenu::PopulateMenu(FMenuBuilder& MenuBuilder)
 		);
 	}
 	MenuBuilder.EndSection(); // SequencerSections
+
+	MenuBuilder.BeginSection("SequencerSectionOrder", LOCTEXT("SectionOrderLabel", "Order"));
+	{
+		MenuBuilder.AddMenuEntry(LOCTEXT("BringToFront", "Bring To Front"), FText(), FSlateIcon(),
+			FUIAction(FExecuteAction::CreateLambda([=]{ return Shared->BringToFront(); })));
+
+		MenuBuilder.AddMenuEntry(LOCTEXT("SendToBack", "Send To Back"), FText(), FSlateIcon(),
+			FUIAction(FExecuteAction::CreateLambda([=]{ return Shared->SendToBack(); })));
+
+		MenuBuilder.AddMenuEntry(LOCTEXT("BringForward", "Bring Forward"), FText(), FSlateIcon(),
+			FUIAction(FExecuteAction::CreateLambda([=]{ return Shared->BringForward(); })));
+
+		MenuBuilder.AddMenuEntry(LOCTEXT("SendBackward", "Send Backward"), FText(), FSlateIcon(),
+			FUIAction(FExecuteAction::CreateLambda([=]{ return Shared->SendBackward(); })));
+
+	}
+	MenuBuilder.EndSection(); // SequencerSectionOrder
 }
 
 void FSectionContextMenu::AddEditMenu(FMenuBuilder& MenuBuilder)
@@ -565,6 +582,272 @@ void FSectionContextMenu::DeleteSection()
 	Sequencer->DeleteSections(Sequencer->GetSelection().GetSelectedSections());
 }
 
+/** Information pertaining to a specific row in a track, required for z-ordering operations */
+struct FTrackSectionRow
+{
+	/** The minimum z-order value for all the sections in this row */
+	int32 MinOrderValue;
+
+	/** The maximum z-order value for all the sections in this row */
+	int32 MaxOrderValue;
+
+	/** All the sections contained in this row */
+	TArray<UMovieSceneSection*> Sections;
+
+	/** A set of sections that are to be operated on */
+	TSet<UMovieSceneSection*> SectionToReOrder;
+
+	void AddSection(UMovieSceneSection* InSection)
+	{
+		Sections.Add(InSection);
+		MinOrderValue = FMath::Min(MinOrderValue, InSection->GetOverlapPriority());
+		MaxOrderValue = FMath::Max(MaxOrderValue, InSection->GetOverlapPriority());
+	}
+};
+
+/** Generate the data required for re-ordering rows based on the current sequencer selection */
+/** @note: Produces a map of track -> rows, keyed on row index. Only returns rows that contain selected sections */
+TMap<UMovieSceneTrack*, TMap<int32, FTrackSectionRow>> GenerateTrackRowsFromSelection(FSequencer& Sequencer)
+{
+	TMap<UMovieSceneTrack*, TMap<int32, FTrackSectionRow>> TrackRows;
+
+	for (const TWeakObjectPtr<UMovieSceneSection>& SectionPtr : Sequencer.GetSelection().GetSelectedSections())
+	{
+		UMovieSceneSection* Section = SectionPtr.Get();
+		if (!Section)
+		{
+			continue;
+		}
+
+		UMovieSceneTrack* Track = Section->GetTypedOuter<UMovieSceneTrack>();
+		if (!Track)
+		{
+			continue;
+		}
+
+		FTrackSectionRow& Row = TrackRows.FindOrAdd(Track).FindOrAdd(Section->GetRowIndex());
+		Row.SectionToReOrder.Add(Section);
+	}
+
+	// Now ensure all rows that we're operating on are fully populated
+	for (auto& Pair : TrackRows)
+	{
+		UMovieSceneTrack* Track = Pair.Key;
+		for (auto& RowPair : Pair.Value)
+		{
+			const int32 RowIndex = RowPair.Key;
+			for (UMovieSceneSection* Section : Track->GetAllSections())
+			{
+				if (Section->GetRowIndex() == RowIndex)
+				{
+					RowPair.Value.AddSection(Section);
+				}
+			}
+		}
+	}
+
+	return TrackRows;
+}
+
+/** Modify all the sections contained within the specified data structure */
+void ModifySections(TMap<UMovieSceneTrack*, TMap<int32, FTrackSectionRow>>& TrackRows)
+{
+	for (auto& Pair : TrackRows)
+	{
+		UMovieSceneTrack* Track = Pair.Key;
+		for (auto& RowPair : Pair.Value)
+		{
+			for (UMovieSceneSection* Section : RowPair.Value.Sections)
+			{
+				Section->Modify();
+			}
+		}
+	}
+}
+
+void FSectionContextMenu::BringToFront()
+{
+	TMap<UMovieSceneTrack*, TMap<int32, FTrackSectionRow>> TrackRows = GenerateTrackRowsFromSelection(*Sequencer);
+	if (TrackRows.Num() == 0)
+	{
+		return;
+	}
+
+	FScopedTransaction Transaction(LOCTEXT("BringToFrontTransaction", "Bring to Front"));
+	ModifySections(TrackRows);
+
+	for (auto& Pair : TrackRows)
+	{
+		UMovieSceneTrack* Track = Pair.Key;
+		TMap<int32, FTrackSectionRow>& Rows = Pair.Value;
+
+		for (auto& RowPair : Rows)
+		{
+			FTrackSectionRow& Row = RowPair.Value;
+
+			Row.Sections.StableSort([&](UMovieSceneSection& A, UMovieSceneSection& B){
+				bool bIsActiveA = Row.SectionToReOrder.Contains(&A);
+				bool bIsActiveB = Row.SectionToReOrder.Contains(&B);
+
+				// Sort secondarily on overlap priority
+				if (bIsActiveA == bIsActiveB)
+				{
+					return A.GetOverlapPriority() < B.GetOverlapPriority();
+				}
+				// Sort and primarily on whether we're sending to the back or not (bIsActive)
+				else
+				{
+					return !bIsActiveA;
+				}
+			});
+
+			int32 CurrentPriority = Row.MinOrderValue;
+			for (UMovieSceneSection* Section : Row.Sections)
+			{
+				Section->SetOverlapPriority(CurrentPriority++);
+			}
+		}
+	}
+
+	Sequencer->SetGlobalTime(Sequencer->GetGlobalTime());
+}
+
+void FSectionContextMenu::SendToBack()
+{
+	TMap<UMovieSceneTrack*, TMap<int32, FTrackSectionRow>> TrackRows = GenerateTrackRowsFromSelection(*Sequencer);
+	if (TrackRows.Num() == 0)
+	{
+		return;
+	}
+
+	FScopedTransaction Transaction(LOCTEXT("SendToBackTransaction", "Send to Back"));
+	ModifySections(TrackRows);
+
+	for (auto& Pair : TrackRows)
+	{
+		UMovieSceneTrack* Track = Pair.Key;
+		TMap<int32, FTrackSectionRow>& Rows = Pair.Value;
+
+		for (auto& RowPair : Rows)
+		{
+			FTrackSectionRow& Row = RowPair.Value;
+
+			Row.Sections.StableSort([&](UMovieSceneSection& A, UMovieSceneSection& B){
+				bool bIsActiveA = Row.SectionToReOrder.Contains(&A);
+				bool bIsActiveB = Row.SectionToReOrder.Contains(&B);
+
+				// Sort secondarily on overlap priority
+				if (bIsActiveA == bIsActiveB)
+				{
+					return A.GetOverlapPriority() < B.GetOverlapPriority();
+				}
+				// Sort and primarily on whether we're bringing to the front or not (bIsActive)
+				else
+				{
+					return bIsActiveA;
+				}
+			});
+
+			int32 CurrentPriority = Row.MinOrderValue;
+			for (UMovieSceneSection* Section : Row.Sections)
+			{
+				Section->SetOverlapPriority(CurrentPriority++);
+			}
+		}
+	}
+
+	Sequencer->SetGlobalTime(Sequencer->GetGlobalTime());
+}
+
+void FSectionContextMenu::BringForward()
+{
+	TMap<UMovieSceneTrack*, TMap<int32, FTrackSectionRow>> TrackRows = GenerateTrackRowsFromSelection(*Sequencer);
+	if (TrackRows.Num() == 0)
+	{
+		return;
+	}
+
+	FScopedTransaction Transaction(LOCTEXT("BringForwardTransaction", "Bring Forward"));
+	ModifySections(TrackRows);
+
+	for (auto& Pair : TrackRows)
+	{
+		UMovieSceneTrack* Track = Pair.Key;
+		TMap<int32, FTrackSectionRow>& Rows = Pair.Value;
+
+		for (auto& RowPair : Rows)
+		{
+			FTrackSectionRow& Row = RowPair.Value;
+
+			Row.Sections.Sort([&](UMovieSceneSection& A, UMovieSceneSection& B){
+				return A.GetOverlapPriority() < B.GetOverlapPriority();
+			});
+
+			for (int32 SectionIndex = Row.Sections.Num() - 1; SectionIndex >= 0; --SectionIndex)
+			{
+				UMovieSceneSection* ThisSection = Row.Sections[SectionIndex];
+				if (Row.SectionToReOrder.Contains(ThisSection))
+				{
+					UMovieSceneSection* OtherSection = Row.Sections[SectionIndex + 1];
+
+					Row.Sections.Swap(SectionIndex, SectionIndex+1);
+
+					const int32 SwappedPriority = OtherSection->GetOverlapPriority();
+					OtherSection->SetOverlapPriority(ThisSection->GetOverlapPriority());
+					ThisSection->SetOverlapPriority(SwappedPriority);
+				}
+			}
+		}
+	}
+
+	Sequencer->SetGlobalTime(Sequencer->GetGlobalTime());
+}
+
+void FSectionContextMenu::SendBackward()
+{
+	TMap<UMovieSceneTrack*, TMap<int32, FTrackSectionRow>> TrackRows = GenerateTrackRowsFromSelection(*Sequencer);
+	if (TrackRows.Num() == 0)
+	{
+		return;
+	}
+
+	FScopedTransaction Transaction(LOCTEXT("SendBackwardTransaction", "Send Backward"));
+	ModifySections(TrackRows);
+
+	for (auto& Pair : TrackRows)
+	{
+		UMovieSceneTrack* Track = Pair.Key;
+		TMap<int32, FTrackSectionRow>& Rows = Pair.Value;
+
+		for (auto& RowPair : Rows)
+		{
+			FTrackSectionRow& Row = RowPair.Value;
+
+			Row.Sections.Sort([&](UMovieSceneSection& A, UMovieSceneSection& B){
+				return A.GetOverlapPriority() < B.GetOverlapPriority();
+			});
+
+			for (int32 SectionIndex = 1; SectionIndex < Row.Sections.Num(); ++SectionIndex)
+			{
+				UMovieSceneSection* ThisSection = Row.Sections[SectionIndex];
+				if (Row.SectionToReOrder.Contains(ThisSection))
+				{
+					UMovieSceneSection* OtherSection = Row.Sections[SectionIndex - 1];
+
+					Row.Sections.Swap(SectionIndex, SectionIndex - 1);
+
+					const int32 SwappedPriority = OtherSection->GetOverlapPriority();
+					OtherSection->SetOverlapPriority(ThisSection->GetOverlapPriority());
+					ThisSection->SetOverlapPriority(SwappedPriority);
+				}
+			}
+		}
+	}
+
+	Sequencer->SetGlobalTime(Sequencer->GetGlobalTime());
+}
+
+
 bool FPasteContextMenu::BuildMenu(FMenuBuilder& MenuBuilder, FSequencer& InSequencer, const FPasteContextMenuArgs& Args)
 {
 	TSharedRef<FPasteContextMenu> Menu = MakeShareable(new FPasteContextMenu(InSequencer, Args));
diff --git a/Engine/Source/Editor/Sequencer/Private/SequencerContextMenus.h b/Engine/Source/Editor/Sequencer/Private/SequencerContextMenus.h
--- a/Engine/Source/Editor/Sequencer/Private/SequencerContextMenus.h
+++ b/Engine/Source/Editor/Sequencer/Private/SequencerContextMenus.h
@@ -57,6 +57,14 @@ private:
 
 	void DeleteSection();
 
+	void BringToFront();
+
+	void SendToBack();
+
+	void BringForward();
+
+	void SendBackward();
+
 	/** The sequencer */
 	TSharedRef<FSequencer> Sequencer;
 
diff --git a/Engine/Source/Editor/Sequencer/Private/SequencerHotspots.cpp b/Engine/Source/Editor/Sequencer/Private/SequencerHotspots.cpp
--- a/Engine/Source/Editor/Sequencer/Private/SequencerHotspots.cpp
+++ b/Engine/Source/Editor/Sequencer/Private/SequencerHotspots.cpp
@@ -18,7 +18,28 @@ void FSectionHotspot::PopulateContextMenu(FMenuBuilder& MenuBuilder, ISequencer&
 	FSequencer& Sequencer = static_cast<FSequencer&>(InSequencer);
 
 	TSharedPtr<ISequencerSection> SectionInterface = Section.TrackNode->GetSections()[Section.SectionIndex];
-	SectionInterface->BuildSectionContextMenu(MenuBuilder);
+
+	FGuid ObjectBinding;
+	if (Section.TrackNode.IsValid())
+	{
+		TSharedPtr<FSequencerDisplayNode> ParentNode = Section.TrackNode;
+
+		while (ParentNode.IsValid())
+		{
+			if (ParentNode.Get()->GetType() == ESequencerNode::Object)
+			{
+				TSharedPtr<FSequencerObjectBindingNode> ObjectNode = StaticCastSharedPtr<FSequencerObjectBindingNode>(ParentNode);
+				if (ObjectNode.IsValid())
+				{
+					ObjectBinding = ObjectNode.Get()->GetObjectBinding();
+					break;
+				}
+			}
+			ParentNode = ParentNode.Get()->GetParent();
+		}
+	}
+
+	SectionInterface->BuildSectionContextMenu(MenuBuilder, ObjectBinding);
 
 	FSectionContextMenu::BuildMenu(MenuBuilder, Sequencer, MouseDownTime);
 }
diff --git a/Engine/Source/Editor/Sequencer/Private/SequencerNodeTree.cpp b/Engine/Source/Editor/Sequencer/Private/SequencerNodeTree.cpp
--- a/Engine/Source/Editor/Sequencer/Private/SequencerNodeTree.cpp
+++ b/Engine/Source/Editor/Sequencer/Private/SequencerNodeTree.cpp
@@ -209,19 +209,20 @@ TSharedRef<FSequencerObjectBindingNode> FSequencerNodeTree::AddObjectBinding(con
 		TSharedPtr<FSequencerObjectBindingNode> ParentNode;
 
 		UMovieSceneSequence* Sequence = Sequencer.GetFocusedMovieSceneSequence();
+		TSharedRef<FMovieSceneSequenceInstance> SequenceInstance = Sequencer.GetFocusedMovieSceneSequenceInstance();
 
 		// Prefer to use the parent spawnable if possible, rather than relying on runtime object presence
 		FMovieScenePossessable* Possessable = Sequence->GetMovieScene()->FindPossessable(ObjectBinding);
-		if (Possessable && Possessable->GetParentSpawnable().IsValid())
+		if (Possessable && Possessable->GetParent().IsValid())
 		{
-			const FMovieSceneBinding* ParentBinding = GuidToBindingMap.FindRef(Possessable->GetParentSpawnable());
+			const FMovieSceneBinding* ParentBinding = GuidToBindingMap.FindRef(Possessable->GetParent());
 			if (ParentBinding)
 			{
-				ParentNode = AddObjectBinding( ParentBinding->GetName(), Possessable->GetParentSpawnable(), GuidToBindingMap, OutNodeList );
+				ParentNode = AddObjectBinding( ParentBinding->GetName(), Possessable->GetParent(), GuidToBindingMap, OutNodeList );
 			}
 		}
 		
-		UObject* RuntimeObject = Sequence->FindObject(ObjectBinding);
+		UObject* RuntimeObject = SequenceInstance->FindObject(ObjectBinding, Sequencer);
 
 		// fallback to using the parent runtime object
 		if (!ParentNode.IsValid() && RuntimeObject)
@@ -230,7 +231,7 @@ TSharedRef<FSequencerObjectBindingNode> FSequencerNodeTree::AddObjectBinding(con
 
 			if (ParentObject != nullptr)
 			{
-				FGuid ParentBinding = Sequence->FindObjectId(*ParentObject);
+				FGuid ParentBinding = SequenceInstance->FindObjectId(*ParentObject);
 				TSharedPtr<FSequencerObjectBindingNode>* FoundParentNode = ObjectBindingMap.Find( ParentBinding );
 				if ( FoundParentNode != nullptr )
 				{
@@ -314,17 +315,40 @@ bool FSequencerNodeTree::IsNodeFiltered( const TSharedRef<const FSequencerDispla
  * @param FilterStrings		The filter strings which need to be matched
  * @param OutFilteredNodes	The list of all filtered nodes
  */
-static bool FilterNodesRecursive( const TSharedRef<FSequencerDisplayNode>& StartNode, const TArray<FString>& FilterStrings, TSet<TSharedRef<const FSequencerDisplayNode>>& OutFilteredNodes )
+static bool FilterNodesRecursive( FSequencer& Sequencer, const TSharedRef<FSequencerDisplayNode>& StartNode, const TArray<FString>& FilterStrings, TSet<TSharedRef<const FSequencerDisplayNode>>& OutFilteredNodes )
 {
 	// assume the filter is acceptable
 	bool bFilterAcceptable = true;
 
 	// check each string in the filter strings list against 
+	UMovieScene* MovieScene = Sequencer.GetFocusedMovieSceneSequence()->GetMovieScene();
+
 	for (const FString& String : FilterStrings)
 	{
-		if (!StartNode->GetDisplayName().ToString().Contains(String)) 
+		if (String.StartsWith(TEXT("label:")))
+		{
+			if (StartNode->GetType() == ESequencerNode::Object)
+			{
+				auto ObjectBindingNode = StaticCastSharedRef<FSequencerObjectBindingNode>(StartNode);
+				auto Labels = MovieScene->GetObjectLabels(ObjectBindingNode->GetObjectBinding());
+
+				if (!Labels.Strings.Contains(String.RightChop(6)))
+				{
+					bFilterAcceptable = false;
+				}
+			}
+			else
+			{
+				bFilterAcceptable = false;
+			}
+		}
+		else if (!StartNode->GetDisplayName().ToString().Contains(String)) 
 		{
 			bFilterAcceptable = false;
+		}
+
+		if (!bFilterAcceptable)
+		{
 			break;
 		}
 	}
@@ -345,7 +369,7 @@ static bool FilterNodesRecursive( const TSharedRef<FSequencerDisplayNode>& Start
 	for (const auto& Node : ChildNodes)
 	{
 		// Mark the parent as filtered if any child node was filtered
-		bFilterAcceptable |= FilterNodesRecursive(Node, FilterStrings, OutFilteredNodes);
+		bFilterAcceptable |= FilterNodesRecursive(Sequencer, Node, FilterStrings, OutFilteredNodes);
 
 		if (bFilterAcceptable && !bInFilter)
 		{
@@ -381,7 +405,7 @@ void FSequencerNodeTree::FilterNodes(const FString& InFilter)
 		for (auto It = ObjectBindingMap.CreateIterator(); It; ++It)
 		{
 			// Recursively filter all nodes, matching them against the list of filter strings.  All filter strings must be matched
-			FilterNodesRecursive(It.Value().ToSharedRef(), FilterStrings, FilteredNodes);
+			FilterNodesRecursive(Sequencer, It.Value().ToSharedRef(), FilterStrings, FilteredNodes);
 		}
 	}
 }
diff --git a/Engine/Source/Editor/Sequencer/Private/SequencerObjectChangeListener.cpp b/Engine/Source/Editor/Sequencer/Private/SequencerObjectChangeListener.cpp
--- a/Engine/Source/Editor/Sequencer/Private/SequencerObjectChangeListener.cpp
+++ b/Engine/Source/Editor/Sequencer/Private/SequencerObjectChangeListener.cpp
@@ -27,6 +27,18 @@ FSequencerObjectChangeListener::~FSequencerObjectChangeListener()
 void FSequencerObjectChangeListener::OnPropertyChanged(const TArray<UObject*>& ChangedObjects, const IPropertyHandle& PropertyHandle) const
 {
 	BroadcastPropertyChanged(FKeyPropertyParams(ChangedObjects, PropertyHandle));
+
+	for (UObject* Object : ChangedObjects)
+	{
+		if (Object)
+		{
+			const FOnObjectPropertyChanged* Event = ObjectToPropertyChangedEvent.Find(Object);
+			if (Event)
+			{
+				Event->Broadcast(*Object);
+			}
+		}
+	}
 }
 
 void FSequencerObjectChangeListener::BroadcastPropertyChanged( FKeyPropertyParams KeyPropertyParams ) const
@@ -108,6 +120,16 @@ FOnPropagateObjectChanges& FSequencerObjectChangeListener::GetOnPropagateObjectC
 	return OnPropagateObjectChanges;
 }
 
+FOnObjectPropertyChanged& FSequencerObjectChangeListener::GetOnAnyPropertyChanged(UObject& Object)
+{
+	return ObjectToPropertyChangedEvent.FindOrAdd(&Object);
+}
+
+void FSequencerObjectChangeListener::ReportObjectDestroyed(UObject& Object)
+{
+	ObjectToPropertyChangedEvent.Remove(&Object);
+}
+
 bool FSequencerObjectChangeListener::FindPropertySetter( const UClass& ObjectClass, const FName PropertyTypeName, const FString& InPropertyVarName, const UStructProperty* StructProperty ) const
 {
 	bool bFound = ClassToPropertyChangedMap.Contains( PropertyTypeName );
diff --git a/Engine/Source/Editor/Sequencer/Private/SequencerObjectChangeListener.h b/Engine/Source/Editor/Sequencer/Private/SequencerObjectChangeListener.h
--- a/Engine/Source/Editor/Sequencer/Private/SequencerObjectChangeListener.h
+++ b/Engine/Source/Editor/Sequencer/Private/SequencerObjectChangeListener.h
@@ -3,6 +3,7 @@
 #pragma once
 
 //#include "ISequencerObjectChangeListener.h"
+#include "ObjectKey.h"
 
 class IPropertyHandle;
 
@@ -18,6 +19,8 @@ public:
 	/** ISequencerObjectChangeListener interface */
 	virtual FOnAnimatablePropertyChanged& GetOnAnimatablePropertyChanged( FName PropertyTypeName ) override;
 	virtual FOnPropagateObjectChanges& GetOnPropagateObjectChanges() override;
+	virtual FOnObjectPropertyChanged& GetOnAnyPropertyChanged(UObject& Object) override;
+	virtual void ReportObjectDestroyed(UObject& Object) override;
 	virtual bool CanKeyProperty(FCanKeyPropertyParams KeyPropertyParams) const override;
 	virtual void KeyProperty(FKeyPropertyParams KeyPropertyParams) const override;
 	virtual void TriggerAllPropertiesChanged(UObject* Object) override;
@@ -69,6 +72,9 @@ private:
 	/** A mapping of property classes to multi-cast delegate that is broadcast when properties of that type change */
 	TMap< FName, FOnAnimatablePropertyChanged > ClassToPropertyChangedMap;
 
+	/** A mapping of object instance to property change event */
+	TMap< FObjectKey, FOnObjectPropertyChanged > ObjectToPropertyChangedEvent;
+
 	/** Delegate to call when object changes should be propagated */
 	FOnPropagateObjectChanges OnPropagateObjectChanges;
 
diff --git a/Engine/Source/Editor/Sequencer/Private/SequencerPrivatePCH.h b/Engine/Source/Editor/Sequencer/Private/SequencerPrivatePCH.h
--- a/Engine/Source/Editor/Sequencer/Private/SequencerPrivatePCH.h
+++ b/Engine/Source/Editor/Sequencer/Private/SequencerPrivatePCH.h
@@ -1,10 +1,9 @@
 // Copyright 1998-2015 Epic Games, Inc. All Rights Reserved.
 
-#ifndef __SequencerPrivatePCH_h__
-#define __SequencerPrivatePCH_h__
+#pragma once
 
-#include "UnrealEd.h"
 
+#include "UnrealEd.h"
 #include "SBreadcrumbTrail.h"
 
 #include "KeyPropertyParams.h"
@@ -33,7 +32,8 @@
 #include "SSequencerCurveEditor.h"
 #include "SSequencerCurveEditorToolBar.h"
 #include "CurveEditorSettings.h"
+#include "SSequencerLabelBrowser.h"
+#include "SSequencerLabelEditor.h"
 
-DECLARE_LOG_CATEGORY_EXTERN(LogSequencer, Log, All);
 
-#endif // __SequencerPrivatePCH_h__
+DECLARE_LOG_CATEGORY_EXTERN(LogSequencer, Log, All);
diff --git a/Engine/Source/Editor/Sequencer/Private/SequencerSettings.cpp b/Engine/Source/Editor/Sequencer/Private/SequencerSettings.cpp
--- a/Engine/Source/Editor/Sequencer/Private/SequencerSettings.cpp
+++ b/Engine/Source/Editor/Sequencer/Private/SequencerSettings.cpp
@@ -23,6 +23,7 @@ USequencerSettings::USequencerSettings( const FObjectInitializer& ObjectInitiali
 	CurveValueSnapInterval = 10.0f;
 	bSnapCurveValueToInterval = true;
 	bDetailsViewVisible = false;
+	bLabelBrowserVisible = false;
 	bAutoScrollEnabled = false;
 	bShowCurveEditor = false;
 	bShowCurveEditorCurveToolTips = true;
@@ -273,11 +274,25 @@ bool USequencerSettings::GetDetailsViewVisible() const
 	return bDetailsViewVisible;
 }
 
-void USequencerSettings::SetDetailsViewVisible(bool InbDetailsViewVisible )
+void USequencerSettings::SetDetailsViewVisible(bool Visible)
 {
-	if (bDetailsViewVisible != InbDetailsViewVisible)
+	if (bDetailsViewVisible != Visible)
 	{
-		bDetailsViewVisible = InbDetailsViewVisible;
+		bDetailsViewVisible = Visible;
+		SaveConfig();
+	}
+}
+
+bool USequencerSettings::GetLabelBrowserVisible() const
+{
+	return bLabelBrowserVisible;
+}
+
+void USequencerSettings::SetLabelBrowserVisible(bool Visible)
+{
+	if (bLabelBrowserVisible != Visible)
+	{
+		bLabelBrowserVisible = Visible;
 		SaveConfig();
 	}
 }
diff --git a/Engine/Source/Editor/Sequencer/Private/SequencerSettings.h b/Engine/Source/Editor/Sequencer/Private/SequencerSettings.h
--- a/Engine/Source/Editor/Sequencer/Private/SequencerSettings.h
+++ b/Engine/Source/Editor/Sequencer/Private/SequencerSettings.h
@@ -141,9 +141,16 @@ public:
 
 	/** Gets whether or not the details view is visible. */
 	bool GetDetailsViewVisible() const;
+
 	/** Sets whether or not the details view is visible. */
 	void SetDetailsViewVisible(bool InbDetailsViewVisible);
 
+	/** Gets whether or not the label browser is visible. */
+	bool GetLabelBrowserVisible() const;
+
+	/** Sets whether or not the label browser is visible. */
+	void SetLabelBrowserVisible(bool Visible);
+
 	/** Gets whether or not auto-scroll is enabled. */
 	bool GetAutoScrollEnabled() const;
 	/** Sets whether or not auto-scroll is enabled. */
@@ -230,6 +237,9 @@ protected:
 	UPROPERTY( config )
 	bool bDetailsViewVisible;
 
+	UPROPERTY( config )
+	bool bLabelBrowserVisible;
+
 	UPROPERTY( config )
 	bool bAutoScrollEnabled;
 
diff --git a/Engine/Source/Editor/Sequencer/Private/TimeSliderController.cpp b/Engine/Source/Editor/Sequencer/Private/TimeSliderController.cpp
--- a/Engine/Source/Editor/Sequencer/Private/TimeSliderController.cpp
+++ b/Engine/Source/Editor/Sequencer/Private/TimeSliderController.cpp
@@ -291,7 +291,7 @@ int32 FSequencerTimeSliderController::OnPaintTimeSlider( bool bMirrorLabels, con
 		FLinearColor ScrubColor = InWidgetStyle.GetColorAndOpacityTint();
 
 		// @todo Sequencer this color should be specified in the style
-		ScrubColor.A = ScrubColor.A*0.5f;
+		ScrubColor.A = ScrubColor.A*0.75f;
 		ScrubColor.B *= 0.1f;
 		ScrubColor.G *= 0.2f;
 		FSlateDrawElement::MakeBox( 
diff --git a/Engine/Source/Editor/Sequencer/Private/Tools/SequencerEditTool_Movement.cpp b/Engine/Source/Editor/Sequencer/Private/Tools/SequencerEditTool_Movement.cpp
--- a/Engine/Source/Editor/Sequencer/Private/Tools/SequencerEditTool_Movement.cpp
+++ b/Engine/Source/Editor/Sequencer/Private/Tools/SequencerEditTool_Movement.cpp
@@ -30,7 +30,7 @@ FReply FSequencerEditTool_Movement::OnMouseButtonDown(SWidget& OwnerWidget, cons
 
 		DelayedDrag = FDelayedDrag_Hotspot(VirtualTrackArea.CachedTrackAreaGeometry().AbsoluteToLocal(MouseEvent.GetScreenSpacePosition()), MouseEvent.GetEffectingButton(), Hotspot);
 
-		if (Sequencer.Pin()->GetSettings()->GetSnapPlayTimeToDraggedKey())
+ 		if (Sequencer.Pin()->GetSettings()->GetSnapPlayTimeToDraggedKey() || (MouseEvent.IsShiftDown() && MouseEvent.GetEffectingButton() == EKeys::LeftMouseButton) )
 		{
 			if (DelayedDrag->Hotspot.IsValid())
 			{
diff --git a/Engine/Source/Editor/Sequencer/Public/ISequencer.h b/Engine/Source/Editor/Sequencer/Public/ISequencer.h
--- a/Engine/Source/Editor/Sequencer/Public/ISequencer.h
+++ b/Engine/Source/Editor/Sequencer/Public/ISequencer.h
@@ -57,6 +57,11 @@ public:
 	 */
 	virtual void FocusSequenceInstance(TSharedRef<FMovieSceneSequenceInstance> SequenceInstance) = 0;
 	
+	/**
+	 * Create a new binding for the specified object
+	 */
+	virtual FGuid CreateBinding(UObject& InObject, const FString& InName) = 0;
+
 	/**
 	 * Given a sub-movie scene section, returns the instance of the movie scene for that section.
 	 *
@@ -166,6 +171,11 @@ public:
 	 */
 	virtual class ISequencerObjectChangeListener& GetObjectChangeListener() = 0;
 
+	/**
+	 * Get a set of all properties currently keyed on the specified object
+	 */
+	virtual void GetAllKeyedProperties(UObject& Object, TSet<UProperty*>& OutProperties) = 0;
+
 	virtual bool CanKeyProperty(FCanKeyPropertyParams CanKeyPropertyParams) const = 0;
 
 	virtual void KeyProperty(FKeyPropertyParams KeyPropertyParams) = 0;
