Copying //UE4/Dev-Sequencer to Dev-Main (//UE4/Dev-Main)

diff --git a/Engine/Source/Editor/Sequencer/Private/SSequencerTrackLane.cpp b/Engine/Source/Editor/Sequencer/Private/SSequencerTrackLane.cpp
--- a/Engine/Source/Editor/Sequencer/Private/SSequencerTrackLane.cpp
+++ b/Engine/Source/Editor/Sequencer/Private/SSequencerTrackLane.cpp
@@ -4,43 +4,123 @@
 #include "SSequencerTrackLane.h"
 #include "SSequencerTreeView.h"
 
+
 void SSequencerTrackLane::Construct(const FArguments& InArgs, const TSharedRef<FSequencerDisplayNode>& InDisplayNode, const TSharedRef<SSequencerTreeView>& InTreeView)
 {
 	DisplayNode = InDisplayNode;
 	TreeView = InTreeView;
 
 	ChildSlot
-	.HAlign(HAlign_Fill)
-	.Padding(0, DisplayNode->GetNodePadding().Top, 0, 0)
-	[
-		InArgs._Content.Widget
-	];
+		.HAlign(HAlign_Fill)
+		.Padding(0, DisplayNode->GetNodePadding().Top, 0, 0)
+		[
+			InArgs._Content.Widget
+		];
 }
 
-int32 SSequencerTrackLane::OnPaint(const FPaintArgs& Args, const FGeometry& AllottedGeometry, const FSlateRect& MyClippingRect, FSlateWindowElementList& OutDrawElements, int32 LayerId, const FWidgetStyle& InWidgetStyle, bool bParentEnabled) const
+
+void DrawLaneRecursive(const TSharedRef<FSequencerDisplayNode>& DisplayNode, const FGeometry& AllottedGeometry, float& YOffs, const FSlateRect& MyClippingRect, FSlateWindowElementList& OutDrawElements, int32 LayerId, const FWidgetStyle& InWidgetStyle)
 {
-	// Draw a separator below the track
-	auto PaintGeometry = AllottedGeometry.ToPaintGeometry();
+	static const FName HoverBorderName("Sequencer.AnimationOutliner.SelectionBorderHover");
+	static const FName SelectionBorderName("Sequencer.AnimationOutliner.SelectionBorder");
+	static const FName SelectionBorderInactiveName("Sequencer.AnimationOutliner.SelectionBorderInactive");
+	static const FName SelectionColorName("SelectionColor_Pressed");
+	static const FName SelectionColorInactiveName("SelectionColor_Inactive");
 
-	TArray<FVector2D> Points;
-	Points.Emplace(0, PaintGeometry.GetLocalSize().Y);
-	Points.Emplace(PaintGeometry.GetLocalSize().X, PaintGeometry.GetLocalSize().Y);
+	if (DisplayNode->IsHidden())
+	{
+		return;
+	}
+
+	float TotalNodeHeight = DisplayNode->GetNodeHeight() + DisplayNode->GetNodePadding().Combined();
+
+	// draw selection border
+	FSequencerSelection& SequencerSelection = DisplayNode->GetSequencer().GetSelection();
+
+	if (SequencerSelection.IsSelected(DisplayNode))
+	{
+		FLinearColor SelectionColor = FEditorStyle::GetSlateColor(
+			(SequencerSelection.GetActiveSelection() == FSequencerSelection::EActiveSelection::OutlinerNode)
+				? SelectionColorName
+				: SelectionColorInactiveName
+		).GetColor(InWidgetStyle);
+
+		FName SelectionBrush = (SequencerSelection.GetActiveSelection() == FSequencerSelection::EActiveSelection::OutlinerNode)
+			? SelectionBorderName
+			: SelectionBorderInactiveName;
+
+		FSlateDrawElement::MakeBox(
+			OutDrawElements,
+			LayerId,
+			AllottedGeometry.ToPaintGeometry(
+				FVector2D(0, YOffs),
+				FVector2D(AllottedGeometry.Size.X, TotalNodeHeight)
+			),
+			FEditorStyle::GetBrush(SelectionBrush),
+			MyClippingRect,
+			ESlateDrawEffect::None,
+			SelectionColor
+		);
+	}
+
+	// draw hover border
+	else if (DisplayNode->IsHovered())
+	{
+		FSlateDrawElement::MakeBox(
+			OutDrawElements,
+			LayerId,
+			AllottedGeometry.ToPaintGeometry(
+				FVector2D(0, YOffs),
+				FVector2D(AllottedGeometry.Size.X, TotalNodeHeight)
+			),
+			FEditorStyle::GetBrush(HoverBorderName),
+			MyClippingRect,
+			ESlateDrawEffect::None,
+			FLinearColor(1.0f, 1.0f, 1.0f, 0.1f)
+		);
+	}
 
-	static FLinearColor SeparatorColor(0.f, 0.f, 0.f, 0.15f);
+	YOffs += TotalNodeHeight;
+
+/*	// draw lane separator
+	static const FName LaneColorName("Sequencer.TrackArea.LaneColor");
+
+	TArray<FVector2D> Points;
+	{
+		Points.Emplace(0, YOffs);
+		Points.Emplace(PaintGeometry.GetLocalSize().X, YOffs);
+	}
 
 	FSlateDrawElement::MakeLines( 
 		OutDrawElements, 
-		LayerId,
-		PaintGeometry,
+		LayerId + 1,
+		AllottedGeometry.ToPaintGeometry(),
 		Points,
 		MyClippingRect,
 		ESlateDrawEffect::None,
-		SeparatorColor
-		);
+		FEditorStyle::GetSlateColor(LaneColorName).GetColor(InWidgetStyle),
+		false //bAntialias
+	);*/
+
+	if (DisplayNode->IsExpanded())
+	{
+		for (const auto& ChildNode : DisplayNode->GetChildNodes())
+		{
+			DrawLaneRecursive(ChildNode, AllottedGeometry, YOffs, MyClippingRect, OutDrawElements, LayerId, InWidgetStyle);
+		}
+	}
+}
+
+
+int32 SSequencerTrackLane::OnPaint(const FPaintArgs& Args, const FGeometry& AllottedGeometry, const FSlateRect& MyClippingRect, FSlateWindowElementList& OutDrawElements, int32 LayerId, const FWidgetStyle& InWidgetStyle, bool bParentEnabled) const
+{
+	float YOffs = 0;
+	DrawLaneRecursive(DisplayNode.ToSharedRef(), AllottedGeometry, YOffs, MyClippingRect, OutDrawElements, LayerId, InWidgetStyle);
 
 	return SCompoundWidget::OnPaint(Args, AllottedGeometry, MyClippingRect, OutDrawElements, LayerId + 1, InWidgetStyle, bParentEnabled);
 }
 
+
 float SSequencerTrackLane::GetPhysicalPosition() const
 {
 	// Positioning strategy:
@@ -74,4 +154,4 @@ float SSequencerTrackLane::GetPhysicalPosition() const
 	}
 
 	return NegativeOffset + Top.GetValue();
-}
\ No newline at end of file
+}
diff --git a/Engine/Source/Editor/Sequencer/Private/Sequencer.cpp b/Engine/Source/Editor/Sequencer/Private/Sequencer.cpp
--- a/Engine/Source/Editor/Sequencer/Private/Sequencer.cpp
+++ b/Engine/Source/Editor/Sequencer/Private/Sequencer.cpp
@@ -53,136 +53,128 @@
 #include "MovieSceneCaptureDialogModule.h"
 #include "AutomatedLevelSequenceCapture.h"
 #include "MovieSceneCommonHelpers.h"
-
 #include "SceneOutlinerModule.h"
 #include "SceneOutlinerPublicTypes.h"
 
+
 #define LOCTEXT_NAMESPACE "Sequencer"
 
 DEFINE_LOG_CATEGORY(LogSequencer);
 
 
-bool FSequencer::IsSequencerEnabled()
-{
-	return true;
-}
-
 namespace
 {
 	TSharedPtr<FSequencer> ActiveSequencer;
 }
 
+
 void FSequencer::InitSequencer(const FSequencerInitParams& InitParams, const TSharedRef<ISequencerObjectChangeListener>& InObjectChangeListener, const TSharedRef<IDetailKeyframeHandler>& InDetailKeyframeHandler, const TArray<FOnCreateTrackEditor>& TrackEditorDelegates)
 {
-	if( IsSequencerEnabled() )
-	{
-		bIsEditingWithinLevelEditor = InitParams.bEditWithinLevelEditor;
+	bIsEditingWithinLevelEditor = InitParams.bEditWithinLevelEditor;
 		
-		if (InitParams.SpawnRegister.IsValid())
-		{
-			SpawnRegister = InitParams.SpawnRegister;
-		}
-		else
-		{
-			// Spawnables not supported
-			SpawnRegister = MakeShareable(new FNullMovieSceneSpawnRegister);
-		}
-
-		// If this sequencer edits the level, close out the existing active sequencer and mark this as the active sequencer.
-		if (bIsEditingWithinLevelEditor)
-		{
-			Settings = USequencerSettingsContainer::GetOrCreate<ULevelEditorSequencerSettings>(*InitParams.ViewParams.UniqueName);
-
-			if (ActiveSequencer.IsValid())
-			{
-				ActiveSequencer->Close();
-			}
-			ActiveSequencer = SharedThis(this);
+	if (InitParams.SpawnRegister.IsValid())
+	{
+		SpawnRegister = InitParams.SpawnRegister;
+	}
+	else
+	{
+		// Spawnables not supported
+		SpawnRegister = MakeShareable(new FNullMovieSceneSpawnRegister);
+	}
 
-			// Register for saving the level so that the state of the scene can be restored before saving and updated after saving.
-			FEditorDelegates::PreSaveWorld.AddSP(this, &FSequencer::OnPreSaveWorld);
-			FEditorDelegates::PostSaveWorld.AddSP(this, &FSequencer::OnPostSaveWorld);
-			FEditorDelegates::PreBeginPIE.AddSP(this, &FSequencer::OnPreBeginPIE);
-			FEditorDelegates::EndPIE.AddSP(this, &FSequencer::OnEndPIE);
+	// If this sequencer edits the level, close out the existing active sequencer and mark this as the active sequencer.
+	if (bIsEditingWithinLevelEditor)
+	{
+		Settings = USequencerSettingsContainer::GetOrCreate<ULevelEditorSequencerSettings>(*InitParams.ViewParams.UniqueName);
 
-			FEditorDelegates::NewCurrentLevel.AddSP(this, &FSequencer::OnNewCurrentLevel);
-			FEditorDelegates::OnMapOpened.AddSP(this, &FSequencer::OnMapOpened);
-		}
-		else
+		if (ActiveSequencer.IsValid())
 		{
-			Settings = USequencerSettingsContainer::GetOrCreate<USequencerSettings>(*InitParams.ViewParams.UniqueName);
+			ActiveSequencer->Close();
 		}
+		ActiveSequencer = SharedThis(this);
 
-		ToolkitHost = InitParams.ToolkitHost;
+		// Register for saving the level so that the state of the scene can be restored before saving and updated after saving.
+		FEditorDelegates::PreSaveWorld.AddSP(this, &FSequencer::OnPreSaveWorld);
+		FEditorDelegates::PostSaveWorld.AddSP(this, &FSequencer::OnPostSaveWorld);
+		FEditorDelegates::PreBeginPIE.AddSP(this, &FSequencer::OnPreBeginPIE);
+		FEditorDelegates::EndPIE.AddSP(this, &FSequencer::OnEndPIE);
 
-		ScrubPosition = InitParams.ViewParams.InitialScrubPosition;
-		ObjectChangeListener = InObjectChangeListener;
-		DetailKeyframeHandler = InDetailKeyframeHandler;
+		FEditorDelegates::NewCurrentLevel.AddSP(this, &FSequencer::OnNewCurrentLevel);
+		FEditorDelegates::OnMapOpened.AddSP(this, &FSequencer::OnMapOpened);
+	}
+	else
+	{
+		Settings = USequencerSettingsContainer::GetOrCreate<USequencerSettings>(*InitParams.ViewParams.UniqueName);
+	}
+
+	ToolkitHost = InitParams.ToolkitHost;
 
-		check( ObjectChangeListener.IsValid() );
+	ScrubPosition = InitParams.ViewParams.InitialScrubPosition;
+	ObjectChangeListener = InObjectChangeListener;
+	DetailKeyframeHandler = InDetailKeyframeHandler;
+
+	check( ObjectChangeListener.IsValid() );
 		
-		UMovieSceneSequence& RootSequence = *InitParams.RootSequence;
+	UMovieSceneSequence& RootSequence = *InitParams.RootSequence;
 		
-		// Focusing the initial movie scene needs to be done before the first time GetFocusedMovieSceneInstane or GetRootMovieSceneInstance is used
-		RootMovieSceneSequenceInstance = MakeShareable(new FMovieSceneSequenceInstance(RootSequence));
-		SequenceInstanceStack.Add(RootMovieSceneSequenceInstance.ToSharedRef());
-
-		// Make internal widgets
-		SequencerWidget = SNew( SSequencer, SharedThis( this ) )
-			.ViewRange( this, &FSequencer::GetViewRange )
-			.ClampRange( this, &FSequencer::GetClampRange )
-			.PlaybackRange( this, &FSequencer::GetPlaybackRange )
-			.OnPlaybackRangeChanged( this, &FSequencer::SetPlaybackRange )
-			.OnBeginPlaybackRangeDrag( this, &FSequencer::OnBeginPlaybackRangeDrag )
-			.OnEndPlaybackRangeDrag( this, &FSequencer::OnEndPlaybackRangeDrag )
-			.ScrubPosition( this, &FSequencer::OnGetScrubPosition )
-			.OnBeginScrubbing( this, &FSequencer::OnBeginScrubbing )
-			.OnEndScrubbing( this, &FSequencer::OnEndScrubbing )
-			.OnScrubPositionChanged( this, &FSequencer::OnScrubPositionChanged )
-			.OnViewRangeChanged( this, &FSequencer::SetViewRange )
-			.OnClampRangeChanged( this, &FSequencer::OnClampRangeChanged )
-			.OnGetAddMenuContent(InitParams.ViewParams.OnGetAddMenuContent)
-			.AddMenuExtender(InitParams.ViewParams.AddMenuExtender);
-
-		// When undo occurs, get a notification so we can make sure our view is up to date
-		GEditor->RegisterForUndo(this);
-
-		if( bIsEditingWithinLevelEditor )
-		{
-			// @todo remove when world-centric mode is added
-			// Hook into the editor's mechanism for checking whether we need live capture of PIE/SIE actor state
-			GEditor->GetActorRecordingState().AddSP(this, &FSequencer::GetActorRecordingState);
-
-			ActivateDetailKeyframeHandler();
-				
-			FPropertyEditorModule& EditModule = FModuleManager::Get().GetModuleChecked<FPropertyEditorModule>("PropertyEditor");
-			EditModule.OnPropertyEditorOpened().AddSP(this, &FSequencer::OnPropertyEditorOpened);
-		}
+	// Focusing the initial movie scene needs to be done before the first time GetFocusedMovieSceneInstane or GetRootMovieSceneInstance is used
+	RootMovieSceneSequenceInstance = MakeShareable(new FMovieSceneSequenceInstance(RootSequence));
+	SequenceInstanceStack.Add(RootMovieSceneSequenceInstance.ToSharedRef());
 
-		// Create tools and bind them to this sequencer
-		for( int32 DelegateIndex = 0; DelegateIndex < TrackEditorDelegates.Num(); ++DelegateIndex )
-		{
-			check( TrackEditorDelegates[DelegateIndex].IsBound() );
-			// Tools may exist in other modules, call a delegate that will create one for us 
-			TSharedRef<ISequencerTrackEditor> TrackEditor = TrackEditorDelegates[DelegateIndex].Execute( SharedThis( this ) );
-			TrackEditors.Add( TrackEditor );
-		}
+	// Make internal widgets
+	SequencerWidget = SNew( SSequencer, SharedThis( this ) )
+		.ViewRange( this, &FSequencer::GetViewRange )
+		.ClampRange( this, &FSequencer::GetClampRange )
+		.PlaybackRange( this, &FSequencer::GetPlaybackRange )
+		.OnPlaybackRangeChanged( this, &FSequencer::SetPlaybackRange )
+		.OnBeginPlaybackRangeDrag( this, &FSequencer::OnBeginPlaybackRangeDrag )
+		.OnEndPlaybackRangeDrag( this, &FSequencer::OnEndPlaybackRangeDrag )
+		.ScrubPosition( this, &FSequencer::OnGetScrubPosition )
+		.OnBeginScrubbing( this, &FSequencer::OnBeginScrubbing )
+		.OnEndScrubbing( this, &FSequencer::OnEndScrubbing )
+		.OnScrubPositionChanged( this, &FSequencer::OnScrubPositionChanged )
+		.OnViewRangeChanged( this, &FSequencer::SetViewRange )
+		.OnClampRangeChanged( this, &FSequencer::OnClampRangeChanged )
+		.OnGetAddMenuContent(InitParams.ViewParams.OnGetAddMenuContent)
+		.AddMenuExtender(InitParams.ViewParams.AddMenuExtender);
+
+	// When undo occurs, get a notification so we can make sure our view is up to date
+	GEditor->RegisterForUndo(this);
 
-		ZoomAnimation = FCurveSequence();
-		ZoomCurve = ZoomAnimation.AddCurve(0.f, 0.2f, ECurveEaseFunction::QuadIn);
-		OverlayAnimation = FCurveSequence();
-		OverlayCurve = OverlayAnimation.AddCurve(0.f, 0.2f, ECurveEaseFunction::QuadIn);
+	if( bIsEditingWithinLevelEditor )
+	{
+		// @todo remove when world-centric mode is added
+		// Hook into the editor's mechanism for checking whether we need live capture of PIE/SIE actor state
+		GEditor->GetActorRecordingState().AddSP(this, &FSequencer::GetActorRecordingState);
 
-		// Update initial movie scene data
-		NotifyMovieSceneDataChanged();
-		UpdateTimeBoundsToFocusedMovieScene();
+		ActivateDetailKeyframeHandler();
+				
+		FPropertyEditorModule& EditModule = FModuleManager::Get().GetModuleChecked<FPropertyEditorModule>("PropertyEditor");
+		EditModule.OnPropertyEditorOpened().AddSP(this, &FSequencer::OnPropertyEditorOpened);
+	}
 
-		// NOTE: Could fill in asset editor commands here!
+	// Create tools and bind them to this sequencer
+	for( int32 DelegateIndex = 0; DelegateIndex < TrackEditorDelegates.Num(); ++DelegateIndex )
+	{
+		check( TrackEditorDelegates[DelegateIndex].IsBound() );
+		// Tools may exist in other modules, call a delegate that will create one for us 
+		TSharedRef<ISequencerTrackEditor> TrackEditor = TrackEditorDelegates[DelegateIndex].Execute( SharedThis( this ) );
+		TrackEditors.Add( TrackEditor );
+	}
 
-		BindSequencerCommands();
+	ZoomAnimation = FCurveSequence();
+	ZoomCurve = ZoomAnimation.AddCurve(0.f, 0.2f, ECurveEaseFunction::QuadIn);
+	OverlayAnimation = FCurveSequence();
+	OverlayCurve = OverlayAnimation.AddCurve(0.f, 0.2f, ECurveEaseFunction::QuadIn);
 
-		ActivateSequencerEditorMode();
-	}
+	// Update initial movie scene data
+	NotifyMovieSceneDataChanged();
+	UpdateTimeBoundsToFocusedMovieScene();
+
+	// NOTE: Could fill in asset editor commands here!
+
+	BindSequencerCommands();
+	ActivateSequencerEditorMode();
 }
 
 
@@ -201,6 +193,7 @@ FSequencer::FSequencer()
 	Selection.GetOnOutlinerNodeSelectionChanged().AddRaw(this, &FSequencer::OnSelectedOutlinerNodesChanged);
 }
 
+
 FSequencer::~FSequencer()
 {
 	GEditor->GetActorRecordingState().RemoveAll( this );
@@ -210,6 +203,7 @@ FSequencer::~FSequencer()
 	SequencerWidget.Reset();
 }
 
+
 void FSequencer::Close()
 {
 	if (ActiveSequencer.IsValid() && ActiveSequencer.Get() == this)
@@ -247,6 +241,7 @@ void FSequencer::Close()
 	SequencerWidget.Reset();
 }
 
+
 void FSequencer::Tick(float InDeltaTime)
 {
 	if (bNeedTreeRefresh)
@@ -300,6 +295,7 @@ UMovieSceneSequence* FSequencer::GetRootMovieSceneSequence() const
 	return nullptr;
 }
 
+
 UMovieSceneSequence* FSequencer::GetFocusedMovieSceneSequence() const
 {
 	// the last item is the focused movie scene
@@ -311,6 +307,7 @@ UMovieSceneSequence* FSequencer::GetFocusedMovieSceneSequence() const
 	return nullptr;
 }
 
+
 void FSequencer::ResetToNewRootSequence(UMovieSceneSequence& NewSequence)
 {
 	if (RootMovieSceneSequenceInstance.IsValid())
@@ -365,6 +362,7 @@ void FSequencer::FocusSequenceInstance(TSharedRef<FMovieSceneSequenceInstance> S
 	SequencerWidget->UpdateBreadcrumbs();
 }
 
+
 FGuid FSequencer::CreateBinding(UObject& InObject, const FString& InName)
 {
 	UMovieSceneSequence* OwnerSequence = GetFocusedMovieSceneSequence();
@@ -404,20 +402,24 @@ FGuid FSequencer::CreateBinding(UObject& InObject, const FString& InName)
 	return PossessableGuid;
 }
 
+
 UObject* FSequencer::GetPlaybackContext() const
 {
 	if (bIsEditingWithinLevelEditor)
 	{
 		return GWorld;
 	}
+
 	return nullptr;
 }
 
+
 void FSequencer::GetAllKeyedProperties(UObject& Object, TSet<UProperty*>& OutProperties)
 {
 	GetAllKeyedPropertiesForInstance(Object, *GetRootMovieSceneSequenceInstance(), OutProperties);
 }
 
+
 void FSequencer::GetAllKeyedPropertiesForInstance(UObject& Object, FMovieSceneSequenceInstance& Instance, TSet<UProperty*>& OutProperties)
 {
 	GetKeyedProperties(Object, Instance, OutProperties);
@@ -435,6 +437,7 @@ void FSequencer::GetAllKeyedPropertiesForInstance(UObject& Object, FMovieSceneSe
 	}
 }
 
+
 void FSequencer::GetKeyedProperties(UObject& Object, FMovieSceneSequenceInstance& Instance, TSet<UProperty*>& OutProperties)
 {
 	FGuid BindingId = Instance.FindObjectId(Object);
@@ -476,6 +479,7 @@ void FSequencer::GetKeyedProperties(UObject& Object, FMovieSceneSequenceInstance
 	}
 }
 
+
 TSharedRef<FMovieSceneSequenceInstance> FSequencer::GetSequenceInstanceForSection(UMovieSceneSection& Section) const
 {
 	return SequenceInstanceBySection.FindChecked(&Section);
@@ -656,6 +660,7 @@ bool FSequencer::CanSnapToFrame() const
 	return bKeysSelected && CanShowFrameNumbers();
 }
 
+
 void FSequencer::NotifyMapChanged( class UWorld* NewWorld, EMapChangeType MapChangeType )
 {
 	// @todo sequencer: We should only wipe/respawn puppets that are affected by the world that is being changed! (multi-UWorld support)
@@ -665,6 +670,7 @@ void FSequencer::NotifyMapChanged( class UWorld* NewWorld, EMapChangeType MapCha
 	}
 }
 
+
 void FSequencer::OnActorsDropped( const TArray<TWeakObjectPtr<AActor> >& Actors )
 {
 	// @todo sequencer: Undo doesn't seem to be working at all
@@ -691,11 +697,14 @@ void FSequencer::OnActorsDropped( const TArray<TWeakObjectPtr<AActor> >& Actors
 	}
 }
 
+
 void FSequencer::NotifyMovieSceneDataChanged()
 {
 	LabelManager.SetMovieScene(SequenceInstanceStack.Top()->GetSequence()->GetMovieScene());
 	PlaybackState = EMovieScenePlayerStatus::Stopped;
 	bNeedTreeRefresh = true;
+
+	UpdatePlaybackRange();
 }
 
 
@@ -712,17 +721,20 @@ FAnimatedRange FSequencer::GetViewRange() const
 	return AnimatedRange;
 }
 
+
 FAnimatedRange FSequencer::GetClampRange() const
 {
 	UMovieScene* FocussedMovieScene = GetFocusedMovieSceneSequence()->GetMovieScene();
 	return FocussedMovieScene->GetEditorData().WorkingRange;
 }
 
+
 TRange<float> FSequencer::GetPlaybackRange() const
 {
 	return GetFocusedMovieSceneSequence()->GetMovieScene()->GetPlaybackRange();
 }
 
+
 void FSequencer::SetPlaybackRange(TRange<float> InNewPlaybackRange)
 {
 	if (ensure(InNewPlaybackRange.HasLowerBound() && InNewPlaybackRange.HasUpperBound() && !InNewPlaybackRange.IsDegenerate()))
@@ -734,61 +746,96 @@ void FSequencer::SetPlaybackRange(TRange<float> InNewPlaybackRange)
 	}
 }
 
+
 void FSequencer::SetStartPlaybackRange()
 {
 	SetPlaybackRange(TRange<float>(GetGlobalTime(), GetPlaybackRange().GetUpperBoundValue()));
 }
 
+
 void FSequencer::SetEndPlaybackRange()
 {
 	SetPlaybackRange(TRange<float>( GetPlaybackRange().GetLowerBoundValue(), GetGlobalTime()));
 }
 
-bool FSequencer::GetAutoKeyEnabled() const 
+
+void FSequencer::UpdatePlaybackRange()
+{
+	if (Settings->ShouldKeepPlayRangeInSectionBounds())
+	{
+		UMovieScene* FocusedMovieScene = GetFocusedMovieSceneSequence()->GetMovieScene();
+		TArray<UMovieSceneSection*> AllSections = FocusedMovieScene->GetAllSections();
+
+		if (AllSections.Num() > 0)
+		{
+			TRange<float> NewBounds = AllSections[0]->GetRange();
+	
+			for (auto MovieSceneSection : AllSections)
+			{
+				NewBounds = TRange<float>( FMath::Min(NewBounds.GetLowerBoundValue(), MovieSceneSection->GetRange().GetLowerBoundValue()),
+											FMath::Max(NewBounds.GetUpperBoundValue(), MovieSceneSection->GetRange().GetUpperBoundValue()) );
+			}
+
+			FocusedMovieScene->SetPlaybackRange(NewBounds.GetLowerBoundValue(), NewBounds.GetUpperBoundValue());
+		}
+	}
+}
+
+
+EAutoKeyMode FSequencer::GetAutoKeyMode() const 
 {
-	return Settings->GetAutoKeyEnabled();
+	return Settings->GetAutoKeyMode();
 }
 
-void FSequencer::SetAutoKeyEnabled(bool bAutoKeyEnabled)
+
+void FSequencer::SetAutoKeyMode(EAutoKeyMode AutoKeyMode)
 {
-	Settings->SetAutoKeyEnabled(bAutoKeyEnabled);
+	Settings->SetAutoKeyMode(AutoKeyMode);
 }
 
+
 bool FSequencer::GetKeyAllEnabled() const 
 {
 	return Settings->GetKeyAllEnabled();
 }
 
+
 void FSequencer::SetKeyAllEnabled(bool bKeyAllEnabled) 
 {
 	Settings->SetKeyAllEnabled(bKeyAllEnabled);
 }
 
+
 bool FSequencer::GetKeyInterpPropertiesOnly() const 
 {
 	return Settings->GetKeyInterpPropertiesOnly();
 }
 
+
 void FSequencer::SetKeyInterpPropertiesOnly(bool bKeyInterpPropertiesOnly) 
 {
 	Settings->SetKeyInterpPropertiesOnly(bKeyInterpPropertiesOnly);
 }
 
+
 EMovieSceneKeyInterpolation FSequencer::GetKeyInterpolation() const
 {
 	return Settings->GetKeyInterpolation();
 }
 
+
 void FSequencer::SetKeyInterpolation(EMovieSceneKeyInterpolation InKeyInterpolation)
 {
 	Settings->SetKeyInterpolation(InKeyInterpolation);
 }
 
+
 bool FSequencer::IsRecordingLive() const 
 {
 	return PlaybackState == EMovieScenePlayerStatus::Recording && GIsPlayInEditorWorld;
 }
 
+
 float FSequencer::GetCurrentLocalTime( UMovieSceneSequence& InMovieSceneSequence )
 {
 	//@todo Sequencer - Nested movie scenes:  Figure out the parent of the passed in movie scene and 
@@ -796,11 +843,13 @@ float FSequencer::GetCurrentLocalTime( UMovieSceneSequence& InMovieSceneSequence
 	return ScrubPosition;
 }
 
+
 float FSequencer::GetGlobalTime()
 {
 	return ScrubPosition;
 }
 
+
 void FSequencer::SetGlobalTime( float NewTime )
 {
 	if (IsAutoScrollEnabled())
@@ -822,6 +871,7 @@ void FSequencer::SetGlobalTime( float NewTime )
 	SetGlobalTimeDirectly(NewTime);
 }
 
+
 void FSequencer::SetGlobalTimeDirectly( float NewTime )
 {
 	float LastTime = ScrubPosition;
@@ -847,6 +897,7 @@ void FSequencer::SetGlobalTimeDirectly( float NewTime )
 	OnGlobalTimeChangedDelegate.Broadcast();
 }
 
+
 void FSequencer::UpdateAutoScroll(float NewTime)
 {
 	float ThresholdPercentage = 0.025f;
@@ -888,6 +939,7 @@ void FSequencer::UpdateAutoScroll(float NewTime)
 	AutoscrubOffset = AutoscrollOffset;
 }
 
+
 TOptional<float> FSequencer::CalculateAutoscrollEncroachment(float NewTime, float ThresholdPercentage) const
 {
 	enum class EDirection { Positive, Negative };
@@ -913,16 +965,66 @@ TOptional<float> FSequencer::CalculateAutoscrollEncroachment(float NewTime, floa
 	}
 }
 
+
 void FSequencer::SetPerspectiveViewportPossessionEnabled(bool bEnabled)
 {
 	bPerspectiveViewportPossessionEnabled = bEnabled;
 }
 
+
 void FSequencer::SetPerspectiveViewportCameraCutEnabled(bool bEnabled)
 {
 	bPerspectiveViewportCameraCutEnabled = bEnabled;
 }
 
+/** Attempt to find an object binding ID that relates to an unspawned spawnable object */
+FGuid FindUnspawnedObjectGuid(UObject& InObject, UMovieSceneSequence& Sequence)
+{
+	UMovieScene* MovieScene = Sequence.GetMovieScene();
+
+	// If the object is a CDO, the it relates to an unspawned spawnable.
+	UObject* ParentObject = Sequence.GetParentObject(&InObject);
+	if (ParentObject && ParentObject->HasAnyFlags(RF_ClassDefaultObject))
+	{
+		// Find the unspawned spawnable object's CDO
+		UClass* ObjectClass = ParentObject->GetClass();
+
+		FMovieSceneSpawnable* ParentSpawnable = MovieScene->FindSpawnable([&](FMovieSceneSpawnable& InSpawnable){
+			return InSpawnable.GetClass() == ObjectClass;
+		});
+
+		if (ParentSpawnable)
+		{
+			UObject* ParentContext = ObjectClass->GetDefaultObject();
+
+			// The only way to find the object now is to resolve all the child bindings, and see if they are the same
+			for (const FGuid& ChildGuid : ParentSpawnable->GetChildPossessables())
+			{
+				UObject* ChildObject = Sequence.FindPossessableObject(ChildGuid, ParentContext);
+				if (ChildObject == &InObject)
+				{
+					return ChildGuid;
+				}
+			}
+		}
+	}
+	else if (InObject.HasAnyFlags(RF_ClassDefaultObject))
+	{
+		UClass* ObjectClass = InObject.GetClass();
+
+		FMovieSceneSpawnable* SpawnableByClass = MovieScene->FindSpawnable([&](FMovieSceneSpawnable& InSpawnable){
+			return InSpawnable.GetClass() == ObjectClass;
+		});
+
+		if (SpawnableByClass)
+		{
+			return SpawnableByClass->GetGuid();
+		}
+	}
+
+	return FGuid();
+}
+
 FGuid FSequencer::GetHandleToObject( UObject* Object, bool bCreateHandleIfMissing )
 {
 	if (Object == nullptr)
@@ -954,31 +1056,36 @@ FGuid FSequencer::GetHandleToObject( UObject* Object, bool bCreateHandleIfMissin
 		ObjectGuid.Invalidate();
 	}
 
-	bool bPossessableAdded = false;
+	if (!ObjectGuid.IsValid())
+	{
+		ObjectGuid = FindUnspawnedObjectGuid(*Object, *FocusedMovieSceneSequence);
+	}
+
+	if (ObjectGuid.IsValid())
+	{
+		return ObjectGuid;
+	}
 
 	// If the object guid was not found attempt to add it
 	// Note: Only possessed actors can be added like this
-	if (!ObjectGuid.IsValid() && FocusedMovieSceneSequence->CanPossessObject(*Object) && bCreateHandleIfMissing)
+	if (FocusedMovieSceneSequence->CanPossessObject(*Object) && bCreateHandleIfMissing)
 	{
 		AActor* PossessedActor = Cast<AActor>(Object);
 
 		ObjectGuid = CreateBinding(*Object, PossessedActor != nullptr ? PossessedActor->GetActorLabel() : Object->GetName());
-		bPossessableAdded = true;
-	}
-
-	if (bPossessableAdded)
-	{
 		NotifyMovieSceneDataChanged();
 	}
 	
 	return ObjectGuid;
 }
 
+
 ISequencerObjectChangeListener& FSequencer::GetObjectChangeListener()
 { 
 	return *ObjectChangeListener;
 }
 
+
 void FSequencer::GetRuntimeObjects( TSharedRef<FMovieSceneSequenceInstance> MovieSceneInstance, const FGuid& ObjectHandle, TArray< UObject*>& OutObjects ) const
 {
 	UObject* FoundObject = MovieSceneInstance->FindObject(ObjectHandle, *this);
@@ -988,6 +1095,7 @@ void FSequencer::GetRuntimeObjects( TSharedRef<FMovieSceneSequenceInstance> Movi
 	}
 }
 
+
 void FSequencer::UpdateCameraCut(UObject* CameraObject, UObject* UnlockIfCameraObject) const
 {
 	if (!IsPerspectiveViewportCameraCutEnabled())
@@ -1011,6 +1119,7 @@ void FSequencer::UpdateCameraCut(UObject* CameraObject, UObject* UnlockIfCameraO
 	}
 }
 
+
 void FSequencer::SetViewportSettings(const TMap<FViewportClient*, EMovieSceneViewportParams>& ViewportParamsMap)
 {
 	if(!IsPerspectiveViewportPossessionEnabled())
@@ -1045,6 +1154,7 @@ void FSequencer::SetViewportSettings(const TMap<FViewportClient*, EMovieSceneVie
 	}
 }
 
+
 void FSequencer::GetViewportSettings(TMap<FViewportClient*, EMovieSceneViewportParams>& ViewportParamsMap) const
 {
 	for(FLevelEditorViewportClient* LevelVC : GEditor->LevelViewportClients)
@@ -1061,11 +1171,13 @@ void FSequencer::GetViewportSettings(TMap<FViewportClient*, EMovieSceneViewportP
 	}
 }
 
+
 EMovieScenePlayerStatus::Type FSequencer::GetPlaybackStatus() const
 {
 	return PlaybackState;
 }
 
+
 void FSequencer::AddOrUpdateMovieSceneInstance( UMovieSceneSection& MovieSceneSection, TSharedRef<FMovieSceneSequenceInstance> InstanceToAdd )
 {
 	if( !SequenceInstanceBySection.Contains( &MovieSceneSection ) )
@@ -1074,6 +1186,7 @@ void FSequencer::AddOrUpdateMovieSceneInstance( UMovieSceneSection& MovieSceneSe
 	}
 }
 
+
 void FSequencer::RemoveMovieSceneInstance( UMovieSceneSection& MovieSceneSection, TSharedRef<FMovieSceneSequenceInstance> InstanceToRemove )
 {
 	UMovieSceneSequence* Sequence = InstanceToRemove->GetSequence();
@@ -1083,6 +1196,7 @@ void FSequencer::RemoveMovieSceneInstance( UMovieSceneSection& MovieSceneSection
 	SequenceInstanceBySection.Remove( &MovieSceneSection );
 }
 
+
 void FSequencer::AddReferencedObjects( FReferenceCollector& Collector )
 {
 	Collector.AddReferencedObject( Settings );
@@ -1094,6 +1208,7 @@ void FSequencer::AddReferencedObjects( FReferenceCollector& Collector )
 	}
 }
 
+
 void FSequencer::ResetPerMovieSceneData()
 {
 	//@todo Sequencer - We may want to preserve selections when moving between movie scenes
@@ -1103,6 +1218,7 @@ void FSequencer::ResetPerMovieSceneData()
 	//  needed for audio track decompression
 }
 
+
 void FSequencer::UpdateRuntimeInstances()
 {
 	// Refresh the current root instance
@@ -1123,6 +1239,7 @@ void FSequencer::UpdateRuntimeInstances()
 	GEditor->RedrawLevelEditingViewports();
 }
 
+
 FReply FSequencer::OnPlay(bool bTogglePlay)
 {
 	if( (PlaybackState == EMovieScenePlayerStatus::Playing ||
@@ -1143,6 +1260,7 @@ FReply FSequencer::OnPlay(bool bTogglePlay)
 	return FReply::Handled();
 }
 
+
 FReply FSequencer::OnRecord()
 {
 	if (PlaybackState != EMovieScenePlayerStatus::Recording)
@@ -1164,6 +1282,7 @@ FReply FSequencer::OnRecord()
 	return FReply::Handled();
 }
 
+
 FReply FSequencer::OnStepForward()
 {
 	PlaybackState = EMovieScenePlayerStatus::Stopped;
@@ -1176,6 +1295,7 @@ FReply FSequencer::OnStepForward()
 	return FReply::Handled();
 }
 
+
 FReply FSequencer::OnStepBackward()
 {
 	PlaybackState = EMovieScenePlayerStatus::Stopped;
@@ -1188,6 +1308,7 @@ FReply FSequencer::OnStepBackward()
 	return FReply::Handled();
 }
 
+
 FReply FSequencer::OnStepToEnd()
 {
 	PlaybackState = EMovieScenePlayerStatus::Stopped;
@@ -1195,6 +1316,7 @@ FReply FSequencer::OnStepToEnd()
 	return FReply::Handled();
 }
 
+
 FReply FSequencer::OnStepToBeginning()
 {
 	PlaybackState = EMovieScenePlayerStatus::Stopped;
@@ -1202,17 +1324,20 @@ FReply FSequencer::OnStepToBeginning()
 	return FReply::Handled();
 }
 
+
 FReply FSequencer::OnToggleLooping()
 {
 	Settings->SetLooping(!Settings->IsLooping());
 	return FReply::Handled();
 }
 
+
 bool FSequencer::IsLooping() const
 {
 	return Settings->IsLooping();
 }
 
+
 void FSequencer::SetGlobalTimeLooped(float InTime)
 {
 	if (Settings->IsLooping())
@@ -1262,11 +1387,13 @@ void FSequencer::SetGlobalTimeLooped(float InTime)
 	SetGlobalTime(InTime);
 }
 
+
 bool FSequencer::CanShowFrameNumbers() const
 {
 	return SequencerSnapValues::IsTimeSnapIntervalFrameRate(Settings->GetTimeSnapInterval());
 }
 
+
 EPlaybackMode::Type FSequencer::GetPlaybackMode() const
 {
 	return PlaybackState == EMovieScenePlayerStatus::Playing ? EPlaybackMode::PlayingForward :
@@ -1274,6 +1401,7 @@ EPlaybackMode::Type FSequencer::GetPlaybackMode() const
 		EPlaybackMode::Stopped;
 }
 
+
 void FSequencer::UpdateTimeBoundsToFocusedMovieScene()
 {
 	UMovieScene* FocusedMovieScene = GetFocusedMovieSceneSequence()->GetMovieScene();
@@ -1304,6 +1432,7 @@ void FSequencer::UpdateTimeBoundsToFocusedMovieScene()
 	}
 }
 
+
 TRange<float> FSequencer::GetTimeBounds() const
 {
 	const UMovieSceneSequence* FocusedSequence = GetFocusedMovieSceneSequence();
@@ -1318,6 +1447,7 @@ TRange<float> FSequencer::GetTimeBounds() const
 	return FocusedSequence->GetMovieScene()->GetPlaybackRange();
 }
 
+
 void FSequencer::SetViewRange(TRange<float> NewViewRange, EViewRangeInterpolation Interpolation)
 {
 	if (!ensure(NewViewRange.HasUpperBound() && NewViewRange.HasLowerBound() && !NewViewRange.IsDegenerate()))
@@ -1367,6 +1497,7 @@ void FSequencer::SetViewRange(TRange<float> NewViewRange, EViewRangeInterpolatio
 	}
 }
 
+
 void FSequencer::OnClampRangeChanged( TRange<float> NewClampRange )
 {
 	if (!NewClampRange.IsEmpty())
@@ -1375,6 +1506,7 @@ void FSequencer::OnClampRangeChanged( TRange<float> NewClampRange )
 	}
 }
 
+
 void FSequencer::OnScrubPositionChanged( float NewScrubPosition, bool bScrubbing )
 {
 	bool bClampToViewRange = true;
@@ -1408,12 +1540,14 @@ void FSequencer::OnScrubPositionChanged( float NewScrubPosition, bool bScrubbing
 	SetGlobalTimeDirectly( NewScrubPosition );
 }
 
+
 void FSequencer::OnBeginScrubbing()
 {
 	PlaybackState = EMovieScenePlayerStatus::Scrubbing;
 	SequencerWidget->RegisterActiveTimerForPlayback();
 }
 
+
 void FSequencer::OnEndScrubbing()
 {
 	PlaybackState = EMovieScenePlayerStatus::Stopped;
@@ -1421,36 +1555,43 @@ void FSequencer::OnEndScrubbing()
 	StopAutoscroll();
 }
 
+
 void FSequencer::OnBeginPlaybackRangeDrag()
 {
 	GEditor->BeginTransaction(LOCTEXT("SetPlaybackRange_Transaction", "Set playback range"));
 }
 
+
 void FSequencer::OnEndPlaybackRangeDrag()
 {
 	GEditor->EndTransaction();
 }
 
+
 void FSequencer::StartAutoscroll(float UnitsPerS)
 {
 	AutoscrollOffset = UnitsPerS;
 }
 
+
 void FSequencer::StopAutoscroll()
 {
 	AutoscrollOffset.Reset();
 }
 
+
 void FSequencer::OnToggleAutoScroll()
 {
 	Settings->SetAutoScrollEnabled(!Settings->GetAutoScrollEnabled());
 }
 
+
 bool FSequencer::IsAutoScrollEnabled() const
 {
 	return Settings->GetAutoScrollEnabled();
 }
 
+
 void FSequencer::FindInContentBrowser()
 {
 	if (GetFocusedMovieSceneSequence())
@@ -1462,6 +1603,7 @@ void FSequencer::FindInContentBrowser()
 	}
 }
 
+
 UObject* FSequencer::GetCurrentAsset() const
 {
 	// For now we find the asset by looking at the root movie scene's outer.
@@ -1469,11 +1611,13 @@ UObject* FSequencer::GetCurrentAsset() const
 	return GetFocusedMovieSceneSequence()->GetMovieScene()->GetOuter();
 }
 
+
 void FSequencer::VerticalScroll(float ScrollAmountUnits)
 {
 	SequencerWidget->GetTreeView()->ScrollByDelta(ScrollAmountUnits);
 }
 
+
 FGuid FSequencer::AddSpawnableForAssetOrClass( UObject* Object )
 {
 	UMovieSceneSequence* Sequence = GetFocusedMovieSceneSequence();
@@ -1603,17 +1747,17 @@ FGuid FSequencer::AddSpawnableForAssetOrClass( UObject* Object )
 
 				UMovieScene3DTransformSection* NewSection = CastChecked<UMovieScene3DTransformSection>(TransformTrack->CreateNewSection());
 
-				NewSection->SetDefault( FTransformKey( EKey3DTransformChannel::Translation, EKey3DTransformChannel::Default, EAxis::X, NewTransform.Translation.X, bUnwindRotation ) );
-				NewSection->SetDefault( FTransformKey( EKey3DTransformChannel::Translation, EKey3DTransformChannel::Default, EAxis::Y, NewTransform.Translation.Y, bUnwindRotation ) );
-				NewSection->SetDefault( FTransformKey( EKey3DTransformChannel::Translation, EKey3DTransformChannel::Default, EAxis::Z, NewTransform.Translation.Z, bUnwindRotation ) );
+				NewSection->SetDefault( FTransformKey( EKey3DTransformChannel::Translation, EAxis::X, NewTransform.Translation.X, bUnwindRotation ) );
+				NewSection->SetDefault( FTransformKey( EKey3DTransformChannel::Translation, EAxis::Y, NewTransform.Translation.Y, bUnwindRotation ) );
+				NewSection->SetDefault( FTransformKey( EKey3DTransformChannel::Translation, EAxis::Z, NewTransform.Translation.Z, bUnwindRotation ) );
 
-				NewSection->SetDefault( FTransformKey( EKey3DTransformChannel::Rotation, EKey3DTransformChannel::Default, EAxis::X, NewTransform.Rotation.Euler().X, bUnwindRotation ) );
-				NewSection->SetDefault( FTransformKey( EKey3DTransformChannel::Rotation, EKey3DTransformChannel::Default, EAxis::Y, NewTransform.Rotation.Euler().Y, bUnwindRotation ) );
-				NewSection->SetDefault( FTransformKey( EKey3DTransformChannel::Rotation, EKey3DTransformChannel::Default, EAxis::Z, NewTransform.Rotation.Euler().Z, bUnwindRotation ) );
+				NewSection->SetDefault( FTransformKey( EKey3DTransformChannel::Rotation, EAxis::X, NewTransform.Rotation.Euler().X, bUnwindRotation ) );
+				NewSection->SetDefault( FTransformKey( EKey3DTransformChannel::Rotation, EAxis::Y, NewTransform.Rotation.Euler().Y, bUnwindRotation ) );
+				NewSection->SetDefault( FTransformKey( EKey3DTransformChannel::Rotation, EAxis::Z, NewTransform.Rotation.Euler().Z, bUnwindRotation ) );
 
-				NewSection->SetDefault( FTransformKey( EKey3DTransformChannel::Scale, EKey3DTransformChannel::Default, EAxis::X, NewTransform.Scale.X, bUnwindRotation ) );
-				NewSection->SetDefault( FTransformKey( EKey3DTransformChannel::Scale, EKey3DTransformChannel::Default, EAxis::Y, NewTransform.Scale.Y, bUnwindRotation ) );
-				NewSection->SetDefault( FTransformKey( EKey3DTransformChannel::Scale, EKey3DTransformChannel::Default, EAxis::Z, NewTransform.Scale.Z, bUnwindRotation ) );
+				NewSection->SetDefault( FTransformKey( EKey3DTransformChannel::Scale, EAxis::X, NewTransform.Scale.X, bUnwindRotation ) );
+				NewSection->SetDefault( FTransformKey( EKey3DTransformChannel::Scale, EAxis::Y, NewTransform.Scale.Y, bUnwindRotation ) );
+				NewSection->SetDefault( FTransformKey( EKey3DTransformChannel::Scale, EAxis::Z, NewTransform.Scale.Z, bUnwindRotation ) );
 
 				TransformTrack->AddSection(*NewSection);
 			}
@@ -1623,6 +1767,7 @@ FGuid FSequencer::AddSpawnableForAssetOrClass( UObject* Object )
 	return NewSpawnableGuid;
 }
 
+
 void FSequencer::AddSubSequence(UMovieSceneSequence* Sequence)
 {
 	// @todo Sequencer - sub-moviescenes This should be moved to the sub-moviescene editor
@@ -1639,6 +1784,7 @@ void FSequencer::AddSubSequence(UMovieSceneSequence* Sequence)
 	SubTrack->AddSequence(*Sequence, ScrubPosition);
 }
 
+
 bool FSequencer::OnHandleAssetDropped(UObject* DroppedAsset, const FGuid& TargetObjectGuid)
 {
 	bool bWasConsumed = false;
@@ -1657,6 +1803,7 @@ bool FSequencer::OnHandleAssetDropped(UObject* DroppedAsset, const FGuid& Target
 	return bWasConsumed;
 }
 
+
 // Takes a display node and traverses it's parents to find the nearest track node if any.  Also collects the names of the nodes which make
 // up the path from the track node to the display node being checked.  The name path includes the name of the node being checked, but not
 // the name of the track node.
@@ -1680,6 +1827,7 @@ void GetParentTrackNodeAndNamePath(TSharedRef<const FSequencerDisplayNode> Displ
 	}
 }
 
+
 void FSequencer::OnRequestNodeDeleted( TSharedRef<const FSequencerDisplayNode> NodeToBeDeleted )
 {
 	bool bAnythingRemoved = false;
@@ -1777,6 +1925,7 @@ void FSequencer::OnRequestNodeDeleted( TSharedRef<const FSequencerDisplayNode> N
 	}
 }
 
+
 void FSequencer::PlaceActorInFrontOfCamera( AActor* ActorCDO )
 {
 	// Place the actor in front of the active perspective camera if we have one
@@ -1830,17 +1979,20 @@ void FSequencer::GetActorRecordingState( bool& bIsRecording /* In+Out */ ) const
 	}
 }
 
+
 void FSequencer::PostUndo(bool bSuccess)
 {
 	NotifyMovieSceneDataChanged();
 }
 
+
 void FSequencer::OnPreSaveWorld(uint32 SaveFlags, class UWorld* World)
 {
 	// Restore the saved state so that the level save can save that instead of the animated state.
 	RootMovieSceneSequenceInstance->RestoreState(*this);
 }
 
+
 void FSequencer::OnPostSaveWorld(uint32 SaveFlags, class UWorld* World, bool bSuccess)
 {
 	// Reset the time after saving so that an update will be triggered to put objects back to their animated state.
@@ -1848,11 +2000,13 @@ void FSequencer::OnPostSaveWorld(uint32 SaveFlags, class UWorld* World, bool bSu
 	SetGlobalTime(GetGlobalTime());
 }
 
+
 void FSequencer::OnNewCurrentLevel()
 {
 	ActivateSequencerEditorMode();
 }
 
+
 void FSequencer::OnMapOpened(const FString& Filename, bool bLoadAsTemplate)
 {
 	ActivateSequencerEditorMode();
@@ -1872,16 +2026,19 @@ void FSequencer::ActivateSequencerEditorMode()
 	SequencerEdMode->SetSequencer(ActiveSequencer.Get());
 }
 
+
 void FSequencer::OnPreBeginPIE(bool bIsSimulating)
 {
 	RootMovieSceneSequenceInstance->RestoreState(*this);
 }
 
+
 void FSequencer::OnEndPIE(bool bIsSimulating)
 {
 	UpdateRuntimeInstances();
 }
 
+
 void FSequencer::ActivateDetailKeyframeHandler()
 {			
 	// Add sequencer detail keyframe handler
@@ -1899,6 +2056,7 @@ void FSequencer::ActivateDetailKeyframeHandler()
 	}
 }
 
+
 void FSequencer::DeactivateDetailKeyframeHandler()
 {
 	static const FName DetailsTabIdentifiers[] = { "LevelEditorSelectionDetails", "LevelEditorSelectionDetails2", "LevelEditorSelectionDetails3", "LevelEditorSelectionDetails4" };
@@ -1918,11 +2076,13 @@ void FSequencer::DeactivateDetailKeyframeHandler()
 	}
 }
 
+
 void FSequencer::OnPropertyEditorOpened()
 {
 	ActivateDetailKeyframeHandler();
 }
 
+
 void FSequencer::UpdatePreviewLevelViewportClientFromCameraCut(FLevelEditorViewportClient& InViewportClient, UObject* InCameraObject) const
 {
 	ACameraActor* CameraActor = Cast<ACameraActor>(InCameraObject);
@@ -1970,6 +2130,7 @@ void FSequencer::UpdatePreviewLevelViewportClientFromCameraCut(FLevelEditorViewp
 	InViewportClient.UpdateViewForLockedActor();
 }
 
+
 void FSequencer::SaveCurrentMovieScene()
 {
 	UPackage* MovieScenePackage = GetCurrentAsset()->GetOutermost();
@@ -1981,6 +2142,7 @@ void FSequencer::SaveCurrentMovieScene()
 	FEditorFileUtils::PromptForCheckoutAndSave( PackagesToSave, false, false );
 }
 
+
 void FSequencer::AddSelectedObjects()
 {
 	UMovieSceneSequence* OwnerSequence = GetFocusedMovieSceneSequence();
@@ -2024,10 +2186,12 @@ void FSequencer::AddSelectedObjects()
 	}
 }
 
+
 void FSequencer::OnSectionSelectionChanged()
 {
 }
 
+
 void FSequencer::OnSelectedOutlinerNodesChanged()
 {
 	if (IsLevelEditorSequencer())
@@ -2045,6 +2209,7 @@ void FSequencer::OnSelectedOutlinerNodesChanged()
 	}
 }
 
+
 void FSequencer::ZoomToSelectedSections()
 {
 	TArray< TRange<float> > Bounds;
@@ -2077,31 +2242,37 @@ void FSequencer::ZoomToSelectedSections()
 	}
 }
 
+
 bool FSequencer::CanKeyProperty(FCanKeyPropertyParams CanKeyPropertyParams) const
 {
 	return ObjectChangeListener->CanKeyProperty(CanKeyPropertyParams);
 } 
 
+
 void FSequencer::KeyProperty(FKeyPropertyParams KeyPropertyParams) 
 {
 	ObjectChangeListener->KeyProperty(KeyPropertyParams);
 }
 
+
 FSequencerSelection& FSequencer::GetSelection()
 {
 	return Selection;
 }
 
+
 FSequencerSelectionPreview& FSequencer::GetSelectionPreview()
 {
 	return SelectionPreview;
 }
 
+
 float FSequencer::GetOverlayFadeCurve() const
 {
 	return OverlayCurve.GetLerp();
 }
 
+
 void FSequencer::DeleteSelectedItems()
 {
 	if (Selection.GetActiveSelection() == FSequencerSelection::EActiveSelection::KeyAndSection)
@@ -2118,6 +2289,7 @@ void FSequencer::DeleteSelectedItems()
 	}
 }
 
+
 void FSequencer::AssignActor(FMenuBuilder& MenuBuilder, FGuid InObjectBinding)
 {
 	UObject* RuntimeObject = SequenceInstanceStack.Top()->FindObject(InObjectBinding, *this);
@@ -2165,6 +2337,7 @@ void FSequencer::AssignActor(FMenuBuilder& MenuBuilder, FGuid InObjectBinding)
 	MenuBuilder.EndSection();
 }
 
+
 void FSequencer::DoAssignActor(AActor*const* InActors, int32 NumActors, FGuid InObjectBinding)
 {
 	if (NumActors > 0)
@@ -2253,6 +2426,7 @@ void FSequencer::DoAssignActor(AActor*const* InActors, int32 NumActors, FGuid In
 	}
 }
 
+
 void FSequencer::DeleteNode(TSharedRef<FSequencerDisplayNode> NodeToBeDeleted)
 {
 	// If this node is selected, delete all selected nodes
@@ -2268,6 +2442,7 @@ void FSequencer::DeleteNode(TSharedRef<FSequencerDisplayNode> NodeToBeDeleted)
 	}
 }
 
+
 void FSequencer::DeleteSelectedNodes()
 {
 	TSet< TSharedRef<FSequencerDisplayNode> > SelectedNodesCopy = GetSelection().GetSelectedOutlinerNodes();
@@ -2288,6 +2463,7 @@ void FSequencer::DeleteSelectedNodes()
 	}
 }
 
+
 void FSequencer::ToggleNodeActive()
 {
 	bool bIsActive = !IsNodeActive();
@@ -2306,6 +2482,7 @@ void FSequencer::ToggleNodeActive()
 	}
 }
 
+
 bool FSequencer::IsNodeActive() const
 {
 	// Active only if all are active
@@ -2325,6 +2502,7 @@ bool FSequencer::IsNodeActive() const
 	return true;
 }
 
+
 void FSequencer::ToggleNodeLocked()
 {
 	bool bIsLocked = !IsNodeLocked();
@@ -2343,6 +2521,7 @@ void FSequencer::ToggleNodeLocked()
 	}
 }
 
+
 bool FSequencer::IsNodeLocked() const
 {
 	// Locked only if all are locked
@@ -2364,75 +2543,106 @@ bool FSequencer::IsNodeLocked() const
 	return NumSections > 0;
 }
 
+
 void FSequencer::TogglePlay()
 {
 	OnPlay();
 }
 
+
 void FSequencer::PlayForward()
 {
 	OnPlay(false);
 }
 
+
 void FSequencer::Rewind()
 {
 	OnStepToBeginning();
 }
 
+
 void FSequencer::StepForward()
 {
 	OnStepForward();
 }
 
+
 void FSequencer::StepBackward()
 {
 	OnStepBackward();
 }
 
+
 void FSequencer::StepToNextKey()
 {
 	SequencerWidget->StepToNextKey();
 }
 
+
 void FSequencer::StepToPreviousKey()
 {
 	SequencerWidget->StepToPreviousKey();
 }
 
+
 void FSequencer::StepToNextCameraKey()
 {
 	SequencerWidget->StepToNextCameraKey();
 }
 
+
 void FSequencer::StepToPreviousCameraKey()
 {
 	SequencerWidget->StepToPreviousCameraKey();
 }
 
+
 void FSequencer::ExpandAllNodesAndDescendants()
 {
 	const bool bExpandAll = true;
 	SequencerWidget->GetTreeView()->ExpandNodes(ETreeRecursion::Recursive, bExpandAll);
 }
 
+
 void FSequencer::CollapseAllNodesAndDescendants()
 {
 	const bool bExpandAll = true;
 	SequencerWidget->GetTreeView()->CollapseNodes(ETreeRecursion::Recursive, bExpandAll);
 }
 
+
 void FSequencer::ToggleExpandCollapseNodes()
 {
 	SequencerWidget->GetTreeView()->ToggleExpandCollapseNodes(ETreeRecursion::NonRecursive);
 }
 
+
 void FSequencer::ToggleExpandCollapseNodesAndDescendants()
 {
 	SequencerWidget->GetTreeView()->ToggleExpandCollapseNodes(ETreeRecursion::Recursive);
 }
 
+
 void FSequencer::SetKey()
 {
+	FScopedTransaction SetKeyTransaction( NSLOCTEXT("Sequencer", "SetKey_Transaction", "Set Key") );
+
+	for (auto OutlinerNode : Selection.GetSelectedOutlinerNodes())
+	{
+		if (OutlinerNode->GetType() == ESequencerNode::Track)
+		{
+			TSharedRef<FSequencerTrackNode> TrackNode = StaticCastSharedRef<FSequencerTrackNode>(OutlinerNode);
+						
+			TSharedRef<FSequencerDisplayNode> ObjectBindingNode = OutlinerNode;
+			if (SequencerHelpers::FindObjectBindingNode(TrackNode, ObjectBindingNode))
+			{
+				FGuid ObjectGuid = StaticCastSharedRef<FSequencerObjectBindingNode>(ObjectBindingNode)->GetObjectBinding();
+				TrackNode->AddKey(ObjectGuid);
+			}
+		}
+	}
+
 	TSet<TSharedPtr<IKeyArea> > KeyAreas;
 	for (auto OutlinerNode : Selection.GetSelectedOutlinerNodes())
 	{
@@ -2440,9 +2650,7 @@ void FSequencer::SetKey()
 	}
 
 	if (KeyAreas.Num() > 0)
-	{
-		FScopedTransaction SetKeyTransaction( NSLOCTEXT("Sequencer", "SetKey_Transaction", "Set Key") );
-	
+	{	
 		for (auto KeyArea : KeyAreas)
 		{
 			if (KeyArea->GetOwningSection()->TryModify())
@@ -2453,11 +2661,13 @@ void FSequencer::SetKey()
 	}
 }
 
+
 bool FSequencer::CanSetKeyTime() const
 {
 	return Selection.GetSelectedKeys().Num() > 0;
 }
 
+
 void FSequencer::SetKeyTime(const bool bUseFrames)
 {
 	TArray<FSequencerSelectedKey> SelectedKeysArray = Selection.GetSelectedKeys().Array();
@@ -2481,6 +2691,7 @@ void FSequencer::SetKeyTime(const bool bUseFrames)
 		);
 }
 
+
 void FSequencer::OnSetKeyTimeTextCommitted(const FText& InText, ETextCommit::Type CommitInfo, const bool bUseFrames)
 {
 	CloseEntryPopupMenu();
@@ -2511,8 +2722,64 @@ void FSequencer::OnSetKeyTimeTextCommitted(const FText& InText, ETextCommit::Typ
 	}
 }
 
+
+void FSequencer::SelectTrackKeys(TWeakObjectPtr<UMovieSceneSection> Section, float KeyTime, bool bAddToSelection, bool bToggleSelection)
+{
+	if (!bAddToSelection && !bToggleSelection)
+	{
+		Selection.EmptySelectedKeys();
+	}
+
+	TSet<TWeakObjectPtr<UMovieSceneSection> > Sections;
+	Sections.Add(Section);
+	TArray<FSectionHandle> SectionHandles = SequencerWidget->GetSectionHandles(Sections);
+	for (auto SectionHandle : SectionHandles)
+	{
+		if (SectionHandle.TrackNode.IsValid())
+		{
+			TSet<TSharedPtr<IKeyArea> > KeyAreas;
+			SequencerHelpers::GetAllKeyAreas(SectionHandle.TrackNode, KeyAreas);
+
+			for (auto KeyArea : KeyAreas)
+			{
+				if (KeyArea.IsValid())
+				{
+					TArray<FKeyHandle> KeyHandles = KeyArea.Get()->GetUnsortedKeyHandles();
+					for (auto KeyHandle : KeyHandles)
+					{
+						float KeyHandleTime = KeyArea.Get()->GetKeyTime(KeyHandle);
+
+						if (FMath::IsNearlyEqual(KeyHandleTime, KeyTime, KINDA_SMALL_NUMBER))
+						{
+							FSequencerSelectedKey SelectedKey(*Section.Get(), KeyArea, KeyHandle);
+
+							if (bToggleSelection)
+							{
+								if (Selection.IsSelected(SelectedKey))
+								{
+									Selection.RemoveFromSelection(SelectedKey);
+								}
+								else
+								{
+									Selection.AddToSelection(SelectedKey);
+								}
+							}
+							else
+							{
+								Selection.AddToSelection(SelectedKey);
+							}
+						}
+					}
+				}
+			}
+		}
+	}
+}
+
+
 TArray<TSharedPtr<FMovieSceneClipboard>> GClipboardStack;
 
+
 void FSequencer::CopySelectedKeys()
 {
 	if (Selection.GetSelectedKeys().Num() == 0)
@@ -2567,6 +2834,7 @@ void FSequencer::CopySelectedKeys()
 	}
 }
 
+
 void FSequencer::CutSelectedKeys()
 {
 	FScopedTransaction CutKeysTransaction( LOCTEXT("CutSelectedKeys_Transaction", "Cut Selected Keys") );
@@ -2574,11 +2842,13 @@ void FSequencer::CutSelectedKeys()
 	DeleteSelectedKeys();
 }
 
+
 const TArray<TSharedPtr<FMovieSceneClipboard>>& FSequencer::GetClipboardStack() const
 {
 	return GClipboardStack;
 }
 
+
 void FSequencer::OnClipboardUsed(TSharedPtr<FMovieSceneClipboard> Clipboard)
 {
 	Clipboard->GetEnvironment().DateTime = FDateTime::UtcNow();
@@ -2589,6 +2859,7 @@ void FSequencer::OnClipboardUsed(TSharedPtr<FMovieSceneClipboard> Clipboard)
 	});
 }
 
+
 void FSequencer::GenericTextEntryModeless(const FText& DialogText, const FText& DefaultText, FOnTextCommitted OnTextComitted)
 {
 	TSharedRef<STextEntryPopup> TextEntryPopup = 
@@ -2609,6 +2880,7 @@ void FSequencer::GenericTextEntryModeless(const FText& DialogText, const FText&
 		);
 }
 
+
 void FSequencer::CloseEntryPopupMenu()
 {
 	if (EntryPopupMenu.IsValid())
@@ -2617,6 +2889,7 @@ void FSequencer::CloseEntryPopupMenu()
 	}
 }
 
+
 void FSequencer::TrimSection(bool bTrimLeft)
 {
 	FScopedTransaction TrimSectionTransaction( NSLOCTEXT("Sequencer", "TrimSection_Transaction", "Trim Section") );
@@ -2626,6 +2899,7 @@ void FSequencer::TrimSection(bool bTrimLeft)
 	NotifyMovieSceneDataChanged();
 }
 
+
 void FSequencer::SplitSection()
 {
 	FScopedTransaction SplitSectionTransaction( NSLOCTEXT("Sequencer", "SplitSection_Transaction", "Split Section") );
@@ -2635,11 +2909,13 @@ void FSequencer::SplitSection()
 	NotifyMovieSceneDataChanged();
 }
 
+
 ISequencerEditTool& FSequencer::GetEditTool()
 {
 	return SequencerWidget->GetEditTool();
 }
 
+
 void FSequencer::BindSequencerCommands()
 {
 	const FSequencerCommands& Commands = FSequencerCommands::Get();
@@ -2753,10 +3029,22 @@ void FSequencer::BindSequencerCommands()
 		FExecuteAction::CreateSP( this, &FSequencer::SplitSection ) );
 
 	SequencerCommandBindings->MapAction(
-		Commands.ToggleAutoKeyEnabled,
-		FExecuteAction::CreateLambda( [this]{ Settings->SetAutoKeyEnabled( !Settings->GetAutoKeyEnabled() ); } ),
+		Commands.SetAutoKeyModeAll,
+		FExecuteAction::CreateLambda( [this]{ Settings->SetAutoKeyMode( EAutoKeyMode::KeyAll ); } ),
 		FCanExecuteAction::CreateLambda( []{ return true; } ),
-		FIsActionChecked::CreateLambda( [this]{ return Settings->GetAutoKeyEnabled(); } ) );
+		FIsActionChecked::CreateLambda( [this]{ return Settings->GetAutoKeyMode() == EAutoKeyMode::KeyAll; } ) );
+
+	SequencerCommandBindings->MapAction(
+		Commands.SetAutoKeyModeAnimated,
+		FExecuteAction::CreateLambda([this] { Settings->SetAutoKeyMode(EAutoKeyMode::KeyAnimated); }),
+		FCanExecuteAction::CreateLambda([] { return true; }),
+		FIsActionChecked::CreateLambda([this] { return Settings->GetAutoKeyMode() == EAutoKeyMode::KeyAnimated; }));
+
+	SequencerCommandBindings->MapAction(
+		Commands.SetAutoKeyModeNone,
+		FExecuteAction::CreateLambda([this] { Settings->SetAutoKeyMode(EAutoKeyMode::KeyNone); }),
+		FCanExecuteAction::CreateLambda([] { return true; }),
+		FIsActionChecked::CreateLambda([this] { return Settings->GetAutoKeyMode() == EAutoKeyMode::KeyNone; }));
 
 	SequencerCommandBindings->MapAction(
 		Commands.ToggleKeyAllEnabled,
@@ -2895,6 +3183,12 @@ void FSequencer::BindSequencerCommands()
 		FCanExecuteAction::CreateLambda( []{ return true; } ),
 		FIsActionChecked::CreateLambda( [this]{ return Settings->ShouldKeepCursorInPlayRange(); } ) );
 
+	SequencerCommandBindings->MapAction(
+		Commands.ToggleKeepPlaybackRangeInSectionBounds,
+		FExecuteAction::CreateLambda( [this]{ Settings->SetKeepPlayRangeInSectionBounds( !Settings->ShouldKeepPlayRangeInSectionBounds() ); NotifyMovieSceneDataChanged(); } ),
+		FCanExecuteAction::CreateLambda( []{ return true; } ),
+		FIsActionChecked::CreateLambda( [this]{ return Settings->ShouldKeepPlayRangeInSectionBounds(); } ) );
+
 	SequencerCommandBindings->MapAction(
 		Commands.RenderMovie,
 		FExecuteAction::CreateLambda([this]{
@@ -2904,16 +3198,6 @@ void FSequencer::BindSequencerCommands()
 			UAutomatedLevelSequenceCapture* MovieSceneCapture = NewObject<UAutomatedLevelSequenceCapture>(GetTransientPackage(), UAutomatedLevelSequenceCapture::StaticClass(), NAME_None, RF_Transient);
 			MovieSceneCapture->LoadConfig();
 
-			// Attempt to find a level sequence actor in the world that references this asset
-			for (auto It = TActorIterator<ALevelSequenceActor>(GWorld); It; ++It)
-			{
-				if (It->LevelSequence == GetCurrentAsset())
-				{
-					MovieSceneCapture->SetLevelSequenceActor(*It);
-					break;
-				}
-			}
-
 			MovieSceneCapture->SetLevelSequenceAsset(GetCurrentAsset()->GetPathName());
 
 			if (CanShowFrameNumbers())
@@ -2940,6 +3224,7 @@ void FSequencer::BindSequencerCommands()
 	}
 }
 
+
 void FSequencer::BuildAddTrackMenu(class FMenuBuilder& MenuBuilder)
 {
 	for (int32 i = 0; i < TrackEditors.Num(); ++i)
@@ -2948,6 +3233,7 @@ void FSequencer::BuildAddTrackMenu(class FMenuBuilder& MenuBuilder)
 	}
 }
 
+
 void FSequencer::BuildObjectBindingTrackMenu(FMenuBuilder& MenuBuilder, const FGuid& ObjectBinding, const UClass* ObjectClass)
 {
 	for (int32 i = 0; i < TrackEditors.Num(); ++i)
@@ -2956,6 +3242,7 @@ void FSequencer::BuildObjectBindingTrackMenu(FMenuBuilder& MenuBuilder, const FG
 	}
 }
 
+
 void FSequencer::BuildObjectBindingEditButtons(TSharedPtr<SHorizontalBox> EditBox, const FGuid& ObjectBinding, const UClass* ObjectClass)
 {
 	for (int32 i = 0; i < TrackEditors.Num(); ++i)
@@ -2964,10 +3251,12 @@ void FSequencer::BuildObjectBindingEditButtons(TSharedPtr<SHorizontalBox> EditBo
 	}
 }
 
+
 bool FSequencer::IsSequencerWidgetFocused() const
 {
 	// Only perform certain operations if the key was pressed in the sequencer widget and not in, for example, the viewport.
 	return SequencerWidget->HasFocusedDescendants() || SequencerWidget->HasKeyboardFocus();
 }
 
+
 #undef LOCTEXT_NAMESPACE
diff --git a/Engine/Source/Editor/Sequencer/Private/Sequencer.h b/Engine/Source/Editor/Sequencer/Private/Sequencer.h
--- a/Engine/Source/Editor/Sequencer/Private/Sequencer.h
+++ b/Engine/Source/Editor/Sequencer/Private/Sequencer.h
@@ -35,7 +35,13 @@
 {
 public:
 
-	static bool IsSequencerEnabled();
+	/** Constructor */
+	FSequencer();
 
+	/** Virtual destructor */
+	virtual ~FSequencer();
+
+public:
+
 	/**
 	 * Initializes sequencer
@@ -47,68 +53,6 @@
 	 */
 	void InitSequencer(const FSequencerInitParams& InitParams, const TSharedRef<ISequencerObjectChangeListener>& InObjectChangeListener, const TSharedRef<IDetailKeyframeHandler>& InDetailKeyframeHandler, const TArray<FOnCreateTrackEditor>& TrackEditorDelegates);
 
-	/** Constructor */
-	FSequencer();
-
-	/** Destructor */
-	virtual ~FSequencer();
-
-public:
-
-	//~ Begin FGCObject Interface
-
-	virtual void AddReferencedObjects(FReferenceCollector& Collector) override;
-
-public:
-
-	//~ Begin FTickableEditorObject Interface
-
-	virtual void Tick(float DeltaTime) override;
-	virtual bool IsTickable() const override { return true; }
-	virtual TStatId GetStatId() const override { RETURN_QUICK_DECLARE_CYCLE_STAT(FSequencer, STATGROUP_Tickables); };
-
-public:
-
-	//~ Begin ISequencer Interface
-
-	virtual void Close() override;
-	virtual TSharedRef<SWidget> GetSequencerWidget() const override { return SequencerWidget.ToSharedRef(); }
-	virtual UMovieSceneSequence* GetRootMovieSceneSequence() const override;
-	virtual UMovieSceneSequence* GetFocusedMovieSceneSequence() const override;
-	virtual void ResetToNewRootSequence(UMovieSceneSequence& NewSequence) override;
-	virtual TSharedRef<FMovieSceneSequenceInstance> GetRootMovieSceneSequenceInstance() const override;
-	virtual TSharedRef<FMovieSceneSequenceInstance> GetFocusedMovieSceneSequenceInstance() const override;
-	virtual void FocusSequenceInstance( TSharedRef<FMovieSceneSequenceInstance> SequenceInstance ) override;
-	virtual TSharedRef<FMovieSceneSequenceInstance> GetSequenceInstanceForSection(UMovieSceneSection& Section) const override;
-	virtual bool GetAutoKeyEnabled() const override;
-	virtual void SetAutoKeyEnabled(bool bAutoKeyEnabled) override;
-	virtual bool GetKeyAllEnabled() const override;
-	virtual void SetKeyAllEnabled(bool bKeyAllEnabled) override;
-	virtual bool GetKeyInterpPropertiesOnly() const override;
-	virtual void SetKeyInterpPropertiesOnly(bool bKeyInterpPropertiesOnly) override;
-	virtual EMovieSceneKeyInterpolation GetKeyInterpolation() const override;
-	virtual void SetKeyInterpolation(EMovieSceneKeyInterpolation) override;
-	virtual bool IsRecordingLive() const override;
-	virtual float GetCurrentLocalTime(UMovieSceneSequence& InMovieSceneSequence) override;
-	virtual float GetGlobalTime() override;
-	virtual void SetGlobalTime(float Time) override;
-	virtual void SetPerspectiveViewportPossessionEnabled(bool bEnabled) override;
-	virtual void SetPerspectiveViewportCameraCutEnabled(bool bEnabled) override;
-	virtual FGuid GetHandleToObject(UObject* Object, bool bCreateHandleIfMissing = true) override;
-	virtual ISequencerObjectChangeListener& GetObjectChangeListener() override;
-	virtual void NotifyMovieSceneDataChanged() override;
-	virtual void UpdateRuntimeInstances() override;
-	virtual void AddSubSequence(UMovieSceneSequence* Sequence) override;
-	virtual bool CanKeyProperty(FCanKeyPropertyParams CanKeyPropertyParams) const override;
-	virtual void KeyProperty(FKeyPropertyParams KeyPropertyParams) override;
-	virtual FSequencerSelection& GetSelection() override;
-	virtual FSequencerSelectionPreview& GetSelectionPreview() override;
-	virtual void NotifyMapChanged(UWorld* NewWorld, EMapChangeType MapChangeType) override;
-	virtual FOnGlobalTimeChanged& OnGlobalTimeChanged() override { return OnGlobalTimeChangedDelegate; }
-	virtual FGuid CreateBinding(UObject& InObject, const FString& InName) override;
-	virtual UObject* GetPlaybackContext() const override;
-	virtual void GetAllKeyedProperties(UObject& Object, TSet<UProperty*>& OutProperties) override;
-
 	/** Set the global time directly, without performing any auto-scroll */
 	void SetGlobalTimeDirectly(float Time);
 
@@ -126,44 +70,42 @@
 public:
 
 	/** Access the user-supplied settings object */
-	USequencerSettings* GetSettings() const { return Settings; }
+	USequencerSettings* GetSettings() const
+	{
+		return Settings;
+	}
 
-	bool IsPerspectiveViewportPossessionEnabled() const override { return bPerspectiveViewportPossessionEnabled; }
-	bool IsPerspectiveViewportCameraCutEnabled() const override { return bPerspectiveViewportCameraCutEnabled; }
+	bool IsPerspectiveViewportPossessionEnabled() const override
+	{
+		return bPerspectiveViewportPossessionEnabled;
+	}
+
+	bool IsPerspectiveViewportCameraCutEnabled() const override
+	{
+		return bPerspectiveViewportCameraCutEnabled;
+	}
 
 	/**
 	 * Pops the current focused movie scene from the stack.  The parent of this movie scene will be come the focused one
 	 */
 	void PopToSequenceInstance( TSharedRef<FMovieSceneSequenceInstance> SequenceInstance );
 
-	/** 
-	 * Deletes the passed in sections
-	 */
+	/** Deletes the passed in sections. */
 	void DeleteSections(const TSet<TWeakObjectPtr<UMovieSceneSection> > & Sections);
 
-	/**
-	 * Deletes the currently selected in keys
-	 */
+	/** Deletes the currently selected in keys. */
 	void DeleteSelectedKeys();
 
-	/**
-	 * Set interpolation modes
-	 */
+	/** Set interpolation modes. */
 	void SetInterpTangentMode(ERichCurveInterpMode InterpMode, ERichCurveTangentMode TangentMode);
 
-	/**
-	 * Is interpolation mode selected
-	 */
+	/** Is interpolation mode selected. */
 	bool IsInterpTangentModeSelected(ERichCurveInterpMode InterpMode, ERichCurveTangentMode TangentMode) const;
 
-	/**
-	 * Snap the currently selected keys to frame
-	 */
+	/** Snap the currently selected keys to frame. */
 	void SnapToFrame();
 
-	/**
- 	 * Are there keys to snap? 
-	 */
+	/** Are there keys to snap? */
 	bool CanSnapToFrame() const;
 
 	/**
@@ -210,14 +152,10 @@
 	 */
 	virtual void CopyActorProperties( AActor* PuppetActor, AActor* TargetActor ) const;
 
-	/**
-	 * Zooms to the edges of all currently selected sections
-	 */
+	/** Zooms to the edges of all currently selected sections. */
 	void ZoomToSelectedSections();
 
-	/**
-	 * Gets the overlay fading animation curve lerp
-	 */
+	/** Gets the overlay fading animation curve lerp. */
 	float GetOverlayFadeCurve() const;
 
 	/** Gets the command bindings for the sequencer */
@@ -248,16 +186,6 @@
 	 */
 	void BuildObjectBindingEditButtons(TSharedPtr<SHorizontalBox> EditBox, const FGuid& ObjectBinding, const UClass* ObjectClass);
 
-	/** IMovieScenePlayer interface */
-	virtual void GetRuntimeObjects( TSharedRef<FMovieSceneSequenceInstance> MovieSceneInstance, const FGuid& ObjectHandle, TArray< UObject* >& OutObjects ) const override;
-	virtual void UpdateCameraCut(UObject* CameraObject, UObject* UnlockIfCameraObject) const override;
-	virtual void SetViewportSettings(const TMap<FViewportClient*, EMovieSceneViewportParams>& ViewportParamsMap) override;
-	virtual void GetViewportSettings(TMap<FViewportClient*, EMovieSceneViewportParams>& ViewportParamsMap) const override;
-	virtual EMovieScenePlayerStatus::Type GetPlaybackStatus() const override;
-	virtual void AddOrUpdateMovieSceneInstance(UMovieSceneSection& MovieSceneSection, TSharedRef<FMovieSceneSequenceInstance> InstanceToAdd) override;
-	virtual void RemoveMovieSceneInstance(UMovieSceneSection& MovieSceneSection, TSharedRef<FMovieSceneSequenceInstance> InstanceToRemove) override;
-	virtual IMovieSceneSpawnRegister& GetSpawnRegister() override { return *SpawnRegister; }
-
 	/** Called when an actor is dropped into Sequencer */
 	void OnActorsDropped( const TArray<TWeakObjectPtr<AActor> >& Actors );
 
@@ -318,6 +246,9 @@
 		return LabelManager;
 	}
 
+	/** Select keys belonging to a section at the key time */
+	void SelectTrackKeys(TWeakObjectPtr<UMovieSceneSection> Section, float KeyTime, bool bAddToSelection, bool bToggleSelection);
+
 public:
 
 	/** Copy the selected keys to the clipboard */
@@ -337,6 +268,87 @@
 	/** Access the currently enabled edit tool */
 	ISequencerEditTool& GetEditTool();
 
+public:
+
+	/** Put the sequencer in a horizontally auto-scrolling state with the given rate */
+	void StartAutoscroll(float UnitsPerS);
+
+	/** Stop the sequencer from auto-scrolling */
+	void StopAutoscroll();
+
+	/** Scroll the sequencer vertically by the specified number of slate units */
+	void VerticalScroll(float ScrollAmountUnits);
+
+public:
+
+	//~ FGCObject Interface
+
+	virtual void AddReferencedObjects(FReferenceCollector& Collector) override;
+
+public:
+
+	//~ FTickableEditorObject Interface
+
+	virtual void Tick(float DeltaTime) override;
+	virtual bool IsTickable() const override { return true; }
+	virtual TStatId GetStatId() const override { RETURN_QUICK_DECLARE_CYCLE_STAT(FSequencer, STATGROUP_Tickables); };
+
+public:
+
+	//~ ISequencer Interface
+
+	virtual void Close() override;
+	virtual TSharedRef<SWidget> GetSequencerWidget() const override { return SequencerWidget.ToSharedRef(); }
+	virtual UMovieSceneSequence* GetRootMovieSceneSequence() const override;
+	virtual UMovieSceneSequence* GetFocusedMovieSceneSequence() const override;
+	virtual void ResetToNewRootSequence(UMovieSceneSequence& NewSequence) override;
+	virtual TSharedRef<FMovieSceneSequenceInstance> GetRootMovieSceneSequenceInstance() const override;
+	virtual TSharedRef<FMovieSceneSequenceInstance> GetFocusedMovieSceneSequenceInstance() const override;
+	virtual void FocusSequenceInstance( TSharedRef<FMovieSceneSequenceInstance> SequenceInstance ) override;
+	virtual TSharedRef<FMovieSceneSequenceInstance> GetSequenceInstanceForSection(UMovieSceneSection& Section) const override;
+	virtual EAutoKeyMode GetAutoKeyMode() const override;
+	virtual void SetAutoKeyMode(EAutoKeyMode AutoKeyMode) override;
+	virtual bool GetKeyAllEnabled() const override;
+	virtual void SetKeyAllEnabled(bool bKeyAllEnabled) override;
+	virtual bool GetKeyInterpPropertiesOnly() const override;
+	virtual void SetKeyInterpPropertiesOnly(bool bKeyInterpPropertiesOnly) override;
+	virtual EMovieSceneKeyInterpolation GetKeyInterpolation() const override;
+	virtual void SetKeyInterpolation(EMovieSceneKeyInterpolation) override;
+	virtual bool IsRecordingLive() const override;
+	virtual float GetCurrentLocalTime(UMovieSceneSequence& InMovieSceneSequence) override;
+	virtual float GetGlobalTime() override;
+	virtual void SetGlobalTime(float Time) override;
+	virtual void SetPerspectiveViewportPossessionEnabled(bool bEnabled) override;
+	virtual void SetPerspectiveViewportCameraCutEnabled(bool bEnabled) override;
+	virtual FGuid GetHandleToObject(UObject* Object, bool bCreateHandleIfMissing = true) override;
+	virtual ISequencerObjectChangeListener& GetObjectChangeListener() override;
+	virtual void NotifyMovieSceneDataChanged() override;
+	virtual void UpdateRuntimeInstances() override;
+	virtual void UpdatePlaybackRange() override;
+	virtual void AddSubSequence(UMovieSceneSequence* Sequence) override;
+	virtual bool CanKeyProperty(FCanKeyPropertyParams CanKeyPropertyParams) const override;
+	virtual void KeyProperty(FKeyPropertyParams KeyPropertyParams) override;
+	virtual FSequencerSelection& GetSelection() override;
+	virtual FSequencerSelectionPreview& GetSelectionPreview() override;
+	virtual void NotifyMapChanged(UWorld* NewWorld, EMapChangeType MapChangeType) override;
+	virtual FOnGlobalTimeChanged& OnGlobalTimeChanged() override { return OnGlobalTimeChangedDelegate; }
+	virtual FGuid CreateBinding(UObject& InObject, const FString& InName) override;
+	virtual UObject* GetPlaybackContext() const override;
+	virtual void GetAllKeyedProperties(UObject& Object, TSet<UProperty*>& OutProperties) override;
+
+public:
+
+	// IMovieScenePlayer interface
+
+	virtual void GetRuntimeObjects( TSharedRef<FMovieSceneSequenceInstance> MovieSceneInstance, const FGuid& ObjectHandle, TArray< UObject* >& OutObjects ) const override;
+	virtual void UpdateCameraCut(UObject* CameraObject, UObject* UnlockIfCameraObject) const override;
+	virtual void SetViewportSettings(const TMap<FViewportClient*, EMovieSceneViewportParams>& ViewportParamsMap) override;
+	virtual void GetViewportSettings(TMap<FViewportClient*, EMovieSceneViewportParams>& ViewportParamsMap) const override;
+	virtual EMovieScenePlayerStatus::Type GetPlaybackStatus() const override;
+	virtual void AddOrUpdateMovieSceneInstance(UMovieSceneSection& MovieSceneSection, TSharedRef<FMovieSceneSequenceInstance> InstanceToAdd) override;
+	virtual void RemoveMovieSceneInstance(UMovieSceneSection& MovieSceneSection, TSharedRef<FMovieSceneSequenceInstance> InstanceToRemove) override;
+	virtual IMovieSceneSpawnRegister& GetSpawnRegister() override { return *SpawnRegister; }
+
 protected:
 
 	/** Reset data about a movie scene when pushing or popping a movie scene. */
@@ -398,20 +410,10 @@
 	/** Called when the user has finished dragging the play range */
 	void OnEndPlaybackRangeDrag();
 
-public:
-
-	/** Put the sequencer in a horizontally autoscrolling state with the given rate */
-	void StartAutoscroll(float UnitsPerS);
-
-	/** Stop the sequencer from autoscrolling */
-	void StopAutoscroll();
-
-	/** Scroll the sequencer vertically by the specified number of slate units */
-	void VerticalScroll(float ScrollAmountUnits);
-
 protected:
+
 	/**
-	 * Update autoscroll mechanics as a result of a new time position
+	 * Update auto-scroll mechanics as a result of a new time position
 	 */
 	void UpdateAutoScroll(float NewTime);
 
@@ -430,14 +432,10 @@
 	 */
 	bool IsAutoScrollEnabled() const;
 
-	/** 
-	 * Find the viewed sequence asset in the content browser
-	 */
+	/** Find the viewed sequence asset in the content browser. */
 	void FindInContentBrowser();
 
-	/**
-	 * Get the asset we're currently editing, if applicable
-	 */
+	/** Get the asset we're currently editing, if applicable. */
 	UObject* GetCurrentAsset() const;
 
 protected:
diff --git a/Engine/Source/Editor/Sequencer/Private/SequencerCommands.cpp b/Engine/Source/Editor/Sequencer/Private/SequencerCommands.cpp
--- a/Engine/Source/Editor/Sequencer/Private/SequencerCommands.cpp
+++ b/Engine/Source/Editor/Sequencer/Private/SequencerCommands.cpp
@@ -19,6 +19,7 @@ void FSequencerCommands::RegisterCommands()
 	UI_COMMAND( SetStartPlaybackRange, "Set Start Playback Range", "Set the start playback range", EUserInterfaceActionType::Button, FInputChord(EKeys::LeftBracket) );
 	UI_COMMAND( SetEndPlaybackRange, "Set End Playback Range", "Set the end playback range", EUserInterfaceActionType::Button, FInputChord(EKeys::RightBracket) );
 	UI_COMMAND( ToggleKeepCursorInPlaybackRange, "Keep Cursor In Playback Range", "When checked, the cursor will be constrained to the current playback range during playback", EUserInterfaceActionType::ToggleButton, FInputChord() );
+	UI_COMMAND( ToggleKeepPlaybackRangeInSectionBounds, "Keep Playback Range In Section Bounds", "When checked, the playback range will be synchronized to the section bounds", EUserInterfaceActionType::ToggleButton, FInputChord() );
 
 	UI_COMMAND( ExpandAllNodesAndDescendants, "Expand All Nodes", "Expand all nodes", EUserInterfaceActionType::Button, FInputChord() );
 	UI_COMMAND( CollapseAllNodesAndDescendants, "Collapse All Nodes", "Collapse all selected nodes", EUserInterfaceActionType::Button, FInputChord() );
@@ -36,8 +37,11 @@ void FSequencerCommands::RegisterCommands()
 	UI_COMMAND( TrimSectionRight, "Trim Section Right", "Trim section at current time to the right (keeps the left)", EUserInterfaceActionType::Button, FInputChord(EModifierKey::Control, EKeys::Period) );
 	UI_COMMAND( SplitSection, "Split Section", "Split section at current time", EUserInterfaceActionType::Button, FInputChord(EModifierKey::Control, EKeys::Slash) );
 
-	UI_COMMAND( ToggleAutoKeyEnabled, "Auto Key", "Enables and disables auto keying", EUserInterfaceActionType::ToggleButton, FInputChord() );
-	UI_COMMAND( ToggleKeyAllEnabled, "Key All", "Enables and disables key all", EUserInterfaceActionType::ToggleButton, FInputChord() );
+	UI_COMMAND( SetAutoKeyModeAll, "Auto-key All", "Enables auto keying for all properties.", EUserInterfaceActionType::ToggleButton, FInputChord() );
+	UI_COMMAND( SetAutoKeyModeAnimated, "Auto-key Animated", "Enables auto keying for properties with existing animations.", EUserInterfaceActionType::ToggleButton, FInputChord());
+	UI_COMMAND( SetAutoKeyModeNone, "Disable Auto-key", "Disables auto keying.", EUserInterfaceActionType::ToggleButton, FInputChord());
+
+	UI_COMMAND(ToggleKeyAllEnabled, "Key All", "Enables and disables key all", EUserInterfaceActionType::ToggleButton, FInputChord());
 
 	UI_COMMAND( ToggleAutoScroll, "Auto Scroll", "Toggle auto-scroll: When enabled, automatically scrolls the sequencer view to keep the current time visible", EUserInterfaceActionType::ToggleButton, FInputChord(EModifierKey::Shift, EKeys::S) );
 
diff --git a/Engine/Source/Editor/Sequencer/Private/SequencerCommands.h b/Engine/Source/Editor/Sequencer/Private/SequencerCommands.h
--- a/Engine/Source/Editor/Sequencer/Private/SequencerCommands.h
+++ b/Engine/Source/Editor/Sequencer/Private/SequencerCommands.h
@@ -51,6 +51,9 @@ public:
 	/** Toggle constraining the time cursor to the playback range */
 	TSharedPtr< FUICommandInfo > ToggleKeepCursorInPlaybackRange;
 
+	/** Toggle constraining the playback range to the section bounds */
+	TSharedPtr< FUICommandInfo > ToggleKeepPlaybackRangeInSectionBounds;
+
 	/** Expand all nodes and descendants */
 	TSharedPtr< FUICommandInfo > ExpandAllNodesAndDescendants;
 
@@ -90,8 +93,14 @@ public:
 	/** Split section */
 	TSharedPtr< FUICommandInfo > SplitSection;
 
-	/** Turns auto keying on and off. */
-	TSharedPtr< FUICommandInfo > ToggleAutoKeyEnabled;
+	/** Set the auto key mode to Key All. */
+	TSharedPtr< FUICommandInfo > SetAutoKeyModeAll;
+
+	/** Set the auto key mode to Key Animated. */
+	TSharedPtr< FUICommandInfo > SetAutoKeyModeAnimated;
+
+	/** Set the auto key mode to Key None. */
+	TSharedPtr< FUICommandInfo > SetAutoKeyModeNone;
 
 	/** Turns key all on and off. */
 	TSharedPtr< FUICommandInfo > ToggleKeyAllEnabled;
diff --git a/Engine/Source/Editor/Sequencer/Private/SequencerCommonHelpers.cpp b/Engine/Source/Editor/Sequencer/Private/SequencerCommonHelpers.cpp
--- a/Engine/Source/Editor/Sequencer/Private/SequencerCommonHelpers.cpp
+++ b/Engine/Source/Editor/Sequencer/Private/SequencerCommonHelpers.cpp
@@ -90,6 +90,22 @@ void SequencerHelpers::GetAllSections(TSharedRef<FSequencerDisplayNode> DisplayN
 	}
 }
 
+bool SequencerHelpers::FindObjectBindingNode(TSharedRef<FSequencerDisplayNode> DisplayNode, TSharedRef<FSequencerDisplayNode>& ObjectBindingNode)
+{
+	if (DisplayNode->GetType() == ESequencerNode::Object)
+	{
+		ObjectBindingNode = DisplayNode;
+		return true;
+	}
+
+	if (DisplayNode->GetParent().IsValid())
+	{
+		return FindObjectBindingNode(DisplayNode->GetParent().ToSharedRef(), ObjectBindingNode);
+	}
+
+	return false;
+}
+
 int32 SequencerHelpers::TimeToFrame(float Time, float FrameRate)
 {
 	float Frame = Time * FrameRate;
diff --git a/Engine/Source/Editor/Sequencer/Private/SequencerCommonHelpers.h b/Engine/Source/Editor/Sequencer/Private/SequencerCommonHelpers.h
--- a/Engine/Source/Editor/Sequencer/Private/SequencerCommonHelpers.h
+++ b/Engine/Source/Editor/Sequencer/Private/SequencerCommonHelpers.h
@@ -25,6 +25,11 @@ public:
 	*/
 	static void GetAllSections(TSharedRef<FSequencerDisplayNode> DisplayNode, TSet<TWeakObjectPtr<UMovieSceneSection>>& Sections);
 
+	/**
+	* Find object binding node from the given display node
+	*/
+	static bool FindObjectBindingNode(TSharedRef<FSequencerDisplayNode> DisplayNode, TSharedRef<FSequencerDisplayNode>& ObjectBindingNode);
+
 	/**
 	* Convert time to frame
 	*/
diff --git a/Engine/Source/Editor/Sequencer/Private/SequencerContextMenus.cpp b/Engine/Source/Editor/Sequencer/Private/SequencerContextMenus.cpp
--- a/Engine/Source/Editor/Sequencer/Private/SequencerContextMenus.cpp
+++ b/Engine/Source/Editor/Sequencer/Private/SequencerContextMenus.cpp
@@ -104,8 +104,8 @@ void FKeyContextMenu::PopulateMenu(FMenuBuilder& MenuBuilder)
 		const bool bUseFrames = SequencerSnapValues::IsTimeSnapIntervalFrameRate(Sequencer->GetSettings()->GetTimeSnapInterval());
 
 		MenuBuilder.AddMenuEntry(
-			bUseFrames ? LOCTEXT("SetKeyFrame", "Set Key Frame") : LOCTEXT("SetKeyTime", "Set Key MouseDownTime"),
-			bUseFrames ? LOCTEXT("SetKeyFrameTooltip", "Set key frame") : LOCTEXT("SetKeyTimeTooltip", "Set key MouseDownTime"),
+			bUseFrames ? LOCTEXT("SetKeyFrame", "Set Key Frame") : LOCTEXT("SetKeyTime", "Set Key Time"),
+			bUseFrames ? LOCTEXT("SetKeyFrameTooltip", "Set key frame") : LOCTEXT("SetKeyTimeTooltip", "Set key time"),
 			FSlateIcon(),
 			FUIAction(
 				FExecuteAction::CreateSP(SequencerPtr, &FSequencer::SetKeyTime, bUseFrames),
diff --git a/Engine/Source/Editor/Sequencer/Private/SequencerEdMode.cpp b/Engine/Source/Editor/Sequencer/Private/SequencerEdMode.cpp
--- a/Engine/Source/Editor/Sequencer/Private/SequencerEdMode.cpp
+++ b/Engine/Source/Editor/Sequencer/Private/SequencerEdMode.cpp
@@ -3,13 +3,16 @@
 #include "SequencerPrivatePCH.h"
 #include "SequencerEdMode.h"
 #include "Sequencer.h"
+#include "MovieSceneHitProxy.h"
+#include "Tracks/MovieScene3DTransformTrack.h"
+#include "Sections/MovieScene3DTransformSection.h"
 
 const FEditorModeID FSequencerEdMode::EM_SequencerMode(TEXT("EM_SequencerMode"));
 
 FSequencerEdMode::FSequencerEdMode() :
-	Sequencer(NULL)
+	Sequencer(nullptr)
 {
-	FSequencerEdModeTool* SequencerEdModeTool = new FSequencerEdModeTool();
+	FSequencerEdModeTool* SequencerEdModeTool = new FSequencerEdModeTool(this);
 
 	Tools.Add( SequencerEdModeTool );
 	SetCurrentTool( SequencerEdModeTool );
@@ -39,7 +42,7 @@ bool FSequencerEdMode::IsCompatibleWith(FEditorModeID OtherModeID) const
 
 bool FSequencerEdMode::InputKey( FEditorViewportClient* ViewportClient, FViewport* Viewport, FKey Key, EInputEvent Event )
 {
-	if (Sequencer == NULL)
+	if (Sequencer == nullptr)
 	{
 		return false;
 	}
@@ -57,7 +60,323 @@ bool FSequencerEdMode::InputKey( FEditorViewportClient* ViewportClient, FViewpor
 	return FEdMode::InputKey(ViewportClient, Viewport, Key, Event);
 }
 
-FSequencerEdModeTool::FSequencerEdModeTool()
+void FSequencerEdMode::Render(const FSceneView* View, FViewport* Viewport, FPrimitiveDrawInterface* PDI)
+{
+	if (Sequencer == nullptr)
+	{
+		return;
+	}
+
+	FEdMode::Render(View, Viewport, PDI);
+
+#if WITH_EDITORONLY_DATA
+	if (View->Family->EngineShowFlags.Splines)
+	{
+		DrawTracks3D(View, PDI);
+	}
+#endif
+}
+
+void FSequencerEdMode::OnKeySelected(FViewport* Viewport, HMovieSceneKeyProxy* KeyProxy)
+{
+	bool bCtrlDown = Viewport->KeyState(EKeys::LeftControl) || Viewport->KeyState(EKeys::RightControl);
+	bool bAltDown = Viewport->KeyState(EKeys::LeftAlt) || Viewport->KeyState(EKeys::RightAlt);
+	bool bShiftDown = Viewport->KeyState(EKeys::LeftShift) || Viewport->KeyState(EKeys::RightShift);
+
+	Sequencer->SetGlobalTimeDirectly(KeyProxy->Time);
+	Sequencer->SelectTrackKeys(KeyProxy->MovieSceneSection, KeyProxy->Time, bShiftDown, bCtrlDown);
+}
+
+namespace SequencerEdMode_Draw3D
+{
+static const FColor	KeySelectedColor(255,128,0);
+static const float	DrawTrackTimeRes = 0.1f;
+static const float	CurveHandleScale = 0.5f;
+}
+
+FTransform GetRefFrame(const UObject* InObject)
+{
+	FTransform RefTM = FTransform::Identity;
+
+	const AActor* Actor = Cast<AActor>(InObject);
+	if (Actor != nullptr && Actor->GetRootComponent() != nullptr && Actor->GetRootComponent()->AttachParent != nullptr)
+	{
+		RefTM = Actor->GetRootComponent()->AttachParent->GetSocketTransform(Actor->GetRootComponent()->AttachSocketName);
+	}
+
+	return RefTM;
+}
+
+void GetLocationAtTime(UMovieScene3DTransformSection* TransformSection, float KeyTime, FVector& KeyPos, FRotator& KeyRot)
+{
+	TransformSection->EvalTranslation(KeyTime, KeyPos);
+	TransformSection->EvalRotation(KeyTime, KeyRot);
+}
+
+void DrawTransformTrack(const FSceneView* View, FPrimitiveDrawInterface* PDI, UMovieScene3DTransformTrack* TransformTrack, const TArray<UObject*>& BoundObjects)
+{
+	const bool bHitTesting = PDI->IsHitTesting();
+	FLinearColor TrackColor = FLinearColor::Yellow; //@todo - customizable per track
+
+	for (int32 SectionIndex = 0; SectionIndex < TransformTrack->GetAllSections().Num(); ++SectionIndex)
+	{
+		UMovieSceneSection* Section = TransformTrack->GetAllSections()[SectionIndex];
+		UMovieScene3DTransformSection* TransformSection = Cast<UMovieScene3DTransformSection>(Section);
+
+		FRichCurve& TransXCurve = TransformSection->GetTranslationCurve(EAxis::X);
+		FRichCurve& TransYCurve = TransformSection->GetTranslationCurve(EAxis::Y);
+		FRichCurve& TransZCurve = TransformSection->GetTranslationCurve(EAxis::Z);
+
+		TSet<float> KeyTimes;
+
+		for (auto KeyIt(TransXCurve.GetKeyHandleIterator()); KeyIt; ++KeyIt)
+		{
+			FKeyHandle KeyHandle = KeyIt.Key();
+			float KeyTime = TransXCurve.GetKeyTime(KeyHandle);
+			KeyTimes.Add(KeyTime);
+		}
+
+		for (auto KeyIt(TransYCurve.GetKeyHandleIterator()); KeyIt; ++KeyIt)
+		{
+			FKeyHandle KeyHandle = KeyIt.Key();
+			float KeyTime = TransYCurve.GetKeyTime(KeyHandle);
+			KeyTimes.Add(KeyTime);
+		}
+
+		for (auto KeyIt(TransZCurve.GetKeyHandleIterator()); KeyIt; ++KeyIt)
+		{
+			FKeyHandle KeyHandle = KeyIt.Key();
+			float KeyTime = TransZCurve.GetKeyTime(KeyHandle);
+			KeyTimes.Add(KeyTime);
+		}
+
+		FVector OldKeyPos(0);
+		float OldKeyTime = 0.f;
+		int KeyTimeIndex = 0;
+
+		for (auto NewKeyTime : KeyTimes)
+		{
+			FVector NewKeyPos(0);
+			FRotator NewKeyRot(0,0,0);
+			GetLocationAtTime(TransformSection, NewKeyTime, NewKeyPos, NewKeyRot);
+
+			// If not the first keypoint, draw a line to the last keypoint.
+			if(KeyTimeIndex > 0)
+			{
+				int32 NumSteps = FMath::CeilToInt( (NewKeyTime - OldKeyTime)/SequencerEdMode_Draw3D::DrawTrackTimeRes );
+								// Limit the number of steps to prevent a rendering performance hit
+				NumSteps = FMath::Min( 100, NumSteps );
+				float DrawSubstep = (NewKeyTime - OldKeyTime)/NumSteps;
+
+				// Find position on first keyframe.
+				float OldTime = OldKeyTime;
+
+				FVector OldPos(0);
+				FRotator OldRot(0,0,0);
+				GetLocationAtTime(TransformSection, OldKeyTime, OldPos, OldRot);
+
+				FKeyHandle OldXKeyHandle = TransXCurve.FindKey(OldKeyTime);
+				FKeyHandle OldYKeyHandle = TransYCurve.FindKey(OldKeyTime);
+				FKeyHandle OldZKeyHandle = TransZCurve.FindKey(OldKeyTime);
+
+				bool bIsConstantKey = false;
+				if (TransXCurve.IsKeyHandleValid(OldXKeyHandle) && TransXCurve.GetKeyInterpMode(OldXKeyHandle) == ERichCurveInterpMode::RCIM_Constant &&
+					TransYCurve.IsKeyHandleValid(OldYKeyHandle) && TransYCurve.GetKeyInterpMode(OldYKeyHandle) == ERichCurveInterpMode::RCIM_Constant &&
+					TransZCurve.IsKeyHandleValid(OldZKeyHandle) && TransZCurve.GetKeyInterpMode(OldZKeyHandle) == ERichCurveInterpMode::RCIM_Constant)
+				{
+					bIsConstantKey = true;
+				}
+
+				for (auto BoundObject : BoundObjects)
+				{
+					FTransform RefTM = GetRefFrame(BoundObject);
+
+					FVector OldPos_G = RefTM.TransformPosition(OldPos);
+					FVector NewKeyPos_G = RefTM.TransformPosition(NewKeyPos);
+
+					// For constant interpolation - don't draw ticks - just draw dotted line.
+					if (bIsConstantKey)
+					{
+						DrawDashedLine(PDI, OldPos_G, NewKeyPos_G, TrackColor, 20, SDPG_Foreground);
+					}
+					else
+					{
+						// Then draw a line for each substep.
+						for (int32 j=1; j<NumSteps+1; j++)
+						{
+							float NewTime = OldKeyTime + j*DrawSubstep;
+
+							FVector NewPos(0);
+							FRotator NewRot(0,0,0);
+							GetLocationAtTime(TransformSection, NewTime, NewPos, NewRot);
+
+							FVector NewPos_G = RefTM.TransformPosition(NewPos);
+
+							PDI->DrawLine(OldPos_G, NewPos_G, TrackColor, SDPG_Foreground);
+
+							// Don't draw point for last one - its the keypoint drawn above.
+							if (j != NumSteps)
+							{
+								PDI->DrawPoint(NewPos_G, TrackColor, 3.f, SDPG_Foreground);
+							}
+							OldTime = NewTime;
+							OldPos_G = NewPos_G;
+						}
+					}
+				}
+			}
+			
+			OldKeyTime = NewKeyTime;			
+			OldKeyPos = NewKeyPos;
+			++KeyTimeIndex;			
+		}
+
+		// Draw keypoints on top of curve
+		for (auto NewKeyTime : KeyTimes)
+		{
+			// Find if this key is one of the selected ones.
+			bool bKeySelected = false;
+			//@todo
+			//for(int32 j=0; j<SelectedKeys.Num() && !bKeySelected; j++)
+			//{
+			//	if( SelectedKeys[j].Group == Group && 
+			//		SelectedKeys[j].Track == this && 
+			//		SelectedKeys[j].KeyIndex == i )
+			//		bKeySelected = true;
+			//}
+
+			// Find the time, position and orientation of this Key.
+
+			FVector NewKeyPos(0);
+			FRotator NewKeyRot(0,0,0);
+			GetLocationAtTime(TransformSection, NewKeyTime, NewKeyPos, NewKeyRot);
+
+			for (auto BoundObject : BoundObjects)
+			{
+				FTransform RefTM = GetRefFrame(BoundObject);
+
+				FColor KeyColor = bKeySelected ? SequencerEdMode_Draw3D::KeySelectedColor : TrackColor.ToFColor(true);
+
+				if (bHitTesting) 
+				{
+					PDI->SetHitProxy( new HMovieSceneKeyProxy(TransformTrack, TransformSection, NewKeyTime) );
+				}
+
+				FVector NewKeyPos_G = RefTM.TransformPosition(NewKeyPos);
+
+				PDI->DrawPoint(NewKeyPos_G, KeyColor, 6.f, SDPG_Foreground);
+
+				//@todo
+				// If desired, draw directional arrow at each keyframe.
+				//if(bShowArrowAtKeys)
+				//{
+				//	FRotationTranslationMatrix ArrowToWorld(NewKeyRot,NewKeyPos);
+				//	DrawDirectionalArrow(PDI, FScaleMatrix(FVector(16.f,16.f,16.f)) * ArrowToWorld, KeyColor, 3.f, 1.f, SDPG_Foreground );
+				//}
+				if (bHitTesting) 
+				{
+					PDI->SetHitProxy( NULL );
+				}
+
+				// If a selected key, draw handles.
+				/*
+				if (bKeySelected)
+				{
+					FVector ArriveTangent = PosTrack.Points[i].ArriveTangent;
+					FVector LeaveTangent = PosTrack.Points[i].LeaveTangent;
+
+					EInterpCurveMode PrevMode = (i > 0)							? GetKeyInterpMode(i-1) : EInterpCurveMode(255);
+					EInterpCurveMode NextMode = (i < PosTrack.Points.Num()-1)	? GetKeyInterpMode(i)	: EInterpCurveMode(255);
+
+					// If not first point, and previous mode was a curve type.
+					if(PrevMode == CIM_CurveAuto || PrevMode == CIM_CurveAutoClamped || PrevMode == CIM_CurveUser || PrevMode == CIM_CurveBreak)
+					{
+						FVector HandlePos = NewKeyPos - RefTM.TransformVector(ArriveTangent * SequencerEdMode_Draw3D::CurveHandleScale);
+						PDI->DrawLine(NewKeyPos, HandlePos, FColor(128,255,0), SDPG_Foreground);
+
+						if (bHitTesting) 
+						{ 
+							//@todo
+							//PDI->SetHitProxy( new HInterpTrackKeyHandleProxy(Group, TrackIndex, i, true) );
+						}
+						PDI->DrawPoint(HandlePos, FColor(128,255,0), 5.f, SDPG_Foreground);
+						if (bHitTesting) 
+						{
+							//@todo
+							//PDI->SetHitProxy( NULL );
+						}
+					}
+
+					// If next section is a curve, draw leaving handle.
+					if(NextMode == CIM_CurveAuto || NextMode == CIM_CurveAutoClamped || NextMode == CIM_CurveUser || NextMode == CIM_CurveBreak)
+					{
+						FVector HandlePos = NewKeyPos + RefTM.TransformVector(LeaveTangent * SequencerEdMode_Draw3D::CurveHandleScale);
+						PDI->DrawLine(NewKeyPos, HandlePos, FColor(128,255,0), SDPG_Foreground);
+
+						if (bHitTesting) 
+						{
+							//@todo
+							//PDI->SetHitProxy( new HInterpTrackKeyHandleProxy(Group, TrackIndex, i, false) );
+						}
+						PDI->DrawPoint(HandlePos, FColor(128,255,0), 5.f, SDPG_Foreground);
+						if (bHitTesting) 
+						{
+							//@todo
+							//PDI->SetHitProxy( NULL );
+						}
+					}
+				}
+				*/
+			}
+		}
+	}
+}
+
+
+void FSequencerEdMode::DrawTracks3D(const FSceneView* View, FPrimitiveDrawInterface* PDI)
+{
+	TSet<TSharedRef<FSequencerDisplayNode> > ObjectBindingNodes;
+
+	// Look for the object binding nodes from the selected nodes
+	for (auto OutlinerNode : Sequencer->GetSelection().GetSelectedOutlinerNodes())
+	{
+		TSharedRef<FSequencerDisplayNode> ObjectBindingNode = OutlinerNode;
+		if (SequencerHelpers::FindObjectBindingNode(OutlinerNode, ObjectBindingNode))
+		{
+			ObjectBindingNodes.Add(ObjectBindingNode);
+		}
+	}
+
+	// Gather up the transform track nodes from the object binding nodes
+
+	for (auto ObjectBindingNode : ObjectBindingNodes)
+	{
+		TSet<TSharedRef<FSequencerDisplayNode> > AllNodes;
+		SequencerHelpers::GetDescendantNodes(ObjectBindingNode, AllNodes);
+
+		FGuid ObjectBinding = StaticCastSharedRef<FSequencerObjectBindingNode>(ObjectBindingNode)->GetObjectBinding();
+
+		TArray<UObject*> BoundObjects;
+		Sequencer->GetRuntimeObjects(Sequencer->GetFocusedMovieSceneSequenceInstance(), ObjectBinding, BoundObjects);
+
+		for (auto DisplayNode : AllNodes)
+		{
+			if (DisplayNode->GetType() == ESequencerNode::Track)
+			{
+				TSharedRef<FSequencerTrackNode> TrackNode = StaticCastSharedRef<FSequencerTrackNode>(DisplayNode);
+				UMovieSceneTrack* TrackNodeTrack = TrackNode->GetTrack();
+				UMovieScene3DTransformTrack* TransformTrack = Cast<UMovieScene3DTransformTrack>(TrackNodeTrack);
+				if (TransformTrack != nullptr)
+				{
+					DrawTransformTrack(View, PDI, TransformTrack, BoundObjects);
+				}
+			}
+		}
+	}
+}
+
+FSequencerEdModeTool::FSequencerEdModeTool(FSequencerEdMode* InSequencerEdMode) :
+	SequencerEdMode(InSequencerEdMode)
 {
 }
 
@@ -67,5 +386,24 @@ FSequencerEdModeTool::~FSequencerEdModeTool()
 
 bool FSequencerEdModeTool::InputKey(FEditorViewportClient* ViewportClient, FViewport* Viewport, FKey Key, EInputEvent Event)
 {
+	if( Key == EKeys::LeftMouseButton )
+	{
+		if( Event == IE_Pressed)
+		{
+			int32 HitX = ViewportClient->Viewport->GetMouseX();
+			int32 HitY = ViewportClient->Viewport->GetMouseY();
+			HHitProxy*HitResult = ViewportClient->Viewport->GetHitProxy(HitX, HitY);
+
+			if(HitResult)
+			{
+				if( HitResult->IsA(HMovieSceneKeyProxy::StaticGetType()) )
+				{
+					HMovieSceneKeyProxy* KeyProxy = (HMovieSceneKeyProxy*)HitResult;
+					SequencerEdMode->OnKeySelected(ViewportClient->Viewport, KeyProxy);
+				}
+			}
+		}
+	}
+
 	return FModeTool::InputKey(ViewportClient, Viewport, Key, Event);
 }
diff --git a/Engine/Source/Editor/Sequencer/Private/SequencerEdMode.h b/Engine/Source/Editor/Sequencer/Private/SequencerEdMode.h
--- a/Engine/Source/Editor/Sequencer/Private/SequencerEdMode.h
+++ b/Engine/Source/Editor/Sequencer/Private/SequencerEdMode.h
@@ -2,6 +2,7 @@
 
 #pragma once
 #include "EditorModeTools.h"
+#include "MovieSceneHitProxy.h"
 
 class FSequencer;
 
@@ -22,9 +23,15 @@ public:
 	virtual void Exit() override;
 	virtual bool IsCompatibleWith(FEditorModeID OtherModeID) const override;
 	virtual bool InputKey( FEditorViewportClient* ViewportClient, FViewport* Viewport, FKey Key, EInputEvent Event ) override;
+	virtual void Render(const FSceneView* View,FViewport* Viewport,FPrimitiveDrawInterface* PDI) override;
 
 	void SetSequencer(FSequencer* InSequencer) { Sequencer = InSequencer; }
 
+	void OnKeySelected(FViewport* Viewport, HMovieSceneKeyProxy* KeyProxy);
+
+protected:
+	void DrawTracks3D(const FSceneView* View, FPrimitiveDrawInterface* PDI);
+
 private:
 	FSequencer* Sequencer;
 };
@@ -35,7 +42,7 @@ private:
 class FSequencerEdModeTool : public FModeTool
 {
 public:
-	FSequencerEdModeTool();
+	FSequencerEdModeTool(FSequencerEdMode* InSequencerEdMode);
 	virtual ~FSequencerEdModeTool();
 
 	virtual FString GetName() const override { return TEXT("Sequencer Edit"); }
@@ -44,4 +51,7 @@ public:
 	 * @return		true if the key was handled by this editor mode tool.
 	 */
 	virtual bool InputKey(FEditorViewportClient* ViewportClient, FViewport* Viewport, FKey Key, EInputEvent Event) override;
+
+private:
+	FSequencerEdMode* SequencerEdMode;
 };
diff --git a/Engine/Source/Editor/Sequencer/Private/SequencerLabelManager.cpp b/Engine/Source/Editor/Sequencer/Private/SequencerLabelManager.cpp
--- a/Engine/Source/Editor/Sequencer/Private/SequencerLabelManager.cpp
+++ b/Engine/Source/Editor/Sequencer/Private/SequencerLabelManager.cpp
@@ -102,3 +102,28 @@ bool FSequencerLabelManager::LabelExists(const FString& Label) const
 
 	return false;
 }
+
+
+bool FSequencerLabelManager::RenameLabel(const FString& OldLabel, const FString& NewLabel)
+{
+	if (!MovieScene.IsValid() || (OldLabel == NewLabel))
+	{
+		return false;
+	}
+
+	bool FoundOldLabel = false;
+
+	for (auto& LabelsPair : MovieScene->GetObjectsToLabels())
+	{
+		TArray<FString>& Strings = LabelsPair.Value.Strings;
+		int32 OldLabelIndex = Strings.IndexOfByKey(OldLabel);
+
+		if (OldLabelIndex > INDEX_NONE)
+		{
+			Strings[OldLabelIndex] = NewLabel;
+			FoundOldLabel = true;
+		}
+	}
+
+	return FoundOldLabel;
+}
diff --git a/Engine/Source/Editor/Sequencer/Private/SequencerLabelManager.h b/Engine/Source/Editor/Sequencer/Private/SequencerLabelManager.h
--- a/Engine/Source/Editor/Sequencer/Private/SequencerLabelManager.h
+++ b/Engine/Source/Editor/Sequencer/Private/SequencerLabelManager.h
@@ -64,6 +64,15 @@ public:
 	 */
 	bool LabelExists(const FString& Label) const;
 
+	/**
+	 * Rename the specified label to the given string.
+	 *
+	 * @param OldLabel The label to rename.
+	 * @param NewLabel The new label string.
+	 * @return true if at least one label was renamed, false otherwise.
+	 */
+	bool RenameLabel(const FString& OldLabel, const FString& NewLabel);
+
 public:
 
 	/** Get an event delegate that is executed when the movie scene's labels changed. */
diff --git a/Engine/Source/Editor/Sequencer/Private/SequencerSettings.cpp b/Engine/Source/Editor/Sequencer/Private/SequencerSettings.cpp
--- a/Engine/Source/Editor/Sequencer/Private/SequencerSettings.cpp
+++ b/Engine/Source/Editor/Sequencer/Private/SequencerSettings.cpp
@@ -5,7 +5,7 @@
 USequencerSettings::USequencerSettings( const FObjectInitializer& ObjectInitializer )
 	: Super( ObjectInitializer )
 {
-	bAutoKeyEnabled = false;
+	AutoKeyMode = EAutoKeyMode::KeyNone;
 	bKeyAllEnabled = false;
 	bKeyInterpPropertiesOnly = false;
 	KeyInterpolation = EMovieSceneKeyInterpolation::Auto;
@@ -29,18 +29,19 @@ USequencerSettings::USequencerSettings( const FObjectInitializer& ObjectInitiali
 	bShowCurveEditorCurveToolTips = true;
 	bLooping = false;
 	bKeepCursorInPlayRange = true;
+	bKeepPlayRangeInSectionBounds = true;
 }
 
-bool USequencerSettings::GetAutoKeyEnabled() const
+EAutoKeyMode USequencerSettings::GetAutoKeyMode() const
 {
-	return bAutoKeyEnabled;
+	return AutoKeyMode;
 }
 
-void USequencerSettings::SetAutoKeyEnabled(bool InbAutoKeyEnabled)
+void USequencerSettings::SetAutoKeyMode(EAutoKeyMode InAutoKeyMode)
 {
-	if ( bAutoKeyEnabled != InbAutoKeyEnabled )
+	if ( AutoKeyMode != InAutoKeyMode )
 	{
-		bAutoKeyEnabled = InbAutoKeyEnabled;
+		AutoKeyMode = InAutoKeyMode;
 		SaveConfig();
 	}
 }
@@ -355,6 +356,20 @@ void USequencerSettings::SetKeepCursorInPlayRange(bool bInKeepCursorInPlayRange)
 	}
 }
 
+bool USequencerSettings::ShouldKeepPlayRangeInSectionBounds() const
+{
+	return bKeepPlayRangeInSectionBounds;
+}
+
+void USequencerSettings::SetKeepPlayRangeInSectionBounds(bool bInKeepPlayRangeInSectionBounds)
+{
+	if (bKeepPlayRangeInSectionBounds != bInKeepPlayRangeInSectionBounds)
+	{
+		bKeepPlayRangeInSectionBounds = bInKeepPlayRangeInSectionBounds;
+		SaveConfig();
+	}
+}
+
 bool USequencerSettings::GetShowCurveEditorCurveToolTips() const
 {
 	return bShowCurveEditorCurveToolTips;
@@ -388,4 +403,5 @@ ULevelEditorSequencerSettings::ULevelEditorSequencerSettings( const FObjectIniti
 	bKeyInterpPropertiesOnly = true;
 	TimeSnapInterval = 0.033334f;
 	bShowRangeSlider = true;
+	bKeepPlayRangeInSectionBounds = false;
 }
\ No newline at end of file
diff --git a/Engine/Source/Editor/Sequencer/Private/SequencerSettings.h b/Engine/Source/Editor/Sequencer/Private/SequencerSettings.h
--- a/Engine/Source/Editor/Sequencer/Private/SequencerSettings.h
+++ b/Engine/Source/Editor/Sequencer/Private/SequencerSettings.h
@@ -54,10 +54,10 @@ public:
 
 	DECLARE_MULTICAST_DELEGATE( FOnShowCurveEditorChanged );
 
-	/** Gets whether or not auto key is enabled. */
-	bool GetAutoKeyEnabled() const;
-	/** Sets whether or not auto key is enabled. */
-	void SetAutoKeyEnabled(bool InbAutoKeyEnabled);
+	/** Gets the current auto-key mode. */
+	EAutoKeyMode GetAutoKeyMode() const;
+	/** Sets the current auto-key mode. */
+	void SetAutoKeyMode(EAutoKeyMode AutoKeyMode);
 
 	/** Gets whether or not key all is enabled. */
 	bool GetKeyAllEnabled() const;
@@ -176,6 +176,11 @@ public:
 	/** Set whether or not the cursor should be kept within the playback range during playback in sequencer */
 	void SetKeepCursorInPlayRange(bool bInKeepCursorInPlayRange);
 
+	/** @return true if the playback range should be synced to the section bounds, false otherwise */
+	bool ShouldKeepPlayRangeInSectionBounds() const;
+	/** Set whether or not the playback range should be synced to the section bounds */
+	void SetKeepPlayRangeInSectionBounds(bool bInKeepPlayRangeInSectionBounds);
+
 	/** Snaps a time value in seconds to the currently selected interval. */
 	float SnapTimeToInterval(float InTimeValue) const;
 
@@ -184,7 +189,7 @@ public:
 
 protected:
 	UPROPERTY( config )
-	bool bAutoKeyEnabled;
+	TEnumAsByte<EAutoKeyMode> AutoKeyMode;
 
 	UPROPERTY( config )
 	bool bKeyAllEnabled;
@@ -255,6 +260,9 @@ protected:
 	UPROPERTY( config )
 	bool bKeepCursorInPlayRange;
 
+	UPROPERTY( config )
+	bool bKeepPlayRangeInSectionBounds;
+
 	FOnShowCurveEditorChanged OnShowCurveEditorChanged;
 };
 
