Copying //UE4/Dev-Sequencer to Dev-Main (//UE4/Dev-Main)

diff --git a/Engine/Source/Editor/Sequencer/Private/DisplayNodes/SequencerTrackNode.cpp b/Engine/Source/Editor/Sequencer/Private/DisplayNodes/SequencerTrackNode.cpp
--- a/Engine/Source/Editor/Sequencer/Private/DisplayNodes/SequencerTrackNode.cpp
+++ b/Engine/Source/Editor/Sequencer/Private/DisplayNodes/SequencerTrackNode.cpp
@@ -8,7 +8,7 @@
 #include "MovieSceneNameableTrack.h"
 #include "Sequencer.h"
 #include "ISequencerTrackEditor.h"
-
+#include "SKeyNavigationButtons.h"
 
 namespace SequencerNodeConstants
 {
@@ -114,11 +114,9 @@ bool FSequencerTrackNode::CanRenameNode() const
 	return (Track != nullptr) && Track->IsA(UMovieSceneNameableTrack::StaticClass());
 }
 
-
-TSharedRef<SWidget> FSequencerTrackNode::GenerateEditWidgetForOutliner()
+TSharedRef<SWidget> FSequencerTrackNode::GetCustomOutlinerContent()
 {
 	TSharedPtr<FSequencerSectionKeyAreaNode> KeyAreaNode = GetTopLevelKeyNode();
-
 	if (KeyAreaNode.IsValid())
 	{
 		// @todo - Sequencer - Support multiple sections/key areas?
@@ -126,10 +124,44 @@ TSharedRef<SWidget> FSequencerTrackNode::GenerateEditWidgetForOutliner()
 
 		if (KeyAreas.Num() > 0)
 		{
+			// Create the widgets for the key editor and key navigation buttons
+			TSharedRef<SHorizontalBox> BoxPanel = SNew(SHorizontalBox);
+
 			if (KeyAreas[0]->CanCreateKeyEditor())
 			{
-				return KeyAreas[0]->CreateKeyEditor(&GetSequencer());
+				BoxPanel->AddSlot()
+				.HAlign(HAlign_Right)
+				.VAlign(VAlign_Center)
+				[
+					SNew(SBox)
+					.WidthOverride(100)
+					.HAlign(HAlign_Left)
+					[
+						KeyAreas[0]->CreateKeyEditor(&GetSequencer())
+					]
+				];
 			}
+			else
+			{
+				BoxPanel->AddSlot()
+				[
+					SNew(SSpacer)
+				];
+			}
+
+			BoxPanel->AddSlot()
+			.AutoWidth()
+			.VAlign(VAlign_Center)
+			[
+				SNew(SKeyNavigationButtons, AsShared())
+			];
+
+			return SNew(SBox)
+				.VAlign(VAlign_Center)
+				.HAlign(HAlign_Right)
+				[
+					BoxPanel
+				];
 		}
 	}
 	else
@@ -141,15 +173,30 @@ TSharedRef<SWidget> FSequencerTrackNode::GenerateEditWidgetForOutliner()
 		{
 			ObjectBinding = StaticCastSharedPtr<FSequencerObjectBindingNode>(ParentNode)->GetObjectBinding();
 		}
+		FBuildEditWidgetParams Params;
+		Params.NodeIsHovered = TAttribute<bool>::Create(TAttribute<bool>::FGetter::CreateSP(this, &FSequencerDisplayNode::IsHovered));
 
-		TSharedPtr<SWidget> EditWidget = AssociatedEditor.BuildOutlinerEditWidget(ObjectBinding, AssociatedTrack.Get());
-		if (EditWidget.IsValid())
+		TSharedPtr<SWidget> Widget = AssociatedEditor.BuildOutlinerEditWidget(ObjectBinding, AssociatedTrack.Get(), Params);
+
+		if (Widget.IsValid())
 		{
-			return EditWidget.ToSharedRef();
+			return SNew(SBox)
+				.VAlign(VAlign_Center)
+				.HAlign(HAlign_Right)
+				[
+					Widget.ToSharedRef()
+				];
 		}
 	}
 
-	return FSequencerDisplayNode::GenerateEditWidgetForOutliner();
+
+	return SNew(SSpacer);
+}
+
+
+const FSlateBrush* FSequencerTrackNode::GetIconBrush() const
+{
+	return AssociatedEditor.GetIconBrush();
 }
 
 
diff --git a/Engine/Source/Editor/Sequencer/Private/DisplayNodes/SequencerTrackNode.h b/Engine/Source/Editor/Sequencer/Private/DisplayNodes/SequencerTrackNode.h
--- a/Engine/Source/Editor/Sequencer/Private/DisplayNodes/SequencerTrackNode.h
+++ b/Engine/Source/Editor/Sequencer/Private/DisplayNodes/SequencerTrackNode.h
@@ -90,13 +90,14 @@ public:
 
 	virtual void BuildContextMenu( FMenuBuilder& MenuBuilder );
 	virtual bool CanRenameNode() const override;
-	virtual TSharedRef<SWidget> GenerateEditWidgetForOutliner() override;
+	virtual TSharedRef<SWidget> GetCustomOutlinerContent() override;
 	virtual void GetChildKeyAreaNodesRecursively(TArray<TSharedRef<FSequencerSectionKeyAreaNode>>& OutNodes) const override;
 	virtual FText GetDisplayName() const override;
 	virtual float GetNodeHeight() const override;
 	virtual FNodePadding GetNodePadding() const override;
 	virtual ESequencerNode::Type GetType() const override;
 	virtual void SetDisplayName(const FText& NewDisplayName) override;
+	virtual const FSlateBrush* GetIconBrush() const override;
 
 private:
 
diff --git a/Engine/Source/Editor/Sequencer/Private/GroupedKeyArea.cpp b/Engine/Source/Editor/Sequencer/Private/GroupedKeyArea.cpp
--- a/Engine/Source/Editor/Sequencer/Private/GroupedKeyArea.cpp
+++ b/Engine/Source/Editor/Sequencer/Private/GroupedKeyArea.cpp
@@ -343,6 +343,18 @@ ERichCurveExtrapolation FGroupedKeyArea::GetExtrapolationMode(bool bPreInfinity)
 	return ExtrapMode;
 }
 
+bool FGroupedKeyArea::CanSetExtrapolationMode() const
+{
+	for (auto& Area : KeyAreas)
+	{
+		if (Area->CanSetExtrapolationMode())
+		{
+			return true;
+		}
+	}
+	return false;
+}
+
 TArray<FKeyHandle> FGroupedKeyArea::AddKeyUnique(float Time, EMovieSceneKeyInterpolation InKeyInterpolation, float TimeToCopyFrom)
 {
 	TArray<FKeyHandle> AddedKeyHandles;
@@ -471,6 +483,26 @@ FLinearColor FGroupedKeyArea::GetColor()
 }
 
 
+TSharedPtr<FStructOnScope> FGroupedKeyArea::GetKeyStruct(FKeyHandle KeyHandle)
+{
+	FKeyGrouping* Group = FindGroup(KeyHandle);
+	
+	if (Group == nullptr)
+	{
+		return nullptr;
+	}
+
+	TArray<FKeyHandle> KeyHandles;
+
+	for (auto& Key : Group->Keys)
+	{
+		KeyHandles.Add(Key.KeyHandle);
+	}
+
+	return Section->GetKeyStruct(KeyHandles);
+}
+
+
 void FGroupedKeyArea::PasteKeys(const FMovieSceneClipboardKeyTrack& KeyTrack, const FMovieSceneClipboardEnvironment& SrcEnvironment, const FSequencerPasteEnvironment& DstEnvironment)
 {
 	checkf(false, TEXT("Pasting into grouped key areas is not supported, and should not be used. Iterate child tracks instead."));
diff --git a/Engine/Source/Editor/Sequencer/Private/GroupedKeyArea.h b/Engine/Source/Editor/Sequencer/Private/GroupedKeyArea.h
--- a/Engine/Source/Editor/Sequencer/Private/GroupedKeyArea.h
+++ b/Engine/Source/Editor/Sequencer/Private/GroupedKeyArea.h
@@ -4,9 +4,11 @@
 
 #include "IKeyArea.h"
 
+
 class FSequencerDisplayNode;
 class UMovieSceneSection;
 
+
 /** Keyable structure used to identify a particular FGroupedKeyArea */
 struct FIndexKey
 {
@@ -23,14 +25,18 @@ struct FIndexKey
 	}
 
 private:
+
 	/** Cached path of the node that we relate to */
 	FString NodePath;
+
 	/** The movie scene section that we relate to */
 	TWeakObjectPtr<UMovieSceneSection> Section;
+
 	/** Cached hash for efficiency (to avoid having to re-hash strings all the time) */
 	uint32 CachedHash;
 };
 
+
 /** Structure that represents a number of keys */
 struct FKeyGrouping
 {
@@ -68,37 +74,17 @@ struct FKeyGrouping
 	TArray<FKeyIndex> Keys;
 };
 
+
 /** A key area that represents multiple, distinct key areas */
-class FGroupedKeyArea : public IKeyArea
+class FGroupedKeyArea
+	: public IKeyArea
 {
 public:
+
 	/** Construct the area from a display node */
 	FGroupedKeyArea(FSequencerDisplayNode& InNode, int32 InSectionIndex);
-	
-	/* IKeyArea interface */
-	virtual TArray<FKeyHandle>		GetUnsortedKeyHandles() const override;
-	virtual void 					SetKeyTime(FKeyHandle KeyHandle, float NewKeyTime) const override;
-	virtual float					GetKeyTime(FKeyHandle KeyHandle) const override;
-	virtual FKeyHandle				MoveKey(FKeyHandle KeyHandle, float DeltaPosition) override;
-	virtual void					DeleteKey(FKeyHandle KeyHandle) override;
-	virtual void                    SetKeyInterpMode(FKeyHandle KeyHandle, ERichCurveInterpMode InterpMode) override;
-	virtual ERichCurveInterpMode    GetKeyInterpMode(FKeyHandle KeyHandle) const override;
-	virtual void                    SetKeyTangentMode(FKeyHandle KeyHandle, ERichCurveTangentMode TangentMode) override;
-	virtual ERichCurveTangentMode   GetKeyTangentMode(FKeyHandle KeyHandle) const override;
-	virtual void                    SetExtrapolationMode(ERichCurveExtrapolation ExtrapMode, bool bPreInfinity) override;
-	virtual ERichCurveExtrapolation GetExtrapolationMode(bool bPreInfinity) const override;
-	virtual TArray<FKeyHandle>		AddKeyUnique(float Time, EMovieSceneKeyInterpolation InKeyInterpolation, float TimeToCopyFrom = FLT_MAX) override;
-	virtual TOptional<FKeyHandle>				DuplicateKey(FKeyHandle KeyToDuplicate) override;
-	virtual FRichCurve*				GetRichCurve() override;
-	virtual UMovieSceneSection*		GetOwningSection() override;
-	virtual bool					CanCreateKeyEditor() override;
-	virtual TSharedRef<SWidget>		CreateKeyEditor(ISequencer* Sequencer) override;
-	virtual void					SetName(FName Name) override {}
-	virtual FName					GetName() const override { return NAME_None; }
-	virtual void					CopyKeys(FMovieSceneClipboardBuilder& ClipboardBuilder, const TFunctionRef<bool(FKeyHandle, const IKeyArea&)>& KeyMask) const override;
-	virtual void					PasteKeys(const FMovieSceneClipboardKeyTrack& KeyTrack, const FMovieSceneClipboardEnvironment& SrcEnvironment, const FSequencerPasteEnvironment& DstEnvironment) override;
-	virtual FLinearColor GetColor() override;
-	/* End IKeyArea interface */
+
+public:
 
 	/** Get the desired tint for the specified key handle */
 	FLinearColor GetKeyTint( FKeyHandle InHandle ) const;
@@ -106,6 +92,35 @@ public:
 	/** Get the brush name for the specified key handle */
 	const FSlateBrush* GetBrush(FKeyHandle InHandle) const;
 
+public:
+
+	// IKeyArea interface
+
+	virtual TArray<FKeyHandle> GetUnsortedKeyHandles() const override;
+	virtual void SetKeyTime(FKeyHandle KeyHandle, float NewKeyTime) const override;
+	virtual float GetKeyTime(FKeyHandle KeyHandle) const override;
+	virtual FKeyHandle MoveKey(FKeyHandle KeyHandle, float DeltaPosition) override;
+	virtual void DeleteKey(FKeyHandle KeyHandle) override;
+	virtual void SetKeyInterpMode(FKeyHandle KeyHandle, ERichCurveInterpMode InterpMode) override;
+	virtual ERichCurveInterpMode GetKeyInterpMode(FKeyHandle KeyHandle) const override;
+	virtual void SetKeyTangentMode(FKeyHandle KeyHandle, ERichCurveTangentMode TangentMode) override;
+	virtual TSharedPtr<FStructOnScope> GetKeyStruct(FKeyHandle KeyHandle) override;
+	virtual ERichCurveTangentMode GetKeyTangentMode(FKeyHandle KeyHandle) const override;
+	virtual void SetExtrapolationMode(ERichCurveExtrapolation ExtrapMode, bool bPreInfinity) override;
+	virtual ERichCurveExtrapolation GetExtrapolationMode(bool bPreInfinity) const override;
+	virtual bool CanSetExtrapolationMode() const override;
+	virtual TArray<FKeyHandle> AddKeyUnique(float Time, EMovieSceneKeyInterpolation InKeyInterpolation, float TimeToCopyFrom = FLT_MAX) override;
+	virtual TOptional<FKeyHandle> DuplicateKey(FKeyHandle KeyToDuplicate) override;
+	virtual FRichCurve* GetRichCurve() override;
+	virtual UMovieSceneSection* GetOwningSection() override;
+	virtual bool CanCreateKeyEditor() override;
+	virtual TSharedRef<SWidget> CreateKeyEditor(ISequencer* Sequencer) override;
+	virtual void SetName(FName Name) override {}
+	virtual FName GetName() const override { return NAME_None; }
+	virtual void CopyKeys(FMovieSceneClipboardBuilder& ClipboardBuilder, const TFunctionRef<bool(FKeyHandle, const IKeyArea&)>& KeyMask) const override;
+	virtual void PasteKeys(const FMovieSceneClipboardKeyTrack& KeyTrack, const FMovieSceneClipboardEnvironment& SrcEnvironment, const FSequencerPasteEnvironment& DstEnvironment) override;
+	virtual FLinearColor GetColor() override;
+
 protected:
 
 	/** Update the persistent index with our new key handles */
@@ -121,10 +136,13 @@ protected:
 	
 	/** Pointer to the section to which this key area relates */
 	UMovieSceneSection* Section;
+
 	/** Array of (child) key areas that we are reflecting */
 	TArray<TSharedRef<IKeyArea>> KeyAreas;
+
 	/** Generated array of groups of keys harvested from the above array */
 	TArray<FKeyGrouping> Groups;
+
 	/** Key into our index */
 	FIndexKey IndexKey;
-};
\ No newline at end of file
+};
diff --git a/Engine/Source/Editor/Sequencer/Private/KeyAreaLayout.cpp b/Engine/Source/Editor/Sequencer/Private/KeyAreaLayout.cpp
--- a/Engine/Source/Editor/Sequencer/Private/KeyAreaLayout.cpp
+++ b/Engine/Source/Editor/Sequencer/Private/KeyAreaLayout.cpp
@@ -3,13 +3,14 @@
 #include "SequencerPrivatePCH.h"
 #include "KeyAreaLayout.h"
 
-FKeyAreaLayoutElement FKeyAreaLayoutElement::FromGroup(const TSharedRef<FGroupedKeyArea>& InKeyAreaGroup, float InOffset, TOptional<float> InHeight)
+FKeyAreaLayoutElement FKeyAreaLayoutElement::FromGroup(const TSharedRef<FSequencerDisplayNode>& InNode, const TSharedRef<FGroupedKeyArea>& InKeyAreaGroup, float InOffset, TOptional<float> InHeight)
 {
 	FKeyAreaLayoutElement Tmp;
 	Tmp.Type = Group;
 	Tmp.KeyArea = InKeyAreaGroup;
 	Tmp.LocalOffset = InOffset;
 	Tmp.Height = InHeight;
+	Tmp.DisplayNode = InNode;
 	return Tmp;
 }
 
@@ -19,7 +20,7 @@ FKeyAreaLayoutElement FKeyAreaLayoutElement::FromKeyAreaNode(const TSharedRef<FS
 	Tmp.Type = Single;
 	Tmp.KeyArea = InKeyAreaNode->GetKeyArea(SectionIndex);
 	Tmp.LocalOffset = InOffset;
-
+	Tmp.DisplayNode = InKeyAreaNode;
 	if (!InKeyAreaNode->IsTopLevel())
 	{
 		Tmp.Height = InKeyAreaNode->GetNodeHeight();
@@ -47,6 +48,11 @@ TSharedRef<IKeyArea> FKeyAreaLayoutElement::GetKeyArea() const
 	return KeyArea.ToSharedRef();
 }
 
+TSharedPtr<FSequencerDisplayNode> FKeyAreaLayoutElement::GetDisplayNode() const
+{
+	return DisplayNode;
+}
+
 FKeyAreaLayout::FKeyAreaLayout(FSequencerDisplayNode& InNode, int32 InSectionIndex)
 {
 	TotalHeight = 0.f;
@@ -61,7 +67,7 @@ FKeyAreaLayout::FKeyAreaLayout(FSequencerDisplayNode& InNode, int32 InSectionInd
 		}
 		else if (!Node.IsExpanded())
 		{
-			Elements.Add(FKeyAreaLayoutElement::FromGroup(Node.UpdateKeyGrouping(InSectionIndex), TotalHeight, Node.GetNodeHeight()));
+			Elements.Add(FKeyAreaLayoutElement::FromGroup(Node.AsShared(), Node.UpdateKeyGrouping(InSectionIndex), TotalHeight, Node.GetNodeHeight()));
 		}
 
 	};
diff --git a/Engine/Source/Editor/Sequencer/Private/KeyAreaLayout.h b/Engine/Source/Editor/Sequencer/Private/KeyAreaLayout.h
--- a/Engine/Source/Editor/Sequencer/Private/KeyAreaLayout.h
+++ b/Engine/Source/Editor/Sequencer/Private/KeyAreaLayout.h
@@ -11,7 +11,10 @@ struct FKeyAreaLayoutElement
 	enum EType { Single, Group };
 
 	/** Construct this element from a grouped key area */
-	static FKeyAreaLayoutElement FromGroup(const TSharedRef<FGroupedKeyArea>& InKeyAreaGroup, float InOffset, TOptional<float> InHeight = TOptional<float>());
+	static FKeyAreaLayoutElement FromNonKeyArea(const TSharedRef<FSequencerDisplayNode>& InNode, float InOffset, TOptional<float> InHeight);
+
+	/** Construct this element from a grouped key area */
+	static FKeyAreaLayoutElement FromGroup(const TSharedRef<FSequencerDisplayNode>& InNode, const TSharedRef<FGroupedKeyArea>& InKeyAreaGroup, float InOffset, TOptional<float> InHeight = TOptional<float>());
 
 	/** Construct this element from a single Key area node */
 	static FKeyAreaLayoutElement FromKeyAreaNode(const TSharedRef<FSequencerSectionKeyAreaNode>& InKeyAreaNode, int32 SectionIndex, float InOffset);
@@ -28,11 +31,17 @@ struct FKeyAreaLayoutElement
 	/** Access the key area that this layout element was generated for */
 	TSharedRef<IKeyArea> GetKeyArea() const;
 
+	/** Access the display node that this layout element was generated for */
+	TSharedPtr<FSequencerDisplayNode> GetDisplayNode() const;
+
 private:
 
 	/** Pointer to the key area that we were generated from */
 	TSharedPtr<IKeyArea> KeyArea;
 
+	/** The specific node that this key area relates to */
+	TSharedPtr<FSequencerDisplayNode> DisplayNode;
+
 	/** The type of this layout element */
 	EType Type;
 
diff --git a/Engine/Source/Editor/Sequencer/Private/SAnimationOutlinerTreeNode.cpp b/Engine/Source/Editor/Sequencer/Private/SAnimationOutlinerTreeNode.cpp
--- a/Engine/Source/Editor/Sequencer/Private/SAnimationOutlinerTreeNode.cpp
+++ b/Engine/Source/Editor/Sequencer/Private/SAnimationOutlinerTreeNode.cpp
@@ -15,21 +15,30 @@
 
 #define LOCTEXT_NAMESPACE "AnimationOutliner"
 
+SAnimationOutlinerTreeNode::~SAnimationOutlinerTreeNode()
+{
+	DisplayNode->OnRenameRequested().RemoveAll(this);
+}
 
 void SAnimationOutlinerTreeNode::Construct( const FArguments& InArgs, TSharedRef<FSequencerDisplayNode> Node, const TSharedRef<SSequencerTreeViewRow>& InTableRow )
 {
 	DisplayNode = Node;
+	bIsTopLevelNode = !Node->GetParent().IsValid();
+
+	if (bIsTopLevelNode)
+	{
+		ExpandedBackgroundBrush = FEditorStyle::GetBrush( "Sequencer.AnimationOutliner.TopLevelBorder_Expanded" );
+		CollapsedBackgroundBrush = FEditorStyle::GetBrush( "Sequencer.AnimationOutliner.TopLevelBorder_Collapsed" );
+	}
+	else
+	{
+		ExpandedBackgroundBrush = FEditorStyle::GetBrush( "Sequencer.AnimationOutliner.DefaultBorder" );
+		CollapsedBackgroundBrush = FEditorStyle::GetBrush( "Sequencer.AnimationOutliner.DefaultBorder" );
+	}
 
-	HoveredBrush = FEditorStyle::GetBrush( "Sequencer.AnimationOutliner.SelectionBorderHover" );
-	SelectedBrush = FEditorStyle::GetBrush( "Sequencer.AnimationOutliner.SelectionBorder" );
-	SelectedBrushInactive = FEditorStyle::GetBrush("Sequencer.AnimationOutliner.SelectionBorderInactive");
-	NotSelectedBrush = FEditorStyle::GetBrush( "NoBorder" );
 	TableRowStyle = &FEditorStyle::Get().GetWidgetStyle<FTableRowStyle>("TableView.Row");
 
-	// Choose the font.  If the node is a root node or an object node, we show a larger font for it.
-	FSlateFontInfo NodeFont = Node->GetParent().IsValid() && Node->GetType() != ESequencerNode::Object ? 
-		FEditorStyle::GetFontStyle("Sequencer.AnimationOutliner.RegularFont") :
-		FEditorStyle::GetFontStyle("Sequencer.AnimationOutliner.BoldFont");
+	FSlateFontInfo NodeFont = FEditorStyle::GetFontStyle("Sequencer.AnimationOutliner.RegularFont");
 
 	EditableLabel = SNew(SEditableLabel)
 		.CanEdit(this, &SAnimationOutlinerTreeNode::HandleNodeLabelCanEdit)
@@ -37,122 +46,78 @@ void SAnimationOutlinerTreeNode::Construct( const FArguments& InArgs, TSharedRef
 		.OnTextChanged(this, &SAnimationOutlinerTreeNode::HandleNodeLabelTextChanged)
 		.Text(this, &SAnimationOutlinerTreeNode::GetDisplayName);
 
+
+	Node->OnRenameRequested().AddRaw(this, &SAnimationOutlinerTreeNode::EnterRenameMode);
+
 	auto NodeHeight = [=]() -> FOptionalSize { return DisplayNode->GetNodeHeight(); };
 
-	TAttribute<FLinearColor> HoverTint(this, &SAnimationOutlinerTreeNode::GetHoverTint);
 	ForegroundColor.Bind(this, &SAnimationOutlinerTreeNode::GetForegroundBasedOnSelection);
 
 	TSharedRef<SWidget>	FinalWidget = 
 		SNew( SBorder )
 		.VAlign( VAlign_Center )
 		.BorderImage( this, &SAnimationOutlinerTreeNode::GetNodeBorderImage )
+		.BorderBackgroundColor( this, &SAnimationOutlinerTreeNode::GetNodeBackgroundTint )
 		.Padding(FMargin(0, Node->GetNodePadding().Combined() / 2))
 		[
 			SNew(SBox)
 			.HeightOverride_Lambda(NodeHeight)
-			.Padding(FMargin(0.0f, 0.0f, 10.0f, 0.0f))
+			.Padding(FMargin(5.0f, 0.0f))
 			[
 				SNew( SHorizontalBox )
 
 				// Expand track lanes button
 				+ SHorizontalBox::Slot()
-				.Padding(FMargin(2.f, 0.f))
+				.Padding(FMargin(2.f, 0.f, 4.f, 0.f))
 				.VAlign( VAlign_Center )
 				.AutoWidth()
 				[
 					SNew(SExpanderArrow, InTableRow).IndentAmount(SequencerLayoutConstants::IndentAmount)
 				]
 
-				// Label Slot
-				+ SHorizontalBox::Slot()
-				.VAlign( VAlign_Center )
-				.FillWidth(3)
-				[
-					EditableLabel.ToSharedRef()
-				]
-
-				// Editor slot
+				// Icon
 				+ SHorizontalBox::Slot()
+				.Padding(FMargin(0.f, 0.f, 4.f, 0.f))
 				.VAlign(VAlign_Center)
 				.AutoWidth()
 				[
-					// @todo Sequencer - Remove this box and width override.
-					SNew(SBox)
-					.HAlign(HAlign_Left)
-					.WidthOverride(100)
+					SNew(SOverlay)
+
+					+ SOverlay::Slot()
 					[
-						DisplayNode->GenerateEditWidgetForOutliner()
+						SNew(SImage)
+						.Image(InArgs._IconBrush)
 					]
-				]
-				// Previous key slot
-				+ SHorizontalBox::Slot()
-				.VAlign(VAlign_Center)
-				.AutoWidth()
-				.Padding(3, 0, 0, 0)
-				[
-					SNew(SBorder)
-					.Padding(0)
-					.BorderImage(NotSelectedBrush)
-					.ColorAndOpacity(HoverTint)
+
+					+ SOverlay::Slot()
+					.VAlign(VAlign_Top)
+					.HAlign(HAlign_Right)
 					[
-						SNew(SButton)
-						.ButtonStyle(FEditorStyle::Get(), "FlatButton")
-						.ToolTipText(LOCTEXT("PreviousKeyButton", "Set the time to the previous key"))
-						.OnClicked(this, &SAnimationOutlinerTreeNode::OnPreviousKeyClicked)
-						.ForegroundColor( FSlateColor::UseForeground() )
-						.ContentPadding(0)
-						[
-							SNew(STextBlock)
-							.Font(FEditorStyle::Get().GetFontStyle("FontAwesome.7"))
-							.Text(FText::FromString(FString(TEXT("\xf060"))) /*fa-arrow-left*/)
-						]
+						SNew(SImage)
+						.Image(InArgs._IconOverlayBrush)
+					]
+
+					+ SOverlay::Slot()
+					[
+						SNew(SSpacer)
+						.Visibility(EVisibility::Visible)
+						.ToolTipText(InArgs._IconToolTipText)
 					]
 				]
-				// Add key slot
+
+				// Label Slot
 				+ SHorizontalBox::Slot()
 				.VAlign(VAlign_Center)
+				.Padding(FMargin(0.f, 0.f, 4.f, 0.f))
 				.AutoWidth()
 				[
-					SNew(SBorder)
-					.Padding(0)
-					.BorderImage(NotSelectedBrush)
-					.ColorAndOpacity(HoverTint)
-					[
-						SNew(SButton)
-						.ButtonStyle(FEditorStyle::Get(), "FlatButton")
-						.ToolTipText(LOCTEXT("AddKeyButton", "Add a new key at the current time"))
-						.OnClicked(this, &SAnimationOutlinerTreeNode::OnAddKeyClicked)
-						.ForegroundColor( FSlateColor::UseForeground() )
-						.ContentPadding(0)
-						[
-							SNew(STextBlock)
-							.Font(FEditorStyle::Get().GetFontStyle("FontAwesome.7"))
-							.Text(FText::FromString(FString(TEXT("\xf055"))) /*fa-plus-circle*/)
-						]
-					]
+					EditableLabel.ToSharedRef()
 				]
-				// Next key slot
+
+				// Arbitrary customization slot
 				+ SHorizontalBox::Slot()
-				.VAlign(VAlign_Center)
-				.AutoWidth()
 				[
-					SNew(SBorder)
-					.Padding(0)
-					.BorderImage(NotSelectedBrush)
-					.ColorAndOpacity(HoverTint)
-					[
-						SNew(SButton)
-						.ButtonStyle(FEditorStyle::Get(), "FlatButton")
-						.ToolTipText(LOCTEXT("NextKeyButton", "Set the time to the next key"))
-						.OnClicked(this, &SAnimationOutlinerTreeNode::OnNextKeyClicked)
-						.ContentPadding(0)
-						.ForegroundColor( FSlateColor::UseForeground() )
-						[
-							SNew(STextBlock)
-							.Font(FEditorStyle::Get().GetFontStyle("FontAwesome.7"))
-							.Text(FText::FromString(FString(TEXT("\xf061"))) /*fa-arrow-right*/)
-						]
-					]
+					InArgs._CustomContent.Widget
 				]
 			]
 		];
@@ -169,117 +134,6 @@ void SAnimationOutlinerTreeNode::EnterRenameMode()
 	EditableLabel->EnterTextMode();
 }
 
-
-FLinearColor SAnimationOutlinerTreeNode::GetHoverTint() const
-{
-	return IsHovered() ? FLinearColor(1,1,1,0.9f) : FLinearColor(1,1,1,0.4f);
-}
-
-
-FReply SAnimationOutlinerTreeNode::OnPreviousKeyClicked()
-{
-	FSequencer& Sequencer = DisplayNode->GetSequencer();
-	float ClosestPreviousKeyDistance = MAX_FLT;
-	float CurrentTime = Sequencer.GetCurrentLocalTime(*Sequencer.GetFocusedMovieSceneSequence());
-	float PreviousTime = 0;
-	bool PreviousKeyFound = false;
-
-	TSet<TSharedPtr<IKeyArea>> KeyAreas;
-	SequencerHelpers::GetAllKeyAreas(DisplayNode, KeyAreas);
-	for (TSharedPtr<IKeyArea> KeyArea : KeyAreas)
-	{
-		for (FKeyHandle& KeyHandle : KeyArea->GetUnsortedKeyHandles())
-		{
-			float KeyTime = KeyArea->GetKeyTime(KeyHandle);
-			if (KeyTime < CurrentTime && CurrentTime - KeyTime < ClosestPreviousKeyDistance)
-			{
-				PreviousTime = KeyTime;
-				ClosestPreviousKeyDistance = CurrentTime - KeyTime;
-				PreviousKeyFound = true;
-			}
-		}
-	}
-
-	if (PreviousKeyFound)
-	{
-		Sequencer.SetGlobalTime(PreviousTime);
-	}
-	return FReply::Handled();
-}
-
-
-FReply SAnimationOutlinerTreeNode::OnNextKeyClicked()
-{
-	FSequencer& Sequencer = DisplayNode->GetSequencer();
-	float ClosestNextKeyDistance = MAX_FLT;
-	float CurrentTime = Sequencer.GetCurrentLocalTime(*Sequencer.GetFocusedMovieSceneSequence());
-	float NextTime = 0;
-	bool NextKeyFound = false;
-
-	TSet<TSharedPtr<IKeyArea>> KeyAreas;
-	SequencerHelpers::GetAllKeyAreas(DisplayNode, KeyAreas);
-	for (TSharedPtr<IKeyArea> KeyArea : KeyAreas)
-	{
-		for (FKeyHandle& KeyHandle : KeyArea->GetUnsortedKeyHandles())
-		{
-			float KeyTime = KeyArea->GetKeyTime(KeyHandle);
-			if (KeyTime > CurrentTime && KeyTime - CurrentTime < ClosestNextKeyDistance)
-			{
-				NextTime = KeyTime;
-				ClosestNextKeyDistance = KeyTime - CurrentTime;
-				NextKeyFound = true;
-			}
-		}
-	}
-
-	if (NextKeyFound)
-	{
-		Sequencer.SetGlobalTime(NextTime);
-	}
-
-	return FReply::Handled();
-}
-
-
-FReply SAnimationOutlinerTreeNode::OnAddKeyClicked()
-{
-	FSequencer& Sequencer = DisplayNode->GetSequencer();
-	float CurrentTime = Sequencer.GetCurrentLocalTime(*Sequencer.GetFocusedMovieSceneSequence());
-
-	TSet<TSharedPtr<IKeyArea>> KeyAreas;
-	SequencerHelpers::GetAllKeyAreas(DisplayNode, KeyAreas);
-
-	TArray<UMovieSceneSection*> KeyAreaSections;
-	for (TSharedPtr<IKeyArea> KeyArea : KeyAreas)
-	{
-		UMovieSceneSection* OwningSection = KeyArea->GetOwningSection();
-		KeyAreaSections.Add(OwningSection);
-	}
-
-	UMovieSceneSection* NearestSection = MovieSceneHelpers::FindNearestSectionAtTime(KeyAreaSections, CurrentTime);
-	if (!NearestSection)
-	{
-		return FReply::Unhandled();
-	}
-
-	FScopedTransaction Transaction(LOCTEXT("AddKeys", "Add keys at current time"));
-	for (TSharedPtr<IKeyArea> KeyArea : KeyAreas)
-	{
-		UMovieSceneSection* OwningSection = KeyArea->GetOwningSection();
-		if (OwningSection == NearestSection)
-		{
-			OwningSection->SetFlags(RF_Transactional);
-			if (OwningSection->TryModify())
-			{
-				KeyArea->AddKeyUnique(CurrentTime, Sequencer.GetKeyInterpolation());
-			}
-		}
-	}
-
-	return FReply::Handled();
-}
-
-
 TSharedPtr<FSequencerDisplayNode> SAnimationOutlinerTreeNode::GetRootNode(TSharedPtr<FSequencerDisplayNode> ObjectNode)
 {
 	if (!ObjectNode->GetParent().IsValid())
@@ -391,7 +245,7 @@ FReply SAnimationOutlinerTreeNode::OnMouseButtonDown( const FGeometry& MyGeometr
 			Sequencer.GetSelection().AddToSelection(DisplayNode.ToSharedRef());
 		}
 	}
-	else
+	else if (MouseEvent.GetEffectingButton() != EKeys::RightMouseButton || !bSelected)
 	{
 		// Deselect the other nodes and select this node.
 		Sequencer.GetSelection().EmptySelectedOutlinerNodes();
@@ -434,31 +288,50 @@ FReply SAnimationOutlinerTreeNode::OnMouseButtonUp( const FGeometry& MyGeometry,
 	return FReply::Handled().SetUserFocus(MenuContent.ToSharedRef(), EFocusCause::SetDirectly);
 }
 
+void SAnimationOutlinerTreeNode::OnMouseEnter(const FGeometry& MyGeometry, const FPointerEvent& MouseEvent)
+{
+	DisplayNode->GetParentTree().SetHoveredNode(DisplayNode);
+	SWidget::OnMouseEnter(MyGeometry, MouseEvent);
+}
+
+void SAnimationOutlinerTreeNode::OnMouseLeave(const FPointerEvent& MouseEvent)
+{
+	DisplayNode->GetParentTree().SetHoveredNode(nullptr);
+	SWidget::OnMouseLeave(MouseEvent);
+}
 
 const FSlateBrush* SAnimationOutlinerTreeNode::GetNodeBorderImage() const
+{
+	return DisplayNode->IsExpanded() ? ExpandedBackgroundBrush : CollapsedBackgroundBrush;
+}
+
+FSlateColor SAnimationOutlinerTreeNode::GetNodeBackgroundTint() const
 {
 	FSequencer& Sequencer = DisplayNode->GetSequencer();
 	const bool bIsSelected = Sequencer.GetSelection().IsSelected(DisplayNode.ToSharedRef());
+	const bool bIsSelectionActive = Sequencer.GetSelection().GetActiveSelection() == FSequencerSelection::EActiveSelection::OutlinerNode;
 
-	if (!bIsSelected)
+	if (bIsSelected)
 	{
-		if (DisplayNode->IsHovered())
+		if (bIsSelectionActive)
 		{
-			return HoveredBrush;
+			return FEditorStyle::GetSlateColor("SelectionColor_Pressed");
+		}
+		else
+		{
+			return FEditorStyle::GetSlateColor("SelectionColor_Inactive");
 		}
-
-		return NotSelectedBrush;
 	}
-
-	if (Sequencer.GetSelection().GetActiveSelection() == FSequencerSelection::EActiveSelection::OutlinerNode)
+	else if (DisplayNode->IsHovered())
 	{
-		return SelectedBrush;
+		return bIsTopLevelNode ? FLinearColor(FColor(52, 52, 52, 255)) : FLinearColor(FColor(72, 72, 72, 255));
+	}
+	else
+	{
+		return bIsTopLevelNode ? FLinearColor(FColor(48, 48, 48, 255)) : FLinearColor(FColor(62, 62, 62, 255));
 	}
-
-	return SelectedBrushInactive;
 }
 
-
 FSlateColor SAnimationOutlinerTreeNode::GetForegroundBasedOnSelection() const
 {
 	FSequencer& Sequencer = DisplayNode->GetSequencer();
diff --git a/Engine/Source/Editor/Sequencer/Private/SAnimationOutlinerTreeNode.h b/Engine/Source/Editor/Sequencer/Private/SAnimationOutlinerTreeNode.h
--- a/Engine/Source/Editor/Sequencer/Private/SAnimationOutlinerTreeNode.h
+++ b/Engine/Source/Editor/Sequencer/Private/SAnimationOutlinerTreeNode.h
@@ -16,7 +16,14 @@ class SAnimationOutlinerTreeNode
 {
 public:
 
+	SAnimationOutlinerTreeNode(){}
+	~SAnimationOutlinerTreeNode();
+	
 	SLATE_BEGIN_ARGS(SAnimationOutlinerTreeNode){}
+		SLATE_ATTRIBUTE(const FSlateBrush*, IconBrush)
+		SLATE_ATTRIBUTE(const FSlateBrush*, IconOverlayBrush)
+		SLATE_ATTRIBUTE(FText, IconToolTipText)
+		SLATE_NAMED_SLOT(FArguments, CustomContent)
 	SLATE_END_ARGS()
 
 	void Construct( const FArguments& InArgs, TSharedRef<FSequencerDisplayNode> Node, const TSharedRef<SSequencerTreeViewRow>& InTableRow );
@@ -34,11 +41,6 @@ public:
 		return DisplayNode;
 	}
 
-	/**
-	 * Gets a tint to apply buttons on hover.
-	 */
-	FLinearColor GetHoverTint() const;
-
 private:
 
 	// SWidget interface
@@ -46,6 +48,9 @@ private:
 	virtual FReply OnMouseButtonDown(const FGeometry& MyGeometry, const FPointerEvent& MouseEvent) override;
 	virtual FReply OnMouseButtonUp(const FGeometry& MyGeometry, const FPointerEvent& MouseEvent) override;
 
+	void OnMouseEnter(const FGeometry& MyGeometry, const FPointerEvent& MouseEvent) override;
+	void OnMouseLeave(const FPointerEvent& MouseEvent) override;
+
 	FSlateColor GetForegroundBasedOnSelection() const;
 
 	/**
@@ -53,6 +58,11 @@ private:
 	 */
 	const FSlateBrush* GetNodeBorderImage() const;
 	
+	/**
+	 * @return The tint to apply to the border image
+	 */
+	FSlateColor GetNodeBackgroundTint() const;
+
 	/**
 	 * @return The expander visibility of this node.
 	 */
@@ -69,15 +79,6 @@ private:
 	/** Callback for when the node label text has changed. */
 	void HandleNodeLabelTextChanged(const FText& NewLabel);
 
-	/** Handles the previous key button being clicked. */
-	FReply OnPreviousKeyClicked();
-
-	/** Handles the next key button being clicked. */
-	FReply OnNextKeyClicked();
-
-	/** Handles the add key button being clicked. */
-	FReply OnAddKeyClicked();
-
 	/** Get all descendant nodes from the given root node. */
 	void GetAllDescendantNodes(TSharedPtr<FSequencerDisplayNode> RootNode, TArray<TSharedRef<FSequencerDisplayNode> >& AllNodes);
 
@@ -95,17 +96,14 @@ private:
 	/** Holds the editable text label widget. */
 	TSharedPtr<SEditableLabel> EditableLabel;
 
-	/** Brush to use if the node is hovered. */
-	const FSlateBrush* HoveredBrush;
-
-	/** Brush to display a border around the widget when it is selected. */
-	const FSlateBrush* SelectedBrush;
+	/** True if this node is a top level node, false otherwise */
+	bool bIsTopLevelNode;
 
-	/** Brush to display a border around the widget when it is selected but inactive. */
-	const FSlateBrush* SelectedBrushInactive;
+	/** Default background brush for this node when expanded */
+	const FSlateBrush* ExpandedBackgroundBrush;
 
-	/** Brush to use if the node is not selected. */
-	const FSlateBrush* NotSelectedBrush;
+	/** Default background brush for this node when collapsed */
+	const FSlateBrush* CollapsedBackgroundBrush;
 
 	/** The table row style used for nodes in the tree. This is required as we don't actually use the tree for selection. */
 	const FTableRowStyle* TableRowStyle;
diff --git a/Engine/Source/Editor/Sequencer/Private/SKeyNavigationButtons.h b/Engine/Source/Editor/Sequencer/Private/SKeyNavigationButtons.h
new file mode 100644
--- /dev/null
+++ b/Engine/Source/Editor/Sequencer/Private/SKeyNavigationButtons.h
@@ -0,0 +1,230 @@
+// Copyright 1998-2015 Epic Games, Inc. All Rights Reserved.
+
+#pragma once
+
+#include "MovieSceneCommonHelpers.h"
+#include "ScopedTransaction.h"
+#include "MovieSceneSection.h"
+
+#define LOCTEXT_NAMESPACE "Arse"
+
+/**
+ * A widget for navigating between keys on a sequencer track
+ */
+class SKeyNavigationButtons
+	: public SCompoundWidget
+{
+public:
+
+	SLATE_BEGIN_ARGS(SKeyNavigationButtons){}
+	SLATE_END_ARGS()
+
+	void Construct(const FArguments& InArgs, const TSharedRef<FSequencerDisplayNode>& InDisplayNode)
+	{
+		DisplayNode = InDisplayNode;
+
+		const FSlateBrush* NoBorder = FEditorStyle::GetBrush( "NoBorder" );
+
+		TAttribute<FLinearColor> HoverTint(this, &SKeyNavigationButtons::GetHoverTint);
+
+		ChildSlot
+		[
+			SNew(SHorizontalBox)
+			
+			// Previous key slot
+			+ SHorizontalBox::Slot()
+			.VAlign(VAlign_Center)
+			.AutoWidth()
+			.Padding(3, 0, 0, 0)
+			[
+				SNew(SBorder)
+				.Padding(0)
+				.BorderImage(NoBorder)
+				.ColorAndOpacity(HoverTint)
+				[
+					SNew(SButton)
+					.ButtonStyle(FEditorStyle::Get(), "FlatButton")
+					.ToolTipText(LOCTEXT("PreviousKeyButton", "Set the time to the previous key"))
+					.OnClicked(this, &SKeyNavigationButtons::OnPreviousKeyClicked)
+					.ForegroundColor( FSlateColor::UseForeground() )
+					.ContentPadding(0)
+					[
+						SNew(STextBlock)
+						.Font(FEditorStyle::Get().GetFontStyle("FontAwesome.7"))
+						.Text(FText::FromString(FString(TEXT("\xf060"))) /*fa-arrow-left*/)
+					]
+				]
+			]
+			// Add key slot
+			+ SHorizontalBox::Slot()
+			.VAlign(VAlign_Center)
+			.AutoWidth()
+			[
+				SNew(SBorder)
+				.Padding(0)
+				.BorderImage(NoBorder)
+				.ColorAndOpacity(HoverTint)
+				[
+					SNew(SButton)
+					.ButtonStyle(FEditorStyle::Get(), "FlatButton")
+					.ToolTipText(LOCTEXT("AddKeyButton", "Add a new key at the current time"))
+					.OnClicked(this, &SKeyNavigationButtons::OnAddKeyClicked)
+					.ForegroundColor( FSlateColor::UseForeground() )
+					.ContentPadding(0)
+					[
+						SNew(STextBlock)
+						.Font(FEditorStyle::Get().GetFontStyle("FontAwesome.7"))
+						.Text(FText::FromString(FString(TEXT("\xf055"))) /*fa-plus-circle*/)
+					]
+				]
+			]
+			// Next key slot
+			+ SHorizontalBox::Slot()
+			.VAlign(VAlign_Center)
+			.AutoWidth()
+			[
+				SNew(SBorder)
+				.Padding(0)
+				.BorderImage(NoBorder)
+				.ColorAndOpacity(HoverTint)
+				[
+					SNew(SButton)
+					.ButtonStyle(FEditorStyle::Get(), "FlatButton")
+					.ToolTipText(LOCTEXT("NextKeyButton", "Set the time to the next key"))
+					.OnClicked(this, &SKeyNavigationButtons::OnNextKeyClicked)
+					.ContentPadding(0)
+					.ForegroundColor( FSlateColor::UseForeground() )
+					[
+						SNew(STextBlock)
+						.Font(FEditorStyle::Get().GetFontStyle("FontAwesome.7"))
+						.Text(FText::FromString(FString(TEXT("\xf061"))) /*fa-arrow-right*/)
+					]
+				]
+			]
+		];
+	}
+
+	/** Handles the previous key button being clicked. */
+	//FReply OnPreviousKeyClicked();
+
+	/** Handles the next key button being clicked. */
+	//FReply OnNextKeyClicked();
+
+	/** Handles the add key button being clicked. */
+	//FReply OnAddKeyClicked();
+
+
+	FLinearColor GetHoverTint() const
+	{
+		return DisplayNode->IsHovered() ? FLinearColor(1,1,1,0.9f) : FLinearColor(1,1,1,0.4f);
+	}
+
+	FReply OnPreviousKeyClicked()
+	{
+		FSequencer& Sequencer = DisplayNode->GetSequencer();
+		float ClosestPreviousKeyDistance = MAX_FLT;
+		float CurrentTime = Sequencer.GetCurrentLocalTime(*Sequencer.GetFocusedMovieSceneSequence());
+		float PreviousTime = 0;
+		bool PreviousKeyFound = false;
+
+		TSet<TSharedPtr<IKeyArea>> KeyAreas;
+		SequencerHelpers::GetAllKeyAreas(DisplayNode, KeyAreas);
+		for (TSharedPtr<IKeyArea> KeyArea : KeyAreas)
+		{
+			for (FKeyHandle& KeyHandle : KeyArea->GetUnsortedKeyHandles())
+			{
+				float KeyTime = KeyArea->GetKeyTime(KeyHandle);
+				if (KeyTime < CurrentTime && CurrentTime - KeyTime < ClosestPreviousKeyDistance)
+				{
+					PreviousTime = KeyTime;
+					ClosestPreviousKeyDistance = CurrentTime - KeyTime;
+					PreviousKeyFound = true;
+				}
+			}
+		}
+
+		if (PreviousKeyFound)
+		{
+			Sequencer.SetGlobalTime(PreviousTime);
+		}
+		return FReply::Handled();
+	}
+
+
+	FReply OnNextKeyClicked()
+	{
+		FSequencer& Sequencer = DisplayNode->GetSequencer();
+		float ClosestNextKeyDistance = MAX_FLT;
+		float CurrentTime = Sequencer.GetCurrentLocalTime(*Sequencer.GetFocusedMovieSceneSequence());
+		float NextTime = 0;
+		bool NextKeyFound = false;
+
+		TSet<TSharedPtr<IKeyArea>> KeyAreas;
+		SequencerHelpers::GetAllKeyAreas(DisplayNode, KeyAreas);
+		for (TSharedPtr<IKeyArea> KeyArea : KeyAreas)
+		{
+			for (FKeyHandle& KeyHandle : KeyArea->GetUnsortedKeyHandles())
+			{
+				float KeyTime = KeyArea->GetKeyTime(KeyHandle);
+				if (KeyTime > CurrentTime && KeyTime - CurrentTime < ClosestNextKeyDistance)
+				{
+					NextTime = KeyTime;
+					ClosestNextKeyDistance = KeyTime - CurrentTime;
+					NextKeyFound = true;
+				}
+			}
+		}
+
+		if (NextKeyFound)
+		{
+			Sequencer.SetGlobalTime(NextTime);
+		}
+
+		return FReply::Handled();
+	}
+
+
+	FReply OnAddKeyClicked()
+	{
+		FSequencer& Sequencer = DisplayNode->GetSequencer();
+		float CurrentTime = Sequencer.GetCurrentLocalTime(*Sequencer.GetFocusedMovieSceneSequence());
+
+		TSet<TSharedPtr<IKeyArea>> KeyAreas;
+		SequencerHelpers::GetAllKeyAreas(DisplayNode, KeyAreas);
+
+		TArray<UMovieSceneSection*> KeyAreaSections;
+		for (TSharedPtr<IKeyArea> KeyArea : KeyAreas)
+		{
+			UMovieSceneSection* OwningSection = KeyArea->GetOwningSection();
+			KeyAreaSections.Add(OwningSection);
+		}
+
+		UMovieSceneSection* NearestSection = MovieSceneHelpers::FindNearestSectionAtTime(KeyAreaSections, CurrentTime);
+		if (!NearestSection)
+		{
+			return FReply::Unhandled();
+		}
+
+		FScopedTransaction Transaction(LOCTEXT("AddKeys", "Add keys at current time"));
+		for (TSharedPtr<IKeyArea> KeyArea : KeyAreas)
+		{
+			UMovieSceneSection* OwningSection = KeyArea->GetOwningSection();
+			if (OwningSection == NearestSection)
+			{
+				OwningSection->SetFlags(RF_Transactional);
+				if (OwningSection->TryModify())
+				{
+					KeyArea->AddKeyUnique(CurrentTime, Sequencer.GetKeyInterpolation());
+				}
+			}
+		}
+
+		Sequencer.UpdatePlaybackRange();
+
+		return FReply::Handled();
+	}
+
+	TSharedPtr<FSequencerDisplayNode> DisplayNode;
+};
+
+#undef LOCTEXT_NAMESPACE
\ No newline at end of file
diff --git a/Engine/Source/Editor/Sequencer/Private/SSequencer.cpp b/Engine/Source/Editor/Sequencer/Private/SSequencer.cpp
--- a/Engine/Source/Editor/Sequencer/Private/SSequencer.cpp
+++ b/Engine/Source/Editor/Sequencer/Private/SSequencer.cpp
@@ -18,10 +18,8 @@
 #include "DragAndDrop/AssetDragDropOp.h"
 #include "DragAndDrop/ActorDragDropGraphEdOp.h"
 #include "DragAndDrop/ClassDragDropOp.h"
-#include "Tools/SequencerEditTool_Selection.h"
-#include "Tools/SequencerEditTool_Movement.h"
 #include "AssetSelection.h"
-#include "MovieSceneShotSection.h"
+#include "MovieSceneCameraCutSection.h"
 #include "CommonMovieSceneTools.h"
 #include "SSearchBox.h"
 #include "SNumericDropDown.h"
@@ -32,45 +30,52 @@
 #include "IKeyArea.h"
 #include "VirtualTrackArea.h"
 #include "SequencerHotspots.h"
-#include "IDetailsView.h"
-#include "PropertyEditorModule.h"
 #include "SSequencerShotFilterOverlay.h"
 #include "GenericCommands.h"
 #include "SequencerContextMenus.h"
 #include "SSequencerTreeViewBox.h"
 #include "NumericTypeInterface.h"
 #include "NumericUnitTypeInterface.inl"
+#include "SNumericEntryBox.h"
+#include "SSequencerGotoBox.h"
+
 
 #define LOCTEXT_NAMESPACE "Sequencer"
 
+DECLARE_DELEGATE_RetVal(bool, FOnGetShowFrames)
+DECLARE_DELEGATE_RetVal(uint8, FOnGetZeroPad)
 
 /* Numeric type interface for showing numbers as frames or times */
 struct FFramesOrTimeInterface : public TNumericUnitTypeInterface<float>
 {
-	DECLARE_DELEGATE_RetVal(bool, FOnGetShowFrames);
-
-	FFramesOrTimeInterface(FOnGetShowFrames InShowFrameNumbers, TSharedPtr<FSequencerTimeSliderController> InController)
+	FFramesOrTimeInterface(FOnGetShowFrames InShowFrameNumbers,TSharedPtr<FSequencerTimeSliderController> InController, FOnGetZeroPad InOnGetZeroPad)
 		: TNumericUnitTypeInterface(EUnit::Seconds)
-		, ShowFrameNumbers(InShowFrameNumbers)
+		, ShowFrameNumbers(MoveTemp(InShowFrameNumbers))
 		, TimeSliderController(MoveTemp(InController))
+		, OnGetZeroPad(MoveTemp(InOnGetZeroPad))
 	{}
 
 private:
 	FOnGetShowFrames ShowFrameNumbers;
 	TSharedPtr<FSequencerTimeSliderController> TimeSliderController;
+	FOnGetZeroPad OnGetZeroPad;
 
 	virtual FString ToString(const float& Value) const override
 	{
 		if (ShowFrameNumbers.Execute())
 		{
 			int32 Frame = TimeSliderController->TimeToFrame(Value);
+			if (OnGetZeroPad.IsBound())
+			{
+				return FString::Printf(*FString::Printf(TEXT("%%0%dd"), OnGetZeroPad.Execute()), Frame);
+			}
 			return FString::Printf(TEXT("%d"), Frame);
 		}
 
 		return FString::Printf(TEXT("%.2fs"), Value);
 	}
 
-	virtual TOptional<float> FromString(const FString& InString) override
+	virtual TOptional<float> FromString(const FString& InString, const float& InExistingValue) override
 	{
 		if (ShowFrameNumbers.Execute())
 		{
@@ -78,17 +83,17 @@ private:
 			return float(TimeSliderController->FrameToTime(NewEndFrame));
 		}
 
-		return TNumericUnitTypeInterface::FromString(InString);
+		return TNumericUnitTypeInterface::FromString(InString, InExistingValue);
 	}
 };
 
 
 /* SSequencer interface
  *****************************************************************************/
-
+PRAGMA_DISABLE_OPTIMIZATION
 void SSequencer::Construct(const FArguments& InArgs, TSharedRef<FSequencer> InSequencer)
 {
-	Sequencer = InSequencer;
+	SequencerPtr = InSequencer;
 	bIsActiveTimerRegistered = false;
 	bUserIsSelecting = false;
 
@@ -124,7 +129,20 @@ void SSequencer::Construct(const FArguments& InArgs, TSharedRef<FSequencer> InSe
 
 	TSharedRef<FSequencerTimeSliderController> TimeSliderController( new FSequencerTimeSliderController( TimeSliderArgs ) );
 	
-	NumericTypeInterface = MakeShareable( new FFramesOrTimeInterface(FFramesOrTimeInterface::FOnGetShowFrames::CreateSP(this, &SSequencer::ShowFrameNumbers), TimeSliderController) );
+	{
+		auto ShowFrameNumbersDelegate = FOnGetShowFrames::CreateSP(this, &SSequencer::ShowFrameNumbers);
+		USequencerSettings* SequencerSettings = Settings;
+		auto GetZeroPad = [=]() -> uint8 {
+			if (SequencerSettings)
+			{
+				return SequencerSettings->GetZeroPadFrames();
+			}
+			return 0;
+		};
+
+		NumericTypeInterface = MakeShareable( new FFramesOrTimeInterface(ShowFrameNumbersDelegate, TimeSliderController, FOnGetZeroPad()) );
+		ZeroPadNumericTypeInterface = MakeShareable( new FFramesOrTimeInterface(ShowFrameNumbersDelegate, TimeSliderController, FOnGetZeroPad::CreateLambda(GetZeroPad)) );
+	}
 
 	bool bMirrorLabels = false;
 	// Create the top and bottom sliders
@@ -151,10 +169,7 @@ void SSequencer::Construct(const FArguments& InArgs, TSharedRef<FSequencer> InSe
 		.Thickness(FVector2D(5.0f, 5.0f));
 
 	SAssignNew(TrackOutliner, SSequencerTrackOutliner);
-
-	SAssignNew(TrackArea, SSequencerTrackArea, TimeSliderController, InSequencer)
-		.Visibility(this, &SSequencer::GetTrackAreaVisibility );
-	
+	SAssignNew(TrackArea, SSequencerTrackArea, TimeSliderController, InSequencer);
 	SAssignNew(TreeView, SSequencerTreeView, SequencerNodeTree.ToSharedRef(), TrackArea.ToSharedRef())
 		.ExternalScrollbar(ScrollBar);
 
@@ -164,289 +179,297 @@ void SSequencer::Construct(const FArguments& InArgs, TSharedRef<FSequencer> InSe
 		.ViewRange(InArgs._ViewRange);
 
 	CurveEditor->SetAllowAutoFrame(Settings->GetShowCurveEditor());
-
 	TrackArea->SetTreeView(TreeView);
 
-	EditTool.Reset(new FSequencerEditTool_Movement(InSequencer, SharedThis(this)));
-
-	// initialize details view
-	FDetailsViewArgs DetailsViewArgs;
-	{
-		DetailsViewArgs.bAllowSearch = false;
-		DetailsViewArgs.bCustomFilterAreaLocation = true;
-		DetailsViewArgs.bCustomNameAreaLocation = true;
-		DetailsViewArgs.bHideSelectionTip = true;
-		DetailsViewArgs.bLockable = false;
-		DetailsViewArgs.bSearchInitialKeyFocus = true;
-		DetailsViewArgs.bUpdatesFromSelection = false;
-		DetailsViewArgs.NotifyHook = this;
-		DetailsViewArgs.bShowOptions = false;
-		DetailsViewArgs.bShowModifiedPropertiesOption = false;
-	}
-
-	DetailsView = FModuleManager::GetModuleChecked<FPropertyEditorModule>("PropertyEditor").CreateDetailView(DetailsViewArgs);
-	{
-		DetailsView->SetEnabled(TAttribute<bool>::Create(TAttribute<bool>::FGetter::CreateSP(this, &SSequencer::HandleDetailsViewEnabled)));
-		DetailsView->SetVisibility(TAttribute<EVisibility>::Create(TAttribute<EVisibility>::FGetter::CreateSP(this, &SSequencer::HandleDetailsViewVisibility)));
-	}
-
 	const int32 Column0 = 0, Column1 = 1;
-	const int32 Row0 = 0, Row1 = 1, Row2 = 2;
+	const int32 Row0 = 0, Row1 = 1, Row2 = 2, Row3 = 3;
+
+	const float CommonPadding = 3.f;
+	const FMargin ResizeBarPadding(4.f, 0, 0, 0);
 
 	ChildSlot
 	[
-		SNew(SBorder)
-			.Padding(2.0f)
-			.BorderImage(FEditorStyle::GetBrush("ToolPanel.GroupBorder"))
+		SNew(SVerticalBox)
+
+		+ SVerticalBox::Slot()
+		[
+			SNew(SSplitter)
+			.Orientation(Orient_Horizontal)
+			
+			+ SSplitter::Slot()
+			.Value(0.1f)
+			[
+				SNew(SBorder)
+				.BorderImage(FEditorStyle::GetBrush("ToolPanel.GroupBorder"))
+				.Visibility(this, &SSequencer::HandleLabelBrowserVisibility)
+				[
+					// track label browser
+					SAssignNew(LabelBrowser, SSequencerLabelBrowser, InSequencer)
+					.OnSelectionChanged(this, &SSequencer::HandleLabelBrowserSelectionChanged)
+				]
+			]
+
+			+ SSplitter::Slot()
+			.Value(0.9f)
 			[
-				SNew(SVerticalBox)
+				SNew(SOverlay)
 
-				+ SVerticalBox::Slot()
-					.AutoHeight()
+				+ SOverlay::Slot()
+				[
+					// track area grid panel
+					SNew( SGridPanel )
+					.FillRow( 2, 1.f )
+					.FillColumn( 0, FillCoefficient_0 )
+					.FillColumn( 1, FillCoefficient_1 )
+
+					// Toolbar
+					+ SGridPanel::Slot( Column0, Row0, SGridPanel::Layer(10) )
+					.ColumnSpan(2)
 					[
-						SNew(SHorizontalBox)
+						SNew(SBorder)
+						.BorderImage(FEditorStyle::GetBrush("ToolPanel.GroupBorder"))
+						.Padding(FMargin(CommonPadding, 0.f))
+						[
+							SNew(SHorizontalBox)
 
-						+SHorizontalBox::Slot()
+							+SHorizontalBox::Slot()
 							.AutoWidth()
 							[
 								MakeToolBar()
 							]
 
-						+SHorizontalBox::Slot()
+							+SHorizontalBox::Slot()
 							.AutoWidth()
 							[
 								SNew(SSequencerCurveEditorToolBar, InSequencer, CurveEditor->GetCommands())
-									.Visibility(this, &SSequencer::GetCurveEditorToolBarVisibility)
+								.Visibility(this, &SSequencer::GetCurveEditorToolBarVisibility)
+							]
+
+							+ SHorizontalBox::Slot()
+							.HAlign(HAlign_Right)
+							.VAlign(VAlign_Center)
+							[
+								SAssignNew(BreadcrumbTrail, SBreadcrumbTrail<FSequencerBreadcrumb>)
+								.Visibility(this, &SSequencer::GetBreadcrumbTrailVisibility)
+								.OnCrumbClicked(this, &SSequencer::OnCrumbClicked)
+								.ButtonStyle(FEditorStyle::Get(), "FlatButton")
+								.DelimiterImage(FEditorStyle::GetBrush("Sequencer.BreadcrumbIcon"))
+								.TextStyle(FEditorStyle::Get(), "Sequencer.BreadcrumbText")
 							]
+						]
 					]
 
-				+ SVerticalBox::Slot()
+					+ SGridPanel::Slot( Column0, Row1 )
 					[
-						SNew(SSplitter)
-							.Orientation(Orient_Horizontal)
-						
-						+ SSplitter::Slot()
-							.Value(0.10f)
+						SNew(SBorder)
+						.BorderImage(FEditorStyle::GetBrush("ToolPanel.GroupBorder"))
+						[
+							SNew(SSpacer)
+						]
+					]
+
+					// outliner search box
+					+ SGridPanel::Slot( Column0, Row1, SGridPanel::Layer(10) )
+					[
+						SNew(SBorder)
+						.BorderImage(FEditorStyle::GetBrush("ToolPanel.GroupBorder"))
+						.Padding(FMargin(CommonPadding*2, CommonPadding))
+						[
+							SNew(SHorizontalBox)
+
+							+ SHorizontalBox::Slot()
+							.AutoWidth()
+							.VAlign(VAlign_Center)
+							.Padding(FMargin(0.f, 0.f, CommonPadding, 0.f))
 							[
-								// track label browser
-								SAssignNew(LabelBrowser, SSequencerLabelBrowser, InSequencer)
-									.OnSelectionChanged(this, &SSequencer::HandleLabelBrowserSelectionChanged)
-									.Visibility(this, &SSequencer::HandleLabelBrowserVisibility)
+								MakeAddButton()
 							]
 
-						+ SSplitter::Slot()
-							.Value(0.7f)
+							+ SHorizontalBox::Slot()
+							.VAlign(VAlign_Center)
+							[
+								SAssignNew(SearchBox, SSearchBox)
+								.HintText(LOCTEXT("FilterNodesHint", "Filter"))
+								.OnTextChanged( this, &SSequencer::OnOutlinerSearchChanged )
+							]
+						]
+					]
+
+					// main sequencer area
+					+ SGridPanel::Slot( Column0, Row2, SGridPanel::Layer(10) )
+					.ColumnSpan(2)
+					[
+						SNew(SHorizontalBox)
+
+						+ SHorizontalBox::Slot()
+						[
+							SNew( SOverlay )
+
+							+ SOverlay::Slot()
 							[
-								SNew(SOverlay)
+								SNew(SScrollBorder, TreeView.ToSharedRef())
+								[
+									SNew(SHorizontalBox)
 
-								+ SOverlay::Slot()
+									// outliner tree
+									+ SHorizontalBox::Slot()
+									.FillWidth( FillCoefficient_0 )
 									[
-										// track area grid panel
-										SNew(SBox)
-											.Padding(FMargin(0.0f, 3.0f, 0.0f, 0.0f))
-											[
-												SNew( SGridPanel )
-													.FillRow( 1, 1.f )
-													.FillColumn( 0, FillCoefficient_0 )
-													.FillColumn( 1, FillCoefficient_1 )
-
-												// outliner search box
-												+ SGridPanel::Slot( Column0, Row0 )
-													.Padding(FMargin(0.0f, 0.0f, 3.0f, 0.0f))
-													.VAlign( VAlign_Center )
-													[
-														SNew(SBorder)
-															.BorderImage(FEditorStyle::GetBrush("ToolPanel.GroupBorder"))
-															.BorderBackgroundColor(FLinearColor(.50f, .50f, .50f, 1.0f))
-															.Padding(FMargin(4.0f, 3.0f))
-															[
-																SAssignNew(SearchBox, SSearchBox)
-																	.OnTextChanged( this, &SSequencer::OnOutlinerSearchChanged )
-															]
-													]
-
-												// main sequencer area
-												+ SGridPanel::Slot( Column0, Row1 )
-													.ColumnSpan(2)
-													[
-														SNew(SHorizontalBox)
-
-														+ SHorizontalBox::Slot()
-															[
-																SNew( SOverlay )
-
-																+ SOverlay::Slot()
-																	[
-																		SNew(SHorizontalBox)
-								
-																		// outliner tree
-																		+ SHorizontalBox::Slot()
-																			.FillWidth( FillCoefficient_0 )
-																			[
-																				SNew(SSequencerTreeViewBox, InSequencer, SharedThis(this))
-																					[
-																						TreeView.ToSharedRef()
-																					]
-																			]
-
-																		// track area
-																		+ SHorizontalBox::Slot()
-																			.FillWidth( FillCoefficient_1 )
-																			[
-																				TrackArea.ToSharedRef()
-																			]
-																	]
-
-																+ SOverlay::Slot()
-																	.HAlign( HAlign_Right )
-																	[
-																		ScrollBar
-																	]
-															]
-
-														+ SHorizontalBox::Slot()
-															.FillWidth( TAttribute<float>( this, &SSequencer::GetOutlinerSpacerFill ) )
-															[
-																SNew(SSpacer)
-															]
-													]
-
-												// playback buttons
-												+ SGridPanel::Slot( Column0, Row2 )
-													.Padding(FMargin(0.0f, 0.0f, 3.0f, 0.0f))
-													[
-														SNew(SBorder)
-															.BorderImage(FEditorStyle::GetBrush("ToolPanel.GroupBorder"))
-															.BorderBackgroundColor(FLinearColor(.50f, .50f, .50f, 1.0f))
-															.HAlign(HAlign_Center)
-															[
-																MakeTransportControls()
-															]
-													]
-
-												// Second column
-												+ SGridPanel::Slot( Column1, Row0 )
-													[
-														SNew( SBorder )
-															.BorderImage( FEditorStyle::GetBrush("ToolPanel.GroupBorder") )
-															.BorderBackgroundColor( FLinearColor(.50f, .50f, .50f, 1.0f ) )
-															.Padding(0)
-															[
-																TopTimeSlider
-															]
-													]
-
-												// Overlay that draws the tick lines
-												+ SGridPanel::Slot( Column1, Row1, SGridPanel::Layer(0) )
-													[
-														SNew( SSequencerSectionOverlay, TimeSliderController )
-															.Visibility( EVisibility::HitTestInvisible )
-															.DisplayScrubPosition( false )
-															.DisplayTickLines( true )
-													]
-
-												// Curve editor
-												+ SGridPanel::Slot( Column1, Row1 )
-													[
-														CurveEditor.ToSharedRef()
-													]
-
-												// Overlay that draws the scrub position
-												+ SGridPanel::Slot( Column1, Row1, SGridPanel::Layer(30) )
-													[
-														SNew( SSequencerSectionOverlay, TimeSliderController )
-															.Visibility( EVisibility::HitTestInvisible )
-															.DisplayScrubPosition( true )
-															.DisplayTickLines( false )
-															.PaintPlaybackRangeArgs(this, &SSequencer::GetSectionPlaybackRangeArgs)
-													]
-
-												// play range sliders
-												+ SGridPanel::Slot( Column1, Row2 )
-													[
-														SNew( SBorder )
-															.BorderImage( FEditorStyle::GetBrush("ToolPanel.GroupBorder") )
-															.BorderBackgroundColor( FLinearColor(.50f, .50f, .50f, 1.0f ) )
-															.Padding(0)
-															[
-																SNew(SVerticalBox)
-
-																+ SVerticalBox::Slot()
-																	.AutoHeight()
-																	[
-																		SNew( SOverlay )
-
-																		+ SOverlay::Slot()
-																			[
-																				BottomTimeSlider
-																			]
-
-																		+ SOverlay::Slot()
-																			[
-																				BottomTimeRange
-																			]
-																	]
-
-																+ SVerticalBox::Slot()
-																	.AutoHeight()
-																	[
-																		SAssignNew(BreadcrumbTrail, SBreadcrumbTrail<FSequencerBreadcrumb>)
-																		.Visibility(this, &SSequencer::GetBreadcrumbTrailVisibility)
-																		.OnCrumbClicked(this, &SSequencer::OnCrumbClicked)
-																	]
-															]
-													]
-											]
+										SNew(SSequencerTreeViewBox, InSequencer, SharedThis(this))
+										[
+											TreeView.ToSharedRef()
+										]
 									]
 
-								+ SOverlay::Slot()
+									// track area
+									+ SHorizontalBox::Slot()
+									.FillWidth( FillCoefficient_1 )
 									[
-										// track area virtual splitter overlay
-										SNew(SSequencerSplitterOverlay)
-											.Style(FEditorStyle::Get(), "Sequencer.AnimationOutliner.Splitter")
-											.Visibility(EVisibility::SelfHitTestInvisible)
-
-										+ SSplitter::Slot()
-											.Value(FillCoefficient_0)
-											.OnSlotResized(SSplitter::FOnSlotResized::CreateSP(this, &SSequencer::OnColumnFillCoefficientChanged, 0))
-											[
-												SNew(SSpacer)
-											]
-
-										+ SSplitter::Slot()
-											.Value(FillCoefficient_1)
-											.OnSlotResized(SSplitter::FOnSlotResized::CreateSP(this, &SSequencer::OnColumnFillCoefficientChanged, 1))
-											[
-												SNew(SSpacer)
-											]
+										SNew(SBox)
+										.Padding(ResizeBarPadding)
+										.Visibility(this, &SSequencer::GetTrackAreaVisibility )
+										[
+											TrackArea.ToSharedRef()
+										]
 									]
+								]
 							]
 
-						+ SSplitter::Slot()
-							.Value(0.2f)
+							+ SOverlay::Slot()
+							.HAlign( HAlign_Right )
 							[
-								// details view panel
-								SNew(SVerticalBox)
-									.Visibility(this, &SSequencer::HandleDetailsViewVisibility)
-
-								+ SVerticalBox::Slot()
-									.AutoHeight()
-									.HAlign(HAlign_Center)
-									.Padding(2.0f, 8.0f, 2.0f, 2.0f)
-									[
-										SNew(STextBlock)
-											.Text(LOCTEXT("NoObjectsSelected", "Select a key or section for details."))
-											.Visibility(this, &SSequencer::HandleDetailsViewTipVisibility)
-									]
+								ScrollBar
+							]
+						]
 
-								+ SVerticalBox::Slot()
-									.FillHeight(1.0f)
-									[
-										DetailsView.ToSharedRef()
-									]
+						+ SHorizontalBox::Slot()
+						.FillWidth( TAttribute<float>( this, &SSequencer::GetOutlinerSpacerFill ) )
+						[
+							SNew(SSpacer)
+						]
+					]
+
+					// playback buttons
+					+ SGridPanel::Slot( Column0, Row3, SGridPanel::Layer(10) )
+					[
+						SNew(SBorder)
+						.BorderImage(FEditorStyle::GetBrush("ToolPanel.GroupBorder"))
+						//.BorderBackgroundColor(FLinearColor(.50f, .50f, .50f, 1.0f))
+						.HAlign(HAlign_Center)
+						[
+							MakeTransportControls()
+						]
+					]
+
+					// Second column
+
+					+ SGridPanel::Slot( Column1, Row1 )
+					.Padding(ResizeBarPadding)
+					.RowSpan(3)
+					[
+						SNew(SBorder)
+						.BorderImage(FEditorStyle::GetBrush("ToolPanel.GroupBorder"))
+						[
+							SNew(SSpacer)
+						]
+					]
+
+					+ SGridPanel::Slot( Column1, Row1, SGridPanel::Layer(10) )
+					.Padding(ResizeBarPadding)
+					[
+						SNew( SBorder )
+						.BorderImage( FEditorStyle::GetBrush("ToolPanel.GroupBorder") )
+						.BorderBackgroundColor( FLinearColor(.50f, .50f, .50f, 1.0f ) )
+						.Padding(0)
+						[
+							TopTimeSlider
+						]
+					]
+
+					// Overlay that draws the tick lines
+					+ SGridPanel::Slot( Column1, Row2, SGridPanel::Layer(10) )
+					.Padding(ResizeBarPadding)
+					[
+						SNew( SSequencerSectionOverlay, TimeSliderController )
+						.Visibility( EVisibility::HitTestInvisible )
+						.DisplayScrubPosition( false )
+						.DisplayTickLines( true )
+					]
+
+					// Goto box
+					+ SGridPanel::Slot( Column1, Row2, SGridPanel::Layer(11) )
+						.Padding(ResizeBarPadding)
+						.HAlign(HAlign_Left)
+						.VAlign(VAlign_Top)
+						[
+							SAssignNew(GotoBox, SSequencerGotoBox, SequencerPtr.Pin().ToSharedRef(), *Settings, NumericTypeInterface.ToSharedRef())
+						]
+
+					// Curve editor
+					+ SGridPanel::Slot( Column1, Row2, SGridPanel::Layer(20) )
+					.Padding(ResizeBarPadding)
+					[
+						CurveEditor.ToSharedRef()
+					]
+
+					// Overlay that draws the scrub position
+					+ SGridPanel::Slot( Column1, Row2, SGridPanel::Layer(30) )
+					.Padding(ResizeBarPadding)
+					[
+						SNew( SSequencerSectionOverlay, TimeSliderController )
+						.Visibility( EVisibility::HitTestInvisible )
+						.DisplayScrubPosition( true )
+						.DisplayTickLines( false )
+						.PaintPlaybackRangeArgs(this, &SSequencer::GetSectionPlaybackRangeArgs)
+					]
+
+					// play range sliders
+					+ SGridPanel::Slot( Column1, Row3, SGridPanel::Layer(10) )
+					.Padding(ResizeBarPadding)
+					[
+						SNew( SBorder )
+						.BorderImage( FEditorStyle::GetBrush("ToolPanel.GroupBorder") )
+						.BorderBackgroundColor( FLinearColor(.50f, .50f, .50f, 1.0f ) )
+						.Padding(0)
+						[
+							SNew( SOverlay )
+
+							+ SOverlay::Slot()
+							[
+								BottomTimeSlider
+							]
+
+							+ SOverlay::Slot()
+							[
+								BottomTimeRange
 							]
+						]
+					]
+				]
+
+				+ SOverlay::Slot()
+				[
+					// track area virtual splitter overlay
+					SNew(SSequencerSplitterOverlay)
+					.Style(FEditorStyle::Get(), "Sequencer.AnimationOutliner.Splitter")
+					.Visibility(EVisibility::SelfHitTestInvisible)
+
+					+ SSplitter::Slot()
+					.Value(FillCoefficient_0)
+					.OnSlotResized(SSplitter::FOnSlotResized::CreateSP(this, &SSequencer::OnColumnFillCoefficientChanged, 0))
+					[
+						SNew(SSpacer)
+					]
+
+					+ SSplitter::Slot()
+					.Value(FillCoefficient_1)
+					.OnSlotResized(SSplitter::FOnSlotResized::CreateSP(this, &SSequencer::OnColumnFillCoefficientChanged, 1))
+					[
+						SNew(SSpacer)
 					]
+				]
 			]
+		]
 	];
 
 	InSequencer->GetSelection().GetOnKeySelectionChanged().AddSP(this, &SSequencer::HandleKeySelectionChanged);
@@ -454,41 +477,19 @@ void SSequencer::Construct(const FArguments& InArgs, TSharedRef<FSequencer> InSe
 
 	ResetBreadcrumbs();
 }
-
+PRAGMA_ENABLE_OPTIMIZATION
 
 void SSequencer::BindCommands(TSharedRef<FUICommandList> SequencerCommandBindings)
 {
-	SequencerCommandBindings->MapAction(
-		FSequencerCommands::Get().MoveTool,
-		FExecuteAction::CreateLambda([this] {
-			EditTool.Reset(new FSequencerEditTool_Movement(Sequencer.Pin(), SharedThis(this)));
-		}),
-		FCanExecuteAction::CreateLambda([] {
-			return true; 
-		}),
-		FIsActionChecked::CreateSP(this, &SSequencer::IsEditToolEnabled, FName("Movement"))
-	);
-
-	SequencerCommandBindings->MapAction(
-		FSequencerCommands::Get().MarqueeTool,
-		FExecuteAction::CreateLambda( [this] {
-			EditTool.Reset(new FSequencerEditTool_Selection(Sequencer.Pin(), SharedThis(this)));
-		}),
-		FCanExecuteAction::CreateLambda([] {
-			return true;
-		}),
-		FIsActionChecked::CreateSP(this, &SSequencer::IsEditToolEnabled, FName("Selection"))
-	);
-
 	auto CanPaste = [this]{
 		if (!HasFocusedDescendants() && !HasKeyboardFocus())
 		{
 			return false;
 		}
 
-		return Sequencer.Pin()->GetClipboardStack().Num() != 0;
+		return SequencerPtr.Pin()->GetClipboardStack().Num() != 0;
 	};
-
+	
 	SequencerCommandBindings->MapAction(
 		FGenericCommands::Get().Paste,
 		FExecuteAction::CreateSP(this, &SSequencer::Paste),
@@ -517,66 +518,16 @@ TSharedRef<INumericTypeInterface<float>> SSequencer::GetNumericTypeInterface()
 	return NumericTypeInterface.ToSharedRef();
 }
 
-
-void SSequencer::UpdateDetailsView()
+TSharedRef<INumericTypeInterface<float>> SSequencer::GetZeroPadNumericTypeInterface()
 {
-	TArray<TWeakObjectPtr<UObject>> Sections;
-
-	// get selected sections
-	for (auto Section : Sequencer.Pin()->GetSelection().GetSelectedSections())
-	{
-		Sections.Add(Section);
-	}
-
-	// get sections from selected keys
-	const TArray<FSequencerSelectedKey> SelectedKeys = Sequencer.Pin()->GetSelection().GetSelectedKeys().Array();
-
-	for (const auto& Key : SelectedKeys)
-	{
-		if (Key.Section != nullptr)
-		{
-			Sections.AddUnique(Key.Section);
-		}
-	}
-
-	// @todo sequencer: highlight selected keys in details view
-
-	// update details view
-	DetailsView->SetObjects(Sections, true);
+	return ZeroPadNumericTypeInterface.ToSharedRef();
 }
 
-
 /* SSequencer callbacks
  *****************************************************************************/
 
-bool SSequencer::HandleDetailsViewEnabled() const
-{
-	return true;
-}
-
-
-EVisibility SSequencer::HandleDetailsViewTipVisibility() const
-{
-	return (DetailsView->GetSelectedObjects().Num() > 0)
-		? EVisibility::Collapsed
-		: EVisibility::Visible;
-}
-
-
-EVisibility SSequencer::HandleDetailsViewVisibility() const
-{
-	if (Sequencer.Pin()->GetSettings()->GetDetailsViewVisible())
-	{
-		return EVisibility::Visible;
-	}
-
-	return EVisibility::Collapsed;
-}
-
-
 void SSequencer::HandleKeySelectionChanged()
 {
-	UpdateDetailsView();
 }
 
 
@@ -600,7 +551,7 @@ void SSequencer::HandleLabelBrowserSelectionChanged(FString NewLabel, ESelectInf
 
 EVisibility SSequencer::HandleLabelBrowserVisibility() const
 {
-	if (Sequencer.Pin()->GetSettings()->GetLabelBrowserVisible())
+	if (Settings->GetLabelBrowserVisible())
 	{
 		return EVisibility::Visible;
 	}
@@ -611,66 +562,60 @@ EVisibility SSequencer::HandleLabelBrowserVisibility() const
 
 void SSequencer::HandleSectionSelectionChanged()
 {
-	UpdateDetailsView();
 }
 
 
-bool SSequencer::IsEditToolEnabled(FName InIdentifier)
+TSharedRef<SWidget> SSequencer::MakeAddButton()
 {
-	return GetEditTool().GetIdentifier() == InIdentifier;
+	return SNew(SComboButton)
+	.OnGetMenuContent(this, &SSequencer::MakeAddMenu)
+	.ButtonStyle(FEditorStyle::Get(), "FlatButton.Success")
+	.ContentPadding(FMargin(2.0f, 1.0f))
+	.HasDownArrow(false)
+	.ButtonContent()
+	[
+		SNew(SHorizontalBox)
+
+		+ SHorizontalBox::Slot()
+		.VAlign(VAlign_Center)
+		.AutoWidth()
+		[
+			SNew(STextBlock)
+			.TextStyle(FEditorStyle::Get(), "NormalText.Important")
+			.Font(FEditorStyle::Get().GetFontStyle("FontAwesome.10"))
+			.Text(FText::FromString(FString(TEXT("\xf067"))) /*fa-plus*/)
+		]
+
+		+ SHorizontalBox::Slot()
+		.AutoWidth()
+		.Padding(4, 0, 0, 0)
+		[
+			SNew(STextBlock)
+			.TextStyle(FEditorStyle::Get(), "NormalText.Important")
+			.Text(LOCTEXT("AddButton", "Add"))
+		]
+
+		+ SHorizontalBox::Slot()
+		.VAlign(VAlign_Center)
+		.AutoWidth()
+		.Padding(4, 0, 0, 0)
+		[
+			SNew(STextBlock)
+			.TextStyle(FEditorStyle::Get(), "NormalText.Important")
+			.Font(FEditorStyle::Get().GetFontStyle("FontAwesome.10"))
+			.Text(FText::FromString(FString(TEXT("\xf0d7"))) /*fa-caret-down*/)
+		]
+	];
 }
 
-
 TSharedRef<SWidget> SSequencer::MakeToolBar()
 {
-	FToolBarBuilder ToolBarBuilder( Sequencer.Pin()->GetCommandBindings(), FMultiBoxCustomization::None, TSharedPtr<FExtender>(), Orient_Horizontal, true);
+	FToolBarBuilder ToolBarBuilder( SequencerPtr.Pin()->GetCommandBindings(), FMultiBoxCustomization::None, TSharedPtr<FExtender>(), Orient_Horizontal, true);
 	
 	ToolBarBuilder.BeginSection("Base Commands");
 	{
-		// Add Track menu
-		ToolBarBuilder.AddWidget(
-			SNew(SComboButton)
-			.OnGetMenuContent(this, &SSequencer::MakeAddMenu)
-			.ButtonStyle(FEditorStyle::Get(), "FlatButton.Success")
-			.ContentPadding(FMargin(2.0f, 1.0f))
-			.HasDownArrow(false)
-			.ButtonContent()
-			[
-				SNew(SHorizontalBox)
-
-				+ SHorizontalBox::Slot()
-					.VAlign(VAlign_Center)
-					.AutoWidth()
-					[
-						SNew(STextBlock)
-							.TextStyle(FEditorStyle::Get(), "NormalText.Important")
-							.Font(FEditorStyle::Get().GetFontStyle("FontAwesome.10"))
-							.Text(FText::FromString(FString(TEXT("\xf067"))) /*fa-plus*/)
-					]
-
-				+ SHorizontalBox::Slot()
-					.AutoWidth()
-					.Padding(4, 0, 0, 0)
-					[
-						SNew(STextBlock)
-							.TextStyle(FEditorStyle::Get(), "NormalText.Important")
-							.Text(LOCTEXT("AddButton", "Add"))
-					]
-
-				+ SHorizontalBox::Slot()
-					.VAlign(VAlign_Center)
-					.AutoWidth()
-					.Padding(4, 0, 0, 0)
-					[
-						SNew(STextBlock)
-							.TextStyle(FEditorStyle::Get(), "NormalText.Important")
-							.Font(FEditorStyle::Get().GetFontStyle("FontAwesome.10"))
-							.Text(FText::FromString(FString(TEXT("\xf0d7"))) /*fa-caret-down*/)
-					]
-			]);
-
 		// General 
-		if( Sequencer.Pin()->IsLevelEditorSequencer() )
+		if( SequencerPtr.Pin()->IsLevelEditorSequencer() )
 		{
 			ToolBarBuilder.AddToolBarButton(
 				FUIAction(FExecuteAction::CreateSP(this, &SSequencer::OnSaveMovieSceneClicked)),
@@ -694,14 +639,14 @@ TSharedRef<SWidget> SSequencer::MakeToolBar()
 
 		ToolBarBuilder.AddSeparator();
 
-		if( Sequencer.Pin()->IsLevelEditorSequencer() )
+		if( SequencerPtr.Pin()->IsLevelEditorSequencer() )
 		{
 			TAttribute<FSlateIcon> KeyAllIcon;
 			KeyAllIcon.Bind(TAttribute<FSlateIcon>::FGetter::CreateLambda([&]{
 				static FSlateIcon KeyAllEnabledIcon(FEditorStyle::GetStyleSetName(), "Sequencer.KeyAllEnabled");
 				static FSlateIcon KeyAllDisabledIcon(FEditorStyle::GetStyleSetName(), "Sequencer.KeyAllDisabled");
 
-				return Sequencer.Pin()->GetKeyAllEnabled() ? KeyAllEnabledIcon : KeyAllDisabledIcon;
+				return SequencerPtr.Pin()->GetKeyAllEnabled() ? KeyAllEnabledIcon : KeyAllDisabledIcon;
 			}));
 
 			ToolBarBuilder.AddToolBarButton( FSequencerCommands::Get().ToggleKeyAllEnabled, NAME_None, TAttribute<FText>(), TAttribute<FText>(), KeyAllIcon );
@@ -709,7 +654,7 @@ TSharedRef<SWidget> SSequencer::MakeToolBar()
 
 		TAttribute<FSlateIcon> AutoKeyModeIcon;
 		AutoKeyModeIcon.Bind(TAttribute<FSlateIcon>::FGetter::CreateLambda( [&] {
-			switch ( Sequencer.Pin()->GetAutoKeyMode() )
+			switch ( SequencerPtr.Pin()->GetAutoKeyMode() )
 			{
 			case EAutoKeyMode::KeyAll:
 				return FSequencerCommands::Get().SetAutoKeyModeAll->GetIcon();
@@ -722,7 +667,7 @@ TSharedRef<SWidget> SSequencer::MakeToolBar()
 
 		TAttribute<FText> AutoKeyModeToolTip;
 		AutoKeyModeToolTip.Bind( TAttribute<FText>::FGetter::CreateLambda( [&] {
-			switch ( Sequencer.Pin()->GetAutoKeyMode() )
+			switch ( SequencerPtr.Pin()->GetAutoKeyMode() )
 			{
 			case EAutoKeyMode::KeyAll:
 				return FSequencerCommands::Get().SetAutoKeyModeAll->GetDescription();
@@ -742,13 +687,6 @@ TSharedRef<SWidget> SSequencer::MakeToolBar()
 	}
 	ToolBarBuilder.EndSection();
 
-	ToolBarBuilder.BeginSection("Tools");
-	{
-		ToolBarBuilder.SetLabelVisibility( EVisibility::Collapsed );
-		ToolBarBuilder.AddToolBarButton( FSequencerCommands::Get().MoveTool );
-		ToolBarBuilder.AddToolBarButton( FSequencerCommands::Get().MarqueeTool );
-	}
-	ToolBarBuilder.EndSection();
 
 	ToolBarBuilder.BeginSection("Snapping");
 	{
@@ -798,17 +736,17 @@ TSharedRef<SWidget> SSequencer::MakeAddMenu()
 
 		// let toolkits populate the menu
 		MenuBuilder.BeginSection("MainMenu");
-		OnGetAddMenuContent.ExecuteIfBound(MenuBuilder, Sequencer.Pin().ToSharedRef());
+		OnGetAddMenuContent.ExecuteIfBound(MenuBuilder, SequencerPtr.Pin().ToSharedRef());
 		MenuBuilder.EndSection();
 
 		// let track editors populate the menu
-		TSharedPtr<FSequencer> PinnedSequencer = Sequencer.Pin();
+		TSharedPtr<FSequencer> Sequencer = SequencerPtr.Pin();
 
 		// Always create the section so that we afford extension
 		MenuBuilder.BeginSection("AddTracks");
-		if (PinnedSequencer.IsValid())
+		if (Sequencer.IsValid())
 		{
-			PinnedSequencer->BuildAddTrackMenu(MenuBuilder);
+			Sequencer->BuildAddTrackMenu(MenuBuilder);
 		}
 		MenuBuilder.EndSection();
 	}
@@ -819,17 +757,16 @@ TSharedRef<SWidget> SSequencer::MakeAddMenu()
 
 TSharedRef<SWidget> SSequencer::MakeGeneralMenu()
 {
-	FMenuBuilder MenuBuilder( true, Sequencer.Pin()->GetCommandBindings() );
+	FMenuBuilder MenuBuilder( true, SequencerPtr.Pin()->GetCommandBindings() );
 
 	MenuBuilder.BeginSection( "ViewOptions", LOCTEXT( "ViewMenuHeader", "View" ) );
 	{
-		if (Sequencer.Pin()->IsLevelEditorSequencer())
+		if (SequencerPtr.Pin()->IsLevelEditorSequencer())
 		{
-			MenuBuilder.AddMenuEntry( FSequencerCommands::Get().ToggleDetailsView );
 			MenuBuilder.AddMenuEntry( FSequencerCommands::Get().ToggleLabelBrowser );
 		}
 
-		if (Sequencer.Pin()->IsLevelEditorSequencer())
+		if (SequencerPtr.Pin()->IsLevelEditorSequencer())
 		{
 			MenuBuilder.AddMenuEntry( FSequencerCommands::Get().FindInContentBrowser );
 		}
@@ -843,12 +780,12 @@ TSharedRef<SWidget> SSequencer::MakeGeneralMenu()
 
 	MenuBuilder.BeginSection( "Ranges", LOCTEXT( "RangesHeader", "Playback Range" ) );
 	{
-		TSharedPtr<FSequencer> PinnedSequencer = Sequencer.Pin();
-
+		TSharedPtr<FSequencer> Sequencer = SequencerPtr.Pin();
+		
 		// Menu entry for the start position
 		auto OnStartChanged = [=](float NewValue){
-			float Upper = PinnedSequencer->GetPlaybackRange().GetUpperBoundValue();
-			PinnedSequencer->SetPlaybackRange(TRange<float>(FMath::Min(NewValue, Upper), Upper));
+			float Upper = Sequencer->GetPlaybackRange().GetUpperBoundValue();
+			Sequencer->SetPlaybackRange(TRange<float>(FMath::Min(NewValue, Upper), Upper));
 		};
 		MenuBuilder.AddWidget(
 			SNew(SHorizontalBox)
@@ -867,21 +804,21 @@ TSharedRef<SWidget> SSequencer::MakeGeneralMenu()
 				.OnBeginSliderMovement(OnBeginPlaybackRangeDrag)
 				.OnEndSliderMovement_Lambda([=](float Value){ OnStartChanged(Value); OnEndPlaybackRangeDrag.ExecuteIfBound(); })
 				.MinValue_Lambda([=]() -> float {
-					return PinnedSequencer->GetClampRange().GetLowerBoundValue(); 
+					return Sequencer->GetClampRange().GetLowerBoundValue(); 
 				})
 				.MaxValue_Lambda([=]() -> float {
-					return PinnedSequencer->GetPlaybackRange().GetUpperBoundValue(); 
+					return Sequencer->GetPlaybackRange().GetUpperBoundValue(); 
 				})
 				.Value_Lambda([=]() -> float {
-					return PinnedSequencer->GetPlaybackRange().GetLowerBoundValue();
+					return Sequencer->GetPlaybackRange().GetLowerBoundValue();
 				})
 			],
 			LOCTEXT("PlaybackStartLabel", "Start"));
 
 		// Menu entry for the end position
 		auto OnEndChanged = [=](float NewValue){
-			float Lower = PinnedSequencer->GetPlaybackRange().GetLowerBoundValue();
-			PinnedSequencer->SetPlaybackRange(TRange<float>(Lower, FMath::Max(NewValue, Lower)));
+			float Lower = Sequencer->GetPlaybackRange().GetLowerBoundValue();
+			Sequencer->SetPlaybackRange(TRange<float>(Lower, FMath::Max(NewValue, Lower)));
 		};
 		MenuBuilder.AddWidget(
 			SNew(SHorizontalBox)
@@ -900,28 +837,55 @@ TSharedRef<SWidget> SSequencer::MakeGeneralMenu()
 				.OnBeginSliderMovement(OnBeginPlaybackRangeDrag)
 				.OnEndSliderMovement_Lambda([=](float Value){ OnEndChanged(Value); OnEndPlaybackRangeDrag.ExecuteIfBound(); })
 				.MinValue_Lambda([=]() -> float {
-					return PinnedSequencer->GetPlaybackRange().GetLowerBoundValue(); 
+					return Sequencer->GetPlaybackRange().GetLowerBoundValue(); 
 				})
 				.MaxValue_Lambda([=]() -> float {
-					return PinnedSequencer->GetClampRange().GetUpperBoundValue(); 
+					return Sequencer->GetClampRange().GetUpperBoundValue(); 
 				})
 				.Value_Lambda([=]() -> float {
-					return PinnedSequencer->GetPlaybackRange().GetUpperBoundValue();
+					return Sequencer->GetPlaybackRange().GetUpperBoundValue();
 				})
 			],
 			LOCTEXT("PlaybackStartEnd", "End"));
 
 		MenuBuilder.AddMenuEntry( FSequencerCommands::Get().ToggleKeepCursorInPlaybackRange );
 		MenuBuilder.AddMenuEntry( FSequencerCommands::Get().ToggleKeepPlaybackRangeInSectionBounds );
-	}
 
+		// Menu entry for zero padding
+		auto OnZeroPadChanged = [=](uint8 NewValue){
+			Settings->SetZeroPadFrames(NewValue);
+		};
+		MenuBuilder.AddWidget(
+			SNew(SHorizontalBox)
+			+ SHorizontalBox::Slot()
+			[
+				SNew(SSpacer)
+			]
+			+ SHorizontalBox::Slot()
+			.AutoWidth()
+			[
+				SNew(SSpinBox<uint8>)
+				.Style(&FEditorStyle::GetWidgetStyle<FSpinBoxStyle>("Sequencer.HyperlinkSpinBox"))
+				.OnValueCommitted_Lambda([=](uint8 Value, ETextCommit::Type){ OnZeroPadChanged(Value); })
+				.OnValueChanged_Lambda(OnZeroPadChanged)
+				.MinValue(0)
+				.MaxValue(8)
+				.Value_Lambda([=]() -> uint8 {
+					return Settings->GetZeroPadFrames();
+				})
+			],
+			LOCTEXT("ZeroPaddingText", "Zero Pad Frame Numbers"));
+
+		MenuBuilder.AddMenuSeparator();
+		MenuBuilder.AddMenuEntry( FSequencerCommands::Get().ToggleShowGotoBox );
+	}
 
 	return MenuBuilder.MakeWidget();
 }
 
 TSharedRef<SWidget> SSequencer::MakeSnapMenu()
 {
-	FMenuBuilder MenuBuilder( false, Sequencer.Pin()->GetCommandBindings() );
+	FMenuBuilder MenuBuilder( false, SequencerPtr.Pin()->GetCommandBindings() );
 
 	MenuBuilder.BeginSection("FramesRanges", LOCTEXT("SnappingMenuFrameRangesHeader", "Frame Ranges") );
 	{
@@ -949,6 +913,7 @@ TSharedRef<SWidget> SSequencer::MakeSnapMenu()
 	{
 		MenuBuilder.AddMenuEntry( FSequencerCommands::Get().ToggleSnapPlayTimeToInterval );
 		MenuBuilder.AddMenuEntry( FSequencerCommands::Get().ToggleSnapPlayTimeToDraggedKey );
+		MenuBuilder.AddMenuEntry( FSequencerCommands::Get().ToggleFixedTimeStepPlayback );
 	}
 	MenuBuilder.EndSection();
 
@@ -964,7 +929,7 @@ TSharedRef<SWidget> SSequencer::MakeSnapMenu()
 
 TSharedRef<SWidget> SSequencer::MakeAutoKeyMenu()
 {
-	FMenuBuilder MenuBuilder(false, Sequencer.Pin()->GetCommandBindings());
+	FMenuBuilder MenuBuilder(false, SequencerPtr.Pin()->GetCommandBindings());
 
 	MenuBuilder.AddMenuEntry(FSequencerCommands::Get().SetAutoKeyModeAll);
 	MenuBuilder.AddMenuEntry(FSequencerCommands::Get().SetAutoKeyModeAnimated);
@@ -978,17 +943,17 @@ TSharedRef<SWidget> SSequencer::MakeAutoKeyMenu()
 TSharedRef<SWidget> SSequencer::MakeTransportControls()
 {
 	FEditorWidgetsModule& EditorWidgetsModule = FModuleManager::Get().LoadModuleChecked<FEditorWidgetsModule>( "EditorWidgets" );
-	TSharedRef<FSequencer> SequencerPinned = Sequencer.Pin().ToSharedRef();
+	TSharedRef<FSequencer> Sequencer = SequencerPtr.Pin().ToSharedRef();
 	FTransportControlArgs TransportControlArgs;
 	{
-		TransportControlArgs.OnBackwardEnd.BindSP( SequencerPinned, &FSequencer::OnStepToBeginning );
-		TransportControlArgs.OnBackwardStep.BindSP( SequencerPinned, &FSequencer::OnStepBackward );
-		TransportControlArgs.OnForwardPlay.BindSP( SequencerPinned, &FSequencer::OnPlay, true );
-		TransportControlArgs.OnForwardStep.BindSP( SequencerPinned, &FSequencer::OnStepForward );
-		TransportControlArgs.OnForwardEnd.BindSP( SequencerPinned, &FSequencer::OnStepToEnd );
-		TransportControlArgs.OnToggleLooping.BindSP( SequencerPinned, &FSequencer::OnToggleLooping );
-		TransportControlArgs.OnGetLooping.BindSP( SequencerPinned, &FSequencer::IsLooping );
-		TransportControlArgs.OnGetPlaybackMode.BindSP( SequencerPinned, &FSequencer::GetPlaybackMode );
+		TransportControlArgs.OnBackwardEnd.BindSP( Sequencer, &FSequencer::OnStepToBeginning );
+		TransportControlArgs.OnBackwardStep.BindSP( Sequencer, &FSequencer::OnStepBackward );
+		TransportControlArgs.OnForwardPlay.BindSP( Sequencer, &FSequencer::OnPlay, true );
+		TransportControlArgs.OnForwardStep.BindSP( Sequencer, &FSequencer::OnStepForward );
+		TransportControlArgs.OnForwardEnd.BindSP( Sequencer, &FSequencer::OnStepToEnd );
+		TransportControlArgs.OnToggleLooping.BindSP( Sequencer, &FSequencer::OnToggleLooping );
+		TransportControlArgs.OnGetLooping.BindSP( Sequencer, &FSequencer::IsLooping );
+		TransportControlArgs.OnGetPlaybackMode.BindSP( Sequencer, &FSequencer::GetPlaybackMode );
 	}
 
 	return EditorWidgetsModule.CreateTransportControl( TransportControlArgs );
@@ -1014,9 +979,9 @@ void SSequencer::RegisterActiveTimerForPlayback()
 
 EActiveTimerReturnType SSequencer::EnsureSlateTickDuringPlayback(double InCurrentTime, float InDeltaTime)
 {
-	if (Sequencer.IsValid())
+	if (SequencerPtr.IsValid())
 	{
-		auto PlaybackStatus = Sequencer.Pin()->GetPlaybackStatus();
+		auto PlaybackStatus = SequencerPtr.Pin()->GetPlaybackStatus();
 		if (PlaybackStatus == EMovieScenePlayerStatus::Playing || PlaybackStatus == EMovieScenePlayerStatus::Recording || PlaybackStatus == EMovieScenePlayerStatus::Scrubbing)
 		{
 			return EActiveTimerReturnType::Continue;
@@ -1048,13 +1013,13 @@ void SSequencer::UpdateLayoutTree()
 {
 	// Cache the selected path names so selection can be restored after the update.
 	TSet<FString> SelectedPathNames;
-	for (TSharedRef<const FSequencerDisplayNode> SelectedDisplayNode : Sequencer.Pin()->GetSelection().GetSelectedOutlinerNodes().Array())
+	for (TSharedRef<const FSequencerDisplayNode> SelectedDisplayNode : SequencerPtr.Pin()->GetSelection().GetSelectedOutlinerNodes().Array())
 	{
 		SelectedPathNames.Add(SelectedDisplayNode->GetPathName());
 	}
 
 	// Suspend broadcasting selection changes because we don't want unnecessary rebuilds.
-	Sequencer.Pin()->GetSelection().SuspendBroadcast();
+	SequencerPtr.Pin()->GetSelection().SuspendBroadcast();
 
 	// Update the node tree
 	SequencerNodeTree->Update();
@@ -1062,20 +1027,20 @@ void SSequencer::UpdateLayoutTree()
 	TreeView->Refresh();
 
 	// Restore the selection state.
-	RestoreSelectionState(SequencerNodeTree->GetRootNodes(), SelectedPathNames, Sequencer.Pin()->GetSelection());	// Update to actor selection.
+	RestoreSelectionState(SequencerNodeTree->GetRootNodes(), SelectedPathNames, SequencerPtr.Pin()->GetSelection());	// Update to actor selection.
 	OnActorSelectionChanged(nullptr);
 
 	// This must come after the selection state has been restored so that the curve editor is populated with the correctly selected nodes
 	CurveEditor->SetSequencerNodeTree(SequencerNodeTree);
 
 	// Continue broadcasting selection changes
-	Sequencer.Pin()->GetSelection().ResumeBroadcast();
+	SequencerPtr.Pin()->GetSelection().ResumeBroadcast();
 }
 
 
 void SSequencer::UpdateBreadcrumbs()
 {
-	TSharedRef<FMovieSceneSequenceInstance> FocusedMovieSceneInstance = Sequencer.Pin()->GetFocusedMovieSceneSequenceInstance();
+	TSharedRef<FMovieSceneSequenceInstance> FocusedMovieSceneInstance = SequencerPtr.Pin()->GetFocusedMovieSceneSequenceInstance();
 
 	if (BreadcrumbTrail->PeekCrumb().BreadcrumbType == FSequencerBreadcrumb::ShotType)
 	{
@@ -1094,7 +1059,7 @@ void SSequencer::UpdateBreadcrumbs()
 void SSequencer::ResetBreadcrumbs()
 {
 	BreadcrumbTrail->ClearCrumbs();
-	BreadcrumbTrail->PushCrumb(TAttribute<FText>::Create(TAttribute<FText>::FGetter::CreateSP(this, &SSequencer::GetRootAnimationName)), FSequencerBreadcrumb(Sequencer.Pin()->GetRootMovieSceneSequenceInstance()));
+	BreadcrumbTrail->PushCrumb(TAttribute<FText>::Create(TAttribute<FText>::FGetter::CreateSP(this, &SSequencer::GetRootAnimationName)), FSequencerBreadcrumb(SequencerPtr.Pin()->GetRootMovieSceneSequenceInstance()));
 }
 
 
@@ -1214,7 +1179,7 @@ FReply SSequencer::OnDrop( const FGeometry& MyGeometry, const FDragDropEvent& Dr
 FReply SSequencer::OnKeyDown( const FGeometry& MyGeometry, const FKeyEvent& InKeyEvent ) 
 {
 	// A toolkit tab is active, so direct all command processing to it
-	if( Sequencer.Pin()->GetCommandBindings()->ProcessCommandBindings( InKeyEvent ) )
+	if( SequencerPtr.Pin()->GetCommandBindings()->ProcessCommandBindings( InKeyEvent ) )
 	{
 		return FReply::Handled();
 	}
@@ -1225,7 +1190,7 @@ FReply SSequencer::OnKeyDown( const FGeometry& MyGeometry, const FKeyEvent& InKe
 
 void SSequencer::OnAssetsDropped( const FAssetDragDropOp& DragDropOp )
 {
-	FSequencer& SequencerRef = *Sequencer.Pin();
+	FSequencer& SequencerRef = *SequencerPtr.Pin();
 
 	bool bObjectAdded = false;
 	TArray< UObject* > DroppedObjects;
@@ -1247,7 +1212,7 @@ void SSequencer::OnAssetsDropped( const FAssetDragDropOp& DragDropOp )
 		}
 	}
 
-	const TSet< TSharedRef<FSequencerDisplayNode> >& SelectedNodes = Sequencer.Pin()->GetSelection().GetSelectedOutlinerNodes();
+	const TSet< TSharedRef<FSequencerDisplayNode> >& SelectedNodes = SequencerPtr.Pin()->GetSelection().GetSelectedOutlinerNodes();
 	FGuid TargetObjectGuid;
 	// if exactly one object node is selected, we have a target object guid
 	TSharedPtr<const FSequencerDisplayNode> DisplayNode;
@@ -1270,7 +1235,7 @@ void SSequencer::OnAssetsDropped( const FAssetDragDropOp& DragDropOp )
 
 		if (!SequencerRef.OnHandleAssetDropped(CurObject, TargetObjectGuid))
 		{
-			SequencerRef.AddSpawnableForAssetOrClass( CurObject );
+			SequencerRef.MakeNewSpawnable( *CurObject );
 		}
 		bObjectAdded = true;
 	}
@@ -1285,7 +1250,7 @@ void SSequencer::OnAssetsDropped( const FAssetDragDropOp& DragDropOp )
 
 void SSequencer::OnClassesDropped( const FClassDragDropOp& DragDropOp )
 {
-	FSequencer& SequencerRef = *Sequencer.Pin();
+	FSequencer& SequencerRef = *SequencerPtr.Pin();
 
 	for( auto ClassIter = DragDropOp.ClassesToDrop.CreateConstIterator(); ClassIter; ++ClassIter )
 	{
@@ -1294,7 +1259,7 @@ void SSequencer::OnClassesDropped( const FClassDragDropOp& DragDropOp )
 		{
 			UObject* Object = Class->GetDefaultObject();
 
-			FGuid NewGuid = SequencerRef.AddSpawnableForAssetOrClass( Object );
+			FGuid NewGuid = SequencerRef.MakeNewSpawnable( *Object );
 
 			if (NewGuid.IsValid())
 			{
@@ -1308,7 +1273,7 @@ void SSequencer::OnClassesDropped( const FClassDragDropOp& DragDropOp )
 
 void SSequencer::OnUnloadedClassesDropped( const FUnloadedClassDragDropOp& DragDropOp )
 {
-	FSequencer& SequencerRef = *Sequencer.Pin();
+	FSequencer& SequencerRef = *SequencerPtr.Pin();
 	for( auto ClassDataIter = DragDropOp.AssetsToDrop->CreateConstIterator(); ClassDataIter; ++ClassDataIter )
 	{
 		auto& ClassData = *ClassDataIter;
@@ -1346,7 +1311,7 @@ void SSequencer::OnUnloadedClassesDropped( const FUnloadedClassDragDropOp& DragD
 
 		if( Object != nullptr )
 		{
-			FGuid NewGuid = SequencerRef.AddSpawnableForAssetOrClass( Object );
+			FGuid NewGuid = SequencerRef.MakeNewSpawnable( *Object );
 			if (NewGuid.IsValid())
 			{
 				SequencerRef.NotifyMovieSceneDataChanged();
@@ -1358,18 +1323,18 @@ void SSequencer::OnUnloadedClassesDropped( const FUnloadedClassDragDropOp& DragD
 
 void SSequencer::OnActorsDropped( FActorDragDropGraphEdOp& DragDropOp )
 {
-	Sequencer.Pin()->OnActorsDropped( DragDropOp.Actors );
+	SequencerPtr.Pin()->OnActorsDropped( DragDropOp.Actors );
 }
 
 
 void SSequencer::OnActorSelectionChanged(UObject*)
 {
-	if (!Sequencer.IsValid())
+	if (!SequencerPtr.IsValid())
 	{
 		return;
 	}
 
-	if (!Sequencer.Pin()->IsLevelEditorSequencer())
+	if (!SequencerPtr.Pin()->IsLevelEditorSequencer())
 	{
 		return;
 	}
@@ -1387,7 +1352,7 @@ void SSequencer::OnActorSelectionChanged(UObject*)
 	{
 		TSharedRef<FSequencerObjectBindingNode> ObjectBindingNode = StaticCastSharedRef<FSequencerObjectBindingNode>(Node);
 		TArray<UObject*> RuntimeObjects;
-		Sequencer.Pin()->GetRuntimeObjects( Sequencer.Pin()->GetFocusedMovieSceneSequenceInstance(), ObjectBindingNode->GetObjectBinding(), RuntimeObjects );
+		SequencerPtr.Pin()->GetRuntimeObjects( SequencerPtr.Pin()->GetFocusedMovieSceneSequenceInstance(), ObjectBindingNode->GetObjectBinding(), RuntimeObjects );
 		
 		bool bDoSelect = false;
 		for (int32 RuntimeIndex = 0; RuntimeIndex < RuntimeObjects.Num(); ++RuntimeIndex )
@@ -1399,20 +1364,20 @@ void SSequencer::OnActorSelectionChanged(UObject*)
 			}
 		}
 
-		if (Sequencer.Pin()->GetSelection().IsSelected(Node) != bDoSelect)
+		if (SequencerPtr.Pin()->GetSelection().IsSelected(Node) != bDoSelect)
 		{
 			if (bDoSelect)
 			{
-				Sequencer.Pin()->GetSelection().AddToSelection(Node);
+				SequencerPtr.Pin()->GetSelection().AddToSelection(Node);
 			}
 			else
 			{
-				Sequencer.Pin()->GetSelection().RemoveFromSelection(Node);
+				SequencerPtr.Pin()->GetSelection().RemoveFromSelection(Node);
 			}
 		}
 	}
 
-	const TSet<TSharedRef<FSequencerDisplayNode>>& OutlinerSelection = Sequencer.Pin()->GetSelection().GetSelectedOutlinerNodes();
+	const TSet<TSharedRef<FSequencerDisplayNode>>& OutlinerSelection = SequencerPtr.Pin()->GetSelection().GetSelectedOutlinerNodes();
 
 	if (OutlinerSelection.Num() == 1)
 	{
@@ -1429,13 +1394,13 @@ void SSequencer::OnCrumbClicked(const FSequencerBreadcrumb& Item)
 {
 	if (Item.BreadcrumbType != FSequencerBreadcrumb::ShotType)
 	{
-		if( Sequencer.Pin()->GetFocusedMovieSceneSequenceInstance() == Item.MovieSceneInstance.Pin() ) 
+		if( SequencerPtr.Pin()->GetFocusedMovieSceneSequenceInstance() == Item.MovieSceneInstance.Pin() ) 
 		{
 			// then do zooming
 		}
 		else
 		{
-			Sequencer.Pin()->PopToSequenceInstance( Item.MovieSceneInstance.Pin().ToSharedRef() );
+			SequencerPtr.Pin()->PopToSequenceInstance( Item.MovieSceneInstance.Pin().ToSharedRef() );
 		}
 	}
 }
@@ -1443,13 +1408,7 @@ void SSequencer::OnCrumbClicked(const FSequencerBreadcrumb& Item)
 
 FText SSequencer::GetRootAnimationName() const
 {
-	return Sequencer.Pin()->GetRootMovieSceneSequence()->GetDisplayName();
-}
-
-
-FText SSequencer::GetShotSectionTitle(UMovieSceneSection* ShotSection) const
-{
-	return Cast<UMovieSceneShotSection>(ShotSection)->GetShotDisplayName();
+	return SequencerPtr.Pin()->GetRootMovieSceneSequence()->GetDisplayName();
 }
 
 
@@ -1490,7 +1449,7 @@ TArray<FSectionHandle> SSequencer::GetSectionHandles(const TSet<TWeakObjectPtr<U
 
 void SSequencer::OnSaveMovieSceneClicked()
 {
-	Sequencer.Pin()->SaveCurrentMovieScene();
+	SequencerPtr.Pin()->SaveCurrentMovieScene();
 }
 
 
@@ -1544,7 +1503,7 @@ void SSequencer::StepToKey(bool bStepToNextKey, bool bCameraOnly)
 		{
 			TSharedRef<FSequencerObjectBindingNode> ObjectBindingNode = StaticCastSharedRef<FSequencerObjectBindingNode>(RootNode);
 			TArray<UObject*> RuntimeObjects;
-			Sequencer.Pin()->GetRuntimeObjects( Sequencer.Pin()->GetFocusedMovieSceneSequenceInstance(), ObjectBindingNode->GetObjectBinding(), RuntimeObjects );
+			SequencerPtr.Pin()->GetRuntimeObjects( SequencerPtr.Pin()->GetFocusedMovieSceneSequenceInstance(), ObjectBindingNode->GetObjectBinding(), RuntimeObjects );
 		
 			for (int32 RuntimeIndex = 0; RuntimeIndex < RuntimeObjects.Num(); ++RuntimeIndex )
 			{
@@ -1558,7 +1517,7 @@ void SSequencer::StepToKey(bool bStepToNextKey, bool bCameraOnly)
 	}
 	else
 	{
-		const TSet< TSharedRef<FSequencerDisplayNode> >& SelectedNodes = Sequencer.Pin()->GetSelection().GetSelectedOutlinerNodes();
+		const TSet< TSharedRef<FSequencerDisplayNode> >& SelectedNodes = SequencerPtr.Pin()->GetSelection().GetSelectedOutlinerNodes();
 		Nodes = SelectedNodes;
 
 		if (Nodes.Num() == 0)
@@ -1576,7 +1535,7 @@ void SSequencer::StepToKey(bool bStepToNextKey, bool bCameraOnly)
 	if (Nodes.Num() > 0)
 	{
 		float ClosestKeyDistance = MAX_FLT;
-		float CurrentTime = Sequencer.Pin()->GetCurrentLocalTime(*Sequencer.Pin()->GetFocusedMovieSceneSequence());
+		float CurrentTime = SequencerPtr.Pin()->GetCurrentLocalTime(*SequencerPtr.Pin()->GetFocusedMovieSceneSequence());
 		float StepToTime = 0;
 		bool StepToKeyFound = false;
 
@@ -1617,7 +1576,7 @@ void SSequencer::StepToKey(bool bStepToNextKey, bool bCameraOnly)
 
 		if (StepToKeyFound)
 		{
-			Sequencer.Pin()->SetGlobalTime(StepToTime);
+			SequencerPtr.Pin()->SetGlobalTime(StepToTime);
 		}
 	}
 }
@@ -1625,7 +1584,7 @@ void SSequencer::StepToKey(bool bStepToNextKey, bool bCameraOnly)
 
 EVisibility SSequencer::GetBreadcrumbTrailVisibility() const
 {
-	return Sequencer.Pin()->IsLevelEditorSequencer() ? EVisibility::Visible : EVisibility::Collapsed;
+	return SequencerPtr.Pin()->IsLevelEditorSequencer() ? EVisibility::Visible : EVisibility::Collapsed;
 }
 
 
@@ -1649,7 +1608,7 @@ EVisibility SSequencer::GetTimeRangeVisibility() const
 
 bool SSequencer::ShowFrameNumbers() const
 {
-	return Sequencer.Pin()->CanShowFrameNumbers() && Settings->GetShowFrameNumbers();
+	return SequencerPtr.Pin()->CanShowFrameNumbers() && Settings->GetShowFrameNumbers();
 }
 
 
@@ -1704,12 +1663,12 @@ FPaintPlaybackRangeArgs SSequencer::GetSectionPlaybackRangeArgs() const
 
 FVirtualTrackArea SSequencer::GetVirtualTrackArea() const
 {
-	return FVirtualTrackArea(*Sequencer.Pin(), *TreeView.Get(), TrackArea->GetCachedGeometry());
+	return FVirtualTrackArea(*SequencerPtr.Pin(), *TreeView.Get(), TrackArea->GetCachedGeometry());
 }
 
 FPasteContextMenuArgs SSequencer::GeneratePasteArgs(float PasteAtTime, TSharedPtr<FMovieSceneClipboard> Clipboard)
 {
-	TSharedPtr<FSequencer> SequencerPinned = Sequencer.Pin();
+	TSharedPtr<FSequencer> Sequencer = SequencerPtr.Pin();
 	if (Settings->GetIsSnapEnabled())
 	{
 		PasteAtTime = Settings->SnapTimeToInterval(PasteAtTime);
@@ -1724,8 +1683,8 @@ FPasteContextMenuArgs SSequencer::GeneratePasteArgs(float PasteAtTime, TSharedPt
 	// Paste into the currently selected sections, or hit test the mouse position as a last resort
 	TArray<TSharedRef<FSequencerDisplayNode>> PasteIntoNodes;
 	{
-		TSet<TWeakObjectPtr<UMovieSceneSection>> Sections = SequencerPinned->GetSelection().GetSelectedSections();
-		for (const FSequencerSelectedKey& Key : SequencerPinned->GetSelection().GetSelectedKeys())
+		TSet<TWeakObjectPtr<UMovieSceneSection>> Sections = Sequencer->GetSelection().GetSelectedSections();
+		for (const FSequencerSelectedKey& Key : Sequencer->GetSelection().GetSelectedKeys())
 		{
 			Sections.Add(Key.Section);
 		}
@@ -1752,11 +1711,11 @@ void SSequencer::Paste()
 {
 	TSharedPtr<FPasteContextMenu> ContextMenu;
 
-	TSharedPtr<FSequencer> SequencerPinned = Sequencer.Pin();
-	if (SequencerPinned->GetClipboardStack().Num() != 0)
+	TSharedPtr<FSequencer> Sequencer = SequencerPtr.Pin();
+	if (Sequencer->GetClipboardStack().Num() != 0)
 	{
-		FPasteContextMenuArgs Args = GeneratePasteArgs(SequencerPinned->GetGlobalTime(), SequencerPinned->GetClipboardStack().Last());
-		ContextMenu = FPasteContextMenu::CreateMenu(*SequencerPinned, Args);
+		FPasteContextMenuArgs Args = GeneratePasteArgs(Sequencer->GetGlobalTime(), Sequencer->GetClipboardStack().Last());
+		ContextMenu = FPasteContextMenu::CreateMenu(*Sequencer, Args);
 	}
 
 	if (!ContextMenu.IsValid() || !ContextMenu->IsValidPaste())
@@ -1769,7 +1728,7 @@ void SSequencer::Paste()
 	}
 
 	const bool bShouldCloseWindowAfterMenuSelection = true;
-	FMenuBuilder MenuBuilder(bShouldCloseWindowAfterMenuSelection, Sequencer.Pin()->GetCommandBindings());
+	FMenuBuilder MenuBuilder(bShouldCloseWindowAfterMenuSelection, SequencerPtr.Pin()->GetCommandBindings());
 
 	ContextMenu->PopulateMenu(MenuBuilder);
 
@@ -1787,19 +1746,19 @@ void SSequencer::Paste()
 
 void SSequencer::PasteFromHistory()
 {
-	TSharedPtr<FSequencer> SequencerPinned = Sequencer.Pin();
-	if (SequencerPinned->GetClipboardStack().Num() == 0)
+	TSharedPtr<FSequencer> Sequencer = SequencerPtr.Pin();
+	if (Sequencer->GetClipboardStack().Num() == 0)
 	{
 		return;
 	}
 
-	FPasteContextMenuArgs Args = GeneratePasteArgs(SequencerPinned->GetGlobalTime());
-	TSharedPtr<FPasteFromHistoryContextMenu> ContextMenu = FPasteFromHistoryContextMenu::CreateMenu(*SequencerPinned, Args);
+	FPasteContextMenuArgs Args = GeneratePasteArgs(Sequencer->GetGlobalTime());
+	TSharedPtr<FPasteFromHistoryContextMenu> ContextMenu = FPasteFromHistoryContextMenu::CreateMenu(*Sequencer, Args);
 
 	if (ContextMenu.IsValid())
 	{
 		const bool bShouldCloseWindowAfterMenuSelection = true;
-		FMenuBuilder MenuBuilder(bShouldCloseWindowAfterMenuSelection, SequencerPinned->GetCommandBindings());
+		FMenuBuilder MenuBuilder(bShouldCloseWindowAfterMenuSelection, Sequencer->GetCommandBindings());
 
 		ContextMenu->PopulateMenu(MenuBuilder);
 
diff --git a/Engine/Source/Editor/Sequencer/Private/SSequencer.h b/Engine/Source/Editor/Sequencer/Private/SSequencer.h
--- a/Engine/Source/Editor/Sequencer/Private/SSequencer.h
+++ b/Engine/Source/Editor/Sequencer/Private/SSequencer.h
@@ -180,12 +180,6 @@ public:
 	/** Access the tree view for this sequencer */
 	TSharedPtr<SSequencerTreeView> GetTreeView() const;
 
-	/** Access the currently active edit tool */
-	ISequencerEditTool& GetEditTool() const
-	{
-		return *EditTool;
-	}
-
 	/** Generate a helper structure that can be used to transform between phsyical space and virtual space in the track area */
 	FVirtualTrackArea GetVirtualTrackArea() const;
 
@@ -195,17 +189,15 @@ public:
 	/** @return a numeric type interface that will parse and display numbers as frames and times correctly */
 	TSharedRef<INumericTypeInterface<float>> GetNumericTypeInterface();
 
+	/** @return a numeric type interface that will parse and display numbers as frames and times correctly, including any zero padding, if necessary */
+	TSharedRef<INumericTypeInterface<float>> GetZeroPadNumericTypeInterface();
+
 public:
 
 	// FNotifyHook overrides
 
 	void NotifyPostChange(const FPropertyChangedEvent& PropertyChangedEvent, FEditPropertyChain* PropertyThatChanged);
 
-protected:
-
-	/** Update the details view from the currently selected keys and sections. */
-	void UpdateDetailsView();
-
 protected:
 
 	// SWidget interface
@@ -219,15 +211,6 @@ protected:
 
 private:
 	
-	/** Handles checking whether the details view is enabled. */
-	bool HandleDetailsViewEnabled() const;
-
-	/** Handles determining the visibility of the details view selection tip. */
-	EVisibility HandleDetailsViewTipVisibility() const;
-
-	/** Handles determining the visibility of the details view. */
-	EVisibility HandleDetailsViewVisibility() const;
-
 	/** Handles key selection changes. */
 	void HandleKeySelectionChanged();
 
@@ -240,15 +223,15 @@ private:
 	/** Handles section selection changes. */
 	void HandleSectionSelectionChanged();
 
-	/** Check whether the specified edit tool is enabled */
-	bool IsEditToolEnabled(FName InIdentifier);
-
 	/** Empty active timer to ensure Slate ticks during Sequencer playback */
 	EActiveTimerReturnType EnsureSlateTickDuringPlayback(double InCurrentTime, float InDeltaTime);	
 
 	/** Makes the toolbar. */
 	TSharedRef<SWidget> MakeToolBar();
 
+	/** Makes add button. */
+	TSharedRef<SWidget> MakeAddButton();
+
 	/** Makes the add menu for the toolbar. */
 	TSharedRef<SWidget> MakeAddMenu();
 
@@ -261,9 +244,13 @@ private:
 	/** Makes the auto-key menu for the toolbar. */
 	TSharedRef<SWidget> MakeAutoKeyMenu();
 
+public:
+
 	/** Makes and configures a set of the standard UE transport controls. */
 	TSharedRef<SWidget> MakeTransportControls();
 
+private:
+	
 	/**
 	* @return The value of the current time snap interval.
 	*/
@@ -350,9 +337,6 @@ private:
 	/** Gets the root movie scene name */
 	FText GetRootAnimationName() const;
 
-	/** Gets the title of the passed in shot section */
-	FText GetShotSectionTitle(UMovieSceneSection* ShotSection) const;
-
 	/** Gets whether or not the breadcrumb trail should be visible. */
 	EVisibility GetBreadcrumbTrailVisibility() const;
 
@@ -388,11 +372,10 @@ public:
 	/** Generate a paste menu args structure */
 	struct FPasteContextMenuArgs GeneratePasteArgs(float PasteAtTime, TSharedPtr<FMovieSceneClipboard> Clipboard = nullptr);
 
-
 private:
 
-	/** Holds the details view. */
-	TSharedPtr<IDetailsView> DetailsView;
+	/** Goto box widget. */
+	TSharedPtr<SWidget> GotoBox;
 
 	/** Section area widget */
 	TSharedPtr<SSequencerTrackArea> TrackArea;
@@ -419,7 +402,7 @@ private:
 	TSharedPtr<SSequencerTreeView> TreeView;
 
 	/** The main sequencer interface */
-	TWeakPtr<FSequencer> Sequencer;
+	TWeakPtr<FSequencer> SequencerPtr;
 
 	/** Cached settings provided to the sequencer itself on creation */
 	USequencerSettings* Settings;
@@ -433,14 +416,12 @@ private:
 	/** Whether the user is selecting. Ignore selection changes from the level when the user is selecting. */
 	bool bUserIsSelecting;
 
-	/** The current edit tool */
-	TUniquePtr<ISequencerEditTool> EditTool;
-
 	/** Extender to use for the 'add' menu */
 	TSharedPtr<FExtender> AddMenuExtender;
 
 	/** Numeric type interface used for converting parsing and generating strings from numbers */
 	TSharedPtr<INumericTypeInterface<float>> NumericTypeInterface;
+	TSharedPtr<INumericTypeInterface<float>> ZeroPadNumericTypeInterface;
 
 	FOnGetAddMenuContent OnGetAddMenuContent;
 	/** Called when the user has begun dragging the playback range */
