Copying //UE4/Dev-Sequencer to Dev-Main (//UE4/Dev-Main)

diff --git a/Engine/Source/Editor/Sequencer/Private/SSequencerGotoBox.cpp b/Engine/Source/Editor/Sequencer/Private/SSequencerGotoBox.cpp
new file mode 100644
--- /dev/null
+++ b/Engine/Source/Editor/Sequencer/Private/SSequencerGotoBox.cpp
@@ -0,0 +1,87 @@
+// Copyright 1998-2015 Epic Games, Inc. All Rights Reserved.
+
+#include "SequencerPrivatePCH.h"
+#include "Sequencer.h"
+#include "SequencerSettings.h"
+#include "SSequencerGotoBox.h"
+
+
+#define LOCTEXT_NAMESPACE "Sequencer"
+
+
+void SSequencerGotoBox::Construct(const FArguments& InArgs, const TSharedRef<FSequencer>& InSequencer, USequencerSettings& InSettings, const TSharedRef<INumericTypeInterface<float>>& InNumericTypeInterface)
+{
+	WasHidden = true;
+	SequencerPtr = InSequencer;
+	Settings = &InSettings;
+	NumericTypeInterface = InNumericTypeInterface;
+
+	ChildSlot
+	[
+		SNew(SBorder)
+			.BorderImage(FEditorStyle::GetBrush("ToolPanel.DarkGroupBorder"))
+			.Padding(6.0f)
+			.Visibility_Lambda([this]() -> EVisibility { return Settings->GetShowGotoBox() ? EVisibility::Visible : EVisibility::Collapsed; })
+			[
+				SNew(SHorizontalBox)
+
+				+ SHorizontalBox::Slot()
+					.AutoWidth()
+					.VAlign(VAlign_Center)
+					[
+						SNew(STextBlock)
+							.Text(LOCTEXT("GotoLabel", "Go to:"))
+					]
+
+				+ SHorizontalBox::Slot()
+					.Padding(6.0f, 0.0f, 0.0f, 0.0f)
+					.AutoWidth()
+					[
+						SAssignNew(EntryBox, SNumericEntryBox<float>)
+							.MinDesiredValueWidth(64.0f)
+							.OnValueCommitted(this, &SSequencerGotoBox::HandleEntryBoxValueCommitted)
+							.TypeInterface(NumericTypeInterface)
+							.Value_Lambda([this](){ return SequencerPtr.Pin()->GetGlobalTime(); })
+					]
+			]
+	];
+}
+
+
+void SSequencerGotoBox::Tick(const FGeometry& AllottedGeometry, const double InCurrentTime, const float InDeltaTime)
+{
+	if (WasHidden)
+	{
+		FSlateApplication::Get().SetAllUserFocus(EntryBox, EFocusCause::Navigation);
+		WasHidden = false;
+	}
+}
+
+
+void SSequencerGotoBox::HandleEntryBoxValueCommitted(float Value, ETextCommit::Type CommitType)
+{
+	Settings->SetShowGotoBox(false);
+	WasHidden = true;
+
+	if (CommitType == ETextCommit::OnCleared)
+	{
+		return;
+	}
+
+	TSharedPtr<FSequencer> Sequencer = SequencerPtr.Pin();
+
+	// scroll view range if new value is not visible
+	const FAnimatedRange ViewRange = Sequencer->GetViewRange();
+
+	if (!ViewRange.Contains(Value))
+	{
+		const float HalfViewWidth = 0.5f * ViewRange.Size<float>();
+		const TRange<float> NewRange = TRange<float>(Value - HalfViewWidth, Value + HalfViewWidth);
+		Sequencer->SetViewRange(NewRange);
+	}
+
+	Sequencer->SetGlobalTime(Value);
+}
+
+
+#undef LOCTEXT_NAMESPACE
diff --git a/Engine/Source/Editor/Sequencer/Private/SSequencerGotoBox.h b/Engine/Source/Editor/Sequencer/Private/SSequencerGotoBox.h
new file mode 100644
--- /dev/null
+++ b/Engine/Source/Editor/Sequencer/Private/SSequencerGotoBox.h
@@ -0,0 +1,50 @@
+// Copyright 1998-2015 Epic Games, Inc. All Rights Reserved.
+
+#pragma once
+
+#include "NumericTypeInterface.h"
+#include "SNumericEntryBox.h"
+
+
+class FSequencer;
+class USequencerSettings;
+
+
+class SSequencerGotoBox
+	: public SCompoundWidget
+{
+public:
+
+	SLATE_BEGIN_ARGS(SSequencerGotoBox) { }
+	SLATE_END_ARGS()
+
+	void Construct(const FArguments& InArgs, const TSharedRef<FSequencer>& InSequencer, USequencerSettings& InSettings, const TSharedRef<INumericTypeInterface<float>>& InNumericTypeInterface);
+
+public:
+
+	// SWidget interface
+
+	virtual void Tick(const FGeometry& AllottedGeometry, const double InCurrentTime, const float InDeltaTime) override;
+
+private:
+
+	/** Callback for committed values in the entry box. */
+	void HandleEntryBoxValueCommitted(float Value, ETextCommit::Type CommitType);
+
+private:
+
+	/** The entry box widget. */
+	TSharedPtr<SNumericEntryBox<float>> EntryBox;
+
+	/** Numeric type interface used for converting parsing and generating strings from numbers. */
+	TSharedPtr<INumericTypeInterface<float>> NumericTypeInterface;
+
+	/** The main sequencer interface. */
+	TWeakPtr<FSequencer> SequencerPtr;
+
+	/** Cached settings provided to the sequencer itself on creation. */
+	USequencerSettings* Settings;
+
+	/** Whether the box was hidden. */
+	bool WasHidden;
+};
diff --git a/Engine/Source/Editor/Sequencer/Private/SSequencerSection.cpp b/Engine/Source/Editor/Sequencer/Private/SSequencerSection.cpp
--- a/Engine/Source/Editor/Sequencer/Private/SSequencerSection.cpp
+++ b/Engine/Source/Editor/Sequencer/Private/SSequencerSection.cpp
@@ -6,7 +6,6 @@
 #include "IKeyArea.h"
 #include "ISequencerSection.h"
 #include "MovieSceneSection.h"
-#include "MovieSceneShotSection.h"
 #include "MovieSceneToolHelpers.h"
 #include "CommonMovieSceneTools.h"
 #include "SequencerHotspots.h"
@@ -193,11 +192,11 @@ void SSequencerSection::CheckForEdgeInteraction( const FPointerEvent& MouseEvent
 
 	if( SectionRectLeft.IsUnderLocation( MouseEvent.GetScreenSpacePosition() ) )
 	{
-		GetSequencer().GetEditTool().SetHotspot(MakeShareable( new FSectionResizeHotspot(FSectionResizeHotspot::Left, FSectionHandle(ParentSectionArea, SectionIndex))) );
+		GetSequencer().SetHotspot(MakeShareable( new FSectionResizeHotspot(FSectionResizeHotspot::Left, FSectionHandle(ParentSectionArea, SectionIndex))) );
 	}
 	else if( SectionRectRight.IsUnderLocation( MouseEvent.GetScreenSpacePosition() ) )
 	{
-		GetSequencer().GetEditTool().SetHotspot(MakeShareable( new FSectionResizeHotspot(FSectionResizeHotspot::Right, FSectionHandle(ParentSectionArea, SectionIndex))) );
+		GetSequencer().SetHotspot(MakeShareable( new FSectionResizeHotspot(FSectionResizeHotspot::Right, FSectionHandle(ParentSectionArea, SectionIndex))) );
 	}
 }
 
@@ -279,6 +278,8 @@ int32 SSequencerSection::OnPaint( const FPaintArgs& Args, const FGeometry& Allot
 
 void SSequencerSection::PaintKeys( const FGeometry& AllottedGeometry, const FSlateRect& MyClippingRect, FSlateWindowElementList& OutDrawElements, int32 LayerId, const FWidgetStyle& InWidgetStyle, bool bParentEnabled ) const
 {
+	static const FName HighlightBrushName("Sequencer.AnimationOutliner.DefaultBorder");
+
 	static const FName BackgroundBrushName("Sequencer.SectionArea.Background");
 	static const FName CircleKeyBrushName("Sequencer.KeyCircle");
 	static const FName DiamondKeyBrushName("Sequencer.KeyDiamond");
@@ -295,7 +296,7 @@ void SSequencerSection::PaintKeys( const FGeometry& AllottedGeometry, const FSla
 	FLinearColor SelectionColor = FEditorStyle::GetSlateColor(SelectionColorName).GetColor( InWidgetStyle );
 	FLinearColor SelectedKeyColor = SelectionColor;
 	FSequencer& Sequencer = ParentSectionArea->GetSequencer();
-	TSharedPtr<ISequencerHotspot> Hotspot = Sequencer.GetEditTool().GetHotspot();
+	TSharedPtr<ISequencerHotspot> Hotspot = Sequencer.GetHotspot();
 
 	// get hovered key
 	FSequencerSelectedKey HoveredKey;
@@ -320,6 +321,7 @@ void SSequencerSection::PaintKeys( const FGeometry& AllottedGeometry, const FSla
 	// draw all keys in each key area
 	UMovieSceneSection& SectionObject = *SectionInterface->GetSectionObject();
 
+	const FSlateBrush* HighlightBrush = FEditorStyle::GetBrush(HighlightBrushName);
 	const FSlateBrush* BackgroundBrush = FEditorStyle::GetBrush(BackgroundBrushName);
 	const FSlateBrush* CircleKeyBrush = FEditorStyle::GetBrush(CircleKeyBrushName);
 	const FSlateBrush* DiamondKeyBrush = FEditorStyle::GetBrush(DiamondKeyBrushName);
@@ -356,6 +358,19 @@ void SSequencerSection::PaintKeys( const FGeometry& AllottedGeometry, const FSla
 			); 
 		}
 
+		if (LayoutElement.GetDisplayNode().IsValid() && LayoutElement.GetDisplayNode()->IsHovered())
+		{
+			FSlateDrawElement::MakeBox( 
+				OutDrawElements,
+				LayerId + 1,
+				KeyAreaGeometry.ToPaintGeometry(),
+				HighlightBrush,
+				MyClippingRect,
+				DrawEffects,
+				FLinearColor(1.f, 1.f, 1.f, 0.2f)
+			); 
+		}
+
 		TArray<FKeyHandle> KeyHandles = KeyArea->GetUnsortedKeyHandles();
 
 		if (!KeyHandles.Num())
@@ -537,7 +552,7 @@ void SSequencerSection::DrawSectionHandles( const FGeometry& AllottedGeometry, c
 	bool bRightHandleActive = false;
 
 	// Get the hovered/selected state for the section handles from the hotspot
-	TSharedPtr<ISequencerHotspot> Hotspot = GetSequencer().GetEditTool().GetHotspot();
+	TSharedPtr<ISequencerHotspot> Hotspot = GetSequencer().GetHotspot();
 
 	if (Hotspot.IsValid() && (
 		Hotspot->GetType() == ESequencerHotspot::SectionResize_L ||
@@ -644,7 +659,7 @@ FReply SSequencerSection::OnMouseButtonDown( const FGeometry& MyGeometry, const
 	FSequencerSelectedKey HoveredKey;
 	
 	// The hovered key is defined from the sequencer hotspot
-	TSharedPtr<ISequencerHotspot> Hotspot = GetSequencer().GetEditTool().GetHotspot();
+	TSharedPtr<ISequencerHotspot> Hotspot = GetSequencer().GetHotspot();
 	if (Hotspot.IsValid() && Hotspot->GetType() == ESequencerHotspot::Key)
 	{
 		HoveredKey = static_cast<FKeyHotspot*>(Hotspot.Get())->Key;
@@ -663,8 +678,7 @@ FReply SSequencerSection::OnMouseButtonDown( const FGeometry& MyGeometry, const
 			Sequencer.GetSelection().AddToSelection(NewKey);
 
 			// Pass the event to the tool to copy the hovered key and move it
-			auto& EditTool = GetSequencer().GetEditTool();
-			EditTool.SetHotspot( MakeShareable( new FKeyHotspot(NewKey, ParentSectionArea) ) );
+			GetSequencer().SetHotspot( MakeShareable( new FKeyHotspot(NewKey) ) );
 
 			// Return unhandled so that the EditTool can handle the mouse down based on the newly created keyframe and prepare to move it
 			return FReply::Unhandled();
@@ -706,17 +720,15 @@ FReply SSequencerSection::OnMouseButtonDoubleClick( const FGeometry& MyGeometry,
 
 FReply SSequencerSection::OnMouseMove( const FGeometry& MyGeometry, const FPointerEvent& MouseEvent )
 {
-	auto& EditTool = GetSequencer().GetEditTool();
-
 	// Checked for hovered key
 	FSequencerSelectedKey KeyUnderMouse = GetKeyUnderMouse( MouseEvent.GetScreenSpacePosition(), MyGeometry );
 	if ( KeyUnderMouse.IsValid() )
 	{
-		EditTool.SetHotspot( MakeShareable( new FKeyHotspot(KeyUnderMouse, ParentSectionArea) ) );
+		GetSequencer().SetHotspot( MakeShareable( new FKeyHotspot(KeyUnderMouse) ) );
 	}
 	else
 	{
-		EditTool.SetHotspot( MakeShareable( new FSectionHotspot(FSectionHandle(ParentSectionArea, SectionIndex))) );
+		GetSequencer().SetHotspot( MakeShareable( new FSectionHotspot(FSectionHandle(ParentSectionArea, SectionIndex))) );
 
 		// Only check for edge interaction if not hovering over a key
 		CheckForEdgeInteraction( MouseEvent, MyGeometry );
@@ -729,7 +741,7 @@ FReply SSequencerSection::OnMouseMove( const FGeometry& MyGeometry, const FPoint
 void SSequencerSection::OnMouseLeave( const FPointerEvent& MouseEvent )
 {
 	SCompoundWidget::OnMouseLeave( MouseEvent );
-	GetSequencer().GetEditTool().SetHotspot(nullptr);
+	GetSequencer().SetHotspot(nullptr);
 }
 
 static float ThrobDurationSeconds = .5f;
diff --git a/Engine/Source/Editor/Sequencer/Private/SSequencerSectionAreaView.cpp b/Engine/Source/Editor/Sequencer/Private/SSequencerSectionAreaView.cpp
--- a/Engine/Source/Editor/Sequencer/Private/SSequencerSectionAreaView.cpp
+++ b/Engine/Source/Editor/Sequencer/Private/SSequencerSectionAreaView.cpp
@@ -8,7 +8,6 @@
 #include "SequencerDisplayNode.h"
 #include "SSequencer.h"
 #include "Sequencer.h"
-#include "MovieSceneShotSection.h"
 #include "CommonMovieSceneTools.h"
 
 namespace SequencerSectionAreaConstants
@@ -83,11 +82,18 @@ FVector2D SSequencerSectionAreaView::ComputeDesiredSize(float) const
 {
 	// Note: X Size is not used
 	FVector2D Size(100, 0.f);
-	for (int32 Index = 0; Index < Children.Num(); ++Index)
+	if (Children.Num())
+	{
+		for (int32 Index = 0; Index < Children.Num(); ++Index)
+		{
+			Size.Y = FMath::Max(Size.Y, Children[Index]->GetDesiredSize().Y);
+		}
+	}
+	else
 	{
-		Size.Y = FMath::Max(Size.Y, Children[Index]->GetDesiredSize().Y);
+		Size.Y = SectionAreaNode->GetNodeHeight();
 	}
-	return Size + SectionAreaNode->GetNodePadding().Bottom;
+	return Size;
 }
 
 void SSequencerSectionAreaView::GenerateSectionWidgets()
@@ -159,7 +165,6 @@ void SSequencerSectionAreaView::OnArrangeChildren( const FGeometry& AllottedGeom
 	int32 MaxTracks = MaxRowIndex + 1;
 
 
-	float SectionHeight = AllottedGeometry.GetLocalSize().Y - SectionAreaNode->GetNodePadding().Bottom;
 
 	FTimeToPixel TimeToPixelConverter = GetTimeToPixel( AllottedGeometry );
 
@@ -171,7 +176,7 @@ void SSequencerSectionAreaView::OnArrangeChildren( const FGeometry& AllottedGeom
 
 		int32 RowIndex = SectionInterface->GetSectionObject()->GetRowIndex();
 
-		FGeometry SectionGeometry = SequencerSectionUtils::GetSectionGeometry( AllottedGeometry, RowIndex, MaxTracks, SectionHeight, SectionInterface, TimeToPixelConverter );
+		FGeometry SectionGeometry = SequencerSectionUtils::GetSectionGeometry( AllottedGeometry, RowIndex, MaxTracks, Widget->GetDesiredSize().Y, SectionInterface, TimeToPixelConverter );
 
 		EVisibility Visibility = Widget->GetVisibility();
 		if( ArrangedChildren.Accepts( Visibility ) )
diff --git a/Engine/Source/Editor/Sequencer/Private/SSequencerTrackArea.cpp b/Engine/Source/Editor/Sequencer/Private/SSequencerTrackArea.cpp
--- a/Engine/Source/Editor/Sequencer/Private/SSequencerTrackArea.cpp
+++ b/Engine/Source/Editor/Sequencer/Private/SSequencerTrackArea.cpp
@@ -10,6 +10,11 @@
 #include "IKeyArea.h"
 #include "ISequencerSection.h"
 #include "SSequencerSection.h"
+#include "SequencerHotspots.h"
+#include "Tools/SequencerEditTool_Movement.h"
+#include "Tools/SequencerEditTool_Selection.h"
+#include "MovieSceneSection.h"
+#include "VirtualTrackArea.h"
 #include "ISequencerTrackEditor.h"
 
 FTrackAreaSlot::FTrackAreaSlot(const TSharedPtr<SSequencerTrackLane>& InSlotContent)
@@ -147,8 +152,39 @@ int32 SSequencerTrackArea::OnPaint(const FPaintArgs& Args, const FGeometry& Allo
 			LayerId = FMath::Max(LayerId, ThisWidgetLayerId);
 		}
 
-		return Sequencer.Pin()->GetEditTool().OnPaint(AllottedGeometry, MyClippingRect, OutDrawElements, LayerId + 2);
+		if (EditTool.IsValid())
+		{
+			LayerId = EditTool->OnPaint(AllottedGeometry, MyClippingRect, OutDrawElements, LayerId + 1);
+		}
+
+		static const FName SelectionColorPressedName("SelectionColor");
+		FLinearColor SelectionColor = FEditorStyle::GetSlateColor(SelectionColorPressedName).GetColor(InWidgetStyle).CopyWithNewOpacity(0.5f);
+
+		TOptional<FHighlightRegion> HighlightRegion = TreeView.Pin()->GetHighlightRegion();
+		if (HighlightRegion.IsSet())
+		{
+			FSlateDrawElement::MakeBox(
+				OutDrawElements,
+				LayerId+1,
+				AllottedGeometry.ToPaintGeometry(FVector2D(0.f, HighlightRegion->Top - 4.f), FVector2D(AllottedGeometry.Size.X, 4.f)),
+				FEditorStyle::GetBrush("Sequencer.TrackHoverHighlight_Top"),
+				MyClippingRect,
+				ESlateDrawEffect::None,
+				SelectionColor
+			);
+		
+			FSlateDrawElement::MakeBox(
+				OutDrawElements,
+				LayerId+1,
+				AllottedGeometry.ToPaintGeometry(FVector2D(0.f, HighlightRegion->Bottom), FVector2D(AllottedGeometry.Size.X, 4.f)),
+				FEditorStyle::GetBrush("Sequencer.TrackHoverHighlight_Bottom"),
+				MyClippingRect,
+				ESlateDrawEffect::None,
+				SelectionColor
+			);
+		}
 	}
+
 	return LayerId;
 }
 
@@ -158,7 +194,7 @@ FReply SSequencerTrackArea::OnMouseButtonDown( const FGeometry& MyGeometry, cons
 	if ( Sequencer.IsValid() )
 	{
 		// Always ensure the edit tool is set up
-		InputStack.SetHandlerAt(0, &Sequencer.Pin()->GetEditTool());
+		InputStack.SetHandlerAt(0, EditTool.Get());
 		return InputStack.HandleMouseButtonDown(*this, MyGeometry, MouseEvent);
 	}
 	return FReply::Unhandled();
@@ -172,19 +208,112 @@ FReply SSequencerTrackArea::OnMouseButtonUp( const FGeometry& MyGeometry, const
 		FContextMenuSuppressor SuppressContextMenus(TimeSliderController.ToSharedRef());
 
 		// Always ensure the edit tool is set up
-		InputStack.SetHandlerAt(0, &Sequencer.Pin()->GetEditTool());
+		InputStack.SetHandlerAt(0, EditTool.Get());
 		return InputStack.HandleMouseButtonUp(*this, MyGeometry, MouseEvent);
 	}
 	return FReply::Unhandled();
 }
 
+bool SSequencerTrackArea::CanActivateEditTool(FName Identifier) const
+{
+	if (!EditTool.IsValid())
+	{
+		return true;
+	}
+	// Can't activate a tool that's already active
+	else if (EditTool->GetIdentifier() == Identifier)
+	{
+		return false;
+	}
+	// Can only activate a new tool if the current one will let us
+	else
+	{
+		return EditTool->CanDeactivate();
+	}
+}
+
+template<typename EditToolType>
+bool SSequencerTrackArea::AttemptToActivateTool()
+{
+	if ( Sequencer.IsValid() )
+	{
+		if (CanActivateEditTool(EditToolType::Identifier))
+		{
+			EditTool.Reset(new EditToolType(*Sequencer.Pin()));
+			return true;
+		}
+	}
+
+	return false;
+}
+
+void SSequencerTrackArea::UpdateHoverStates( const FGeometry& MyGeometry, const FPointerEvent& MouseEvent )
+{
+	TSharedPtr<SSequencerTreeView> PinnedTreeView = TreeView.Pin();
+
+	// Set the node that we are hovering
+	TSharedPtr<FSequencerDisplayNode> NewHoveredNode = PinnedTreeView->HitTestNode(MyGeometry.AbsoluteToLocal(MouseEvent.GetScreenSpacePosition()).Y);
+	PinnedTreeView->GetNodeTree()->SetHoveredNode(NewHoveredNode);
+
+	if ( Sequencer.IsValid() )
+	{
+		TSharedPtr<ISequencerHotspot> Hotspot = Sequencer.Pin()->GetHotspot();
+		if (Hotspot.IsValid())
+		{
+			const ESequencerHotspot Type = Hotspot->GetType();
+
+			// Activate the move tool if we're over a key or section resize handle
+			if (Type == ESequencerHotspot::Key || Type == ESequencerHotspot::SectionResize_L || Type == ESequencerHotspot::SectionResize_R)
+			{
+				AttemptToActivateTool<FSequencerEditTool_Movement>();
+			}
+			else if (Type == ESequencerHotspot::Section)
+			{
+				UMovieSceneSection* Section = StaticCastSharedPtr<FSectionHotspot>(Hotspot)->Section.GetSectionObject();
+
+				// Move sections if they are selected
+				if (Sequencer.Pin()->GetSelection().IsSelected(Section))
+				{
+					AttemptToActivateTool<FSequencerEditTool_Movement>();
+				}
+				else
+				{
+					// Activate selection mode if the section has keys, otherwise just move it
+					TSet<FKeyHandle> KeyHandles;
+					Section->GetKeyHandles(KeyHandles);
+
+					if (KeyHandles.Num())
+					{
+						AttemptToActivateTool<FSequencerEditTool_Selection>();
+					}
+					else
+					{
+						AttemptToActivateTool<FSequencerEditTool_Movement>();
+					}
+				}
+			}
+			else
+			{
+				// Any other region implies selection mode
+				AttemptToActivateTool<FSequencerEditTool_Selection>();
+			}
+		}
+		else
+		{
+			// Any other region implies selection mode
+			AttemptToActivateTool<FSequencerEditTool_Selection>();
+		}
+	}
+}
 
 FReply SSequencerTrackArea::OnMouseMove( const FGeometry& MyGeometry, const FPointerEvent& MouseEvent )
 {
 	if ( Sequencer.IsValid() )
 	{
+		UpdateHoverStates(MyGeometry, MouseEvent);
+
 		// Always ensure the edit tool is set up
-		InputStack.SetHandlerAt(0, &Sequencer.Pin()->GetEditTool());
+		InputStack.SetHandlerAt(0, EditTool.Get());
 
 		FReply Reply = InputStack.HandleMouseMove(*this, MyGeometry, MouseEvent);
 
@@ -208,7 +337,7 @@ FReply SSequencerTrackArea::OnMouseWheel( const FGeometry& MyGeometry, const FPo
 	if ( Sequencer.IsValid() )
 	{
 		// Always ensure the edit tool is set up
-		InputStack.SetHandlerAt(0, &Sequencer.Pin()->GetEditTool());
+		InputStack.SetHandlerAt(0, EditTool.Get());
 		return InputStack.HandleMouseWheel(*this, MyGeometry, MouseEvent);
 	}
 	return FReply::Unhandled();
@@ -219,7 +348,10 @@ void SSequencerTrackArea::OnMouseEnter(const FGeometry& MyGeometry, const FPoint
 {
 	if ( Sequencer.IsValid() )
 	{
-		Sequencer.Pin()->GetEditTool().OnMouseEnter(*this, MyGeometry, MouseEvent);
+		if (EditTool.IsValid())
+		{
+			EditTool->OnMouseEnter(*this, MyGeometry, MouseEvent);
+		}
 	}
 }
 
@@ -228,7 +360,12 @@ void SSequencerTrackArea::OnMouseLeave(const FPointerEvent& MouseEvent)
 {
 	if ( Sequencer.IsValid() )
 	{
-		Sequencer.Pin()->GetEditTool().OnMouseLeave(*this, MouseEvent);
+		if (EditTool.IsValid())
+		{
+			EditTool->OnMouseLeave(*this, MouseEvent);
+		}
+
+		TreeView.Pin()->GetNodeTree()->SetHoveredNode(nullptr);
 	}
 }
 
@@ -237,7 +374,10 @@ void SSequencerTrackArea::OnMouseCaptureLost()
 {
 	if ( Sequencer.IsValid() )
 	{
-		Sequencer.Pin()->GetEditTool().OnMouseCaptureLost();
+		if (EditTool.IsValid())
+		{
+			EditTool->OnMouseCaptureLost();
+		}
 	}
 }
 
@@ -251,8 +391,12 @@ FCursorReply SSequencerTrackArea::OnCursorQuery( const FGeometry& MyGeometry, co
 			return FCursorReply::Cursor(EMouseCursor::GrabHandClosed);
 		}
 
-		return Sequencer.Pin()->GetEditTool().OnCursorQuery(MyGeometry, CursorEvent);
+		if (EditTool.IsValid())
+		{
+			return EditTool->OnCursorQuery(MyGeometry, CursorEvent);
+		}
 	}
+
 	return FCursorReply::Unhandled();
 }
 
@@ -261,11 +405,6 @@ void SSequencerTrackArea::Tick( const FGeometry& AllottedGeometry, const double
 {
 	CachedGeometry = AllottedGeometry;
 
-	if ( Sequencer.IsValid() )
-	{
-		Sequencer.Pin()->GetEditTool().Tick(AllottedGeometry, InCurrentTime, InDeltaTime);
-	}
-
 	FVector2D Size = AllottedGeometry.GetLocalSize();
 
 	if (SizeLastFrame.IsSet() && Size.X != SizeLastFrame->X)
diff --git a/Engine/Source/Editor/Sequencer/Private/SSequencerTrackArea.h b/Engine/Source/Editor/Sequencer/Private/SSequencerTrackArea.h
--- a/Engine/Source/Editor/Sequencer/Private/SSequencerTrackArea.h
+++ b/Engine/Source/Editor/Sequencer/Private/SSequencerTrackArea.h
@@ -9,6 +9,7 @@ class FSequencerTimeSliderController;
 class FSequencer;
 
 #include "SequencerInputHandlerStack.h"
+#include "ISequencerEditTool.h"
 
 /**
  * Structure representing a slot in the track area.
@@ -83,6 +84,18 @@ public:
 	virtual FVector2D ComputeDesiredSize(float) const override;
 	virtual FChildren* GetChildren() override;
 
+protected:
+
+	/** Check whether it's possible to activate the specified tool */
+	bool CanActivateEditTool(FName Identifier) const;
+
+	/** Attempt to activate the tool specified by the template parameter */
+	template<typename EditToolType>
+	bool AttemptToActivateTool();
+
+	/** Update any hover state required for the track area */
+	void UpdateHoverStates( const FGeometry& MyGeometry, const FPointerEvent& MouseEvent );
+
 private:
 
 	/** The track area's children. */
@@ -108,6 +121,9 @@ private:
 	/** Keep a hold of the size of the area so we can maintain zoom levels. */
 	TOptional<FVector2D> SizeLastFrame;
 
+	/** The currently actuve edit tool on this track area */
+	TUniquePtr<ISequencerEditTool> EditTool;
+
 private:
 
 	/** Input handler stack responsible for routing input to the different handlers */
diff --git a/Engine/Source/Editor/Sequencer/Private/SSequencerTrackLane.cpp b/Engine/Source/Editor/Sequencer/Private/SSequencerTrackLane.cpp
--- a/Engine/Source/Editor/Sequencer/Private/SSequencerTrackLane.cpp
+++ b/Engine/Source/Editor/Sequencer/Private/SSequencerTrackLane.cpp
@@ -11,27 +11,25 @@ void SSequencerTrackLane::Construct(const FArguments& InArgs, const TSharedRef<F
 	TreeView = InTreeView;
 
 	ChildSlot
-		.HAlign(HAlign_Fill)
-		.Padding(0, DisplayNode->GetNodePadding().Top, 0, 0)
-		[
-			InArgs._Content.Widget
-		];
+	.HAlign(HAlign_Fill)
+	.Padding(0, DisplayNode->GetNodePadding().Top, 0, 0)
+	[
+		InArgs._Content.Widget
+	];
 }
 
 
 void DrawLaneRecursive(const TSharedRef<FSequencerDisplayNode>& DisplayNode, const FGeometry& AllottedGeometry, float& YOffs, const FSlateRect& MyClippingRect, FSlateWindowElementList& OutDrawElements, int32 LayerId, const FWidgetStyle& InWidgetStyle)
 {
-	static const FName HoverBorderName("Sequencer.AnimationOutliner.SelectionBorderHover");
-	static const FName SelectionBorderName("Sequencer.AnimationOutliner.SelectionBorder");
-	static const FName SelectionBorderInactiveName("Sequencer.AnimationOutliner.SelectionBorderInactive");
-	static const FName SelectionColorName("SelectionColor_Pressed");
+	static const FName BorderName("Sequencer.AnimationOutliner.DefaultBorder");
+	static const FName SelectionColorName("SelectionColor");
 	static const FName SelectionColorInactiveName("SelectionColor_Inactive");
 
 	if (DisplayNode->IsHidden())
 	{
 		return;
 	}
-
+	
 	float TotalNodeHeight = DisplayNode->GetNodeHeight() + DisplayNode->GetNodePadding().Combined();
 
 	// draw selection border
@@ -45,10 +43,6 @@ void DrawLaneRecursive(const TSharedRef<FSequencerDisplayNode>& DisplayNode, con
 				: SelectionColorInactiveName
 		).GetColor(InWidgetStyle);
 
-		FName SelectionBrush = (SequencerSelection.GetActiveSelection() == FSequencerSelection::EActiveSelection::OutlinerNode)
-			? SelectionBorderName
-			: SelectionBorderInactiveName;
-
 		FSlateDrawElement::MakeBox(
 			OutDrawElements,
 			LayerId,
@@ -56,7 +50,7 @@ void DrawLaneRecursive(const TSharedRef<FSequencerDisplayNode>& DisplayNode, con
 				FVector2D(0, YOffs),
 				FVector2D(AllottedGeometry.Size.X, TotalNodeHeight)
 			),
-			FEditorStyle::GetBrush(SelectionBrush),
+			FEditorStyle::GetBrush(BorderName),
 			MyClippingRect,
 			ESlateDrawEffect::None,
 			SelectionColor
@@ -73,7 +67,7 @@ void DrawLaneRecursive(const TSharedRef<FSequencerDisplayNode>& DisplayNode, con
 				FVector2D(0, YOffs),
 				FVector2D(AllottedGeometry.Size.X, TotalNodeHeight)
 			),
-			FEditorStyle::GetBrush(HoverBorderName),
+			FEditorStyle::GetBrush(BorderName),
 			MyClippingRect,
 			ESlateDrawEffect::None,
 			FLinearColor(1.0f, 1.0f, 1.0f, 0.1f)
@@ -120,38 +114,26 @@ int32 SSequencerTrackLane::OnPaint(const FPaintArgs& Args, const FGeometry& Allo
 	return SCompoundWidget::OnPaint(Args, AllottedGeometry, MyClippingRect, OutDrawElements, LayerId + 1, InWidgetStyle, bParentEnabled);
 }
 
-
-float SSequencerTrackLane::GetPhysicalPosition() const
+FVector2D SSequencerTrackLane::ComputeDesiredSize(float LayoutScale) const
 {
-	// Positioning strategy:
-	// We know that there must be at least one tree view row around that references something in our subtree (otherwise this widget wouldn't exist)
-	// Knowing this, we just need to root it out by traversing the visible tree, and asking for the first rows's offset
-	float NegativeOffset = 0.f;
-	TOptional<float> Top;
-
-	auto TreeViewPinned = TreeView.Pin();
+	float Height = DisplayNode->GetNodeHeight() + DisplayNode->GetNodePadding().Combined();
 	
-	// Iterate parent first until we find a tree view row we can use for the offset height
-	auto Iter = [&](FSequencerDisplayNode& InNode){
-		
-		auto ChildRowGeometry = TreeViewPinned->GetPhysicalGeometryForNode(InNode.AsShared());
-		if (ChildRowGeometry.IsSet())
-		{
-			Top = ChildRowGeometry->PhysicalTop;
-			// Stop iterating
-			return false;
-		}
-
-		NegativeOffset -= InNode.GetNodeHeight() + InNode.GetNodePadding().Combined();
-		return true;
-	};
-
-	DisplayNode->TraverseVisible_ParentFirst(Iter);
-
-	if (!Top.IsSet())
+	
+	if (DisplayNode->GetType() == ESequencerNode::Track || DisplayNode->GetType() == ESequencerNode::Category)
 	{
-		Top = 0.f;
+		const bool bIncludeThisNode = false;
+		
+		// These types of nodes need to consider the entire visible hierarchy for their desired size
+		DisplayNode->TraverseVisible_ParentFirst([&](FSequencerDisplayNode& InNode){
+			Height += InNode.GetNodeHeight() + InNode.GetNodePadding().Combined();
+			return true;
+		}, bIncludeThisNode);
 	}
 
-	return NegativeOffset + Top.GetValue();
+	return FVector2D(100.f, Height);
+}
+
+float SSequencerTrackLane::GetPhysicalPosition() const
+{
+	return TreeView.Pin()->ComputeNodePosition(DisplayNode.ToSharedRef()).Get(0.f);
 }
diff --git a/Engine/Source/Editor/Sequencer/Private/SSequencerTrackLane.h b/Engine/Source/Editor/Sequencer/Private/SSequencerTrackLane.h
--- a/Engine/Source/Editor/Sequencer/Private/SSequencerTrackLane.h
+++ b/Engine/Source/Editor/Sequencer/Private/SSequencerTrackLane.h
@@ -23,7 +23,10 @@ public:
 
 	/** Get the desired physical vertical position of this track lane */
 	float GetPhysicalPosition() const;
-	
+
+protected:
+	virtual FVector2D ComputeDesiredSize(float LayoutScale) const override;
+
 private:
 	/** The authoritative display node that created us */
 	TSharedPtr<FSequencerDisplayNode> DisplayNode;
diff --git a/Engine/Source/Editor/Sequencer/Private/SSequencerTreeView.cpp b/Engine/Source/Editor/Sequencer/Private/SSequencerTreeView.cpp
--- a/Engine/Source/Editor/Sequencer/Private/SSequencerTreeView.cpp
+++ b/Engine/Source/Editor/Sequencer/Private/SSequencerTreeView.cpp
@@ -122,6 +122,65 @@ void SSequencerTreeView::Tick(const FGeometry& AllottedGeometry, const double In
 	PhysicalNodes.Sort([](const FCachedGeometry& A, const FCachedGeometry& B){
 		return A.PhysicalTop < B.PhysicalTop;
 	});
+
+	HighlightRegion = TOptional<FHighlightRegion>();
+
+	if (SequencerNodeTree->GetHoveredNode().IsValid())
+	{
+		TSharedRef<FSequencerDisplayNode> OutermostParent = SequencerNodeTree->GetHoveredNode()->GetOutermostParent();
+		if (OutermostParent->GetType() == ESequencerNode::Spacer)
+		{
+			return;
+		}
+
+		TOptional<float> PhysicalTop = ComputeNodePosition(OutermostParent);
+
+		if (PhysicalTop.IsSet())
+		{
+			// Compute total height of the highlight
+			float TotalHeight = 0.f;
+			OutermostParent->TraverseVisible_ParentFirst([&](FSequencerDisplayNode& InNode){
+				TotalHeight += InNode.GetNodeHeight() + InNode.GetNodePadding().Combined();
+				return true;
+			});
+
+			HighlightRegion = FHighlightRegion(PhysicalTop.GetValue(), PhysicalTop.GetValue() + TotalHeight);
+		}
+	}
+}
+
+int32 SSequencerTreeView::OnPaint(const FPaintArgs& Args, const FGeometry& AllottedGeometry, const FSlateRect& MyClippingRect, FSlateWindowElementList& OutDrawElements, int32 LayerId, const FWidgetStyle& InWidgetStyle, bool bParentEnabled) const
+{
+	LayerId = STreeView::OnPaint(Args, AllottedGeometry, MyClippingRect, OutDrawElements, LayerId, InWidgetStyle, bParentEnabled);
+
+	if (HighlightRegion.IsSet())
+	{
+		static const FName SelectionColorPressedName("SelectionColor");
+		FLinearColor SelectionColor = FEditorStyle::GetSlateColor(SelectionColorPressedName).GetColor(InWidgetStyle).CopyWithNewOpacity(0.3f);
+
+		// Black tint for unhighlighted regions
+		FSlateDrawElement::MakeBox(
+			OutDrawElements,
+			LayerId+1,
+			AllottedGeometry.ToPaintGeometry(FVector2D(2.f, HighlightRegion->Top - 4.f), FVector2D(AllottedGeometry.Size.X - 4.f, 4.f)),
+			FEditorStyle::GetBrush("Sequencer.TrackHoverHighlight_Top"),
+			MyClippingRect,
+			ESlateDrawEffect::None,
+			SelectionColor
+		);
+		
+		FSlateDrawElement::MakeBox(
+			OutDrawElements,
+			LayerId+1,
+			AllottedGeometry.ToPaintGeometry(FVector2D(2.f, HighlightRegion->Bottom), FVector2D(AllottedGeometry.Size.X - 4.f, 4.f)),
+			FEditorStyle::GetBrush("Sequencer.TrackHoverHighlight_Bottom"),
+			MyClippingRect,
+			ESlateDrawEffect::None,
+			SelectionColor
+		);
+	}
+
+	return LayerId + 1;
 }
 
 TOptional<SSequencerTreeView::FCachedGeometry> SSequencerTreeView::GetPhysicalGeometryForNode(const FDisplayNodeRef& InNode) const
@@ -134,6 +193,38 @@ TOptional<SSequencerTreeView::FCachedGeometry> SSequencerTreeView::GetPhysicalGe
 	return TOptional<FCachedGeometry>();
 }
 
+TOptional<float> SSequencerTreeView::ComputeNodePosition(const FDisplayNodeRef& InNode) const
+{
+	// Positioning strategy:
+	// Attempt to root out any visible node in the specified node's sub-hierarchy, and compute the node's offset from that
+	float NegativeOffset = 0.f;
+	TOptional<float> Top;
+	
+	// Iterate parent first until we find a tree view row we can use for the offset height
+	auto Iter = [&](FSequencerDisplayNode& InDisplayNode){
+		
+		TOptional<FCachedGeometry> ChildRowGeometry = GetPhysicalGeometryForNode(InDisplayNode.AsShared());
+		if (ChildRowGeometry.IsSet())
+		{
+			Top = ChildRowGeometry->PhysicalTop;
+			// Stop iterating
+			return false;
+		}
+
+		NegativeOffset -= InDisplayNode.GetNodeHeight() + InDisplayNode.GetNodePadding().Combined();
+		return true;
+	};
+
+	InNode->TraverseVisible_ParentFirst(Iter);
+
+	if (Top.IsSet())
+	{
+		return NegativeOffset + Top.GetValue();
+	}
+
+	return Top;
+}
+
 void SSequencerTreeView::ReportChildRowGeometry(const FDisplayNodeRef& InNode, const FGeometry& InGeometry)
 {
 	float ChildOffset = TransformPoint(
diff --git a/Engine/Source/Editor/Sequencer/Private/SSequencerTreeView.h b/Engine/Source/Editor/Sequencer/Private/SSequencerTreeView.h
--- a/Engine/Source/Editor/Sequencer/Private/SSequencerTreeView.h
+++ b/Engine/Source/Editor/Sequencer/Private/SSequencerTreeView.h
@@ -17,6 +17,13 @@ enum class ETreeRecursion
 	Recursive, NonRecursive
 };
 
+/** Structure to represent the top/bottom bounds of a highlight region */
+struct FHighlightRegion
+{
+	FHighlightRegion(float InTop, float InBottom) : Top(InTop), Bottom(InBottom) {}
+	float Top, Bottom;
+};
+
 /** Structure used to define a column in the tree view */
 struct FSequencerTreeViewColumn
 {
@@ -44,10 +51,14 @@ public:
 	/** Construct this widget */
 	void Construct(const FArguments& InArgs, const TSharedRef<FSequencerNodeTree>& InNodeTree, const TSharedRef<SSequencerTrackArea>& InTrackArea);
 	virtual void Tick( const FGeometry& AllottedGeometry, const double InCurrentTime, const float InDeltaTime );
+	virtual int32 OnPaint(const FPaintArgs& Args, const FGeometry& AllottedGeometry, const FSlateRect& MyClippingRect, FSlateWindowElementList& OutDrawElements, int32 LayerId, const FWidgetStyle& InWidgetStyle, bool bParentEnabled) const override;
 
 	/** Access the underlying tree data */
 	TSharedPtr<FSequencerNodeTree> GetNodeTree() { return SequencerNodeTree; }
 
+	/** @return an optional region specifying the vertical bounds in which a highlight should be drawn */
+	const TOptional<FHighlightRegion>& GetHighlightRegion() const { return HighlightRegion; }
+
 public:
 
 	/** Get the display node at the specified physical vertical position */
@@ -123,6 +134,9 @@ public:
 	/** Retrieve the last reported physical geometry for the specified node, if available */
 	TOptional<FCachedGeometry> GetPhysicalGeometryForNode(const FDisplayNodeRef& InNode) const;
 
+	/** Attempt to compute the physical position of the specified node */
+	TOptional<float> ComputeNodePosition(const FDisplayNodeRef& InNode) const;
+
 	/** Report geometry for a child row */
 	void ReportChildRowGeometry(const FDisplayNodeRef& InNode, const FGeometry& InGeometry);
 
@@ -161,6 +175,9 @@ private:
 
 	/** Strong pointer to the track area so we can generate track lanes as we need them */
 	TSharedPtr<SSequencerTrackArea> TrackArea;
+
+	/** A global highlight for the currently hovered tree node hierarchy */
+	TOptional<FHighlightRegion> HighlightRegion;
 };
 
 /** Widget that represents a row in the sequencer's tree control. */
diff --git a/Engine/Source/Editor/Sequencer/Private/Sequencer.cpp b/Engine/Source/Editor/Sequencer/Private/Sequencer.cpp
--- a/Engine/Source/Editor/Sequencer/Private/Sequencer.cpp
+++ b/Engine/Source/Editor/Sequencer/Private/Sequencer.cpp
@@ -20,16 +20,13 @@
 #include "Kismet2/BlueprintEditorUtils.h"
 #include "BlueprintEditorModule.h"
 #include "MovieSceneTrack.h"
-#include "MovieSceneShotTrack.h"
+#include "MovieSceneCameraCutTrack.h"
 #include "MovieSceneAudioTrack.h"
 #include "MovieSceneSkeletalAnimationTrack.h"
 #include "MovieSceneTrackEditor.h"
 #include "MovieSceneToolHelpers.h"
-#include "Kismet2/KismetEditorUtilities.h"
-#include "Kismet2/BlueprintEditorUtils.h"
-#include "AssetSelection.h"
 #include "ScopedTransaction.h"
-#include "MovieSceneShotSection.h"
+#include "MovieSceneCameraCutSection.h"
 #include "MovieScene3DTransformSection.h"
 #include "MovieSceneSubTrack.h"
 #include "ISequencerSection.h"
@@ -55,6 +52,8 @@
 #include "MovieSceneCommonHelpers.h"
 #include "SceneOutlinerModule.h"
 #include "SceneOutlinerPublicTypes.h"
+#include "ContentBrowserModule.h"
+#include "IContentBrowserSingleton.h"
 
 
 #define LOCTEXT_NAMESPACE "Sequencer"
@@ -117,7 +116,7 @@ void FSequencer::InitSequencer(const FSequencerInitParams& InitParams, const TSh
 		
 	UMovieSceneSequence& RootSequence = *InitParams.RootSequence;
 		
-	// Focusing the initial movie scene needs to be done before the first time GetFocusedMovieSceneInstane or GetRootMovieSceneInstance is used
+	// Focusing the initial movie scene needs to be done before the first time GetFocusedMovieSceneSequenceInstance or GetRootMovieSceneInstance is used
 	RootMovieSceneSequenceInstance = MakeShareable(new FMovieSceneSequenceInstance(RootSequence));
 	SequenceInstanceStack.Add(RootMovieSceneSequenceInstance.ToSharedRef());
 
@@ -175,6 +174,12 @@ void FSequencer::InitSequencer(const FSequencerInitParams& InitParams, const TSh
 
 	BindSequencerCommands();
 	ActivateSequencerEditorMode();
+	AttachTransportControlsToViewports();
+
+	for (auto TrackEditor : TrackEditors)
+	{
+		TrackEditor->OnInitialize();
+	}
 }
 
 
@@ -184,6 +189,8 @@ FSequencer::FSequencer()
 	, LastViewRange(0.f, 5.f)
 	, ViewRangeBeforeZoom(TRange<float>::Empty())
 	, PlaybackState( EMovieScenePlayerStatus::Stopped )
+	, PlaybackStartRealTime ( 0.0 )
+	, NumContinuousFixedTimeStepFrames ( 0 )
 	, ScrubPosition( 0.0f )
 	, bPerspectiveViewportPossessionEnabled( true )
 	, bPerspectiveViewportCameraCutEnabled( false )
@@ -199,6 +206,12 @@ FSequencer::~FSequencer()
 	GEditor->GetActorRecordingState().RemoveAll( this );
 	GEditor->UnregisterForUndo( this );
 
+	DetachTransportControlsFromViewports();
+
+	for (auto TrackEditor : TrackEditors)
+	{
+		TrackEditor->OnRelease();
+	}
 	TrackEditors.Empty();
 	SequencerWidget.Reset();
 }
@@ -237,6 +250,8 @@ void FSequencer::Close()
 		FPropertyEditorModule& EditModule = FModuleManager::Get().GetModuleChecked<FPropertyEditorModule>("PropertyEditor");			
 		EditModule.OnPropertyEditorOpened().RemoveAll(this);
 	}
+	
+	DetachTransportControlsFromViewports();
 
 	SequencerWidget.Reset();
 }
@@ -271,6 +286,8 @@ void FSequencer::Tick(float InDeltaTime)
 
 	float NewTime = GetGlobalTime() + InDeltaTime * GWorld->GetWorldSettings()->MatineeTimeDilation;
 
+	ConstrainFixedTimeStepFrameRate();
+
 	if (PlaybackState == EMovieScenePlayerStatus::Playing ||
 		PlaybackState == EMovieScenePlayerStatus::Recording)
 	{
@@ -284,6 +301,63 @@ void FSequencer::Tick(float InDeltaTime)
 	}
 }
 
+void FSequencer::ConstrainFixedTimeStepFrameRate()
+{
+	if( SequencerSnapValues::IsTimeSnapIntervalFrameRate(Settings->GetTimeSnapInterval()) && Settings->GetFixedTimeStepPlayback() )
+	{
+		// NOTE: Its important that PlaybackStartRealTime and NumContinuousFixedTimeStepFrames are reset
+		//    when anything timing-related changes, like FApp::GetFixedDeltaTime() or playback direction.
+
+		double CurRealTime = FPlatformTime::Seconds();
+
+		// Minor hack to handle changes to world TimeDilation.  We reset our frame rate gate state
+		// when we detect a change to time dilation.
+		static float LastTimeDilation =  GWorld->GetWorldSettings()->GetEffectiveTimeDilation();
+		if( LastTimeDilation != GWorld->GetWorldSettings()->GetEffectiveTimeDilation() )
+		{
+			// Looks like time dilation has changed!
+			NumContinuousFixedTimeStepFrames = 0;
+			PlaybackStartRealTime = CurRealTime;
+
+			LastTimeDilation = GWorld->GetWorldSettings()->GetEffectiveTimeDilation();
+		}
+
+		// How long should have it taken to get to the current frame?
+		const float PlaybackSpeed = 1.f;
+		const double ExpectedPlaybackTime =
+			NumContinuousFixedTimeStepFrames * FApp::GetFixedDeltaTime() * PlaybackSpeed;
+
+		// How long has it been (in real-time) since we started playback?
+		double RealTimeSincePlaybackStarted = CurRealTime - PlaybackStartRealTime;
+
+		// If we're way ahead of schedule (more than 5 ms), then we'll perform a long sleep
+		float WaitTime = ExpectedPlaybackTime - RealTimeSincePlaybackStarted;
+		if( WaitTime > 5 / 1000.0f )
+		{
+			FPlatformProcess::Sleep( WaitTime - 3 / 1000.0f );
+
+			// Update timing info after our little snooze
+			CurRealTime = FPlatformTime::Seconds();
+			RealTimeSincePlaybackStarted = CurRealTime - PlaybackStartRealTime;
+			WaitTime = ExpectedPlaybackTime - RealTimeSincePlaybackStarted;
+		}
+
+		while( RealTimeSincePlaybackStarted < ExpectedPlaybackTime )
+		{
+			// OK, we're running ahead of schedule so we need to wait a bit before the next frame
+			FPlatformProcess::Sleep( 0.0f );
+
+			// Check the time again
+			CurRealTime = FPlatformTime::Seconds();
+			RealTimeSincePlaybackStarted = CurRealTime - PlaybackStartRealTime;
+			WaitTime = ExpectedPlaybackTime - RealTimeSincePlaybackStarted;
+		}
+
+		// Increment number of continuous fixed time step frames
+		++NumContinuousFixedTimeStepFrames;
+	}
+}
+
 
 UMovieSceneSequence* FSequencer::GetRootMovieSceneSequence() const
 {
@@ -486,6 +560,12 @@ TSharedRef<FMovieSceneSequenceInstance> FSequencer::GetSequenceInstanceForSectio
 }
 
 
+bool FSequencer::HasSequenceInstanceForSection(UMovieSceneSection& Section) const
+{
+	return SequenceInstanceBySection.Contains(&Section);
+}
+
+
 void FSequencer::PopToSequenceInstance(TSharedRef<FMovieSceneSequenceInstance> SequenceInstance)
 {
 	if( SequenceInstanceStack.Num() > 1 )
@@ -758,6 +838,33 @@ void FSequencer::SetEndPlaybackRange()
 	SetPlaybackRange(TRange<float>( GetPlaybackRange().GetLowerBoundValue(), GetGlobalTime()));
 }
 
+void FSequencer::ZoomViewRange(float InZoomDelta)
+{
+	float LocalViewRangeMax = TargetViewRange.GetUpperBoundValue();
+	float LocalViewRangeMin = TargetViewRange.GetLowerBoundValue();
+	const float OutputViewSize = LocalViewRangeMax - LocalViewRangeMin;
+	const float OutputChange = OutputViewSize * InZoomDelta;
+
+	float CurrentPositionFraction = (ScrubPosition - LocalViewRangeMin) / OutputViewSize;
+
+	float NewViewOutputMin = LocalViewRangeMin - (OutputChange * CurrentPositionFraction);
+	float NewViewOutputMax = LocalViewRangeMax + (OutputChange * (1.f - CurrentPositionFraction));
+
+	if( NewViewOutputMin < NewViewOutputMax )
+	{
+		SetViewRange(TRange<float>(NewViewOutputMin, NewViewOutputMax), EViewRangeInterpolation::Immediate);
+	}
+}
+
+void FSequencer::ZoomInViewRange()
+{
+	ZoomViewRange(-0.1f);
+}
+
+void FSequencer::ZoomOutViewRange()
+{
+	ZoomViewRange(0.1f);
+}
 
 void FSequencer::UpdatePlaybackRange()
 {
@@ -779,6 +886,25 @@ void FSequencer::UpdatePlaybackRange()
 			FocusedMovieScene->SetPlaybackRange(NewBounds.GetLowerBoundValue(), NewBounds.GetUpperBoundValue());
 		}
 	}
+	else
+	{
+		UMovieScene* FocusedMovieScene = GetFocusedMovieSceneSequence()->GetMovieScene();
+		TRange<float> NewBounds = TRange<float>::Empty();
+
+		for (auto MasterTrack : FocusedMovieScene->GetMasterTracks())
+		{
+			if (MasterTrack->AddsSectionBoundsToPlayRange())
+			{
+				NewBounds = TRange<float>( FMath::Min(NewBounds.GetLowerBoundValue(), MasterTrack->GetSectionBoundaries().GetLowerBoundValue()),
+											FMath::Max(NewBounds.GetUpperBoundValue(), MasterTrack->GetSectionBoundaries().GetUpperBoundValue()) );
+			}
+		}
+			
+		if (!NewBounds.IsEmpty())
+		{
+			FocusedMovieScene->SetPlaybackRange(NewBounds.GetLowerBoundValue(), NewBounds.GetUpperBoundValue());
+		}
+	}
 }
 
 
@@ -850,7 +976,7 @@ float FSequencer::GetGlobalTime()
 }
 
 
-void FSequencer::SetGlobalTime( float NewTime )
+void FSequencer::SetGlobalTime( float NewTime, const bool& bAllowSnappingToFrames )
 {
 	if (IsAutoScrollEnabled())
 	{
@@ -868,18 +994,24 @@ void FSequencer::SetGlobalTime( float NewTime )
 		}
 	}
 
-	SetGlobalTimeDirectly(NewTime);
+	SetGlobalTimeDirectly(NewTime, bAllowSnappingToFrames);
 }
 
 
-void FSequencer::SetGlobalTimeDirectly( float NewTime )
+void FSequencer::SetGlobalTimeDirectly( float NewTime, const bool& bAllowSnappingToFrames  )
 {
 	float LastTime = ScrubPosition;
 
+	if (bAllowSnappingToFrames && Settings->GetIsSnapEnabled())
+	{
+		NewTime = Settings->SnapTimeToInterval(NewTime);
+	}
+
 	// Update the position
 	ScrubPosition = NewTime;
 
-	SequenceInstanceStack.Top()->Update(NewTime, LastTime, *this);
+	EMovieSceneUpdateData UpdateData(NewTime, LastTime);
+	SequenceInstanceStack.Top()->Update(UpdateData, *this);
 
 	// Ensure any relevant spawned objects are cleaned up if we're playing back the master sequence
 	if (SequenceInstanceStack.Num() == 1)
@@ -977,6 +1109,26 @@ void FSequencer::SetPerspectiveViewportCameraCutEnabled(bool bEnabled)
 	bPerspectiveViewportCameraCutEnabled = bEnabled;
 }
 
+ISequencer::FOnCameraCut& FSequencer::OnCameraCut()
+{
+	return OnCameraCutEvent;
+}
+
+TSharedRef<INumericTypeInterface<float>> FSequencer::GetNumericTypeInterface()
+{
+	return SequencerWidget->GetNumericTypeInterface();
+}
+
+TSharedRef<INumericTypeInterface<float>> FSequencer::GetZeroPadNumericTypeInterface()
+{
+	return SequencerWidget->GetZeroPadNumericTypeInterface();
+}
+
+TSharedRef<SWidget> FSequencer::MakeTransportControls()
+{
+	return SequencerWidget->MakeTransportControls();
+}
+
 /** Attempt to find an object binding ID that relates to an unspawned spawnable object */
 FGuid FindUnspawnedObjectGuid(UObject& InObject, UMovieSceneSequence& Sequence)
 {
@@ -1098,6 +1250,8 @@ void FSequencer::GetRuntimeObjects( TSharedRef<FMovieSceneSequenceInstance> Movi
 
 void FSequencer::UpdateCameraCut(UObject* CameraObject, UObject* UnlockIfCameraObject) const
 {
+	OnCameraCutEvent.Broadcast(CameraObject);
+
 	if (!IsPerspectiveViewportCameraCutEnabled())
 	{
 		return;
@@ -1115,7 +1269,7 @@ void FSequencer::UpdateCameraCut(UObject* CameraObject, UObject* UnlockIfCameraO
 		if ((CameraObject != nullptr) || LevelVC->IsLockedToActor(UnlockIfCameraActor))
 		{
 			UpdatePreviewLevelViewportClientFromCameraCut(*LevelVC, CameraObject);
-		}			
+		}
 	}
 }
 
@@ -1223,7 +1377,9 @@ void FSequencer::UpdateRuntimeInstances()
 {
 	// Refresh the current root instance
 	SequenceInstanceStack.Top()->RefreshInstance( *this );
-	SequenceInstanceStack.Top()->Update(ScrubPosition, ScrubPosition, *this);
+
+	EMovieSceneUpdateData UpdateData(ScrubPosition, ScrubPosition);
+	SequenceInstanceStack.Top()->Update(UpdateData, *this);
 
 	// Ensure any relevant spawned objects are cleaned up if we're playing back the master sequence
 	if (SequenceInstanceStack.Num() == 1)
@@ -1239,6 +1395,89 @@ void FSequencer::UpdateRuntimeInstances()
 	GEditor->RedrawLevelEditingViewports();
 }
 
+void FSequencer::AttachTransportControlsToViewports()
+{
+	FLevelEditorModule* Module = FModuleManager::Get().LoadModulePtr<FLevelEditorModule>("LevelEditor");
+	if (Module)
+	{
+		TSharedPtr<ILevelEditor> LevelEditor = Module->GetFirstLevelEditor();
+		const TArray< TSharedPtr<ILevelViewport> >& LevelViewports = LevelEditor->GetViewports();
+		
+		FEditorWidgetsModule& EditorWidgetsModule = FModuleManager::Get().LoadModuleChecked<FEditorWidgetsModule>( "EditorWidgets" );
+
+		FTransportControlArgs TransportControlArgs;
+		TransportControlArgs.OnForwardPlay.BindSP(this, &FSequencer::OnPlay, true);
+		TransportControlArgs.OnForwardStep.BindSP(this, &FSequencer::OnStepForward);
+		TransportControlArgs.OnBackwardStep.BindSP(this, &FSequencer::OnStepBackward);
+		TransportControlArgs.OnForwardEnd.BindSP(this, &FSequencer::OnStepToEnd);
+		TransportControlArgs.OnBackwardEnd.BindSP(this, &FSequencer::OnStepToBeginning);
+		TransportControlArgs.OnToggleLooping.BindSP(this, &FSequencer::OnToggleLooping);
+		TransportControlArgs.OnGetLooping.BindSP(this, &FSequencer::IsLooping);
+		TransportControlArgs.OnGetPlaybackMode.BindSP(this, &FSequencer::GetPlaybackMode);
+
+		for (int32 i = 0; i < LevelViewports.Num(); ++i)
+		{
+			const TSharedPtr<ILevelViewport>& LevelViewport = LevelViewports[i];
+			
+			TSharedPtr<SWidget> TransportControl =
+				SNew(SHorizontalBox)
+				.Visibility(EVisibility::SelfHitTestInvisible)
+				+SHorizontalBox::Slot()
+				.FillWidth(1)
+				.HAlign(HAlign_Center)
+				.VAlign(VAlign_Bottom)
+				.Padding(4.f)
+				[
+					SNew(SBorder)
+					.Padding(4.f)
+					.Cursor( EMouseCursor::Default )
+					.BorderImage( FEditorStyle::GetBrush( "FilledBorder" ) )
+					.Visibility(this, &FSequencer::GetTransportControlVisibility, LevelViewport)
+					.Content()
+					[
+						EditorWidgetsModule.CreateTransportControl(TransportControlArgs)
+					]
+				];
+
+			LevelViewport->AddOverlayWidget(TransportControl.ToSharedRef());
+
+			TransportControls.Add(LevelViewport, TransportControl);
+		}
+	}
+}
+
+
+void FSequencer::DetachTransportControlsFromViewports()
+{
+	FLevelEditorModule* Module = FModuleManager::Get().LoadModulePtr<FLevelEditorModule>("LevelEditor");
+	if (Module)
+	{
+		TSharedPtr<ILevelEditor> LevelEditor = Module->GetFirstLevelEditor();
+		if (LevelEditor.IsValid())
+		{
+			const TArray< TSharedPtr<ILevelViewport> >& LevelViewports = LevelEditor->GetViewports();
+		
+			for (int32 i = 0; i < LevelViewports.Num(); ++i)
+			{
+				const TSharedPtr<ILevelViewport>& LevelViewport = LevelViewports[i];
+
+				TSharedPtr<SWidget>* TransportControl = TransportControls.Find(LevelViewport);
+				if (TransportControl && TransportControl->IsValid())
+				{
+					LevelViewport->RemoveOverlayWidget(TransportControl->ToSharedRef());
+				}
+			}
+		}
+	}
+}
+
+
+EVisibility FSequencer::GetTransportControlVisibility(TSharedPtr<ILevelViewport> LevelViewport) const
+{
+	FLevelEditorViewportClient& ViewportClient = LevelViewport->GetLevelViewportClient();
+	return ViewportClient.ViewportType == LVT_Perspective &&
+		ViewportClient.AllowsCinematicPreview() ? EVisibility::Visible : EVisibility::Collapsed;
+}
 
 FReply FSequencer::OnPlay(bool bTogglePlay)
 {
@@ -1246,12 +1485,22 @@ FReply FSequencer::OnPlay(bool bTogglePlay)
 		 PlaybackState == EMovieScenePlayerStatus::Recording) && bTogglePlay )
 	{
 		PlaybackState = EMovieScenePlayerStatus::Stopped;
+
+		// Snap time when stopping play
+		const bool bAllowSnappingToFrames = true;
+		SetGlobalTime(ScrubPosition, bAllowSnappingToFrames);
+
 		// Update on stop (cleans up things like sounds that are playing)
-		SequenceInstanceStack.Top()->Update( ScrubPosition, ScrubPosition, *this );
+		EMovieSceneUpdateData UpdateData(ScrubPosition, ScrubPosition);
+		SequenceInstanceStack.Top()->Update(UpdateData, *this );
 	}
 	else
 	{
 		PlaybackState = EMovieScenePlayerStatus::Playing;
+
+		// Reset our playback start time so fixed time step playback can gate frame rate properly
+		PlaybackStartRealTime = FPlatformTime::Seconds();
+		NumContinuousFixedTimeStepFrames = 0;
 		
 		// Make sure Slate ticks during playback
 		SequencerWidget->RegisterActiveTimerForPlayback();
@@ -1287,11 +1536,8 @@ FReply FSequencer::OnStepForward()
 {
 	PlaybackState = EMovieScenePlayerStatus::Stopped;
 	float NewPosition = ScrubPosition + Settings->GetTimeSnapInterval();
-	if (Settings->GetIsSnapEnabled())
-	{
-		NewPosition = Settings->SnapTimeToInterval(NewPosition);
-	}
-	SetGlobalTime(NewPosition);
+	const bool bAllowSnappingToFrames = true;
+	SetGlobalTime(NewPosition, bAllowSnappingToFrames);
 	return FReply::Handled();
 }
 
@@ -1300,11 +1546,8 @@ FReply FSequencer::OnStepBackward()
 {
 	PlaybackState = EMovieScenePlayerStatus::Stopped;
 	float NewPosition = ScrubPosition - Settings->GetTimeSnapInterval();
-	if (Settings->GetIsSnapEnabled())
-	{
-		NewPosition = Settings->SnapTimeToInterval(NewPosition);
-	}
-	SetGlobalTime(NewPosition);
+	const bool bAllowSnappingToFrames = true;
+	SetGlobalTime(NewPosition, bAllowSnappingToFrames);
 	return FReply::Handled();
 }
 
@@ -1617,8 +1860,7 @@ void FSequencer::VerticalScroll(float ScrollAmountUnits)
 	SequencerWidget->GetTreeView()->ScrollByDelta(ScrollAmountUnits);
 }
 
-
-FGuid FSequencer::AddSpawnableForAssetOrClass( UObject* Object )
+FGuid FSequencer::AddSpawnable(UObject& Object)
 {
 	UMovieSceneSequence* Sequence = GetFocusedMovieSceneSequence();
 	if (!Sequence->AllowsSpawnableObjects())
@@ -1630,144 +1872,133 @@ FGuid FSequencer::AddSpawnableForAssetOrClass( UObject* Object )
 	// MovieScene asset.
 	UMovieScene* OwnerMovieScene = Sequence->GetMovieScene();
 
-	// @todo sequencer: Undo doesn't seem to be working at all
-	const FScopedTransaction Transaction( LOCTEXT("UndoAddingObject", "Add Object to MovieScene") );
-
-	// Use the class as the spawnable's name if this is an actor class, otherwise just use the object name (asset)
-	const bool bIsActorClass = Object->IsA( AActor::StaticClass() ) && !Object->HasAnyFlags( RF_ArchetypeObject );
-	const FName AssetName = bIsActorClass ? Object->GetClass()->GetFName() : Object->GetFName();
+	TValueOrError<FNewSpawnable, FText> Result = SpawnRegister->CreateNewSpawnableType(Object, *OwnerMovieScene);
+	if (!Result.IsValid())
+	{
+		FNotificationInfo Info(Result.GetError());
+		Info.ExpireDuration = 3.0f;
+		FSlateNotificationManager::Get().AddNotification(Info);
+		return FGuid();
+	}
 
-	// Inner objects don't need a name (it will be auto-generated by the UObject system), but we want one in this case
-	// because the class of any actors that are created from this Blueprint may end up being user-facing.
-	const FName BlueprintName = MakeUniqueObjectName( OwnerMovieScene, UBlueprint::StaticClass(), AssetName );
+	FNewSpawnable& NewSpawnable = Result.GetValue();
 
-	// Use the asset name as the initial spawnable name
-	const FString NewSpawnableName = AssetName.ToString();		// @todo sequencer: Need UI to allow user to rename these slots
+	auto DuplName = [&](FMovieSceneSpawnable& InSpawnable)
+	{
+		return InSpawnable.GetName() == NewSpawnable.Name;
+	};
 
-	FText ErrorText;
+	int32 Index = 2;
+	FString UniqueString;
+	while (OwnerMovieScene->FindSpawnable(DuplName))
+	{
+		NewSpawnable.Name.RemoveFromEnd(UniqueString);
+		UniqueString = FString::Printf(TEXT(" (%d)"), Index);
+		NewSpawnable.Name += UniqueString;
+	}
 
-	// Create our new blueprint!
-	UBlueprint* NewBlueprint = nullptr;
+	FGuid NewGuid = OwnerMovieScene->AddSpawnable(NewSpawnable.Name, NewSpawnable.Blueprint);
 
-	// If we're dropping a blueprint, create a new blueprint out of this one, so we can propagate changes from the world to it without affecting the dragged BP.
-	if (UBlueprint* SourceBlueprint = Cast<UBlueprint>(Object))
-	{
-		UClass* SourceParentClass = SourceBlueprint->GeneratedClass;
+	UpdateRuntimeInstances();
 
-		if (!FKismetEditorUtilities::CanCreateBlueprintOfClass(SourceParentClass))
-		{
-			ErrorText = FText::Format(LOCTEXT("UnableToAddSpawnableBlueprint", "Unable to add spawnable for class of type '{0}' since it is not a valid blueprint parent class."),
-				FText::FromString(SourceParentClass->GetName())
-				);
-		}
-		else
-		{
-			NewBlueprint = FKismetEditorUtilities::CreateBlueprint(SourceParentClass, OwnerMovieScene, BlueprintName, EBlueprintType::BPTYPE_Normal, UBlueprint::StaticClass(), UBlueprintGeneratedClass::StaticClass());
-		}
-	}
-	else
-	{
-		// @todo sequencer: Add support for forcing specific factories for an asset
-		UActorFactory* FactoryToUse = nullptr;
-		if( bIsActorClass )
-		{
-			// Placing an actor class directly::
-			FactoryToUse = GEditor->FindActorFactoryForActorClass( Object->GetClass() );
-		}
-		else
-		{
-			// Placing an asset
-			FactoryToUse = FActorFactoryAssetProxy::GetFactoryForAssetObject( Object );
-		}
+	return NewGuid;
+}
 
-		if( FactoryToUse != nullptr )
-		{
-			// Create the blueprint
-			NewBlueprint = FactoryToUse->CreateBlueprint( Object, OwnerMovieScene, BlueprintName );
-		}
-		else if( bIsActorClass )
-		{
-			// We don't have a factory, but we can still try to create a blueprint for this actor class
-			NewBlueprint = FKismetEditorUtilities::CreateBlueprint( Object->GetClass(), OwnerMovieScene, BlueprintName, EBlueprintType::BPTYPE_Normal, UBlueprint::StaticClass(), UBlueprintGeneratedClass::StaticClass() );
-		}
-	}
+FGuid FSequencer::MakeNewSpawnable( UObject& Object )
+{
+	// @todo sequencer: Undo doesn't seem to be working at all
+	const FScopedTransaction Transaction( LOCTEXT("UndoAddingObject", "Add Object to MovieScene") );
 
-	if (!NewBlueprint)
+	FGuid NewGuid = AddSpawnable(Object);
+	if (!NewGuid.IsValid())
 	{
-		if (ErrorText.IsEmpty())
-		{
-			ErrorText = FText::Format(LOCTEXT("UnableToAddSpawnable", "Unable to add spawnable for '{0}'."),
-				FText::FromString(Object->GetName())
-				);
-		}
-
-		FNotificationInfo Info(ErrorText);
-		Info.ExpireDuration = 3.0f;
-		FSlateNotificationManager::Get().AddNotification(Info);
 		return FGuid();
 	}
 
-	FTransformData NewTransform;
+	// Spawn the object so we can position it correctly, it's going to get spawned anyway since things default to spawned.
+	UObject* SpawnedObject = SpawnRegister->SpawnObject(NewGuid, *GetFocusedMovieSceneSequenceInstance(), *this);
+
+	FTransformData DefaultTransform;
 
-	if ((NewBlueprint->GeneratedClass != nullptr) && FBlueprintEditorUtils::IsActorBased(NewBlueprint))
+	AActor* SpawnedActor = Cast<AActor>(SpawnedObject);
+	if (SpawnedActor)
 	{
-		AActor* ActorCDO = CastChecked< AActor >( NewBlueprint->GeneratedClass->ClassDefaultObject );
 		// Place the new spawnable in front of the camera (unless we were automatically created from a PIE actor)
 		if (Settings->GetSpawnPosition() == SSP_PlaceInFrontOfCamera)
 		{
-			PlaceActorInFrontOfCamera( ActorCDO );
+			PlaceActorInFrontOfCamera( SpawnedActor );
 		}
-		NewTransform.Translation = ActorCDO->GetActorLocation();
-		NewTransform.Rotation = ActorCDO->GetActorRotation();
-		NewTransform.Scale = FVector(1.0f, 1.0f, 1.0f);
-		NewTransform.bValid = true;
+		DefaultTransform.Translation = SpawnedActor->GetActorLocation();
+		DefaultTransform.Rotation = SpawnedActor->GetActorRotation();
+		DefaultTransform.Scale = FVector(1.0f, 1.0f, 1.0f);
+		DefaultTransform.bValid = true;
 	}
 
-	FGuid NewSpawnableGuid = OwnerMovieScene->AddSpawnable(NewSpawnableName, NewBlueprint);
-	if (NewSpawnableGuid.IsValid())
+	SetupDefaultsForSpawnable(NewGuid, DefaultTransform);
+	return NewGuid;
+}
+
+void FSequencer::SetupDefaultsForSpawnable( const FGuid& Guid, const FTransformData& DefaultTransform )
+{
+	UMovieSceneSequence* Sequence = GetFocusedMovieSceneSequence();
+	UMovieScene* OwnerMovieScene = Sequence->GetMovieScene();
+
+	// Ensure it has a spawn track
+	UMovieSceneSpawnTrack* SpawnTrack = Cast<UMovieSceneSpawnTrack>(OwnerMovieScene->FindTrack(UMovieSceneSpawnTrack::StaticClass(), Guid, NAME_None));
+	if (!SpawnTrack)
+	{
+		SpawnTrack = Cast<UMovieSceneSpawnTrack>(OwnerMovieScene->AddTrack(UMovieSceneSpawnTrack::StaticClass(), Guid));
+	}
+
+	if (SpawnTrack)
+	{
+		SpawnTrack->AddSection(*SpawnTrack->CreateNewSection());
+		SpawnTrack->SetObjectId(Guid);
+	}
+	
+	// Ensure it will spawn in the right place
+	if (DefaultTransform.bValid)
 	{
-		UMovieSceneSpawnTrack* SpawnTrack = Cast<UMovieSceneSpawnTrack>(OwnerMovieScene->AddTrack(UMovieSceneSpawnTrack::StaticClass(), NewSpawnableGuid));
-		if (SpawnTrack)
+		UMovieScene3DTransformTrack* TransformTrack = Cast<UMovieScene3DTransformTrack>(OwnerMovieScene->FindTrack(UMovieScene3DTransformTrack::StaticClass(), Guid, "Transform"));
+		if (!TransformTrack)
 		{
-			SpawnTrack->AddSection(*SpawnTrack->CreateNewSection());
-			SpawnTrack->SetObjectId(NewSpawnableGuid);
+			TransformTrack = Cast<UMovieScene3DTransformTrack>(OwnerMovieScene->AddTrack(UMovieScene3DTransformTrack::StaticClass(), Guid));
 		}
-						
-		if (NewTransform.bValid)
-		{
-			UMovieScene3DTransformTrack* TransformTrack = Cast<UMovieScene3DTransformTrack>(OwnerMovieScene->AddTrack(UMovieScene3DTransformTrack::StaticClass(), NewSpawnableGuid));
-			if (TransformTrack)
-			{
-				static FName Transform("Transform");
-				TransformTrack->SetPropertyNameAndPath(Transform, Transform.ToString());
-				const bool bUnwindRotation = false;
 
-				float Time = GetCurrentLocalTime( *Sequence );
-				EMovieSceneKeyInterpolation Interpolation = GetKeyInterpolation();
+		if (TransformTrack)
+		{
+			static FName Transform("Transform");
+			TransformTrack->SetPropertyNameAndPath(Transform, Transform.ToString());
+			const bool bUnwindRotation = false;
 
-				UMovieScene3DTransformSection* NewSection = CastChecked<UMovieScene3DTransformSection>(TransformTrack->CreateNewSection());
+			EMovieSceneKeyInterpolation Interpolation = GetKeyInterpolation();
 
-				NewSection->SetDefault( FTransformKey( EKey3DTransformChannel::Translation, EAxis::X, NewTransform.Translation.X, bUnwindRotation ) );
-				NewSection->SetDefault( FTransformKey( EKey3DTransformChannel::Translation, EAxis::Y, NewTransform.Translation.Y, bUnwindRotation ) );
-				NewSection->SetDefault( FTransformKey( EKey3DTransformChannel::Translation, EAxis::Z, NewTransform.Translation.Z, bUnwindRotation ) );
+			const TArray<UMovieSceneSection*>& Sections = TransformTrack->GetAllSections();
+			if (!Sections.Num())
+			{
+				TransformTrack->AddSection(*TransformTrack->CreateNewSection());
+			}
+			
+			for (UMovieSceneSection* Section : Sections)
+			{
+				UMovieScene3DTransformSection* TransformSection = CastChecked<UMovieScene3DTransformSection>(Section);
 
-				NewSection->SetDefault( FTransformKey( EKey3DTransformChannel::Rotation, EAxis::X, NewTransform.Rotation.Euler().X, bUnwindRotation ) );
-				NewSection->SetDefault( FTransformKey( EKey3DTransformChannel::Rotation, EAxis::Y, NewTransform.Rotation.Euler().Y, bUnwindRotation ) );
-				NewSection->SetDefault( FTransformKey( EKey3DTransformChannel::Rotation, EAxis::Z, NewTransform.Rotation.Euler().Z, bUnwindRotation ) );
+				TransformSection->SetDefault( FTransformKey( EKey3DTransformChannel::Translation, EAxis::X, DefaultTransform.Translation.X, bUnwindRotation ) );
+				TransformSection->SetDefault( FTransformKey( EKey3DTransformChannel::Translation, EAxis::Y, DefaultTransform.Translation.Y, bUnwindRotation ) );
+				TransformSection->SetDefault( FTransformKey( EKey3DTransformChannel::Translation, EAxis::Z, DefaultTransform.Translation.Z, bUnwindRotation ) );
 
-				NewSection->SetDefault( FTransformKey( EKey3DTransformChannel::Scale, EAxis::X, NewTransform.Scale.X, bUnwindRotation ) );
-				NewSection->SetDefault( FTransformKey( EKey3DTransformChannel::Scale, EAxis::Y, NewTransform.Scale.Y, bUnwindRotation ) );
-				NewSection->SetDefault( FTransformKey( EKey3DTransformChannel::Scale, EAxis::Z, NewTransform.Scale.Z, bUnwindRotation ) );
+				TransformSection->SetDefault( FTransformKey( EKey3DTransformChannel::Rotation, EAxis::X, DefaultTransform.Rotation.Euler().X, bUnwindRotation ) );
+				TransformSection->SetDefault( FTransformKey( EKey3DTransformChannel::Rotation, EAxis::Y, DefaultTransform.Rotation.Euler().Y, bUnwindRotation ) );
+				TransformSection->SetDefault( FTransformKey( EKey3DTransformChannel::Rotation, EAxis::Z, DefaultTransform.Rotation.Euler().Z, bUnwindRotation ) );
 
-				TransformTrack->AddSection(*NewSection);
+				TransformSection->SetDefault( FTransformKey( EKey3DTransformChannel::Scale, EAxis::X, DefaultTransform.Scale.X, bUnwindRotation ) );
+				TransformSection->SetDefault( FTransformKey( EKey3DTransformChannel::Scale, EAxis::Y, DefaultTransform.Scale.Y, bUnwindRotation ) );
+				TransformSection->SetDefault( FTransformKey( EKey3DTransformChannel::Scale, EAxis::Z, DefaultTransform.Scale.Z, bUnwindRotation ) );
 			}
 		}
 	}
-
-	return NewSpawnableGuid;
 }
 
-
 void FSequencer::AddSubSequence(UMovieSceneSequence* Sequence)
 {
 	// @todo Sequencer - sub-moviescenes This should be moved to the sub-moviescene editor
@@ -1781,7 +2012,8 @@ void FSequencer::AddSubSequence(UMovieSceneSequence* Sequence)
 	OwnerMovieScene->Modify();
 
 	UMovieSceneSubTrack* SubTrack = OwnerMovieScene->AddMasterTrack<UMovieSceneSubTrack>();
-	SubTrack->AddSequence(*Sequence, ScrubPosition);
+	float Duration = Sequence->GetMovieScene()->GetPlaybackRange().Size<float>();
+	SubTrack->AddSequence(Sequence, ScrubPosition, Duration);
 }
 
 
@@ -1880,9 +2112,9 @@ void FSequencer::OnRequestNodeDeleted( TSharedRef<const FSequencerDisplayNode> N
 			{
 				OwnerMovieScene->RemoveMasterTrack(*Track);
 			}
-			else if (OwnerMovieScene->GetShotTrack() == Track)
+			else if (OwnerMovieScene->GetCameraCutTrack() == Track)
 			{
-				OwnerMovieScene->RemoveShotTrack();
+				OwnerMovieScene->RemoveCameraCutTrack();
 			}
 			else
 			{
@@ -1958,19 +2190,6 @@ void FSequencer::PlaceActorInFrontOfCamera( AActor* ActorCDO )
 }
 
 
-void FSequencer::CopyActorProperties( AActor* PuppetActor, AActor* TargetActor ) const
-{
-	// @todo sequencer: How do we make this undoable with the original action that caused the change in the first place? 
-	//       -> Ideally we are still within the transaction scope while performing this change
-
-	// @todo sequencer: If we decide to propagate changes BACK to an actor after changing the CDO, we may need to
-	// do additional work to make sure transient state is refreshed (e.g. call PostEditMove() is transform is changed)
-
-	// @todo sequencer: Pass option here to call PostEditChange, if needed
-	const int32 CopiedPropertyCount = EditorUtilities::CopyActorProperties( PuppetActor, TargetActor );
-}
-
-
 void FSequencer::GetActorRecordingState( bool& bIsRecording /* In+Out */ ) const
 {
 	if( IsRecordingLive() )
@@ -2133,6 +2352,48 @@ void FSequencer::UpdatePreviewLevelViewportClientFromCameraCut(FLevelEditorViewp
 
 void FSequencer::SaveCurrentMovieScene()
 {
+	// Capture thumbnail
+	// Convert UObject* array to FAssetData array
+	TArray<FAssetData> AssetDataList;
+	AssetDataList.Add(FAssetData(GetCurrentAsset()));
+
+	FViewport* Viewport = GEditor->GetActiveViewport();
+
+	// If there's no active viewport, find any other viewport that allows cinematic preview.
+	if (Viewport == nullptr)
+	{
+		for (FLevelEditorViewportClient* LevelVC : GEditor->LevelViewportClients)
+		{
+			if ((LevelVC == nullptr) || !LevelVC->IsPerspective() || !LevelVC->AllowsCinematicPreview())
+			{
+				continue;
+			}
+
+			Viewport = LevelVC->Viewport;
+		}
+	}
+
+	if ( ensure(GCurrentLevelEditingViewportClient) && Viewport != nullptr )
+	{
+		bool bIsInGameView = GCurrentLevelEditingViewportClient->IsInGameView();
+		GCurrentLevelEditingViewportClient->SetGameView(true);
+
+		//have to re-render the requested viewport
+		FLevelEditorViewportClient* OldViewportClient = GCurrentLevelEditingViewportClient;
+		//remove selection box around client during render
+		GCurrentLevelEditingViewportClient = NULL;
+
+		Viewport->Draw();
+
+		IContentBrowserSingleton& ContentBrowser = FModuleManager::LoadModuleChecked<FContentBrowserModule>("ContentBrowser").Get();
+		ContentBrowser.CaptureThumbnailFromViewport(Viewport, AssetDataList);
+
+		//redraw viewport to have the yellow highlight again
+		GCurrentLevelEditingViewportClient = OldViewportClient;
+		GCurrentLevelEditingViewportClient->SetGameView(bIsInGameView);
+		Viewport->Draw();
+	}
+
 	UPackage* MovieScenePackage = GetCurrentAsset()->GetOutermost();
 
 	TArray<UPackage*> PackagesToSave;
@@ -2543,6 +2804,135 @@ bool FSequencer::IsNodeLocked() const
 	return NumSections > 0;
 }
 
+void FSequencer::ConvertToSpawnable(TSharedRef<FSequencerObjectBindingNode> NodeToBeConverted)
+{
+	const FScopedTransaction Transaction( LOCTEXT("ConvertSelectedSpawnable", "Convert node to spawnables") );
+
+	// Ensure we're in a non-possessed state
+	RootMovieSceneSequenceInstance->RestoreState(*this);
+	GetFocusedMovieSceneSequence()->GetMovieScene()->Modify();
+
+	ConvertToSpawnableInternal(NodeToBeConverted);
+
+	NotifyMovieSceneDataChanged();
+}
+
+void FSequencer::ConvertSelectedNodesToSpawnables()
+{
+	// @todo sequencer: Undo doesn't seem to be working at all
+	const FScopedTransaction Transaction( LOCTEXT("ConvertSelectedSpawnable", "Convert selected nodes to spawnables") );
+
+	UMovieScene* MovieScene = GetFocusedMovieSceneSequence()->GetMovieScene();
+
+	// Ensure we're in a non-possessed state
+	RootMovieSceneSequenceInstance->RestoreState(*this);
+	MovieScene->Modify();
+
+	TArray<TSharedRef<FSequencerObjectBindingNode>> ObjectBindingNodes;
+
+	for (const TSharedRef<FSequencerDisplayNode>& Node : Selection.GetSelectedOutlinerNodes())
+	{
+		if (Node->GetType() == ESequencerNode::Object)
+		{
+			auto ObjectBindingNode = StaticCastSharedRef<FSequencerObjectBindingNode>(Node);
+
+			// If we have a possessable for this node, and it has no parent, we can convert it to a spawnable
+			FMovieScenePossessable* Possessable = MovieScene->FindPossessable(ObjectBindingNode->GetObjectBinding());
+			if (Possessable && !Possessable->GetParent().IsValid())
+			{
+				ObjectBindingNodes.Add(ObjectBindingNode);
+			}
+		}
+	}
+
+	FScopedSlowTask SlowTask(ObjectBindingNodes.Num(), LOCTEXT("ConvertSpawnableProgress", "Converting selected possessable nodes to spawnables"));
+	SlowTask.MakeDialog(true);
+
+	for (const TSharedRef<FSequencerObjectBindingNode>& ObjectBindingNode : ObjectBindingNodes)
+	{
+		SlowTask.EnterProgressFrame();
+		ConvertToSpawnableInternal(ObjectBindingNode);
+
+		if (GWarn->ReceivedUserCancel())
+		{
+			break;
+		}
+	}
+
+	NotifyMovieSceneDataChanged();
+}
+
+void FSequencer::ConvertToSpawnableInternal(TSharedRef<FSequencerObjectBindingNode> NodeToBeConverted)
+{
+	UMovieSceneSequence* Sequence = GetFocusedMovieSceneSequence();
+	UMovieScene* MovieScene = Sequence->GetMovieScene();
+
+	FMovieScenePossessable* Possessable = MovieScene->FindPossessable(NodeToBeConverted->GetObjectBinding());
+
+	if (!Possessable)
+	{
+		return;
+	}
+
+	// Find the object in the environment
+	UObject* FoundObject = GetFocusedMovieSceneSequenceInstance()->FindObject(NodeToBeConverted->GetObjectBinding(), *this);
+	if (!FoundObject)
+	{
+		return;
+	}
+
+	Sequence->Modify();
+
+	FMovieSceneSpawnable* Spawnable = MovieScene->FindSpawnable(AddSpawnable(*FoundObject));
+	if (!Spawnable)
+	{
+		return;
+	}
+
+	// Swap the guids, so the possessable's tracks now belong to the spawnable
+	{
+		FGuid BenignSpawnableGuid = Spawnable->GetGuid();
+		FGuid PersistentGuid = Possessable->GetGuid();
+
+		Possessable->SetGuid(BenignSpawnableGuid);
+		Spawnable->SetGuid(PersistentGuid);
+
+		if (MovieScene->RemovePossessable(BenignSpawnableGuid))
+		{
+			Sequence->UnbindPossessableObjects(PersistentGuid);
+		}
+
+		for (int32 Index = 0; Index < MovieScene->GetPossessableCount(); ++Index)
+		{
+			FMovieScenePossessable& MovieScenePossessable = MovieScene->GetPossessable(Index);
+			if (MovieScenePossessable.GetParent() == PersistentGuid)
+			{
+				Spawnable->AddChildPossessable(MovieScenePossessable.GetGuid());
+			}
+		}
+	}
+	
+	FTransformData DefaultTransform;
+
+	// @todo: Where should this code live (not here)?
+	AActor* OldActor = Cast<AActor>(FoundObject);
+	if (OldActor)
+	{
+		DefaultTransform.Translation = OldActor->GetActorLocation();
+		DefaultTransform.Rotation = OldActor->GetActorRotation();
+		DefaultTransform.Scale = OldActor->GetActorScale();
+		DefaultTransform.bValid = true;
+
+		AActor* CDO = CastChecked<AActor>(Spawnable->GetClass()->GetDefaultObject());
+
+		EditorUtilities::CopyActorProperties(OldActor, CDO, EditorUtilities::ECopyOptions::Default);
+		GWorld->DestroyActor(OldActor);
+	}
+
+	SetupDefaultsForSpawnable(Spawnable->GetGuid(), DefaultTransform);
+
+	SetGlobalTimeDirectly(ScrubPosition);
+}
 
 void FSequencer::TogglePlay()
 {
@@ -2659,6 +3049,8 @@ void FSequencer::SetKey()
 			}
 		}
 	}
+
+	UpdatePlaybackRange();
 }
 
 
@@ -2790,7 +3182,6 @@ void FSequencer::CopySelectedKeys()
 	TOptional<float> CopyRelativeTo;
 	
 	// Copy relative to the current key hotspot, if applicable
-	TSharedPtr<ISequencerHotspot> Hotspot = SequencerWidget->GetEditTool().GetHotspot();
 	if (Hotspot.IsValid() && Hotspot->GetType() == ESequencerHotspot::Key)
 	{
 		FKeyHotspot* KeyHotspot = StaticCastSharedPtr<FKeyHotspot>(Hotspot).Get();
@@ -2909,12 +3300,17 @@ void FSequencer::SplitSection()
 	NotifyMovieSceneDataChanged();
 }
 
-
-ISequencerEditTool& FSequencer::GetEditTool()
+/** Get the current active hotspot */
+TSharedPtr<ISequencerHotspot> FSequencer::GetHotspot() const
 {
-	return SequencerWidget->GetEditTool();
+	return Hotspot;
 }
 
+/** Set the hotspot to something else */
+void FSequencer::SetHotspot(TSharedPtr<ISequencerHotspot> NewHotspot)
+{
+	Hotspot = MoveTemp(NewHotspot);
+}
 
 void FSequencer::BindSequencerCommands()
 {
@@ -2972,6 +3368,18 @@ void FSequencer::BindSequencerCommands()
 		Commands.SetStartPlaybackRange,
 		FExecuteAction::CreateSP( this, &FSequencer::SetStartPlaybackRange ) );
 
+	SequencerCommandBindings->MapAction(
+		Commands.ZoomInViewRange,
+		FExecuteAction::CreateSP( this, &FSequencer::ZoomInViewRange ),
+		FCanExecuteAction(),
+		EUIActionRepeatMode::RepeatEnabled );
+
+	SequencerCommandBindings->MapAction(
+		Commands.ZoomOutViewRange,
+		FExecuteAction::CreateSP( this, &FSequencer::ZoomOutViewRange ),		
+		FCanExecuteAction(),
+		EUIActionRepeatMode::RepeatEnabled );
+
 	SequencerCommandBindings->MapAction(
 		Commands.SetEndPlaybackRange,
 		FExecuteAction::CreateSP( this, &FSequencer::SetEndPlaybackRange ) );
@@ -3028,6 +3436,29 @@ void FSequencer::BindSequencerCommands()
 		Commands.SplitSection,
 		FExecuteAction::CreateSP( this, &FSequencer::SplitSection ) );
 
+	// We can convert to spawnables if anything selected is a root-level possessable
+	auto CanConvertToSpawnables = [this]{
+		UMovieScene* MovieScene = GetFocusedMovieSceneSequence()->GetMovieScene();
+
+		for (const TSharedRef<FSequencerDisplayNode>& Node : Selection.GetSelectedOutlinerNodes())
+		{
+			if (Node->GetType() == ESequencerNode::Object)
+			{
+				FMovieScenePossessable* Possessable = MovieScene->FindPossessable(static_cast<FSequencerObjectBindingNode&>(*Node).GetObjectBinding());
+				if (Possessable && !Possessable->GetParent().IsValid())
+				{
+					return true;
+				}
+			}
+		}
+		return false;
+	};
+	SequencerCommandBindings->MapAction(
+		FSequencerCommands::Get().ConvertToSpawnable,
+		FExecuteAction::CreateSP(this, &FSequencer::ConvertSelectedNodesToSpawnables),
+		FCanExecuteAction::CreateLambda(CanConvertToSpawnables)
+	);
+
 	SequencerCommandBindings->MapAction(
 		Commands.SetAutoKeyModeAll,
 		FExecuteAction::CreateLambda( [this]{ Settings->SetAutoKeyMode( EAutoKeyMode::KeyAll ); } ),
@@ -3081,9 +3512,13 @@ void FSequencer::BindSequencerCommands()
 	SequencerCommandBindings->MapAction(
 		Commands.ToggleShowFrameNumbers,
 		FExecuteAction::CreateLambda( [this]{ Settings->SetShowFrameNumbers( !Settings->GetShowFrameNumbers() ); } ),
-		FCanExecuteAction::CreateSP(this, &FSequencer::CanShowFrameNumbers ),
+		FCanExecuteAction::CreateSP( this, &FSequencer::CanShowFrameNumbers ),
 		FIsActionChecked::CreateLambda( [this]{ return Settings->GetShowFrameNumbers(); } ) );
 
+	SequencerCommandBindings->MapAction(
+		Commands.ToggleShowGotoBox,
+		FExecuteAction::CreateLambda( [this]{ Settings->SetShowGotoBox( !Settings->GetShowGotoBox() ); } ) );
+
 	SequencerCommandBindings->MapAction(
 		Commands.ToggleShowRangeSlider,
 		FExecuteAction::CreateLambda( [this]{ Settings->SetShowRangeSlider( !Settings->GetShowRangeSlider() ); } ),
@@ -3132,6 +3567,12 @@ void FSequencer::BindSequencerCommands()
 		FCanExecuteAction::CreateLambda( []{ return true; } ),
 		FIsActionChecked::CreateLambda( [this]{ return Settings->GetSnapPlayTimeToDraggedKey(); } ) );
 
+	SequencerCommandBindings->MapAction(
+		Commands.ToggleFixedTimeStepPlayback,
+		FExecuteAction::CreateLambda( [this]{ Settings->SetFixedTimeStepPlayback( !Settings->GetFixedTimeStepPlayback() ); } ),
+		FCanExecuteAction::CreateLambda( []{ return true; } ),
+		FIsActionChecked::CreateLambda( [this]{ return Settings->GetFixedTimeStepPlayback(); } ) );
+
 	SequencerCommandBindings->MapAction(
 		Commands.ToggleSnapCurveValueToInterval,
 		FExecuteAction::CreateLambda( [this]{ Settings->SetSnapCurveValueToInterval( !Settings->GetSnapCurveValueToInterval() ); } ),
diff --git a/Engine/Source/Editor/Sequencer/Private/Sequencer.h b/Engine/Source/Editor/Sequencer/Private/Sequencer.h
--- a/Engine/Source/Editor/Sequencer/Private/Sequencer.h
+++ b/Engine/Source/Editor/Sequencer/Private/Sequencer.h
@@ -23,6 +23,8 @@ class UMovieSceneSequence;
 class UWorld;
 class IMovieSceneSpawnRegister;
 
+struct ISequencerHotspot;
+struct FTransformData;
 
 /**
  * Sequencer is the editing tool for MovieScene assets.
@@ -54,10 +56,11 @@ public:
 	void InitSequencer(const FSequencerInitParams& InitParams, const TSharedRef<ISequencerObjectChangeListener>& InObjectChangeListener, const TSharedRef<IDetailKeyframeHandler>& InDetailKeyframeHandler, const TArray<FOnCreateTrackEditor>& TrackEditorDelegates);
 
 	/** Set the global time directly, without performing any auto-scroll */
-	void SetGlobalTimeDirectly(float Time);
+	void SetGlobalTimeDirectly(float Time, const bool& bAllowSnappingToFrames = false);
 
 	/** @return The current view range */
 	virtual FAnimatedRange GetViewRange() const override;
+	virtual void SetViewRange(TRange<float> NewViewRange, EViewRangeInterpolation Interpolation = EViewRangeInterpolation::Animated) override;
 
 	/** @return The current clamp range */
 	FAnimatedRange GetClampRange() const;
@@ -67,6 +70,10 @@ public:
 	void SetStartPlaybackRange();
 	void SetEndPlaybackRange();
 
+	void ZoomViewRange(float InZoomDelta);
+	void ZoomInViewRange();
+	void ZoomOutViewRange();
+
 public:
 
 	/** Access the user-supplied settings object */
@@ -116,14 +123,48 @@ public:
 		return TrackEditors;
 	}
 
+public:
+
 	/**
-	 * Attempts to add a new spawnable to the MovieScene for the specified asset or class
-	 *
-	 * @param	Object	The asset or class object to add a spawnable for
+	 * Attempts to add a new spawnable to the MovieScene for the specified object (asset, class or actor instance)
 	 *
+	 * @param	Object	The asset, class, or actor to add a spawnable for
 	 * @return	The spawnable guid for the spawnable, or an invalid Guid if we were not able to create a spawnable
 	 */
-	FGuid AddSpawnableForAssetOrClass( UObject* Object );
+	FGuid MakeNewSpawnable(UObject& SourceObject);
+
+	/**
+	 * Converts the specified possessable GUID to a spawnable
+	 *
+	 * @param	PossessableGuid		The guid of the possessable to convert
+	 */
+	void ConvertToSpawnable(TSharedRef<FSequencerObjectBindingNode> NodeToBeConverted);
+
+	/**
+	 * Converts all the currently selected nodes to be spawnables, if possible
+	 */
+	void ConvertSelectedNodesToSpawnables();
+
+protected:
+
+	/**
+	 * Attempts to add a new spawnable to the MovieScene for the specified asset, class, or actor
+	 *
+	 * @param	Object	The asset, class, or actor to add a spawnable for
+	 *
+	 * @return	The spawnable ID, or invalid ID on failure
+	 */
+	FGuid AddSpawnable( UObject& Object );
+
+	/**
+	 * Setup a new spawnable object with some default tracks and keys
+	 *
+	 * @param	Guid		The ID of the spawnable to setup defaults for
+	 * @param	Transform	The default transform for the spawnable
+	 */
+	void SetupDefaultsForSpawnable( const FGuid& Guid, const FTransformData& Transform );
+
+public:
 
 	/**
 	 * Call when an asset is dropped into the sequencer. Will proprogate this
@@ -143,15 +184,6 @@ public:
 	 */
 	virtual void OnRequestNodeDeleted( TSharedRef<const FSequencerDisplayNode> NodeToBeDeleted );
 
-	/**
-	 * Copies properties from one actor to another.  Only properties that are different are copied.  This is used
-	 * to propagate changed properties from a living actor in the scene to properties of a blueprint CDO actor
-	 *
-	 * @param	PuppetActor	The source actor (the puppet actor in the scene we're copying properties from)
-	 * @param	TargetActor	The destination actor (data-only blueprint actor CDO)
-	 */
-	virtual void CopyActorProperties( AActor* PuppetActor, AActor* TargetActor ) const;
-
 	/** Zooms to the edges of all currently selected sections. */
 	void ZoomToSelectedSections();
 
@@ -159,7 +191,7 @@ public:
 	float GetOverlayFadeCurve() const;
 
 	/** Gets the command bindings for the sequencer */
-	TSharedPtr<FUICommandList> GetCommandBindings() { return SequencerCommandBindings; }
+	virtual TSharedPtr<FUICommandList> GetCommandBindings() const override { return SequencerCommandBindings; }
 
 	/**
 	 * Builds up the sequencer's "Add Track" menu.
@@ -189,6 +221,14 @@ public:
 	/** Called when an actor is dropped into Sequencer */
 	void OnActorsDropped( const TArray<TWeakObjectPtr<AActor> >& Actors );
 
+	/** Generates and attaches transport control widgets to the main level editor viewports */
+	void AttachTransportControlsToViewports();
+	/** Purges all transport control widgets from the main level editor viewports */
+	void DetachTransportControlsFromViewports();
+	
+	/** Gets the visibility of a transport control based on the viewport it's attached to */
+	EVisibility GetTransportControlVisibility(TSharedPtr<ILevelViewport> LevelViewport) const;
+
 	/** Functions to push on to the transport controls we use */
 	FReply OnPlay(bool bTogglePlay=true);
 	FReply OnRecord();
@@ -265,8 +305,16 @@ public:
 
 public:
 
-	/** Access the currently enabled edit tool */
-	ISequencerEditTool& GetEditTool();
+	/** Get the current active hotspot */
+	TSharedPtr<ISequencerHotspot> GetHotspot() const;
+
+	/** Set the hotspot to something else */
+	void SetHotspot(TSharedPtr<ISequencerHotspot> NewHotspot);
+
+protected:
+
+	/** The current hotspot that can be set from anywhere to initiate drags */
+	TSharedPtr<ISequencerHotspot> Hotspot;
 
 public:
 
@@ -306,6 +354,7 @@ public:
 	virtual TSharedRef<FMovieSceneSequenceInstance> GetFocusedMovieSceneSequenceInstance() const override;
 	virtual void FocusSequenceInstance( TSharedRef<FMovieSceneSequenceInstance> SequenceInstance ) override;
 	virtual TSharedRef<FMovieSceneSequenceInstance> GetSequenceInstanceForSection(UMovieSceneSection& Section) const override;
+	virtual bool HasSequenceInstanceForSection(UMovieSceneSection& Section) const override;
 	virtual EAutoKeyMode GetAutoKeyMode() const override;
 	virtual void SetAutoKeyMode(EAutoKeyMode AutoKeyMode) override;
 	virtual bool GetKeyAllEnabled() const override;
@@ -317,7 +366,7 @@ public:
 	virtual bool IsRecordingLive() const override;
 	virtual float GetCurrentLocalTime(UMovieSceneSequence& InMovieSceneSequence) override;
 	virtual float GetGlobalTime() override;
-	virtual void SetGlobalTime(float Time) override;
+	virtual void SetGlobalTime(float Time, const bool& bAllowSnappingToFrames = false) override;
 	virtual void SetPerspectiveViewportPossessionEnabled(bool bEnabled) override;
 	virtual void SetPerspectiveViewportCameraCutEnabled(bool bEnabled) override;
 	virtual FGuid GetHandleToObject(UObject* Object, bool bCreateHandleIfMissing = true) override;
@@ -335,6 +384,10 @@ public:
 	virtual FGuid CreateBinding(UObject& InObject, const FString& InName) override;
 	virtual UObject* GetPlaybackContext() const override;
 	virtual void GetAllKeyedProperties(UObject& Object, TSet<UProperty*>& OutProperties) override;
+	virtual FOnCameraCut& OnCameraCut() override;
+	virtual TSharedRef<INumericTypeInterface<float>> GetNumericTypeInterface() override;
+	virtual TSharedRef<INumericTypeInterface<float>> GetZeroPadNumericTypeInterface() override;
+	virtual TSharedRef<SWidget> MakeTransportControls() override;
 
 public:
 
@@ -376,12 +429,8 @@ protected:
 	/** @return The current scrub position */
 	virtual float OnGetScrubPosition() const { return ScrubPosition; }
 
-	/**
-	 * Set the view range, growing the working range to accomodate, if necessary
-	 * @param NewViewRange The new view range. Must be a finite range
-	 * @param Interpolation How to interpolate to the new view range
-	 */
-	void SetViewRange(TRange<float> NewViewRange, EViewRangeInterpolation Interpolation = EViewRangeInterpolation::Animated);
+	/** Constrains the maximum frame rate to the fixed time step rate when playing back in that mode */
+	void ConstrainFixedTimeStepFrameRate();
 
 	/**
 	 * Called when the clamp range is changed by the user
@@ -543,6 +592,9 @@ protected:
 	/** Is the sequencer widget focused? */
 	bool IsSequencerWidgetFocused() const;
 
+	/** Internal conversion function that doesn't perform expensive reset/update tasks */
+	void ConvertToSpawnableInternal(TSharedRef<FSequencerObjectBindingNode> NodeToBeConverted);
+	
 private:
 
 	/** User-supplied settings object for this sequencer */
@@ -608,6 +660,13 @@ private:
 	/** Whether we are playing, recording, etc. */
 	EMovieScenePlayerStatus::Type PlaybackState;
 
+	/** The real-time that we started playback last */
+	double PlaybackStartRealTime;
+
+	/** Number of continuous fixed time step frames we've played so far without any change in play back state,
+	    such as time step, reverse mode, etc. */
+	uint32 NumContinuousFixedTimeStepFrames;
+
 	/** The current scrub position */
 	// @todo sequencer: Should use FTimespan or "double" for Time Cursor Position! (cascades)
 	float ScrubPosition;
@@ -632,4 +691,9 @@ private:
 
 	/** A delegate which is called any time the global time changes. */
 	FOnGlobalTimeChanged OnGlobalTimeChangedDelegate;
+
+	/** A map of all the transport controls to viewports that this sequencer has made */
+	TMap< TSharedPtr<class ILevelViewport>, TSharedPtr<class SWidget> > TransportControls;
+
+	FOnCameraCut OnCameraCutEvent;
 };
diff --git a/Engine/Source/Editor/Sequencer/Private/SequencerCommands.cpp b/Engine/Source/Editor/Sequencer/Private/SequencerCommands.cpp
--- a/Engine/Source/Editor/Sequencer/Private/SequencerCommands.cpp
+++ b/Engine/Source/Editor/Sequencer/Private/SequencerCommands.cpp
@@ -18,6 +18,8 @@ void FSequencerCommands::RegisterCommands()
 	UI_COMMAND( StepToPreviousCameraKey, "Step to Previous Camera Key", "Step to the previous camera key", EUserInterfaceActionType::Button, FInputChord(EModifierKey::Alt, EKeys::Comma) );
 	UI_COMMAND( SetStartPlaybackRange, "Set Start Playback Range", "Set the start playback range", EUserInterfaceActionType::Button, FInputChord(EKeys::LeftBracket) );
 	UI_COMMAND( SetEndPlaybackRange, "Set End Playback Range", "Set the end playback range", EUserInterfaceActionType::Button, FInputChord(EKeys::RightBracket) );
+	UI_COMMAND( ZoomInViewRange, "Zoom into the View Range", "Zoom into the view range", EUserInterfaceActionType::Button, FInputChord(EKeys::Equals) );
+	UI_COMMAND( ZoomOutViewRange, "Zoom out of the View Range", "Zoom out of the view range", EUserInterfaceActionType::Button, FInputChord(EKeys::Hyphen) );
 	UI_COMMAND( ToggleKeepCursorInPlaybackRange, "Keep Cursor In Playback Range", "When checked, the cursor will be constrained to the current playback range during playback", EUserInterfaceActionType::ToggleButton, FInputChord() );
 	UI_COMMAND( ToggleKeepPlaybackRangeInSectionBounds, "Keep Playback Range In Section Bounds", "When checked, the playback range will be synchronized to the section bounds", EUserInterfaceActionType::ToggleButton, FInputChord() );
 
@@ -46,6 +48,7 @@ void FSequencerCommands::RegisterCommands()
 	UI_COMMAND( ToggleAutoScroll, "Auto Scroll", "Toggle auto-scroll: When enabled, automatically scrolls the sequencer view to keep the current time visible", EUserInterfaceActionType::ToggleButton, FInputChord(EModifierKey::Shift, EKeys::S) );
 
 	UI_COMMAND( ToggleShowFrameNumbers, "Show Frame Numbers", "Enables and disables showing frame numbers", EUserInterfaceActionType::ToggleButton, FInputChord() );
+	UI_COMMAND( ToggleShowGotoBox, "Go to time...", "Go to a particular point on the time line", EUserInterfaceActionType::Button, FInputChord(EModifierKey::Control, EKeys::G) );
 	UI_COMMAND( ToggleShowRangeSlider, "Show Range Slider", "Enables and disables showing the time range slider", EUserInterfaceActionType::ToggleButton, FInputChord() );
 	UI_COMMAND( ToggleIsSnapEnabled, "Enable Snapping", "Enables and disables snapping", EUserInterfaceActionType::ToggleButton, FInputChord() );
 
@@ -58,6 +61,8 @@ void FSequencerCommands::RegisterCommands()
 	UI_COMMAND( ToggleSnapPlayTimeToInterval, "Snap to the Interval While Scrubbing", "Snap the play time to the interval while scrubbing", EUserInterfaceActionType::ToggleButton, FInputChord() );
 	UI_COMMAND( ToggleSnapPlayTimeToDraggedKey, "Snap to the Dragged Key", "Snap the play time to the dragged key", EUserInterfaceActionType::ToggleButton, FInputChord() );
 
+	UI_COMMAND( ToggleFixedTimeStepPlayback, "Fixed Time Step Playback", "Lock the play rate to the frame rate", EUserInterfaceActionType::ToggleButton, FInputChord() );
+
 	UI_COMMAND( ToggleSnapCurveValueToInterval, "Snap Curve Key Values", "Snap curve keys to the value snapping interval", EUserInterfaceActionType::ToggleButton, FInputChord() );
 
 	UI_COMMAND( FindInContentBrowser, "Find in Content Browser", "Find the viewed sequence asset in the content browser", EUserInterfaceActionType::Button, FInputChord() );
@@ -72,6 +77,8 @@ void FSequencerCommands::RegisterCommands()
 	UI_COMMAND( RenderMovie, "Render Movie", "Render this movie to a video, or image frame sequence", EUserInterfaceActionType::Button, FInputChord() );
 
 	UI_COMMAND( PasteFromHistory, "Paste From History", "Paste from the sequencer clipboard history", EUserInterfaceActionType::Button, FInputChord(EModifierKey::Control | EModifierKey::Shift, EKeys::V) );
+
+	UI_COMMAND( ConvertToSpawnable, "Convert to Spawnable", "Make the specified possessed objects spawnable from sequencer. This will allow sequencer to have control over the lifetime of the object.", EUserInterfaceActionType::Button, FInputChord() );
 }
 
 #undef LOCTEXT_NAMESPACE
\ No newline at end of file
diff --git a/Engine/Source/Editor/Sequencer/Private/SequencerCommonHelpers.cpp b/Engine/Source/Editor/Sequencer/Private/SequencerCommonHelpers.cpp
--- a/Engine/Source/Editor/Sequencer/Private/SequencerCommonHelpers.cpp
+++ b/Engine/Source/Editor/Sequencer/Private/SequencerCommonHelpers.cpp
@@ -3,6 +3,9 @@
 #include "SequencerPrivatePCH.h"
 #include "SequencerCommonHelpers.h"
 #include "MovieSceneSection.h"
+#include "SequencerHotspots.h"
+#include "VirtualTrackArea.h"
+#include "SequencerContextMenus.h"
 #include "MovieSceneTrack.h"
 
 void SequencerHelpers::GetAllKeyAreas(TSharedPtr<FSequencerDisplayNode> DisplayNode, TSet<TSharedPtr<IKeyArea>>& KeyAreas)
@@ -117,3 +120,116 @@ float SequencerHelpers::FrameToTime(int32 Frame, float FrameRate)
 	return Frame / FrameRate;
 }
 
+void SequencerHelpers::PerformDefaultSelection(FSequencer& Sequencer, const FPointerEvent& MouseEvent)
+{
+	FSequencerSelection& Selection = Sequencer.GetSelection();
+
+	// @todo: selection in transactions
+	auto ConditionallyClearSelection = [&]{
+		if (!MouseEvent.IsShiftDown() && !MouseEvent.IsControlDown())
+		{
+			Selection.EmptySelectedSections();
+			Selection.EmptySelectedKeys();
+		}
+	};
+
+	TSharedPtr<ISequencerHotspot> Hotspot = Sequencer.GetHotspot();
+	if (!Hotspot.IsValid())
+	{
+		ConditionallyClearSelection();
+		return;
+	}
+
+	// Handle right-click selection separately since we never deselect on right click (except for clearing on exclusive selection)
+	if (MouseEvent.GetEffectingButton() == EKeys::RightMouseButton)
+	{
+		if (Hotspot->GetType() == ESequencerHotspot::Key)
+		{
+			FSequencerSelectedKey Key = static_cast<FKeyHotspot*>(Hotspot.Get())->Key;
+			if (!Selection.IsSelected(Key))
+			{
+				ConditionallyClearSelection();
+				Selection.AddToSelection(Key);
+			}
+		}
+		else if (Hotspot->GetType() == ESequencerHotspot::Section)
+		{
+			UMovieSceneSection* Section = static_cast<FSectionHotspot*>(Hotspot.Get())->Section.GetSectionObject();
+			if (!Selection.IsSelected(Section))
+			{
+				ConditionallyClearSelection();
+				Selection.AddToSelection(Section);
+			}
+		}
+
+		return;
+	}
+
+	// Normal selection
+	ConditionallyClearSelection();
+
+	bool bForceSelect = !MouseEvent.IsControlDown();
+
+	if (Hotspot->GetType() == ESequencerHotspot::Key)
+	{
+		FSequencerSelectedKey Key = static_cast<FKeyHotspot*>(Hotspot.Get())->Key;
+		if (bForceSelect || !Selection.IsSelected(Key))
+		{
+			Selection.AddToSelection(Key);
+		}
+		else
+		{
+			Selection.RemoveFromSelection(Key);
+		}
+	}
+	else if (Hotspot->GetType() == ESequencerHotspot::Section)
+	{
+		UMovieSceneSection* Section = static_cast<FSectionHotspot*>(Hotspot.Get())->Section.GetSectionObject();
+		if (bForceSelect || !Selection.IsSelected(Section))
+		{
+			Selection.AddToSelection(Section);
+		}
+		else
+		{
+			Selection.RemoveFromSelection(Section);
+		}
+	}
+}
+
+TSharedPtr<SWidget> SequencerHelpers::SummonContextMenu(FSequencer& Sequencer, const FGeometry& MyGeometry, const FPointerEvent& MouseEvent)
+{
+	// @todo sequencer replace with UI Commands instead of faking it
+
+	FVector2D MouseDownPos = MyGeometry.AbsoluteToLocal(MouseEvent.GetScreenSpacePosition());
+
+	// Attempt to paste into either the current node selection, or the clicked on track
+	TSharedRef<SSequencer> SequencerWidget = StaticCastSharedRef<SSequencer>(Sequencer.GetSequencerWidget());
+	const float PasteAtTime = SequencerWidget->GetVirtualTrackArea().PixelToTime(MouseDownPos.X);
+
+	const bool bShouldCloseWindowAfterMenuSelection = true;
+	FMenuBuilder MenuBuilder(bShouldCloseWindowAfterMenuSelection, Sequencer.GetCommandBindings());
+
+	TSharedPtr<ISequencerHotspot> Hotspot = Sequencer.GetHotspot();
+
+	if (Hotspot.IsValid())
+	{
+		if (Hotspot->GetType() == ESequencerHotspot::Section || Hotspot->GetType() == ESequencerHotspot::Key)
+		{
+			Hotspot->PopulateContextMenu(MenuBuilder, Sequencer, PasteAtTime);
+
+			return MenuBuilder.MakeWidget();
+		}
+	}
+	else if (Sequencer.GetClipboardStack().Num() != 0)
+	{
+		TSharedPtr<FPasteContextMenu> PasteMenu = FPasteContextMenu::CreateMenu(Sequencer, SequencerWidget->GeneratePasteArgs(PasteAtTime));
+		if (PasteMenu.IsValid() && PasteMenu->IsValidPaste())
+		{
+			PasteMenu->PopulateMenu(MenuBuilder);
+
+			return MenuBuilder.MakeWidget();
+		}
+	}
+
+	return nullptr;
+}
\ No newline at end of file
diff --git a/Engine/Source/Editor/Sequencer/Private/SequencerCommonHelpers.h b/Engine/Source/Editor/Sequencer/Private/SequencerCommonHelpers.h
--- a/Engine/Source/Editor/Sequencer/Private/SequencerCommonHelpers.h
+++ b/Engine/Source/Editor/Sequencer/Private/SequencerCommonHelpers.h
@@ -39,6 +39,16 @@ public:
 	* Convert frame to time
 	*/
 	static float FrameToTime(int32 Frame, float FrameRate);
+
+	/**
+	 * Perform default selection for the specified mouse event, based on the current hotspot
+	 */
+	static void PerformDefaultSelection(FSequencer& Sequencer, const FPointerEvent& MouseEvent);
+	
+	/**
+	 * Attempt to summon a context menu for the current hotspot
+	 */
+	static TSharedPtr<SWidget> SummonContextMenu(FSequencer& Sequencer, const FGeometry& MyGeometry, const FPointerEvent& MouseEvent);
 };
 
 
