Copying //UE4/Dev-Sequencer to Dev-Main (//UE4/Dev-Main)

diff --git a/Engine/Source/Editor/Sequencer/Private/SequencerContextMenus.cpp b/Engine/Source/Editor/Sequencer/Private/SequencerContextMenus.cpp
--- a/Engine/Source/Editor/Sequencer/Private/SequencerContextMenus.cpp
+++ b/Engine/Source/Editor/Sequencer/Private/SequencerContextMenus.cpp
@@ -9,24 +9,48 @@
 #include "SSequencerSection.h"
 #include "MovieSceneTrack.h"
 #include "MovieSceneCommonHelpers.h"
+#include "MovieSceneKeyStruct.h"
 #include "GenericCommands.h"
 #include "SequencerCommands.h"
+#include "ModuleManager.h"
+#include "IDetailsView.h"
+#include "IStructureDetailsView.h"
+#include "PropertyEditorModule.h"
+#include "MovieSceneSequence.h"
+
 
 #define LOCTEXT_NAMESPACE "SequencerContextMenus"
 
+
 void FKeyContextMenu::BuildMenu(FMenuBuilder& MenuBuilder, FSequencer& InSequencer)
 {
 	TSharedRef<FKeyContextMenu> Menu = MakeShareable(new FKeyContextMenu(InSequencer));
 	Menu->PopulateMenu(MenuBuilder);
 }
 
+
 void FKeyContextMenu::PopulateMenu(FMenuBuilder& MenuBuilder)
 {
 	FSequencer* SequencerPtr = &Sequencer.Get();
+	TSharedRef<FKeyContextMenu> Shared = AsShared();
+
+	MenuBuilder.AddSubMenu(
+		LOCTEXT("KeyProperties", "Properties"),
+		LOCTEXT("KeyPropertiesTooltip", "Modify the key properties"),
+		FNewMenuDelegate::CreateLambda([=](FMenuBuilder& SubMenuBuilder){ Shared->AddPropertiesMenu(SubMenuBuilder); }),
+		FUIAction (
+			FExecuteAction(),
+			// @todo sequencer: only one struct per structure view supported right now :/
+			FCanExecuteAction::CreateLambda([=]{ return (Shared->Sequencer->GetSelection().GetSelectedKeys().Num() == 1); })
+		),
+		NAME_None,
+		EUserInterfaceActionType::Button
+	);
 
 	MenuBuilder.BeginSection("SequencerKeyEdit", LOCTEXT("EditMenu", "Edit"));
 	{
-		TSharedPtr<ISequencerHotspot> Hotspot = SequencerPtr->GetEditTool().GetHotspot();
+		TSharedPtr<ISequencerHotspot> Hotspot = SequencerPtr->GetHotspot();
+
 		if (Hotspot.IsValid() && Hotspot->GetType() == ESequencerHotspot::Key)
 		{
 			MenuBuilder.AddMenuEntry(FGenericCommands::Get().Cut);
@@ -131,17 +155,89 @@ void FKeyContextMenu::PopulateMenu(FMenuBuilder& MenuBuilder)
 	MenuBuilder.EndSection(); // SequencerKeys
 }
 
+
+void FKeyContextMenu::AddPropertiesMenu(FMenuBuilder& MenuBuilder)
+{
+	FDetailsViewArgs DetailsViewArgs;
+	{
+		DetailsViewArgs.bAllowSearch = false;
+		DetailsViewArgs.bCustomFilterAreaLocation = true;
+		DetailsViewArgs.bCustomNameAreaLocation = true;
+		DetailsViewArgs.bHideSelectionTip = true;
+		DetailsViewArgs.bLockable = false;
+		DetailsViewArgs.bSearchInitialKeyFocus = true;
+		DetailsViewArgs.bUpdatesFromSelection = false;
+		DetailsViewArgs.bShowOptions = false;
+		DetailsViewArgs.bShowModifiedPropertiesOption = false;
+	}
+
+	FStructureDetailsViewArgs StructureViewArgs;
+	{
+		StructureViewArgs.bShowObjects = false;
+		StructureViewArgs.bShowAssets = true;
+		StructureViewArgs.bShowClasses = true;
+		StructureViewArgs.bShowInterfaces = false;
+	}
+
+	TArray<TSharedPtr<FStructOnScope>> Keys;
+	{
+		for (const FSequencerSelectedKey& Key : Sequencer->GetSelection().GetSelectedKeys())
+		{
+			if (Key.KeyArea.IsValid() && Key.KeyHandle.IsSet())
+			{
+				TSharedPtr<FStructOnScope> KeyStruct = Key.KeyArea->GetKeyStruct(Key.KeyHandle.GetValue());
+
+				if (KeyStruct.IsValid())
+				{
+					Keys.Add(KeyStruct);
+				}
+			}
+		}
+	}
+
+	TSharedRef<IStructureDetailsView> StructureDetailsView = FModuleManager::GetModuleChecked<FPropertyEditorModule>("PropertyEditor")
+		.CreateStructureDetailView(DetailsViewArgs, StructureViewArgs, nullptr, LOCTEXT("MessageData", "Message Data"));
+	{
+		// @todo sequencer: only one struct per structure view supported right now :/
+		if (Keys.Num() == 1)
+		{
+			TSharedPtr<FStructOnScope>& Key = Keys[0];
+			StructureDetailsView->SetStructureData(Key);
+			StructureDetailsView->GetOnFinishedChangingPropertiesDelegate().AddLambda(
+				[=](const FPropertyChangedEvent& ChangeEvent) {
+					if (Key->GetStruct()->IsChildOf(FMovieSceneKeyStruct::StaticStruct()))
+					{
+						((FMovieSceneKeyStruct*)Key->GetStructMemory())->PropagateChanges(ChangeEvent);
+					}
+					Sequencer->GetFocusedMovieSceneSequence()->Modify();
+					Sequencer->UpdateRuntimeInstances();
+				}
+			);
+		}
+	}
+	
+	MenuBuilder.AddWidget(StructureDetailsView->GetWidget().ToSharedRef(), FText::GetEmpty(), true);
+}
+
+
 void FSectionContextMenu::BuildMenu(FMenuBuilder& MenuBuilder, FSequencer& InSequencer, float InMouseDownTime)
 {
 	TSharedRef<FSectionContextMenu> Menu = MakeShareable(new FSectionContextMenu(InSequencer, InMouseDownTime));
 	Menu->PopulateMenu(MenuBuilder);
 }
 
+
 void FSectionContextMenu::PopulateMenu(FMenuBuilder& MenuBuilder)
 {
 	// Copy a reference to the context menu by value into each lambda handler to ensure the type stays alive until the menu is closed
 	TSharedRef<FSectionContextMenu> Shared = AsShared();
 
+	MenuBuilder.AddSubMenu(
+		LOCTEXT("SectionProperties", "Properties"),
+		LOCTEXT("SectionPropertiesTooltip", "Modify the section properties"),
+		FNewMenuDelegate::CreateLambda([=](FMenuBuilder& SubMenuBuilder){ Shared->AddPropertiesMenu(SubMenuBuilder); })
+	);
+
 	MenuBuilder.BeginSection("SequencerKeyEdit", LOCTEXT("EditMenu", "Edit"));
 	{
 		TSharedPtr<FPasteFromHistoryContextMenu> PasteFromHistoryMenu;
@@ -182,30 +278,38 @@ void FSectionContextMenu::PopulateMenu(FMenuBuilder& MenuBuilder)
 
 	MenuBuilder.BeginSection("SequencerSections", LOCTEXT("SectionsMenu", "Sections"));
 	{
-		bool bCanSelectAll = CanSelectAllKeys();
-		MenuBuilder.AddMenuEntry(
-			LOCTEXT("SelectAllKeys", "Select All Keys"),
-			LOCTEXT("SelectAllKeysTooltip", "Select all keys in section"),
-			FSlateIcon(),
-			FUIAction(
-				FExecuteAction::CreateLambda([=]{ return Shared->SelectAllKeys(); }),
-				FCanExecuteAction::CreateLambda([=]{ return bCanSelectAll; }))
-		);
+		if (CanSelectAllKeys())
+		{
+			MenuBuilder.AddMenuEntry(
+				LOCTEXT("SelectAllKeys", "Select All Keys"),
+				LOCTEXT("SelectAllKeysTooltip", "Select all keys in section"),
+				FSlateIcon(),
+				FUIAction(FExecuteAction::CreateLambda([=]{ return Shared->SelectAllKeys(); }))
+			);
+		}
 
 		MenuBuilder.AddSubMenu(
 			LOCTEXT("EditSection", "Edit"),
 			LOCTEXT("EditSectionTooltip", "Edit section"),
 			FNewMenuDelegate::CreateLambda([=](FMenuBuilder& InMenuBuilder){ Shared->AddEditMenu(InMenuBuilder); }));
 
-		MenuBuilder.AddSubMenu(
-			LOCTEXT("SetPreInfinityExtrap", "Pre-Infinity"),
-			LOCTEXT("SetPreInfinityExtrapTooltip", "Set pre-infinity extrapolation"),
-			FNewMenuDelegate::CreateLambda([=](FMenuBuilder& SubMenuBuilder){ Shared->AddExtrapolationMenu(SubMenuBuilder, true); }));
+		if (CanSetExtrapolationMode())
+		{
+			MenuBuilder.AddSubMenu(
+				LOCTEXT("SetPreInfinityExtrap", "Pre-Infinity"),
+				LOCTEXT("SetPreInfinityExtrapTooltip", "Set pre-infinity extrapolation"),
+				FNewMenuDelegate::CreateLambda([=](FMenuBuilder& SubMenuBuilder){ Shared->AddExtrapolationMenu(SubMenuBuilder, true); }));
+
+			MenuBuilder.AddSubMenu(
+				LOCTEXT("SetPostInfinityExtrap", "Post-Infinity"),
+				LOCTEXT("SetPostInfinityExtrapTooltip", "Set post-infinity extrapolation"),
+				FNewMenuDelegate::CreateLambda([=](FMenuBuilder& SubMenuBuilder){ Shared->AddExtrapolationMenu(SubMenuBuilder, false); }));
+		}
 
 		MenuBuilder.AddSubMenu(
-			LOCTEXT("SetPostInfinityExtrap", "Post-Infinity"),
-			LOCTEXT("SetPostInfinityExtrapTooltip", "Set post-infinity extrapolation"),
-			FNewMenuDelegate::CreateLambda([=](FMenuBuilder& SubMenuBuilder){ Shared->AddExtrapolationMenu(SubMenuBuilder, false); }));
+			LOCTEXT("OrderSection", "Order"),
+			LOCTEXT("OrderSectionTooltip", "Order section"),
+			FNewMenuDelegate::CreateLambda([=](FMenuBuilder& SubMenuBuilder){ Shared->AddOrderMenu(SubMenuBuilder); }));			
 
 		MenuBuilder.AddMenuEntry(
 			LOCTEXT("ToggleSectionActive", "Active"),
@@ -241,25 +345,9 @@ void FSectionContextMenu::PopulateMenu(FMenuBuilder& MenuBuilder)
 		);
 	}
 	MenuBuilder.EndSection(); // SequencerSections
-
-	MenuBuilder.BeginSection("SequencerSectionOrder", LOCTEXT("SectionOrderLabel", "Order"));
-	{
-		MenuBuilder.AddMenuEntry(LOCTEXT("BringToFront", "Bring To Front"), FText(), FSlateIcon(),
-			FUIAction(FExecuteAction::CreateLambda([=]{ return Shared->BringToFront(); })));
-
-		MenuBuilder.AddMenuEntry(LOCTEXT("SendToBack", "Send To Back"), FText(), FSlateIcon(),
-			FUIAction(FExecuteAction::CreateLambda([=]{ return Shared->SendToBack(); })));
-
-		MenuBuilder.AddMenuEntry(LOCTEXT("BringForward", "Bring Forward"), FText(), FSlateIcon(),
-			FUIAction(FExecuteAction::CreateLambda([=]{ return Shared->BringForward(); })));
-
-		MenuBuilder.AddMenuEntry(LOCTEXT("SendBackward", "Send Backward"), FText(), FSlateIcon(),
-			FUIAction(FExecuteAction::CreateLambda([=]{ return Shared->SendBackward(); })));
-
-	}
-	MenuBuilder.EndSection(); // SequencerSectionOrder
 }
 
+
 void FSectionContextMenu::AddEditMenu(FMenuBuilder& MenuBuilder)
 {
 	// Copy a reference to the context menu by value into each lambda handler to ensure the type stays alive until the menu is closed
@@ -293,6 +381,7 @@ void FSectionContextMenu::AddEditMenu(FMenuBuilder& MenuBuilder)
 	);
 }
 
+
 void FSectionContextMenu::AddExtrapolationMenu(FMenuBuilder& MenuBuilder, bool bPreInfinity)
 {
 	// Copy a reference to the context menu by value into each lambda handler to ensure the type stays alive until the menu is closed
@@ -364,6 +453,58 @@ void FSectionContextMenu::AddExtrapolationMenu(FMenuBuilder& MenuBuilder, bool b
 	);
 }
 
+
+void FSectionContextMenu::AddPropertiesMenu(FMenuBuilder& MenuBuilder)
+{
+	FDetailsViewArgs DetailsViewArgs;
+	{
+		DetailsViewArgs.bAllowSearch = false;
+		DetailsViewArgs.bCustomFilterAreaLocation = true;
+		DetailsViewArgs.bCustomNameAreaLocation = true;
+		DetailsViewArgs.bHideSelectionTip = true;
+		DetailsViewArgs.bLockable = false;
+		DetailsViewArgs.bSearchInitialKeyFocus = true;
+		DetailsViewArgs.bUpdatesFromSelection = false;
+		DetailsViewArgs.bShowOptions = false;
+		DetailsViewArgs.bShowModifiedPropertiesOption = false;
+	}
+
+	TArray<TWeakObjectPtr<UObject>> Sections;
+	{
+		for (auto Section : Sequencer->GetSelection().GetSelectedSections())
+		{
+			Sections.Add(Section);
+		}
+	}
+
+	TSharedRef<IDetailsView> DetailsView = FModuleManager::GetModuleChecked<FPropertyEditorModule>("PropertyEditor").CreateDetailView(DetailsViewArgs);
+	{
+		DetailsView->SetObjects(Sections);
+	}
+	
+	MenuBuilder.AddWidget(DetailsView, FText::GetEmpty(), true);
+}
+
+
+void FSectionContextMenu::AddOrderMenu(FMenuBuilder& MenuBuilder)
+{
+	// Copy a reference to the context menu by value into each lambda handler to ensure the type stays alive until the menu is closed
+	TSharedRef<FSectionContextMenu> Shared = AsShared();
+
+	MenuBuilder.AddMenuEntry(LOCTEXT("BringToFront", "Bring To Front"), FText(), FSlateIcon(),
+		FUIAction(FExecuteAction::CreateLambda([=]{ return Shared->BringToFront(); })));
+
+	MenuBuilder.AddMenuEntry(LOCTEXT("SendToBack", "Send To Back"), FText(), FSlateIcon(),
+		FUIAction(FExecuteAction::CreateLambda([=]{ return Shared->SendToBack(); })));
+
+	MenuBuilder.AddMenuEntry(LOCTEXT("BringForward", "Bring Forward"), FText(), FSlateIcon(),
+		FUIAction(FExecuteAction::CreateLambda([=]{ return Shared->BringForward(); })));
+
+	MenuBuilder.AddMenuEntry(LOCTEXT("SendBackward", "Send Backward"), FText(), FSlateIcon(),
+		FUIAction(FExecuteAction::CreateLambda([=]{ return Shared->SendBackward(); })));
+}
+
+
 void FSectionContextMenu::SelectAllKeys()
 {
 	TArray<FSectionHandle> SelectedSections = StaticCastSharedRef<SSequencer>(Sequencer->GetSequencerWidget())->GetSectionHandles(Sequencer->GetSelection().GetSelectedSections());
@@ -388,6 +529,7 @@ void FSectionContextMenu::SelectAllKeys()
 	}
 }
 
+
 bool FSectionContextMenu::CanSelectAllKeys() const
 {
 	TArray<FSectionHandle> SelectedSections = StaticCastSharedRef<SSequencer>(Sequencer->GetSequencerWidget())->GetSectionHandles(Sequencer->GetSelection().GetSelectedSections());
@@ -413,6 +555,24 @@ bool FSectionContextMenu::CanSelectAllKeys() const
 	return false;
 }
 
+bool FSectionContextMenu::CanSetExtrapolationMode() const
+{
+	TArray<FSectionHandle> SelectedSections = StaticCastSharedRef<SSequencer>(Sequencer->GetSequencerWidget())->GetSectionHandles(Sequencer->GetSelection().GetSelectedSections());
+	for (const FSectionHandle& Handle : SelectedSections)
+	{
+		FKeyAreaLayout Layout(*Handle.TrackNode, Handle.SectionIndex);
+		for (const FKeyAreaLayoutElement& Element : Layout.GetElements())
+		{
+			TSharedPtr<IKeyArea> KeyArea = Element.GetKeyArea();
+			if (KeyArea->CanSetExtrapolationMode())
+			{
+				return true;
+			}
+		}
+	}
+	return false;
+}
+
 void FSectionContextMenu::TrimSection(bool bTrimLeft)
 {
 	FScopedTransaction TrimSectionTransaction(LOCTEXT("TrimSection_Transaction", "Trim Section"));
@@ -421,6 +581,7 @@ void FSectionContextMenu::TrimSection(bool bTrimLeft)
 	Sequencer->NotifyMovieSceneDataChanged();
 }
 
+
 void FSectionContextMenu::SplitSection()
 {
 	FScopedTransaction SplitSectionTransaction(LOCTEXT("SplitSection_Transaction", "Split Section"));
@@ -429,6 +590,7 @@ void FSectionContextMenu::SplitSection()
 	Sequencer->NotifyMovieSceneDataChanged();
 }
 
+
 bool FSectionContextMenu::IsTrimmable() const
 {
 	for (auto Section : Sequencer->GetSelection().GetSelectedSections())
@@ -441,6 +603,7 @@ bool FSectionContextMenu::IsTrimmable() const
 	return false;
 }
 
+
 void FSectionContextMenu::SetExtrapolationMode(ERichCurveExtrapolation ExtrapMode, bool bPreInfinity)
 {
 	FScopedTransaction SetExtrapolationModeTransaction(LOCTEXT("SetExtrapolationMode_Transaction", "Set Extrapolation Mode"));
@@ -478,6 +641,7 @@ void FSectionContextMenu::SetExtrapolationMode(ERichCurveExtrapolation ExtrapMod
 	}
 }
 
+
 bool FSectionContextMenu::IsExtrapolationModeSelected(ERichCurveExtrapolation ExtrapMode, bool bPreInfinity) const
 {
 	// @todo Sequencer should operate on selected sections
@@ -503,6 +667,7 @@ bool FSectionContextMenu::IsExtrapolationModeSelected(ERichCurveExtrapolation Ex
 	return bAllSelected;
 }
 
+
 void FSectionContextMenu::ToggleSectionActive()
 {
 	FScopedTransaction ToggleSectionActiveTransaction( LOCTEXT("ToggleSectionActive_Transaction", "Toggle Section Active") );
@@ -526,6 +691,7 @@ void FSectionContextMenu::ToggleSectionActive()
 	}
 }
 
+
 bool FSectionContextMenu::IsSectionActive() const
 {
 	// Active only if all are active
@@ -540,6 +706,7 @@ bool FSectionContextMenu::IsSectionActive() const
 	return true;
 }
 
+
 void FSectionContextMenu::ToggleSectionLocked()
 {
 	FScopedTransaction ToggleSectionLockedTransaction( NSLOCTEXT("Sequencer", "ToggleSectionLocked_Transaction", "Toggle Section Locked") );
@@ -563,6 +730,7 @@ void FSectionContextMenu::ToggleSectionLocked()
 	}
 }
 
+
 bool FSectionContextMenu::IsSectionLocked() const
 {
 	// Locked only if all are locked
@@ -577,11 +745,13 @@ bool FSectionContextMenu::IsSectionLocked() const
 	return true;
 }
 
+
 void FSectionContextMenu::DeleteSection()
 {
 	Sequencer->DeleteSections(Sequencer->GetSelection().GetSelectedSections());
 }
 
+
 /** Information pertaining to a specific row in a track, required for z-ordering operations */
 struct FTrackSectionRow
 {
@@ -605,6 +775,7 @@ struct FTrackSectionRow
 	}
 };
 
+
 /** Generate the data required for re-ordering rows based on the current sequencer selection */
 /** @note: Produces a map of track -> rows, keyed on row index. Only returns rows that contain selected sections */
 TMap<UMovieSceneTrack*, TMap<int32, FTrackSectionRow>> GenerateTrackRowsFromSelection(FSequencer& Sequencer)
@@ -649,6 +820,7 @@ TMap<UMovieSceneTrack*, TMap<int32, FTrackSectionRow>> GenerateTrackRowsFromSele
 	return TrackRows;
 }
 
+
 /** Modify all the sections contained within the specified data structure */
 void ModifySections(TMap<UMovieSceneTrack*, TMap<int32, FTrackSectionRow>>& TrackRows)
 {
@@ -665,6 +837,7 @@ void ModifySections(TMap<UMovieSceneTrack*, TMap<int32, FTrackSectionRow>>& Trac
 	}
 }
 
+
 void FSectionContextMenu::BringToFront()
 {
 	TMap<UMovieSceneTrack*, TMap<int32, FTrackSectionRow>> TrackRows = GenerateTrackRowsFromSelection(*Sequencer);
@@ -712,6 +885,7 @@ void FSectionContextMenu::BringToFront()
 	Sequencer->SetGlobalTime(Sequencer->GetGlobalTime());
 }
 
+
 void FSectionContextMenu::SendToBack()
 {
 	TMap<UMovieSceneTrack*, TMap<int32, FTrackSectionRow>> TrackRows = GenerateTrackRowsFromSelection(*Sequencer);
@@ -759,6 +933,7 @@ void FSectionContextMenu::SendToBack()
 	Sequencer->SetGlobalTime(Sequencer->GetGlobalTime());
 }
 
+
 void FSectionContextMenu::BringForward()
 {
 	TMap<UMovieSceneTrack*, TMap<int32, FTrackSectionRow>> TrackRows = GenerateTrackRowsFromSelection(*Sequencer);
@@ -803,6 +978,7 @@ void FSectionContextMenu::BringForward()
 	Sequencer->SetGlobalTime(Sequencer->GetGlobalTime());
 }
 
+
 void FSectionContextMenu::SendBackward()
 {
 	TMap<UMovieSceneTrack*, TMap<int32, FTrackSectionRow>> TrackRows = GenerateTrackRowsFromSelection(*Sequencer);
@@ -861,6 +1037,7 @@ bool FPasteContextMenu::BuildMenu(FMenuBuilder& MenuBuilder, FSequencer& InSeque
 	return true;
 }
 
+
 TSharedRef<FPasteContextMenu> FPasteContextMenu::CreateMenu(FSequencer& InSequencer, const FPasteContextMenuArgs& Args)
 {
 	TSharedRef<FPasteContextMenu> Menu = MakeShareable(new FPasteContextMenu(InSequencer, Args));
@@ -868,6 +1045,7 @@ TSharedRef<FPasteContextMenu> FPasteContextMenu::CreateMenu(FSequencer& InSequen
 	return Menu;
 }
 
+
 TArray<TSharedRef<FSequencerSectionKeyAreaNode>> KeyAreaNodesBuffer;
 
 void FPasteContextMenu::GatherPasteDestinationsForNode(FSequencerDisplayNode& InNode, int32 SectionIndex, const FName& CurrentScope, TMap<FName, FSequencerClipboardReconciler>& Map)
@@ -918,6 +1096,7 @@ void FPasteContextMenu::GatherPasteDestinationsForNode(FSequencerDisplayNode& In
 	}
 }
 
+
 void GetFullNodePath(FSequencerDisplayNode& InNode, FString& Path)
 {
 	TSharedPtr<FSequencerDisplayNode> Parent = InNode.GetParent();
@@ -966,6 +1145,7 @@ TSharedPtr<FSequencerTrackNode> GetTrackFromNode(FSequencerDisplayNode& InNode,
 	return nullptr;
 }
 
+
 void FPasteContextMenu::Setup()
 {
 	if (!Args.Clipboard.IsValid())
@@ -1060,11 +1240,13 @@ void FPasteContextMenu::Setup()
 	}
 }
 
+
 bool FPasteContextMenu::IsValidPaste() const
 {
 	return Args.Clipboard.IsValid() && PasteDestinations.Num() != 0;
 }
 
+
 void FPasteContextMenu::PopulateMenu(FMenuBuilder& MenuBuilder)
 {
 	// Copy a reference to the context menu by value into each lambda handler to ensure the type stays alive until the menu is closed
@@ -1089,6 +1271,7 @@ void FPasteContextMenu::PopulateMenu(FMenuBuilder& MenuBuilder)
 	}
 }
 
+
 void FPasteContextMenu::AddPasteMenuForTrackType(FMenuBuilder& MenuBuilder, int32 DestinationIndex)
 {
 	// Copy a reference to the context menu by value into each lambda handler to ensure the type stays alive until the menu is closed
@@ -1107,6 +1290,7 @@ void FPasteContextMenu::AddPasteMenuForTrackType(FMenuBuilder& MenuBuilder, int3
 	}
 }
 
+
 bool FPasteContextMenu::AutoPaste()
 {
 	if (PasteDestinations.Num() == 1)
@@ -1124,6 +1308,7 @@ bool FPasteContextMenu::AutoPaste()
 	return false;
 }
 
+
 void FPasteContextMenu::PasteInto(int32 DestinationIndex, FName KeyAreaName)
 {
 	FSequencerClipboardReconciler& Reconciler = PasteDestinations[DestinationIndex].Reconcilers[KeyAreaName];
@@ -1161,6 +1346,7 @@ void FPasteContextMenu::PasteInto(int32 DestinationIndex, FName KeyAreaName)
 	}
 }
 
+
 bool FPasteFromHistoryContextMenu::BuildMenu(FMenuBuilder& MenuBuilder, FSequencer& InSequencer, const FPasteContextMenuArgs& Args)
 {
 	if (InSequencer.GetClipboardStack().Num() == 0)
@@ -1173,6 +1359,7 @@ bool FPasteFromHistoryContextMenu::BuildMenu(FMenuBuilder& MenuBuilder, FSequenc
 	return true;
 }
 
+
 TSharedPtr<FPasteFromHistoryContextMenu> FPasteFromHistoryContextMenu::CreateMenu(FSequencer& InSequencer, const FPasteContextMenuArgs& Args)
 {
 	if (InSequencer.GetClipboardStack().Num() == 0)
@@ -1183,6 +1370,7 @@ TSharedPtr<FPasteFromHistoryContextMenu> FPasteFromHistoryContextMenu::CreateMen
 	return MakeShareable(new FPasteFromHistoryContextMenu(InSequencer, Args));
 }
 
+
 void FPasteFromHistoryContextMenu::PopulateMenu(FMenuBuilder& MenuBuilder)
 {
 	// Copy a reference to the context menu by value into each lambda handler to ensure the type stays alive until the menu is closed
@@ -1213,4 +1401,5 @@ void FPasteFromHistoryContextMenu::PopulateMenu(FMenuBuilder& MenuBuilder)
 	MenuBuilder.EndSection();
 }
 
-#undef LOCTEXT_NAMESPACE
\ No newline at end of file
+
+#undef LOCTEXT_NAMESPACE
diff --git a/Engine/Source/Editor/Sequencer/Private/SequencerContextMenus.h b/Engine/Source/Editor/Sequencer/Private/SequencerContextMenus.h
--- a/Engine/Source/Editor/Sequencer/Private/SequencerContextMenus.h
+++ b/Engine/Source/Editor/Sequencer/Private/SequencerContextMenus.h
@@ -33,10 +33,18 @@ private:
 	/** Add extrapolation menu for pre and post infinity */
 	void AddExtrapolationMenu(FMenuBuilder& MenuBuilder, bool bPreInfinity);
 
+	/** Add the Properties sub-menu. */
+	void AddPropertiesMenu(FMenuBuilder& MenuBuilder);
+
+	/** Add the Order sub-menu. */
+	void AddOrderMenu(FMenuBuilder& MenuBuilder);
+
 	void SelectAllKeys();
 
 	bool CanSelectAllKeys() const;
 
+	bool CanSetExtrapolationMode() const;
+
 	void TrimSection(bool bTrimLeft);
 
 	void SplitSection();
@@ -191,6 +199,9 @@ private:
 	/** Hidden AsShared() methods to discourage CreateSP delegate use. */
 	using TSharedFromThis::AsShared;
 
+	/** Add the Properties sub-menu. */
+	void AddPropertiesMenu(FMenuBuilder& MenuBuilder);
+
 	void PopulateMenu(FMenuBuilder& MenuBuilder);
 
 	/** The sequencer */
diff --git a/Engine/Source/Editor/Sequencer/Private/SequencerEdMode.cpp b/Engine/Source/Editor/Sequencer/Private/SequencerEdMode.cpp
--- a/Engine/Source/Editor/Sequencer/Private/SequencerEdMode.cpp
+++ b/Engine/Source/Editor/Sequencer/Private/SequencerEdMode.cpp
@@ -150,6 +150,8 @@ void DrawTransformTrack(const FSceneView* View, FPrimitiveDrawInterface* PDI, UM
 			KeyTimes.Add(KeyTime);
 		}
 
+		KeyTimes.Sort([](const float& A, const float& B ) { return A < B; });
+
 		FVector OldKeyPos(0);
 		float OldKeyTime = 0.f;
 		int KeyTimeIndex = 0;
diff --git a/Engine/Source/Editor/Sequencer/Private/SequencerHotspots.h b/Engine/Source/Editor/Sequencer/Private/SequencerHotspots.h
--- a/Engine/Source/Editor/Sequencer/Private/SequencerHotspots.h
+++ b/Engine/Source/Editor/Sequencer/Private/SequencerHotspots.h
@@ -11,9 +11,8 @@
 struct FKeyHotspot
 	: ISequencerHotspot
 {
-	FKeyHotspot(FSequencerSelectedKey InKey, TSharedPtr<FSequencerTrackNode> InTrackNode)
+	FKeyHotspot(FSequencerSelectedKey InKey)
 		: Key(InKey)
-		, TrackNode(MoveTemp(InTrackNode))
 	{ }
 
 	virtual ESequencerHotspot GetType() const override { return ESequencerHotspot::Key; }
@@ -22,9 +21,6 @@ struct FKeyHotspot
 
 	/** The key itself */
 	FSequencerSelectedKey Key;
-
-	/** The track node that holds the section the key is contained in */
-	TSharedPtr<FSequencerTrackNode> TrackNode;
 };
 
 
diff --git a/Engine/Source/Editor/Sequencer/Private/SequencerInputHandlerStack.h b/Engine/Source/Editor/Sequencer/Private/SequencerInputHandlerStack.h
--- a/Engine/Source/Editor/Sequencer/Private/SequencerInputHandlerStack.h
+++ b/Engine/Source/Editor/Sequencer/Private/SequencerInputHandlerStack.h
@@ -21,6 +21,8 @@ class FSequencerInputHandlerStack
 {
 public:
 
+	FSequencerInputHandlerStack() : CapturedIndex(INDEX_NONE) {}
+
 	/** Add a handler to the stack */
 	void AddHandler(ISequencerInputHandler* Handler) { Handlers.Add(Handler); }
 
diff --git a/Engine/Source/Editor/Sequencer/Private/SequencerNodeTree.cpp b/Engine/Source/Editor/Sequencer/Private/SequencerNodeTree.cpp
--- a/Engine/Source/Editor/Sequencer/Private/SequencerNodeTree.cpp
+++ b/Engine/Source/Editor/Sequencer/Private/SequencerNodeTree.cpp
@@ -4,6 +4,7 @@
 #include "MovieSceneSequence.h"
 #include "MovieSceneSection.h"
 #include "MovieSceneTrack.h"
+#include "MovieSceneCinematicShotTrack.h"
 #include "SequencerNodeTree.h"
 #include "Sequencer.h"
 #include "ScopedTransaction.h"
@@ -12,7 +13,7 @@
 #include "SequencerSectionLayoutBuilder.h"
 #include "ISequencerSection.h"
 #include "ISequencerTrackEditor.h"
-
+#include "SequencerSpacerNode.h"
 
 void FSequencerNodeTree::Empty()
 {
@@ -26,10 +27,14 @@ void FSequencerNodeTree::Empty()
 
 void FSequencerNodeTree::Update()
 {
+	HoveredNode = nullptr;
+
 	// @todo Sequencer - This update pass is too aggressive.  Some nodes may still be valid
 	Empty();
 
 	UMovieScene* MovieScene = Sequencer.GetFocusedMovieSceneSequence()->GetMovieScene();
+	UMovieSceneCinematicShotTrack* CinematicShotTrack = MovieScene->FindMasterTrack<UMovieSceneCinematicShotTrack>();
+
 	TArray<TSharedRef<FSequencerDisplayNode>> NewRootNodes;
 
 	// Get the master tracks  so we can get sections from them
@@ -37,12 +42,15 @@ void FSequencerNodeTree::Update()
 
 	for (UMovieSceneTrack* Track : MasterTracks)
 	{
-		UMovieSceneTrack& TrackRef = *Track;
+		if (Track != CinematicShotTrack)
+		{
+			UMovieSceneTrack& TrackRef = *Track;
 
-		TSharedRef<FSequencerTrackNode> SectionNode = MakeShareable(new FSequencerTrackNode(TrackRef, *FindOrAddTypeEditor(TrackRef), nullptr, *this));
-		NewRootNodes.Add(SectionNode);
+			TSharedRef<FSequencerTrackNode> SectionNode = MakeShareable(new FSequencerTrackNode(TrackRef, *FindOrAddTypeEditor(TrackRef), nullptr, *this));
+			NewRootNodes.Add(SectionNode);
 	
-		MakeSectionInterfaces(TrackRef, SectionNode);
+			MakeSectionInterfaces(TrackRef, SectionNode);
+		}
 	}
 
 	const TArray<FMovieSceneBinding>& Bindings = MovieScene->GetBindings();
@@ -102,33 +110,51 @@ void FSequencerNodeTree::Update()
 
 	NewRootNodes.Append(NewObjectNodes);
 
-	// Look for a shot track.  It will always come first if it exists
-	UMovieSceneTrack* ShotTrack = MovieScene->GetShotTrack();
+	// Cinematic shot track always comes first
+	if (CinematicShotTrack)
+	{
+		TSharedRef<FSequencerTrackNode> SectionNode = MakeShareable(new FSequencerTrackNode(*CinematicShotTrack, *FindOrAddTypeEditor(*CinematicShotTrack), nullptr, *this));
+
+		RootNodes.Add(SectionNode);
+		MakeSectionInterfaces(*CinematicShotTrack, SectionNode);
+	}
+
+	// Then comes the camera cut track
+	UMovieSceneTrack* CameraCutTrack = MovieScene->GetCameraCutTrack();
 	
-	if(ShotTrack)
+	if (CameraCutTrack)
 	{
-		TSharedRef<FSequencerTrackNode> SectionNode = MakeShareable(new FSequencerTrackNode(*ShotTrack, *FindOrAddTypeEditor(*ShotTrack), nullptr, *this));
+		TSharedRef<FSequencerTrackNode> SectionNode = MakeShareable(new FSequencerTrackNode(*CameraCutTrack, *FindOrAddTypeEditor(*CameraCutTrack), nullptr, *this));
 
-		// Shot track always comes first
 		RootNodes.Add(SectionNode);
-		MakeSectionInterfaces(*ShotTrack, SectionNode);
+		MakeSectionInterfaces(*CameraCutTrack, SectionNode);
 	}
 
-	// Add all other nodes after the shot track
+	// Add all other nodes after the camera cut track
 	RootNodes.Append(NewRootNodes);
 
-	// Set up virtual offsets, and expansion states
+
+	RootNodes.Reserve(RootNodes.Num()*2);
+	for (int32 Index = 0; Index < RootNodes.Num(); Index += 2)
+	{
+		RootNodes.Insert(MakeShareable(new FSequencerSpacerNode(2.f, nullptr, *this)), Index);
+	}
+	RootNodes.Add(MakeShareable(new FSequencerSpacerNode(2.f, nullptr, *this)));
+
+	// Set up virtual offsets, expansion states, and tints
 	float VerticalOffset = 0.f;
 
-	for (auto& Node : RootNodes)
+	for (TSharedRef<FSequencerDisplayNode>& Node : RootNodes)
 	{
 		Node->Traverse_ParentFirst([&](FSequencerDisplayNode& InNode) {
+
+			// Set up the virtual node position
 			float VerticalTop = VerticalOffset;
 			VerticalOffset += InNode.GetNodeHeight() + InNode.GetNodePadding().Combined();
 			InNode.Initialize(VerticalTop, VerticalOffset);
+
 			return true;
 		});
-
 	}
 
 	// Re-filter the tree after updating 
@@ -307,6 +333,18 @@ bool FSequencerNodeTree::IsNodeFiltered( const TSharedRef<const FSequencerDispla
 	return FilteredNodes.Contains( Node );
 }
 
+void FSequencerNodeTree::SetHoveredNode(const TSharedPtr<FSequencerDisplayNode>& InHoveredNode)
+{
+	if (InHoveredNode != HoveredNode)
+	{
+		HoveredNode = InHoveredNode;
+	}
+}
+
+const TSharedPtr<FSequencerDisplayNode>& FSequencerNodeTree::GetHoveredNode() const
+{
+	return HoveredNode;
+}
 
 /**
  * Recursively filters nodes
diff --git a/Engine/Source/Editor/Sequencer/Private/SequencerNodeTree.h b/Engine/Source/Editor/Sequencer/Private/SequencerNodeTree.h
--- a/Engine/Source/Editor/Sequencer/Private/SequencerNodeTree.h
+++ b/Engine/Source/Editor/Sequencer/Private/SequencerNodeTree.h
@@ -79,6 +79,16 @@ public:
 	 */
 	bool GetDefaultExpansionState( const FSequencerDisplayNode& Node ) const;
 
+	/**
+	 * Set the single hovered node in the tree
+	 */
+	void SetHoveredNode(const TSharedPtr<FSequencerDisplayNode>& InHoveredNode);
+
+	/**
+	 * Get the single hovered node in the tree, possibly nullptr
+	 */
+	const TSharedPtr<FSequencerDisplayNode>& GetHoveredNode() const;
+
 private:
 	/**
 	 * Finds or adds a type editor for the track
@@ -109,6 +119,8 @@ private:
 	TMap< FGuid, TSharedPtr<FSequencerObjectBindingNode> > ObjectBindingMap;
 	/** Set of all filtered nodes */
 	TSet< TSharedRef<const FSequencerDisplayNode> > FilteredNodes;
+	/** Cardinal hovered node */
+	TSharedPtr<FSequencerDisplayNode> HoveredNode;
 	/** Active filter string if any */
 	FString FilterString;
 	/** Sequencer interface */
diff --git a/Engine/Source/Editor/Sequencer/Private/SequencerSelectedKey.h b/Engine/Source/Editor/Sequencer/Private/SequencerSelectedKey.h
--- a/Engine/Source/Editor/Sequencer/Private/SequencerSelectedKey.h
+++ b/Engine/Source/Editor/Sequencer/Private/SequencerSelectedKey.h
@@ -2,9 +2,12 @@
 
 #pragma once
 
-class UMovieSceneSection;
+#include "MovieSceneSection.h"
+
+
 class IKeyArea;
 
+
 /**
  * Represents a selected key in the sequencer
  */
@@ -12,38 +15,42 @@ struct FSequencerSelectedKey
 {	
 	/** Section that the key belongs to */
 	UMovieSceneSection* Section;
+
 	/** Key area providing the key */
 	TSharedPtr<IKeyArea> KeyArea;
+
 	/** Index of the key in the key area */
 	TOptional<FKeyHandle> KeyHandle;
 
-	FSequencerSelectedKey( UMovieSceneSection& InSection, TSharedPtr<IKeyArea> InKeyArea, FKeyHandle InKeyHandle )
-		: Section( &InSection )
-		, KeyArea( InKeyArea )
+public:
+
+	/** Create and initialize a new instance. */
+	FSequencerSelectedKey(UMovieSceneSection& InSection, TSharedPtr<IKeyArea> InKeyArea, FKeyHandle InKeyHandle)
+		: Section(&InSection)
+		, KeyArea(InKeyArea)
 		, KeyHandle(InKeyHandle)
 	{}
 
+	/** Default constructor. */
 	FSequencerSelectedKey()
-		: Section( NULL )
-		, KeyArea( NULL )
+		: Section(nullptr)
+		, KeyArea(nullptr)
 		, KeyHandle()
 	{}
 
 	/** @return Whether or not this is a valid selected key */
-	bool IsValid() const { return Section != NULL && KeyArea.IsValid() && KeyHandle.IsSet(); }
+	bool IsValid() const { return Section != nullptr && KeyArea.IsValid() && KeyHandle.IsSet(); }
 
-	friend uint32 GetTypeHash( const FSequencerSelectedKey& SelectedKey )
+	friend uint32 GetTypeHash(const FSequencerSelectedKey& SelectedKey)
 	{
-		return GetTypeHash( SelectedKey.Section ) ^ GetTypeHash( SelectedKey.KeyArea ) ^ 
+		return GetTypeHash(SelectedKey.Section) ^ GetTypeHash(SelectedKey.KeyArea) ^ 
 			(SelectedKey.KeyHandle.IsSet() ? GetTypeHash(SelectedKey.KeyHandle.GetValue()) : 0);
 	} 
 
-	bool operator==( const FSequencerSelectedKey& OtherKey ) const
+	bool operator==(const FSequencerSelectedKey& OtherKey) const
 	{
 		return Section == OtherKey.Section && KeyArea == OtherKey.KeyArea &&
 			KeyHandle.IsSet() && OtherKey.KeyHandle.IsSet() &&
 			KeyHandle.GetValue() == OtherKey.KeyHandle.GetValue();
 	}
-
 };
-
diff --git a/Engine/Source/Editor/Sequencer/Private/SequencerSettings.cpp b/Engine/Source/Editor/Sequencer/Private/SequencerSettings.cpp
--- a/Engine/Source/Editor/Sequencer/Private/SequencerSettings.cpp
+++ b/Engine/Source/Editor/Sequencer/Private/SequencerSettings.cpp
@@ -20,6 +20,7 @@ USequencerSettings::USequencerSettings( const FObjectInitializer& ObjectInitiali
 	bSnapSectionTimesToSections = true;
 	bSnapPlayTimeToInterval = true;
 	bSnapPlayTimeToDraggedKey = false;
+	bFixedTimeStepPlayback = false;
 	CurveValueSnapInterval = 10.0f;
 	bSnapCurveValueToInterval = true;
 	bDetailsViewVisible = false;
@@ -30,6 +31,7 @@ USequencerSettings::USequencerSettings( const FObjectInitializer& ObjectInitiali
 	bLooping = false;
 	bKeepCursorInPlayRange = true;
 	bKeepPlayRangeInSectionBounds = true;
+	ZeroPadFrames = 0;
 }
 
 EAutoKeyMode USequencerSettings::GetAutoKeyMode() const
@@ -116,6 +118,20 @@ void USequencerSettings::SetShowFrameNumbers(bool InbShowFrameNumbers)
 	}
 }
 
+bool USequencerSettings::GetShowGotoBox() const
+{
+	return bShowGotoBox;
+}
+
+void USequencerSettings::SetShowGotoBox(bool InbShowGotoBox)
+{
+	if ( bShowGotoBox != InbShowGotoBox )
+	{
+		bShowGotoBox = InbShowGotoBox;
+		SaveConfig();
+	}
+}
+
 bool USequencerSettings::GetShowRangeSlider() const
 {
 	return bShowRangeSlider;
@@ -242,6 +258,20 @@ void USequencerSettings::SetSnapPlayTimeToDraggedKey(bool InbSnapPlayTimeToDragg
 	}
 }
 
+bool USequencerSettings::GetFixedTimeStepPlayback() const
+{
+	return bFixedTimeStepPlayback;
+}
+
+void USequencerSettings::SetFixedTimeStepPlayback(bool InbFixedTimeStepPlayback)
+{
+	if ( bFixedTimeStepPlayback != InbFixedTimeStepPlayback )
+	{
+		bFixedTimeStepPlayback = InbFixedTimeStepPlayback;
+		SaveConfig();
+	}
+}
+
 float USequencerSettings::GetCurveValueSnapInterval() const
 {
 	return CurveValueSnapInterval;
@@ -384,6 +414,21 @@ void USequencerSettings::SetShowCurveEditorCurveToolTips(bool InbShowCurveEditor
 	}
 }
 
+
+uint8 USequencerSettings::GetZeroPadFrames() const
+{
+	return ZeroPadFrames;
+}
+
+void USequencerSettings::SetZeroPadFrames(uint8 InZeroPadFrames)
+{
+	if (ZeroPadFrames != InZeroPadFrames)
+	{
+		ZeroPadFrames = InZeroPadFrames;
+		SaveConfig();
+	}
+}
+
 float USequencerSettings::SnapTimeToInterval( float InTimeValue ) const
 {
 	return TimeSnapInterval > 0
@@ -404,4 +449,5 @@ ULevelEditorSequencerSettings::ULevelEditorSequencerSettings( const FObjectIniti
 	TimeSnapInterval = 0.033334f;
 	bShowRangeSlider = true;
 	bKeepPlayRangeInSectionBounds = false;
+	ZeroPadFrames = 4;
 }
\ No newline at end of file
diff --git a/Engine/Source/Editor/Sequencer/Private/SequencerSettings.h b/Engine/Source/Editor/Sequencer/Private/SequencerSettings.h
--- a/Engine/Source/Editor/Sequencer/Private/SequencerSettings.h
+++ b/Engine/Source/Editor/Sequencer/Private/SequencerSettings.h
@@ -84,6 +84,11 @@ public:
 	/** Sets whether or not to show frame numbers. */
 	void SetShowFrameNumbers(bool InbShowFrameNumbers);
 
+	/** Gets whether or not to show the goto box. */
+	bool GetShowGotoBox() const;
+	/** Sets whether or not to show the goto box. */
+	void SetShowGotoBox(bool InbShowGotoBox);
+
 	/** Gets whether or not to show the time range slider. */
 	bool GetShowRangeSlider() const;
 	/** Sets whether or not to show frame numbers. */
@@ -129,6 +134,11 @@ public:
 	/** Sets whether or not to snap the play time to the dragged key. */
 	void SetSnapPlayTimeToDraggedKey(bool InbSnapPlayTimeToDraggedKey);
 
+	/** Gets whether to lock the play rate to the frame rate. */
+	bool GetFixedTimeStepPlayback() const;
+	/** Sets whether or not to lock the play rate to the frame rate. */
+	void SetFixedTimeStepPlayback(bool InbFTixedTimeStepPlayback);
+
 	/** Gets the snapping interval for curve values. */
 	float GetCurveValueSnapInterval() const;
 	/** Sets the snapping interval for curve values. */
@@ -181,6 +191,11 @@ public:
 	/** Set whether or not the playback range should be synced to the section bounds */
 	void SetKeepPlayRangeInSectionBounds(bool bInKeepPlayRangeInSectionBounds);
 
+	/** Get the number of digits we should zero-pad to when showing frame numbers in sequencer */
+	uint8 GetZeroPadFrames() const;
+	/** Set the number of digits we should zero-pad to when showing frame numbers in sequencer */
+	void SetZeroPadFrames(uint8 InZeroPadFrames);
+
 	/** Snaps a time value in seconds to the currently selected interval. */
 	float SnapTimeToInterval(float InTimeValue) const;
 
@@ -206,6 +221,9 @@ protected:
 	UPROPERTY( config )
 	bool bShowFrameNumbers;
 
+	UPROPERTY( config )
+	bool bShowGotoBox;
+
 	UPROPERTY( config )
 	bool bShowRangeSlider;
 
@@ -233,6 +251,9 @@ protected:
 	UPROPERTY( config )
 	bool bSnapPlayTimeToDraggedKey;
 
+	UPROPERTY( config )
+	bool bFixedTimeStepPlayback;
+
 	UPROPERTY( config )
 	float CurveValueSnapInterval;
 
@@ -263,6 +284,9 @@ protected:
 	UPROPERTY( config )
 	bool bKeepPlayRangeInSectionBounds;
 
+	UPROPERTY( config )
+	uint8 ZeroPadFrames;
+
 	FOnShowCurveEditorChanged OnShowCurveEditorChanged;
 };
 
diff --git a/Engine/Source/Editor/Sequencer/Private/SequencerUtilities.cpp b/Engine/Source/Editor/Sequencer/Private/SequencerUtilities.cpp
new file mode 100644
--- /dev/null
+++ b/Engine/Source/Editor/Sequencer/Private/SequencerUtilities.cpp
@@ -0,0 +1,65 @@
+// Copyright 1998-2015 Epic Games, Inc. All Rights Reserved.
+
+#include "SequencerPrivatePCH.h"
+#include "SequencerUtilities.h"
+
+static EVisibility GetRolloverVisibility(TAttribute<bool> HoverState, TWeakPtr<SComboButton> WeakComboButton)
+{
+	TSharedPtr<SComboButton> ComboButton = WeakComboButton.Pin();
+	if (HoverState.Get() || ComboButton->IsOpen())
+	{
+		return EVisibility::SelfHitTestInvisible;
+	}
+	else
+	{
+		return EVisibility::Collapsed;
+	}
+}
+
+TSharedRef<SWidget> FSequencerUtilities::MakeAddButton(FText HoverText, FOnGetContent MenuContent, const TAttribute<bool>& HoverState)
+{
+	FSlateFontInfo SmallLayoutFont( FPaths::EngineContentDir() / TEXT("Slate/Fonts/Roboto-Regular.ttf"), 8 );
+
+	TSharedRef<STextBlock> ComboButtonText = SNew(STextBlock)
+		.Text(HoverText)
+		.Font(SmallLayoutFont)
+		.ColorAndOpacity( FSlateColor::UseForeground() );
+
+	TSharedRef<SComboButton> ComboButton =
+
+		SNew(SComboButton)
+		.HasDownArrow(false)
+		.ButtonStyle(FEditorStyle::Get(), "HoverHintOnly")
+		.ForegroundColor( FSlateColor::UseForeground() )
+		.OnGetMenuContent(MenuContent)
+		.ContentPadding(FMargin(5, 2))
+		.HAlign(HAlign_Center)
+		.VAlign(VAlign_Center)
+		.ButtonContent()
+		[
+			SNew(SHorizontalBox)
+
+			+ SHorizontalBox::Slot()
+			.AutoWidth()
+			.VAlign(VAlign_Center)
+			.Padding(FMargin(0,0,2,0))
+			[
+				SNew(SImage)
+				.ColorAndOpacity( FSlateColor::UseForeground() )
+				.Image(FEditorStyle::GetBrush("Plus"))
+			]
+
+			+ SHorizontalBox::Slot()
+			.VAlign(VAlign_Center)
+			.AutoWidth()
+			[
+				ComboButtonText
+			]
+		];
+
+	TAttribute<EVisibility> Visibility = TAttribute<EVisibility>::Create(TAttribute<EVisibility>::FGetter::CreateStatic(GetRolloverVisibility, HoverState, TWeakPtr<SComboButton>(ComboButton)));
+	ComboButtonText->SetVisibility(Visibility);
+
+	return ComboButton;
+}
+
diff --git a/Engine/Source/Editor/Sequencer/Private/TimeSliderController.cpp b/Engine/Source/Editor/Sequencer/Private/TimeSliderController.cpp
--- a/Engine/Source/Editor/Sequencer/Private/TimeSliderController.cpp
+++ b/Engine/Source/Editor/Sequencer/Private/TimeSliderController.cpp
@@ -423,7 +423,7 @@ int32 FSequencerTimeSliderController::DrawPlaybackRange(const FGeometry& Allotte
 		FEditorStyle::GetBrush("WhiteBrush"),
 		MyClippingRect,
 		ESlateDrawEffect::None,
-		FLinearColor::Black.CopyWithNewOpacity(0.2f)
+		FLinearColor::Black.CopyWithNewOpacity(0.3f)
 	);
 
 	FSlateDrawElement::MakeBox(
@@ -433,7 +433,7 @@ int32 FSequencerTimeSliderController::DrawPlaybackRange(const FGeometry& Allotte
 		FEditorStyle::GetBrush("WhiteBrush"),
 		MyClippingRect,
 		ESlateDrawEffect::None,
-		FLinearColor::Black.CopyWithNewOpacity(0.2f)
+		FLinearColor::Black.CopyWithNewOpacity(0.3f)
 	);
 	return LayerId + 1;
 }
@@ -533,10 +533,12 @@ FReply FSequencerTimeSliderController::OnMouseButtonUp( SWidget& WidgetOwner, co
 			}
 
 			// Zoom in
+			bool bDoSetRange = false;
 			if (NewValue > DownValue)
 			{
 				// push the current value onto the stack
 				RangeStack.Add(FVector2D(TimeSliderArgs.ViewRange.Get().GetLowerBoundValue(), TimeSliderArgs.ViewRange.Get().GetUpperBoundValue()));
+				bDoSetRange = true;
 			}
 			// Zoom out
 			else if (RangeStack.Num())
@@ -545,14 +547,18 @@ FReply FSequencerTimeSliderController::OnMouseButtonUp( SWidget& WidgetOwner, co
 				FVector2D LastRange = RangeStack.Pop();
 				DownValue = LastRange[0];
 				NewValue = LastRange[1];
+				bDoSetRange = true;
 			}
 
-			TimeSliderArgs.OnViewRangeChanged.ExecuteIfBound(TRange<float>(DownValue, NewValue), EViewRangeInterpolation::Immediate);
+			if (bDoSetRange)
+			{
+				TimeSliderArgs.OnViewRangeChanged.ExecuteIfBound(TRange<float>(DownValue, NewValue), EViewRangeInterpolation::Immediate);
 					
-			if( !TimeSliderArgs.ViewRange.IsBound() )
-			{	
-				// The output is not bound to a delegate so we'll manage the value ourselves
-				TimeSliderArgs.ViewRange.Set( TRange<float>( DownValue, NewValue ) );
+				if( !TimeSliderArgs.ViewRange.IsBound() )
+				{	
+					// The output is not bound to a delegate so we'll manage the value ourselves
+					TimeSliderArgs.ViewRange.Set( TRange<float>( DownValue, NewValue ) );
+				}
 			}
 		}
 		else
diff --git a/Engine/Source/Editor/Sequencer/Private/Tools/EditToolDragOperations.cpp b/Engine/Source/Editor/Sequencer/Private/Tools/EditToolDragOperations.cpp
--- a/Engine/Source/Editor/Sequencer/Private/Tools/EditToolDragOperations.cpp
+++ b/Engine/Source/Editor/Sequencer/Private/Tools/EditToolDragOperations.cpp
@@ -9,7 +9,6 @@
 #include "CommonMovieSceneTools.h"
 #include "MovieScene.h"
 #include "MovieSceneTrack.h"
-#include "MovieSceneShotTrack.h"
 #include "VirtualTrackArea.h"
 
 struct FDefaultKeySnappingCandidates : ISequencerSnapCandidate
diff --git a/Engine/Source/Editor/Sequencer/Private/Tools/SequencerEditTool.h b/Engine/Source/Editor/Sequencer/Private/Tools/SequencerEditTool.h
--- a/Engine/Source/Editor/Sequencer/Private/Tools/SequencerEditTool.h
+++ b/Engine/Source/Editor/Sequencer/Private/Tools/SequencerEditTool.h
@@ -13,6 +13,10 @@ class FSequencerEditTool
 {
 public:
 
+	FSequencerEditTool(FSequencer& InSequencer)
+		: Sequencer(InSequencer)
+	{ }
+
 	/** Virtual destructor. */
 	~FSequencerEditTool() { }
 
@@ -60,30 +64,18 @@ public:
 		return LayerId;
 	}
 
-	virtual void Tick(const FGeometry& AllottedGeometry, const double InCurrentTime, const float InDeltaTime) override
-	{
-		// do nothing
-	}
-
 	virtual FCursorReply OnCursorQuery(const FGeometry& MyGeometry, const FPointerEvent& CursorEvent) const override
 	{
 		return FCursorReply::Unhandled();
 	}
 
-	/** Get the current active hotspot */
-	virtual TSharedPtr<ISequencerHotspot> GetHotspot() const override
-	{
-		return Hotspot;
-	}
-
-	/** Set the hotspot to something else */
-	virtual void SetHotspot(TSharedPtr<ISequencerHotspot> NewHotspot) override
+	virtual ISequencer& GetSequencer() const override
 	{
-		Hotspot = MoveTemp(NewHotspot);
+		return Sequencer;
 	}
 
 protected:
 
-	/** The current hotspot that can be set from anywhere to initiate drags */
-	TSharedPtr<ISequencerHotspot> Hotspot;
+	/** This edit tool's sequencer */
+	FSequencer& Sequencer;
 };
diff --git a/Engine/Source/Editor/Sequencer/Private/Tools/SequencerEditTool_Movement.cpp b/Engine/Source/Editor/Sequencer/Private/Tools/SequencerEditTool_Movement.cpp
--- a/Engine/Source/Editor/Sequencer/Private/Tools/SequencerEditTool_Movement.cpp
+++ b/Engine/Source/Editor/Sequencer/Private/Tools/SequencerEditTool_Movement.cpp
@@ -7,21 +7,18 @@
 #include "SequencerHotspots.h"
 #include "EditToolDragOperations.h"
 
+const FName FSequencerEditTool_Movement::Identifier = "Movement";
 
-FSequencerEditTool_Movement::FSequencerEditTool_Movement(TSharedPtr<FSequencer> InSequencer, TSharedPtr<SSequencer> InSequencerWidget)
-	: Sequencer(InSequencer)
-	, SequencerWidget(InSequencerWidget)
+FSequencerEditTool_Movement::FSequencerEditTool_Movement(FSequencer& InSequencer)
+	: FSequencerEditTool(InSequencer)
+	, SequencerWidget(StaticCastSharedRef<SSequencer>(InSequencer.GetSequencerWidget()))
 { }
 
 
-ISequencer& FSequencerEditTool_Movement::GetSequencer() const
-{
-	return *Sequencer.Pin();
-}
-
-
 FReply FSequencerEditTool_Movement::OnMouseButtonDown(SWidget& OwnerWidget, const FGeometry& MyGeometry, const FPointerEvent& MouseEvent)
 {
+	TSharedPtr<ISequencerHotspot> Hotspot = Sequencer.GetHotspot();
+
 	DelayedDrag.Reset();
 
 	if (MouseEvent.GetEffectingButton() == EKeys::LeftMouseButton || MouseEvent.GetEffectingButton() == EKeys::MiddleMouseButton)
@@ -30,14 +27,14 @@ FReply FSequencerEditTool_Movement::OnMouseButtonDown(SWidget& OwnerWidget, cons
 
 		DelayedDrag = FDelayedDrag_Hotspot(VirtualTrackArea.CachedTrackAreaGeometry().AbsoluteToLocal(MouseEvent.GetScreenSpacePosition()), MouseEvent.GetEffectingButton(), Hotspot);
 
- 		if (Sequencer.Pin()->GetSettings()->GetSnapPlayTimeToDraggedKey() || (MouseEvent.IsShiftDown() && MouseEvent.GetEffectingButton() == EKeys::LeftMouseButton) )
+ 		if (Sequencer.GetSettings()->GetSnapPlayTimeToDraggedKey() || (MouseEvent.IsShiftDown() && MouseEvent.GetEffectingButton() == EKeys::LeftMouseButton) )
 		{
 			if (DelayedDrag->Hotspot.IsValid())
 			{
 				if (DelayedDrag->Hotspot->GetType() == ESequencerHotspot::Key)
 				{
 					FSequencerSelectedKey& ThisKey = StaticCastSharedPtr<FKeyHotspot>(DelayedDrag->Hotspot)->Key;
-					Sequencer.Pin()->SetGlobalTime(ThisKey.KeyArea->GetKeyTime(ThisKey.KeyHandle.GetValue()));
+					Sequencer.SetGlobalTime(ThisKey.KeyArea->GetKeyTime(ThisKey.KeyHandle.GetValue()));
 				}
 			}
 		}
@@ -85,13 +82,12 @@ FReply FSequencerEditTool_Movement::OnMouseMove(SWidget& OwnerWidget, const FGeo
 
 TSharedPtr<ISequencerEditToolDragOperation> FSequencerEditTool_Movement::CreateDrag(const FPointerEvent& MouseEvent)
 {
-	auto PinnedSequencer = Sequencer.Pin();
-	FSequencerSelection& Selection = PinnedSequencer->GetSelection();
+	FSequencerSelection& Selection = Sequencer.GetSelection();
 
 	if (DelayedDrag->Hotspot.IsValid())
 	{
 		// Let the hotspot start a drag first, if it wants to
-		auto HotspotDrag = DelayedDrag->Hotspot->InitiateDrag(*PinnedSequencer);
+		auto HotspotDrag = DelayedDrag->Hotspot->InitiateDrag(Sequencer);
 		if (HotspotDrag.IsValid())
 		{
 			return HotspotDrag;
@@ -118,7 +114,7 @@ TSharedPtr<ISequencerEditToolDragOperation> FSequencerEditTool_Movement::CreateD
 				Selection.AddToSelection(ThisSection);
 				SectionHandles.Add(ThisHandle);
 			}
-			return MakeShareable( new FMoveSection( *PinnedSequencer, SectionHandles ) );
+			return MakeShareable( new FMoveSection( Sequencer, SectionHandles ) );
 		}
 		// Moving key(s)?
 		else if (HotspotType == ESequencerHotspot::Key)
@@ -136,22 +132,22 @@ TSharedPtr<ISequencerEditToolDragOperation> FSequencerEditTool_Movement::CreateD
 			// @todo sequencer: Make this a customizable UI command modifier?
 			if (MouseEvent.IsAltDown())
 			{
-				return MakeShareable( new FDuplicateKeys( *PinnedSequencer, Selection.GetSelectedKeys() ) );
+				return MakeShareable( new FDuplicateKeys( Sequencer, Selection.GetSelectedKeys() ) );
 			}
 			else
 			{
-				return MakeShareable( new FMoveKeys( *PinnedSequencer, Selection.GetSelectedKeys() ) );
+				return MakeShareable( new FMoveKeys( Sequencer, Selection.GetSelectedKeys() ) );
 			}
 		}
 	}
 	// If we're not dragging a hotspot, sections take precedence over keys
 	else if (Selection.GetSelectedSections().Num())
 	{
-		return MakeShareable( new FMoveSection( *PinnedSequencer, SequencerWidget.Pin()->GetSectionHandles(Selection.GetSelectedSections()) ) );
+		return MakeShareable( new FMoveSection( Sequencer, SequencerWidget.Pin()->GetSectionHandles(Selection.GetSelectedSections()) ) );
 	}
 	else if (Selection.GetSelectedKeys().Num())
 	{
-		return MakeShareable( new FMoveKeys( *PinnedSequencer, Selection.GetSelectedKeys() ) );
+		return MakeShareable( new FMoveKeys( Sequencer, Selection.GetSelectedKeys() ) );
 	}
 
 	return nullptr;
@@ -171,17 +167,11 @@ FReply FSequencerEditTool_Movement::OnMouseButtonUp(SWidget& OwnerWidget, const
 		return FReply::Handled().ReleaseMouseCapture();
 	}
 
-	FReply Reply = FReply::Unhandled();
-
-	if (Hotspot.IsValid())
-	{
-		PerformHotspotSelection(MouseEvent);
-		Reply = FReply::Handled();
-	}
+	SequencerHelpers::PerformDefaultSelection(Sequencer, MouseEvent);
 
 	if (MouseEvent.GetEffectingButton() == EKeys::RightMouseButton)
 	{
-		TSharedPtr<SWidget> MenuContent = OnSummonContextMenu( MyGeometry, MouseEvent );
+		TSharedPtr<SWidget> MenuContent = SequencerHelpers::SummonContextMenu( Sequencer, MyGeometry, MouseEvent );
 		if (MenuContent.IsValid())
 		{
 			FWidgetPath WidgetPath = MouseEvent.GetEventPath() != nullptr ? *MouseEvent.GetEventPath() : FWidgetPath();
@@ -198,12 +188,7 @@ FReply FSequencerEditTool_Movement::OnMouseButtonUp(SWidget& OwnerWidget, const
 		}
 	}
 
-	if (!Reply.IsEventHandled())
-	{
-		Reply = FSequencerEditTool_Default::OnMouseButtonUp(OwnerWidget, MyGeometry, MouseEvent);
-	}
-
-	return Reply;
+	return FReply::Handled();
 }
 
 
@@ -216,9 +201,10 @@ void FSequencerEditTool_Movement::OnMouseCaptureLost()
 
 FCursorReply FSequencerEditTool_Movement::OnCursorQuery(const FGeometry& MyGeometry, const FPointerEvent& CursorEvent) const
 {
-	if (DragOperation.IsValid())
+	TSharedPtr<ISequencerHotspot> Hotspot = Sequencer.GetHotspot();
+	if (DelayedDrag.IsSet())
 	{
-		return FCursorReply::Cursor(EMouseCursor::CardinalCross);
+		Hotspot = DelayedDrag->Hotspot;
 	}
 
 	if (Hotspot.IsValid())
@@ -229,12 +215,16 @@ FCursorReply FSequencerEditTool_Movement::OnCursorQuery(const FGeometry& MyGeome
 			return Reply;
 		}
 	}
-	return FCursorReply::Cursor(EMouseCursor::Default);
-}
 
+	return FCursorReply::Cursor(EMouseCursor::CardinalCross);
+}
 
 FName FSequencerEditTool_Movement::GetIdentifier() const
 {
-	static FName Identifier("Movement");
 	return Identifier;
 }
+
+bool FSequencerEditTool_Movement::CanDeactivate() const
+{
+	return !DelayedDrag.IsSet();
+}
\ No newline at end of file
diff --git a/Engine/Source/Editor/Sequencer/Private/Tools/SequencerEditTool_Movement.h b/Engine/Source/Editor/Sequencer/Private/Tools/SequencerEditTool_Movement.h
--- a/Engine/Source/Editor/Sequencer/Private/Tools/SequencerEditTool_Movement.h
+++ b/Engine/Source/Editor/Sequencer/Private/Tools/SequencerEditTool_Movement.h
@@ -2,7 +2,7 @@
 
 #pragma once
 
-#include "SequencerEditTool_Default.h"
+#include "SequencerEditTool.h"
 #include "DelayedDrag.h"
 
 
@@ -10,12 +10,15 @@ class ISequencerEditToolDragOperation;
 
 
 class FSequencerEditTool_Movement
-	: public FSequencerEditTool_Default
+	: public FSequencerEditTool
 {
 public:
 
+	/** Static identifier for this edit tool */
+	static const FName Identifier;
+
 	/** Create and initialize a new instance. */
-	FSequencerEditTool_Movement(TSharedPtr<FSequencer> InSequencer, TSharedPtr<SSequencer> InSequencerWidget);
+	FSequencerEditTool_Movement(FSequencer& InSequencer);
 
 public:
 
@@ -27,15 +30,12 @@ public:
 	virtual void OnMouseCaptureLost() override;
 	virtual FCursorReply OnCursorQuery(const FGeometry& MyGeometry, const FPointerEvent& CursorEvent) const override;
 	virtual FName GetIdentifier() const override;
-	virtual ISequencer& GetSequencer() const override;
+	virtual bool CanDeactivate() const override;
 
 private:
 
 	TSharedPtr<ISequencerEditToolDragOperation> CreateDrag(const FPointerEvent& MouseEvent);
 
-	/** The sequencer itself */
-	TWeakPtr<FSequencer> Sequencer;
-
 	/** Sequencer widget */
 	TWeakPtr<SSequencer> SequencerWidget;
 
diff --git a/Engine/Source/Editor/Sequencer/Private/Tools/SequencerEditTool_Selection.cpp b/Engine/Source/Editor/Sequencer/Private/Tools/SequencerEditTool_Selection.cpp
--- a/Engine/Source/Editor/Sequencer/Private/Tools/SequencerEditTool_Selection.cpp
+++ b/Engine/Source/Editor/Sequencer/Private/Tools/SequencerEditTool_Selection.cpp
@@ -60,9 +60,9 @@ class FMarqueeDragOperation
 {
 public:
 
-	FMarqueeDragOperation(TWeakPtr<FSequencer> InSequencer, TWeakPtr<SSequencer> InSequencerWidget)
-		: Sequencer(MoveTemp(InSequencer))
-		, SequencerWidget(MoveTemp(InSequencerWidget))
+	FMarqueeDragOperation(FSequencer& InSequencer)
+		: Sequencer(InSequencer)
+		, SequencerWidget(StaticCastSharedRef<SSequencer>(InSequencer.GetSequencerWidget()))
 		, PreviewState(ESelectionPreviewState::Selected)
 	{}
 
@@ -94,7 +94,7 @@ public:
 			PreviewState = ESelectionPreviewState::Selected;
 
 			// @todo: selection in transactions
-			Sequencer.Pin()->GetSelection().Empty();
+			Sequencer.GetSelection().Empty();
 		}
 	}
 
@@ -102,7 +102,6 @@ public:
 	{
 		// hange the current marquee selection
 		const FVector2D MouseDelta = MouseEvent.GetCursorDelta();
-		auto PinnedSequencer = Sequencer.Pin();
 
 		// Handle virtual scrolling when at the vertical extremes of the widget (performed before we clamp the mouse pos)
 		{
@@ -111,13 +110,13 @@ public:
 			float Difference = LocalMousePos.Y - ScrollThresholdV;
 			if (Difference < 0 && MouseDelta.Y < 0)
 			{
-				PinnedSequencer->VerticalScroll( Difference * 0.1f );
+				Sequencer.VerticalScroll( Difference * 0.1f );
 			}
 
 			Difference = LocalMousePos.Y - (VirtualTrackArea.GetPhysicalSize().Y - ScrollThresholdV);
 			if (Difference > 0 && MouseDelta.Y > 0)
 			{
-				PinnedSequencer->VerticalScroll( Difference * 0.1f );
+				Sequencer.VerticalScroll( Difference * 0.1f );
 			}
 		}
 
@@ -129,7 +128,7 @@ public:
 		CurrentMousePos = LocalMousePos;
 		CurrentMousePos.X = FMath::Clamp(CurrentMousePos.X, 0.f, VirtualTrackArea.GetPhysicalSize().X);
 
-		TRange<float> ViewRange = PinnedSequencer->GetViewRange();
+		TRange<float> ViewRange = Sequencer.GetViewRange();
 
 		// Handle virtual scrolling when at the horizontal extremes of the widget
 		{
@@ -138,18 +137,18 @@ public:
 			float Difference = CurrentPosition.X - (ViewRange.GetLowerBoundValue() + ScrollThresholdH);
 			if (Difference < 0 && MouseDelta.X < 0)
 			{
-				PinnedSequencer->StartAutoscroll(Difference);
+				Sequencer.StartAutoscroll(Difference);
 			}
 			else
 			{
 				Difference = CurrentPosition.X - (ViewRange.GetUpperBoundValue() - ScrollThresholdH);
 				if (Difference > 0 && MouseDelta.X > 0)
 				{
-					PinnedSequencer->StartAutoscroll(Difference);
+					Sequencer.StartAutoscroll(Difference);
 				}
 				else
 				{
-					PinnedSequencer->StopAutoscroll();
+					Sequencer.StopAutoscroll();
 				}
 			}
 		}
@@ -161,24 +160,23 @@ public:
 		VirtualKeySize.Y = SequencerSectionConstants::KeySize.Y;
 
 		// Visit everything using the preview selection primarily as well as the 
-		auto& SelectionPreview = PinnedSequencer->GetSelectionPreview();
+		auto& SelectionPreview = Sequencer.GetSelectionPreview();
 
 		// Ensure the preview is empty before calculating the intersection
 		SelectionPreview.Empty();
 
-		const auto& RootNodes = SequencerWidget.Pin()->GetTreeView()->GetNodeTree()->GetRootNodes();
+		const auto& RootNodes = SequencerWidget->GetTreeView()->GetNodeTree()->GetRootNodes();
 
 		// Now walk everything within the current marquee range, setting preview selection states as we go
 		FSequencerEntityWalker Walker(FSequencerEntityRange(TopLeft(), BottomRight()), VirtualKeySize);
-		Walker.Traverse(FSelectionPreviewVisitor(SelectionPreview, PinnedSequencer->GetSelection(), PreviewState), RootNodes);
+		Walker.Traverse(FSelectionPreviewVisitor(SelectionPreview, Sequencer.GetSelection(), PreviewState), RootNodes);
 	}
 
 	virtual void OnEndDrag(const FPointerEvent& MouseEvent, FVector2D LocalMousePos, const FVirtualTrackArea& VirtualTrackArea) override
 	{
 		// finish dragging the marquee selection
-		auto PinnedSequencer = Sequencer.Pin();
-		auto& Selection = PinnedSequencer->GetSelection();
-		auto& SelectionPreview = PinnedSequencer->GetSelectionPreview();
+		auto& Selection = Sequencer.GetSelection();
+		auto& SelectionPreview = Sequencer.GetSelectionPreview();
 
 		// Patch everything from the selection preview into the actual selection
 		for (const auto& Pair : SelectionPreview.GetDefinedKeyStates())
@@ -215,7 +213,7 @@ public:
 	virtual int32 OnPaint(const FGeometry& AllottedGeometry, const FSlateRect& MyClippingRect, FSlateWindowElementList& OutDrawElements, int32 LayerId) const
 	{
 		// convert to physical space for rendering
-		const FVirtualTrackArea VirtualTrackArea = SequencerWidget.Pin()->GetVirtualTrackArea();
+		const FVirtualTrackArea VirtualTrackArea = SequencerWidget->GetVirtualTrackArea();
 
 		FVector2D SelectionTopLeft = VirtualTrackArea.VirtualToPhysical(TopLeft());
 		FVector2D SelectionBottomRight = VirtualTrackArea.VirtualToPhysical(BottomRight());
@@ -250,10 +248,10 @@ private:
 	}
 
 	/** The sequencer itself */
-	TWeakPtr<FSequencer> Sequencer;
+	FSequencer& Sequencer;
 
 	/** Sequencer widget */
-	TWeakPtr<SSequencer> SequencerWidget;
+	TSharedRef<SSequencer> SequencerWidget;
 
 	/** Whether we should select/deselect things in this marquee operation */
 	ESelectionPreviewState PreviewState;
@@ -264,19 +262,15 @@ private:
 };
 
 
-FSequencerEditTool_Selection::FSequencerEditTool_Selection(TSharedPtr<FSequencer> InSequencer, TSharedPtr<SSequencer> InSequencerWidget)
-	: Sequencer(InSequencer)
-	, SequencerWidget(InSequencerWidget)
+const FName FSequencerEditTool_Selection::Identifier = "Selection";
+
+FSequencerEditTool_Selection::FSequencerEditTool_Selection(FSequencer& InSequencer)
+	: FSequencerEditTool(InSequencer)
+	, SequencerWidget(StaticCastSharedRef<SSequencer>(InSequencer.GetSequencerWidget()))
 	, CursorDecorator(nullptr)
 { }
 
 
-ISequencer& FSequencerEditTool_Selection::GetSequencer() const
-{
-	return *Sequencer.Pin();
-}
-
-
 FCursorReply FSequencerEditTool_Selection::OnCursorQuery(const FGeometry& MyGeometry, const FPointerEvent& CursorEvent) const
 {
 	return FCursorReply::Cursor(EMouseCursor::Crosshairs);
@@ -313,7 +307,7 @@ FReply FSequencerEditTool_Selection::OnMouseButtonDown(SWidget& OwnerWidget, con
 
 	if (MouseEvent.GetEffectingButton() == EKeys::LeftMouseButton)
 	{
-		DelayedDrag = FDelayedDrag_Hotspot(MyGeometry.AbsoluteToLocal(MouseEvent.GetScreenSpacePosition()), EKeys::LeftMouseButton, Hotspot);
+		DelayedDrag = FDelayedDrag_Hotspot(MyGeometry.AbsoluteToLocal(MouseEvent.GetScreenSpacePosition()), EKeys::LeftMouseButton, Sequencer.GetHotspot());
 		return FReply::Handled();
 	}
 	return FReply::Unhandled();
@@ -343,19 +337,18 @@ FReply FSequencerEditTool_Selection::OnMouseMove(SWidget& OwnerWidget, const FGe
 				auto HotspotType = DelayedDrag->Hotspot->GetType();
 				if (HotspotType != ESequencerHotspot::Section && HotspotType != ESequencerHotspot::Key)
 				{
-					DragOperation = DelayedDrag->Hotspot->InitiateDrag(*Sequencer.Pin());
+					DragOperation = DelayedDrag->Hotspot->InitiateDrag(Sequencer);
 				}
 			}
 
 			if (!DragOperation.IsValid())
 			{
-				DragOperation = MakeShareable( new FMarqueeDragOperation(Sequencer, SequencerWidget) );
+				DragOperation = MakeShareable( new FMarqueeDragOperation(Sequencer) );
 			}
 
 			if (DragOperation.IsValid())
 			{
-				FVector2D LocalPosition = MyGeometry.AbsoluteToLocal(MouseEvent.GetScreenSpacePosition());
-				DragOperation->OnBeginDrag(MouseEvent, LocalPosition, VirtualTrackArea);
+				DragOperation->OnBeginDrag(MouseEvent, DelayedDrag->GetInitialPosition(), VirtualTrackArea);
 
 				// Steal the capture, as we're now the authoritative widget in charge of a mouse-drag operation
 				Reply.CaptureMouse(OwnerWidget.AsShared());
@@ -381,11 +374,34 @@ FReply FSequencerEditTool_Selection::OnMouseButtonUp(SWidget& OwnerWidget, const
 
 		CursorDecorator = nullptr;
 
-		Sequencer.Pin()->StopAutoscroll();
+		Sequencer.StopAutoscroll();
 		return FReply::Handled().ReleaseMouseCapture();
 	}
+	else
+	{
+		SequencerHelpers::PerformDefaultSelection(Sequencer, MouseEvent);
+
+		if (MouseEvent.GetEffectingButton() == EKeys::RightMouseButton)
+		{
+			TSharedPtr<SWidget> MenuContent = SequencerHelpers::SummonContextMenu( Sequencer, MyGeometry, MouseEvent );
+			if (MenuContent.IsValid())
+			{
+				FWidgetPath WidgetPath = MouseEvent.GetEventPath() != nullptr ? *MouseEvent.GetEventPath() : FWidgetPath();
+
+				FSlateApplication::Get().PushMenu(
+					OwnerWidget.AsShared(),
+					WidgetPath,
+					MenuContent.ToSharedRef(),
+					MouseEvent.GetScreenSpacePosition(),
+					FPopupTransitionEffect( FPopupTransitionEffect::ContextMenu )
+					);
+
+				return FReply::Handled().SetUserFocus(MenuContent.ToSharedRef(), EFocusCause::SetDirectly).ReleaseMouseCapture();
+			}
+		}
 
-	return FSequencerEditTool_Default::OnMouseButtonUp(OwnerWidget, MyGeometry, MouseEvent);
+		return FReply::Handled();
+	}
 }
 
 
@@ -408,10 +424,13 @@ void FSequencerEditTool_Selection::OnMouseCaptureLost()
 
 FName FSequencerEditTool_Selection::GetIdentifier() const
 {
-	static FName Identifier("Selection");
 	return Identifier;
 }
 
+bool FSequencerEditTool_Selection::CanDeactivate() const
+{
+	return !DelayedDrag.IsSet();
+}
 
 void FSequencerEditTool_Selection::UpdateCursor(const FGeometry& MyGeometry, const FPointerEvent& MouseEvent)
 {
@@ -430,7 +449,7 @@ void FSequencerEditTool_Selection::UpdateCursor(const FGeometry& MyGeometry, con
 		}
 		else
 		{
-			CursorDecorator = FEditorStyle::Get().GetBrush(TEXT("Sequencer.CursorDecorator_Marquee"));
+			CursorDecorator = nullptr;
 		}
 	}
 }
diff --git a/Engine/Source/Editor/Sequencer/Private/Tools/SequencerEditTool_Selection.h b/Engine/Source/Editor/Sequencer/Private/Tools/SequencerEditTool_Selection.h
--- a/Engine/Source/Editor/Sequencer/Private/Tools/SequencerEditTool_Selection.h
+++ b/Engine/Source/Editor/Sequencer/Private/Tools/SequencerEditTool_Selection.h
@@ -2,7 +2,7 @@
 
 #pragma once
 
-#include "SequencerEditTool_Default.h"
+#include "SequencerEditTool.h"
 #include "UniquePtr.h"
 #include "DelayedDrag.h"
 
@@ -11,18 +11,19 @@ struct FMarqueeSelectData;
 
 
 class FSequencerEditTool_Selection
-	: public FSequencerEditTool_Default
+	: public FSequencerEditTool
 {
 public:
 
+	/** Static identifier for this edit tool */
+	static const FName Identifier;
+
 	/** Create and initialize a new instance. */
-	FSequencerEditTool_Selection(TSharedPtr<FSequencer> InSequencer, TSharedPtr<SSequencer> InSequencerWidget);
+	FSequencerEditTool_Selection(FSequencer& InSequencer);
 
 public:
 
-	// ISequencerEditTool interface
-
-	virtual ISequencer& GetSequencer() const override;
+	// ISequencerEditTool interface	
 	virtual int32 OnPaint(const FGeometry& AllottedGeometry, const FSlateRect& MyClippingRect, FSlateWindowElementList& OutDrawElements, int32 LayerId) const override;
 	virtual FReply OnMouseButtonDown(SWidget& OwnerWidget, const FGeometry& MyGeometry, const FPointerEvent& MouseEvent) override;
 	virtual FReply OnMouseButtonUp(SWidget& OwnerWidget, const FGeometry& MyGeometry, const FPointerEvent& MouseEvent) override;
@@ -31,6 +32,7 @@ public:
 	virtual void OnMouseCaptureLost() override;
 	virtual FCursorReply OnCursorQuery(const FGeometry& MyGeometry, const FPointerEvent& CursorEvent) const override;
 	virtual FName GetIdentifier() const override;
+	virtual bool CanDeactivate() const override;
 	
 private:
 
@@ -39,9 +41,6 @@ private:
 
 private:
 
-	/** The sequencer itself */
-	TWeakPtr<FSequencer> Sequencer;
-
 	/** Sequencer widget */
 	TWeakPtr<SSequencer> SequencerWidget;
 
diff --git a/Engine/Source/Editor/Sequencer/Private/VirtualTrackArea.cpp b/Engine/Source/Editor/Sequencer/Private/VirtualTrackArea.cpp
--- a/Engine/Source/Editor/Sequencer/Private/VirtualTrackArea.cpp
+++ b/Engine/Source/Editor/Sequencer/Private/VirtualTrackArea.cpp
@@ -129,14 +129,17 @@ FSequencerSelectedKey FVirtualTrackArea::HitTestKey(FVector2D InPhysicalPosition
 	else if (!Node->IsExpanded())
 	{
 		TSharedPtr<FSequencerTrackNode> TrackNode = GetParentTrackNode(*Node);
-		const TArray< TSharedRef<ISequencerSection> >& Sections = TrackNode->GetSections();
-
-		for ( int32 SectionIndex = 0; SectionIndex < Sections.Num(); ++SectionIndex )
+		if (TrackNode.IsValid())
 		{
-			UMovieSceneSection* Section = Sections[SectionIndex]->GetSectionObject();
-			if (Section->GetStartTime() <= KeyRight && Section->GetEndTime() >= KeyLeft)
+			const TArray< TSharedRef<ISequencerSection> >& Sections = TrackNode->GetSections();
+
+			for ( int32 SectionIndex = 0; SectionIndex < Sections.Num(); ++SectionIndex )
 			{
-				KeyAreas.Add(Node->GetKeyGrouping(SectionIndex));
+				UMovieSceneSection* Section = Sections[SectionIndex]->GetSectionObject();
+				if (Section->GetStartTime() <= KeyRight && Section->GetEndTime() >= KeyLeft)
+				{
+					KeyAreas.Add(Node->GetKeyGrouping(SectionIndex));
+				}
 			}
 		}
 	}
diff --git a/Engine/Source/Editor/Sequencer/Public/IKeyArea.h b/Engine/Source/Editor/Sequencer/Public/IKeyArea.h
--- a/Engine/Source/Editor/Sequencer/Public/IKeyArea.h
+++ b/Engine/Source/Editor/Sequencer/Public/IKeyArea.h
@@ -123,6 +123,11 @@ public:
 	 */
 	virtual ERichCurveExtrapolation GetExtrapolationMode(bool bPreInfinity) const = 0;
 
+	/**
+	 * @return Whether this key area can set extrapolation mode
+	 */
+	virtual bool CanSetExtrapolationMode() const { return false; }
+
 	/**
 	 * Adds a key at the specified time if there isn't already a key present.  The value of the added key should
 	 * be the value which would be returned if the animation containing this key area was evaluated at the specified time.
@@ -190,4 +195,12 @@ public:
 	 * @return Key area color.
 	 */
 	virtual FLinearColor GetColor() = 0;
+
+	/**
+	 * Get the data structure representing the specified key.
+	 *
+	 * @param KeyHandle The handle of the key.
+	 * @return The key's data structure representation, or nullptr if key not found or no structure available.
+	 */
+	virtual TSharedPtr<FStructOnScope> GetKeyStruct(FKeyHandle KeyHandle) = 0;
 };
diff --git a/Engine/Source/Editor/Sequencer/Public/ISequencer.h b/Engine/Source/Editor/Sequencer/Public/ISequencer.h
--- a/Engine/Source/Editor/Sequencer/Public/ISequencer.h
+++ b/Engine/Source/Editor/Sequencer/Public/ISequencer.h
@@ -87,6 +87,11 @@ public:
 	 */
 	virtual TSharedRef<FMovieSceneSequenceInstance> GetSequenceInstanceForSection(UMovieSceneSection& Section) const = 0;
 
+	/**
+	 * @return Whether the section has a sequence instance
+	 */
+	virtual bool HasSequenceInstanceForSection(UMovieSceneSection& Section) const = 0;
+
 	/**
 	 * Adds a movie scene as a section inside the current movie scene
 	 * 
@@ -140,9 +145,10 @@ public:
 	 * Sets the global position to the time.
 	 *
 	 * @param Time The global time to set.
+	 * @param bAllowSnappingToFrames Whether or not to allow snapping to frames if.
 	 * @see GetGlobalTime
 	 */
-	virtual void SetGlobalTime(float Time) = 0;
+	virtual void SetGlobalTime(float Time, const bool& bAllowSnappingToFrames = false) = 0;
 
 	/** @return The current view range */
 	virtual FAnimatedRange GetViewRange() const
@@ -150,6 +156,13 @@ public:
 		return FAnimatedRange();
 	}
 
+	/**
+	 * Set the view range, growing the working range to accomodate, if necessary
+	 * @param NewViewRange The new view range. Must be a finite range
+	 * @param Interpolation How to interpolate to the new view range
+	 */
+	virtual void SetViewRange(TRange<float> NewViewRange, EViewRangeInterpolation Interpolation = EViewRangeInterpolation::Animated) = 0;
+
 	/**
 	 * Sets whether perspective viewport hijacking is enabled.
 	 *
@@ -174,6 +187,9 @@ public:
 	 */ 
 	virtual bool IsPerspectiveViewportCameraCutEnabled() const { return true; }
 
+	DECLARE_EVENT_OneParam(ISequencer, FOnCameraCut, UObject*)
+	virtual FOnCameraCut& OnCameraCut() = 0;
+
 	/**
 	 * Gets a handle to runtime information about the object being manipulated by a movie scene
 	 * 
@@ -210,4 +226,16 @@ public:
 
 	/** Gets a multicast delegate which is executed whenever the global time changes. */
 	virtual FOnGlobalTimeChanged& OnGlobalTimeChanged() = 0;
+
+	/** @return a numeric type interface that will parse and display numbers as frames and times correctly */
+	virtual TSharedRef<INumericTypeInterface<float>> GetNumericTypeInterface() = 0;
+
+	/** @return a numeric type interface that will parse and display numbers as frames and times correctly (including zero-pad, if applicable) */
+	virtual TSharedRef<INumericTypeInterface<float>> GetZeroPadNumericTypeInterface() = 0;
+
+	/** @return the command bindings for this sequencer */
+	virtual TSharedPtr<FUICommandList> GetCommandBindings() const = 0;
+
+	/** @return Returns a widget containing the sequencer's playback controls */
+	virtual TSharedRef<SWidget> MakeTransportControls() = 0;
 };
diff --git a/Engine/Source/Editor/Sequencer/Public/ISequencerEditTool.h b/Engine/Source/Editor/Sequencer/Public/ISequencerEditTool.h
--- a/Engine/Source/Editor/Sequencer/Public/ISequencerEditTool.h
+++ b/Engine/Source/Editor/Sequencer/Public/ISequencerEditTool.h
@@ -61,20 +61,14 @@ public:
 
 	// @todo sequencer: documentation needed
 	virtual void OnMouseCaptureLost() = 0;
-	virtual void OnMouseEnter(SWidget& OwnerWidget, const FGeometry& MyGeometry, const FPointerEvent& MouseEvent) = 0;
-	virtual void OnMouseLeave(SWidget& OwnerWidget, const FPointerEvent& MouseEvent) = 0;
 	virtual int32 OnPaint(const FGeometry& AllottedGeometry, const FSlateRect& MyClippingRect, FSlateWindowElementList& OutDrawElements, int32 LayerId) const = 0;
-	virtual void Tick(const FGeometry& AllottedGeometry, const double InCurrentTime, const float InDeltaTime) = 0;
 	virtual FCursorReply OnCursorQuery(const FGeometry& MyGeometry, const FPointerEvent& CursorEvent) const = 0;
 	virtual ISequencer& GetSequencer() const = 0;
+	virtual void OnMouseEnter(SWidget& OwnerWidget, const FGeometry& MyGeometry, const FPointerEvent& MouseEvent) = 0;
+	virtual void OnMouseLeave(SWidget& OwnerWidget, const FPointerEvent& MouseEvent) = 0;
 	virtual FName GetIdentifier() const = 0;
-
-	/** Get the current active hotspot */
-	virtual TSharedPtr<ISequencerHotspot> GetHotspot() const = 0;
-
-	/** Set the hotspot to something else */
-	virtual void SetHotspot(TSharedPtr<ISequencerHotspot> NewHotspot) = 0;
-
+	virtual bool CanDeactivate() const = 0;
+	
 public:
 
 	/** Virtual destructor. */
diff --git a/Engine/Source/Editor/Sequencer/Public/ISequencerTrackEditor.h b/Engine/Source/Editor/Sequencer/Public/ISequencerTrackEditor.h
--- a/Engine/Source/Editor/Sequencer/Public/ISequencerTrackEditor.h
+++ b/Engine/Source/Editor/Sequencer/Public/ISequencerTrackEditor.h
@@ -7,6 +7,12 @@ class ISequencerSection;
 class UMovieSceneSection;
 class UMovieSceneTrack;
 
+/** Data structure containing information required to build an edit widget */
+struct FBuildEditWidgetParams
+{
+	/** Attribute that specifies when the node relating to this edit widget is hovered */
+	TAttribute<bool> NodeIsHovered;
+};
 
 /**
  * Interface for sequencer track editors.
@@ -63,9 +69,10 @@ public:
 	 * Builds an edit widget for the outliner nodes which represent tracks which are edited by this editor.
 	 * @param ObjectBinding The object binding associated with the track being edited by this editor.
 	 * @param Track The track being edited by this editor.
+	 * @param Params Parameter struct containing data relevant to the edit widget
 	 * @returns The the widget to display in the outliner, or an empty shared ptr if not widget is to be displayed.
 	 */
-	virtual TSharedPtr<SWidget> BuildOutlinerEditWidget(const FGuid& ObjectBinding, UMovieSceneTrack* Track) = 0;
+	virtual TSharedPtr<SWidget> BuildOutlinerEditWidget(const FGuid& ObjectBinding, UMovieSceneTrack* Track, const FBuildEditWidgetParams& Params) = 0;
 
 	/**
 	 * Builds the context menu for the track.
@@ -97,6 +104,12 @@ public:
 	 */
 	virtual TSharedRef<ISequencerSection> MakeSectionInterface(UMovieSceneSection& SectionObject, UMovieSceneTrack& Track) = 0;
 
+	/** Gets an icon brush for this track editor */
+	virtual const FSlateBrush* GetIconBrush() const { return nullptr; }
+
+	/** Called when the instance of this track editor is initialized */
+	virtual void OnInitialize() = 0;
+
 	/** Called when the instance of this track editor is released */
 	virtual void OnRelease() = 0;
 
diff --git a/Engine/Source/Editor/Sequencer/Public/MovieSceneTrackEditor.h b/Engine/Source/Editor/Sequencer/Public/MovieSceneTrackEditor.h
--- a/Engine/Source/Editor/Sequencer/Public/MovieSceneTrackEditor.h
+++ b/Engine/Source/Editor/Sequencer/Public/MovieSceneTrackEditor.h
@@ -212,7 +212,7 @@ public:
 	virtual void BuildAddTrackMenu(FMenuBuilder& MenuBuilder) override { }
 	virtual void BuildObjectBindingEditButtons(TSharedPtr<SHorizontalBox> EditBox, const FGuid& ObjectBinding, const UClass* ObjectClass) override { }
 	virtual void BuildObjectBindingTrackMenu(FMenuBuilder& MenuBuilder, const FGuid& ObjectBinding, const UClass* ObjectClass) override { }
-	virtual TSharedPtr<SWidget> BuildOutlinerEditWidget(const FGuid& ObjectBinding, UMovieSceneTrack* Track) override  { return TSharedPtr<SWidget>(); }
+	virtual TSharedPtr<SWidget> BuildOutlinerEditWidget(const FGuid& ObjectBinding, UMovieSceneTrack* Track, const FBuildEditWidgetParams& Params) override  { return TSharedPtr<SWidget>(); }
 	virtual void BuildTrackContextMenu( FMenuBuilder& MenuBuilder, UMovieSceneTrack* Track ) override { }
 	virtual bool HandleAssetAdded(UObject* Asset, const FGuid& TargetObjectGuid) override { return false; }
 
@@ -229,6 +229,7 @@ public:
 	}
 
 	virtual TSharedRef<ISequencerSection> MakeSectionInterface(class UMovieSceneSection& SectionObject, UMovieSceneTrack& Track) = 0;
+	virtual void OnInitialize() override { };
 	virtual void OnRelease() override { };
 
 	virtual int32 PaintTrackArea(const FPaintArgs& Args, const FGeometry& AllottedGeometry, const FSlateRect& MyClippingRect, FSlateWindowElementList& OutDrawElements, int32 LayerId, const FWidgetStyle& InWidgetStyle) override
diff --git a/Engine/Source/Editor/Sequencer/Public/SequencerCommands.h b/Engine/Source/Editor/Sequencer/Public/SequencerCommands.h
new file mode 100644
--- /dev/null
+++ b/Engine/Source/Editor/Sequencer/Public/SequencerCommands.h
@@ -0,0 +1,187 @@
+// Copyright 1998-2016 Epic Games, Inc. All Rights Reserved.
+
+#pragma once
+
+class FSequencerCommands : public TCommands<FSequencerCommands>
+{
+
+public:
+	FSequencerCommands() : TCommands<FSequencerCommands>
+	(
+		"Sequencer",
+		NSLOCTEXT("Contexts", "Sequencer", "Sequencer"),
+		NAME_None, // "MainFrame" // @todo Fix this crash
+		FEditorStyle::GetStyleSetName() // Icon Style Set
+	)
+	{}
+	
+	/** Toggle play */
+	TSharedPtr< FUICommandInfo > TogglePlay;
+
+	/** Play forward */
+	TSharedPtr< FUICommandInfo > PlayForward;
+
+	/** Rewind */
+	TSharedPtr< FUICommandInfo > Rewind;
+
+	/** Step forward */
+	TSharedPtr< FUICommandInfo > StepForward;
+
+	/** Step backward */
+	TSharedPtr< FUICommandInfo > StepBackward;
+
+	/** Step to next key */
+	TSharedPtr< FUICommandInfo > StepToNextKey;
+
+	/** Step to previous key */
+	TSharedPtr< FUICommandInfo > StepToPreviousKey;
+
+	/** Step to next camera key */
+	TSharedPtr< FUICommandInfo > StepToNextCameraKey;
+
+	/** Step to previous camera key */
+	TSharedPtr< FUICommandInfo > StepToPreviousCameraKey;
+
+	/** Set start playback range */
+	TSharedPtr< FUICommandInfo > SetStartPlaybackRange;
+
+	/** Set end playback range */
+	TSharedPtr< FUICommandInfo > SetEndPlaybackRange;
+
+	/** Zoom into the view range */
+	TSharedPtr< FUICommandInfo > ZoomInViewRange;
+
+	/** Zoom out of the view range */
+	TSharedPtr< FUICommandInfo > ZoomOutViewRange;
+
+	/** Toggle constraining the time cursor to the playback range */
+	TSharedPtr< FUICommandInfo > ToggleKeepCursorInPlaybackRange;
+
+	/** Toggle constraining the playback range to the section bounds */
+	TSharedPtr< FUICommandInfo > ToggleKeepPlaybackRangeInSectionBounds;
+
+	/** Expand all nodes and descendants */
+	TSharedPtr< FUICommandInfo > ExpandAllNodesAndDescendants;
+
+	/** Collapse all nodes and descendants */
+	TSharedPtr< FUICommandInfo > CollapseAllNodesAndDescendants;
+
+	/** Expand/collapse nodes */
+	TSharedPtr< FUICommandInfo > ToggleExpandCollapseNodes;
+
+	/** Expand/collapse nodes and descendants */
+	TSharedPtr< FUICommandInfo > ToggleExpandCollapseNodesAndDescendants;
+
+	/** Sets a key at the current time for the selected actor */
+	TSharedPtr< FUICommandInfo > SetKey;
+
+	/** Sets the interp tangent mode for the selected keys to auto */
+	TSharedPtr< FUICommandInfo > SetInterpolationCubicAuto;
+
+	/** Sets the interp tangent mode for the selected keys to user */
+	TSharedPtr< FUICommandInfo > SetInterpolationCubicUser;
+
+	/** Sets the interp tangent mode for the selected keys to break */
+	TSharedPtr< FUICommandInfo > SetInterpolationCubicBreak;
+
+	/** Sets the interp tangent mode for the selected keys to linear */
+	TSharedPtr< FUICommandInfo > SetInterpolationLinear;
+
+	/** Sets the interp tangent mode for the selected keys to constant */
+	TSharedPtr< FUICommandInfo > SetInterpolationConstant;
+
+	/** Trim section to the left, keeping the right portion */
+	TSharedPtr< FUICommandInfo > TrimSectionLeft;
+
+	/** Trim section to the right, keeping the left portion */
+	TSharedPtr< FUICommandInfo > TrimSectionRight;
+
+	/** Split section */
+	TSharedPtr< FUICommandInfo > SplitSection;
+
+	/** Set the auto key mode to Key All. */
+	TSharedPtr< FUICommandInfo > SetAutoKeyModeAll;
+
+	/** Set the auto key mode to Key Animated. */
+	TSharedPtr< FUICommandInfo > SetAutoKeyModeAnimated;
+
+	/** Set the auto key mode to Key None. */
+	TSharedPtr< FUICommandInfo > SetAutoKeyModeNone;
+
+	/** Turns key all on and off. */
+	TSharedPtr< FUICommandInfo > ToggleKeyAllEnabled;
+
+	/** Turns show frame numbers on and off. */
+	TSharedPtr< FUICommandInfo > ToggleShowFrameNumbers;
+
+	/** Toggle the visibility of the goto box. */
+	TSharedPtr< FUICommandInfo > ToggleShowGotoBox;
+
+	/** Turns the range slider on and off. */
+	TSharedPtr< FUICommandInfo > ToggleShowRangeSlider;
+
+	/** Toggles whether to lock the in/out to the start/end range on scrolling. */
+	TSharedPtr< FUICommandInfo > ToggleLockInOutToStartEndRange;
+
+	/** Turns snapping on and off. */
+	TSharedPtr< FUICommandInfo > ToggleIsSnapEnabled;
+
+	/** Toggles whether or not keys should snap to the selected interval. */
+	TSharedPtr< FUICommandInfo > ToggleSnapKeyTimesToInterval;
+
+	/** Toggles whether or not keys should snap to other keys in the section. */
+	TSharedPtr< FUICommandInfo > ToggleSnapKeyTimesToKeys;
+
+	/** Toggles whether or not sections should snap to the selected interval. */
+	TSharedPtr< FUICommandInfo > ToggleSnapSectionTimesToInterval;
+
+	/** Toggles whether or not sections should snap to other sections. */
+	TSharedPtr< FUICommandInfo > ToggleSnapSectionTimesToSections;
+
+	/** Toggles whether or not the play time should snap to the selected interval. */
+	TSharedPtr< FUICommandInfo > ToggleSnapPlayTimeToInterval;
+
+	/** Toggles whether or not the play time should snap to the dragged key. */
+	TSharedPtr< FUICommandInfo > ToggleSnapPlayTimeToDraggedKey;
+
+	/** Toggles whether or not to lock the play rate to the frame rate. */
+	TSharedPtr< FUICommandInfo > ToggleFixedTimeStepPlayback;
+
+	/** Toggles whether or not to snap curve values to the interval. */
+	TSharedPtr< FUICommandInfo > ToggleSnapCurveValueToInterval;
+
+	/** Finds the viewed sequence asset in the content browser. */
+	TSharedPtr< FUICommandInfo > FindInContentBrowser;
+
+	/** Toggles whether the details view is enabled in the level editor. */
+	TSharedPtr< FUICommandInfo > ToggleDetailsView;
+
+	/** Toggles whether the label browser is enabled in the level editor. */
+	TSharedPtr< FUICommandInfo > ToggleLabelBrowser;
+
+	/** Turns auto scroll on and off. */
+	TSharedPtr< FUICommandInfo > ToggleAutoScroll;
+
+	/** Toggles whether or not the curve editor should be shown. */
+	TSharedPtr< FUICommandInfo > ToggleShowCurveEditor;
+
+	/** Enable the move tool */
+	TSharedPtr< FUICommandInfo > MoveTool;
+
+	/** Enable the marquee selection tool */
+	TSharedPtr< FUICommandInfo > MarqueeTool;
+
+	/** Open a panel that enables exporting the sequence to a movie */
+	TSharedPtr< FUICommandInfo > RenderMovie;
+
+	/** Paste from the sequencer clipboard history */
+	TSharedPtr< FUICommandInfo > PasteFromHistory;
+
+	/** Convert the selected possessed objects to spawnables. These will be spawned and destroyed by sequencer as per object's the spawn track. */
+	TSharedPtr< FUICommandInfo > ConvertToSpawnable;
+
+	/**
+	 * Initialize commands
+	 */
+	virtual void RegisterCommands() override;
+};
diff --git a/Engine/Source/Editor/Sequencer/Public/SequencerUtilities.h b/Engine/Source/Editor/Sequencer/Public/SequencerUtilities.h
new file mode 100644
--- /dev/null
+++ b/Engine/Source/Editor/Sequencer/Public/SequencerUtilities.h
@@ -0,0 +1,8 @@
+// Copyright 1998-2015 Epic Games, Inc. All Rights Reserved.
+
+#pragma once
+
+struct SEQUENCER_API FSequencerUtilities
+{
+	static TSharedRef<SWidget> MakeAddButton(FText HoverText, FOnGetContent MenuContent, const TAttribute<bool>& HoverState);
+};
\ No newline at end of file
diff --git a/Engine/Source/Editor/SequencerWidgets/Private/STimeRange.cpp b/Engine/Source/Editor/SequencerWidgets/Private/STimeRange.cpp
--- a/Engine/Source/Editor/SequencerWidgets/Private/STimeRange.cpp
+++ b/Engine/Source/Editor/SequencerWidgets/Private/STimeRange.cpp
@@ -306,16 +306,23 @@ FString STimeRange::ToString(const float& Value) const
 	return FString::Printf(TEXT("%.2f"), Value);
 }
 
-TOptional<float> STimeRange::FromString(const FString& InString)
+TOptional<float> STimeRange::FromString(const FString& InString, const float& ExistingValue)
 {
 	bool bShowFrameNumbers = ShowFrameNumbers.IsBound() ? ShowFrameNumbers.Get() : false;
 	if (bShowFrameNumbers && TimeSliderController.IsValid())
 	{
-		int32 NewEndFrame = FCString::Atoi(*InString);
-		return float(TimeSliderController.Get()->FrameToTime(NewEndFrame));
-	}
+		// Convert existing value to frames
+		float ExistingValueInFrames = TimeSliderController.Get()->TimeToFrame(ExistingValue);
+		TOptional<float> Result = TDefaultNumericTypeInterface<float>::FromString(InString, ExistingValueInFrames);
 
-	return FCString::Atof(*InString);
+		if (Result.IsSet())
+		{
+			int32 NewEndFrame = FMath::RoundToInt(Result.GetValue());
+			return float(TimeSliderController.Get()->FrameToTime(NewEndFrame));
+		}
+	}
+		
+	return TDefaultNumericTypeInterface<float>::FromString(InString, ExistingValue);
 }
 
 
diff --git a/Engine/Source/Editor/SequencerWidgets/Private/STimeRange.h b/Engine/Source/Editor/SequencerWidgets/Private/STimeRange.h
--- a/Engine/Source/Editor/SequencerWidgets/Private/STimeRange.h
+++ b/Engine/Source/Editor/SequencerWidgets/Private/STimeRange.h
@@ -60,7 +60,7 @@ protected:
 
 	/** Convert the type to/from a string */
 	virtual FString ToString(const float& Value) const override;
-	virtual TOptional<float> FromString(const FString& InString) override;
+	virtual TOptional<float> FromString(const FString& InString, const float& ExistingValue) override;
 
 private:
 	TSharedPtr<ITimeSliderController> TimeSliderController;
diff --git a/Engine/Source/Editor/UMGEditor/Private/Details/SWidgetDetailsView.cpp b/Engine/Source/Editor/UMGEditor/Private/Details/SWidgetDetailsView.cpp
--- a/Engine/Source/Editor/UMGEditor/Private/Details/SWidgetDetailsView.cpp
+++ b/Engine/Source/Editor/UMGEditor/Private/Details/SWidgetDetailsView.cpp
@@ -339,6 +339,8 @@ FText SWidgetDetailsView::GetCategoryText() const
 	return FText::GetEmpty();
 }
 
+const FSlateBrush* GetEditorIcon_Deprecated(UWidget* Widget);
+
 const FSlateBrush* SWidgetDetailsView::GetNameIcon() const
 {
 	if ( SelectedObjects.Num() == 1 )
@@ -346,7 +348,9 @@ const FSlateBrush* SWidgetDetailsView::GetNameIcon() const
 		UWidget* Widget = Cast<UWidget>(SelectedObjects[0].Get());
 		if ( Widget )
 		{
-			return Widget->GetEditorIcon();
+			// @todo UMG: remove after 4.12
+			return GetEditorIcon_Deprecated(Widget);
+			// return FClassIconFinder::FindIconForClass(Widget->GetClass());
 		}
 	}
 
diff --git a/Engine/Source/Editor/UMGEditor/Private/Hierarchy/SHierarchyViewItem.cpp b/Engine/Source/Editor/UMGEditor/Private/Hierarchy/SHierarchyViewItem.cpp
--- a/Engine/Source/Editor/UMGEditor/Private/Hierarchy/SHierarchyViewItem.cpp
+++ b/Engine/Source/Editor/UMGEditor/Private/Hierarchy/SHierarchyViewItem.cpp
@@ -885,9 +885,13 @@ FText FHierarchyWidget::GetLabelToolTipText() const
 	return FText::GetEmpty();
 }
 
+const FSlateBrush* GetEditorIcon_Deprecated(UWidget* Widget);
+
 const FSlateBrush* FHierarchyWidget::GetImage() const
 {
-	return Item.GetTemplate()->GetEditorIcon();
+	// @todo UMG: remove after 4.12
+	return GetEditorIcon_Deprecated(Item.GetTemplate());
+	// return FClassIconFinder::FindIconForClass(Item.GetTemplate()->GetClass());
 }
 
 FSlateFontInfo FHierarchyWidget::GetFont() const
diff --git a/Engine/Source/Editor/UMGEditor/Private/Templates/WidgetTemplateBlueprintClass.cpp b/Engine/Source/Editor/UMGEditor/Private/Templates/WidgetTemplateBlueprintClass.cpp
--- a/Engine/Source/Editor/UMGEditor/Private/Templates/WidgetTemplateBlueprintClass.cpp
+++ b/Engine/Source/Editor/UMGEditor/Private/Templates/WidgetTemplateBlueprintClass.cpp
@@ -64,10 +64,14 @@ UWidget* FWidgetTemplateBlueprintClass::Create(UWidgetTree* Tree)
 	return FWidgetTemplateClass::CreateNamed(Tree, FName(*FBlueprintEditorUtils::GetClassNameWithoutSuffix(WidgetClass.Get())));
 }
 
+const FSlateBrush* GetEditorIcon_Deprecated(UWidget* Widget);
+
 const FSlateBrush* FWidgetTemplateBlueprintClass::GetIcon() const
 {
+	// @todo UMG: remove after 4.12
 	auto DefaultUserWidget = UUserWidget::StaticClass()->GetDefaultObject<UUserWidget>();
-	return DefaultUserWidget->GetEditorIcon();
+	return GetEditorIcon_Deprecated(DefaultUserWidget);
+	// return FClassIconFinder::FindIconForClass(UUserWidget::StaticClass());
 }
 
 TSharedRef<IToolTip> FWidgetTemplateBlueprintClass::GetToolTip() const
diff --git a/Engine/Source/Editor/UMGEditor/Private/Templates/WidgetTemplateClass.cpp b/Engine/Source/Editor/UMGEditor/Private/Templates/WidgetTemplateClass.cpp
--- a/Engine/Source/Editor/UMGEditor/Private/Templates/WidgetTemplateClass.cpp
+++ b/Engine/Source/Editor/UMGEditor/Private/Templates/WidgetTemplateClass.cpp
@@ -40,10 +40,14 @@ UWidget* FWidgetTemplateClass::Create(UWidgetTree* Tree)
 	return CreateNamed(Tree, NAME_None);
 }
 
+const FSlateBrush* GetEditorIcon_Deprecated(UWidget* Widget);
+
 const FSlateBrush* FWidgetTemplateClass::GetIcon() const
 {
-	UWidget* DefaultWidget = WidgetClass->GetDefaultObject<UWidget>();
-	return DefaultWidget->GetEditorIcon();
+	// @todo UMG: remove after 4.12
+	auto DefaultUserWidget = WidgetClass->GetDefaultObject<UUserWidget>();
+	return GetEditorIcon_Deprecated(DefaultUserWidget);
+	// return FClassIconFinder::FindIconForClass(WidgetClass);
 }
 
 TSharedRef<IToolTip> FWidgetTemplateClass::GetToolTip() const
diff --git a/Engine/Source/Editor/UMGEditor/Private/UMGEditorModule.cpp b/Engine/Source/Editor/UMGEditor/Private/UMGEditorModule.cpp
--- a/Engine/Source/Editor/UMGEditor/Private/UMGEditorModule.cpp
+++ b/Engine/Source/Editor/UMGEditor/Private/UMGEditorModule.cpp
@@ -19,6 +19,9 @@
 #include "DesignerCommands.h"
 #include "WidgetBlueprint.h"
 
+#include "SlateStyleRegistry.h"
+#include "ClassIconFinder.h"
+
 #include "Settings/WidgetDesignerSettings.h"
 #include "ISettingsModule.h"
 
@@ -26,6 +29,14 @@
 
 const FName UMGEditorAppIdentifier = FName(TEXT("UMGEditorApp"));
 
+PRAGMA_DISABLE_DEPRECATION_WARNINGS
+const FSlateBrush* GetEditorIcon_Deprecated(UWidget* Widget)
+{
+	const FSlateBrush* Brush = Widget->GetEditorIcon();
+	return Brush ? Brush : FClassIconFinder::FindIconForClass(Widget->GetClass());
+}
+PRAGMA_ENABLE_DEPRECATION_WARNINGS
+
 class FUMGEditorModule : public IUMGEditorModule, public IBlueprintCompiler
 {
 public:
@@ -72,11 +83,23 @@ public:
 				GetMutableDefault<UWidgetDesignerSettings>()
 				);
 		}
+
+		const ISlateStyle* UMGStyle = FSlateStyleRegistry::FindSlateStyle("UMGStyle");
+		if (UMGStyle)
+		{
+			FClassIconFinder::RegisterIconSource(UMGStyle);
+		}
 	}
 
 	/** Called before the module is unloaded, right before the module object is destroyed. */
 	virtual void ShutdownModule() override
 	{
+		const ISlateStyle* UMGStyle = FSlateStyleRegistry::FindSlateStyle("UMGStyle");
+		if (UMGStyle)
+		{
+			FClassIconFinder::UnregisterIconSource(UMGStyle);
+		}
+
 		MenuExtensibilityManager.Reset();
 		ToolBarExtensibilityManager.Reset();
 
diff --git a/Engine/Source/Editor/UnrealEd/Private/EditorEngine.cpp b/Engine/Source/Editor/UnrealEd/Private/EditorEngine.cpp
--- a/Engine/Source/Editor/UnrealEd/Private/EditorEngine.cpp
+++ b/Engine/Source/Editor/UnrealEd/Private/EditorEngine.cpp
@@ -122,8 +122,6 @@
 #include "Engine/CoreSettings.h"
 #include "ShaderCompiler.h"
 
-#include "AnimationRecorder.h"
-
 DEFINE_LOG_CATEGORY_STATIC(LogEditor, Log, All);
 
 #define LOCTEXT_NAMESPACE "UnrealEd.Editor"
@@ -1589,8 +1587,8 @@ void UEditorEngine::Tick( float DeltaSeconds, bool bIdleMode )
 
 	FUnrealEdMisc::Get().TickPerformanceAnalytics();
 
-	FAnimationRecorderManager::Get().Tick(DeltaSeconds);
-	
+	BroadcastPostEditorTick(DeltaSeconds);
+
 	// If the fadeout animation has completed for the undo/redo notification item, allow it to be deleted
 	if(UndoRedoNotificationItem.IsValid() && UndoRedoNotificationItem->GetCompletionState() == SNotificationItem::CS_None)
 	{
diff --git a/Engine/Source/Editor/UnrealEd/Private/EditorServer.cpp b/Engine/Source/Editor/UnrealEd/Private/EditorServer.cpp
--- a/Engine/Source/Editor/UnrealEd/Private/EditorServer.cpp
+++ b/Engine/Source/Editor/UnrealEd/Private/EditorServer.cpp
@@ -70,7 +70,6 @@
 #include "UnrealEngine.h"
 #include "AI/Navigation/NavLinkRenderingComponent.h"
 #include "PhysicsPublic.h"
-#include "AnimationRecorder.h"
 #include "Analytics/AnalyticsPrivacySettings.h"
 #include "KismetReinstanceUtilities.h"
 
@@ -4892,6 +4891,15 @@ bool UEditorEngine::SnapObjectTo( FActorOrComponent Object, const bool InAlign,
 		{
 			RebuildAlteredBSP();
 		}
+
+		TArray<FEdMode*> ActiveModes; 
+		GCurrentLevelEditingViewportClient->GetModeTools()->GetActiveModes(ActiveModes);
+		for( int32 ModeIndex = 0; ModeIndex < ActiveModes.Num(); ++ModeIndex )
+		{
+			// Notify active modes
+			ActiveModes[ModeIndex]->ActorMoveNotify();
+		}
+
 		return true;
 	}
 
diff --git a/Engine/Source/Editor/UnrealEd/Private/FileHelpers.cpp b/Engine/Source/Editor/UnrealEd/Private/FileHelpers.cpp
--- a/Engine/Source/Editor/UnrealEd/Private/FileHelpers.cpp
+++ b/Engine/Source/Editor/UnrealEd/Private/FileHelpers.cpp
@@ -598,9 +598,10 @@ bool RenameStreamingLevel( FString& LevelToRename, const FString& OldBaseLevelNa
 	return false;
 }
 
-static bool OpenLevelSaveAsDialog(const FString& InDefaultPath, const FString& InNewNameSuggestion, FString& OutPackageName)
+static bool OpenSaveAsDialog(UClass* SavedClass, const FString& InDefaultPath, const FString& InNewNameSuggestion, FString& OutPackageName)
 {
 	FString DefaultPath = InDefaultPath;
+
 	if (DefaultPath.IsEmpty())
 	{
 		DefaultPath = TEXT("/Game/Maps");
@@ -610,14 +611,19 @@ static bool OpenLevelSaveAsDialog(const FString& InDefaultPath, const FString& I
 	check(!NewNameSuggestion.IsEmpty());
 
 	FSaveAssetDialogConfig SaveAssetDialogConfig;
-	SaveAssetDialogConfig.DialogTitleOverride = LOCTEXT("SaveLevelDialogTitle", "Save Level As");
-	SaveAssetDialogConfig.DefaultPath = DefaultPath;
-	SaveAssetDialogConfig.DefaultAssetName = NewNameSuggestion;
-	SaveAssetDialogConfig.AssetClassNames.Add(UWorld::StaticClass()->GetFName());
-	SaveAssetDialogConfig.ExistingAssetPolicy = ESaveAssetDialogExistingAssetPolicy::AllowButWarn;
+	{
+		SaveAssetDialogConfig.DefaultPath = DefaultPath;
+		SaveAssetDialogConfig.DefaultAssetName = NewNameSuggestion;
+		SaveAssetDialogConfig.AssetClassNames.Add(SavedClass->GetFName());
+		SaveAssetDialogConfig.ExistingAssetPolicy = ESaveAssetDialogExistingAssetPolicy::AllowButWarn;
+		SaveAssetDialogConfig.DialogTitleOverride = (SavedClass == UWorld::StaticClass())
+			? LOCTEXT("SaveLevelDialogTitle", "Save Level As")
+			: LOCTEXT("SaveAssetDialogTitle", "Save Asset As");
+	}
 
 	FContentBrowserModule& ContentBrowserModule = FModuleManager::LoadModuleChecked<FContentBrowserModule>("ContentBrowser");
 	FString SaveObjectPath = ContentBrowserModule.Get().CreateModalSaveAssetDialog(SaveAssetDialogConfig);
+	
 	if ( !SaveObjectPath.IsEmpty() )
 	{
 		OutPackageName = FPackageName::ObjectPathToPackageName(SaveObjectPath);
@@ -673,7 +679,8 @@ static bool SaveAsImplementation( UWorld* InWorld, const FString& DefaultFilenam
 			FPackageName::TryConvertFilenameToLongPackageName(DefaultDirectory / Filename, DefaultPackagePath);
 
 			FString PackageName;
-			bSaveFileLocationSelected = OpenLevelSaveAsDialog(
+			bSaveFileLocationSelected = OpenSaveAsDialog(
+				UWorld::StaticClass(),
 				FPackageName::GetLongPackagePath(DefaultPackagePath),
 				FPaths::GetBaseFilename(Filename),
 				PackageName);
@@ -895,7 +902,7 @@ static bool IsWorldDirty()
  * @param	InLevel		The level to be SaveAs'd.
  * @return				true if the world was saved.
  */
-bool FEditorFileUtils::SaveAs(ULevel* InLevel)
+bool FEditorFileUtils::SaveLevelAs(ULevel* InLevel)
 {
 	FString DefaultFilename;
 
@@ -2454,7 +2461,8 @@ static int32 InternalSavePackage( UPackage* PackageToSave, bool& bOutPackageLoca
 				}
 
 				FString SaveAsPackageName;
-				bSaveFile = OpenLevelSaveAsDialog(
+				bSaveFile = OpenSaveAsDialog(
+					UWorld::StaticClass(),
 					FPackageName::GetLongPackagePath(DefaultPackagePath),
 					FPaths::GetBaseFilename(FinalPackageFilename),
 					SaveAsPackageName);
diff --git a/Engine/Source/Editor/UnrealEd/Public/FileHelpers.h b/Engine/Source/Editor/UnrealEd/Public/FileHelpers.h
--- a/Engine/Source/Editor/UnrealEd/Public/FileHelpers.h
+++ b/Engine/Source/Editor/UnrealEd/Public/FileHelpers.h
@@ -129,7 +129,7 @@ public:
 	 * @param	Level		The Level to be SaveAs'd.
 	 * @return				true if the world was saved.
 	 */
-	UNREALED_API static bool SaveAs(ULevel* Level);
+	UNREALED_API static bool SaveLevelAs(ULevel* Level);
 
 	/**
 	 * Saves all levels to the specified directory.
diff --git a/Engine/Source/Editor/WorldBrowser/Private/LevelCollectionModel.cpp b/Engine/Source/Editor/WorldBrowser/Private/LevelCollectionModel.cpp
--- a/Engine/Source/Editor/WorldBrowser/Private/LevelCollectionModel.cpp
+++ b/Engine/Source/Editor/WorldBrowser/Private/LevelCollectionModel.cpp
@@ -1391,7 +1391,7 @@ void FLevelCollectionModel::SaveSelectedLevelAs_Executed()
 		ULevel* Level = SelectedLevelsList[0]->GetLevelObject();
 		if (Level)
 		{
-			FEditorFileUtils::SaveAs(Level);
+			FEditorFileUtils::SaveLevelAs(Level);
 		}
 	}
 }
diff --git a/Engine/Source/Runtime/AIModule/Private/EnvironmentQuery/Tests/EnvQueryTest_Pathfinding.cpp b/Engine/Source/Runtime/AIModule/Private/EnvironmentQuery/Tests/EnvQueryTest_Pathfinding.cpp
--- a/Engine/Source/Runtime/AIModule/Private/EnvironmentQuery/Tests/EnvQueryTest_Pathfinding.cpp
+++ b/Engine/Source/Runtime/AIModule/Private/EnvironmentQuery/Tests/EnvQueryTest_Pathfinding.cpp
@@ -234,6 +234,6 @@ ANavigationData* UEnvQueryTest_Pathfinding::FindNavigationData(UNavigationSystem
 	return NavSys.GetMainNavData(FNavigationSystem::DontCreate);
 }
 
-#undef LOCTEXT_NAMESPACE
+#endif //WITH_DEV_AUTOMATION_TESTS
 
-#endif //WITH_DEV_AUTOMATION_TESTS
\ No newline at end of file
+#undef LOCTEXT_NAMESPACE
