Copying //UE4/Dev-Sequencer to //UE4/Main

diff --git a/Engine/Source/Editor/MovieSceneTools/Public/KeyframeTrackEditor.h b/Engine/Source/Editor/MovieSceneTools/Public/KeyframeTrackEditor.h
--- a/Engine/Source/Editor/MovieSceneTools/Public/KeyframeTrackEditor.h
+++ b/Engine/Source/Editor/MovieSceneTools/Public/KeyframeTrackEditor.h
@@ -110,6 +110,7 @@ private:
 		TFunction<void(TrackType*)> OnSetIntermediateValue)
 	{
 		bool bTrackCreated = false;
+		bool bSectionCreated = false;
 
 		// Try to find an existing Track, and if one doesn't exist check the key params and create one if requested.
 		FFindOrCreateTrackResult TrackResult = FindOrCreateTrackForObject( ObjectHandle, TrackClass, PropertyName, KeyParams.bCreateTrackIfMissing );
@@ -126,18 +127,21 @@ private:
 
 		if ( Track )
 		{
-			AddKeysToTrack( Track, KeyTime, NewKeys, DefaultKeys, KeyParams, bTrackCreated, OnSetIntermediateValue );
+			bSectionCreated = AddKeysToTrack( Track, KeyTime, NewKeys, DefaultKeys, KeyParams, bTrackCreated, OnSetIntermediateValue );
 		}
 
-		return bTrackCreated;
+		return bTrackCreated || bSectionCreated;
 	}
 
-	void AddKeysToTrack(
+	/* Returns whether a section was added */
+	bool AddKeysToTrack(
 		TrackType* Track, float KeyTime,
 		const TArray<KeyDataType>& NewKeys, const TArray<KeyDataType>& DefaultKeys,
 		FKeyParams KeyParams, bool bNewTrack,
 		TFunction<void(TrackType*)> OnSetIntermediateValue)
 	{
+		bool bSectionCreated = false;
+
 		bool bSettingIntermediateValue = KeyParams.bCreateKeyOnlyWhenAutoKeying && GetSequencer()->GetAutoKeyMode() == EAutoKeyMode::KeyNone;
 		if ( bSettingIntermediateValue )
 		{
@@ -154,11 +158,11 @@ private:
 				{
 					if ( HasKeys( Track, NewKey ) || KeyParams.bCreateKeyIfEmpty )
 					{
-						AddKey( Track, KeyTime, NewKey, InterpolationMode );
+						bSectionCreated |= AddKey( Track, KeyTime, NewKey, InterpolationMode );
 					}
 					else
 					{
-						SetDefault( Track, KeyTime, NewKey );
+						bSectionCreated |= SetDefault( Track, KeyTime, NewKey );
 					}
 				}
 			}
@@ -167,10 +171,11 @@ private:
 			{
 				for ( const KeyDataType& DefaultKey : DefaultKeys )
 				{
-					SetDefault( Track, KeyTime, DefaultKey );
+					bSectionCreated |= SetDefault( Track, KeyTime, DefaultKey );
 				}
 			}
 		}
+		return bSectionCreated;
 	}
 
 	bool NewKeyIsNewData( TrackType* Track, float Time, const KeyDataType& KeyData ) const
@@ -194,15 +199,20 @@ private:
 
 	using FMovieSceneTrackEditor::AddKey;
 
-	void AddKey( TrackType* Track, float Time, const KeyDataType& KeyData, EMovieSceneKeyInterpolation KeyInterpolation )
+	/* Return whether a section was added */
+	bool AddKey( TrackType* Track, float Time, const KeyDataType& KeyData, EMovieSceneKeyInterpolation KeyInterpolation )
 	{
+		bool bSectionAdded = false;
 		Track->Modify();
-		IKeyframeSection<KeyDataType>* KeyframeSection = CastChecked<SectionType>( Track->FindOrAddSection( Time ) );
+		IKeyframeSection<KeyDataType>* KeyframeSection = CastChecked<SectionType>( Track->FindOrAddSection( Time, bSectionAdded ) );
 		KeyframeSection->AddKey( Time, KeyData, KeyInterpolation );
+		return bSectionAdded;
 	}
 
-	void SetDefault( TrackType* Track, float Time, const KeyDataType& KeyData )
+	/* Return whether a section was added */
+	bool SetDefault( TrackType* Track, float Time, const KeyDataType& KeyData )
 	{
+		bool bSectionAdded = false;
 		Track->Modify();
 		const TArray<UMovieSceneSection*>& Sections = Track->GetAllSections();
 		if ( Sections.Num() > 0)
@@ -215,9 +225,10 @@ private:
 		}
 		else
 		{
-			IKeyframeSection<KeyDataType>* KeyframeSection = CastChecked<SectionType>( Track->FindOrAddSection( Time ) );
+			IKeyframeSection<KeyDataType>* KeyframeSection = CastChecked<SectionType>( Track->FindOrAddSection( Time, bSectionAdded ) );
 			KeyframeSection->SetDefault( KeyData );
 		}
+		return bSectionAdded;
 	}
 };
 
diff --git a/Engine/Source/Editor/MovieSceneTools/Public/MovieSceneToolHelpers.h b/Engine/Source/Editor/MovieSceneTools/Public/MovieSceneToolHelpers.h
--- a/Engine/Source/Editor/MovieSceneTools/Public/MovieSceneToolHelpers.h
+++ b/Engine/Source/Editor/MovieSceneTools/Public/MovieSceneToolHelpers.h
@@ -61,8 +61,9 @@ public:
 	 * 
 	 * @param Section The section to gather takes from
 	 * @param TakeNumbers The gathered take numbers
+	 * @param CurrentTakeNumber The current take number of the section
 	 */
-	static void GatherTakes(const UMovieSceneSection* Section, TArray<uint32>& TakeNumbers);
+	static void GatherTakes(const UMovieSceneSection* Section, TArray<uint32>& TakeNumbers, uint32& CurrentTakeNumber);
 
 
 	/**
diff --git a/Engine/Source/Editor/MovieSceneTools/Public/NameCurveKeyArea.h b/Engine/Source/Editor/MovieSceneTools/Public/NameCurveKeyArea.h
--- a/Engine/Source/Editor/MovieSceneTools/Public/NameCurveKeyArea.h
+++ b/Engine/Source/Editor/MovieSceneTools/Public/NameCurveKeyArea.h
@@ -33,7 +33,7 @@ public:
 	virtual bool CanCreateKeyEditor() override;
 	virtual TSharedRef<SWidget> CreateKeyEditor(ISequencer* Sequencer) override;
 	virtual void DeleteKey(FKeyHandle KeyHandle) override;
-	virtual FLinearColor GetColor() override;
+	virtual TOptional<FLinearColor> GetColor() override;
 	virtual ERichCurveExtrapolation GetExtrapolationMode(bool bPreInfinity) const override;
 	virtual ERichCurveInterpMode GetKeyInterpMode(FKeyHandle Keyhandle) const override;
 	virtual TSharedPtr<FStructOnScope> GetKeyStruct(FKeyHandle KeyHandle) override;
diff --git a/Engine/Source/Editor/MovieSceneTools/Public/PropertySection.h b/Engine/Source/Editor/MovieSceneTools/Public/PropertySection.h
--- a/Engine/Source/Editor/MovieSceneTools/Public/PropertySection.h
+++ b/Engine/Source/Editor/MovieSceneTools/Public/PropertySection.h
@@ -46,21 +46,9 @@ public:
 
 	virtual void GenerateSectionLayout( class ISectionLayoutBuilder& LayoutBuilder ) const override { }
 
-	virtual int32 OnPaintSection( const FGeometry& AllottedGeometry, const FSlateRect& SectionClippingRect, FSlateWindowElementList& OutDrawElements, int32 LayerId, bool bParentEnabled ) const override 
+	virtual int32 OnPaintSection( FSequencerSectionPainter& Painter ) const override 
 	{
-		const ESlateDrawEffect::Type DrawEffects = bParentEnabled ? ESlateDrawEffect::None : ESlateDrawEffect::DisabledEffect;
-		
-		// Add a box for the section
-		FSlateDrawElement::MakeBox( 
-			OutDrawElements,
-			LayerId,
-			AllottedGeometry.ToPaintGeometry(),
-			FEditorStyle::GetBrush("Sequencer.GenericSection.Background"),
-			SectionClippingRect,
-			DrawEffects
-		);
-
-		return LayerId;
+		return Painter.PaintSectionBackground();
 	}
 
 	void SetIntermediateValueForTrack( UMovieSceneTrack* Track, FPropertyChangedParams PropertyChangedParams )
diff --git a/Engine/Source/Editor/NiagaraEditor/Private/NiagaraSequencer.h b/Engine/Source/Editor/NiagaraEditor/Private/NiagaraSequencer.h
--- a/Engine/Source/Editor/NiagaraEditor/Private/NiagaraSequencer.h
+++ b/Engine/Source/Editor/NiagaraEditor/Private/NiagaraSequencer.h
@@ -57,25 +57,25 @@ public:
 		return EmitterSection;
 	}
 
-	int32 OnPaintSection(const FGeometry &AllottedGeometry, const FSlateRect &SectionClippingRect, FSlateWindowElementList &OutDrawElements, int32 LayerId, bool  bParentEnabled) const
+	virtual int32 OnPaintSection( FSequencerSectionPainter& InPainter ) const override
 	{
 		// draw the first run of the emitter
 		FSlateDrawElement::MakeBox
 			(
-			OutDrawElements,
-			LayerId,
-			AllottedGeometry.ToPaintGeometry(),
+			InPainter.DrawElements,
+			InPainter.LayerId,
+			InPainter.SectionGeometry.ToPaintGeometry(),
 			FEditorStyle::GetBrush("CurveEd.TimelineArea"),
-			SectionClippingRect,
+			InPainter.SectionClippingRect,
 			ESlateDrawEffect::None,
 			FLinearColor(0.3f, 0.3f, 0.6f)
 			);
 
 		// draw all loops of the emitter as 'ghosts' of the original section
-		float X = AllottedGeometry.AbsolutePosition.X;
-		float GeomW = AllottedGeometry.GetDrawSize().X;
-		float ClipW = SectionClippingRect.Right-SectionClippingRect.Left;
-		FSlateRect NewClipRect = SectionClippingRect;
+		float X = InPainter.SectionGeometry.AbsolutePosition.X;
+		float GeomW = InPainter.SectionGeometry.GetDrawSize().X;
+		float ClipW = InPainter.SectionClippingRect.Right-InPainter.SectionClippingRect.Left;
+		FSlateRect NewClipRect = InPainter.SectionClippingRect;
 		NewClipRect.Left += 5;
 		NewClipRect.Right -= 10;
 		for (int32 Loop = 0; Loop < EmitterSection->GetEmitterProps()->NumLoops - 1; Loop++)
@@ -84,9 +84,9 @@ public:
 			NewClipRect.Right += GeomW;
 			FSlateDrawElement::MakeBox
 				(
-				OutDrawElements,
-				LayerId,
-				AllottedGeometry.ToPaintGeometry(FVector2D(GeomW*(Loop+1), 0.0f), AllottedGeometry.GetDrawSize(), 1.0f),
+				InPainter.DrawElements,
+				InPainter.LayerId,
+				InPainter.SectionGeometry.ToPaintGeometry(FVector2D(GeomW*(Loop+1), 0.0f), InPainter.SectionGeometry.GetDrawSize(), 1.0f),
 				FEditorStyle::GetBrush("CurveEd.TimelineArea"),
 				NewClipRect,
 				ESlateDrawEffect::None,
@@ -94,7 +94,7 @@ public:
 				);
 		}
 
-		return LayerId;
+		return InPainter.LayerId;
 	}
 
 	FText GetDisplayName(void) const	{ return EmitterSection->GetEmitterName(); }
diff --git a/Engine/Source/Editor/Persona/Private/Persona.cpp b/Engine/Source/Editor/Persona/Private/Persona.cpp
--- a/Engine/Source/Editor/Persona/Private/Persona.cpp
+++ b/Engine/Source/Editor/Persona/Private/Persona.cpp
@@ -1799,24 +1799,6 @@ FLinearColor FPersona::GetWorldCentricTabColorScale() const
 	return FLinearColor( 0.5f, 0.25f, 0.35f, 0.5f );
 }
 
-void FPersona::SaveAsset_Execute()
-{
-	if( ensure( GetEditingObjects().Num() > 0 ) )
-	{
-		const FName CurrentMode = GetCurrentMode();
-
-		TArray<UObject*> EditorObjects = GetEditorObjectsForMode(CurrentMode);
-
-		TArray< UPackage* > PackagesToSave;
-		for(auto Iter = EditorObjects.CreateIterator(); Iter; ++Iter)
-		{
-			PackagesToSave.Add((*Iter)->GetOutermost());
-		}
-
-		FEditorFileUtils::PromptForCheckoutAndSave( PackagesToSave, /*bCheckDirty=*/ false, /*bPromptToSave=*/ false );
-	}
-}
-
 void FPersona::SaveAnimationAssets_Execute()
 {
 	// only save animation assets related to skeletons
@@ -2574,6 +2556,11 @@ TArray<UObject*> FPersona::GetEditorObjectsForMode(FName Mode) const
 	return TArray<UObject*>();
 }
 
+void FPersona::GetSaveableObjects(TArray<UObject*>& OutObjects) const
+{
+	OutObjects = GetEditorObjectsForMode(GetCurrentMode());
+}
+
 const FSlateBrush* FPersona::GetDirtyImageForMode(FName Mode) const
 {
 	TArray<UObject*> EditorObjects = GetEditorObjectsForMode(Mode);
diff --git a/Engine/Source/Editor/Persona/Private/Persona.h b/Engine/Source/Editor/Persona/Private/Persona.h
--- a/Engine/Source/Editor/Persona/Private/Persona.h
+++ b/Engine/Source/Editor/Persona/Private/Persona.h
@@ -197,7 +197,6 @@ public:
 	virtual FText GetToolkitToolTipText() const override;
 	virtual FString GetWorldCentricTabPrefix() const override;
 	virtual FLinearColor GetWorldCentricTabColorScale() const override;
-	virtual void SaveAsset_Execute() override;
 	//~ End IToolkit Interface
 
 	/** Saves all animation assets related to a skeleton */
@@ -297,6 +296,10 @@ protected:
 	/** Handler connected to open skeleton so we get notified when it changes */
 	void OnSkeletonHierarchyChanged();
 
+protected:
+	//~ IToolkit interface
+	virtual void GetSaveableObjects(TArray<UObject*>& OutObjects) const override;
+
 protected:
 	USkeleton* TargetSkeleton;
 	TWeakObjectPtr<UObject> CachedPreviewAsset;
diff --git a/Engine/Source/Editor/Sequencer/Private/DisplayNodes/SequencerDisplayNode.cpp b/Engine/Source/Editor/Sequencer/Private/DisplayNodes/SequencerDisplayNode.cpp
--- a/Engine/Source/Editor/Sequencer/Private/DisplayNodes/SequencerDisplayNode.cpp
+++ b/Engine/Source/Editor/Sequencer/Private/DisplayNodes/SequencerDisplayNode.cpp
@@ -59,7 +59,7 @@ private:
 	void CollectAllKeyTimes(TArray<float>& OutKeyTimes) const;
 
 	/** Adds a key time uniquely to an array of key times */
-	void AddKeyTime(float NewTime, TArray<float>& OutKeyTimes) const;
+	void AddKeyTime(const float& NewTime, TArray<float>& OutKeyTimes) const;
 
 private:
 
@@ -73,28 +73,32 @@ private:
 
 int32 SSequencerObjectTrack::OnPaint(const FPaintArgs& Args, const FGeometry& AllottedGeometry, const FSlateRect& MyClippingRect, FSlateWindowElementList& OutDrawElements, int32 LayerId, const FWidgetStyle& InWidgetStyle, bool bParentEnabled) const
 {
-	TArray<float> OutKeyTimes;
-	CollectAllKeyTimes(OutKeyTimes);
+	if (RootNode->GetSequencer().GetSettings()->GetShowCombinedKeyframes())
+	{
+		TArray<float> OutKeyTimes;
+		CollectAllKeyTimes(OutKeyTimes);
 	
-	FTimeToPixel TimeToPixelConverter(AllottedGeometry, ViewRange.Get());
+		FTimeToPixel TimeToPixelConverter(AllottedGeometry, ViewRange.Get());
 
-	for (int32 i = 0; i < OutKeyTimes.Num(); ++i)
-	{
-		float KeyPosition = TimeToPixelConverter.TimeToPixel(OutKeyTimes[i]);
-		static const FVector2D KeyMarkSize = FVector2D(3.f, 21.f);
-
-		FSlateDrawElement::MakeBox(
-			OutDrawElements,
-			LayerId+1,
-			AllottedGeometry.ToPaintGeometry(FVector2D(KeyPosition - FMath::CeilToFloat(KeyMarkSize.X/2.f), FMath::CeilToFloat(AllottedGeometry.Size.Y/2.f - KeyMarkSize.Y/2.f)), KeyMarkSize),
-			FEditorStyle::GetBrush("Sequencer.KeyMark"),
-			MyClippingRect,
-			ESlateDrawEffect::None,
-			FLinearColor(1.f, 1.f, 1.f, 1.f)
-		);
+		for (int32 i = 0; i < OutKeyTimes.Num(); ++i)
+		{
+			float KeyPosition = TimeToPixelConverter.TimeToPixel(OutKeyTimes[i]);
+			static const FVector2D KeyMarkSize = FVector2D(3.f, 21.f);
+
+			FSlateDrawElement::MakeBox(
+				OutDrawElements,
+				LayerId+1,
+				AllottedGeometry.ToPaintGeometry(FVector2D(KeyPosition - FMath::CeilToFloat(KeyMarkSize.X/2.f), FMath::CeilToFloat(AllottedGeometry.Size.Y/2.f - KeyMarkSize.Y/2.f)), KeyMarkSize),
+				FEditorStyle::GetBrush("Sequencer.KeyMark"),
+				MyClippingRect,
+				ESlateDrawEffect::None,
+				FLinearColor(1.f, 1.f, 1.f, 1.f)
+			);
+		}
+		return LayerId+1;
 	}
 
-	return LayerId+1;
+	return LayerId;
 }
 
 
@@ -110,7 +114,7 @@ void SSequencerObjectTrack::CollectAllKeyTimes(TArray<float>& OutKeyTimes) const
 	TArray<TSharedRef<FSequencerSectionKeyAreaNode>> OutNodes;
 	RootNode->GetChildKeyAreaNodesRecursively(OutNodes);
 
-	for (int32 i = 0; i < OutNodes.Num(); ++i)
+for (int32 i = 0; i < OutNodes.Num(); ++i)
 	{
 		TArray< TSharedRef<IKeyArea> > KeyAreas = OutNodes[i]->GetAllKeyAreas();
 		for (int32 j = 0; j < KeyAreas.Num(); ++j)
@@ -125,7 +129,7 @@ void SSequencerObjectTrack::CollectAllKeyTimes(TArray<float>& OutKeyTimes) const
 }
 
 
-void SSequencerObjectTrack::AddKeyTime(float NewTime, TArray<float>& OutKeyTimes) const
+void SSequencerObjectTrack::AddKeyTime(const float& NewTime, TArray<float>& OutKeyTimes) const
 {
 	// @todo Sequencer It might be more efficient to add each key and do the pruning at the end
 	for (float& KeyTime : OutKeyTimes)
diff --git a/Engine/Source/Editor/Sequencer/Private/DisplayNodes/SequencerObjectBindingNode.cpp b/Engine/Source/Editor/Sequencer/Private/DisplayNodes/SequencerObjectBindingNode.cpp
--- a/Engine/Source/Editor/Sequencer/Private/DisplayNodes/SequencerObjectBindingNode.cpp
+++ b/Engine/Source/Editor/Sequencer/Private/DisplayNodes/SequencerObjectBindingNode.cpp
@@ -270,13 +270,13 @@ FText FSequencerObjectBindingNode::GetIconToolTipText() const
 
 float FSequencerObjectBindingNode::GetNodeHeight() const
 {
-	return SequencerLayoutConstants::ObjectNodeHeight;
+	return SequencerLayoutConstants::ObjectNodeHeight + SequencerNodeConstants::CommonPadding*2;
 }
 
 
 FNodePadding FSequencerObjectBindingNode::GetNodePadding() const
 {
-	return FNodePadding(SequencerNodeConstants::CommonPadding);
+	return FNodePadding(0.f);//SequencerNodeConstants::CommonPadding);
 }
 
 
diff --git a/Engine/Source/Editor/Sequencer/Private/DisplayNodes/SequencerSectionCategoryNode.cpp b/Engine/Source/Editor/Sequencer/Private/DisplayNodes/SequencerSectionCategoryNode.cpp
--- a/Engine/Source/Editor/Sequencer/Private/DisplayNodes/SequencerSectionCategoryNode.cpp
+++ b/Engine/Source/Editor/Sequencer/Private/DisplayNodes/SequencerSectionCategoryNode.cpp
@@ -2,6 +2,7 @@
 
 #include "SequencerPrivatePCH.h"
 #include "SSequencer.h"
+#include "SKeyNavigationButtons.h"
 
 
 namespace SequencerSectionCategoryNodeConstants
@@ -18,6 +19,23 @@ bool FSequencerSectionCategoryNode::CanRenameNode() const
 	return false;
 }
 
+TSharedRef<SWidget> FSequencerSectionCategoryNode::GetCustomOutlinerContent()
+{
+	return SNew(SBox)
+		.HAlign(HAlign_Right)
+		.VAlign(VAlign_Center)
+		[
+			SNew(SHorizontalBox)
+
+			+ SHorizontalBox::Slot()
+			.AutoWidth()
+			.VAlign(VAlign_Center)
+			[
+				SNew(SKeyNavigationButtons, AsShared())
+			]
+		];
+}
+
 
 FText FSequencerSectionCategoryNode::GetDisplayName() const
 {
@@ -27,13 +45,13 @@ FText FSequencerSectionCategoryNode::GetDisplayName() const
 
 float FSequencerSectionCategoryNode::GetNodeHeight() const 
 {
-	return SequencerLayoutConstants::CategoryNodeHeight;
+	return SequencerLayoutConstants::CategoryNodeHeight + SequencerSectionCategoryNodeConstants::NodePadding*2;
 }
 
 
 FNodePadding FSequencerSectionCategoryNode::GetNodePadding() const
 {
-	return FNodePadding(SequencerSectionCategoryNodeConstants::NodePadding);
+	return FNodePadding(0.f);//FNodePadding(SequencerSectionCategoryNodeConstants::NodePadding);
 }
 
 
diff --git a/Engine/Source/Editor/Sequencer/Private/DisplayNodes/SequencerSectionCategoryNode.h b/Engine/Source/Editor/Sequencer/Private/DisplayNodes/SequencerSectionCategoryNode.h
--- a/Engine/Source/Editor/Sequencer/Private/DisplayNodes/SequencerSectionCategoryNode.h
+++ b/Engine/Source/Editor/Sequencer/Private/DisplayNodes/SequencerSectionCategoryNode.h
@@ -34,6 +34,7 @@ public:
 	// FSequencerDisplayNode interface
 
 	virtual bool CanRenameNode() const override;
+	virtual TSharedRef<SWidget> GetCustomOutlinerContent() override;
 	virtual FText GetDisplayName() const override;
 	virtual float GetNodeHeight() const override;
 	virtual FNodePadding GetNodePadding() const override;
diff --git a/Engine/Source/Editor/Sequencer/Private/DisplayNodes/SequencerSectionKeyAreaNode.cpp b/Engine/Source/Editor/Sequencer/Private/DisplayNodes/SequencerSectionKeyAreaNode.cpp
--- a/Engine/Source/Editor/Sequencer/Private/DisplayNodes/SequencerSectionKeyAreaNode.cpp
+++ b/Engine/Source/Editor/Sequencer/Private/DisplayNodes/SequencerSectionKeyAreaNode.cpp
@@ -79,7 +79,7 @@ float FSequencerSectionKeyAreaNode::GetNodeHeight() const
 
 FNodePadding FSequencerSectionKeyAreaNode::GetNodePadding() const
 {
-	return FNodePadding(0.f, 1.f);
+	return FNodePadding(0.f);//FNodePadding(0.f, 1.f);
 }
 
 
diff --git a/Engine/Source/Editor/Sequencer/Private/DisplayNodes/SequencerTrackNode.cpp b/Engine/Source/Editor/Sequencer/Private/DisplayNodes/SequencerTrackNode.cpp
--- a/Engine/Source/Editor/Sequencer/Private/DisplayNodes/SequencerTrackNode.cpp
+++ b/Engine/Source/Editor/Sequencer/Private/DisplayNodes/SequencerTrackNode.cpp
@@ -178,15 +178,51 @@ TSharedRef<SWidget> FSequencerTrackNode::GetCustomOutlinerContent()
 
 		TSharedPtr<SWidget> Widget = AssociatedEditor.BuildOutlinerEditWidget(ObjectBinding, AssociatedTrack.Get(), Params);
 
+		TSharedRef<SHorizontalBox> BoxPanel = SNew(SHorizontalBox);
+
+		bool bHasKeyableAreas = false;
+
+		TArray<TSharedRef<FSequencerSectionKeyAreaNode>> ChildKeyAreaNodes;
+		FSequencerDisplayNode::GetChildKeyAreaNodesRecursively(ChildKeyAreaNodes);
+		for (int32 ChildIndex = 0; ChildIndex < ChildKeyAreaNodes.Num() && !bHasKeyableAreas; ++ChildIndex)
+		{
+			TArray< TSharedRef<IKeyArea> > ChildKeyAreas = ChildKeyAreaNodes[ChildIndex]->GetAllKeyAreas();
+
+			for (int32 ChildKeyAreaIndex = 0; ChildKeyAreaIndex < ChildKeyAreas.Num() && !bHasKeyableAreas; ++ChildKeyAreaIndex)
+			{
+				if (ChildKeyAreas[ChildKeyAreaIndex]->CanCreateKeyEditor())
+				{
+					bHasKeyableAreas = true;
+				}
+			}
+		}
+
 		if (Widget.IsValid())
 		{
-			return SNew(SBox)
-				.VAlign(VAlign_Center)
-				.HAlign(HAlign_Right)
-				[
-					Widget.ToSharedRef()
-				];
+			BoxPanel->AddSlot()
+			.VAlign(VAlign_Center)
+			.HAlign(HAlign_Right)
+			[
+				Widget.ToSharedRef()
+			];
+		}
+
+		if (bHasKeyableAreas)
+		{
+			BoxPanel->AddSlot()
+			.AutoWidth()
+			.VAlign(VAlign_Center)
+			[
+				SNew(SKeyNavigationButtons, AsShared())
+			];
 		}
+
+		return SNew(SBox)
+			.VAlign(VAlign_Center)
+			.HAlign(HAlign_Right)
+			[
+				BoxPanel
+			];
 	}
 
 
@@ -219,15 +255,17 @@ FText FSequencerTrackNode::GetDisplayName() const
 
 float FSequencerTrackNode::GetNodeHeight() const
 {
-	return (Sections.Num() > 0)
+	float Height = (Sections.Num() > 0)
 		? Sections[0]->GetSectionHeight() * (GetMaxRowIndex() + 1)
 		: SequencerLayoutConstants::SectionAreaDefaultHeight;
+
+	return Height + 2*SequencerNodeConstants::CommonPadding;
 }
 
 
 FNodePadding FSequencerTrackNode::GetNodePadding() const
 {
-	return FNodePadding(SequencerNodeConstants::CommonPadding);
+	return FNodePadding(0.f);//SequencerNodeConstants::CommonPadding);
 }
 
 
diff --git a/Engine/Source/Editor/Sequencer/Private/GroupedKeyArea.cpp b/Engine/Source/Editor/Sequencer/Private/GroupedKeyArea.cpp
--- a/Engine/Source/Editor/Sequencer/Private/GroupedKeyArea.cpp
+++ b/Engine/Source/Editor/Sequencer/Private/GroupedKeyArea.cpp
@@ -477,9 +477,9 @@ TSharedRef<SWidget> FGroupedKeyArea::CreateKeyEditor(ISequencer* Sequencer)
 }
 
 
-FLinearColor FGroupedKeyArea::GetColor()
+TOptional<FLinearColor> FGroupedKeyArea::GetColor()
 {
-	return FLinearColor(0.1f, 0.1f, 0.1f, 0.7f);
+	return TOptional<FLinearColor>();
 }
 
 
diff --git a/Engine/Source/Editor/Sequencer/Private/SAnimationOutlinerTreeNode.cpp b/Engine/Source/Editor/Sequencer/Private/SAnimationOutlinerTreeNode.cpp
--- a/Engine/Source/Editor/Sequencer/Private/SAnimationOutlinerTreeNode.cpp
+++ b/Engine/Source/Editor/Sequencer/Private/SAnimationOutlinerTreeNode.cpp
@@ -11,10 +11,73 @@
 #include "IKeyArea.h"
 #include "SEditableLabel.h"
 #include "SSequencerTreeView.h"
+#include "SColorPicker.h"
 
 
 #define LOCTEXT_NAMESPACE "AnimationOutliner"
 
+class STrackColorPicker : public SCompoundWidget
+{
+public:
+	static void OnOpen()
+	{
+		if (!Transaction.IsValid())
+		{
+			Transaction.Reset(new FScopedTransaction(LOCTEXT("ChangeTrackColor", "Change Track Color")));
+		}
+	}
+
+	static void OnClose()
+	{
+		if (!Transaction.IsValid())
+		{
+			return;
+		}
+
+		if (!bMadeChanges)
+		{
+			Transaction->Cancel();
+		}
+		Transaction.Reset();
+	}
+
+	SLATE_BEGIN_ARGS(STrackColorPicker){}
+	SLATE_END_ARGS()
+
+	void Construct(const FArguments& InArgs, UMovieSceneTrack* InTrack)
+	{
+		bMadeChanges = false;
+		Track = InTrack;
+		ChildSlot
+		[
+			SNew(SColorPicker)
+			.DisplayInlineVersion(true)
+			.TargetColorAttribute(this, &STrackColorPicker::GetTrackColor)
+			.OnColorCommitted(this, &STrackColorPicker::SetTrackColor)
+		];
+	}
+
+	FLinearColor GetTrackColor() const
+	{
+		return Track->GetColorTint().ReinterpretAsLinear();
+	}
+
+	void SetTrackColor(FLinearColor NewColor)
+	{
+		bMadeChanges = true;
+		Track->Modify();
+		Track->SetColorTint(NewColor.ToFColor(false));
+	}
+
+private:
+	UMovieSceneTrack* Track;
+	static TUniquePtr<FScopedTransaction> Transaction;
+	static bool bMadeChanges;
+};
+
+TUniquePtr<FScopedTransaction> STrackColorPicker::Transaction;
+bool STrackColorPicker::bMadeChanges = false;
+
 SAnimationOutlinerTreeNode::~SAnimationOutlinerTreeNode()
 {
 	DisplayNode->OnRenameRequested().RemoveAll(this);
@@ -60,64 +123,100 @@ void SAnimationOutlinerTreeNode::Construct( const FArguments& InArgs, TSharedRef
 		.BorderBackgroundColor( this, &SAnimationOutlinerTreeNode::GetNodeBackgroundTint )
 		.Padding(FMargin(0, Node->GetNodePadding().Combined() / 2))
 		[
-			SNew(SBox)
-			.HeightOverride_Lambda(NodeHeight)
-			.Padding(FMargin(5.0f, 0.0f))
-			[
-				SNew( SHorizontalBox )
-
-				// Expand track lanes button
-				+ SHorizontalBox::Slot()
-				.Padding(FMargin(2.f, 0.f, 4.f, 0.f))
-				.VAlign( VAlign_Center )
-				.AutoWidth()
-				[
-					SNew(SExpanderArrow, InTableRow).IndentAmount(SequencerLayoutConstants::IndentAmount)
-				]
+			SNew( SHorizontalBox )
 
-				// Icon
-				+ SHorizontalBox::Slot()
-				.Padding(FMargin(0.f, 0.f, 4.f, 0.f))
-				.VAlign(VAlign_Center)
-				.AutoWidth()
+			+ SHorizontalBox::Slot()
+			[
+				SNew(SBox)
+				.HeightOverride_Lambda(NodeHeight)
+				.Padding(FMargin(5.0f, 0.0f))
 				[
-					SNew(SOverlay)
+					SNew( SHorizontalBox )
 
-					+ SOverlay::Slot()
+					// Expand track lanes button
+					+ SHorizontalBox::Slot()
+					.Padding(FMargin(2.f, 0.f, 4.f, 0.f))
+					.VAlign( VAlign_Center )
+					.AutoWidth()
 					[
-						SNew(SImage)
-						.Image(InArgs._IconBrush)
+						SNew(SExpanderArrow, InTableRow).IndentAmount(SequencerLayoutConstants::IndentAmount)
 					]
 
-					+ SOverlay::Slot()
-					.VAlign(VAlign_Top)
-					.HAlign(HAlign_Right)
+					// Icon
+					+ SHorizontalBox::Slot()
+					.Padding(FMargin(0.f, 0.f, 4.f, 0.f))
+					.VAlign(VAlign_Center)
+					.AutoWidth()
 					[
-						SNew(SImage)
-						.Image(InArgs._IconOverlayBrush)
+						SNew(SOverlay)
+
+						+ SOverlay::Slot()
+						[
+							SNew(SImage)
+							.Image(InArgs._IconBrush)
+						]
+
+						+ SOverlay::Slot()
+						.VAlign(VAlign_Top)
+						.HAlign(HAlign_Right)
+						[
+							SNew(SImage)
+							.Image(InArgs._IconOverlayBrush)
+						]
+
+						+ SOverlay::Slot()
+						[
+							SNew(SSpacer)
+							.Visibility(EVisibility::Visible)
+							.ToolTipText(InArgs._IconToolTipText)
+						]
 					]
 
-					+ SOverlay::Slot()
+					// Label Slot
+					+ SHorizontalBox::Slot()
+					.VAlign(VAlign_Center)
+					.Padding(FMargin(0.f, 0.f, 4.f, 0.f))
+					.AutoWidth()
 					[
-						SNew(SSpacer)
-						.Visibility(EVisibility::Visible)
-						.ToolTipText(InArgs._IconToolTipText)
+						EditableLabel.ToSharedRef()
 					]
-				]
 
-				// Label Slot
-				+ SHorizontalBox::Slot()
-				.VAlign(VAlign_Center)
-				.Padding(FMargin(0.f, 0.f, 4.f, 0.f))
-				.AutoWidth()
-				[
-					EditableLabel.ToSharedRef()
+					// Arbitrary customization slot
+					+ SHorizontalBox::Slot()
+					[
+						InArgs._CustomContent.Widget
+					]
 				]
+			]
 
-				// Arbitrary customization slot
-				+ SHorizontalBox::Slot()
+			+ SHorizontalBox::Slot()
+			.AutoWidth()
+			[
+				SNew(SComboButton)
+				.ContentPadding(0)
+				.VAlign(VAlign_Fill)
+				.HasDownArrow(false)
+				.ButtonStyle(FEditorStyle::Get(), "Sequencer.AnimationOutliner.ColorStrip")
+				.OnGetMenuContent(this, &SAnimationOutlinerTreeNode::OnGetColorPicker)
+				.OnMenuOpenChanged_Lambda([](bool bIsOpen){
+					if (bIsOpen)
+					{
+						STrackColorPicker::OnOpen();
+					}
+					else if (!bIsOpen)
+					{
+						STrackColorPicker::OnClose();
+					}
+				})
+				.ButtonContent()
 				[
-					InArgs._CustomContent.Widget
+					SNew(SBox)
+					.WidthOverride(6.f)
+					[
+						SNew(SImage)
+						.Image(FEditorStyle::GetBrush("WhiteBrush"))
+						.ColorAndOpacity(this, &SAnimationOutlinerTreeNode::GetTrackColorTint)
+					]
 				]
 			]
 		];
@@ -332,6 +431,51 @@ FSlateColor SAnimationOutlinerTreeNode::GetNodeBackgroundTint() const
 	}
 }
 
+TSharedRef<SWidget> SAnimationOutlinerTreeNode::OnGetColorPicker() const
+{
+	UMovieSceneTrack* Track = nullptr;
+
+	FSequencerDisplayNode* Current = DisplayNode.Get();
+	while (Current && Current->GetType() != ESequencerNode::Object)
+	{
+		if (Current->GetType() == ESequencerNode::Track)
+		{
+			Track = static_cast<FSequencerTrackNode*>(Current)->GetTrack();
+			if (Track)
+			{
+				break;
+			}
+		}
+		Current = Current->GetParent().Get();
+	}
+
+	if (!Track)
+	{
+		return SNullWidget::NullWidget;
+	}
+
+	return SNew(STrackColorPicker, Track);
+}
+
+FSlateColor SAnimationOutlinerTreeNode::GetTrackColorTint() const
+{
+	FSequencerDisplayNode* Current = DisplayNode.Get();
+	while (Current && Current->GetType() != ESequencerNode::Object)
+	{
+		if (Current->GetType() == ESequencerNode::Track)
+		{
+			UMovieSceneTrack* Track = static_cast<FSequencerTrackNode*>(Current)->GetTrack();
+			if (Track)
+			{
+				return FLinearColor(Track->GetColorTint()).CopyWithNewOpacity(1.f);
+			}
+		}
+		Current = Current->GetParent().Get();
+	}
+
+	return FLinearColor::Transparent;
+}
+
 FSlateColor SAnimationOutlinerTreeNode::GetForegroundBasedOnSelection() const
 {
 	FSequencer& Sequencer = DisplayNode->GetSequencer();
diff --git a/Engine/Source/Editor/Sequencer/Private/SAnimationOutlinerTreeNode.h b/Engine/Source/Editor/Sequencer/Private/SAnimationOutlinerTreeNode.h
--- a/Engine/Source/Editor/Sequencer/Private/SAnimationOutlinerTreeNode.h
+++ b/Engine/Source/Editor/Sequencer/Private/SAnimationOutlinerTreeNode.h
@@ -53,6 +53,9 @@ private:
 
 	FSlateColor GetForegroundBasedOnSelection() const;
 
+	/** Get the tint to apply to the color indicator based on this node's track */
+	FSlateColor GetTrackColorTint() const;
+
 	/**
 	 * @return The border image to show in the tree node.
 	 */
@@ -88,6 +91,9 @@ private:
 	/** Called when nodes are selected. */
 	void OnSelectionChanged(TArray<TSharedPtr<FSequencerDisplayNode>> AffectedNodes);
 
+	/** Called when the user clicks the track color */
+	TSharedRef<SWidget> OnGetColorPicker() const;
+
 private:
 
 	/** Layout node the widget is visualizing. */
diff --git a/Engine/Source/Editor/Sequencer/Private/SSequencer.cpp b/Engine/Source/Editor/Sequencer/Private/SSequencer.cpp
--- a/Engine/Source/Editor/Sequencer/Private/SSequencer.cpp
+++ b/Engine/Source/Editor/Sequencer/Private/SSequencer.cpp
@@ -2,7 +2,7 @@
 
 #include "SequencerPrivatePCH.h"
 #include "SSequencer.h"
-#include "Editor/SequencerWidgets/Public/SequencerWidgetsModule.h"
+#include "ISequencerWidgetsModule.h"
 #include "Sequencer.h"
 #include "MovieScene.h"
 #include "MovieSceneSection.h"
@@ -14,7 +14,7 @@
 #include "SSequencerTrackArea.h"
 #include "SSequencerTrackOutliner.h"
 #include "SequencerNodeTree.h"
-#include "TimeSliderController.h"
+#include "SequencerTimeSliderController.h"
 #include "DragAndDrop/AssetDragDropOp.h"
 #include "DragAndDrop/ActorDragDropGraphEdOp.h"
 #include "DragAndDrop/ClassDragDropOp.h"
@@ -96,6 +96,8 @@ void SSequencer::Construct(const FArguments& InArgs, TSharedRef<FSequencer> InSe
 	SequencerPtr = InSequencer;
 	bIsActiveTimerRegistered = false;
 	bUserIsSelecting = false;
+	CachedClampRange = TRange<float>::Empty();
+	CachedViewRange = TRange<float>::Empty();
 
 	USelection::SelectionChangedEvent.AddSP(this, &SSequencer::OnActorSelectionChanged);
 
@@ -105,7 +107,7 @@ void SSequencer::Construct(const FArguments& InArgs, TSharedRef<FSequencer> InSe
 	// Create a node tree which contains a tree of movie scene data to display in the sequence
 	SequencerNodeTree = MakeShareable( new FSequencerNodeTree( InSequencer.Get() ) );
 
-	FSequencerWidgetsModule& SequencerWidgets = FModuleManager::Get().LoadModuleChecked<FSequencerWidgetsModule>( "SequencerWidgets" );
+	ISequencerWidgetsModule& SequencerWidgets = FModuleManager::Get().LoadModuleChecked<ISequencerWidgetsModule>( "SequencerWidgets" );
 
 	OnBeginPlaybackRangeDrag = InArgs._OnBeginPlaybackRangeDrag;
 	OnEndPlaybackRangeDrag = InArgs._OnEndPlaybackRangeDrag;
@@ -625,6 +627,16 @@ TSharedRef<SWidget> SSequencer::MakeToolBar()
 				FSlateIcon(FEditorStyle::GetStyleSetName(), "Sequencer.Save")
 			);
 
+			ToolBarBuilder.AddToolBarButton(
+				FUIAction(FExecuteAction::CreateSP(this, &SSequencer::OnSaveMovieSceneAsClicked)),
+				NAME_None,
+				LOCTEXT("SaveAs", "Save As"),
+				LOCTEXT("SaveAsTooltip", "Saves the current level sequence under a different name"),
+				FSlateIcon(FEditorStyle::GetStyleSetName(), "AssetEditor.SaveAssetAs")
+			);
+
+			// @todo sequencer: implement FSequencer::DiscardChanges()
+			//ToolBarBuilder.AddToolBarButton( FSequencerCommands::Get().DiscardChanges );
 			ToolBarBuilder.AddToolBarButton( FSequencerCommands::Get().RenderMovie );
 			ToolBarBuilder.AddSeparator();
 		}
@@ -766,6 +778,8 @@ TSharedRef<SWidget> SSequencer::MakeGeneralMenu()
 			MenuBuilder.AddMenuEntry( FSequencerCommands::Get().ToggleLabelBrowser );
 		}
 
+		MenuBuilder.AddMenuEntry( FSequencerCommands::Get().ToggleCombinedKeyframes );
+
 		if (SequencerPtr.Pin()->IsLevelEditorSequencer())
 		{
 			MenuBuilder.AddMenuEntry( FSequencerCommands::Get().FindInContentBrowser );
@@ -878,6 +892,7 @@ TSharedRef<SWidget> SSequencer::MakeGeneralMenu()
 
 		MenuBuilder.AddMenuSeparator();
 		MenuBuilder.AddMenuEntry( FSequencerCommands::Get().ToggleShowGotoBox );
+		MenuBuilder.AddMenuEntry( FSequencerCommands::Get().ToggleLinkCurveEditorTimeRange );
 	}
 
 	return MenuBuilder.MakeWidget();
@@ -1453,6 +1468,12 @@ void SSequencer::OnSaveMovieSceneClicked()
 }
 
 
+void SSequencer::OnSaveMovieSceneAsClicked()
+{
+	SequencerPtr.Pin()->SaveCurrentMovieSceneAs();
+}
+
+
 void SSequencer::StepToNextKey()
 {
 	StepToKey(true, false);
@@ -1641,6 +1662,25 @@ void SSequencer::OnCurveEditorVisibilityChanged()
 {
 	if (CurveEditor.IsValid())
 	{
+		if (!Settings->GetLinkCurveEditorTimeRange())
+		{
+			TRange<float> ClampRange = SequencerPtr.Pin()->GetClampRange();
+			if (CachedClampRange.IsEmpty())
+			{
+				CachedClampRange = ClampRange;
+			}
+			SequencerPtr.Pin()->SetClampRange(CachedClampRange);
+			CachedClampRange = ClampRange;
+
+			TRange<float> ViewRange = SequencerPtr.Pin()->GetViewRange();
+			if (CachedViewRange.IsEmpty())
+			{
+				CachedViewRange = ViewRange;
+			}
+			SequencerPtr.Pin()->SetViewRange(CachedViewRange);
+			CachedViewRange = ViewRange;
+		}
+
 		// Only zoom horizontally if the editor is visible
 		CurveEditor->SetAllowAutoFrame(Settings->GetShowCurveEditor());
 	}
diff --git a/Engine/Source/Editor/Sequencer/Private/SSequencer.h b/Engine/Source/Editor/Sequencer/Private/SSequencer.h
--- a/Engine/Source/Editor/Sequencer/Private/SSequencer.h
+++ b/Engine/Source/Editor/Sequencer/Private/SSequencer.h
@@ -177,6 +177,9 @@ public:
 	/** Called when the save button is clicked */
 	void OnSaveMovieSceneClicked();
 
+	/** Called when the save-as button is clicked */
+	void OnSaveMovieSceneAsClicked();
+
 	/** Access the tree view for this sequencer */
 	TSharedPtr<SSequencerTreeView> GetTreeView() const;
 
@@ -429,4 +432,8 @@ private:
 
 	/** Called when the user has finished dragging the playback range */
 	FSimpleDelegate OnEndPlaybackRangeDrag;
+
+	/** Cached clamp and view range for unlinking the curve editor time range */
+	TRange<float> CachedClampRange;
+	TRange<float> CachedViewRange;
 };
diff --git a/Engine/Source/Editor/Sequencer/Private/SSequencerCurveEditor.cpp b/Engine/Source/Editor/Sequencer/Private/SSequencerCurveEditor.cpp
--- a/Engine/Source/Editor/Sequencer/Private/SSequencerCurveEditor.cpp
+++ b/Engine/Source/Editor/Sequencer/Private/SSequencerCurveEditor.cpp
@@ -3,7 +3,7 @@
 #include "SequencerPrivatePCH.h"
 #include "SCurveEditor.h"
 #include "Sequencer.h"
-#include "TimeSliderController.h"
+#include "SequencerTimeSliderController.h"
 
 #define LOCTEXT_NAMESPACE "SequencerCurveEditor"
 
diff --git a/Engine/Source/Editor/Sequencer/Private/SSequencerCurveEditorToolBar.cpp b/Engine/Source/Editor/Sequencer/Private/SSequencerCurveEditorToolBar.cpp
--- a/Engine/Source/Editor/Sequencer/Private/SSequencerCurveEditorToolBar.cpp
+++ b/Engine/Source/Editor/Sequencer/Private/SSequencerCurveEditorToolBar.cpp
@@ -136,6 +136,9 @@ TSharedRef<SWidget> SSequencerCurveEditorToolBar::MakeCurveEditorCurveOptionsMen
 
 	FMenuBuilder MenuBuilder( true, CurveEditorCommandList );
 
+	MenuBuilder.AddMenuEntry( FRichCurveEditorCommands::Get().BakeCurve);
+	MenuBuilder.AddMenuEntry( FRichCurveEditorCommands::Get().ReduceCurve);
+
 	MenuBuilder.AddSubMenu(
 		LOCTEXT( "PreInfinitySubMenu", "Pre-Infinity" ),
 		LOCTEXT( "PreInfinitySubMenuToolTip", "Pre-Infinity Extrapolation" ),
diff --git a/Engine/Source/Editor/Sequencer/Private/SSequencerSection.cpp b/Engine/Source/Editor/Sequencer/Private/SSequencerSection.cpp
--- a/Engine/Source/Editor/Sequencer/Private/SSequencerSection.cpp
+++ b/Engine/Source/Editor/Sequencer/Private/SSequencerSection.cpp
@@ -10,16 +10,156 @@
 #include "CommonMovieSceneTools.h"
 #include "SequencerHotspots.h"
 #include "ScopedTransaction.h"
+#include "SequencerSectionPainter.h"
 
 double SSequencerSection::SelectionThrobEndTime = 0;
 
+struct FSequencerSectionPainterImpl : FSequencerSectionPainter
+{
+	FSequencerSectionPainterImpl(FSequencer& InSequencer, UMovieSceneSection& InSection, FSlateWindowElementList& _OutDrawElements, const FGeometry& InSectionGeometry)
+		: FSequencerSectionPainter(_OutDrawElements, InSectionGeometry, InSection)
+		, Sequencer(InSequencer)
+		, TimeToPixelConverter(InSection.IsInfinite() ?
+			FTimeToPixel(SectionGeometry, Sequencer.GetViewRange()) :
+			FTimeToPixel(SectionGeometry, TRange<float>(InSection.GetStartTime(), InSection.GetEndTime()))
+			)
+	{
+		CalculateSelectionColor();
+	}
+
+	virtual int32 PaintSectionBackground(const FSequencerSectionTintRegion* Tints, uint32 NumTints) override
+	{
+		const ESlateDrawEffect::Type DrawEffects = bParentEnabled
+			? ESlateDrawEffect::None
+			: ESlateDrawEffect::DisabledEffect;
+
+		static const FSlateBrush* SectionBackgroundBrush = FEditorStyle::GetBrush("Sequencer.Section.Background");
+		static const FSlateBrush* SectionBackgroundTintBrush = FEditorStyle::GetBrush("Sequencer.Section.BackgroundTint");
+		static const FSlateBrush* SelectedSectionOverlay = FEditorStyle::GetBrush("Sequencer.Section.SelectedSectionOverlay");
+
+		FGeometry InfiniteGeometry = SectionGeometry.MakeChild(FVector2D(-100.f, 0.f), FVector2D(SectionGeometry.GetLocalSize().X + 200.f, SectionGeometry.GetLocalSize().Y));
+
+		FSlateDrawElement::MakeBox(
+			DrawElements,
+			LayerId++,
+			Section.IsInfinite() ?
+				InfiniteGeometry.ToPaintGeometry() :
+				SectionGeometry.ToPaintGeometry(),
+			SectionBackgroundBrush,
+			SectionClippingRect,
+			DrawEffects
+		);
+
+		const FSlateLayoutTransform LayoutTransform = SectionGeometry.GetAccumulatedLayoutTransform();
+
+		//const float TintOpacity = Args.bHovered ? 0.1f : 0.07f;
+		while (NumTints-- > 0)
+		{
+			FSlateRect SubClippingRect = SectionClippingRect;
+
+			if (Tints->Range.IsSet())
+			{
+				const float StartPixel = TimeToPixelConverter.TimeToPixel(Tints->Range->GetLowerBoundValue());
+				const float EndPixel = TimeToPixelConverter.TimeToPixel(Tints->Range->GetUpperBoundValue());
+
+				SubClippingRect = SubClippingRect.IntersectionWith(
+					TransformRect(LayoutTransform, FSlateRect(StartPixel, 0.0f, EndPixel, SectionGeometry.Size.Y))
+					);
+			}
+			
+			// We never show colors with full tint
+			FLinearColor TintColor(Tints->Color);
+			TintColor.A *= 0.4f;
+
+			FSlateDrawElement::MakeBox(
+				DrawElements,
+				LayerId,
+				SectionGeometry.ToPaintGeometry(),
+				SectionBackgroundTintBrush,
+				SubClippingRect.InsetBy(FMargin(1.f)),
+				DrawEffects,
+				TintColor
+			);
+
+			Tints++;
+		}
+
+		++LayerId;
+
+		// This is now invalid
+		Tints = nullptr;
+
+		// Draw the selection hash
+		if (SelectionColor.IsSet())
+		{
+			FSlateDrawElement::MakeBox(
+				DrawElements,
+				LayerId,
+				SectionGeometry.ToPaintGeometry(FVector2D(1, 1), SectionGeometry.Size - FVector2D(2,2)),
+				SelectedSectionOverlay,
+				SectionClippingRect,
+				DrawEffects,
+				SelectionColor.GetValue().CopyWithNewOpacity(0.8f)
+			);
+		}
+
+		return LayerId;
+	}
+
+	virtual const FTimeToPixel& GetTimeConverter() const
+	{
+		return TimeToPixelConverter;
+	}
+
+	void CalculateSelectionColor()
+	{
+		FSequencerSelection& Selection = Sequencer.GetSelection();
+		FSequencerSelectionPreview& SelectionPreview = Sequencer.GetSelectionPreview();
+
+		ESelectionPreviewState SelectionPreviewState = SelectionPreview.GetSelectionState(&Section);
+
+		if (SelectionPreviewState == ESelectionPreviewState::NotSelected)
+		{
+			// Explicitly not selected in the preview selection
+			return;
+		}
+		
+		if (SelectionPreviewState == ESelectionPreviewState::Undefined && !Selection.IsSelected(&Section))
+		{
+			// No preview selection for this section, and it's not selected
+			return;
+		}
+
+		SelectionColor = FEditorStyle::GetSlateColor(SequencerSectionConstants::SelectionColorName).GetColor(FWidgetStyle());
+
+		// Use a muted selection color for selection previews
+		if (SelectionPreviewState == ESelectionPreviewState::Selected)
+		{
+			SelectionColor.GetValue() = SelectionColor.GetValue().LinearRGBToHSV();
+			SelectionColor.GetValue().R += 0.1f; // +10% hue
+			SelectionColor.GetValue().G = 0.6f; // 60% saturation
+
+			SelectionColor = SelectionColor.GetValue().HSVToLinearRGB();
+		}
+		else if (Selection.GetActiveSelection() != FSequencerSelection::EActiveSelection::KeyAndSection)
+		{
+			// Use an inactive selection color for existing selections that are not active
+			SelectionColor = FEditorStyle::GetSlateColor(SequencerSectionConstants::SelectionInactiveColorName).GetColor(FWidgetStyle());
+		}
+	}
+
+	TOptional<FLinearColor> SelectionColor;
+	FSequencer& Sequencer;
+	FTimeToPixel TimeToPixelConverter;
+};
 
 void SSequencerSection::Construct( const FArguments& InArgs, TSharedRef<FSequencerTrackNode> SectionNode, int32 InSectionIndex )
 {
 	SectionIndex = InSectionIndex;
 	ParentSectionArea = SectionNode;
 	SectionInterface = SectionNode->GetSections()[InSectionIndex];
-	Layout = FKeyAreaLayout(*SectionNode, InSectionIndex);
+	Layout = FSectionLayout(*SectionNode, InSectionIndex);
+	HandleOffsetPx = 0.f;
 
 	ChildSlot
 	[
@@ -34,10 +174,10 @@ FVector2D SSequencerSection::ComputeDesiredSize(float) const
 }
 
 
-FGeometry SSequencerSection::GetKeyAreaGeometry( const FKeyAreaLayoutElement& KeyArea, const FGeometry& SectionGeometry ) const
+FGeometry SSequencerSection::GetKeyAreaGeometry( const FSectionLayoutElement& KeyArea, const FGeometry& SectionGeometry ) const
 {
 	// Compute the geometry for the key area
-	return SectionGeometry.MakeChild( FVector2D( 0, KeyArea.GetOffset() ), FVector2D( SectionGeometry.Size.X, KeyArea.GetHeight(SectionGeometry) ) );
+	return SectionGeometry.MakeChild( FVector2D( 0, KeyArea.GetOffset() ), FVector2D( SectionGeometry.Size.X, KeyArea.GetHeight()  ) );
 }
 
 
@@ -48,9 +188,13 @@ FSequencerSelectedKey SSequencerSection::GetKeyUnderMouse( const FVector2D& Mous
 	UMovieSceneSection& Section = *SectionInterface->GetSectionObject();
 
 	// Search every key area until we find the one under the mouse
-	for (const FKeyAreaLayoutElement& Element : Layout->GetElements())
+	for (const FSectionLayoutElement& Element : Layout->GetElements())
 	{
 		TSharedPtr<IKeyArea> KeyArea = Element.GetKeyArea();
+		if (!KeyArea.IsValid())
+		{
+			continue;
+		}
 
 		// Compute the current key area geometry
 		FGeometry KeyAreaGeometryPadded = GetKeyAreaGeometry( Element, AllottedGeometry );
@@ -99,9 +243,13 @@ FSequencerSelectedKey SSequencerSection::CreateKeyUnderMouse( const FVector2D& M
 	FGeometry SectionGeometry = MakeSectionGeometryWithoutHandles( AllottedGeometry, SectionInterface );
 
 	// Search every key area until we find the one under the mouse
-	for (const FKeyAreaLayoutElement& Element : Layout->GetElements())
+	for (const FSectionLayoutElement& Element : Layout->GetElements())
 	{
 		TSharedPtr<IKeyArea> KeyArea = Element.GetKeyArea();
+		if (!KeyArea.IsValid())
+		{
+			continue;
+		}
 
 		// Compute the current key area geometry
 		FGeometry KeyAreaGeometryPadded = GetKeyAreaGeometry( Element, AllottedGeometry );
@@ -138,12 +286,16 @@ FSequencerSelectedKey SSequencerSection::CreateKeyUnderMouse( const FVector2D& M
 
 					TArray<FKeyHandle> KeyHandles = KeyArea->AddKeyUnique(KeyTime+TimeFuzz, GetSequencer().GetKeyInterpolation(), KeyTime);
 
-					Layout = FKeyAreaLayout(*ParentSectionArea, SectionIndex);
+					Layout = FSectionLayout(*ParentSectionArea, SectionIndex);
 
 					// Look specifically for the key with the offset key time
-					for (const FKeyAreaLayoutElement& NewElement : Layout->GetElements())
+					for (const FSectionLayoutElement& NewElement : Layout->GetElements())
 					{
 						TSharedPtr<IKeyArea> NewKeyArea = NewElement.GetKeyArea();
+						if (!NewKeyArea.IsValid())
+						{
+							continue;
+						}
 
 						for (auto KeyHandle : KeyHandles)
 						{
@@ -160,7 +312,7 @@ FSequencerSelectedKey SSequencerSection::CreateKeyUnderMouse( const FVector2D& M
 				else
 				{
 					KeyArea->AddKeyUnique(KeyTime, GetSequencer().GetKeyInterpolation());
-					Layout = FKeyAreaLayout(*ParentSectionArea, SectionIndex);
+					Layout = FSectionLayout(*ParentSectionArea, SectionIndex);
 						
 					return GetKeyUnderMouse(MousePosition, AllottedGeometry);
 				}
@@ -217,17 +369,37 @@ int32 SSequencerSection::OnPaint( const FPaintArgs& Args, const FGeometry& Allot
 	const bool bLocked = SectionObject.IsLocked();
 	const ESlateDrawEffect::Type DrawEffects = bEnabled ? ESlateDrawEffect::None : ESlateDrawEffect::DisabledEffect;
 
-	int32 StartLayer = SCompoundWidget::OnPaint( Args, AllottedGeometry, MyClippingRect, OutDrawElements, LayerId, InWidgetStyle, bEnabled );
+
 	FGeometry SectionGeometry = MakeSectionGeometryWithoutHandles( AllottedGeometry, SectionInterface );
-	FSlateRect SectionClipRect = SectionGeometry.GetClippingRect().IntersectionWith( MyClippingRect );
+
+	FSequencerSectionPainterImpl Painter(ParentSectionArea->GetSequencer(), SectionObject, OutDrawElements, SectionGeometry);
+
+	FGeometry PaintSpaceParentGeometry = ParentGeometry;
+	PaintSpaceParentGeometry.AppendTransform(FSlateLayoutTransform(Inverse(Args.GetWindowToDesktopTransform())));
+
+	FSlateRect ParentClippingRect = PaintSpaceParentGeometry.GetClippingRect();
+
+	// Clip vertically
+	ParentClippingRect.Top = FMath::Max(ParentClippingRect.Top, MyClippingRect.Top);
+	ParentClippingRect.Bottom = FMath::Min(ParentClippingRect.Bottom, MyClippingRect.Bottom);
+
+	Painter.SectionClippingRect = Painter.SectionGeometry.GetClippingRect().IntersectionWith(ParentClippingRect);
+
+	Painter.LayerId = LayerId;
+	Painter.bParentEnabled = bEnabled;
 
 	// Ask the interface to draw the section
-	int32 PostSectionLayer = SectionInterface->OnPaintSection( SectionGeometry, SectionClipRect, OutDrawElements, LayerId, bEnabled );
+	LayerId = SectionInterface->OnPaintSection(Painter);
+	
+	LayerId = SCompoundWidget::OnPaint( Args, AllottedGeometry, MyClippingRect, OutDrawElements, LayerId, InWidgetStyle, bEnabled );
+
 	FLinearColor SelectionColor = FEditorStyle::GetSlateColor(SequencerSectionConstants::SelectionColorName).GetColor(FWidgetStyle());
+	DrawSectionHandles(AllottedGeometry, MyClippingRect, OutDrawElements, ++LayerId, DrawEffects, SelectionColor);
 
-	DrawSectionHandles(AllottedGeometry, MyClippingRect, OutDrawElements, PostSectionLayer, DrawEffects, SelectionColor);
-	DrawSelectionBorder(AllottedGeometry, MyClippingRect, OutDrawElements, PostSectionLayer, DrawEffects, SelectionColor);
-	PaintKeys( SectionGeometry, MyClippingRect, OutDrawElements, PostSectionLayer, InWidgetStyle, bEnabled );
+	FSlateRect KeyClippingRect = Painter.SectionGeometry.GetClippingRect().ExtendBy(FMargin(10.f, 0.f)).IntersectionWith(ParentClippingRect);
+	
+	Painter.LayerId = ++LayerId;
+	PaintKeys( Painter, InWidgetStyle, KeyClippingRect );
 
 	if (bLocked)
 	{
@@ -235,7 +407,7 @@ int32 SSequencerSection::OnPaint( const FPaintArgs& Args, const FGeometry& Allot
 
 		FSlateDrawElement::MakeBox(
 			OutDrawElements,
-			PostSectionLayer+1,
+			++LayerId,
 			AllottedGeometry.ToPaintGeometry(),
 			FEditorStyle::GetBrush(SelectionBorder),
 			MyClippingRect,
@@ -246,29 +418,30 @@ int32 SSequencerSection::OnPaint( const FPaintArgs& Args, const FGeometry& Allot
 
 	// Section name with drop shadow
 	FText SectionTitle = SectionInterface->GetSectionTitle();
+	FMargin ContentPadding = SectionInterface->GetContentPadding();
 
 	if (!SectionTitle.IsEmpty())
 	{
 		FSlateDrawElement::MakeText(
 			OutDrawElements,
-			PostSectionLayer+1,
-			SectionGeometry.ToOffsetPaintGeometry(FVector2D(6, 6)),
+			++LayerId,
+			Painter.SectionGeometry.ToOffsetPaintGeometry(FVector2D(ContentPadding.Left + 1, ContentPadding.Top + 1)),
 			SectionTitle,
 			FEditorStyle::GetFontStyle("NormalFont"),
 			MyClippingRect,
 			DrawEffects,
-			FLinearColor::Black
+			FLinearColor(0,0,0,.5f)
 		);
 
 		FSlateDrawElement::MakeText(
 			OutDrawElements,
-			PostSectionLayer+2,
-			SectionGeometry.ToOffsetPaintGeometry(FVector2D(5, 5)),
+			++LayerId,
+			Painter.SectionGeometry.ToOffsetPaintGeometry(FVector2D(ContentPadding.Left, ContentPadding.Top)),
 			SectionTitle,
 			FEditorStyle::GetFontStyle("NormalFont"),
 			MyClippingRect,
 			DrawEffects,
-			FLinearColor::White
+			FColor(200, 200, 200)
 		);
 	}
 
@@ -276,15 +449,16 @@ int32 SSequencerSection::OnPaint( const FPaintArgs& Args, const FGeometry& Allot
 }
 
 
-void SSequencerSection::PaintKeys( const FGeometry& AllottedGeometry, const FSlateRect& MyClippingRect, FSlateWindowElementList& OutDrawElements, int32 LayerId, const FWidgetStyle& InWidgetStyle, bool bParentEnabled ) const
+void SSequencerSection::PaintKeys( FSequencerSectionPainter& InPainter, const FWidgetStyle& InWidgetStyle, const FSlateRect& KeyClippingRect ) const
 {
 	static const FName HighlightBrushName("Sequencer.AnimationOutliner.DefaultBorder");
 
-	static const FName BackgroundBrushName("Sequencer.SectionArea.Background");
+	static const FName BackgroundBrushName("Sequencer.Section.BackgroundTint");
 	static const FName CircleKeyBrushName("Sequencer.KeyCircle");
 	static const FName DiamondKeyBrushName("Sequencer.KeyDiamond");
 	static const FName SquareKeyBrushName("Sequencer.KeySquare");
 	static const FName TriangleKeyBrushName("Sequencer.KeyTriangle");
+	static const FName StripeOverlayBrushName("Sequencer.Section.StripeOverlay");
 
 	static const FName SelectionColorName("SelectionColor");
 	static const FName SelectionInactiveColorName("SelectionColorInactive");
@@ -323,54 +497,84 @@ void SSequencerSection::PaintKeys( const FGeometry& AllottedGeometry, const FSla
 
 	const FSlateBrush* HighlightBrush = FEditorStyle::GetBrush(HighlightBrushName);
 	const FSlateBrush* BackgroundBrush = FEditorStyle::GetBrush(BackgroundBrushName);
+	const FSlateBrush* StripeOverlayBrush = FEditorStyle::GetBrush(StripeOverlayBrushName);
 	const FSlateBrush* CircleKeyBrush = FEditorStyle::GetBrush(CircleKeyBrushName);
 	const FSlateBrush* DiamondKeyBrush = FEditorStyle::GetBrush(DiamondKeyBrushName);
 	const FSlateBrush* SquareKeyBrush = FEditorStyle::GetBrush(SquareKeyBrushName);
 	const FSlateBrush* TriangleKeyBrush = FEditorStyle::GetBrush(TriangleKeyBrushName);
 
-	const ESlateDrawEffect::Type DrawEffects = bParentEnabled
+	const ESlateDrawEffect::Type DrawEffects = InPainter.bParentEnabled
 		? ESlateDrawEffect::None
 		: ESlateDrawEffect::DisabledEffect;
 
-	for (const FKeyAreaLayoutElement& LayoutElement : Layout->GetElements())
+	const FTimeToPixel& TimeToPixelConverter = InPainter.GetTimeConverter();
+
+	int32 LayerId = InPainter.LayerId;
+
+	for (const FSectionLayoutElement& LayoutElement : Layout->GetElements())
 	{
 		// get key handles
 		TSharedPtr<IKeyArea> KeyArea = LayoutElement.GetKeyArea();
 
+		FGeometry KeyAreaGeometry = GetKeyAreaGeometry(LayoutElement, InPainter.SectionGeometry);
 
-		FGeometry KeyAreaGeometry = GetKeyAreaGeometry(LayoutElement, AllottedGeometry);
-
-		FTimeToPixel TimeToPixelConverter = SectionObject.IsInfinite()
-			? FTimeToPixel(ParentGeometry, Sequencer.GetViewRange())
-			: FTimeToPixel(KeyAreaGeometry, TRange<float>(SectionObject.GetStartTime(), SectionObject.GetEndTime()));
+		TOptional<FLinearColor> KeyAreaColor = KeyArea.IsValid() ? KeyArea->GetColor() : TOptional<FLinearColor>();
 
 		// draw a box for the key area 
-		if (SectionInterface->ShouldDrawKeyAreaBackground())
+		if (KeyAreaColor.IsSet())
 		{
+			static float BoxThickness = 5.f;
+			FVector2D KeyAreaSize = KeyAreaGeometry.GetLocalSize();
 			FSlateDrawElement::MakeBox( 
-				OutDrawElements,
+				InPainter.DrawElements,
 				LayerId + 1,
-				KeyAreaGeometry.ToPaintGeometry(),
-				BackgroundBrush,
-				MyClippingRect,
+				KeyAreaGeometry.ToPaintGeometry(FVector2D(KeyAreaSize.X, BoxThickness), FSlateLayoutTransform(FVector2D(0.f, KeyAreaSize.Y*.5f - BoxThickness*.5f))),
+				StripeOverlayBrush,
+				InPainter.SectionClippingRect,
 				DrawEffects,
-				KeyArea->GetColor()
+				KeyAreaColor.GetValue()
 			); 
 		}
 
 		if (LayoutElement.GetDisplayNode().IsValid() && LayoutElement.GetDisplayNode()->IsHovered())
 		{
-			FSlateDrawElement::MakeBox( 
-				OutDrawElements,
+			FSlateDrawElement::MakeBox(
+				InPainter.DrawElements,
 				LayerId + 1,
 				KeyAreaGeometry.ToPaintGeometry(),
 				HighlightBrush,
-				MyClippingRect,
+				InPainter.SectionClippingRect,
 				DrawEffects,
-				FLinearColor(1.f, 1.f, 1.f, 0.2f)
+				FLinearColor(1.f, 1.f, 1.f, 0.05f)
 			); 
 		}
 
+		if (Selection.IsSelected(LayoutElement.GetDisplayNode().ToSharedRef()))
+		{
+			static const FName SelectedTrackTint("Sequencer.Section.SelectedTrackTint");
+			static const FName SelectionColorInactiveName("SelectionColor_Inactive");
+
+			FLinearColor KeyAreaOutlineColor = (Selection.GetActiveSelection() == FSequencerSelection::EActiveSelection::OutlinerNode)
+					? SelectionColor
+					: FEditorStyle::GetSlateColor(SelectionColorInactiveName).GetColor(InWidgetStyle);
+
+			FSlateDrawElement::MakeBox(
+				InPainter.DrawElements,
+				++LayerId + 2,
+				KeyAreaGeometry.ToPaintGeometry(),
+				FEditorStyle::GetBrush(SelectedTrackTint),
+				InPainter.SectionClippingRect,
+				DrawEffects,
+				KeyAreaOutlineColor
+			);
+		}
+
+		// Can't do any of the rest if there are no keys
+		if (!KeyArea.IsValid())
+		{
+			continue;
+		}
+
 		TArray<FKeyHandle> KeyHandles = KeyArea->GetUnsortedKeyHandles();
 
 		if (!KeyHandles.Num())
@@ -437,7 +641,7 @@ void SSequencerSection::PaintKeys( const FGeometry& AllottedGeometry, const FSla
 			// allow group & section overrides
 			const FSlateBrush* OverrideBrush = nullptr;
 
-			if (LayoutElement.GetType() == FKeyAreaLayoutElement::Group)
+			if (LayoutElement.GetType() == FSectionLayoutElement::Group)
 			{
 				auto Group = StaticCastSharedPtr<FGroupedKeyArea>(KeyArea);
 				OverrideBrush = Group->GetBrush(KeyHandle);
@@ -490,7 +694,7 @@ void SSequencerSection::PaintKeys( const FGeometry& AllottedGeometry, const FSla
 			}
 
 			// allow group to tint the color
-			if (LayoutElement.GetType() == FKeyAreaLayoutElement::Group)
+			if (LayoutElement.GetType() == FSectionLayoutElement::Group)
 			{
 				auto Group = StaticCastSharedPtr<FGroupedKeyArea>(KeyArea);
 				KeyColor *= Group->GetKeyTint(KeyHandle);
@@ -503,7 +707,7 @@ void SSequencerSection::PaintKeys( const FGeometry& AllottedGeometry, const FSla
 			const FVector2D KeySize = bSelected ? SequencerSectionConstants::KeySize + ThrobAmount * ThrobScaleValue : SequencerSectionConstants::KeySize;
 
 			FSlateDrawElement::MakeBox(
-				OutDrawElements,
+				InPainter.DrawElements,
 				// always draw selected keys on top of other keys
 				bSelected ? KeyLayer + 1 : KeyLayer,
 				// Center the key along Y.  Ensure the middle of the key is at the actual key time
@@ -515,14 +719,14 @@ void SSequencerSection::PaintKeys( const FGeometry& AllottedGeometry, const FSla
 					KeySize
 				),
 				KeyBrush,
-				MyClippingRect,
+				KeyClippingRect,
 				DrawEffects,
 				BorderColor
 			);
 
 			// draw fill
 			FSlateDrawElement::MakeBox(
-				OutDrawElements,
+				InPainter.DrawElements,
 				// always draw selected keys on top of other keys
 				bSelected ? KeyLayer + 2 : KeyLayer + 1,
 				// Center the key along Y.  Ensure the middle of the key is at the actual key time
@@ -534,7 +738,7 @@ void SSequencerSection::PaintKeys( const FGeometry& AllottedGeometry, const FSla
 					KeySize - 2.0f * BrushBorderWidth
 				),
 				KeyBrush,
-				MyClippingRect,
+				KeyClippingRect,
 				DrawEffects,
 				FillColor
 			);
@@ -545,8 +749,14 @@ void SSequencerSection::PaintKeys( const FGeometry& AllottedGeometry, const FSla
 
 void SSequencerSection::DrawSectionHandles( const FGeometry& AllottedGeometry, const FSlateRect& MyClippingRect, FSlateWindowElementList& OutDrawElements, int32 LayerId, ESlateDrawEffect::Type DrawEffects, FLinearColor SelectionColor ) const
 {
-	const FSlateBrush* LeftGripBrush = FEditorStyle::GetBrush(SectionInterface->GetSectionGripLeftBrushName());
-	const FSlateBrush* RightGripBrush = FEditorStyle::GetBrush(SectionInterface->GetSectionGripRightBrushName());
+	UMovieSceneSection* Section = SectionInterface->GetSectionObject();
+	if (!Section || Section->IsInfinite())
+	{
+		return;
+	}
+
+	const FSlateBrush* LeftGripBrush = FEditorStyle::GetBrush("Sequencer.Section.GripLeft");
+	const FSlateBrush* RightGripBrush = FEditorStyle::GetBrush("Sequencer.Section.GripRight");
 
 	bool bLeftHandleActive = false;
 	bool bRightHandleActive = false;
@@ -566,17 +776,21 @@ void SSequencerSection::DrawSectionHandles( const FGeometry& AllottedGeometry, c
 		}
 	}
 
+	float Opacity = FMath::Clamp(.5f + HandleOffsetPx / SectionInterface->GetSectionGripSize() * .5f, .5f, 1.f);
+
 	// Left Grip
 	FSlateDrawElement::MakeBox
 	(
 		OutDrawElements,
 		LayerId,
-		// Center the key along Y.  Ensure the middle of the key is at the actual key time
-		AllottedGeometry.ToPaintGeometry(FVector2D(0.0f, 0.0f), FVector2D(SectionInterface->GetSectionGripSize(), AllottedGeometry.GetDrawSize().Y)),
+		AllottedGeometry.ToPaintGeometry(
+			FVector2D(0.0f, 0.0f),
+			FVector2D(SectionInterface->GetSectionGripSize(), AllottedGeometry.GetDrawSize().Y)
+		),
 		LeftGripBrush,
-		MyClippingRect,
+		MyClippingRect.InsetBy(FMargin(1.f)),
 		DrawEffects,
-		bLeftHandleActive ? SelectionColor : LeftGripBrush->GetTint(FWidgetStyle())
+		(bLeftHandleActive ? SelectionColor : LeftGripBrush->GetTint(FWidgetStyle())).CopyWithNewOpacity(Opacity)
 	);
 	
 	// Right Grip
@@ -584,60 +798,14 @@ void SSequencerSection::DrawSectionHandles( const FGeometry& AllottedGeometry, c
 	(
 		OutDrawElements,
 		LayerId,
-		// Center the key along Y.  Ensure the middle of the key is at the actual key time
-		AllottedGeometry.ToPaintGeometry(FVector2D(AllottedGeometry.Size.X-SectionInterface->GetSectionGripSize(), 0.0f), FVector2D(SectionInterface->GetSectionGripSize(), AllottedGeometry.GetDrawSize().Y)),
+		AllottedGeometry.ToPaintGeometry(
+			FVector2D(AllottedGeometry.Size.X-SectionInterface->GetSectionGripSize(), 0.0f),
+			FVector2D(SectionInterface->GetSectionGripSize(), AllottedGeometry.GetDrawSize().Y)
+		),
 		RightGripBrush,
-		MyClippingRect,
+		MyClippingRect.InsetBy(FMargin(1.f)),
 		DrawEffects,
-		bRightHandleActive ? SelectionColor : RightGripBrush->GetTint(FWidgetStyle())
-	);
-}
-
-
-void SSequencerSection::DrawSelectionBorder( const FGeometry& AllottedGeometry, const FSlateRect& MyClippingRect, FSlateWindowElementList& OutDrawElements, int32 LayerId, ESlateDrawEffect::Type DrawEffects, FLinearColor SelectionColor ) const
-{
-	FSequencerSelection& Selection = ParentSectionArea->GetSequencer().GetSelection();
-	FSequencerSelectionPreview& SelectionPreview = ParentSectionArea->GetSequencer().GetSelectionPreview();
-	UMovieSceneSection* SectionObject = SectionInterface->GetSectionObject();
-	ESelectionPreviewState SelectionPreviewState = SelectionPreview.GetSelectionState( SectionObject );
-
-	if (SelectionPreviewState == ESelectionPreviewState::NotSelected)
-	{
-		// Explicitly not selected in the preview selection
-		return;
-	}
-	
-	if (SelectionPreviewState == ESelectionPreviewState::Undefined && !Selection.IsSelected(SectionObject))
-	{
-		// No preview selection for this section, and it's not selected
-		return;
-	}
-	
-	// Use a muted selection color for selection previews
-	if( SelectionPreviewState == ESelectionPreviewState::Selected )
-	{
-		SelectionColor = SelectionColor.LinearRGBToHSV();
-		SelectionColor.R += 0.1f; // +10% hue
-		SelectionColor.G = 0.6f; // 60% saturation
-		SelectionColor = SelectionColor.HSVToLinearRGB();
-	}
-	else if (Selection.GetActiveSelection() != FSequencerSelection::EActiveSelection::KeyAndSection)
-	{
-		// Use an inactive selection color for existing selections that are not active
-		SelectionColor = FEditorStyle::GetSlateColor(SequencerSectionConstants::SelectionInactiveColorName).GetColor(FWidgetStyle());
-	}
-
-	// draw selection box
-	static const FName SelectionBorder("Sequencer.Section.SelectionBorder");
-
-	FSlateDrawElement::MakeBox(
-		OutDrawElements,
-		LayerId+1,
-		AllottedGeometry.ToPaintGeometry(),
-		FEditorStyle::GetBrush(SelectionBorder),
-		MyClippingRect,
-		DrawEffects,
-		SelectionColor
+		(bRightHandleActive ? SelectionColor : RightGripBrush->GetTint(FWidgetStyle())).CopyWithNewOpacity(Opacity)
 	);
 }
 
@@ -645,9 +813,24 @@ void SSequencerSection::Tick( const FGeometry& AllottedGeometry, const double In
 {	
 	if( GetVisibility() == EVisibility::Visible )
 	{
-		Layout = FKeyAreaLayout(*ParentSectionArea, SectionIndex);
-		FGeometry SectionGeometry = MakeSectionGeometryWithoutHandles( AllottedGeometry, SectionInterface );
+		Layout = FSectionLayout(*ParentSectionArea, SectionIndex);
+
+		UMovieSceneSection* Section = SectionInterface->GetSectionObject();
+		if (Section && !Section->IsInfinite())
+		{
+			FTimeToPixel TimeToPixelConverter(ParentGeometry, GetSequencer().GetViewRange());
+
+			const int32 SectionLengthPx = FMath::Max(0,
+				FMath::RoundToInt(
+					TimeToPixelConverter.TimeToPixel(Section->GetEndTime())) - FMath::RoundToInt(TimeToPixelConverter.TimeToPixel(Section->GetStartTime())
+				)
+			);
+
+			const float SectionGripSize = SectionInterface->GetSectionGripSize();
+			HandleOffsetPx = FMath::Max(FMath::RoundToFloat((2*SectionGripSize - SectionLengthPx) * .5f), 0.f);
+		}
 
+		FGeometry SectionGeometry = MakeSectionGeometryWithoutHandles( AllottedGeometry, SectionInterface );
 		SectionInterface->Tick(SectionGeometry, ParentGeometry, InCurrentTime, InDeltaTime);
 	}
 }
@@ -691,10 +874,10 @@ FReply SSequencerSection::OnMouseButtonDown( const FGeometry& MyGeometry, const
 
 FGeometry SSequencerSection::MakeSectionGeometryWithoutHandles( const FGeometry& AllottedGeometry, const TSharedPtr<ISequencerSection>& InSectionInterface ) const
 {
-	const bool bSectionsAreConnected = InSectionInterface->AreSectionsConnected();
-	const float SectionGripSize = !bSectionsAreConnected ? InSectionInterface->GetSectionGripSize() : 0.0f;
-
-	return AllottedGeometry.MakeChild( FVector2D( SectionGripSize, 0 ), AllottedGeometry.GetDrawSize() - FVector2D( SectionGripSize*2, 0.0f ) );
+	return AllottedGeometry.MakeChild(
+		AllottedGeometry.GetDrawSize() - FVector2D( HandleOffsetPx*2, 0.0f ),
+		FSlateLayoutTransform(FVector2D(HandleOffsetPx, 0 ))
+	);
 }
 
 
diff --git a/Engine/Source/Editor/Sequencer/Private/SSequencerSection.h b/Engine/Source/Editor/Sequencer/Private/SSequencerSection.h
--- a/Engine/Source/Editor/Sequencer/Private/SSequencerSection.h
+++ b/Engine/Source/Editor/Sequencer/Private/SSequencerSection.h
@@ -2,8 +2,9 @@
 
 #pragma once
 
-#include "KeyAreaLayout.h"
+#include "SectionLayout.h"
 
+class FSequencerSectionPainter;
 
 class SSequencerSection : public SCompoundWidget
 {
@@ -17,7 +18,7 @@ public:
 	TSharedPtr<ISequencerSection> GetSectionInterface() const { return SectionInterface; }
 
 	/** Caches the parent geometry to be given to section interfaces that need it on tick */
-	void CacheParentGeometry(const FGeometry& InParentGeometry) {ParentGeometry = InParentGeometry;}
+	void CacheParentGeometry(const FGeometry& InParentGeometry) { ParentGeometry = InParentGeometry; }
 
 	virtual FVector2D ComputeDesiredSize(float) const override;
 
@@ -30,7 +31,7 @@ private:
 	 * @param SectionGeometry	The geometry of the section
 	 * @return The geometry of the key area
 	 */
-	FGeometry GetKeyAreaGeometry( const FKeyAreaLayoutElement& KeyArea, const FGeometry& SectionGeometry ) const;
+	FGeometry GetKeyAreaGeometry( const FSectionLayoutElement& KeyArea, const FGeometry& SectionGeometry ) const;
 
 	/**
 	 * Determines the key that is under the mouse
@@ -70,23 +71,15 @@ private:
 	/**
 	 * Paints keys visible inside the section
 	 *
-	 * @param AllottedGeometry	The geometry of the section where keys are painted
-	 * @param MyClippingRect	ClippingRect of the section
-	 * @param OutDrawElements	List of draw elements to add to
-	 * @param LayerId			The starting draw area
+	 * @param InPainter			Section painter
 	 */
-	void PaintKeys( const FGeometry& AllottedGeometry, const FSlateRect& MyClippingRect, FSlateWindowElementList& OutDrawElements, int32 LayerId, const FWidgetStyle& InWidgetStyle, bool bParentEnabled ) const;
+	void PaintKeys( FSequencerSectionPainter& InPainter, const FWidgetStyle& InWidgetStyle, const FSlateRect& KeyClippingRect ) const;
 
 	/**
 	 * Draw the section resize handles.
 	 */
 	void DrawSectionHandles( const FGeometry& AllottedGeometry, const FSlateRect& MyClippingRect, FSlateWindowElementList& OutDrawElements, int32 LayerId, ESlateDrawEffect::Type DrawEffects, FLinearColor SelectionColor ) const;
 
-	/**
-	 * Draw a box representing this section's selection
-	 */
-	void DrawSelectionBorder( const FGeometry& AllottedGeometry, const FSlateRect& MyClippingRect, FSlateWindowElementList& OutDrawElements, int32 LayerId, ESlateDrawEffect::Type DrawEffects, FLinearColor SelectionColor ) const;
-
 	/** @return the sequencer interface */
 	FSequencer& GetSequencer() const;
 
@@ -110,11 +103,13 @@ private:
 	/** Section area where this section resides */
 	TSharedPtr<FSequencerTrackNode> ParentSectionArea;
 	/** Cached layout generated each tick */
-	TOptional<FKeyAreaLayout> Layout;
+	TOptional<FSectionLayout> Layout;
 	/** The index of this section in the parent section area */
 	int32 SectionIndex;
 	/** Cached parent geometry to pass down to any section interfaces that need it during tick */
 	FGeometry ParentGeometry;
 	/** The end time for a throbbing animation for selected keys */
 	static double SelectionThrobEndTime;
+	/** Handle offset amount in pixels */
+	float HandleOffsetPx;
 };
\ No newline at end of file
diff --git a/Engine/Source/Editor/Sequencer/Private/SSequencerSectionAreaView.cpp b/Engine/Source/Editor/Sequencer/Private/SSequencerSectionAreaView.cpp
--- a/Engine/Source/Editor/Sequencer/Private/SSequencerSectionAreaView.cpp
+++ b/Engine/Source/Editor/Sequencer/Private/SSequencerSectionAreaView.cpp
@@ -35,29 +35,31 @@ namespace SequencerSectionUtils
 		float PixelEndX = TimeToPixelConverter.TimeToPixel( Section->GetEndTime() );
 
 		// If the section is infinite, occupy the entire width of the geometry where the section is located.
-		if (Section->IsInfinite())
+		const bool bIsInfinite = Section->IsInfinite();
+		if (bIsInfinite)
 		{
 			PixelStartX = AllottedGeometry.Position.X;
 			PixelEndX = AllottedGeometry.Position.X + AllottedGeometry.Size.X;
 		}
 
-		// Actual section length without grips.
-		float SectionLengthActual = FMath::Max(1.0f,PixelEndX-PixelStartX);
+		const float MinSectionWidth = 1.f;
+		float SectionLength = FMath::Max(MinSectionWidth, PixelEndX-PixelStartX);
 
-		float SectionLengthWithGrips = SectionLengthActual;
-		float ExtentSize = 0;
-		if( !SectionInterface->AreSectionsConnected() )
+		float GripOffset = 0;
+		if (!bIsInfinite)
 		{
-			// Extend the section to include areas for the grips
-			ExtentSize = SectionInterface->GetSectionGripSize();
-			// Connected sections do not display grips outside of their section area
-			SectionLengthWithGrips += (ExtentSize*2);
+			float NewSectionLength = FMath::Max(SectionLength, MinSectionWidth + SectionInterface->GetSectionGripSize() * 2.f);
+
+			GripOffset = (NewSectionLength - SectionLength) / 2.f;
+			SectionLength = NewSectionLength;
 		}
 
 		float ActualHeight = NodeHeight / MaxTracks;
 
 		// Compute allotted geometry area that can be used to draw the section
-		return AllottedGeometry.MakeChild( FVector2D( PixelStartX-ExtentSize, ActualHeight * RowIndex ), FVector2D( SectionLengthWithGrips, ActualHeight ) );
+		return AllottedGeometry.MakeChild(
+			FVector2D( PixelStartX-GripOffset, ActualHeight * RowIndex ),
+			FVector2D( SectionLength, ActualHeight ) );
 	}
 
 }
@@ -72,8 +74,6 @@ void SSequencerSectionAreaView::Construct( const FArguments& InArgs, TSharedRef<
 	check( Node->GetType() == ESequencerNode::Track );
 	SectionAreaNode = StaticCastSharedRef<FSequencerTrackNode>( Node );
 
-	BackgroundBrush = FEditorStyle::GetBrush("Sequencer.SectionArea.Background");
-
 	// Generate widgets for sections in this view
 	GenerateSectionWidgets();
 };
@@ -148,6 +148,11 @@ void SSequencerSectionAreaView::Tick(const FGeometry& AllottedGeometry, const do
 
 			return SectionA && SectionB && SectionA->GetOverlapPriority() < SectionB->GetOverlapPriority();
 		});
+
+		for( int32 WidgetIndex = 0; WidgetIndex < Children.Num(); ++WidgetIndex )
+		{
+			Children[WidgetIndex]->CacheParentGeometry(AllottedGeometry);
+		}
 	}
 }
 
@@ -176,13 +181,11 @@ void SSequencerSectionAreaView::OnArrangeChildren( const FGeometry& AllottedGeom
 
 		int32 RowIndex = SectionInterface->GetSectionObject()->GetRowIndex();
 
-		FGeometry SectionGeometry = SequencerSectionUtils::GetSectionGeometry( AllottedGeometry, RowIndex, MaxTracks, Widget->GetDesiredSize().Y, SectionInterface, TimeToPixelConverter );
-
 		EVisibility Visibility = Widget->GetVisibility();
 		if( ArrangedChildren.Accepts( Visibility ) )
 		{
-			Widget->CacheParentGeometry(AllottedGeometry);
-
+			FGeometry SectionGeometry = SequencerSectionUtils::GetSectionGeometry( AllottedGeometry, RowIndex, MaxTracks, Widget->GetDesiredSize().Y, SectionInterface, TimeToPixelConverter );
+			
 			ArrangedChildren.AddWidget( 
 				Visibility, 
 				AllottedGeometry.MakeChild( Widget, SectionGeometry.Position, SectionGeometry.GetDrawSize() )
diff --git a/Engine/Source/Editor/Sequencer/Private/SSequencerSectionAreaView.h b/Engine/Source/Editor/Sequencer/Private/SSequencerSectionAreaView.h
--- a/Engine/Source/Editor/Sequencer/Private/SSequencerSectionAreaView.h
+++ b/Engine/Source/Editor/Sequencer/Private/SSequencerSectionAreaView.h
@@ -57,8 +57,6 @@ private:
 	TSharedPtr<FSequencerTrackNode> SectionAreaNode;
 	/** The current view range */
 	TAttribute< TRange<float> > ViewRange;
-	/** Background brush of the section area */
-	const FSlateBrush* BackgroundBrush;
 	/** All the widgets in the panel */
 	TSlotlessChildren<SSequencerSection> Children;
 };
diff --git a/Engine/Source/Editor/Sequencer/Private/SSequencerSectionOverlay.cpp b/Engine/Source/Editor/Sequencer/Private/SSequencerSectionOverlay.cpp
--- a/Engine/Source/Editor/Sequencer/Private/SSequencerSectionOverlay.cpp
+++ b/Engine/Source/Editor/Sequencer/Private/SSequencerSectionOverlay.cpp
@@ -2,7 +2,7 @@
 
 #include "SequencerPrivatePCH.h"
 #include "SSequencerSectionOverlay.h"
-#include "TimeSliderController.h"
+#include "SequencerTimeSliderController.h"
 
 int32 SSequencerSectionOverlay::OnPaint( const FPaintArgs& Args, const FGeometry& AllottedGeometry, const FSlateRect& MyClippingRect, FSlateWindowElementList& OutDrawElements, int32 LayerId, const FWidgetStyle& InWidgetStyle, bool bParentEnabled ) const
 {
diff --git a/Engine/Source/Editor/Sequencer/Private/SSequencerTrackArea.cpp b/Engine/Source/Editor/Sequencer/Private/SSequencerTrackArea.cpp
--- a/Engine/Source/Editor/Sequencer/Private/SSequencerTrackArea.cpp
+++ b/Engine/Source/Editor/Sequencer/Private/SSequencerTrackArea.cpp
@@ -3,7 +3,7 @@
 #include "SequencerPrivatePCH.h"
 #include "SSequencerTrackArea.h"
 #include "SSequencerTrackLane.h"
-#include "TimeSliderController.h"
+#include "SequencerTimeSliderController.h"
 #include "CommonMovieSceneTools.h"
 #include "Sequencer.h"
 #include "SSequencerTreeView.h"
@@ -157,9 +157,6 @@ int32 SSequencerTrackArea::OnPaint(const FPaintArgs& Args, const FGeometry& Allo
 			LayerId = EditTool->OnPaint(AllottedGeometry, MyClippingRect, OutDrawElements, LayerId + 1);
 		}
 
-		static const FName SelectionColorPressedName("SelectionColor");
-		FLinearColor SelectionColor = FEditorStyle::GetSlateColor(SelectionColorPressedName).GetColor(InWidgetStyle).CopyWithNewOpacity(0.5f);
-
 		TOptional<FHighlightRegion> HighlightRegion = TreeView.Pin()->GetHighlightRegion();
 		if (HighlightRegion.IsSet())
 		{
@@ -170,7 +167,7 @@ int32 SSequencerTrackArea::OnPaint(const FPaintArgs& Args, const FGeometry& Allo
 				FEditorStyle::GetBrush("Sequencer.TrackHoverHighlight_Top"),
 				MyClippingRect,
 				ESlateDrawEffect::None,
-				SelectionColor
+				FLinearColor::Black
 			);
 		
 			FSlateDrawElement::MakeBox(
@@ -180,7 +177,7 @@ int32 SSequencerTrackArea::OnPaint(const FPaintArgs& Args, const FGeometry& Allo
 				FEditorStyle::GetBrush("Sequencer.TrackHoverHighlight_Bottom"),
 				MyClippingRect,
 				ESlateDrawEffect::None,
-				SelectionColor
+				FLinearColor::Black
 			);
 		}
 	}
@@ -216,7 +213,11 @@ FReply SSequencerTrackArea::OnMouseButtonUp( const FGeometry& MyGeometry, const
 
 bool SSequencerTrackArea::CanActivateEditTool(FName Identifier) const
 {
-	if (!EditTool.IsValid())
+	if (InputStack.GetCapturedIndex() != INDEX_NONE)
+	{
+		return false;
+	}
+	else if (!EditTool.IsValid())
 	{
 		return true;
 	}
diff --git a/Engine/Source/Editor/Sequencer/Private/SSequencerTrackLane.cpp b/Engine/Source/Editor/Sequencer/Private/SSequencerTrackLane.cpp
--- a/Engine/Source/Editor/Sequencer/Private/SSequencerTrackLane.cpp
+++ b/Engine/Source/Editor/Sequencer/Private/SSequencerTrackLane.cpp
@@ -12,7 +12,7 @@ void SSequencerTrackLane::Construct(const FArguments& InArgs, const TSharedRef<F
 
 	ChildSlot
 	.HAlign(HAlign_Fill)
-	.Padding(0, DisplayNode->GetNodePadding().Top, 0, 0)
+	.Padding(0)
 	[
 		InArgs._Content.Widget
 	];
@@ -70,7 +70,7 @@ void DrawLaneRecursive(const TSharedRef<FSequencerDisplayNode>& DisplayNode, con
 			FEditorStyle::GetBrush(BorderName),
 			MyClippingRect,
 			ESlateDrawEffect::None,
-			FLinearColor(1.0f, 1.0f, 1.0f, 0.1f)
+			FLinearColor(1.0f, 1.0f, 1.0f, 0.05f)
 		);
 	}
 
diff --git a/Engine/Source/Editor/Sequencer/Private/SSequencerTreeView.cpp b/Engine/Source/Editor/Sequencer/Private/SSequencerTreeView.cpp
--- a/Engine/Source/Editor/Sequencer/Private/SSequencerTreeView.cpp
+++ b/Engine/Source/Editor/Sequencer/Private/SSequencerTreeView.cpp
@@ -155,10 +155,7 @@ int32 SSequencerTreeView::OnPaint(const FPaintArgs& Args, const FGeometry& Allot
 
 	if (HighlightRegion.IsSet())
 	{
-		static const FName SelectionColorPressedName("SelectionColor");
-		FLinearColor SelectionColor = FEditorStyle::GetSlateColor(SelectionColorPressedName).GetColor(InWidgetStyle).CopyWithNewOpacity(0.3f);
-
-		// Black tint for unhighlighted regions
+		// Black tint for highlighted regions
 		FSlateDrawElement::MakeBox(
 			OutDrawElements,
 			LayerId+1,
@@ -166,7 +163,7 @@ int32 SSequencerTreeView::OnPaint(const FPaintArgs& Args, const FGeometry& Allot
 			FEditorStyle::GetBrush("Sequencer.TrackHoverHighlight_Top"),
 			MyClippingRect,
 			ESlateDrawEffect::None,
-			SelectionColor
+			FLinearColor::Black
 		);
 		
 		FSlateDrawElement::MakeBox(
@@ -176,7 +173,7 @@ int32 SSequencerTreeView::OnPaint(const FPaintArgs& Args, const FGeometry& Allot
 			FEditorStyle::GetBrush("Sequencer.TrackHoverHighlight_Bottom"),
 			MyClippingRect,
 			ESlateDrawEffect::None,
-			SelectionColor
+			FLinearColor::Black
 		);
 	}
 
diff --git a/Engine/Source/Editor/Sequencer/Private/SectionLayout.cpp b/Engine/Source/Editor/Sequencer/Private/SectionLayout.cpp
new file mode 100644
--- /dev/null
+++ b/Engine/Source/Editor/Sequencer/Private/SectionLayout.cpp
@@ -0,0 +1,148 @@
+
+#include "SequencerPrivatePCH.h"
+#include "SectionLayout.h"
+
+FSectionLayoutElement FSectionLayoutElement::FromGroup(const TSharedRef<FSequencerDisplayNode>& InNode, const TSharedRef<FGroupedKeyArea>& InKeyAreaGroup, float InOffset)
+{
+	FSectionLayoutElement Tmp;
+	Tmp.Type = Group;
+	Tmp.KeyArea = InKeyAreaGroup;
+	Tmp.LocalOffset = InOffset;
+	Tmp.Height = InNode->GetNodeHeight();
+	Tmp.DisplayNode = InNode;
+	return Tmp;
+}
+
+FSectionLayoutElement FSectionLayoutElement::FromKeyAreaNode(const TSharedRef<FSequencerSectionKeyAreaNode>& InKeyAreaNode, int32 SectionIndex, float InOffset)
+{
+	FSectionLayoutElement Tmp;
+	Tmp.Type = Single;
+	Tmp.KeyArea = InKeyAreaNode->GetKeyArea(SectionIndex);
+	Tmp.LocalOffset = InOffset;
+	Tmp.DisplayNode = InKeyAreaNode;
+	Tmp.Height = InKeyAreaNode->GetNodeHeight();
+	return Tmp;
+}
+
+FSectionLayoutElement FSectionLayoutElement::FromTrack(const TSharedRef<FSequencerTrackNode>& InTrackNode, int32 SectionIndex, float InOffset)
+{
+	FSectionLayoutElement Tmp;
+	Tmp.Type = Single;
+	Tmp.KeyArea = InTrackNode->GetTopLevelKeyNode()->GetKeyArea(SectionIndex);
+	Tmp.LocalOffset = InOffset;
+	Tmp.DisplayNode = InTrackNode;
+	Tmp.Height = InTrackNode->GetNodeHeight();
+	return Tmp;
+}
+
+FSectionLayoutElement FSectionLayoutElement::EmptySpace(const TSharedRef<FSequencerDisplayNode>& InNode, float InOffset)
+{
+	FSectionLayoutElement Tmp;
+	Tmp.Type = Single;
+	Tmp.LocalOffset = InOffset;
+	Tmp.DisplayNode = InNode;
+	Tmp.Height = InNode->GetNodeHeight();
+	return Tmp;
+}
+
+FSectionLayoutElement::EType FSectionLayoutElement::GetType() const
+{
+	return Type;
+}
+
+float FSectionLayoutElement::GetOffset() const
+{
+	return LocalOffset;
+}
+
+float FSectionLayoutElement::GetHeight() const
+{
+	return Height;
+}
+
+TSharedPtr<IKeyArea> FSectionLayoutElement::GetKeyArea() const
+{
+	return KeyArea;
+}
+
+TSharedPtr<FSequencerDisplayNode> FSectionLayoutElement::GetDisplayNode() const
+{
+	return DisplayNode;
+}
+
+FSectionLayout::FSectionLayout(FSequencerDisplayNode& InNode, int32 InSectionIndex)
+{
+	float VerticalOffset = 0.f;
+
+	auto SetupKeyArea = [&](FSequencerDisplayNode& Node){
+
+		if (Node.GetType() == ESequencerNode::KeyArea)
+		{
+			Elements.Add(FSectionLayoutElement::FromKeyAreaNode(
+				StaticCastSharedRef<FSequencerSectionKeyAreaNode>(Node.AsShared()),
+				InSectionIndex,
+				VerticalOffset
+			));
+		}
+		else if (Node.GetType() == ESequencerNode::Track && static_cast<FSequencerTrackNode&>(Node).GetTopLevelKeyNode().IsValid())
+		{
+			Elements.Add(FSectionLayoutElement::FromTrack(
+				StaticCastSharedRef<FSequencerTrackNode>(Node.AsShared()),
+				InSectionIndex,
+				VerticalOffset
+			));
+		}
+		else if (!Node.IsExpanded())
+		{
+			Elements.Add(FSectionLayoutElement::FromGroup(
+				Node.AsShared(),
+				Node.UpdateKeyGrouping(InSectionIndex),
+				VerticalOffset
+			));
+		}
+		else
+		{
+			// It's benign space
+			Elements.Add(FSectionLayoutElement::EmptySpace(
+				Node.AsShared(),
+				VerticalOffset
+			));
+		}
+
+	};
+
+	// First, layout the parent
+	{
+		VerticalOffset += InNode.GetNodePadding().Top;
+
+		SetupKeyArea(InNode);
+
+		VerticalOffset += InNode.GetNodeHeight() + InNode.GetNodePadding().Bottom;
+	}
+
+	// Then any children
+	InNode.TraverseVisible_ParentFirst([&](FSequencerDisplayNode& Node){
+		
+		VerticalOffset += Node.GetNodePadding().Top;
+
+		SetupKeyArea(Node);
+
+		VerticalOffset += Node.GetNodeHeight() + Node.GetNodePadding().Bottom;
+		return true;
+
+	}, false);
+}
+
+const TArray<FSectionLayoutElement>& FSectionLayout::GetElements() const
+{
+	return Elements;
+}
+
+float FSectionLayout::GetTotalHeight() const
+{
+	if (Elements.Num())
+	{
+		return Elements.Last().GetOffset() + Elements.Last().GetDisplayNode()->GetNodePadding().Combined() + Elements.Last().GetHeight();
+	}
+	return 0.f;
+}
\ No newline at end of file
diff --git a/Engine/Source/Editor/Sequencer/Private/SectionLayout.h b/Engine/Source/Editor/Sequencer/Private/SectionLayout.h
new file mode 100644
--- /dev/null
+++ b/Engine/Source/Editor/Sequencer/Private/SectionLayout.h
@@ -0,0 +1,74 @@
+
+#pragma once
+
+#include "GroupedKeyArea.h"
+
+/** A layout element specifying the geometry required to render a key area */
+struct FSectionLayoutElement
+{
+	/** Whether this layout element pertains to one or multiple key areas */
+	enum EType { Single, Group };
+
+	/** Construct this element from a grouped key area */
+	static FSectionLayoutElement FromGroup(const TSharedRef<FSequencerDisplayNode>& InNode, const TSharedRef<FGroupedKeyArea>& InKeyAreaGroup, float InOffset);
+
+	/** Construct this element from a single Key area node */
+	static FSectionLayoutElement FromKeyAreaNode(const TSharedRef<FSequencerSectionKeyAreaNode>& InKeyAreaNode, int32 SectionIndex, float InOffset);
+
+	/** Construct this element from a single Key area node */
+	static FSectionLayoutElement FromTrack(const TSharedRef<FSequencerTrackNode>& InTrackNode, int32 SectionIndex, float InOffset);
+
+	/** Construct this element from a single Key area node */
+	static FSectionLayoutElement EmptySpace(const TSharedRef<FSequencerDisplayNode>& InNode, float InOffset);
+
+	/** Retrieve the type of this layout element */
+	EType GetType() const;
+
+	/** Retrieve vertical offset from the top of this element's parent */
+	float GetOffset() const;
+
+	/** Retrieve the desired height of this element based on the specified parent geometry */
+	float GetHeight() const;
+
+	/** Access the key area that this layout element was generated for */
+	TSharedPtr<IKeyArea> GetKeyArea() const;
+
+	/** Access the display node that this layout element was generated for */
+	TSharedPtr<FSequencerDisplayNode> GetDisplayNode() const;
+
+private:
+
+	/** Pointer to the key area that we were generated from */
+	TSharedPtr<IKeyArea> KeyArea;
+
+	/** The specific node that this key area relates to */
+	TSharedPtr<FSequencerDisplayNode> DisplayNode;
+
+	/** The type of this layout element */
+	EType Type;
+
+	/** The vertical offset from the top of the element's parent */
+	float LocalOffset;
+
+	/** Explicit height of the layout element */
+	float Height;
+};
+
+/** Class used for generating, and caching the layout geometry for a given display node's key areas */
+class FSectionLayout
+{
+public:
+
+	/** Constructor that takes a display node, and the index of the section to layout */
+	FSectionLayout(FSequencerDisplayNode& InNode, int32 InSectionIndex);
+
+	/** Get all layout elements that we generated */
+	const TArray<FSectionLayoutElement>& GetElements() const;
+	
+	/** Get the desired total height of this layout */
+	float GetTotalHeight() const;
+
+private:
+	/** Array of layout elements that we generated */
+	TArray<FSectionLayoutElement> Elements;
+};
\ No newline at end of file
