Copying //UE4/Dev-Sequencer to //UE4/Main

diff --git a/Engine/Source/Editor/Sequencer/Private/Sequencer.cpp b/Engine/Source/Editor/Sequencer/Private/Sequencer.cpp
--- a/Engine/Source/Editor/Sequencer/Private/Sequencer.cpp
+++ b/Engine/Source/Editor/Sequencer/Private/Sequencer.cpp
@@ -28,6 +28,7 @@
 #include "ScopedTransaction.h"
 #include "MovieSceneCameraCutSection.h"
 #include "MovieScene3DTransformSection.h"
+#include "MovieSceneSubSection.h"
 #include "MovieSceneSubTrack.h"
 #include "ISequencerSection.h"
 #include "MovieSceneSequenceInstance.h"
@@ -54,6 +55,7 @@
 #include "SceneOutlinerPublicTypes.h"
 #include "ContentBrowserModule.h"
 #include "IContentBrowserSingleton.h"
+#include "PackageTools.h"
 
 
 #define LOCTEXT_NAMESPACE "Sequencer"
@@ -84,12 +86,11 @@ void FSequencer::InitSequencer(const FSequencerInitParams& InitParams, const TSh
 	// If this sequencer edits the level, close out the existing active sequencer and mark this as the active sequencer.
 	if (bIsEditingWithinLevelEditor)
 	{
-		Settings = USequencerSettingsContainer::GetOrCreate<ULevelEditorSequencerSettings>(*InitParams.ViewParams.UniqueName);
-
 		if (ActiveSequencer.IsValid())
 		{
 			ActiveSequencer->Close();
 		}
+
 		ActiveSequencer = SharedThis(this);
 
 		// Register for saving the level so that the state of the scene can be restored before saving and updated after saving.
@@ -101,10 +102,8 @@ void FSequencer::InitSequencer(const FSequencerInitParams& InitParams, const TSh
 		FEditorDelegates::NewCurrentLevel.AddSP(this, &FSequencer::OnNewCurrentLevel);
 		FEditorDelegates::OnMapOpened.AddSP(this, &FSequencer::OnMapOpened);
 	}
-	else
-	{
-		Settings = USequencerSettingsContainer::GetOrCreate<USequencerSettings>(*InitParams.ViewParams.UniqueName);
-	}
+
+	Settings = USequencerSettingsContainer::GetOrCreate<USequencerSettings>(*InitParams.ViewParams.UniqueName);
 
 	ToolkitHost = InitParams.ToolkitHost;
 
@@ -244,13 +243,13 @@ void FSequencer::Close()
 		}
 		ActiveSequencer.Reset();
 
-		if( GLevelEditorModeTools().IsModeActive( FSequencerEdMode::EM_SequencerMode ) )
+		if (GLevelEditorModeTools().IsModeActive(FSequencerEdMode::EM_SequencerMode))
 		{
-			GLevelEditorModeTools().DeactivateMode( FSequencerEdMode::EM_SequencerMode );
+			GLevelEditorModeTools().DeactivateMode(FSequencerEdMode::EM_SequencerMode);
 		}
 	}
 
-	if( bIsEditingWithinLevelEditor )
+	if (bIsEditingWithinLevelEditor)
 	{
 		DeactivateDetailKeyframeHandler();
 			
@@ -308,6 +307,7 @@ void FSequencer::Tick(float InDeltaTime)
 	}
 }
 
+
 void FSequencer::ConstrainFixedTimeStepFrameRate()
 {
 	if( SequencerSnapValues::IsTimeSnapIntervalFrameRate(Settings->GetTimeSnapInterval()) && Settings->GetFixedTimeStepPlayback() )
@@ -425,14 +425,32 @@ TSharedRef<FMovieSceneSequenceInstance> FSequencer::GetFocusedMovieSceneSequence
 }
 
 
-void FSequencer::FocusSequenceInstance(TSharedRef<FMovieSceneSequenceInstance> SequenceInstance)
+void FSequencer::FocusSequenceInstance(UMovieSceneSubSection& InSubSection)
 {
+	TSharedRef<FMovieSceneSequenceInstance> SequenceInstance = GetSequenceInstanceForSection(InSubSection);
+
 	// Check for infinite recursion
 	check(SequenceInstance != SequenceInstanceStack.Top());
 
 	// Focus the movie scene
 	SequenceInstanceStack.Push(SequenceInstance);
 
+	// Find the time in the focused sequence instance that is relative to the current time
+	TRange<float> PlaybackRange = TRange<float>(0, 0);
+	float PlaybackRangeStart = 0.f;
+	if (InSubSection.GetSequence() != nullptr)
+	{
+		UMovieScene* ShotMovieScene = InSubSection.GetSequence()->GetMovieScene();
+		PlaybackRange = ShotMovieScene->GetPlaybackRange();
+	}
+
+	const float ShotOffset = InSubSection.StartOffset + PlaybackRange.GetLowerBoundValue() - InSubSection.PrerollTime;
+	float AbsoluteShotPosition = ShotOffset + (GetGlobalTime() - (InSubSection.GetStartTime() - InSubSection.PrerollTime)) / InSubSection.TimeScale;
+	if (!PlaybackRange.Contains(AbsoluteShotPosition))
+	{
+		AbsoluteShotPosition = PlaybackRange.GetLowerBoundValue();
+	}
+
 	// Reset data that is only used for the previous movie scene
 	ResetPerMovieSceneData();
 
@@ -441,6 +459,9 @@ void FSequencer::FocusSequenceInstance(TSharedRef<FMovieSceneSequenceInstance> S
 	UpdateTimeBoundsToFocusedMovieScene();
 
 	SequencerWidget->UpdateBreadcrumbs();
+
+	const bool bAllowSnappingToFrames = true;
+	SetGlobalTime(AbsoluteShotPosition, bAllowSnappingToFrames);
 }
 
 
@@ -751,8 +772,10 @@ bool FSequencer::CanSnapToFrame() const
 void FSequencer::NotifyMapChanged( class UWorld* NewWorld, EMapChangeType MapChangeType )
 {
 	// @todo sequencer: We should only wipe/respawn puppets that are affected by the world that is being changed! (multi-UWorld support)
-	if( ( MapChangeType == EMapChangeType::LoadMap || MapChangeType == EMapChangeType::NewMap ) )
+	if( ( MapChangeType == EMapChangeType::LoadMap || MapChangeType == EMapChangeType::NewMap || MapChangeType == EMapChangeType::TearDownWorld) )
 	{
+		SpawnRegister->DestroyAllObjects(*this);
+
 		NotifyMovieSceneDataChanged();
 	}
 }
@@ -815,6 +838,11 @@ FAnimatedRange FSequencer::GetClampRange() const
 	return FocussedMovieScene->GetEditorData().WorkingRange;
 }
 
+void FSequencer::SetClampRange(TRange<float> InNewClampRange)
+{
+	UMovieScene* FocussedMovieScene = GetFocusedMovieSceneSequence()->GetMovieScene();
+	FocussedMovieScene->GetEditorData().WorkingRange = InNewClampRange;
+}
 
 TRange<float> FSequencer::GetPlaybackRange() const
 {
@@ -845,6 +873,7 @@ void FSequencer::SetEndPlaybackRange()
 	SetPlaybackRange(TRange<float>( GetPlaybackRange().GetLowerBoundValue(), GetGlobalTime()));
 }
 
+
 void FSequencer::ZoomViewRange(float InZoomDelta)
 {
 	float LocalViewRangeMax = TargetViewRange.GetUpperBoundValue();
@@ -857,22 +886,25 @@ void FSequencer::ZoomViewRange(float InZoomDelta)
 	float NewViewOutputMin = LocalViewRangeMin - (OutputChange * CurrentPositionFraction);
 	float NewViewOutputMax = LocalViewRangeMax + (OutputChange * (1.f - CurrentPositionFraction));
 
-	if( NewViewOutputMin < NewViewOutputMax )
+	if (NewViewOutputMin < NewViewOutputMax)
 	{
 		SetViewRange(TRange<float>(NewViewOutputMin, NewViewOutputMax), EViewRangeInterpolation::Immediate);
 	}
 }
 
+
 void FSequencer::ZoomInViewRange()
 {
 	ZoomViewRange(-0.1f);
 }
 
+
 void FSequencer::ZoomOutViewRange()
 {
 	ZoomViewRange(0.1f);
 }
 
+
 void FSequencer::UpdatePlaybackRange()
 {
 	if (Settings->ShouldKeepPlayRangeInSectionBounds())
@@ -971,7 +1003,7 @@ bool FSequencer::IsRecordingLive() const
 
 float FSequencer::GetCurrentLocalTime( UMovieSceneSequence& InMovieSceneSequence )
 {
-	//@todo Sequencer - Nested movie scenes:  Figure out the parent of the passed in movie scene and 
+	//@todo Sequencer: Nested movie scenes:  Figure out the parent of the passed in movie scene and 
 	// calculate local time
 	return ScrubPosition;
 }
@@ -1093,15 +1125,14 @@ TOptional<float> FSequencer::CalculateAutoscrollEncroachment(float NewTime, floa
 		// Scrolling backwards in time, and have hit the threshold
 		return NewTime - (RangeMin + AutoScrollThreshold);
 	}
-	else if (Movement == EDirection::Positive && NewTime > RangeMax - AutoScrollThreshold)
+	
+	if (Movement == EDirection::Positive && NewTime > RangeMax - AutoScrollThreshold)
 	{
 		// Scrolling forwards in time, and have hit the threshold
 		return NewTime - (RangeMax - AutoScrollThreshold);
 	}
-	else
-	{
-		return TOptional<float>();
-	}
+
+	return TOptional<float>();
 }
 
 
@@ -1264,7 +1295,7 @@ void FSequencer::UpdateCameraCut(UObject* CameraObject, UObject* UnlockIfCameraO
 		return;
 	}
 
-	ACameraActor* UnlockIfCameraActor = Cast<ACameraActor>(UnlockIfCameraObject);
+	AActor* UnlockIfCameraActor = Cast<AActor>(UnlockIfCameraObject);
 
 	for (FLevelEditorViewportClient* LevelVC : GEditor->LevelViewportClients)
 	{
@@ -1283,14 +1314,14 @@ void FSequencer::UpdateCameraCut(UObject* CameraObject, UObject* UnlockIfCameraO
 
 void FSequencer::SetViewportSettings(const TMap<FViewportClient*, EMovieSceneViewportParams>& ViewportParamsMap)
 {
-	if(!IsPerspectiveViewportPossessionEnabled())
+	if (!IsPerspectiveViewportPossessionEnabled())
 	{
 		return;
 	}
 
-	for(FLevelEditorViewportClient* LevelVC : GEditor->LevelViewportClients)
+	for (FLevelEditorViewportClient* LevelVC : GEditor->LevelViewportClients)
 	{
-		if(LevelVC && LevelVC->IsPerspective() && LevelVC->AllowsCinematicPreview())
+		if (LevelVC && LevelVC->IsPerspective() && LevelVC->AllowsCinematicPreview())
 		{
 			if (ViewportParamsMap.Contains(LevelVC))
 			{
@@ -1318,9 +1349,9 @@ void FSequencer::SetViewportSettings(const TMap<FViewportClient*, EMovieSceneVie
 
 void FSequencer::GetViewportSettings(TMap<FViewportClient*, EMovieSceneViewportParams>& ViewportParamsMap) const
 {
-	for(FLevelEditorViewportClient* LevelVC : GEditor->LevelViewportClients)
+	for (FLevelEditorViewportClient* LevelVC : GEditor->LevelViewportClients)
 	{
-		if(LevelVC && LevelVC->IsPerspective() && LevelVC->AllowsCinematicPreview())
+		if (LevelVC && LevelVC->IsPerspective() && LevelVC->AllowsCinematicPreview())
 		{
 			EMovieSceneViewportParams ViewportParams;
 			ViewportParams.FadeAmount = LevelVC->FadeAmount;
@@ -1339,6 +1370,12 @@ EMovieScenePlayerStatus::Type FSequencer::GetPlaybackStatus() const
 }
 
 
+void FSequencer::SetPlaybackStatus(EMovieScenePlayerStatus::Type InPlaybackStatus)
+{
+	PlaybackState = InPlaybackStatus;
+}
+
+
 void FSequencer::AddOrUpdateMovieSceneInstance( UMovieSceneSection& MovieSceneSection, TSharedRef<FMovieSceneSequenceInstance> InstanceToAdd )
 {
 	if( !SequenceInstanceBySection.Contains( &MovieSceneSection ) )
@@ -1405,51 +1442,55 @@ void FSequencer::UpdateRuntimeInstances()
 void FSequencer::AttachTransportControlsToViewports()
 {
 	FLevelEditorModule* Module = FModuleManager::Get().LoadModulePtr<FLevelEditorModule>("LevelEditor");
-	if (Module)
+	
+	if (Module == nullptr)
 	{
-		TSharedPtr<ILevelEditor> LevelEditor = Module->GetFirstLevelEditor();
-		const TArray< TSharedPtr<ILevelViewport> >& LevelViewports = LevelEditor->GetViewports();
+		return;
+	}
+
+	TSharedPtr<ILevelEditor> LevelEditor = Module->GetFirstLevelEditor();
+	const TArray< TSharedPtr<ILevelViewport> >& LevelViewports = LevelEditor->GetViewports();
 		
-		FEditorWidgetsModule& EditorWidgetsModule = FModuleManager::Get().LoadModuleChecked<FEditorWidgetsModule>( "EditorWidgets" );
+	FEditorWidgetsModule& EditorWidgetsModule = FModuleManager::Get().LoadModuleChecked<FEditorWidgetsModule>( "EditorWidgets" );
 
-		FTransportControlArgs TransportControlArgs;
-		TransportControlArgs.OnForwardPlay.BindSP(this, &FSequencer::OnPlay, true);
-		TransportControlArgs.OnForwardStep.BindSP(this, &FSequencer::OnStepForward);
-		TransportControlArgs.OnBackwardStep.BindSP(this, &FSequencer::OnStepBackward);
-		TransportControlArgs.OnForwardEnd.BindSP(this, &FSequencer::OnStepToEnd);
-		TransportControlArgs.OnBackwardEnd.BindSP(this, &FSequencer::OnStepToBeginning);
-		TransportControlArgs.OnToggleLooping.BindSP(this, &FSequencer::OnToggleLooping);
-		TransportControlArgs.OnGetLooping.BindSP(this, &FSequencer::IsLooping);
-		TransportControlArgs.OnGetPlaybackMode.BindSP(this, &FSequencer::GetPlaybackMode);
+	FTransportControlArgs TransportControlArgs;
+	TransportControlArgs.OnForwardPlay.BindSP(this, &FSequencer::OnPlay, true);
+	TransportControlArgs.OnForwardStep.BindSP(this, &FSequencer::OnStepForward);
+	TransportControlArgs.OnBackwardStep.BindSP(this, &FSequencer::OnStepBackward);
+	TransportControlArgs.OnForwardEnd.BindSP(this, &FSequencer::OnStepToEnd);
+	TransportControlArgs.OnBackwardEnd.BindSP(this, &FSequencer::OnStepToBeginning);
+	TransportControlArgs.OnToggleLooping.BindSP(this, &FSequencer::OnToggleLooping);
+	TransportControlArgs.OnGetLooping.BindSP(this, &FSequencer::IsLooping);
+	TransportControlArgs.OnGetPlaybackMode.BindSP(this, &FSequencer::GetPlaybackMode);
 
-		for (int32 i = 0; i < LevelViewports.Num(); ++i)
-		{
-			const TSharedPtr<ILevelViewport>& LevelViewport = LevelViewports[i];
+	for (int32 i = 0; i < LevelViewports.Num(); ++i)
+	{
+		const TSharedPtr<ILevelViewport>& LevelViewport = LevelViewports[i];
 			
-			TSharedPtr<SWidget> TransportControl =
-				SNew(SHorizontalBox)
+		TSharedPtr<SWidget> TransportControl =
+			SNew(SHorizontalBox)
 				.Visibility(EVisibility::SelfHitTestInvisible)
-				+SHorizontalBox::Slot()
+
+			+SHorizontalBox::Slot()
 				.FillWidth(1)
 				.HAlign(HAlign_Center)
 				.VAlign(VAlign_Bottom)
 				.Padding(4.f)
 				[
 					SNew(SBorder)
-					.Padding(4.f)
-					.Cursor( EMouseCursor::Default )
-					.BorderImage( FEditorStyle::GetBrush( "FilledBorder" ) )
-					.Visibility(this, &FSequencer::GetTransportControlVisibility, LevelViewport)
-					.Content()
-					[
-						EditorWidgetsModule.CreateTransportControl(TransportControlArgs)
-					]
+						.Padding(4.f)
+						.Cursor( EMouseCursor::Default )
+						.BorderImage( FEditorStyle::GetBrush( "FilledBorder" ) )
+						.Visibility(this, &FSequencer::GetTransportControlVisibility, LevelViewport)
+						.Content()
+						[
+							EditorWidgetsModule.CreateTransportControl(TransportControlArgs)
+						]
 				];
 
-			LevelViewport->AddOverlayWidget(TransportControl.ToSharedRef());
+		LevelViewport->AddOverlayWidget(TransportControl.ToSharedRef());
 
-			TransportControls.Add(LevelViewport, TransportControl);
-		}
+		TransportControls.Add(LevelViewport, TransportControl);
 	}
 }
 
@@ -1900,7 +1941,7 @@ FGuid FSequencer::AddSpawnable(UObject& Object)
 	while (OwnerMovieScene->FindSpawnable(DuplName))
 	{
 		NewSpawnable.Name.RemoveFromEnd(UniqueString);
-		UniqueString = FString::Printf(TEXT(" (%d)"), Index);
+		UniqueString = FString::Printf(TEXT(" (%d)"), Index++);
 		NewSpawnable.Name += UniqueString;
 	}
 
@@ -1939,6 +1980,19 @@ FGuid FSequencer::MakeNewSpawnable( UObject& Object )
 		DefaultTransform.Rotation = SpawnedActor->GetActorRotation();
 		DefaultTransform.Scale = FVector(1.0f, 1.0f, 1.0f);
 		DefaultTransform.bValid = true;
+
+		// Make spawnable cameras owned by the master sequence by default so that they can persist outside of the shot. 
+		// This allows thumbnails and the look through camera button to function properly since the spawnable camera
+		// will exist.
+		UCameraComponent* CameraComponent = MovieSceneHelpers::CameraComponentFromActor(SpawnedActor);
+		if (CameraComponent)
+		{
+			FMovieSceneSpawnable* NewSpawnable = GetFocusedMovieSceneSequence()->GetMovieScene()->FindSpawnable(NewGuid);
+			if (NewSpawnable)
+			{
+				NewSpawnable->SetSpawnOwnership(ESpawnOwnership::MasterSequence);
+			}
+		}
 	}
 
 	SetupDefaultsForSpawnable(NewGuid, DefaultTransform);
@@ -1974,8 +2028,6 @@ void FSequencer::SetupDefaultsForSpawnable( const FGuid& Guid, const FTransformD
 
 		if (TransformTrack)
 		{
-			static FName Transform("Transform");
-			TransformTrack->SetPropertyNameAndPath(Transform, Transform.ToString());
 			const bool bUnwindRotation = false;
 
 			EMovieSceneKeyInterpolation Interpolation = GetKeyInterpolation();
@@ -2015,7 +2067,6 @@ void FSequencer::AddSubSequence(UMovieSceneSequence* Sequence)
 
 	// @todo sequencer: Undo doesn't seem to be working at all
 	const FScopedTransaction Transaction( LOCTEXT("UndoAddingObject", "Add Object to MovieScene") );
-
 	OwnerMovieScene->Modify();
 
 	UMovieSceneSubTrack* SubTrack = OwnerMovieScene->AddMasterTrack<UMovieSceneSubTrack>();
@@ -2051,11 +2102,13 @@ void GetParentTrackNodeAndNamePath(TSharedRef<const FSequencerDisplayNode> Displ
 	TArray<FName> PathToTrack;
 	PathToTrack.Add( DisplayNode->GetNodeName() );
 	TSharedPtr<FSequencerDisplayNode> CurrentParent = DisplayNode->GetParent();
+
 	while ( CurrentParent.IsValid() && CurrentParent->GetType() != ESequencerNode::Track )
 	{
 		PathToTrack.Add( CurrentParent->GetNodeName() );
 		CurrentParent = CurrentParent->GetParent();
 	}
+
 	if ( CurrentParent.IsValid() )
 	{
 		OutParentTrack = StaticCastSharedPtr<FSequencerTrackNode>( CurrentParent );
@@ -2172,7 +2225,7 @@ void FSequencer::PlaceActorInFrontOfCamera( AActor* ActorCDO )
 	{
 		// Don't allow this when the active viewport is showing a simulation/PIE level
 		const bool bIsViewportShowingPIEWorld = GCurrentLevelEditingViewportClient->GetWorld()->GetOutermost()->HasAnyPackageFlags(PKG_PlayInEditor);
-		if( !bIsViewportShowingPIEWorld )
+		if (!bIsViewportShowingPIEWorld)
 		{
 			// @todo sequencer actors: Ideally we could use the actor's collision to figure out how far to push out
 			// the object (like when placing in viewports), but we can't really do that because we're only dealing with a CDO
@@ -2199,7 +2252,7 @@ void FSequencer::PlaceActorInFrontOfCamera( AActor* ActorCDO )
 
 void FSequencer::GetActorRecordingState( bool& bIsRecording /* In+Out */ ) const
 {
-	if( IsRecordingLive() )
+	if (IsRecordingLive())
 	{
 		bIsRecording = true;
 	}
@@ -2271,7 +2324,7 @@ void FSequencer::ActivateDetailKeyframeHandler()
 	static const FName DetailsTabIdentifiers[] = { "LevelEditorSelectionDetails", "LevelEditorSelectionDetails2", "LevelEditorSelectionDetails3", "LevelEditorSelectionDetails4" };
 	FPropertyEditorModule& EditModule = FModuleManager::Get().GetModuleChecked<FPropertyEditorModule>("PropertyEditor");
 
-	for(const FName& DetailsTabIdentifier : DetailsTabIdentifiers)
+	for (const FName& DetailsTabIdentifier : DetailsTabIdentifiers)
 	{
 		TSharedPtr<IDetailsView> DetailsView = EditModule.FindDetailView(DetailsTabIdentifier);
 
@@ -2288,11 +2341,11 @@ void FSequencer::DeactivateDetailKeyframeHandler()
 	static const FName DetailsTabIdentifiers[] = { "LevelEditorSelectionDetails", "LevelEditorSelectionDetails2", "LevelEditorSelectionDetails3", "LevelEditorSelectionDetails4" };
 	FPropertyEditorModule& EditModule = FModuleManager::Get().GetModuleChecked<FPropertyEditorModule>("PropertyEditor");
 
-	for(const FName& DetailsTabIdentifier : DetailsTabIdentifiers)
+	for (const FName& DetailsTabIdentifier : DetailsTabIdentifiers)
 	{
 		TSharedPtr<IDetailsView> DetailsView = EditModule.FindDetailView(DetailsTabIdentifier);
 
-		if(DetailsView.IsValid())
+		if (DetailsView.IsValid())
 		{
 			if (DetailsView->GetKeyframeHandler() == DetailKeyframeHandler)
 			{
@@ -2311,7 +2364,7 @@ void FSequencer::OnPropertyEditorOpened()
 
 void FSequencer::UpdatePreviewLevelViewportClientFromCameraCut(FLevelEditorViewportClient& InViewportClient, UObject* InCameraObject) const
 {
-	ACameraActor* CameraActor = Cast<ACameraActor>(InCameraObject);
+	AActor* CameraActor = Cast<AActor>(InCameraObject);
 
 	if (CameraActor)
 	{
@@ -2333,22 +2386,26 @@ void FSequencer::UpdatePreviewLevelViewportClientFromCameraCut(FLevelEditorViewp
 	// If viewing through a camera - enforce aspect ratio.
 	if (CameraActor)
 	{
-		if (CameraActor->GetCameraComponent()->AspectRatio == 0)
-		{
-			InViewportClient.AspectRatio = 1.7f;
-		}
-		else
+		UCameraComponent* CameraComponent = Cast<UCameraComponent>(CameraActor->GetComponentByClass(UCameraComponent::StaticClass()));
+		if (CameraComponent)
 		{
-			InViewportClient.AspectRatio = CameraActor->GetCameraComponent()->AspectRatio;
-		}
+			if (CameraComponent->AspectRatio == 0)
+			{
+				InViewportClient.AspectRatio = 1.7f;
+			}
+			else
+			{
+				InViewportClient.AspectRatio = CameraComponent->AspectRatio;
+			}
 
-		//don't stop the camera from zooming when not playing back
-		InViewportClient.ViewFOV = CameraActor->GetCameraComponent()->FieldOfView;
+			//don't stop the camera from zooming when not playing back
+			InViewportClient.ViewFOV = CameraComponent->FieldOfView;
 
-		// If there are selected actors, invalidate the viewports hit proxies, otherwise they won't be selectable afterwards
-		if (InViewportClient.Viewport && GEditor->GetSelectedActorCount() > 0)
-		{
-			InViewportClient.Viewport->InvalidateHitProxy();
+			// If there are selected actors, invalidate the viewports hit proxies, otherwise they won't be selectable afterwards
+			if (InViewportClient.Viewport && GEditor->GetSelectedActorCount() > 0)
+			{
+				InViewportClient.Viewport->InvalidateHitProxy();
+			}
 		}
 	}
 
@@ -2357,6 +2414,33 @@ void FSequencer::UpdatePreviewLevelViewportClientFromCameraCut(FLevelEditorViewp
 }
 
 
+void GetDescendantMovieScenes(UMovieSceneSequence* InSequence, TArray<UMovieScene*> & InMovieScenes)
+{
+	UMovieScene* InMovieScene = InSequence->GetMovieScene();
+	if (InMovieScene == nullptr)
+	{
+		return;
+	}
+
+	InMovieScenes.Add(InMovieScene);
+
+	for (auto MasterTrack : InMovieScene->GetMasterTracks())
+	{
+		for (auto Section : MasterTrack->GetAllSections())
+		{
+			UMovieSceneSubSection* SubSection = Cast<UMovieSceneSubSection>(Section);
+			if (SubSection != nullptr)
+			{
+				UMovieSceneSequence* SubSequence = SubSection->GetSequence();
+				if (SubSequence != nullptr)
+				{
+					GetDescendantMovieScenes(SubSequence, InMovieScenes);
+				}
+			}
+		}
+	}
+}
+
 void FSequencer::SaveCurrentMovieScene()
 {
 	// Capture thumbnail
@@ -2401,13 +2485,47 @@ void FSequencer::SaveCurrentMovieScene()
 		Viewport->Draw();
 	}
 
-	UPackage* MovieScenePackage = GetCurrentAsset()->GetOutermost();
-
 	TArray<UPackage*> PackagesToSave;
+	TArray<UMovieScene*> MovieScenesToSave;
+	GetDescendantMovieScenes(RootMovieSceneSequenceInstance->GetSequence(), MovieScenesToSave);
+	for (auto MovieSceneToSave : MovieScenesToSave)
+	{
+		UPackage* MovieScenePackageToSave = MovieSceneToSave->GetOuter()->GetOutermost();
+		if (MovieScenePackageToSave->IsDirty())
+		{
+			PackagesToSave.Add(MovieScenePackageToSave);
+		}
+	}
+
+	// If there's more than 1 movie scene to save, prompt the user whether to save all dirty movie scenes.
+	const bool bCheckDirty = PackagesToSave.Num() > 1;
+	const bool bPromptToSave = PackagesToSave.Num() > 1;
+
+	FEditorFileUtils::PromptForCheckoutAndSave( PackagesToSave, bCheckDirty, bPromptToSave );
+}
+
+
+void FSequencer::SaveCurrentMovieSceneAs()
+{
+	TSharedPtr<IToolkitHost> MyToolkitHost = GetToolkitHost();
 
-	PackagesToSave.Add( MovieScenePackage );
+	if (!MyToolkitHost.IsValid())
+	{
+		return;
+	}
 
-	FEditorFileUtils::PromptForCheckoutAndSave( PackagesToSave, false, false );
+	TArray<UObject*> AssetsToSave;
+	AssetsToSave.Add(GetCurrentAsset());
+
+	TArray<UObject*> SavedAssets;
+	FEditorFileUtils::SaveAssetsAs(AssetsToSave, SavedAssets);
+
+	if ((SavedAssets[0] != AssetsToSave[0]) && (SavedAssets[0] != nullptr))
+	{
+		FAssetEditorManager& AssetEditorManager = FAssetEditorManager::Get();
+		AssetEditorManager.CloseAllEditorsForAsset(AssetsToSave[0]);
+		AssetEditorManager.OpenEditorForAssets(SavedAssets, EToolkitMode::Standalone, MyToolkitHost.ToSharedRef());
+	}
 }
 
 
@@ -2415,12 +2533,13 @@ void FSequencer::AddSelectedObjects()
 {
 	UMovieSceneSequence* OwnerSequence = GetFocusedMovieSceneSequence();
 	UMovieScene* OwnerMovieScene = OwnerSequence->GetMovieScene();
+	USelection* CurrentSelection = GEditor->GetSelectedActors();
 
 	TArray<AActor*> SelectedActors;
-
-	USelection* CurrentSelection = GEditor->GetSelectedActors();
 	TArray<UObject*> SelectedObjects;
+
 	CurrentSelection->GetSelectedObjects( AActor::StaticClass(), SelectedObjects );
+
 	for (TArray<UObject*>::TIterator It(SelectedObjects); It; ++It)
 	{
 		AActor* Actor = CastChecked<AActor>(*It);
@@ -2601,6 +2720,7 @@ void FSequencer::AssignActor(FMenuBuilder& MenuBuilder, FGuid InObjectBinding)
 				})
 			)
 		];
+
 	MenuBuilder.AddWidget(MiniSceneOutliner, FText::GetEmpty(), true);
 	MenuBuilder.EndSection();
 }
@@ -2608,90 +2728,95 @@ void FSequencer::AssignActor(FMenuBuilder& MenuBuilder, FGuid InObjectBinding)
 
 void FSequencer::DoAssignActor(AActor*const* InActors, int32 NumActors, FGuid InObjectBinding)
 {
-	if (NumActors > 0)
+	if (NumActors <= 0)
 	{
-		AActor* Actor = InActors[0];
-		if (Actor != nullptr)
-		{
-			FScopedTransaction AssignActor( NSLOCTEXT("Sequencer", "AssignActor", "Assign Actor") );
+		return;
+	}
 
-			RootMovieSceneSequenceInstance->SaveState(*this);
+	AActor* Actor = InActors[0];
 
-			UMovieSceneSequence* OwnerSequence = GetFocusedMovieSceneSequence();
-			UMovieScene* OwnerMovieScene = OwnerSequence->GetMovieScene();
+	if (Actor == nullptr)
+	{
+		return;
+	}
 
-			Actor->Modify();
-			OwnerSequence->Modify();
-			OwnerMovieScene->Modify();
+	FScopedTransaction AssignActor( NSLOCTEXT("Sequencer", "AssignActor", "Assign Actor") );
 
-			UObject* RuntimeObject = SequenceInstanceStack.Top()->FindObject(InObjectBinding, *this);
+	RootMovieSceneSequenceInstance->SaveState(*this);
 
-			// Replace the object itself
-			FGuid ParentGuid;
-			{
-				// Get the object guid to assign, remove the binding if it already exists
-				ParentGuid = SequenceInstanceStack.Top()->FindObjectId(*Actor);
-				FString NewActorLabel = Actor->GetActorLabel();
-				if (ParentGuid.IsValid())
-				{
-					OwnerMovieScene->RemovePossessable(ParentGuid);
-					OwnerSequence->UnbindPossessableObjects(ParentGuid);
-				}
+	UMovieSceneSequence* OwnerSequence = GetFocusedMovieSceneSequence();
+	UMovieScene* OwnerMovieScene = OwnerSequence->GetMovieScene();
 
-				// Add this object
-				FMovieScenePossessable NewPossessable( NewActorLabel, Actor->GetClass());
-				ParentGuid = NewPossessable.GetGuid();
-				OwnerSequence->BindPossessableObject(ParentGuid, *Actor, GetPlaybackContext());
+	Actor->Modify();
+	OwnerSequence->Modify();
+	OwnerMovieScene->Modify();
 
-				// Replace
-				OwnerMovieScene->ReplacePossessable(InObjectBinding, ParentGuid, NewActorLabel);
-			}
+	UObject* RuntimeObject = SequenceInstanceStack.Top()->FindObject(InObjectBinding, *this);
+
+	// Replace the object itself
+	FGuid ParentGuid;
+	{
+		// Get the object guid to assign, remove the binding if it already exists
+		ParentGuid = SequenceInstanceStack.Top()->FindObjectId(*Actor);
+		FString NewActorLabel = Actor->GetActorLabel();
+		if (ParentGuid.IsValid())
+		{
+			OwnerMovieScene->RemovePossessable(ParentGuid);
+			OwnerSequence->UnbindPossessableObjects(ParentGuid);
+		}
+
+		// Add this object
+		FMovieScenePossessable NewPossessable( NewActorLabel, Actor->GetClass());
+		ParentGuid = NewPossessable.GetGuid();
+		OwnerSequence->BindPossessableObject(ParentGuid, *Actor, GetPlaybackContext());
+
+		// Replace
+		OwnerMovieScene->ReplacePossessable(InObjectBinding, ParentGuid, NewActorLabel);
+	}
 
-			// Handle components
-			if (AActor* ActorToReplace = Cast<AActor>(RuntimeObject))
+	// Handle components
+	if (AActor* ActorToReplace = Cast<AActor>(RuntimeObject))
+	{
+		for (UActorComponent* ComponentToReplace : ActorToReplace->GetComponents())
+		{
+			if (ComponentToReplace != nullptr)
 			{
-				for (UActorComponent* ComponentToReplace : ActorToReplace->GetComponents())
+				FGuid ComponentGuid = SequenceInstanceStack.Top()->FindObjectId(*ComponentToReplace);
+				if (ComponentGuid.IsValid())
 				{
-					if (ComponentToReplace != nullptr)
+					UActorComponent* NewComponent = Actor->GetComponentByClass(ComponentToReplace->GetClass());
+					if (NewComponent)
 					{
-						FGuid ComponentGuid = SequenceInstanceStack.Top()->FindObjectId(*ComponentToReplace);
-						if (ComponentGuid.IsValid())
+						// Get the object guid to assign, remove the binding if it already exists
+						FGuid NewComponentGuid = SequenceInstanceStack.Top()->FindObjectId(*NewComponent);
+						FString NewComponentLabel = NewComponent->GetName();
+						if (NewComponentGuid.IsValid())
 						{
-							UActorComponent* NewComponent = Actor->GetComponentByClass(ComponentToReplace->GetClass());
-							if (NewComponent)
-							{
-								// Get the object guid to assign, remove the binding if it already exists
-								FGuid NewComponentGuid = SequenceInstanceStack.Top()->FindObjectId(*NewComponent);
-								FString NewComponentLabel = NewComponent->GetName();
-								if (NewComponentGuid.IsValid())
-								{
-									OwnerMovieScene->RemovePossessable(NewComponentGuid);
-									OwnerSequence->UnbindPossessableObjects(NewComponentGuid);
-								}
+							OwnerMovieScene->RemovePossessable(NewComponentGuid);
+							OwnerSequence->UnbindPossessableObjects(NewComponentGuid);
+						}
 
-								// Add this object
-								FMovieScenePossessable NewPossessable( NewComponentLabel, NewComponent->GetClass());
-								NewComponentGuid = NewPossessable.GetGuid();
-								OwnerSequence->BindPossessableObject(NewComponentGuid, *NewComponent, Actor);
+						// Add this object
+						FMovieScenePossessable NewPossessable( NewComponentLabel, NewComponent->GetClass());
+						NewComponentGuid = NewPossessable.GetGuid();
+						OwnerSequence->BindPossessableObject(NewComponentGuid, *NewComponent, Actor);
 
-								// Replace
-								OwnerMovieScene->ReplacePossessable(ComponentGuid, NewComponentGuid, NewComponentLabel);
+						// Replace
+						OwnerMovieScene->ReplacePossessable(ComponentGuid, NewComponentGuid, NewComponentLabel);
 
-								FMovieScenePossessable* ThisPossessable = OwnerMovieScene->FindPossessable(NewComponentGuid);
-								if (ensure(ThisPossessable))
-								{
-									ThisPossessable->SetParent(ParentGuid);
-								}
-							}
+						FMovieScenePossessable* ThisPossessable = OwnerMovieScene->FindPossessable(NewComponentGuid);
+						if (ensure(ThisPossessable))
+						{
+							ThisPossessable->SetParent(ParentGuid);
 						}
 					}
 				}
 			}
-			
-			RootMovieSceneSequenceInstance->RestoreState(*this);
-			NotifyMovieSceneDataChanged();
 		}
 	}
+			
+	RootMovieSceneSequenceInstance->RestoreState(*this);
+	NotifyMovieSceneDataChanged();
 }
 
 
@@ -2705,7 +2830,6 @@ void FSequencer::DeleteNode(TSharedRef<FSequencerDisplayNode> NodeToBeDeleted)
 	else
 	{
 		const FScopedTransaction Transaction( NSLOCTEXT("Sequencer", "UndoDeletingObject", "Delete Node") );
-
 		OnRequestNodeDeleted(NodeToBeDeleted);
 	}
 }
@@ -2715,18 +2839,20 @@ void FSequencer::DeleteSelectedNodes()
 {
 	TSet< TSharedRef<FSequencerDisplayNode> > SelectedNodesCopy = GetSelection().GetSelectedOutlinerNodes();
 
-	if( SelectedNodesCopy.Num() > 0 )
+	if (SelectedNodesCopy.Num() == 0)
 	{
-		const FScopedTransaction Transaction( NSLOCTEXT("Sequencer", "UndoDeletingObject", "Delete Node") );
+		return;
+	}
+
+	const FScopedTransaction Transaction( NSLOCTEXT("Sequencer", "UndoDeletingObject", "Delete Node") );
 
-		for( const TSharedRef<FSequencerDisplayNode>& SelectedNode : SelectedNodesCopy )
+	for( const TSharedRef<FSequencerDisplayNode>& SelectedNode : SelectedNodesCopy )
+	{
+		if( !SelectedNode->IsHidden() )
 		{
-			if( !SelectedNode->IsHidden() )
-			{
-				// Delete everything in the entire node
-				TSharedRef<const FSequencerDisplayNode> NodeToBeDeleted = StaticCastSharedRef<const FSequencerDisplayNode>(SelectedNode);
-				OnRequestNodeDeleted( NodeToBeDeleted );
-			}
+			// Delete everything in the entire node
+			TSharedRef<const FSequencerDisplayNode> NodeToBeDeleted = StaticCastSharedRef<const FSequencerDisplayNode>(SelectedNode);
+			OnRequestNodeDeleted( NodeToBeDeleted );
 		}
 	}
 }
@@ -2735,7 +2861,6 @@ void FSequencer::DeleteSelectedNodes()
 void FSequencer::ToggleNodeActive()
 {
 	bool bIsActive = !IsNodeActive();
-
 	const FScopedTransaction Transaction( NSLOCTEXT("Sequencer", "ToggleNodeActive", "Toggle Node Active") );
 
 	for (auto OutlinerNode : Selection.GetSelectedOutlinerNodes())
@@ -2818,9 +2943,7 @@ void FSequencer::ConvertToSpawnable(TSharedRef<FSequencerObjectBindingNode> Node
 	// Ensure we're in a non-possessed state
 	RootMovieSceneSequenceInstance->RestoreState(*this);
 	GetFocusedMovieSceneSequence()->GetMovieScene()->Modify();
-
 	ConvertToSpawnableInternal(NodeToBeConverted);
-
 	NotifyMovieSceneDataChanged();
 }
 
@@ -2873,7 +2996,6 @@ void FSequencer::ConvertToSpawnableInternal(TSharedRef<FSequencerObjectBindingNo
 {
 	UMovieSceneSequence* Sequence = GetFocusedMovieSceneSequence();
 	UMovieScene* MovieScene = Sequence->GetMovieScene();
-
 	FMovieScenePossessable* Possessable = MovieScene->FindPossessable(NodeToBeConverted->GetObjectBinding());
 
 	if (!Possessable)
@@ -2937,10 +3059,10 @@ void FSequencer::ConvertToSpawnableInternal(TSharedRef<FSequencerObjectBindingNo
 	}
 
 	SetupDefaultsForSpawnable(Spawnable->GetGuid(), DefaultTransform);
-
 	SetGlobalTimeDirectly(ScrubPosition);
 }
 
+
 void FSequencer::TogglePlay()
 {
 	OnPlay();
@@ -3087,7 +3209,7 @@ void FSequencer::SetKeyTime(const bool bUseFrames)
 		bUseFrames ? NSLOCTEXT("Sequencer.Popups", "SetKeyFramePopup", "New Frame") : NSLOCTEXT("Sequencer.Popups", "SetKeyTimePopup", "New Time"),
 		bUseFrames ? FText::AsNumber( SequencerHelpers::TimeToFrame( KeyTime, FrameRate )) : FText::AsNumber( KeyTime ),
 		FOnTextCommitted::CreateSP(this, &FSequencer::OnSetKeyTimeTextCommitted, bUseFrames)
-		);
+	);
 }
 
 
@@ -3105,7 +3227,6 @@ void FSequencer::OnSetKeyTimeTextCommitted(const FText& InText, ETextCommit::Typ
 		const float NewKeyTime = (float)dNewTime;
 
 		FScopedTransaction SetKeyTimeTransaction(NSLOCTEXT("Sequencer", "SetKeyTime_Transaction", "Set Key Time"));
-
 		TArray<FSequencerSelectedKey> SelectedKeysArray = Selection.GetSelectedKeys().Array();
 	
 		for ( const FSequencerSelectedKey& Key : SelectedKeysArray )
@@ -3218,7 +3339,6 @@ void FSequencer::CopySelectedKeys()
 		});
 	}
 
-
 	TSharedRef<FMovieSceneClipboard> Clipboard = MakeShareable( new FMovieSceneClipboard(Builder.Commit(CopyRelativeTo)) );
 	
 	if (Clipboard->GetKeyTrackGroups().Num())
@@ -3258,6 +3378,51 @@ void FSequencer::OnClipboardUsed(TSharedPtr<FMovieSceneClipboard> Clipboard)
 }
 
 
+void FSequencer::DiscardChanges()
+{
+	if (SequenceInstanceStack.Num() == 0)
+	{
+		return;
+	}
+
+	const UMovieSceneSequence* EditedSequence = SequenceInstanceStack.Top()->GetSequence();
+
+	if (EditedSequence == nullptr)
+	{
+		return;
+	}
+
+	FString SequencePath = EditedSequence->GetPathName();
+/*
+	// @todo sequencer: the following is not going to work, because the objects in the unloaded
+	// package are referenced elsewhere and cannot be destroyed. LoadObject will simply load the
+	// already loaded objects (and sub-objects, such as tracks and sections) again.
+
+	// unload asset
+	SequenceInstanceStack.Pop();
+	PackageTools::UnloadPackages(TArrayBuilder<UPackage*>().Add(EditedSequence->GetOutermost()), false);
+
+	// reload asset
+	UMovieSceneSequence* ReloadedSequence = LoadObject<UMovieSceneSequence>(nullptr, *SequencePath);
+
+	if (ReloadedSequence == nullptr)
+	{
+		return;
+	}
+
+	// refresh sequencer
+	if (SequenceInstanceStack.Num() == 0)
+	{
+		RootMovieSceneSequenceInstance = nullptr;
+		ResetToNewRootSequence(*ReloadedSequence);
+	}
+	else
+	{
+		FocusSequenceInstance(MakeShareable(new FMovieSceneSequenceInstance(*ReloadedSequence)));
+	}*/
+}
+
+
 void FSequencer::GenericTextEntryModeless(const FText& DialogText, const FText& DefaultText, FOnTextCommitted OnTextComitted)
 {
 	TSharedRef<STextEntryPopup> TextEntryPopup = 
@@ -3275,7 +3440,7 @@ void FSequencer::GenericTextEntryModeless(const FText& DialogText, const FText&
 		TextEntryPopup,
 		FSlateApplication::Get().GetCursorPos(),
 		FPopupTransitionEffect(FPopupTransitionEffect::TypeInPopup)
-		);
+	);
 }
 
 
@@ -3291,9 +3456,7 @@ void FSequencer::CloseEntryPopupMenu()
 void FSequencer::TrimSection(bool bTrimLeft)
 {
 	FScopedTransaction TrimSectionTransaction( NSLOCTEXT("Sequencer", "TrimSection_Transaction", "Trim Section") );
-
 	MovieSceneToolHelpers::TrimSection(Selection.GetSelectedSections(), GetGlobalTime(), bTrimLeft);
-
 	NotifyMovieSceneDataChanged();
 }
 
@@ -3301,24 +3464,25 @@ void FSequencer::TrimSection(bool bTrimLeft)
 void FSequencer::SplitSection()
 {
 	FScopedTransaction SplitSectionTransaction( NSLOCTEXT("Sequencer", "SplitSection_Transaction", "Split Section") );
-
 	MovieSceneToolHelpers::SplitSection(Selection.GetSelectedSections(), GetGlobalTime());
-
 	NotifyMovieSceneDataChanged();
 }
 
+
 /** Get the current active hotspot */
 TSharedPtr<ISequencerHotspot> FSequencer::GetHotspot() const
 {
 	return Hotspot;
 }
 
+
 /** Set the hotspot to something else */
 void FSequencer::SetHotspot(TSharedPtr<ISequencerHotspot> NewHotspot)
 {
 	Hotspot = MoveTemp(NewHotspot);
 }
 
+
 void FSequencer::BindSequencerCommands()
 {
 	const FSequencerCommands& Commands = FSequencerCommands::Get();
@@ -3501,12 +3665,12 @@ void FSequencer::BindSequencerCommands()
 		FExecuteAction::CreateSP( this, &FSequencer::FindInContentBrowser ) );
 
 	SequencerCommandBindings->MapAction(
-		Commands.ToggleDetailsView,
+		Commands.ToggleCombinedKeyframes,
 		FExecuteAction::CreateLambda( [this]{
-			Settings->SetDetailsViewVisible( !Settings->GetDetailsViewVisible() );
+			Settings->SetShowCombinedKeyframes( !Settings->GetShowCombinedKeyframes() );
 		} ),
 		FCanExecuteAction::CreateLambda( []{ return true; } ),
-		FIsActionChecked::CreateLambda( [this]{ return Settings->GetDetailsViewVisible(); } ) );
+		FIsActionChecked::CreateLambda( [this]{ return Settings->GetShowCombinedKeyframes(); } ) );
 
 	SequencerCommandBindings->MapAction(
 		Commands.ToggleLabelBrowser,
@@ -3592,6 +3756,12 @@ void FSequencer::BindSequencerCommands()
 		FCanExecuteAction::CreateLambda( []{ return true; } ),
 		FIsActionChecked::CreateLambda( [this]{ return Settings->GetShowCurveEditor(); } ) );
 
+	SequencerCommandBindings->MapAction(
+		Commands.ToggleLinkCurveEditorTimeRange,
+		FExecuteAction::CreateLambda( [this]{ Settings->SetLinkCurveEditorTimeRange(!Settings->GetLinkCurveEditorTimeRange()); } ),
+		FCanExecuteAction::CreateLambda( []{ return true; } ),
+		FIsActionChecked::CreateLambda( [this]{ return Settings->GetLinkCurveEditorTimeRange(); } ) );
+
 	auto CanCutOrCopy = [this]{
 		if (!IsSequencerWidgetFocused())
 		{
@@ -3664,6 +3834,22 @@ void FSequencer::BindSequencerCommands()
 		})
 	);
 
+	SequencerCommandBindings->MapAction(
+		Commands.DiscardChanges,
+		FExecuteAction::CreateSP(this, &FSequencer::DiscardChanges),
+		FCanExecuteAction::CreateLambda([this]{
+			if (SequenceInstanceStack.Num() == 0)
+			{
+				return false;
+			}
+
+			UMovieSceneSequence* EditedSequence = SequenceInstanceStack.Top()->GetSequence();
+
+			return ((EditedSequence != nullptr) && EditedSequence->GetOutermost()->IsDirty());
+		})
+	);
+
+
 	SequencerWidget->BindCommands(SequencerCommandBindings);
 
 	for (int32 i = 0; i < TrackEditors.Num(); ++i)
diff --git a/Engine/Source/Editor/Sequencer/Private/Sequencer.h b/Engine/Source/Editor/Sequencer/Private/Sequencer.h
--- a/Engine/Source/Editor/Sequencer/Private/Sequencer.h
+++ b/Engine/Source/Editor/Sequencer/Private/Sequencer.h
@@ -55,15 +55,13 @@ public:
 	 */
 	void InitSequencer(const FSequencerInitParams& InitParams, const TSharedRef<ISequencerObjectChangeListener>& InObjectChangeListener, const TSharedRef<IDetailKeyframeHandler>& InDetailKeyframeHandler, const TArray<FOnCreateTrackEditor>& TrackEditorDelegates);
 
-	/** Set the global time directly, without performing any auto-scroll */
-	void SetGlobalTimeDirectly(float Time, const bool& bAllowSnappingToFrames = false);
-
 	/** @return The current view range */
 	virtual FAnimatedRange GetViewRange() const override;
 	virtual void SetViewRange(TRange<float> NewViewRange, EViewRangeInterpolation Interpolation = EViewRangeInterpolation::Animated) override;
 
 	/** @return The current clamp range */
 	FAnimatedRange GetClampRange() const;
+	void SetClampRange(TRange<float> InNewClampRange);
 
 	TRange<float> GetPlaybackRange() const;
 	void SetPlaybackRange(TRange<float> InNewPlaybackRange);
@@ -257,6 +255,9 @@ public:
 	/** Called to save the current movie scene */
 	void SaveCurrentMovieScene();
 
+	/** Called to save the current movie scene under a new name */
+	void SaveCurrentMovieSceneAs();
+
 	/** Called to add selected objects to the movie scene */
 	void AddSelectedObjects();
 
@@ -303,6 +304,9 @@ public:
 	/** Promote a clipboard to the top of the clipboard stack, and update its timestamp */
 	void OnClipboardUsed(TSharedPtr<FMovieSceneClipboard> Clipboard);
 
+	/** Discard all changes to the current movie scene. */
+	void DiscardChanges();
+
 public:
 
 	/** Get the current active hotspot */
@@ -352,7 +356,7 @@ public:
 	virtual void ResetToNewRootSequence(UMovieSceneSequence& NewSequence) override;
 	virtual TSharedRef<FMovieSceneSequenceInstance> GetRootMovieSceneSequenceInstance() const override;
 	virtual TSharedRef<FMovieSceneSequenceInstance> GetFocusedMovieSceneSequenceInstance() const override;
-	virtual void FocusSequenceInstance( TSharedRef<FMovieSceneSequenceInstance> SequenceInstance ) override;
+	virtual void FocusSequenceInstance( UMovieSceneSubSection& InSubSection ) override;
 	virtual TSharedRef<FMovieSceneSequenceInstance> GetSequenceInstanceForSection(UMovieSceneSection& Section) const override;
 	virtual bool HasSequenceInstanceForSection(UMovieSceneSection& Section) const override;
 	virtual EAutoKeyMode GetAutoKeyMode() const override;
@@ -367,6 +371,7 @@ public:
 	virtual float GetCurrentLocalTime(UMovieSceneSequence& InMovieSceneSequence) override;
 	virtual float GetGlobalTime() override;
 	virtual void SetGlobalTime(float Time, const bool& bAllowSnappingToFrames = false) override;
+	virtual void SetGlobalTimeDirectly(float Time, const bool& bAllowSnappingToFrames = false) override;
 	virtual void SetPerspectiveViewportPossessionEnabled(bool bEnabled) override;
 	virtual void SetPerspectiveViewportCameraCutEnabled(bool bEnabled) override;
 	virtual FGuid GetHandleToObject(UObject* Object, bool bCreateHandleIfMissing = true) override;
@@ -398,6 +403,7 @@ public:
 	virtual void SetViewportSettings(const TMap<FViewportClient*, EMovieSceneViewportParams>& ViewportParamsMap) override;
 	virtual void GetViewportSettings(TMap<FViewportClient*, EMovieSceneViewportParams>& ViewportParamsMap) const override;
 	virtual EMovieScenePlayerStatus::Type GetPlaybackStatus() const override;
+	virtual void SetPlaybackStatus(EMovieScenePlayerStatus::Type InPlaybackStatus) override;
 	virtual void AddOrUpdateMovieSceneInstance(UMovieSceneSection& MovieSceneSection, TSharedRef<FMovieSceneSequenceInstance> InstanceToAdd) override;
 	virtual void RemoveMovieSceneInstance(UMovieSceneSection& MovieSceneSection, TSharedRef<FMovieSceneSequenceInstance> InstanceToRemove) override;
 	virtual IMovieSceneSpawnRegister& GetSpawnRegister() override { return *SpawnRegister; }
diff --git a/Engine/Source/Editor/Sequencer/Private/SequencerCommands.cpp b/Engine/Source/Editor/Sequencer/Private/SequencerCommands.cpp
--- a/Engine/Source/Editor/Sequencer/Private/SequencerCommands.cpp
+++ b/Engine/Source/Editor/Sequencer/Private/SequencerCommands.cpp
@@ -66,10 +66,11 @@ void FSequencerCommands::RegisterCommands()
 	UI_COMMAND( ToggleSnapCurveValueToInterval, "Snap Curve Key Values", "Snap curve keys to the value snapping interval", EUserInterfaceActionType::ToggleButton, FInputChord() );
 
 	UI_COMMAND( FindInContentBrowser, "Find in Content Browser", "Find the viewed sequence asset in the content browser", EUserInterfaceActionType::Button, FInputChord() );
-	UI_COMMAND( ToggleDetailsView, "Details View", "Enable details view for selected sections and keys", EUserInterfaceActionType::ToggleButton, FInputChord() );
+	UI_COMMAND( ToggleCombinedKeyframes, "Combined Keyframes", "Show/hide the combined keyframes at the top node level", EUserInterfaceActionType::ToggleButton, FInputChord() );
 	UI_COMMAND( ToggleLabelBrowser, "Label Browser", "Show/hide the track label browser", EUserInterfaceActionType::ToggleButton, FInputChord() );
 	
 	UI_COMMAND( ToggleShowCurveEditor, "Curve Editor", "Show the animation keys in a curve editor", EUserInterfaceActionType::ToggleButton, FInputChord() );
+	UI_COMMAND( ToggleLinkCurveEditorTimeRange, "Link Curve Editor Time Range", "Link the curve editor time range to the sequence", EUserInterfaceActionType::ToggleButton, FInputChord() );
 
 	UI_COMMAND( MoveTool, "Move", "Activates the move tool", EUserInterfaceActionType::ToggleButton, FInputChord( EKeys::V ) );
 	UI_COMMAND( MarqueeTool, "Marquee", "Activates the marquee selection tool", EUserInterfaceActionType::ToggleButton, FInputChord( EKeys::M ) );
@@ -79,6 +80,8 @@ void FSequencerCommands::RegisterCommands()
 	UI_COMMAND( PasteFromHistory, "Paste From History", "Paste from the sequencer clipboard history", EUserInterfaceActionType::Button, FInputChord(EModifierKey::Control | EModifierKey::Shift, EKeys::V) );
 
 	UI_COMMAND( ConvertToSpawnable, "Convert to Spawnable", "Make the specified possessed objects spawnable from sequencer. This will allow sequencer to have control over the lifetime of the object.", EUserInterfaceActionType::Button, FInputChord() );
+
+	UI_COMMAND(DiscardChanges, "Discard all changes", "Revert the currently edited movie scene to its last saved state.", EUserInterfaceActionType::Button, FInputChord());
 }
 
 #undef LOCTEXT_NAMESPACE
\ No newline at end of file
diff --git a/Engine/Source/Editor/Sequencer/Private/SequencerContextMenus.cpp b/Engine/Source/Editor/Sequencer/Private/SequencerContextMenus.cpp
--- a/Engine/Source/Editor/Sequencer/Private/SequencerContextMenus.cpp
+++ b/Engine/Source/Editor/Sequencer/Private/SequencerContextMenus.cpp
@@ -516,14 +516,17 @@ void FSectionContextMenu::SelectAllKeys()
 			continue;
 		}
 
-		FKeyAreaLayout Layout(*Handle.TrackNode, Handle.SectionIndex);
-		for (const FKeyAreaLayoutElement& Element : Layout.GetElements())
+		FSectionLayout Layout(*Handle.TrackNode, Handle.SectionIndex);
+		for (const FSectionLayoutElement& Element : Layout.GetElements())
 		{
 			TSharedPtr<IKeyArea> KeyArea = Element.GetKeyArea();
-			for (const FKeyHandle& KeyHandle : KeyArea->GetUnsortedKeyHandles())
+			if (KeyArea.IsValid())
 			{
-				FSequencerSelectedKey SelectKey(*Section, KeyArea, KeyHandle);
-				Sequencer->GetSelection().AddToSelection(SelectKey);
+				for (const FKeyHandle& KeyHandle : KeyArea->GetUnsortedKeyHandles())
+				{
+					FSequencerSelectedKey SelectKey(*Section, KeyArea, KeyHandle);
+					Sequencer->GetSelection().AddToSelection(SelectKey);
+				}
 			}
 		}
 	}
@@ -541,11 +544,11 @@ bool FSectionContextMenu::CanSelectAllKeys() const
 			continue;
 		}
 
-		FKeyAreaLayout Layout(*Handle.TrackNode, Handle.SectionIndex);
-		for (const FKeyAreaLayoutElement& Element : Layout.GetElements())
+		FSectionLayout Layout(*Handle.TrackNode, Handle.SectionIndex);
+		for (const FSectionLayoutElement& Element : Layout.GetElements())
 		{
-			TArray<FKeyHandle> KeyHandles = Element.GetKeyArea()->GetUnsortedKeyHandles();
-			if (KeyHandles.Num() > 0)
+			TSharedPtr<IKeyArea> KeyArea = Element.GetKeyArea();
+			if (KeyArea.IsValid() && Element.GetKeyArea()->GetUnsortedKeyHandles().Num() > 0)
 			{
 				return true;
 			}
@@ -560,11 +563,11 @@ bool FSectionContextMenu::CanSetExtrapolationMode() const
 	TArray<FSectionHandle> SelectedSections = StaticCastSharedRef<SSequencer>(Sequencer->GetSequencerWidget())->GetSectionHandles(Sequencer->GetSelection().GetSelectedSections());
 	for (const FSectionHandle& Handle : SelectedSections)
 	{
-		FKeyAreaLayout Layout(*Handle.TrackNode, Handle.SectionIndex);
-		for (const FKeyAreaLayoutElement& Element : Layout.GetElements())
+		FSectionLayout Layout(*Handle.TrackNode, Handle.SectionIndex);
+		for (const FSectionLayoutElement& Element : Layout.GetElements())
 		{
 			TSharedPtr<IKeyArea> KeyArea = Element.GetKeyArea();
-			if (KeyArea->CanSetExtrapolationMode())
+			if (KeyArea.IsValid() && KeyArea->CanSetExtrapolationMode())
 			{
 				return true;
 			}
@@ -621,12 +624,15 @@ void FSectionContextMenu::SetExtrapolationMode(ERichCurveExtrapolation ExtrapMod
 
 		if (Section->TryModify())
 		{
-			FKeyAreaLayout Layout(*Handle.TrackNode, Handle.SectionIndex);
-			for (const FKeyAreaLayoutElement& Element : Layout.GetElements())
+			FSectionLayout Layout(*Handle.TrackNode, Handle.SectionIndex);
+			for (const FSectionLayoutElement& Element : Layout.GetElements())
 			{
 				TSharedPtr<IKeyArea> KeyArea = Element.GetKeyArea();
-				bAnythingChanged = true;
-				KeyArea->SetExtrapolationMode(ExtrapMode, bPreInfinity);
+				if (KeyArea.IsValid())
+				{
+					bAnythingChanged = true;
+					KeyArea->SetExtrapolationMode(ExtrapMode, bPreInfinity);
+				}
 			}
 		}
 	}
@@ -650,16 +656,18 @@ bool FSectionContextMenu::IsExtrapolationModeSelected(ERichCurveExtrapolation Ex
 	TArray<FSectionHandle> SelectedSections = StaticCastSharedRef<SSequencer>(Sequencer->GetSequencerWidget())->GetSectionHandles(Sequencer->GetSelection().GetSelectedSections());
 	for (const FSectionHandle& Handle : SelectedSections)
 	{
-		FKeyAreaLayout Layout(*Handle.TrackNode, Handle.SectionIndex);
-		for (const FKeyAreaLayoutElement& Element : Layout.GetElements())
+		FSectionLayout Layout(*Handle.TrackNode, Handle.SectionIndex);
+		for (const FSectionLayoutElement& Element : Layout.GetElements())
 		{
 			TSharedPtr<IKeyArea> KeyArea = Element.GetKeyArea();
-
-			bAllSelected = true;
-			if (KeyArea->GetExtrapolationMode(bPreInfinity) != ExtrapMode)
+			if (KeyArea.IsValid())
 			{
-				bAllSelected = false;
-				break;
+				bAllSelected = true;
+				if (KeyArea->GetExtrapolationMode(bPreInfinity) != ExtrapMode)
+				{
+					bAllSelected = false;
+					break;
+				}
 			}
 		}
 	}
diff --git a/Engine/Source/Editor/Sequencer/Private/SequencerNodeTree.cpp b/Engine/Source/Editor/Sequencer/Private/SequencerNodeTree.cpp
--- a/Engine/Source/Editor/Sequencer/Private/SequencerNodeTree.cpp
+++ b/Engine/Source/Editor/Sequencer/Private/SequencerNodeTree.cpp
@@ -137,9 +137,9 @@ void FSequencerNodeTree::Update()
 	RootNodes.Reserve(RootNodes.Num()*2);
 	for (int32 Index = 0; Index < RootNodes.Num(); Index += 2)
 	{
-		RootNodes.Insert(MakeShareable(new FSequencerSpacerNode(2.f, nullptr, *this)), Index);
+		RootNodes.Insert(MakeShareable(new FSequencerSpacerNode(1.f, nullptr, *this)), Index);
 	}
-	RootNodes.Add(MakeShareable(new FSequencerSpacerNode(2.f, nullptr, *this)));
+	RootNodes.Add(MakeShareable(new FSequencerSpacerNode(1.f, nullptr, *this)));
 
 	// Set up virtual offsets, expansion states, and tints
 	float VerticalOffset = 0.f;
diff --git a/Engine/Source/Editor/Sequencer/Private/SequencerObjectChangeListener.cpp b/Engine/Source/Editor/Sequencer/Private/SequencerObjectChangeListener.cpp
--- a/Engine/Source/Editor/Sequencer/Private/SequencerObjectChangeListener.cpp
+++ b/Engine/Source/Editor/Sequencer/Private/SequencerObjectChangeListener.cpp
@@ -206,6 +206,11 @@ bool FSequencerObjectChangeListener::FindPropertySetter( const UClass& ObjectCla
 
 bool FSequencerObjectChangeListener::CanKeyProperty(FCanKeyPropertyParams CanKeyPropertyParams) const
 {
+	if (CanKeyPropertyParams.PropertyPath.Num() == 0)
+	{
+		return false;
+	}
+
 	const UStructProperty* StructProperty = Cast<const UStructProperty>(CanKeyPropertyParams.PropertyPath.Last());
 	const UStructProperty* ParentStructProperty = nullptr;
 	if (CanKeyPropertyParams.PropertyPath.Num() > 1)
diff --git a/Engine/Source/Editor/Sequencer/Private/SequencerSectionPainter.cpp b/Engine/Source/Editor/Sequencer/Private/SequencerSectionPainter.cpp
new file mode 100644
--- /dev/null
+++ b/Engine/Source/Editor/Sequencer/Private/SequencerSectionPainter.cpp
@@ -0,0 +1,31 @@
+// Copyright 1998-2015 Epic Games, Inc. All Rights Reserved.
+
+#include "SequencerPrivatePCH.h"
+#include "Sequencer.h"
+#include "SequencerSectionPainter.h"
+
+FSequencerSectionPainter::FSequencerSectionPainter(FSlateWindowElementList& OutDrawElements, const FGeometry& InSectionGeometry, UMovieSceneSection& InSection)
+	: Section(InSection)
+	, DrawElements(OutDrawElements)
+	, SectionGeometry(InSectionGeometry)
+	, LayerId(0)
+	, bParentEnabled(true)
+{
+}
+
+int32 FSequencerSectionPainter::PaintSectionBackground(const FColor& Tint)
+{
+	FSequencerSectionTintRegion Region(Tint);
+	return PaintSectionBackground(&Region, 1);
+}
+
+int32 FSequencerSectionPainter::PaintSectionBackground()
+{
+	FSequencerSectionTintRegion Region(GetTrack()->GetColorTint());
+	return PaintSectionBackground(&Region, 1);
+}
+
+UMovieSceneTrack* FSequencerSectionPainter::GetTrack() const
+{
+	return Section.GetTypedOuter<UMovieSceneTrack>();
+}
\ No newline at end of file
diff --git a/Engine/Source/Editor/Sequencer/Private/SequencerSettings.cpp b/Engine/Source/Editor/Sequencer/Private/SequencerSettings.cpp
--- a/Engine/Source/Editor/Sequencer/Private/SequencerSettings.cpp
+++ b/Engine/Source/Editor/Sequencer/Private/SequencerSettings.cpp
@@ -28,10 +28,12 @@ USequencerSettings::USequencerSettings( const FObjectInitializer& ObjectInitiali
 	bAutoScrollEnabled = false;
 	bShowCurveEditor = false;
 	bShowCurveEditorCurveToolTips = true;
+	bLinkCurveEditorTimeRange = false;
 	bLooping = false;
 	bKeepCursorInPlayRange = true;
 	bKeepPlayRangeInSectionBounds = true;
 	ZeroPadFrames = 0;
+	bShowCombinedKeyframes = true;
 }
 
 EAutoKeyMode USequencerSettings::GetAutoKeyMode() const
@@ -415,6 +417,21 @@ void USequencerSettings::SetShowCurveEditorCurveToolTips(bool InbShowCurveEditor
 }
 
 
+bool USequencerSettings::GetLinkCurveEditorTimeRange() const
+{
+	return bLinkCurveEditorTimeRange;
+}
+
+void USequencerSettings::SetLinkCurveEditorTimeRange(bool InbLinkCurveEditorTimeRange)
+{
+	if (bLinkCurveEditorTimeRange != InbLinkCurveEditorTimeRange)
+	{
+		bLinkCurveEditorTimeRange = InbLinkCurveEditorTimeRange;
+		SaveConfig();
+	}
+}
+
+
 uint8 USequencerSettings::GetZeroPadFrames() const
 {
 	return ZeroPadFrames;
@@ -429,6 +446,21 @@ void USequencerSettings::SetZeroPadFrames(uint8 InZeroPadFrames)
 	}
 }
 
+bool USequencerSettings::GetShowCombinedKeyframes() const
+{
+	return bShowCombinedKeyframes;
+}
+
+void USequencerSettings::SetShowCombinedKeyframes(bool InbShowCombinedKeyframes)
+{
+	if (bShowCombinedKeyframes != InbShowCombinedKeyframes)
+	{
+		bShowCombinedKeyframes = InbShowCombinedKeyframes;
+		SaveConfig();
+	}
+}
+
+
 float USequencerSettings::SnapTimeToInterval( float InTimeValue ) const
 {
 	return TimeSnapInterval > 0
@@ -440,14 +472,3 @@ USequencerSettings::FOnShowCurveEditorChanged& USequencerSettings::GetOnShowCurv
 {
 	return OnShowCurveEditorChanged;
 }
-
-/** Level editor specific sequencer settings */
-ULevelEditorSequencerSettings::ULevelEditorSequencerSettings( const FObjectInitializer& ObjectInitializer )
-	: Super( ObjectInitializer )
-{
-	bKeyInterpPropertiesOnly = true;
-	TimeSnapInterval = 0.033334f;
-	bShowRangeSlider = true;
-	bKeepPlayRangeInSectionBounds = false;
-	ZeroPadFrames = 4;
-}
\ No newline at end of file
diff --git a/Engine/Source/Editor/Sequencer/Private/SequencerSettings.h b/Engine/Source/Editor/Sequencer/Private/SequencerSettings.h
--- a/Engine/Source/Editor/Sequencer/Private/SequencerSettings.h
+++ b/Engine/Source/Editor/Sequencer/Private/SequencerSettings.h
@@ -4,6 +4,7 @@
 
 #include "SequencerSettings.generated.h"
 
+
 UENUM()
 enum ESequencerSpawnPosition
 {
@@ -14,9 +15,11 @@ enum ESequencerSpawnPosition
 	SSP_PlaceInFrontOfCamera UMETA(DisplayName="Place in Front of Camera"),
 };
 
+
 /** Empty class used to house multiple named USequencerSettings */
 UCLASS()
-class USequencerSettingsContainer : public UObject
+class USequencerSettingsContainer
+	: public UObject
 {
 public:
 	GENERATED_BODY()
@@ -45,9 +48,11 @@ public:
 	}
 };
 
+
 /** Serializable options for sequencer. */
 UCLASS(config=EditorPerProjectUserSettings, PerObjectConfig)
-class USequencerSettings : public UObject
+class USequencerSettings
+	: public UObject
 {
 public:
 	GENERATED_UCLASS_BODY()
@@ -176,6 +181,11 @@ public:
 	/** Sets whether or not to show curve tool tips in the curve editor. */
 	void SetShowCurveEditorCurveToolTips(bool InbShowCurveEditorCurveToolTips);
 	
+	/** Gets whether or not to link the curve editor time range. */
+	bool GetLinkCurveEditorTimeRange() const;
+	/** Sets whether or not to link the curve editor time range. */
+	void SetLinkCurveEditorTimeRange(bool InbLinkCurveEditorTimeRange);
+
 	/** Gets the loop mode. */
 	bool IsLooping() const;
 	/** Sets the loop mode. */
@@ -196,6 +206,11 @@ public:
 	/** Set the number of digits we should zero-pad to when showing frame numbers in sequencer */
 	void SetZeroPadFrames(uint8 InZeroPadFrames);
 
+	/** @return true if showing combined keyframes at the top node */
+	bool GetShowCombinedKeyframes() const;
+	/** Set whether to show combined keyframes at the top node */ 
+	void SetShowCombinedKeyframes(bool InShowCombinedKeyframes);
+
 	/** Snaps a time value in seconds to the currently selected interval. */
 	float SnapTimeToInterval(float InTimeValue) const;
 
@@ -203,6 +218,7 @@ public:
 	FOnShowCurveEditorChanged& GetOnShowCurveEditorChanged();
 
 protected:
+
 	UPROPERTY( config )
 	TEnumAsByte<EAutoKeyMode> AutoKeyMode;
 
@@ -275,6 +291,9 @@ protected:
 	UPROPERTY( config )
 	bool bShowCurveEditorCurveToolTips;
 
+	UPROPERTY( config )
+	bool bLinkCurveEditorTimeRange;
+
 	UPROPERTY( config )
 	bool bLooping;
 
@@ -287,13 +306,8 @@ protected:
 	UPROPERTY( config )
 	uint8 ZeroPadFrames;
 
-	FOnShowCurveEditorChanged OnShowCurveEditorChanged;
-};
+	UPROPERTY( config )
+	bool bShowCombinedKeyframes;
 
-/** Level editor specific sequencer settings */
-UCLASS(config=EditorPerProjectUserSettings, PerObjectConfig)
-class ULevelEditorSequencerSettings : public USequencerSettings
-{
-public:
-	GENERATED_UCLASS_BODY()
+	FOnShowCurveEditorChanged OnShowCurveEditorChanged;
 };
diff --git a/Engine/Source/Editor/Sequencer/Private/SequencerTimeSliderController.cpp b/Engine/Source/Editor/Sequencer/Private/SequencerTimeSliderController.cpp
new file mode 100644
--- /dev/null
+++ b/Engine/Source/Editor/Sequencer/Private/SequencerTimeSliderController.cpp
@@ -0,0 +1,1016 @@
+// Copyright 1998-2016 Epic Games, Inc. All Rights Reserved.
+
+#include "SequencerPrivatePCH.h"
+#include "SequencerTimeSliderController.h"
+
+#define LOCTEXT_NAMESPACE "TimeSlider"
+
+namespace ScrubConstants
+{
+	/** The minimum amount of pixels between each major ticks on the widget */
+	const int32 MinPixelsPerDisplayTick = 12;
+
+	/**The smallest number of units between between major tick marks */
+	const float MinDisplayTickSpacing = 0.001f;
+
+	/**The fraction of the current view range to scroll per unit delta  */
+	const float ScrollPanFraction = 0.1f;
+}
+
+/** Utility struct for converting between scrub range space and local/absolute screen space */
+struct FScrubRangeToScreen
+{
+	FVector2D WidgetSize;
+
+	TRange<float> ViewInput;
+	float ViewInputRange;
+	float PixelsPerInput;
+
+	FScrubRangeToScreen(TRange<float> InViewInput, const FVector2D& InWidgetSize )
+	{
+		WidgetSize = InWidgetSize;
+
+		ViewInput = InViewInput;
+		ViewInputRange = ViewInput.Size<float>();
+		PixelsPerInput = ViewInputRange > 0 ? ( WidgetSize.X / ViewInputRange ) : 0;
+	}
+
+	/** Local Widget Space -> Curve Input domain. */
+	float LocalXToInput(float ScreenX) const
+	{
+		float LocalX = ScreenX;
+		return (LocalX/PixelsPerInput) + ViewInput.GetLowerBoundValue();
+	}
+
+	/** Curve Input domain -> local Widget Space */
+	float InputToLocalX(float Input) const
+	{
+		return (Input - ViewInput.GetLowerBoundValue()) * PixelsPerInput;
+	}
+};
+
+
+/**
+ * Gets the the next spacing value in the series 
+ * to determine a good spacing value
+ * E.g, .001,.005,.010,.050,.100,.500,1.000,etc
+ */
+static float GetNextSpacing( uint32 CurrentStep )
+{
+	if(CurrentStep & 0x01) 
+	{
+		// Odd numbers
+		return FMath::Pow( 10.f, 0.5f*((float)(CurrentStep-1)) + 1.f );
+	}
+	else 
+	{
+		// Even numbers
+		return 0.5f * FMath::Pow( 10.f, 0.5f*((float)(CurrentStep)) + 1.f );
+	}
+}
+
+/**
+ * Determines the optimal spacing between tick marks in the slider for a given pixel density
+ * Increments until a minimum amount of slate units specified by MinTick is reached
+ * 
+ * @param InPixelsPerInput	The density of pixels between each input
+ * @param MinTick			The minimum slate units per tick allowed
+ * @param MinTickSpacing	The minimum tick spacing in time units allowed
+ * @return the optimal spacing in time units
+ */
+float DetermineOptimalSpacing( float InPixelsPerInput, uint32 MinTick, float MinTickSpacing )
+{
+	uint32 CurStep = 0;
+
+	// Start with the smallest spacing
+	float Spacing = MinTickSpacing;
+
+	if (InPixelsPerInput > 0)
+	{
+		while( Spacing * InPixelsPerInput < MinTick )
+		{
+			Spacing = MinTickSpacing * GetNextSpacing( CurStep );
+			CurStep++;
+		}
+	}
+
+	return Spacing;
+}
+
+
+FSequencerTimeSliderController::FSequencerTimeSliderController( const FTimeSliderArgs& InArgs )
+	: TimeSliderArgs( InArgs )
+	, DistanceDragged( 0.0f )
+	, MouseDragType( DRAG_NONE )
+	, bPanning( false )
+{
+	ScrubHandleUp = FEditorStyle::GetBrush( TEXT( "Sequencer.Timeline.ScrubHandleUp" ) ); 
+	ScrubHandleDown = FEditorStyle::GetBrush( TEXT( "Sequencer.Timeline.ScrubHandleDown" ) );
+	ScrubHandleSize = 13.f;
+	ContextMenuSuppression = 0;
+}
+
+struct FDrawTickArgs
+{
+	/** Geometry of the area */
+	FGeometry AllottedGeometry;
+	/** Clipping rect of the area */
+	FSlateRect ClippingRect;
+	/** Color of each tick */
+	FLinearColor TickColor;
+	/** Offset in Y where to start the tick */
+	float TickOffset;
+	/** Height in of major ticks */
+	float MajorTickHeight;
+	/** Start layer for elements */
+	int32 StartLayer;
+	/** Draw effects to apply */
+	ESlateDrawEffect::Type DrawEffects;
+	/** Whether or not to only draw major ticks */
+	bool bOnlyDrawMajorTicks;
+	/** Whether or not to mirror labels */
+	bool bMirrorLabels;
+	
+};
+
+void FSequencerTimeSliderController::DrawTicks( FSlateWindowElementList& OutDrawElements, const struct FScrubRangeToScreen& RangeToScreen, FDrawTickArgs& InArgs ) const
+{
+	float MinDisplayTickSpacing = ScrubConstants::MinDisplayTickSpacing;
+	if (SequencerSnapValues::IsTimeSnapIntervalFrameRate(TimeSliderArgs.Settings->GetTimeSnapInterval()) && TimeSliderArgs.Settings->GetShowFrameNumbers())
+	{
+		MinDisplayTickSpacing = TimeSliderArgs.Settings->GetTimeSnapInterval();
+	}
+
+	const float Spacing = DetermineOptimalSpacing( RangeToScreen.PixelsPerInput, ScrubConstants::MinPixelsPerDisplayTick, MinDisplayTickSpacing );
+
+	// Sub divisions
+	// @todo Sequencer may need more robust calculation
+	const int32 Divider = 10;
+	// For slightly larger halfway tick mark
+	const int32 HalfDivider = Divider / 2;
+	// Find out where to start from
+	int32 OffsetNum = FMath::FloorToInt(RangeToScreen.ViewInput.GetLowerBoundValue() / Spacing);
+	
+	FSlateFontInfo SmallLayoutFont( FPaths::EngineContentDir() / TEXT("Slate/Fonts/Roboto-Regular.ttf"), 8 );
+
+	TArray<FVector2D> LinePoints;
+	LinePoints.AddUninitialized(2);
+
+	float Seconds = 0;
+	while( (Seconds = OffsetNum*Spacing) < RangeToScreen.ViewInput.GetUpperBoundValue() )
+	{
+		// X position local to start of the widget area
+		float XPos = RangeToScreen.InputToLocalX( Seconds );
+		uint32 AbsOffsetNum = FMath::Abs(OffsetNum);
+
+		if ( AbsOffsetNum % Divider == 0 )
+		{
+			FVector2D Offset( XPos, InArgs.TickOffset );
+			FVector2D TickSize( 0.0f, InArgs.MajorTickHeight );
+
+			LinePoints[0] = FVector2D( 0.0f,1.0f);
+			LinePoints[1] = TickSize;
+
+			// lines should not need anti-aliasing
+			const bool bAntiAliasLines = false;
+
+			// Draw each tick mark
+			FSlateDrawElement::MakeLines(
+				OutDrawElements,
+				InArgs.StartLayer,
+				InArgs.AllottedGeometry.ToPaintGeometry( Offset, TickSize ),
+				LinePoints,
+				InArgs.ClippingRect,
+				InArgs.DrawEffects,
+				InArgs.TickColor,
+				bAntiAliasLines
+				);
+
+			if( !InArgs.bOnlyDrawMajorTicks )
+			{
+				FString FrameString;
+				if (SequencerSnapValues::IsTimeSnapIntervalFrameRate(TimeSliderArgs.Settings->GetTimeSnapInterval()) && TimeSliderArgs.Settings->GetShowFrameNumbers())
+				{
+					FrameString = FString::Printf( TEXT("%d"), TimeToFrame(Seconds));
+				}
+				else
+				{
+					FrameString = Spacing == ScrubConstants::MinDisplayTickSpacing ? FString::Printf( TEXT("%.3f"), Seconds ) : FString::Printf( TEXT("%.2f"), Seconds );
+				}
+
+				// Space the text between the tick mark but slightly above
+				const TSharedRef< FSlateFontMeasure > FontMeasureService = FSlateApplication::Get().GetRenderer()->GetFontMeasureService();
+				FVector2D TextSize = FontMeasureService->Measure(FrameString, SmallLayoutFont);
+				FVector2D TextOffset( XPos + 5.f, InArgs.bMirrorLabels ? 3.f : FMath::Abs( InArgs.AllottedGeometry.Size.Y - (InArgs.MajorTickHeight+3.f) ) );
+				FSlateDrawElement::MakeText(
+					OutDrawElements,
+					InArgs.StartLayer+1, 
+					InArgs.AllottedGeometry.ToPaintGeometry( TextOffset, TextSize ), 
+					FrameString, 
+					SmallLayoutFont, 
+					InArgs.ClippingRect, 
+					InArgs.DrawEffects,
+					InArgs.TickColor*0.65f 
+				);
+			}
+		}
+		else if( !InArgs.bOnlyDrawMajorTicks )
+		{
+			// Compute the size of each tick mark.  If we are half way between to visible values display a slightly larger tick mark
+			const float MinorTickHeight = AbsOffsetNum % HalfDivider == 0 ? 6.0f : 2.0f;
+
+			FVector2D Offset(XPos, InArgs.bMirrorLabels ? 0.0f : FMath::Abs( InArgs.AllottedGeometry.Size.Y - MinorTickHeight ) );
+			FVector2D TickSize(0.0f, MinorTickHeight);
+
+			LinePoints[0] = FVector2D(0.0f,1.0f);
+			LinePoints[1] = TickSize;
+
+			const bool bAntiAlias = false;
+			// Draw each sub mark
+			FSlateDrawElement::MakeLines(
+				OutDrawElements,
+				InArgs.StartLayer,
+				InArgs.AllottedGeometry.ToPaintGeometry( Offset, TickSize ),
+				LinePoints,
+				InArgs.ClippingRect,
+				InArgs.DrawEffects,
+				InArgs.TickColor,
+				bAntiAlias
+			);
+		}
+		// Advance to next tick mark
+		++OffsetNum;
+	}
+}
+
+
+int32 FSequencerTimeSliderController::OnPaintTimeSlider( bool bMirrorLabels, const FGeometry& AllottedGeometry, const FSlateRect& MyClippingRect, FSlateWindowElementList& OutDrawElements, int32 LayerId, const FWidgetStyle& InWidgetStyle, bool bParentEnabled ) const
+{
+	const bool bEnabled = bParentEnabled;
+	const ESlateDrawEffect::Type DrawEffects = bEnabled ? ESlateDrawEffect::None : ESlateDrawEffect::DisabledEffect;
+
+	TRange<float> LocalViewRange = TimeSliderArgs.ViewRange.Get();
+	const float LocalViewRangeMin = LocalViewRange.GetLowerBoundValue();
+	const float LocalViewRangeMax = LocalViewRange.GetUpperBoundValue();
+	const float LocalSequenceLength = LocalViewRangeMax-LocalViewRangeMin;
+	
+	FVector2D Scale = FVector2D(1.0f,1.0f);
+	if ( LocalSequenceLength > 0)
+	{
+		FScrubRangeToScreen RangeToScreen( LocalViewRange, AllottedGeometry.Size );
+	
+		const float MajorTickHeight = 9.0f;
+	
+		FDrawTickArgs Args;
+		Args.AllottedGeometry = AllottedGeometry;
+		Args.bMirrorLabels = bMirrorLabels;
+		Args.bOnlyDrawMajorTicks = false;
+		Args.TickColor = FLinearColor::White;
+		Args.ClippingRect = MyClippingRect;
+		Args.DrawEffects = DrawEffects;
+		Args.StartLayer = LayerId;
+		Args.TickOffset = bMirrorLabels ? 0.0f : FMath::Abs( AllottedGeometry.Size.Y - MajorTickHeight );
+		Args.MajorTickHeight = MajorTickHeight;
+
+		DrawTicks( OutDrawElements, RangeToScreen, Args );
+
+		FPaintPlaybackRangeArgs PlaybackRangeArgs(
+			bMirrorLabels ? FEditorStyle::GetBrush("Sequencer.Timeline.PlayRange_Bottom_L") : FEditorStyle::GetBrush("Sequencer.Timeline.PlayRange_Top_L"),
+			bMirrorLabels ? FEditorStyle::GetBrush("Sequencer.Timeline.PlayRange_Bottom_R") : FEditorStyle::GetBrush("Sequencer.Timeline.PlayRange_Top_R"),
+			6.f
+		);
+		LayerId = DrawPlaybackRange(AllottedGeometry, MyClippingRect, OutDrawElements, LayerId, RangeToScreen, PlaybackRangeArgs);
+
+		float HalfSize = FMath::CeilToFloat(ScrubHandleSize/2.0f);
+
+		// Draw the scrub handle
+		float XPos = RangeToScreen.InputToLocalX( TimeSliderArgs.ScrubPosition.Get() );
+
+		// Should draw above the text
+		const int32 ArrowLayer = LayerId + 2;
+		FPaintGeometry MyGeometry =	AllottedGeometry.ToPaintGeometry( FVector2D( XPos-HalfSize, 0 ), FVector2D( ScrubHandleSize, AllottedGeometry.Size.Y ) );
+		FLinearColor ScrubColor = InWidgetStyle.GetColorAndOpacityTint();
+
+		// @todo Sequencer this color should be specified in the style
+		ScrubColor.A = ScrubColor.A*0.75f;
+		ScrubColor.B *= 0.1f;
+		ScrubColor.G *= 0.2f;
+		FSlateDrawElement::MakeBox( 
+			OutDrawElements,
+			ArrowLayer, 
+			MyGeometry,
+			bMirrorLabels ? ScrubHandleUp : ScrubHandleDown,
+			MyClippingRect, 
+			DrawEffects, 
+			ScrubColor
+			);
+
+		// Draw the current time next to the scrub handle
+		float Time = TimeSliderArgs.ScrubPosition.Get();
+
+		FString FrameString;
+		if (SequencerSnapValues::IsTimeSnapIntervalFrameRate(TimeSliderArgs.Settings->GetTimeSnapInterval()) && TimeSliderArgs.Settings->GetShowFrameNumbers())
+		{
+			float FrameRate = 1.0f/TimeSliderArgs.Settings->GetTimeSnapInterval();
+			float FrameTime = Time * FrameRate;
+			int32 Frame = SequencerHelpers::TimeToFrame(Time, FrameRate);
+					
+			const float FrameTolerance = 0.001f;
+			if (FMath::IsNearlyEqual(FrameTime, (float)Frame, FrameTolerance))
+			{
+				FrameString = FString::Printf( TEXT("%d"), TimeToFrame(Time));
+			}
+			else
+			{
+				FrameString = FString::Printf( TEXT("%.3f"), FrameTime);
+			}
+		}
+		else
+		{
+			FrameString = FString::Printf( TEXT("%.2f"), Time );
+		}
+
+		FSlateFontInfo SmallLayoutFont( FPaths::EngineContentDir() / TEXT("Slate/Fonts/Roboto-Regular.ttf"), 10 );
+
+		const TSharedRef< FSlateFontMeasure > FontMeasureService = FSlateApplication::Get().GetRenderer()->GetFontMeasureService();
+		FVector2D TextSize = FontMeasureService->Measure(FrameString, SmallLayoutFont);
+
+		// Flip the text position if getting near the end of the view range
+		if ((AllottedGeometry.Size.X - XPos) < (TextSize.X + 14.f))
+		{
+			XPos = XPos - TextSize.X - 12.f;
+		}
+		else
+		{
+			XPos = XPos + 10.f;
+		}
+
+		FVector2D TextOffset( XPos,  Args.bMirrorLabels ? TextSize.Y-6.f : Args.AllottedGeometry.Size.Y - (Args.MajorTickHeight+TextSize.Y) );
+
+		FSlateDrawElement::MakeText(
+			OutDrawElements,
+			Args.StartLayer+1, 
+			Args.AllottedGeometry.ToPaintGeometry( TextOffset, TextSize ), 
+			FrameString, 
+			SmallLayoutFont, 
+			Args.ClippingRect, 
+			Args.DrawEffects,
+			Args.TickColor 
+		);
+		
+		if (MouseDragType == DRAG_SETTING_RANGE)
+		{
+			float MouseStartPosX = RangeToScreen.InputToLocalX(MouseDownRange[0]);
+			float MouseEndPosX = RangeToScreen.InputToLocalX(MouseDownRange[1]);
+
+			float RangePosX = MouseStartPosX < MouseEndPosX ? MouseStartPosX : MouseEndPosX;
+			float RangeSizeX = FMath::Abs(MouseStartPosX - MouseEndPosX);
+
+			FSlateDrawElement::MakeBox(
+				OutDrawElements,
+				LayerId+1,
+				AllottedGeometry.ToPaintGeometry( FVector2D(RangePosX, 0.f), FVector2D(RangeSizeX, AllottedGeometry.Size.Y) ),
+				bMirrorLabels ? ScrubHandleDown : ScrubHandleUp,
+				MyClippingRect,
+				DrawEffects,
+				MouseStartPosX < MouseEndPosX ? FLinearColor(0.5f, 0.5f, 0.5f) : FLinearColor(0.25f, 0.3f, 0.3f)
+			);
+		}
+
+		return ArrowLayer;
+	}
+
+	return LayerId;
+}
+
+int32 FSequencerTimeSliderController::DrawPlaybackRange(const FGeometry& AllottedGeometry, const FSlateRect& MyClippingRect, FSlateWindowElementList& OutDrawElements, int32 LayerId, const FScrubRangeToScreen& RangeToScreen, const FPaintPlaybackRangeArgs& Args) const
+{
+	if (!TimeSliderArgs.PlaybackRange.IsSet())
+	{
+		return LayerId;
+	}
+
+	TRange<float> PlaybackRange = TimeSliderArgs.PlaybackRange.Get();
+
+	float PlaybackRangeL = RangeToScreen.InputToLocalX(PlaybackRange.GetLowerBoundValue()) - 1;
+	float PlaybackRangeR = RangeToScreen.InputToLocalX(PlaybackRange.GetUpperBoundValue()) + 1;
+
+	FSlateDrawElement::MakeBox(
+		OutDrawElements,
+		LayerId+1,
+		AllottedGeometry.ToPaintGeometry(FVector2D(PlaybackRangeL, 0.f), FVector2D(Args.BrushWidth, AllottedGeometry.Size.Y)),
+		Args.StartBrush,
+		MyClippingRect,
+		ESlateDrawEffect::None,
+		FColor(32, 128, 32)	// 120, 75, 50 (HSV)
+	);
+
+	FSlateDrawElement::MakeBox(
+		OutDrawElements,
+		LayerId+1,
+		AllottedGeometry.ToPaintGeometry(FVector2D(PlaybackRangeR - Args.BrushWidth, 0.f), FVector2D(Args.BrushWidth, AllottedGeometry.Size.Y)),
+		Args.EndBrush,
+		MyClippingRect,
+		ESlateDrawEffect::None,
+		FColor(128, 32, 32)	// 0, 75, 50 (HSV)
+	);
+
+	// Black tint for excluded regions
+	FSlateDrawElement::MakeBox(
+		OutDrawElements,
+		LayerId+1,
+		AllottedGeometry.ToPaintGeometry(FVector2D(0.f, 0.f), FVector2D(PlaybackRangeL, AllottedGeometry.Size.Y)),
+		FEditorStyle::GetBrush("WhiteBrush"),
+		MyClippingRect,
+		ESlateDrawEffect::None,
+		FLinearColor::Black.CopyWithNewOpacity(0.3f)
+	);
+
+	FSlateDrawElement::MakeBox(
+		OutDrawElements,
+		LayerId+1,
+		AllottedGeometry.ToPaintGeometry(FVector2D(PlaybackRangeR, 0.f), FVector2D(AllottedGeometry.Size.X - PlaybackRangeR, AllottedGeometry.Size.Y)),
+		FEditorStyle::GetBrush("WhiteBrush"),
+		MyClippingRect,
+		ESlateDrawEffect::None,
+		FLinearColor::Black.CopyWithNewOpacity(0.3f)
+	);
+	return LayerId + 1;
+}
+
+FReply FSequencerTimeSliderController::OnMouseButtonDown( SWidget& WidgetOwner, const FGeometry& MyGeometry, const FPointerEvent& MouseEvent )
+{
+	bool bHandleLeftMouseButton = MouseEvent.GetEffectingButton() == EKeys::LeftMouseButton;
+	bool bHandleRightMouseButton = MouseEvent.GetEffectingButton() == EKeys::RightMouseButton && TimeSliderArgs.AllowZoom;
+	
+	DistanceDragged = 0;
+
+	FScrubRangeToScreen RangeToScreen( TimeSliderArgs.ViewRange.Get(), MyGeometry.Size );
+	MouseDownRange[0] = RangeToScreen.LocalXToInput(MyGeometry.AbsoluteToLocal(MouseEvent.GetLastScreenSpacePosition()).X);
+	MouseDownRange[1] = MouseDownRange[0];
+
+	if ( bHandleLeftMouseButton )
+	{
+		return FReply::Handled().CaptureMouse( WidgetOwner.AsShared() ).PreventThrottling();
+	}
+	else if ( bHandleRightMouseButton )
+	{
+		// Always capture mouse if we left or right click on the widget
+		return FReply::Handled().CaptureMouse( WidgetOwner.AsShared() );
+	}
+
+	return FReply::Unhandled();
+}
+
+FReply FSequencerTimeSliderController::OnMouseButtonUp( SWidget& WidgetOwner, const FGeometry& MyGeometry, const FPointerEvent& MouseEvent )
+{
+	bool bHandleLeftMouseButton = MouseEvent.GetEffectingButton() == EKeys::LeftMouseButton && WidgetOwner.HasMouseCapture();
+	bool bHandleRightMouseButton = MouseEvent.GetEffectingButton() == EKeys::RightMouseButton && WidgetOwner.HasMouseCapture() && TimeSliderArgs.AllowZoom ;
+
+	if ( bHandleRightMouseButton )
+	{
+		if (!bPanning)
+		{
+			// Open a context menu if allowed
+			if (ContextMenuSuppression == 0 && TimeSliderArgs.PlaybackRange.IsSet())
+			{
+				FScrubRangeToScreen RangeToScreen( TimeSliderArgs.ViewRange.Get(), MyGeometry.Size );
+				FVector2D CursorPos = MyGeometry.AbsoluteToLocal( MouseEvent.GetScreenSpacePosition() );
+
+				float MouseValue = RangeToScreen.LocalXToInput( CursorPos.X );
+				if (TimeSliderArgs.Settings->GetIsSnapEnabled())
+				{
+					MouseValue = TimeSliderArgs.Settings->SnapTimeToInterval(MouseValue);
+				}
+
+				TSharedRef<SWidget> MenuContent = OpenSetPlaybackRangeMenu(MouseValue);
+				FSlateApplication::Get().PushMenu(
+					WidgetOwner.AsShared(),
+					MouseEvent.GetEventPath() != nullptr ? *MouseEvent.GetEventPath() : FWidgetPath(),
+					MenuContent,
+					MouseEvent.GetScreenSpacePosition(),
+					FPopupTransitionEffect( FPopupTransitionEffect::ContextMenu )
+					);
+
+				return FReply::Handled().SetUserFocus(MenuContent, EFocusCause::SetDirectly).ReleaseMouseCapture();
+			}
+
+			// return unhandled in case our parent wants to use our right mouse button to open a context menu
+			return FReply::Unhandled().ReleaseMouseCapture();
+		}
+		
+		bPanning = false;
+		
+		return FReply::Handled().ReleaseMouseCapture();
+	}
+	else if ( bHandleLeftMouseButton )
+	{
+		if (MouseDragType == DRAG_START_RANGE)
+		{
+			TimeSliderArgs.OnEndPlaybackRangeDrag.ExecuteIfBound();
+		}
+		// Set the end range time?
+		else if (MouseDragType == DRAG_END_RANGE)
+		{
+			TimeSliderArgs.OnEndPlaybackRangeDrag.ExecuteIfBound();
+		}
+		else if (MouseDragType == DRAG_SETTING_RANGE)
+		{
+			FScrubRangeToScreen RangeToScreen( TimeSliderArgs.ViewRange.Get(), MyGeometry.Size );
+			FVector2D CursorPos = MyGeometry.AbsoluteToLocal(MouseEvent.GetLastScreenSpacePosition());
+			float NewValue = RangeToScreen.LocalXToInput(CursorPos.X);
+
+			if ( TimeSliderArgs.Settings->GetIsSnapEnabled() )
+			{
+				NewValue = TimeSliderArgs.Settings->SnapTimeToInterval(NewValue);
+			}
+
+			float DownValue = MouseDownRange[0];
+				
+			if ( TimeSliderArgs.Settings->GetIsSnapEnabled() )
+			{
+				DownValue = TimeSliderArgs.Settings->SnapTimeToInterval(DownValue);
+			}
+
+			// Zoom in
+			bool bDoSetRange = false;
+			if (NewValue > DownValue)
+			{
+				// push the current value onto the stack
+				RangeStack.Add(FVector2D(TimeSliderArgs.ViewRange.Get().GetLowerBoundValue(), TimeSliderArgs.ViewRange.Get().GetUpperBoundValue()));
+				bDoSetRange = true;
+			}
+			// Zoom out
+			else if (RangeStack.Num())
+			{
+				// pop the stack
+				FVector2D LastRange = RangeStack.Pop();
+				DownValue = LastRange[0];
+				NewValue = LastRange[1];
+				bDoSetRange = true;
+			}
+
+			if (bDoSetRange)
+			{
+				TimeSliderArgs.OnViewRangeChanged.ExecuteIfBound(TRange<float>(DownValue, NewValue), EViewRangeInterpolation::Immediate);
+					
+				if( !TimeSliderArgs.ViewRange.IsBound() )
+				{	
+					// The output is not bound to a delegate so we'll manage the value ourselves
+					TimeSliderArgs.ViewRange.Set( TRange<float>( DownValue, NewValue ) );
+				}
+			}
+		}
+		else
+		{
+			TimeSliderArgs.OnEndScrubberMovement.ExecuteIfBound();
+
+			FScrubRangeToScreen RangeToScreen( TimeSliderArgs.ViewRange.Get(), MyGeometry.Size );
+			FVector2D CursorPos = MyGeometry.AbsoluteToLocal(MouseEvent.GetLastScreenSpacePosition());
+			float NewValue = RangeToScreen.LocalXToInput(CursorPos.X);
+
+			if ( TimeSliderArgs.Settings->GetIsSnapEnabled() && TimeSliderArgs.Settings->GetSnapPlayTimeToInterval() )
+			{
+				NewValue = TimeSliderArgs.Settings->SnapTimeToInterval(NewValue);
+			}
+			
+			CommitScrubPosition( NewValue, /*bIsScrubbing=*/false );
+		}
+
+		MouseDragType = DRAG_NONE;
+		return FReply::Handled().ReleaseMouseCapture();
+
+	}
+
+	return FReply::Unhandled();
+}
+
+FReply FSequencerTimeSliderController::OnMouseMove( SWidget& WidgetOwner, const FGeometry& MyGeometry, const FPointerEvent& MouseEvent )
+{
+	if ( WidgetOwner.HasMouseCapture() )
+	{
+		if (MouseEvent.IsMouseButtonDown( EKeys::RightMouseButton ))
+		{
+			if (!bPanning)
+			{
+				DistanceDragged += FMath::Abs( MouseEvent.GetCursorDelta().X );
+				if ( DistanceDragged > FSlateApplication::Get().GetDragTriggerDistance() )
+				{
+					bPanning = true;
+				}
+			}
+			else
+			{
+				TRange<float> LocalViewRange = TimeSliderArgs.ViewRange.Get();
+				float LocalViewRangeMin = LocalViewRange.GetLowerBoundValue();
+				float LocalViewRangeMax = LocalViewRange.GetUpperBoundValue();
+
+				FScrubRangeToScreen ScaleInfo( LocalViewRange, MyGeometry.Size );
+				FVector2D ScreenDelta = MouseEvent.GetCursorDelta();
+				FVector2D InputDelta;
+				InputDelta.X = ScreenDelta.X/ScaleInfo.PixelsPerInput;
+
+				float NewViewOutputMin = LocalViewRangeMin - InputDelta.X;
+				float NewViewOutputMax = LocalViewRangeMax - InputDelta.X;
+
+				ClampViewRange(NewViewOutputMin, NewViewOutputMax);
+				SetViewRange(NewViewOutputMin, NewViewOutputMax, EViewRangeInterpolation::Immediate);
+			}
+		}
+		else if (MouseEvent.IsMouseButtonDown( EKeys::LeftMouseButton ))
+		{
+			TRange<float> LocalViewRange = TimeSliderArgs.ViewRange.Get();
+			DistanceDragged += FMath::Abs( MouseEvent.GetCursorDelta().X );
+
+			if ( MouseDragType == DRAG_NONE )
+			{
+				if ( DistanceDragged > FSlateApplication::Get().GetDragTriggerDistance() )
+				{
+					FScrubRangeToScreen RangeToScreen(LocalViewRange, MyGeometry.Size);
+					const float ScrubPosition = TimeSliderArgs.ScrubPosition.Get();
+
+					TRange<float> PlaybackRange = TimeSliderArgs.PlaybackRange.Get();
+					float LocalMouseDownPos = RangeToScreen.InputToLocalX(MouseDownRange[0]);
+
+					// Favor dragging the end position
+					if (HitTestPlaybackEnd(RangeToScreen, PlaybackRange, LocalMouseDownPos, ScrubPosition))
+					{
+						MouseDragType = DRAG_END_RANGE;
+						TimeSliderArgs.OnBeginPlaybackRangeDrag.ExecuteIfBound();
+					}
+					else if (HitTestPlaybackStart(RangeToScreen, PlaybackRange, LocalMouseDownPos, ScrubPosition))
+					{
+						MouseDragType = DRAG_START_RANGE;
+						TimeSliderArgs.OnBeginPlaybackRangeDrag.ExecuteIfBound();
+					}
+					else if (FSlateApplication::Get().GetModifierKeys().AreModifersDown(EModifierKey::Control))
+					{
+						MouseDragType = DRAG_SETTING_RANGE;
+					}
+					else
+					{
+						MouseDragType = DRAG_SCRUBBING_TIME;
+						TimeSliderArgs.OnBeginScrubberMovement.ExecuteIfBound();
+					}
+				}
+			}
+			else
+			{
+				FScrubRangeToScreen RangeToScreen( TimeSliderArgs.ViewRange.Get(), MyGeometry.Size );
+				FVector2D CursorPos = MyGeometry.AbsoluteToLocal( MouseEvent.GetLastScreenSpacePosition() );
+				float NewValue = RangeToScreen.LocalXToInput( CursorPos.X );
+
+
+				// Set the start range time?
+				if (MouseDragType == DRAG_START_RANGE)
+				{
+					if (TimeSliderArgs.Settings->GetIsSnapEnabled())
+					{
+						NewValue = TimeSliderArgs.Settings->SnapTimeToInterval(NewValue);
+					}
+
+					SetPlaybackRangeStart(NewValue);
+				}
+				// Set the end range time?
+				else if(MouseDragType == DRAG_END_RANGE)
+				{
+					if (TimeSliderArgs.Settings->GetIsSnapEnabled())
+					{
+						NewValue = TimeSliderArgs.Settings->SnapTimeToInterval(NewValue);
+					}
+					
+					SetPlaybackRangeEnd(NewValue);
+				}
+				else if (MouseDragType == DRAG_SCRUBBING_TIME)
+				{
+					if ( TimeSliderArgs.Settings->GetIsSnapEnabled() && TimeSliderArgs.Settings->GetSnapPlayTimeToInterval() )
+					{
+						NewValue = TimeSliderArgs.Settings->SnapTimeToInterval(NewValue);
+					}
+					
+					// Delegate responsibility for clamping to the current viewrange to the client
+					CommitScrubPosition( NewValue, /*bIsScrubbing=*/true );
+				}
+				else if (MouseDragType == DRAG_SETTING_RANGE)
+				{
+					MouseDownRange[1] = NewValue;
+				}
+			}
+		}
+		return FReply::Handled();
+	}
+
+	return FReply::Unhandled();
+}
+
+void FSequencerTimeSliderController::CommitScrubPosition( float NewValue, bool bIsScrubbing )
+{
+	// Manage the scrub position ourselves if its not bound to a delegate
+	if ( !TimeSliderArgs.ScrubPosition.IsBound() )
+	{
+		TimeSliderArgs.ScrubPosition.Set( NewValue );
+	}
+
+	TimeSliderArgs.OnScrubPositionChanged.ExecuteIfBound( NewValue, bIsScrubbing );
+}
+
+FReply FSequencerTimeSliderController::OnMouseWheel( SWidget& WidgetOwner, const FGeometry& MyGeometry, const FPointerEvent& MouseEvent )
+{
+	TOptional<TRange<float>> NewTargetRange;
+
+	float MouseFractionX = MyGeometry.AbsoluteToLocal(MouseEvent.GetScreenSpacePosition()).X / MyGeometry.GetLocalSize().X;
+	if ( TimeSliderArgs.AllowZoom && MouseEvent.IsControlDown() )
+	{
+		const float ZoomDelta = -0.2f * MouseEvent.GetWheelDelta();
+		if (ZoomByDelta(ZoomDelta, MouseFractionX))
+		{
+			return FReply::Handled();
+		}
+	}
+	else if (MouseEvent.IsShiftDown())
+	{
+		PanByDelta(-MouseEvent.GetWheelDelta());
+		return FReply::Handled();
+	}
+	
+	return FReply::Unhandled();
+}
+
+FCursorReply FSequencerTimeSliderController::OnCursorQuery( TSharedRef<const SWidget> WidgetOwner, const FGeometry& MyGeometry, const FPointerEvent& CursorEvent ) const
+{
+	FScrubRangeToScreen RangeToScreen(TimeSliderArgs.ViewRange.Get(), MyGeometry.Size);
+	TRange<float> PlaybackRange = TimeSliderArgs.PlaybackRange.Get();
+	const float ScrubPosition = TimeSliderArgs.ScrubPosition.Get();
+
+	// Use L/R resize cursor if we're dragging or hovering a playback range bound
+	float HitTestPosition = MyGeometry.AbsoluteToLocal(CursorEvent.GetScreenSpacePosition()).X;
+	if (MouseDragType == DRAG_END_RANGE ||
+		MouseDragType == DRAG_START_RANGE ||
+		HitTestPlaybackStart(RangeToScreen, PlaybackRange, HitTestPosition, ScrubPosition) ||
+		HitTestPlaybackEnd(RangeToScreen, PlaybackRange, HitTestPosition, ScrubPosition)
+		)
+	{
+		return FCursorReply::Cursor(EMouseCursor::ResizeLeftRight);
+	}
+
+	return FCursorReply::Unhandled();
+}
+
+int32 FSequencerTimeSliderController::OnPaintSectionView( const FGeometry& AllottedGeometry, const FSlateRect& MyClippingRect, FSlateWindowElementList& OutDrawElements, int32 LayerId, bool bEnabled, const FPaintSectionAreaViewArgs& Args ) const
+{
+	const ESlateDrawEffect::Type DrawEffects = bEnabled ? ESlateDrawEffect::None : ESlateDrawEffect::DisabledEffect;
+
+	TRange<float> LocalViewRange = TimeSliderArgs.ViewRange.Get();
+	float LocalScrubPosition = TimeSliderArgs.ScrubPosition.Get();
+
+	float ViewRange = LocalViewRange.Size<float>();
+	float PixelsPerInput = ViewRange > 0 ? AllottedGeometry.Size.X / ViewRange : 0;
+	float LinePos =  (LocalScrubPosition - LocalViewRange.GetLowerBoundValue()) * PixelsPerInput;
+
+	FScrubRangeToScreen RangeToScreen( LocalViewRange, AllottedGeometry.Size );
+
+	if (Args.PlaybackRangeArgs.IsSet())
+	{
+		LayerId = DrawPlaybackRange(AllottedGeometry, MyClippingRect, OutDrawElements, LayerId, RangeToScreen, Args.PlaybackRangeArgs.GetValue());
+	}
+
+	if( Args.bDisplayTickLines )
+	{
+		static FLinearColor TickColor(0.f, 0.f, 0.f, 0.3f);
+		
+		// Draw major tick lines in the section area
+		FDrawTickArgs DrawTickArgs;
+		DrawTickArgs.AllottedGeometry = AllottedGeometry;
+		DrawTickArgs.bMirrorLabels = false;
+		DrawTickArgs.bOnlyDrawMajorTicks = true;
+		DrawTickArgs.TickColor = TickColor;
+		DrawTickArgs.ClippingRect = MyClippingRect;
+		DrawTickArgs.DrawEffects = DrawEffects;
+		// Draw major ticks under sections
+		DrawTickArgs.StartLayer = LayerId-1;
+		// Draw the tick the entire height of the section area
+		DrawTickArgs.TickOffset = 0.0f;
+		DrawTickArgs.MajorTickHeight = AllottedGeometry.Size.Y;
+
+		DrawTicks( OutDrawElements, RangeToScreen, DrawTickArgs );
+	}
+
+	if( Args.bDisplayScrubPosition )
+	{
+		// Draw a line for the scrub position
+		TArray<FVector2D> LinePoints;
+		LinePoints.AddUninitialized(2);
+		LinePoints[0] = FVector2D( 0.0f, 0.0f );
+		LinePoints[1] = FVector2D( 0.0f, FMath::RoundToFloat( AllottedGeometry.Size.Y ) );
+
+		FSlateDrawElement::MakeLines(
+			OutDrawElements,
+			LayerId+1,
+			AllottedGeometry.ToPaintGeometry( FVector2D(LinePos, 0.0f ), FVector2D(1.0f,1.0f) ),
+			LinePoints,
+			MyClippingRect,
+			DrawEffects,
+			FLinearColor::White,
+			false
+			);
+	}
+
+	return LayerId;
+}
+
+TSharedRef<SWidget> FSequencerTimeSliderController::OpenSetPlaybackRangeMenu(float MouseTime)
+{
+	const bool bShouldCloseWindowAfterMenuSelection = true;
+	FMenuBuilder MenuBuilder(bShouldCloseWindowAfterMenuSelection, nullptr);
+
+	FText NumericText;
+	if (SequencerSnapValues::IsTimeSnapIntervalFrameRate(TimeSliderArgs.Settings->GetTimeSnapInterval()) && TimeSliderArgs.Settings->GetShowFrameNumbers())
+	{
+		NumericText = FText::Format(LOCTEXT("FrameTextFormat", "Playback Range (at frame {0}):"), FText::AsNumber(TimeToFrame(MouseTime)));
+	}
+	else
+	{
+		NumericText = FText::Format(LOCTEXT("TimeTextFormat", "Playback Range (at {0}s):"), FText::AsNumber(MouseTime));
+	}
+
+	TRange<float> PlaybackRange = TimeSliderArgs.PlaybackRange.Get();
+
+	MenuBuilder.BeginSection("SequencerPlaybackRangeMenu", NumericText);
+	{
+		MenuBuilder.AddMenuEntry(
+			FText::Format(LOCTEXT("SetPlaybackStart", "Set Start Time"), NumericText),
+			FText(),
+			FSlateIcon(),
+			FUIAction(
+				FExecuteAction::CreateLambda([=]{ return SetPlaybackRangeStart(MouseTime); }),
+				FCanExecuteAction::CreateLambda([=]{ return MouseTime <= PlaybackRange.GetUpperBoundValue(); })
+			)
+		);
+
+		MenuBuilder.AddMenuEntry(
+			FText::Format(LOCTEXT("SetPlaybackEnd", "Set End Time"), NumericText),
+			FText(),
+			FSlateIcon(),
+			FUIAction(
+				FExecuteAction::CreateLambda([=]{ return SetPlaybackRangeEnd(MouseTime); }),
+				FCanExecuteAction::CreateLambda([=]{ return MouseTime >= PlaybackRange.GetLowerBoundValue(); })
+			)
+		);
+	}
+	MenuBuilder.EndSection(); // SequencerPlaybackRangeMenu
+
+	return MenuBuilder.MakeWidget();
+}
+
+int32 FSequencerTimeSliderController::TimeToFrame(float Time) const
+{
+	float FrameRate = 1.0f/TimeSliderArgs.Settings->GetTimeSnapInterval();
+	return SequencerHelpers::TimeToFrame(Time, FrameRate);
+}
+
+float FSequencerTimeSliderController::FrameToTime(int32 Frame) const
+{
+	float FrameRate = 1.0f/TimeSliderArgs.Settings->GetTimeSnapInterval();
+	return SequencerHelpers::FrameToTime(Frame, FrameRate);
+}
+
+void FSequencerTimeSliderController::ClampViewRange(float& NewRangeMin, float& NewRangeMax)
+{
+	bool bNeedsClampSet = false;
+	float NewClampRangeMin = TimeSliderArgs.ClampRange.Get().GetLowerBoundValue();
+	if ( NewRangeMin < TimeSliderArgs.ClampRange.Get().GetLowerBoundValue() )
+	{
+		NewClampRangeMin = NewRangeMin;
+		bNeedsClampSet = true;
+	}
+
+	float NewClampRangeMax = TimeSliderArgs.ClampRange.Get().GetUpperBoundValue();
+	if ( NewRangeMax > TimeSliderArgs.ClampRange.Get().GetUpperBoundValue() )
+	{
+		NewClampRangeMax = NewRangeMax;
+		bNeedsClampSet = true;
+	}
+
+	if (bNeedsClampSet)
+	{
+		SetClampRange(NewClampRangeMin, NewClampRangeMax);
+	}
+}
+
+void FSequencerTimeSliderController::SetViewRange( float NewRangeMin, float NewRangeMax, EViewRangeInterpolation Interpolation )
+{
+	const TRange<float> NewRange(NewRangeMin, NewRangeMax);
+
+	TimeSliderArgs.OnViewRangeChanged.ExecuteIfBound( NewRange, Interpolation );
+
+	if( !TimeSliderArgs.ViewRange.IsBound() )
+	{	
+		// The  output is not bound to a delegate so we'll manage the value ourselves (no animation)
+		TimeSliderArgs.ViewRange.Set( NewRange );
+	}
+}
+
+void FSequencerTimeSliderController::SetClampRange( float NewRangeMin, float NewRangeMax )
+{
+	const TRange<float> NewRange(NewRangeMin, NewRangeMax);
+
+	TimeSliderArgs.OnClampRangeChanged.ExecuteIfBound(NewRange);
+
+	if( !TimeSliderArgs.ClampRange.IsBound() )
+	{	
+		// The  output is not bound to a delegate so we'll manage the value ourselves (no animation)
+		TimeSliderArgs.ClampRange.Set(NewRange);
+	}
+}
+
+bool FSequencerTimeSliderController::ZoomByDelta( float InDelta, float MousePositionFraction )
+{
+	TRange<float> LocalViewRange = TimeSliderArgs.ViewRange.Get().GetAnimationTarget();
+	float LocalViewRangeMax = LocalViewRange.GetUpperBoundValue();
+	float LocalViewRangeMin = LocalViewRange.GetLowerBoundValue();
+	const float OutputViewSize = LocalViewRangeMax - LocalViewRangeMin;
+	const float OutputChange = OutputViewSize * InDelta;
+
+	float NewViewOutputMin = LocalViewRangeMin - (OutputChange * MousePositionFraction);
+	float NewViewOutputMax = LocalViewRangeMax + (OutputChange * (1.f - MousePositionFraction));
+
+	if( NewViewOutputMin < NewViewOutputMax )
+	{
+		ClampViewRange(NewViewOutputMin, NewViewOutputMax);
+		SetViewRange(NewViewOutputMin, NewViewOutputMax, EViewRangeInterpolation::Animated);
+		return true;
+	}
+
+	return false;
+}
+
+void FSequencerTimeSliderController::PanByDelta( float InDelta )
+{
+	TRange<float> LocalViewRange = TimeSliderArgs.ViewRange.Get().GetAnimationTarget();
+
+	float CurrentMin = LocalViewRange.GetLowerBoundValue();
+	float CurrentMax = LocalViewRange.GetUpperBoundValue();
+
+	// Adjust the delta to be a percentage of the current range
+	InDelta *= ScrubConstants::ScrollPanFraction * (CurrentMax - CurrentMin);
+
+	float NewViewOutputMin = CurrentMin + InDelta;
+	float NewViewOutputMax = CurrentMax + InDelta;
+
+	ClampViewRange(NewViewOutputMin, NewViewOutputMax);
+	SetViewRange(NewViewOutputMin, NewViewOutputMax, EViewRangeInterpolation::Animated);
+}
+
+bool FSequencerTimeSliderController::HitTestPlaybackStart(const FScrubRangeToScreen& RangeToScreen, const TRange<float>& PlaybackRange, float LocalHitPositionX, float ScrubPosition) const
+{
+	static float BrushSizeInStateUnits = 6.f, DragToleranceSlateUnits = 2.f;
+	float LocalPlaybackStartPos = RangeToScreen.InputToLocalX(PlaybackRange.GetLowerBoundValue());
+
+	// We favor hit testing the scrub bar over hit testing the playback range bounds
+	if (FMath::IsNearlyEqual(LocalPlaybackStartPos, RangeToScreen.InputToLocalX(ScrubPosition), ScrubHandleSize/2.f))
+	{
+		return false;
+	}
+
+	// Hit test against the brush region to the right of the playback start position, +/- DragToleranceSlateUnits
+	return LocalHitPositionX >= LocalPlaybackStartPos - DragToleranceSlateUnits &&
+		LocalHitPositionX <= LocalPlaybackStartPos + BrushSizeInStateUnits + DragToleranceSlateUnits;
+}
+
+bool FSequencerTimeSliderController::HitTestPlaybackEnd(const FScrubRangeToScreen& RangeToScreen, const TRange<float>& PlaybackRange, float LocalHitPositionX, float ScrubPosition) const
+{
+	static float BrushSizeInStateUnits = 6.f, DragToleranceSlateUnits = 2.f;
+	float LocalPlaybackEndPos = RangeToScreen.InputToLocalX(PlaybackRange.GetUpperBoundValue());
+	
+	// We favor hit testing the scrub bar over hit testing the playback range bounds
+	if (FMath::IsNearlyEqual(LocalPlaybackEndPos, RangeToScreen.InputToLocalX(ScrubPosition), ScrubHandleSize/2.f))
+	{
+		return false;
+	}
+	
+	// Hit test against the brush region to the left of the playback end position, +/- DragToleranceSlateUnits
+	return LocalHitPositionX >= LocalPlaybackEndPos - BrushSizeInStateUnits - DragToleranceSlateUnits &&
+		LocalHitPositionX <= LocalPlaybackEndPos + DragToleranceSlateUnits;
+}
+
+void FSequencerTimeSliderController::SetPlaybackRangeStart(float NewStart)
+{
+	TRange<float> PlaybackRange = TimeSliderArgs.PlaybackRange.Get();
+
+	if (NewStart <= PlaybackRange.GetUpperBoundValue())
+	{
+		TimeSliderArgs.OnPlaybackRangeChanged.ExecuteIfBound(TRange<float>(NewStart, PlaybackRange.GetUpperBoundValue()));
+	}
+}
+
+void FSequencerTimeSliderController::SetPlaybackRangeEnd(float NewEnd)
+{
+	TRange<float> PlaybackRange = TimeSliderArgs.PlaybackRange.Get();
+
+	if (NewEnd >= PlaybackRange.GetLowerBoundValue())
+	{
+		TimeSliderArgs.OnPlaybackRangeChanged.ExecuteIfBound(TRange<float>(PlaybackRange.GetLowerBoundValue(), NewEnd));
+	}
+}
+
+#undef LOCTEXT_NAMESPACE
