Copying //UE4/Dev-Sequencer to //UE4/Main

diff --git a/Engine/Source/Editor/Sequencer/Private/IntegralKeyDetailsCustomization.h b/Engine/Source/Editor/Sequencer/Private/IntegralKeyDetailsCustomization.h
new file mode 100644
--- /dev/null
+++ b/Engine/Source/Editor/Sequencer/Private/IntegralKeyDetailsCustomization.h
@@ -0,0 +1,22 @@
+// Copyright 1998-2016 Epic Games, Inc. All Rights Reserved.
+
+#pragma once
+
+#include "IDetailCustomization.h"
+
+class FIntegralKeyDetailsCustomization : public IDetailCustomization
+{
+public:
+	/** Makes a new instance of this detail layout class for a specific detail view requesting it */
+	static TSharedRef<class IDetailCustomization> MakeInstance(TWeakObjectPtr<const UMovieSceneSection> InSection);
+
+	FIntegralKeyDetailsCustomization(TWeakObjectPtr<const UMovieSceneSection> InSection)
+		: Section(InSection)
+	{}
+
+	/** IDetailCustomization interface */
+	virtual void CustomizeDetails( IDetailLayoutBuilder& DetailBuilder ) override;
+
+private:
+	TWeakObjectPtr<const UMovieSceneSection> Section;
+};
\ No newline at end of file
diff --git a/Engine/Source/Editor/Sequencer/Private/SAnimationOutlinerTreeNode.cpp b/Engine/Source/Editor/Sequencer/Private/SAnimationOutlinerTreeNode.cpp
--- a/Engine/Source/Editor/Sequencer/Private/SAnimationOutlinerTreeNode.cpp
+++ b/Engine/Source/Editor/Sequencer/Private/SAnimationOutlinerTreeNode.cpp
@@ -12,7 +12,7 @@
 #include "SEditableLabel.h"
 #include "SSequencerTreeView.h"
 #include "SColorPicker.h"
-
+#include "SequencerSectionPainter.h"
 
 #define LOCTEXT_NAMESPACE "AnimationOutliner"
 
@@ -106,8 +106,10 @@ void SAnimationOutlinerTreeNode::Construct( const FArguments& InArgs, TSharedRef
 	EditableLabel = SNew(SEditableLabel)
 		.CanEdit(this, &SAnimationOutlinerTreeNode::HandleNodeLabelCanEdit)
 		.Font(NodeFont)
+		.ColorAndOpacity(this, &SAnimationOutlinerTreeNode::GetDisplayNameColor)
 		.OnTextChanged(this, &SAnimationOutlinerTreeNode::HandleNodeLabelTextChanged)
-		.Text(this, &SAnimationOutlinerTreeNode::GetDisplayName);
+		.Text(this, &SAnimationOutlinerTreeNode::GetDisplayName)
+		.ToolTipText(this, &SAnimationOutlinerTreeNode::GetDisplayNameToolTipText);
 
 
 	Node->OnRenameRequested().AddRaw(this, &SAnimationOutlinerTreeNode::EnterRenameMode);
@@ -233,16 +235,6 @@ void SAnimationOutlinerTreeNode::EnterRenameMode()
 	EditableLabel->EnterTextMode();
 }
 
-TSharedPtr<FSequencerDisplayNode> SAnimationOutlinerTreeNode::GetRootNode(TSharedPtr<FSequencerDisplayNode> ObjectNode)
-{
-	if (!ObjectNode->GetParent().IsValid())
-	{
-		return ObjectNode;
-	}
-
-	return GetRootNode(ObjectNode->GetParent());
-}
-
 
 void SAnimationOutlinerTreeNode::GetAllDescendantNodes(TSharedPtr<FSequencerDisplayNode> RootNode, TArray<TSharedRef<FSequencerDisplayNode> >& AllNodes)
 {
@@ -262,131 +254,6 @@ void SAnimationOutlinerTreeNode::GetAllDescendantNodes(TSharedPtr<FSequencerDisp
 	}
 }
 
-
-FReply SAnimationOutlinerTreeNode::OnMouseButtonDown( const FGeometry& MyGeometry, const FPointerEvent& MouseEvent )
-{
-	if (!DisplayNode->IsSelectable())
-	{
-		return FReply::Unhandled();
-	}
-
-	if ((MouseEvent.GetEffectingButton() != EKeys::LeftMouseButton) &&
-		(MouseEvent.GetEffectingButton() != EKeys::RightMouseButton))
-	{
-		return FReply::Unhandled();
-	}
-
-	FSequencer& Sequencer = DisplayNode->GetSequencer();
-	bool bSelected = Sequencer.GetSelection().IsSelected(DisplayNode.ToSharedRef());
-
-	TArray<TSharedPtr<FSequencerDisplayNode> > AffectedNodes;
-	AffectedNodes.Add(DisplayNode.ToSharedRef());
-
-	Sequencer.GetSelection().SuspendBroadcast();
-
-	if (MouseEvent.IsShiftDown())
-	{
-		FSequencerNodeTree& ParentTree = DisplayNode->GetParentTree();
-
-		const TArray< TSharedRef<FSequencerDisplayNode> >RootNodes = ParentTree.GetRootNodes();
-
-		// Get all nodes in order
-		TArray<TSharedRef<FSequencerDisplayNode> > AllNodes;
-		for (int32 i = 0; i < RootNodes.Num(); ++i)
-		{
-			GetAllDescendantNodes(RootNodes[i], AllNodes);
-		}
-
-		int32 FirstIndexToSelect = INT32_MAX;
-		int32 LastIndexToSelect = INT32_MIN;
-
-		for (int32 ChildIndex = 0; ChildIndex < AllNodes.Num(); ++ChildIndex)
-		{
-			TSharedRef<FSequencerDisplayNode> ChildNode = AllNodes[ChildIndex];
-
-			if (ChildNode == DisplayNode.ToSharedRef() || Sequencer.GetSelection().IsSelected(ChildNode))
-			{
-				if (ChildIndex < FirstIndexToSelect)
-				{
-					FirstIndexToSelect = ChildIndex;
-				}
-
-				if (ChildIndex > LastIndexToSelect)
-				{
-					LastIndexToSelect = ChildIndex;
-				}
-			}
-		}
-
-		if (FirstIndexToSelect != INT32_MAX && LastIndexToSelect != INT32_MIN)
-		{
-			for (int32 ChildIndex = FirstIndexToSelect; ChildIndex <= LastIndexToSelect; ++ChildIndex)
-			{
-				TSharedRef<FSequencerDisplayNode> ChildNode = AllNodes[ChildIndex];
-
-				if (!Sequencer.GetSelection().IsSelected(ChildNode))
-				{
-					Sequencer.GetSelection().AddToSelection(ChildNode);
-					AffectedNodes.Add(ChildNode);
-				}
-			}
-		}
-	}
-	else if( MouseEvent.IsControlDown() )
-	{
-		// Toggle selection when control is down
-		if (bSelected)
-		{
-			Sequencer.GetSelection().RemoveFromSelection(DisplayNode.ToSharedRef());
-		}
-		else
-		{
-			Sequencer.GetSelection().AddToSelection(DisplayNode.ToSharedRef());
-		}
-	}
-	else if (MouseEvent.GetEffectingButton() != EKeys::RightMouseButton || !bSelected)
-	{
-		// Deselect the other nodes and select this node.
-		Sequencer.GetSelection().EmptySelectedOutlinerNodes();
-		Sequencer.GetSelection().AddToSelection(DisplayNode.ToSharedRef());
-	}
-
-	OnSelectionChanged( AffectedNodes );
-
-	Sequencer.GetSelection().ResumeBroadcast();
-	Sequencer.GetSelection().GetOnOutlinerNodeSelectionChanged().Broadcast();
-
-	return FReply::Handled();
-}
-
-
-FReply SAnimationOutlinerTreeNode::OnMouseButtonUp( const FGeometry& MyGeometry, const FPointerEvent& MouseEvent )
-{
-	if( MouseEvent.GetEffectingButton() != EKeys::RightMouseButton )
-	{
-		return FReply::Unhandled();
-	}
-
-	TSharedPtr<SWidget> MenuContent = DisplayNode->OnSummonContextMenu(MyGeometry, MouseEvent);
-
-	if (!MenuContent.IsValid())
-	{
-		return FReply::Handled();
-	}
-
-	FWidgetPath WidgetPath = MouseEvent.GetEventPath() != nullptr ? *MouseEvent.GetEventPath() : FWidgetPath();
-
-	FSlateApplication::Get().PushMenu(
-		AsShared(),
-		WidgetPath,
-		MenuContent.ToSharedRef(),
-		MouseEvent.GetScreenSpacePosition(),
-		FPopupTransitionEffect( FPopupTransitionEffect::ContextMenu )
-	);
-			
-	return FReply::Handled().SetUserFocus(MenuContent.ToSharedRef(), EFocusCause::SetDirectly);
-}
-
 void SAnimationOutlinerTreeNode::OnMouseEnter(const FGeometry& MyGeometry, const FPointerEvent& MouseEvent)
 {
 	DisplayNode->GetParentTree().SetHoveredNode(DisplayNode);
@@ -408,19 +275,15 @@ FSlateColor SAnimationOutlinerTreeNode::GetNodeBackgroundTint() const
 {
 	FSequencer& Sequencer = DisplayNode->GetSequencer();
 	const bool bIsSelected = Sequencer.GetSelection().IsSelected(DisplayNode.ToSharedRef());
-	const bool bIsSelectionActive = Sequencer.GetSelection().GetActiveSelection() == FSequencerSelection::EActiveSelection::OutlinerNode;
 
 	if (bIsSelected)
 	{
-		if (bIsSelectionActive)
-		{
-			return FEditorStyle::GetSlateColor("SelectionColor_Pressed");
-		}
-		else
-		{
-			return FEditorStyle::GetSlateColor("SelectionColor_Inactive");
-		}
+		return FEditorStyle::GetSlateColor("SelectionColor_Pressed");
 	}
+	else if (Sequencer.GetSelection().NodeHasSelectedKeysOrSections(DisplayNode.ToSharedRef()))
+	{
+		return FLinearColor(FColor(115, 115, 115, 255));
+	}	
 	else if (DisplayNode->IsHovered())
 	{
 		return bIsTopLevelNode ? FLinearColor(FColor(52, 52, 52, 255)) : FLinearColor(FColor(72, 72, 72, 255));
@@ -467,7 +330,7 @@ FSlateColor SAnimationOutlinerTreeNode::GetTrackColorTint() const
 			UMovieSceneTrack* Track = static_cast<FSequencerTrackNode*>(Current)->GetTrack();
 			if (Track)
 			{
-				return FLinearColor(Track->GetColorTint()).CopyWithNewOpacity(1.f);
+				return FSequencerSectionPainter::BlendColor(Track->GetColorTint());
 			}
 		}
 		Current = Current->GetParent().Get();
@@ -491,161 +354,33 @@ EVisibility SAnimationOutlinerTreeNode::GetExpanderVisibility() const
 }
 
 
-FText SAnimationOutlinerTreeNode::GetDisplayName() const
+FSlateColor SAnimationOutlinerTreeNode::GetDisplayNameColor() const
 {
-	return DisplayNode->GetDisplayName();
+	return DisplayNode->GetDisplayNameColor();
 }
 
 
-bool SAnimationOutlinerTreeNode::HandleNodeLabelCanEdit() const
+FText SAnimationOutlinerTreeNode::GetDisplayNameToolTipText() const
 {
-	return DisplayNode->CanRenameNode();
+	return DisplayNode->GetDisplayNameToolTipText();
 }
 
 
-void SAnimationOutlinerTreeNode::HandleNodeLabelTextChanged(const FText& NewLabel)
+FText SAnimationOutlinerTreeNode::GetDisplayName() const
 {
-	DisplayNode->SetDisplayName(NewLabel);
+	return DisplayNode->GetDisplayName();
 }
 
 
-void SAnimationOutlinerTreeNode::OnSelectionChanged( TArray<TSharedPtr<FSequencerDisplayNode> > AffectedNodes )
+bool SAnimationOutlinerTreeNode::HandleNodeLabelCanEdit() const
 {
-	FSequencer& Sequencer = DisplayNode->GetSequencer();
-	if (!Sequencer.IsLevelEditorSequencer())
-	{
-		return;
-	}
-
-	TArray<TSharedPtr<FSequencerDisplayNode> > AffectedRootNodes;
-
-	for (int32 NodeIdx = 0; NodeIdx < AffectedNodes.Num(); ++NodeIdx)
-	{
-		TSharedPtr<FSequencerDisplayNode> RootNode = GetRootNode(AffectedNodes[NodeIdx]);
-		if (RootNode.IsValid() && RootNode->GetType() == ESequencerNode::Object)
-		{
-			AffectedRootNodes.Add(RootNode);
-		}
-	}
-
-	if (!AffectedRootNodes.Num())
-	{
-		return;
-	}
-
-	const FModifierKeysState ModifierKeys = FSlateApplication::Get().GetModifierKeys();
-	bool IsControlDown = ModifierKeys.IsControlDown();
-	bool IsShiftDown = ModifierKeys.IsShiftDown();
-
-	TArray<AActor*> ActorsToSelect;
-	TArray<AActor*> ActorsToRemainSelected;
-	TArray<AActor*> ActorsToDeselect;
-
-	// Find objects bound to the object node and determine what actors need to be selected in the editor
-	for (int32 ObjectIdx = 0; ObjectIdx < AffectedRootNodes.Num(); ++ObjectIdx)
-	{
-		const TSharedPtr<const FSequencerObjectBindingNode> ObjectNode = StaticCastSharedPtr<const FSequencerObjectBindingNode>( AffectedRootNodes[ObjectIdx] );
-
-		// Get the bound objects
-		TArray<UObject*> RuntimeObjects;
-		Sequencer.GetRuntimeObjects( Sequencer.GetFocusedMovieSceneSequenceInstance(), ObjectNode->GetObjectBinding(), RuntimeObjects );
-		
-		if( RuntimeObjects.Num() > 0 )
-		{
-			for( int32 ActorIdx = 0; ActorIdx < RuntimeObjects.Num(); ++ActorIdx )
-			{
-				AActor* Actor = Cast<AActor>( RuntimeObjects[ActorIdx] );
-
-				if (Actor == nullptr)
-				{
-					continue;
-				}
-
-				bool bIsActorSelected = GEditor->GetSelectedActors()->IsSelected(Actor);
-				if (IsControlDown)
-				{
-					if (!bIsActorSelected)
-					{
-						ActorsToSelect.Add(Actor);
-					}
-					else
-					{
-						ActorsToDeselect.Add(Actor);
-					}
-				}
-				else
-				{
-					if (!bIsActorSelected)
-					{
-						ActorsToSelect.Add(Actor);
-					}
-					else
-					{
-						ActorsToRemainSelected.Add(Actor);
-					}
-				}
-			}
-		}
-	}
-
-	if (!IsControlDown && !IsShiftDown)
-	{
-		for (FSelectionIterator SelectionIt(*GEditor->GetSelectedActors()); SelectionIt; ++SelectionIt)
-		{
-			AActor* Actor = CastChecked< AActor >( *SelectionIt );
-			if (Actor)
-			{
-				if (ActorsToSelect.Find(Actor) == INDEX_NONE && 
-					ActorsToDeselect.Find(Actor) == INDEX_NONE &&
-					ActorsToRemainSelected.Find(Actor) == INDEX_NONE)
-				{
-					ActorsToDeselect.Add(Actor);
-				}
-			}
-		}
-	}
-
-	// If there's no selection to change, return early
-	if (ActorsToSelect.Num() + ActorsToDeselect.Num() == 0)
-	{
-		return;
-	}
-
-	// Mark that the user is selecting so that the UI doesn't respond to the selection changes in the following block
-	TSharedRef<SSequencer> SequencerWidget = StaticCastSharedRef<SSequencer>(Sequencer.GetSequencerWidget());
-	SequencerWidget->SetUserIsSelecting(true);
-
-	const FScopedTransaction Transaction(NSLOCTEXT("UnrealEd", "ClickingOnActors", "Clicking on Actors"));
-
-	const bool bNotifySelectionChanged = false;
-	const bool bDeselectBSP = true;
-	const bool bWarnAboutTooManyActors = false;
-	const bool bSelectEvenIfHidden = true;
-
-	GEditor->GetSelectedActors()->Modify();
-
-	if (!IsControlDown && !IsShiftDown)
-	{
-		GEditor->SelectNone(bNotifySelectionChanged, bDeselectBSP, bWarnAboutTooManyActors);
-	}
-
-	GEditor->GetSelectedActors()->BeginBatchSelectOperation();
-
-	for (auto ActorToSelect : ActorsToSelect)
-	{
-		GEditor->SelectActor(ActorToSelect, true, bNotifySelectionChanged, bSelectEvenIfHidden);
-	}
-
-	for (auto ActorToDeselect : ActorsToDeselect)
-	{
-		GEditor->SelectActor(ActorToDeselect, false, bNotifySelectionChanged, bSelectEvenIfHidden);
-	}
+	return DisplayNode->CanRenameNode();
+}
 
-	GEditor->GetSelectedActors()->EndBatchSelectOperation();
-	GEditor->NoteSelectionChange();
 
-	// Unlock the selection so that the sequencer widget can now respond to selection changes in the level
-	SequencerWidget->SetUserIsSelecting(false);
+void SAnimationOutlinerTreeNode::HandleNodeLabelTextChanged(const FText& NewLabel)
+{
+	DisplayNode->SetDisplayName(NewLabel);
 }
 
 
diff --git a/Engine/Source/Editor/Sequencer/Private/SAnimationOutlinerTreeNode.h b/Engine/Source/Editor/Sequencer/Private/SAnimationOutlinerTreeNode.h
--- a/Engine/Source/Editor/Sequencer/Private/SAnimationOutlinerTreeNode.h
+++ b/Engine/Source/Editor/Sequencer/Private/SAnimationOutlinerTreeNode.h
@@ -45,9 +45,6 @@ private:
 
 	// SWidget interface
 
-	virtual FReply OnMouseButtonDown(const FGeometry& MyGeometry, const FPointerEvent& MouseEvent) override;
-	virtual FReply OnMouseButtonUp(const FGeometry& MyGeometry, const FPointerEvent& MouseEvent) override;
-
 	void OnMouseEnter(const FGeometry& MyGeometry, const FPointerEvent& MouseEvent) override;
 	void OnMouseLeave(const FPointerEvent& MouseEvent) override;
 
@@ -71,6 +68,16 @@ private:
 	 */
 	EVisibility GetExpanderVisibility() const;
 
+	/**
+	 * @return The color used to draw the display name.
+	 */
+	FSlateColor GetDisplayNameColor() const;
+
+	/**
+	 * @return The text displayed for the tool tip for the diplay name label. 
+	 */
+	FText GetDisplayNameToolTipText() const;
+
 	/**
 	 * @return The display name for this node.
 	 */
@@ -85,12 +92,6 @@ private:
 	/** Get all descendant nodes from the given root node. */
 	void GetAllDescendantNodes(TSharedPtr<FSequencerDisplayNode> RootNode, TArray<TSharedRef<FSequencerDisplayNode> >& AllNodes);
 
-	/** Return the root node given an object node. */
-	TSharedPtr<FSequencerDisplayNode> GetRootNode(TSharedPtr<FSequencerDisplayNode> ObjectNode);
-
-	/** Called when nodes are selected. */
-	void OnSelectionChanged(TArray<TSharedPtr<FSequencerDisplayNode>> AffectedNodes);
-
 	/** Called when the user clicks the track color */
 	TSharedRef<SWidget> OnGetColorPicker() const;
 
diff --git a/Engine/Source/Editor/Sequencer/Private/SSequencer.cpp b/Engine/Source/Editor/Sequencer/Private/SSequencer.cpp
--- a/Engine/Source/Editor/Sequencer/Private/SSequencer.cpp
+++ b/Engine/Source/Editor/Sequencer/Private/SSequencer.cpp
@@ -45,6 +45,7 @@
 DECLARE_DELEGATE_RetVal(bool, FOnGetShowFrames)
 DECLARE_DELEGATE_RetVal(uint8, FOnGetZeroPad)
 
+
 /* Numeric type interface for showing numbers as frames or times */
 struct FFramesOrTimeInterface : public TNumericUnitTypeInterface<float>
 {
@@ -77,10 +78,18 @@ private:
 
 	virtual TOptional<float> FromString(const FString& InString, const float& InExistingValue) override
 	{
-		if (ShowFrameNumbers.Execute())
+		bool bShowFrameNumbers = ShowFrameNumbers.IsBound() ? ShowFrameNumbers.Execute() : false;
+		if (bShowFrameNumbers)
 		{
-			int32 NewEndFrame = FCString::Atoi(*InString);
-			return float(TimeSliderController->FrameToTime(NewEndFrame));
+			// Convert existing value to frames
+			float ExistingValueInFrames = TimeSliderController->TimeToFrame(InExistingValue);
+			TOptional<float> Result = TNumericUnitTypeInterface<float>::FromString(InString, ExistingValueInFrames);
+
+			if (Result.IsSet())
+			{
+				int32 NewEndFrame = FMath::RoundToInt(Result.GetValue());
+				return float(TimeSliderController->FrameToTime(NewEndFrame));
+			}
 		}
 
 		return TNumericUnitTypeInterface::FromString(InString, InExistingValue);
@@ -99,16 +108,13 @@ void SSequencer::Construct(const FArguments& InArgs, TSharedRef<FSequencer> InSe
 	CachedClampRange = TRange<float>::Empty();
 	CachedViewRange = TRange<float>::Empty();
 
-	USelection::SelectionChangedEvent.AddSP(this, &SSequencer::OnActorSelectionChanged);
-
 	Settings = InSequencer->GetSettings();
 	Settings->GetOnShowCurveEditorChanged().AddSP(this, &SSequencer::OnCurveEditorVisibilityChanged);
 
-	// Create a node tree which contains a tree of movie scene data to display in the sequence
-	SequencerNodeTree = MakeShareable( new FSequencerNodeTree( InSequencer.Get() ) );
-
 	ISequencerWidgetsModule& SequencerWidgets = FModuleManager::Get().LoadModuleChecked<ISequencerWidgetsModule>( "SequencerWidgets" );
 
+	OnBeginInOutRangeDrag = InArgs._OnBeginInOutRangeDrag;
+	OnEndInOutRangeDrag = InArgs._OnEndInOutRangeDrag;
 	OnBeginPlaybackRangeDrag = InArgs._OnBeginPlaybackRangeDrag;
 	OnEndPlaybackRangeDrag = InArgs._OnEndPlaybackRangeDrag;
 
@@ -116,7 +122,11 @@ void SSequencer::Construct(const FArguments& InArgs, TSharedRef<FSequencer> InSe
 	{
 		TimeSliderArgs.ViewRange = InArgs._ViewRange;
 		TimeSliderArgs.ClampRange = InArgs._ClampRange;
+		TimeSliderArgs.InOutRange = InArgs._InOutRange;
 		TimeSliderArgs.PlaybackRange = InArgs._PlaybackRange;
+		TimeSliderArgs.OnInOutRangeChanged = InArgs._OnInOutRangeChanged;
+		TimeSliderArgs.OnBeginInOutRangeDrag = OnBeginInOutRangeDrag;
+		TimeSliderArgs.OnEndInOutRangeDrag = OnEndInOutRangeDrag;
 		TimeSliderArgs.OnPlaybackRangeChanged = InArgs._OnPlaybackRangeChanged;
 		TimeSliderArgs.OnBeginPlaybackRangeDrag = OnBeginPlaybackRangeDrag;
 		TimeSliderArgs.OnEndPlaybackRangeDrag = OnEndPlaybackRangeDrag;
@@ -126,11 +136,15 @@ void SSequencer::Construct(const FArguments& InArgs, TSharedRef<FSequencer> InSe
 		TimeSliderArgs.OnBeginScrubberMovement = InArgs._OnBeginScrubbing;
 		TimeSliderArgs.OnEndScrubberMovement = InArgs._OnEndScrubbing;
 		TimeSliderArgs.OnScrubPositionChanged = InArgs._OnScrubPositionChanged;
+		TimeSliderArgs.PlaybackStatus = InArgs._PlaybackStatus;
+
 		TimeSliderArgs.Settings = Settings;
 	}
 
-	TSharedRef<FSequencerTimeSliderController> TimeSliderController( new FSequencerTimeSliderController( TimeSliderArgs ) );
+	TimeSliderController = MakeShareable( new FSequencerTimeSliderController( TimeSliderArgs ) );
 	
+	TSharedRef<FSequencerTimeSliderController> TimeSliderControllerRef = TimeSliderController.ToSharedRef();
+
 	{
 		auto ShowFrameNumbersDelegate = FOnGetShowFrames::CreateSP(this, &SSequencer::ShowFrameNumbers);
 		USequencerSettings* SequencerSettings = Settings;
@@ -142,18 +156,28 @@ void SSequencer::Construct(const FArguments& InArgs, TSharedRef<FSequencer> InSe
 			return 0;
 		};
 
-		NumericTypeInterface = MakeShareable( new FFramesOrTimeInterface(ShowFrameNumbersDelegate, TimeSliderController, FOnGetZeroPad()) );
-		ZeroPadNumericTypeInterface = MakeShareable( new FFramesOrTimeInterface(ShowFrameNumbersDelegate, TimeSliderController, FOnGetZeroPad::CreateLambda(GetZeroPad)) );
+		NumericTypeInterface = MakeShareable( new FFramesOrTimeInterface(ShowFrameNumbersDelegate, TimeSliderControllerRef, FOnGetZeroPad()) );
+		ZeroPadNumericTypeInterface = MakeShareable( new FFramesOrTimeInterface(ShowFrameNumbersDelegate, TimeSliderControllerRef, FOnGetZeroPad::CreateLambda(GetZeroPad)) );
 	}
 
 	bool bMirrorLabels = false;
+	
 	// Create the top and bottom sliders
-	TSharedRef<ITimeSlider> TopTimeSlider = SequencerWidgets.CreateTimeSlider( TimeSliderController, bMirrorLabels );
+	TSharedRef<ITimeSlider> TopTimeSlider = SequencerWidgets.CreateTimeSlider( TimeSliderControllerRef, bMirrorLabels );
 	bMirrorLabels = true;
-	TSharedRef<ITimeSlider> BottomTimeSlider = SequencerWidgets.CreateTimeSlider( TimeSliderController, TAttribute<EVisibility>(this, &SSequencer::GetBottomTimeSliderVisibility), bMirrorLabels );
+	TSharedRef<ITimeSlider> BottomTimeSlider = SequencerWidgets.CreateTimeSlider( TimeSliderControllerRef, TAttribute<EVisibility>(this, &SSequencer::GetBottomTimeSliderVisibility), bMirrorLabels );
 
 	// Create bottom time range slider
-	TSharedRef<ITimeSlider> BottomTimeRange = SequencerWidgets.CreateTimeRange( TimeSliderController, TAttribute<EVisibility>(this, &SSequencer::GetTimeRangeVisibility), TAttribute<bool>(this, &SSequencer::ShowFrameNumbers), TAttribute<float>(this, &SSequencer::OnGetTimeSnapInterval));
+	TSharedRef<ITimeSlider> BottomTimeRange = SequencerWidgets.CreateTimeRange(
+		FTimeRangeArgs(
+			EShowRange(EShowRange::WorkingRange | EShowRange::ViewRange),
+			TimeSliderControllerRef,
+			TAttribute<EVisibility>(this, &SSequencer::GetTimeRangeVisibility),
+			TAttribute<bool>(this, &SSequencer::ShowFrameNumbers),
+			ZeroPadNumericTypeInterface.ToSharedRef()
+		),
+		SequencerWidgets.CreateTimeRangeSlider(TimeSliderControllerRef, TAttribute<float>(this, &SSequencer::OnGetTimeSnapInterval))
+	);
 
 	OnGetAddMenuContent = InArgs._OnGetAddMenuContent;
 	AddMenuExtender = InArgs._AddMenuExtender;
@@ -171,11 +195,11 @@ void SSequencer::Construct(const FArguments& InArgs, TSharedRef<FSequencer> InSe
 		.Thickness(FVector2D(5.0f, 5.0f));
 
 	SAssignNew(TrackOutliner, SSequencerTrackOutliner);
-	SAssignNew(TrackArea, SSequencerTrackArea, TimeSliderController, InSequencer);
-	SAssignNew(TreeView, SSequencerTreeView, SequencerNodeTree.ToSharedRef(), TrackArea.ToSharedRef())
+	SAssignNew(TrackArea, SSequencerTrackArea, TimeSliderControllerRef, InSequencer);
+	SAssignNew(TreeView, SSequencerTreeView, InSequencer->GetNodeTree(), TrackArea.ToSharedRef())
 		.ExternalScrollbar(ScrollBar);
 
-	SAssignNew(CurveEditor, SSequencerCurveEditor, InSequencer, TimeSliderController)
+	SAssignNew(CurveEditor, SSequencerCurveEditor, InSequencer, TimeSliderControllerRef)
 		.Visibility(this, &SSequencer::GetCurveEditorVisibility)
 		.OnViewRangeChanged(InArgs._OnViewRangeChanged)
 		.ViewRange(InArgs._ViewRange);
@@ -359,7 +383,7 @@ void SSequencer::Construct(const FArguments& InArgs, TSharedRef<FSequencer> InSe
 						//.BorderBackgroundColor(FLinearColor(.50f, .50f, .50f, 1.0f))
 						.HAlign(HAlign_Center)
 						[
-							MakeTransportControls()
+							SequencerPtr.Pin()->MakeTransportControls(true)
 						]
 					]
 
@@ -392,21 +416,12 @@ void SSequencer::Construct(const FArguments& InArgs, TSharedRef<FSequencer> InSe
 					+ SGridPanel::Slot( Column1, Row2, SGridPanel::Layer(10) )
 					.Padding(ResizeBarPadding)
 					[
-						SNew( SSequencerSectionOverlay, TimeSliderController )
+						SNew( SSequencerSectionOverlay, TimeSliderControllerRef )
 						.Visibility( EVisibility::HitTestInvisible )
 						.DisplayScrubPosition( false )
 						.DisplayTickLines( true )
 					]
 
-					// Goto box
-					+ SGridPanel::Slot( Column1, Row2, SGridPanel::Layer(11) )
-						.Padding(ResizeBarPadding)
-						.HAlign(HAlign_Left)
-						.VAlign(VAlign_Top)
-						[
-							SAssignNew(GotoBox, SSequencerGotoBox, SequencerPtr.Pin().ToSharedRef(), *Settings, NumericTypeInterface.ToSharedRef())
-						]
-
 					// Curve editor
 					+ SGridPanel::Slot( Column1, Row2, SGridPanel::Layer(20) )
 					.Padding(ResizeBarPadding)
@@ -418,13 +433,22 @@ void SSequencer::Construct(const FArguments& InArgs, TSharedRef<FSequencer> InSe
 					+ SGridPanel::Slot( Column1, Row2, SGridPanel::Layer(30) )
 					.Padding(ResizeBarPadding)
 					[
-						SNew( SSequencerSectionOverlay, TimeSliderController )
+						SNew( SSequencerSectionOverlay, TimeSliderControllerRef )
 						.Visibility( EVisibility::HitTestInvisible )
 						.DisplayScrubPosition( true )
 						.DisplayTickLines( false )
 						.PaintPlaybackRangeArgs(this, &SSequencer::GetSectionPlaybackRangeArgs)
 					]
 
+					// Goto box
+					+ SGridPanel::Slot(Column1, Row2, SGridPanel::Layer(40))
+						.Padding(ResizeBarPadding)
+						.HAlign(HAlign_Left)
+						.VAlign(VAlign_Top)
+						[
+							SAssignNew(GotoBox, SSequencerGotoBox, SequencerPtr.Pin().ToSharedRef(), *Settings, NumericTypeInterface.ToSharedRef())
+						]
+
 					// play range sliders
 					+ SGridPanel::Slot( Column1, Row3, SGridPanel::Layer(10) )
 					.Padding(ResizeBarPadding)
@@ -476,6 +500,7 @@ void SSequencer::Construct(const FArguments& InArgs, TSharedRef<FSequencer> InSe
 
 	InSequencer->GetSelection().GetOnKeySelectionChanged().AddSP(this, &SSequencer::HandleKeySelectionChanged);
 	InSequencer->GetSelection().GetOnSectionSelectionChanged().AddSP(this, &SSequencer::HandleSectionSelectionChanged);
+	InSequencer->GetSelection().GetOnOutlinerNodeSelectionChanged().AddSP(this, &SSequencer::HandleOutlinerNodeSelectionChanged);
 
 	ResetBreadcrumbs();
 }
@@ -503,6 +528,16 @@ void SSequencer::BindCommands(TSharedRef<FUICommandList> SequencerCommandBinding
 		FExecuteAction::CreateSP(this, &SSequencer::PasteFromHistory),
 		FCanExecuteAction::CreateLambda(CanPaste)
 	);
+
+	SequencerCommandBindings->MapAction(
+		FSequencerCommands::Get().ToggleShowGotoBox,
+		FExecuteAction::CreateLambda([this]{ GotoBox->ToggleVisibility(); })
+	);
+}
+	
+const ISequencerEditTool* SSequencer::GetEditTool() const
+{
+	return TrackArea->GetEditTool();
 }
 
 
@@ -525,6 +560,7 @@ TSharedRef<INumericTypeInterface<float>> SSequencer::GetZeroPadNumericTypeInterf
 	return ZeroPadNumericTypeInterface.ToSharedRef();
 }
 
+
 /* SSequencer callbacks
  *****************************************************************************/
 
@@ -567,6 +603,21 @@ void SSequencer::HandleSectionSelectionChanged()
 }
 
 
+void SSequencer::HandleOutlinerNodeSelectionChanged()
+{
+	const TSet<TSharedRef<FSequencerDisplayNode>>& OutlinerSelection = SequencerPtr.Pin()->GetSelection().GetSelectedOutlinerNodes();
+
+	if ( OutlinerSelection.Num() == 1 )
+	{
+		for ( auto& Node : OutlinerSelection )
+		{
+			TreeView->RequestScrollIntoView( Node );
+			break;
+		}
+	}
+}
+
+
 TSharedRef<SWidget> SSequencer::MakeAddButton()
 {
 	return SNew(SComboButton)
@@ -632,11 +683,10 @@ TSharedRef<SWidget> SSequencer::MakeToolBar()
 				NAME_None,
 				LOCTEXT("SaveAs", "Save As"),
 				LOCTEXT("SaveAsTooltip", "Saves the current level sequence under a different name"),
-				FSlateIcon(FEditorStyle::GetStyleSetName(), "AssetEditor.SaveAssetAs")
+				FSlateIcon(FEditorStyle::GetStyleSetName(), "Sequencer.SaveAs")
 			);
 
-			// @todo sequencer: implement FSequencer::DiscardChanges()
-			//ToolBarBuilder.AddToolBarButton( FSequencerCommands::Get().DiscardChanges );
+			ToolBarBuilder.AddToolBarButton( FSequencerCommands::Get().DiscardChanges );
 			ToolBarBuilder.AddToolBarButton( FSequencerCommands::Get().RenderMovie );
 			ToolBarBuilder.AddSeparator();
 		}
@@ -779,6 +829,7 @@ TSharedRef<SWidget> SSequencer::MakeGeneralMenu()
 		}
 
 		MenuBuilder.AddMenuEntry( FSequencerCommands::Get().ToggleCombinedKeyframes );
+		MenuBuilder.AddMenuEntry( FSequencerCommands::Get().ToggleChannelColors );
 
 		if (SequencerPtr.Pin()->IsLevelEditorSequencer())
 		{
@@ -864,6 +915,7 @@ TSharedRef<SWidget> SSequencer::MakeGeneralMenu()
 
 		MenuBuilder.AddMenuEntry( FSequencerCommands::Get().ToggleKeepCursorInPlaybackRange );
 		MenuBuilder.AddMenuEntry( FSequencerCommands::Get().ToggleKeepPlaybackRangeInSectionBounds );
+		MenuBuilder.AddMenuEntry( FSequencerCommands::Get().ToggleLinkCurveEditorTimeRange );
 
 		// Menu entry for zero padding
 		auto OnZeroPadChanged = [=](uint8 NewValue){
@@ -890,9 +942,13 @@ TSharedRef<SWidget> SSequencer::MakeGeneralMenu()
 			],
 			LOCTEXT("ZeroPaddingText", "Zero Pad Frame Numbers"));
 
-		MenuBuilder.AddMenuSeparator();
 		MenuBuilder.AddMenuEntry( FSequencerCommands::Get().ToggleShowGotoBox );
-		MenuBuilder.AddMenuEntry( FSequencerCommands::Get().ToggleLinkCurveEditorTimeRange );
+
+		if (SequencerPtr.Pin()->IsLevelEditorSequencer())
+		{
+			MenuBuilder.AddMenuSeparator();
+			MenuBuilder.AddMenuEntry( FSequencerCommands::Get().FixActorReferences );
+		}
 	}
 
 	return MenuBuilder.MakeWidget();
@@ -927,6 +983,7 @@ TSharedRef<SWidget> SSequencer::MakeSnapMenu()
 	MenuBuilder.BeginSection( "PlayTimeSnapping", LOCTEXT( "SnappingMenuPlayTimeHeader", "Play Time Snapping" ) );
 	{
 		MenuBuilder.AddMenuEntry( FSequencerCommands::Get().ToggleSnapPlayTimeToInterval );
+		MenuBuilder.AddMenuEntry( FSequencerCommands::Get().ToggleSnapPlayTimeToKeys );
 		MenuBuilder.AddMenuEntry( FSequencerCommands::Get().ToggleSnapPlayTimeToDraggedKey );
 		MenuBuilder.AddMenuEntry( FSequencerCommands::Get().ToggleFixedTimeStepPlayback );
 	}
@@ -954,27 +1011,34 @@ TSharedRef<SWidget> SSequencer::MakeAutoKeyMenu()
 
 }
 
-
-TSharedRef<SWidget> SSequencer::MakeTransportControls()
+TSharedRef<SWidget> SSequencer::MakeTimeRange(const TSharedRef<SWidget>& InnerContent, bool bShowWorkingRange, bool bShowViewRange, bool bShowPlaybackRange)
 {
-	FEditorWidgetsModule& EditorWidgetsModule = FModuleManager::Get().LoadModuleChecked<FEditorWidgetsModule>( "EditorWidgets" );
-	TSharedRef<FSequencer> Sequencer = SequencerPtr.Pin().ToSharedRef();
-	FTransportControlArgs TransportControlArgs;
+	ISequencerWidgetsModule& SequencerWidgets = FModuleManager::Get().LoadModuleChecked<ISequencerWidgetsModule>( "SequencerWidgets" );
+
+	EShowRange ShowRange = EShowRange::None;
+	if (bShowWorkingRange)
+	{
+		ShowRange |= EShowRange::WorkingRange;
+	}
+	if (bShowViewRange)
+	{
+		ShowRange |= EShowRange::ViewRange;
+	}
+	if (bShowPlaybackRange)
 	{
-		TransportControlArgs.OnBackwardEnd.BindSP( Sequencer, &FSequencer::OnStepToBeginning );
-		TransportControlArgs.OnBackwardStep.BindSP( Sequencer, &FSequencer::OnStepBackward );
-		TransportControlArgs.OnForwardPlay.BindSP( Sequencer, &FSequencer::OnPlay, true );
-		TransportControlArgs.OnForwardStep.BindSP( Sequencer, &FSequencer::OnStepForward );
-		TransportControlArgs.OnForwardEnd.BindSP( Sequencer, &FSequencer::OnStepToEnd );
-		TransportControlArgs.OnToggleLooping.BindSP( Sequencer, &FSequencer::OnToggleLooping );
-		TransportControlArgs.OnGetLooping.BindSP( Sequencer, &FSequencer::IsLooping );
-		TransportControlArgs.OnGetPlaybackMode.BindSP( Sequencer, &FSequencer::GetPlaybackMode );
+		ShowRange |= EShowRange::PlaybackRange;
 	}
 
-	return EditorWidgetsModule.CreateTransportControl( TransportControlArgs );
+	FTimeRangeArgs Args(
+		ShowRange,
+		TimeSliderController.ToSharedRef(),
+		TAttribute<EVisibility>(this, &SSequencer::GetTimeRangeVisibility),
+		TAttribute<bool>(this, &SSequencer::ShowFrameNumbers),
+		GetZeroPadNumericTypeInterface()
+		);
+	return SequencerWidgets.CreateTimeRange(Args, InnerContent);
 }
 
-
 SSequencer::~SSequencer()
 {
 	USelection::SelectionChangedEvent.RemoveAll(this);
@@ -1026,30 +1090,36 @@ void RestoreSelectionState(const TArray<TSharedRef<FSequencerDisplayNode>>& Disp
 
 void SSequencer::UpdateLayoutTree()
 {
-	// Cache the selected path names so selection can be restored after the update.
-	TSet<FString> SelectedPathNames;
-	for (TSharedRef<const FSequencerDisplayNode> SelectedDisplayNode : SequencerPtr.Pin()->GetSelection().GetSelectedOutlinerNodes().Array())
+	TSharedPtr<FSequencer> Sequencer = SequencerPtr.Pin();
+	if ( Sequencer.IsValid() )
 	{
-		SelectedPathNames.Add(SelectedDisplayNode->GetPathName());
-	}
-
-	// Suspend broadcasting selection changes because we don't want unnecessary rebuilds.
-	SequencerPtr.Pin()->GetSelection().SuspendBroadcast();
+		// Cache the selected path names so selection can be restored after the update.
+		TSet<FString> SelectedPathNames;
+		for (TSharedRef<const FSequencerDisplayNode> SelectedDisplayNode : Sequencer->GetSelection().GetSelectedOutlinerNodes().Array())
+		{
+			FString PathName = SelectedDisplayNode->GetPathName();
+			if ( FName(*PathName).IsNone() == false )
+			{
+				SelectedPathNames.Add(PathName);
+			}
+		}
 
-	// Update the node tree
-	SequencerNodeTree->Update();
+		// Suspend broadcasting selection changes because we don't want unnecessary rebuilds.
+		Sequencer->GetSelection().SuspendBroadcast();
 
-	TreeView->Refresh();
+		// Update the node tree
+		Sequencer->GetNodeTree()->Update();
 
-	// Restore the selection state.
-	RestoreSelectionState(SequencerNodeTree->GetRootNodes(), SelectedPathNames, SequencerPtr.Pin()->GetSelection());	// Update to actor selection.
-	OnActorSelectionChanged(nullptr);
+		// Restore the selection state.
+		RestoreSelectionState(Sequencer->GetNodeTree()->GetRootNodes(), SelectedPathNames, SequencerPtr.Pin()->GetSelection());	// Update to actor selection.
 
-	// This must come after the selection state has been restored so that the curve editor is populated with the correctly selected nodes
-	CurveEditor->SetSequencerNodeTree(SequencerNodeTree);
+		// This must come after the selection state has been restored so that the tree and curve editor are populated with the correctly selected nodes
+		TreeView->Refresh();
+		CurveEditor->SetSequencerNodeTree(Sequencer->GetNodeTree());
 
-	// Continue broadcasting selection changes
-	SequencerPtr.Pin()->GetSelection().ResumeBroadcast();
+		// Continue broadcasting selection changes
+		SequencerPtr.Pin()->GetSelection().ResumeBroadcast();
+	}
 }
 
 
@@ -1080,18 +1150,22 @@ void SSequencer::ResetBreadcrumbs()
 
 void SSequencer::OnOutlinerSearchChanged( const FText& Filter )
 {
-	const FString FilterString = Filter.ToString();
+	TSharedPtr<FSequencer> Sequencer = SequencerPtr.Pin();
+	if ( Sequencer.IsValid())
+	{
+		const FString FilterString = Filter.ToString();
 
-	SequencerNodeTree->FilterNodes(FilterString);
-	TreeView->Refresh();
+		Sequencer->GetNodeTree()->FilterNodes( FilterString );
+		TreeView->Refresh();
 
-	if (FilterString.StartsWith(TEXT("label:")))
-	{
-		LabelBrowser->SetSelectedLabel(FilterString.RightChop(6));
-	}
-	else
-	{
-		LabelBrowser->SetSelectedLabel(FString());
+		if ( FilterString.StartsWith( TEXT( "label:" ) ) )
+		{
+			LabelBrowser->SetSelectedLabel( FilterString.RightChop( 6 ) );
+		}
+		else
+		{
+			LabelBrowser->SetSelectedLabel( FString() );
+		}
 	}
 }
 
@@ -1258,7 +1332,12 @@ void SSequencer::OnAssetsDropped( const FAssetDragDropOp& DragDropOp )
 	if( bObjectAdded )
 	{
 		// Update the sequencers view of the movie scene data when any object is added
-		SequencerRef.NotifyMovieSceneDataChanged();
+		SequencerRef.UpdateRuntimeInstances();
+
+		// Update the tree and synchronize selection
+		UpdateLayoutTree();
+
+		SequencerRef.SynchronizeSequencerSelectionWithExternalSelection();
 	}
 }
 
@@ -1342,69 +1421,6 @@ void SSequencer::OnActorsDropped( FActorDragDropGraphEdOp& DragDropOp )
 }
 
 
-void SSequencer::OnActorSelectionChanged(UObject*)
-{
-	if (!SequencerPtr.IsValid())
-	{
-		return;
-	}
-
-	if (!SequencerPtr.Pin()->IsLevelEditorSequencer())
-	{
-		return;
-	}
-
-	// If the user is selecting within the sequencer, ignore
-	if (bUserIsSelecting)
-	{
-		return;
-	}
-
-	TSet<TSharedRef<FSequencerDisplayNode>> RootNodes(SequencerNodeTree->GetRootNodes());
-
-	// Select the nodes that have runtime objects that are selected in the level
-	for (auto Node : RootNodes)
-	{
-		TSharedRef<FSequencerObjectBindingNode> ObjectBindingNode = StaticCastSharedRef<FSequencerObjectBindingNode>(Node);
-		TArray<UObject*> RuntimeObjects;
-		SequencerPtr.Pin()->GetRuntimeObjects( SequencerPtr.Pin()->GetFocusedMovieSceneSequenceInstance(), ObjectBindingNode->GetObjectBinding(), RuntimeObjects );
-		
-		bool bDoSelect = false;
-		for (int32 RuntimeIndex = 0; RuntimeIndex < RuntimeObjects.Num(); ++RuntimeIndex )
-		{
-			if (GEditor->GetSelectedActors()->IsSelected(RuntimeObjects[RuntimeIndex]))
-			{
-				bDoSelect = true;
-				break;
-			}
-		}
-
-		if (SequencerPtr.Pin()->GetSelection().IsSelected(Node) != bDoSelect)
-		{
-			if (bDoSelect)
-			{
-				SequencerPtr.Pin()->GetSelection().AddToSelection(Node);
-			}
-			else
-			{
-				SequencerPtr.Pin()->GetSelection().RemoveFromSelection(Node);
-			}
-		}
-	}
-
-	const TSet<TSharedRef<FSequencerDisplayNode>>& OutlinerSelection = SequencerPtr.Pin()->GetSelection().GetSelectedOutlinerNodes();
-
-	if (OutlinerSelection.Num() == 1)
-	{
-		for (auto& Node : OutlinerSelection)
-		{
-			TreeView->RequestScrollIntoView(Node);
-			break;
-		}
-	}
-}
-
-
 void SSequencer::OnCrumbClicked(const FSequencerBreadcrumb& Item)
 {
 	if (Item.BreadcrumbType != FSequencerBreadcrumb::ShotType)
@@ -1415,6 +1431,11 @@ void SSequencer::OnCrumbClicked(const FSequencerBreadcrumb& Item)
 		}
 		else
 		{
+			if (Settings->GetShowCurveEditor())
+			{
+				Settings->SetShowCurveEditor(false);
+			}
+
 			SequencerPtr.Pin()->PopToSequenceInstance( Item.MovieSceneInstance.Pin().ToSharedRef() );
 		}
 	}
@@ -1437,25 +1458,29 @@ TArray<FSectionHandle> SSequencer::GetSectionHandles(const TSet<TWeakObjectPtr<U
 {
 	TArray<FSectionHandle> SectionHandles;
 
-	// @todo sequencer: this is potentially slow as it traverses the entire tree - there's scope for optimization here
-	for (auto& Node : SequencerNodeTree->GetRootNodes())
+	TSharedPtr<FSequencer> Sequencer = SequencerPtr.Pin();
+	if (Sequencer.IsValid())
 	{
-		Node->Traverse_ParentFirst([&](FSequencerDisplayNode& InNode) {
-			if (InNode.GetType() == ESequencerNode::Track)
-			{
-				FSequencerTrackNode& TrackNode = static_cast<FSequencerTrackNode&>(InNode);
-
-				const auto& AllSections = TrackNode.GetTrack()->GetAllSections();
-				for (int32 Index = 0; Index < AllSections.Num(); ++Index)
+		// @todo sequencer: this is potentially slow as it traverses the entire tree - there's scope for optimization here
+		for (auto& Node : Sequencer->GetNodeTree()->GetRootNodes())
+		{
+			Node->Traverse_ParentFirst([&](FSequencerDisplayNode& InNode) {
+				if (InNode.GetType() == ESequencerNode::Track)
 				{
-					if (DesiredSections.Contains(TWeakObjectPtr<UMovieSceneSection>(AllSections[Index])))
+					FSequencerTrackNode& TrackNode = static_cast<FSequencerTrackNode&>(InNode);
+
+					const auto& AllSections = TrackNode.GetTrack()->GetAllSections();
+					for (int32 Index = 0; Index < AllSections.Num(); ++Index)
 					{
-						SectionHandles.Emplace(StaticCastSharedRef<FSequencerTrackNode>(TrackNode.AsShared()), Index);
+						if (DesiredSections.Contains(TWeakObjectPtr<UMovieSceneSection>(AllSections[Index])))
+						{
+							SectionHandles.Emplace(StaticCastSharedRef<FSequencerTrackNode>(TrackNode.AsShared()), Index);
+						}
 					}
 				}
-			}
-			return true;
-		});
+				return true;
+			});
+		}
 	}
 
 	return SectionHandles;
@@ -1500,104 +1525,108 @@ void SSequencer::StepToPreviousCameraKey()
 
 void SSequencer::StepToKey(bool bStepToNextKey, bool bCameraOnly)
 {
-	TSet< TSharedRef<FSequencerDisplayNode> > Nodes;
-
-	if (bCameraOnly)
+	TSharedPtr<FSequencer> Sequencer = SequencerPtr.Pin();
+	if ( Sequencer.IsValid() )
 	{
-		TSet<TSharedRef<FSequencerDisplayNode>> RootNodes(SequencerNodeTree->GetRootNodes());
+		TSet< TSharedRef<FSequencerDisplayNode> > Nodes;
 
-		TSet<TWeakObjectPtr<AActor> > LockedActors;
-		for (int32 i = 0; i < GEditor->LevelViewportClients.Num(); ++i)
-		{		
-			FLevelEditorViewportClient* LevelVC = GEditor->LevelViewportClients[i];
-			if (LevelVC && LevelVC->IsPerspective() && LevelVC->GetViewMode() != VMI_Unknown)
+		if ( bCameraOnly )
+		{
+			TSet<TSharedRef<FSequencerDisplayNode>> RootNodes( Sequencer->GetNodeTree()->GetRootNodes() );
+
+			TSet<TWeakObjectPtr<AActor> > LockedActors;
+			for ( int32 i = 0; i < GEditor->LevelViewportClients.Num(); ++i )
 			{
-				TWeakObjectPtr<AActor> ActorLock = LevelVC->GetActiveActorLock();
-				if (ActorLock.IsValid())
+				FLevelEditorViewportClient* LevelVC = GEditor->LevelViewportClients[i];
+				if ( LevelVC && LevelVC->IsPerspective() && LevelVC->GetViewMode() != VMI_Unknown )
 				{
-					LockedActors.Add(ActorLock);
+					TWeakObjectPtr<AActor> ActorLock = LevelVC->GetActiveActorLock();
+					if ( ActorLock.IsValid() )
+					{
+						LockedActors.Add( ActorLock );
+					}
 				}
 			}
-		}
 
-		for (auto RootNode : RootNodes)
-		{
-			TSharedRef<FSequencerObjectBindingNode> ObjectBindingNode = StaticCastSharedRef<FSequencerObjectBindingNode>(RootNode);
-			TArray<UObject*> RuntimeObjects;
-			SequencerPtr.Pin()->GetRuntimeObjects( SequencerPtr.Pin()->GetFocusedMovieSceneSequenceInstance(), ObjectBindingNode->GetObjectBinding(), RuntimeObjects );
-		
-			for (int32 RuntimeIndex = 0; RuntimeIndex < RuntimeObjects.Num(); ++RuntimeIndex )
+			for ( auto RootNode : RootNodes )
 			{
-				AActor* RuntimeActor = Cast<AActor>(RuntimeObjects[RuntimeIndex]);
-				if (RuntimeActor != nullptr && LockedActors.Contains(RuntimeActor))
+				TSharedRef<FSequencerObjectBindingNode> ObjectBindingNode = StaticCastSharedRef<FSequencerObjectBindingNode>( RootNode );
+				TArray<TWeakObjectPtr<UObject>> RuntimeObjects;
+				Sequencer->GetRuntimeObjects( Sequencer->GetFocusedMovieSceneSequenceInstance(), ObjectBindingNode->GetObjectBinding(), RuntimeObjects );
+
+				for ( int32 RuntimeIndex = 0; RuntimeIndex < RuntimeObjects.Num(); ++RuntimeIndex )
 				{
-					Nodes.Add(RootNode);
+					AActor* RuntimeActor = Cast<AActor>( RuntimeObjects[RuntimeIndex].Get() );
+					if ( RuntimeActor != nullptr && LockedActors.Contains( RuntimeActor ) )
+					{
+						Nodes.Add( RootNode );
+					}
 				}
 			}
 		}
-	}
-	else
-	{
-		const TSet< TSharedRef<FSequencerDisplayNode> >& SelectedNodes = SequencerPtr.Pin()->GetSelection().GetSelectedOutlinerNodes();
-		Nodes = SelectedNodes;
-
-		if (Nodes.Num() == 0)
+		else
 		{
-			TSet<TSharedRef<FSequencerDisplayNode>> RootNodes(SequencerNodeTree->GetRootNodes());
-			for (auto RootNode : RootNodes)
+			const TSet< TSharedRef<FSequencerDisplayNode> >& SelectedNodes = Sequencer->GetSelection().GetSelectedOutlinerNodes();
+			Nodes = SelectedNodes;
+
+			if ( Nodes.Num() == 0 )
 			{
-				Nodes.Add(RootNode);
+				TSet<TSharedRef<FSequencerDisplayNode>> RootNodes( Sequencer->GetNodeTree()->GetRootNodes() );
+				for ( auto RootNode : RootNodes )
+				{
+					Nodes.Add( RootNode );
 
-				SequencerHelpers::GetDescendantNodes(RootNode, Nodes);
+					SequencerHelpers::GetDescendantNodes( RootNode, Nodes );
+				}
 			}
 		}
-	}
-		
-	if (Nodes.Num() > 0)
-	{
-		float ClosestKeyDistance = MAX_FLT;
-		float CurrentTime = SequencerPtr.Pin()->GetCurrentLocalTime(*SequencerPtr.Pin()->GetFocusedMovieSceneSequence());
-		float StepToTime = 0;
-		bool StepToKeyFound = false;
 
-		auto It = Nodes.CreateConstIterator();
-		bool bExpand = !(*It).Get().IsExpanded();
-
-		for (auto Node : Nodes)
+		if ( Nodes.Num() > 0 )
 		{
-			TSet<TSharedPtr<IKeyArea>> KeyAreas;
-			SequencerHelpers::GetAllKeyAreas(Node, KeyAreas);
+			float ClosestKeyDistance = MAX_FLT;
+			float CurrentTime = Sequencer->GetCurrentLocalTime( *Sequencer->GetFocusedMovieSceneSequence() );
+			float StepToTime = 0;
+			bool StepToKeyFound = false;
+
+			auto It = Nodes.CreateConstIterator();
+			bool bExpand = !( *It ).Get().IsExpanded();
 
-			for (TSharedPtr<IKeyArea> KeyArea : KeyAreas)
+			for ( auto Node : Nodes )
 			{
-				for (FKeyHandle& KeyHandle : KeyArea->GetUnsortedKeyHandles())
+				TSet<TSharedPtr<IKeyArea>> KeyAreas;
+				SequencerHelpers::GetAllKeyAreas( Node, KeyAreas );
+
+				for ( TSharedPtr<IKeyArea> KeyArea : KeyAreas )
 				{
-					float KeyTime = KeyArea->GetKeyTime(KeyHandle);
-					if (bStepToNextKey)
+					for ( FKeyHandle& KeyHandle : KeyArea->GetUnsortedKeyHandles() )
 					{
-						if (KeyTime > CurrentTime && KeyTime - CurrentTime < ClosestKeyDistance)
+						float KeyTime = KeyArea->GetKeyTime( KeyHandle );
+						if ( bStepToNextKey )
 						{
-							StepToTime = KeyTime;
-							ClosestKeyDistance = KeyTime - CurrentTime;
-							StepToKeyFound = true;
+							if ( KeyTime > CurrentTime && KeyTime - CurrentTime < ClosestKeyDistance )
+							{
+								StepToTime = KeyTime;
+								ClosestKeyDistance = KeyTime - CurrentTime;
+								StepToKeyFound = true;
+							}
 						}
-					}
-					else
-					{
-						if (KeyTime < CurrentTime && CurrentTime - KeyTime < ClosestKeyDistance)
+						else
 						{
-							StepToTime = KeyTime;
-							ClosestKeyDistance = CurrentTime - KeyTime;
-							StepToKeyFound = true;
+							if ( KeyTime < CurrentTime && CurrentTime - KeyTime < ClosestKeyDistance )
+							{
+								StepToTime = KeyTime;
+								ClosestKeyDistance = CurrentTime - KeyTime;
+								StepToKeyFound = true;
+							}
 						}
 					}
 				}
 			}
-		}
 
-		if (StepToKeyFound)
-		{
-			SequencerPtr.Pin()->SetGlobalTime(StepToTime);
+			if ( StepToKeyFound )
+			{
+				Sequencer->SetGlobalTime( StepToTime );
+			}
 		}
 	}
 }
@@ -1683,6 +1712,11 @@ void SSequencer::OnCurveEditorVisibilityChanged()
 
 		// Only zoom horizontally if the editor is visible
 		CurveEditor->SetAllowAutoFrame(Settings->GetShowCurveEditor());
+
+		if (CurveEditor->GetAutoFrame())
+		{
+			CurveEditor->ZoomToFit();
+		}
 	}
 }
 
diff --git a/Engine/Source/Editor/Sequencer/Private/SSequencer.h b/Engine/Source/Editor/Sequencer/Private/SSequencer.h
--- a/Engine/Source/Editor/Sequencer/Private/SSequencer.h
+++ b/Engine/Source/Editor/Sequencer/Private/SSequencer.h
@@ -4,6 +4,7 @@
 
 #include "UniquePtr.h"
 #include "ISequencerEditTool.h"
+#include "IMovieScenePlayer.h"
 #include "SequencerCommonHelpers.h"
 
 
@@ -18,18 +19,25 @@ class FUnloadedClassDragDropOp;
 class IDetailsView;
 class UMovieSceneSection;
 class SSequencerCurveEditor;
+class SSequencerGotoBox;
 class SSequencerLabelBrowser;
 class SSequencerTrackArea;
 class SSequencerTrackOutliner;
 class SSequencerTreeView;
+class FSequencerTimeSliderController;
+
 struct FSectionHandle;
 struct FPaintPlaybackRangeArgs;
 
+
 namespace SequencerLayoutConstants
 {
 	/** The amount to indent child nodes of the layout tree */
 	const float IndentAmount = 8.0f;
 
+	/** Height of each folder node */
+	const float FolderNodeHeight = 20.0f;
+
 	/** Height of each object node */
 	const float ObjectNodeHeight = 20.0f;
 
@@ -92,9 +100,24 @@ public:
 		/** The current clamp range (seconds) */
 		SLATE_ATTRIBUTE( FAnimatedRange, ClampRange )
 
+		/** The in/out selection range */
+		SLATE_ATTRIBUTE( TRange<float>, InOutRange )
+
 		/** The playback range */
 		SLATE_ATTRIBUTE( TRange<float>, PlaybackRange )
-		
+
+		/** The playback status */
+		SLATE_ATTRIBUTE( EMovieScenePlayerStatus::Type, PlaybackStatus )
+
+		/** Called when the user changes the playback range */
+		SLATE_EVENT( FOnRangeChanged, OnInOutRangeChanged )
+
+		/** Called when the user has begun dragging the selection range */
+		SLATE_EVENT( FSimpleDelegate, OnBeginInOutRangeDrag )
+
+		/** Called when the user has finished dragging the selection range */
+		SLATE_EVENT( FSimpleDelegate, OnEndInOutRangeDrag )
+
 		/** Called when the user changes the playback range */
 		SLATE_EVENT( FOnRangeChanged, OnPlaybackRangeChanged )
 
@@ -163,17 +186,6 @@ public:
 	void StepToPreviousCameraKey();
 	void StepToKey(bool bStepToNextKey, bool bCameraOnly);
 
-	/** Whether the user is selecting. Ignore selection changes from the level when the user is selecting. */
-	void SetUserIsSelecting(bool bUserIsSelectingIn)
-	{
-		bUserIsSelecting = bUserIsSelectingIn;
-	}
-
-	bool UserIsSelecting()
-	{
-		return bUserIsSelecting;
-	}
-
 	/** Called when the save button is clicked */
 	void OnSaveMovieSceneClicked();
 
@@ -194,6 +206,9 @@ public:
 
 	/** @return a numeric type interface that will parse and display numbers as frames and times correctly, including any zero padding, if necessary */
 	TSharedRef<INumericTypeInterface<float>> GetZeroPadNumericTypeInterface();
+	
+	/** Access the currently active track area edit tool */
+	const ISequencerEditTool* GetEditTool() const;
 
 public:
 
@@ -226,6 +241,9 @@ private:
 	/** Handles section selection changes. */
 	void HandleSectionSelectionChanged();
 
+	/** Handles changes to the selected outliner nodes. */
+	void HandleOutlinerNodeSelectionChanged();
+
 	/** Empty active timer to ensure Slate ticks during Sequencer playback */
 	EActiveTimerReturnType EnsureSlateTickDuringPlayback(double InCurrentTime, float InDeltaTime);	
 
@@ -249,11 +267,11 @@ private:
 
 public:
 
-	/** Makes and configures a set of the standard UE transport controls. */
-	TSharedRef<SWidget> MakeTransportControls();
+	/** Makes a time range widget with the specified inner content */
+	TSharedRef<SWidget> MakeTimeRange(const TSharedRef<SWidget>& InnerContent, bool bShowWorkingRange, bool bShowViewRange, bool bShowPlaybackRange);
 
 private:
-	
+
 	/**
 	* @return The value of the current time snap interval.
 	*/
@@ -327,12 +345,6 @@ private:
 	 * @param	DragDropOp	Information about the actor(s) that was dropped
 	 */
 	void OnActorsDropped(FActorDragDropGraphEdOp& DragDropOp); 
-	
-	/**
-	* Delegate used when actor selection changes in the level
-	*
-	*/	
-	void OnActorSelectionChanged( UObject* obj );
 
 	/** Called when a breadcrumb is clicked on in the sequencer */
 	void OnCrumbClicked(const FSequencerBreadcrumb& Item);
@@ -378,7 +390,7 @@ public:
 private:
 
 	/** Goto box widget. */
-	TSharedPtr<SWidget> GotoBox;
+	TSharedPtr<SSequencerGotoBox> GotoBox;
 
 	/** Section area widget */
 	TSharedPtr<SSequencerTrackArea> TrackArea;
@@ -389,9 +401,6 @@ private:
 	/** The curve editor. */
 	TSharedPtr<SSequencerCurveEditor> CurveEditor;
 
-	/** Sequencer node tree for movie scene data */
-	TSharedPtr<FSequencerNodeTree> SequencerNodeTree;
-
 	/** The breadcrumb trail widget for this sequencer */
 	TSharedPtr<SBreadcrumbTrail<FSequencerBreadcrumb>> BreadcrumbTrail;
 
@@ -426,7 +435,17 @@ private:
 	TSharedPtr<INumericTypeInterface<float>> NumericTypeInterface;
 	TSharedPtr<INumericTypeInterface<float>> ZeroPadNumericTypeInterface;
 
+	/** Time slider controller for this sequencer */
+	TSharedPtr<FSequencerTimeSliderController> TimeSliderController;
+
 	FOnGetAddMenuContent OnGetAddMenuContent;
+
+	/** Called when the user has begun dragging the in/out selection range */
+	FSimpleDelegate OnBeginInOutRangeDrag;
+
+	/** Called when the user has finished dragging the in/out selection range */
+	FSimpleDelegate OnEndInOutRangeDrag;
+
 	/** Called when the user has begun dragging the playback range */
 	FSimpleDelegate OnBeginPlaybackRangeDrag;
 
diff --git a/Engine/Source/Editor/Sequencer/Private/SSequencerGotoBox.cpp b/Engine/Source/Editor/Sequencer/Private/SSequencerGotoBox.cpp
--- a/Engine/Source/Editor/Sequencer/Private/SSequencerGotoBox.cpp
+++ b/Engine/Source/Editor/Sequencer/Private/SSequencerGotoBox.cpp
@@ -11,17 +11,16 @@
 
 void SSequencerGotoBox::Construct(const FArguments& InArgs, const TSharedRef<FSequencer>& InSequencer, USequencerSettings& InSettings, const TSharedRef<INumericTypeInterface<float>>& InNumericTypeInterface)
 {
-	WasHidden = true;
 	SequencerPtr = InSequencer;
 	Settings = &InSettings;
 	NumericTypeInterface = InNumericTypeInterface;
 
 	ChildSlot
 	[
-		SNew(SBorder)
+		SAssignNew(Border, SBorder)
 			.BorderImage(FEditorStyle::GetBrush("ToolPanel.DarkGroupBorder"))
 			.Padding(6.0f)
-			.Visibility_Lambda([this]() -> EVisibility { return Settings->GetShowGotoBox() ? EVisibility::Visible : EVisibility::Collapsed; })
+			.Visibility(EVisibility::Collapsed)
 			[
 				SNew(SHorizontalBox)
 
@@ -48,26 +47,33 @@ void SSequencerGotoBox::Construct(const FArguments& InArgs, const TSharedRef<FSe
 }
 
 
-void SSequencerGotoBox::Tick(const FGeometry& AllottedGeometry, const double InCurrentTime, const float InDeltaTime)
+void SSequencerGotoBox::ToggleVisibility()
 {
-	if (WasHidden)
+	FSlateApplication& SlateApplication = FSlateApplication::Get();
+
+	if (Border->GetVisibility() == EVisibility::Visible)
+	{
+		SlateApplication.SetAllUserFocus(LastFocusedWidget.Pin(), EFocusCause::Navigation);
+		Border->SetVisibility(EVisibility::Collapsed);
+	}
+	else
 	{
-		FSlateApplication::Get().SetAllUserFocus(EntryBox, EFocusCause::Navigation);
-		WasHidden = false;
+		Border->SetVisibility(EVisibility::Visible);
+		LastFocusedWidget = SlateApplication.GetUserFocusedWidget(0);
+		SlateApplication.SetAllUserFocus(EntryBox, EFocusCause::Navigation);
 	}
 }
 
 
 void SSequencerGotoBox::HandleEntryBoxValueCommitted(float Value, ETextCommit::Type CommitType)
 {
-	Settings->SetShowGotoBox(false);
-	WasHidden = true;
-
-	if (CommitType == ETextCommit::OnCleared)
+	if (CommitType != ETextCommit::OnEnter)
 	{
 		return;
 	}
 
+	ToggleVisibility();
+
 	TSharedPtr<FSequencer> Sequencer = SequencerPtr.Pin();
 
 	// scroll view range if new value is not visible
diff --git a/Engine/Source/Editor/Sequencer/Private/SSequencerGotoBox.h b/Engine/Source/Editor/Sequencer/Private/SSequencerGotoBox.h
--- a/Engine/Source/Editor/Sequencer/Private/SSequencerGotoBox.h
+++ b/Engine/Source/Editor/Sequencer/Private/SSequencerGotoBox.h
@@ -18,13 +18,11 @@ public:
 	SLATE_BEGIN_ARGS(SSequencerGotoBox) { }
 	SLATE_END_ARGS()
 
+	/** Construct the widget. */
 	void Construct(const FArguments& InArgs, const TSharedRef<FSequencer>& InSequencer, USequencerSettings& InSettings, const TSharedRef<INumericTypeInterface<float>>& InNumericTypeInterface);
 
-public:
-
-	// SWidget interface
-
-	virtual void Tick(const FGeometry& AllottedGeometry, const double InCurrentTime, const float InDeltaTime) override;
+	/** Toggle the widget's visibility. */
+	void ToggleVisibility();
 
 private:
 
@@ -33,9 +31,15 @@ private:
 
 private:
 
+	/** The border widget. */
+	TSharedPtr<SWidget> Border;
+
 	/** The entry box widget. */
 	TSharedPtr<SNumericEntryBox<float>> EntryBox;
 
+	/** The widget that focused prior to this goto box. */
+	TWeakPtr<SWidget> LastFocusedWidget;
+
 	/** Numeric type interface used for converting parsing and generating strings from numbers. */
 	TSharedPtr<INumericTypeInterface<float>> NumericTypeInterface;
 
@@ -44,7 +48,4 @@ private:
 
 	/** Cached settings provided to the sequencer itself on creation. */
 	USequencerSettings* Settings;
-
-	/** Whether the box was hidden. */
-	bool WasHidden;
 };
diff --git a/Engine/Source/Editor/Sequencer/Private/SSequencerSection.cpp b/Engine/Source/Editor/Sequencer/Private/SSequencerSection.cpp
--- a/Engine/Source/Editor/Sequencer/Private/SSequencerSection.cpp
+++ b/Engine/Source/Editor/Sequencer/Private/SSequencerSection.cpp
@@ -27,7 +27,7 @@ struct FSequencerSectionPainterImpl : FSequencerSectionPainter
 		CalculateSelectionColor();
 	}
 
-	virtual int32 PaintSectionBackground(const FSequencerSectionTintRegion* Tints, uint32 NumTints) override
+	virtual int32 PaintSectionBackground(const FLinearColor& Tint) override
 	{
 		const ESlateDrawEffect::Type DrawEffects = bParentEnabled
 			? ESlateDrawEffect::None
@@ -39,55 +39,31 @@ struct FSequencerSectionPainterImpl : FSequencerSectionPainter
 
 		FGeometry InfiniteGeometry = SectionGeometry.MakeChild(FVector2D(-100.f, 0.f), FVector2D(SectionGeometry.GetLocalSize().X + 200.f, SectionGeometry.GetLocalSize().Y));
 
+		const FLinearColor FinalTint = FSequencerSectionPainter::BlendColor(Tint);
+
+		FPaintGeometry PaintGeometry = Section.IsInfinite() ?
+				InfiniteGeometry.ToPaintGeometry() :
+				SectionGeometry.ToPaintGeometry();
+				
 		FSlateDrawElement::MakeBox(
 			DrawElements,
 			LayerId++,
-			Section.IsInfinite() ?
-				InfiniteGeometry.ToPaintGeometry() :
-				SectionGeometry.ToPaintGeometry(),
+			PaintGeometry,
 			SectionBackgroundBrush,
 			SectionClippingRect,
 			DrawEffects
 		);
 
-		const FSlateLayoutTransform LayoutTransform = SectionGeometry.GetAccumulatedLayoutTransform();
-
-		//const float TintOpacity = Args.bHovered ? 0.1f : 0.07f;
-		while (NumTints-- > 0)
-		{
-			FSlateRect SubClippingRect = SectionClippingRect;
-
-			if (Tints->Range.IsSet())
-			{
-				const float StartPixel = TimeToPixelConverter.TimeToPixel(Tints->Range->GetLowerBoundValue());
-				const float EndPixel = TimeToPixelConverter.TimeToPixel(Tints->Range->GetUpperBoundValue());
-
-				SubClippingRect = SubClippingRect.IntersectionWith(
-					TransformRect(LayoutTransform, FSlateRect(StartPixel, 0.0f, EndPixel, SectionGeometry.Size.Y))
-					);
-			}
-			
-			// We never show colors with full tint
-			FLinearColor TintColor(Tints->Color);
-			TintColor.A *= 0.4f;
-
-			FSlateDrawElement::MakeBox(
-				DrawElements,
-				LayerId,
-				SectionGeometry.ToPaintGeometry(),
-				SectionBackgroundTintBrush,
-				SubClippingRect.InsetBy(FMargin(1.f)),
-				DrawEffects,
-				TintColor
-			);
-
-			Tints++;
-		}
-
-		++LayerId;
-
-		// This is now invalid
-		Tints = nullptr;
+		FMargin Inset = Section.IsInfinite() ? FMargin(0.f, 1.f) : FMargin(1.f);
+		FSlateDrawElement::MakeBox(
+			DrawElements,
+			LayerId++,
+			PaintGeometry,
+			SectionBackgroundTintBrush,
+			SectionClippingRect.InsetBy(Inset),
+			DrawEffects,
+			FinalTint
+		);
 
 		// Draw the selection hash
 		if (SelectionColor.IsSet())
@@ -113,6 +89,12 @@ struct FSequencerSectionPainterImpl : FSequencerSectionPainter
 
 	void CalculateSelectionColor()
 	{
+		// Don't draw selected if infinite
+		if (Section.IsInfinite())
+		{
+			return;
+		}
+
 		FSequencerSelection& Selection = Sequencer.GetSelection();
 		FSequencerSelectionPreview& SelectionPreview = Sequencer.GetSelectionPreview();
 
@@ -141,11 +123,6 @@ struct FSequencerSectionPainterImpl : FSequencerSectionPainter
 
 			SelectionColor = SelectionColor.GetValue().HSVToLinearRGB();
 		}
-		else if (Selection.GetActiveSelection() != FSequencerSelection::EActiveSelection::KeyAndSection)
-		{
-			// Use an inactive selection color for existing selections that are not active
-			SelectionColor = FEditorStyle::GetSlateColor(SequencerSectionConstants::SelectionInactiveColorName).GetColor(FWidgetStyle());
-		}
 	}
 
 	TOptional<FLinearColor> SelectionColor;
@@ -369,7 +346,6 @@ int32 SSequencerSection::OnPaint( const FPaintArgs& Args, const FGeometry& Allot
 	const bool bLocked = SectionObject.IsLocked();
 	const ESlateDrawEffect::Type DrawEffects = bEnabled ? ESlateDrawEffect::None : ESlateDrawEffect::DisabledEffect;
 
-
 	FGeometry SectionGeometry = MakeSectionGeometryWithoutHandles( AllottedGeometry, SectionInterface );
 
 	FSequencerSectionPainterImpl Painter(ParentSectionArea->GetSequencer(), SectionObject, OutDrawElements, SectionGeometry);
@@ -393,8 +369,27 @@ int32 SSequencerSection::OnPaint( const FPaintArgs& Args, const FGeometry& Allot
 	
 	LayerId = SCompoundWidget::OnPaint( Args, AllottedGeometry, MyClippingRect, OutDrawElements, LayerId, InWidgetStyle, bEnabled );
 
-	FLinearColor SelectionColor = FEditorStyle::GetSlateColor(SequencerSectionConstants::SelectionColorName).GetColor(FWidgetStyle());
-	DrawSectionHandles(AllottedGeometry, MyClippingRect, OutDrawElements, ++LayerId, DrawEffects, SelectionColor);
+	const ISequencerEditTool* EditTool = GetSequencer().GetEditTool();
+	const ISequencerHotspot* Hotspot = EditTool ? EditTool->GetDragHotspot() : nullptr;
+	if (!Hotspot)
+	{
+		Hotspot = GetSequencer().GetHotspot().Get();
+	}
+
+	const bool bShowHandles =
+		IsHovered() ||			// If it's hovered
+		HandleOffsetPx ||		// Or the handles are being drawn outside
+		(						// Or we're dragging one of the handles
+			Hotspot &&
+			(Hotspot->GetType() == ESequencerHotspot::SectionResize_L || Hotspot->GetType() == ESequencerHotspot::SectionResize_R) &&
+			static_cast<const FSectionResizeHotspot*>(Hotspot)->Section == FSectionHandle(ParentSectionArea, SectionIndex)
+		);
+
+	if (bShowHandles)
+	{
+		FLinearColor SelectionColor = FEditorStyle::GetSlateColor(SequencerSectionConstants::SelectionColorName).GetColor(FWidgetStyle());
+		DrawSectionHandles(AllottedGeometry, MyClippingRect, OutDrawElements, ++LayerId, DrawEffects, SelectionColor, Hotspot);
+	}
 
 	FSlateRect KeyClippingRect = Painter.SectionGeometry.GetClippingRect().ExtendBy(FMargin(10.f, 0.f)).IntersectionWith(ParentClippingRect);
 	
@@ -483,13 +478,6 @@ void SSequencerSection::PaintKeys( FSequencerSectionPainter& InPainter, const FW
 	auto& Selection = Sequencer.GetSelection();
 	auto& SelectionPreview = Sequencer.GetSelectionPreview();
 
-	if (Selection.GetActiveSelection() != FSequencerSelection::EActiveSelection::KeyAndSection)
-	{
-		// Selected key color is different when the active selection is not keys
-		SelectedKeyColor = FEditorStyle::GetSlateColor(SelectionInactiveColorName).GetColor(InWidgetStyle)
-			* FLinearColor(.25, .25, .25, 1);  // Make the color a little darker since it's not very visible next to white key frames.
-	}
-
 	const float ThrobScaleValue = GetSelectionThrobValue();
 
 	// draw all keys in each key area
@@ -521,7 +509,7 @@ void SSequencerSection::PaintKeys( FSequencerSectionPainter& InPainter, const FW
 		TOptional<FLinearColor> KeyAreaColor = KeyArea.IsValid() ? KeyArea->GetColor() : TOptional<FLinearColor>();
 
 		// draw a box for the key area 
-		if (KeyAreaColor.IsSet())
+		if (KeyAreaColor.IsSet() && Sequencer.GetSettings()->GetShowChannelColors())
 		{
 			static float BoxThickness = 5.f;
 			FVector2D KeyAreaSize = KeyAreaGeometry.GetLocalSize();
@@ -530,33 +518,46 @@ void SSequencerSection::PaintKeys( FSequencerSectionPainter& InPainter, const FW
 				LayerId + 1,
 				KeyAreaGeometry.ToPaintGeometry(FVector2D(KeyAreaSize.X, BoxThickness), FSlateLayoutTransform(FVector2D(0.f, KeyAreaSize.Y*.5f - BoxThickness*.5f))),
 				StripeOverlayBrush,
-				InPainter.SectionClippingRect,
+				InPainter.SectionClippingRect.InsetBy(1.f),
 				DrawEffects,
 				KeyAreaColor.GetValue()
 			); 
 		}
 
-		if (LayoutElement.GetDisplayNode().IsValid() && LayoutElement.GetDisplayNode()->IsHovered())
+		if (LayoutElement.GetDisplayNode().IsValid())
 		{
-			FSlateDrawElement::MakeBox(
-				InPainter.DrawElements,
-				LayerId + 1,
-				KeyAreaGeometry.ToPaintGeometry(),
-				HighlightBrush,
-				InPainter.SectionClippingRect,
-				DrawEffects,
-				FLinearColor(1.f, 1.f, 1.f, 0.05f)
-			); 
+			FLinearColor HighlightColor;
+			bool bDrawHighlight = false;
+			if (Sequencer.GetSelection().NodeHasSelectedKeysOrSections(LayoutElement.GetDisplayNode().ToSharedRef()))
+			{
+				bDrawHighlight = true;
+				HighlightColor = FLinearColor(1.0f, 1.0f, 1.0f, 0.15f);
+			}
+			else if (LayoutElement.GetDisplayNode()->IsHovered())
+			{
+				bDrawHighlight = true;
+				HighlightColor = FLinearColor(1.0f, 1.0f, 1.0f, 0.05f);
+			}
+
+			if (bDrawHighlight)
+			{
+				FSlateDrawElement::MakeBox(
+					InPainter.DrawElements,
+					LayerId + 1,
+					KeyAreaGeometry.ToPaintGeometry(),
+					HighlightBrush,
+					InPainter.SectionClippingRect,
+					DrawEffects,
+					HighlightColor
+				); 
+			}
 		}
 
 		if (Selection.IsSelected(LayoutElement.GetDisplayNode().ToSharedRef()))
 		{
 			static const FName SelectedTrackTint("Sequencer.Section.SelectedTrackTint");
-			static const FName SelectionColorInactiveName("SelectionColor_Inactive");
 
-			FLinearColor KeyAreaOutlineColor = (Selection.GetActiveSelection() == FSequencerSelection::EActiveSelection::OutlinerNode)
-					? SelectionColor
-					: FEditorStyle::GetSlateColor(SelectionColorInactiveName).GetColor(InWidgetStyle);
+			FLinearColor KeyAreaOutlineColor = SelectionColor;
 
 			FSlateDrawElement::MakeBox(
 				InPainter.DrawElements,
@@ -597,12 +598,14 @@ void SSequencerSection::PaintKeys( FSequencerSectionPainter& InPainter, const FW
 			// determine the key's brush & color
 			const FSlateBrush* KeyBrush;
 			FLinearColor KeyColor;
-			
+			FVector2D FillOffset(0.0f, 0.0f);
+
 			switch (KeyArea->GetKeyInterpMode(KeyHandle))
 			{
 			case RCIM_Linear:
 				KeyBrush = TriangleKeyBrush;
 				KeyColor = FLinearColor(0.0f, 0.617f, 0.449f, 1.0f); // blueish green
+				FillOffset = FVector2D(0.0f, 1.0f);
 				break;
 
 			case RCIM_Constant:
@@ -654,6 +657,7 @@ void SSequencerSection::PaintKeys( FSequencerSectionPainter& InPainter, const FW
 			if (OverrideBrush != nullptr)
 			{
 				KeyBrush = OverrideBrush;
+				FillOffset = SectionInterface->GetKeyBrushOrigin(KeyHandle);
 			}
 
 			// determine draw colors based on hover, selection, etc.
@@ -731,9 +735,10 @@ void SSequencerSection::PaintKeys( FSequencerSectionPainter& InPainter, const FW
 				bSelected ? KeyLayer + 2 : KeyLayer + 1,
 				// Center the key along Y.  Ensure the middle of the key is at the actual key time
 				KeyAreaGeometry.ToPaintGeometry(
+					FillOffset + 
 					FVector2D(
-						KeyPosition - FMath::CeilToFloat(KeySize.X / 2.0f - BrushBorderWidth),
-						((KeyAreaGeometry.Size.Y / 2.0f) - (KeySize.Y / 2.0f - BrushBorderWidth))
+						(KeyPosition - FMath::CeilToFloat((KeySize.X / 2.0f) - BrushBorderWidth)),
+						((KeyAreaGeometry.Size.Y / 2.0f) - ((KeySize.Y / 2.0f) - BrushBorderWidth))
 					),
 					KeySize - 2.0f * BrushBorderWidth
 				),
@@ -747,7 +752,7 @@ void SSequencerSection::PaintKeys( FSequencerSectionPainter& InPainter, const FW
 }
 
 
-void SSequencerSection::DrawSectionHandles( const FGeometry& AllottedGeometry, const FSlateRect& MyClippingRect, FSlateWindowElementList& OutDrawElements, int32 LayerId, ESlateDrawEffect::Type DrawEffects, FLinearColor SelectionColor ) const
+void SSequencerSection::DrawSectionHandles( const FGeometry& AllottedGeometry, const FSlateRect& MyClippingRect, FSlateWindowElementList& OutDrawElements, int32 LayerId, ESlateDrawEffect::Type DrawEffects, FLinearColor SelectionColor, const ISequencerHotspot* Hotspot ) const
 {
 	UMovieSceneSection* Section = SectionInterface->GetSectionObject();
 	if (!Section || Section->IsInfinite())
@@ -762,14 +767,12 @@ void SSequencerSection::DrawSectionHandles( const FGeometry& AllottedGeometry, c
 	bool bRightHandleActive = false;
 
 	// Get the hovered/selected state for the section handles from the hotspot
-	TSharedPtr<ISequencerHotspot> Hotspot = GetSequencer().GetHotspot();
-
-	if (Hotspot.IsValid() && (
+	if (Hotspot && (
 		Hotspot->GetType() == ESequencerHotspot::SectionResize_L ||
 		Hotspot->GetType() == ESequencerHotspot::SectionResize_R))
 	{
-		FSectionResizeHotspot& ResizeHotspot = static_cast<FSectionResizeHotspot&>(*Hotspot.Get());
-		if (ResizeHotspot.Section == FSectionHandle(ParentSectionArea, SectionIndex))
+		const FSectionResizeHotspot* ResizeHotspot = static_cast<const FSectionResizeHotspot*>(Hotspot);
+		if (ResizeHotspot->Section == FSectionHandle(ParentSectionArea, SectionIndex))
 		{
 			bLeftHandleActive = Hotspot->GetType() == ESequencerHotspot::SectionResize_L;
 			bRightHandleActive = Hotspot->GetType() == ESequencerHotspot::SectionResize_R;
diff --git a/Engine/Source/Editor/Sequencer/Private/SSequencerSection.h b/Engine/Source/Editor/Sequencer/Private/SSequencerSection.h
--- a/Engine/Source/Editor/Sequencer/Private/SSequencerSection.h
+++ b/Engine/Source/Editor/Sequencer/Private/SSequencerSection.h
@@ -78,7 +78,7 @@ private:
 	/**
 	 * Draw the section resize handles.
 	 */
-	void DrawSectionHandles( const FGeometry& AllottedGeometry, const FSlateRect& MyClippingRect, FSlateWindowElementList& OutDrawElements, int32 LayerId, ESlateDrawEffect::Type DrawEffects, FLinearColor SelectionColor ) const;
+	void DrawSectionHandles( const FGeometry& AllottedGeometry, const FSlateRect& MyClippingRect, FSlateWindowElementList& OutDrawElements, int32 LayerId, ESlateDrawEffect::Type DrawEffects, FLinearColor SelectionColor, const ISequencerHotspot* Hotspot ) const;
 
 	/** @return the sequencer interface */
 	FSequencer& GetSequencer() const;
diff --git a/Engine/Source/Editor/Sequencer/Private/SSequencerSectionAreaView.cpp b/Engine/Source/Editor/Sequencer/Private/SSequencerSectionAreaView.cpp
--- a/Engine/Source/Editor/Sequencer/Private/SSequencerSectionAreaView.cpp
+++ b/Engine/Source/Editor/Sequencer/Private/SSequencerSectionAreaView.cpp
@@ -126,16 +126,14 @@ int32 SSequencerSectionAreaView::OnPaint( const FPaintArgs& Args, const FGeometr
 	FArrangedChildren ArrangedChildren(EVisibility::Visible);
 	ArrangeChildren(AllottedGeometry, ArrangedChildren);
 
-	int32 MaxLayerId = LayerId;
 	for (int32 ChildIndex = 0; ChildIndex < ArrangedChildren.Num(); ++ChildIndex)
 	{
 		FArrangedWidget& CurWidget = ArrangedChildren[ChildIndex];
 		FSlateRect ChildClipRect = MyClippingRect.IntersectionWith( CurWidget.Geometry.GetClippingRect() );
-		const int32 ThisLayerId = CurWidget.Widget->Paint( Args.WithNewParent(this), CurWidget.Geometry, ChildClipRect, OutDrawElements, LayerId, InWidgetStyle, ShouldBeEnabled( bParentEnabled ) );
-		MaxLayerId = FMath::Max(ThisLayerId, MaxLayerId);
+		LayerId = CurWidget.Widget->Paint( Args.WithNewParent(this), CurWidget.Geometry, ChildClipRect, OutDrawElements, LayerId, InWidgetStyle, ShouldBeEnabled( bParentEnabled ) );
 	}
 
-	return MaxLayerId + 1;
+	return LayerId + 1;
 }
 
 void SSequencerSectionAreaView::Tick(const FGeometry& AllottedGeometry, const double InCurrentTime, const float InDeltaTime)
diff --git a/Engine/Source/Editor/Sequencer/Private/SSequencerTrackArea.h b/Engine/Source/Editor/Sequencer/Private/SSequencerTrackArea.h
--- a/Engine/Source/Editor/Sequencer/Private/SSequencerTrackArea.h
+++ b/Engine/Source/Editor/Sequencer/Private/SSequencerTrackArea.h
@@ -66,6 +66,9 @@ public:
 	/** Assign a tree view to this track area. */
 	void SetTreeView(const TSharedPtr<SSequencerTreeView>& InTreeView);
 
+	/** Access the currently active track area edit tool */
+	const ISequencerEditTool* GetEditTool() const { return EditTool.IsValid() ? EditTool.Get() : nullptr; }
+
 public:
 
 	// SWidget interface
@@ -121,7 +124,7 @@ private:
 	/** Keep a hold of the size of the area so we can maintain zoom levels. */
 	TOptional<FVector2D> SizeLastFrame;
 
-	/** The currently actuve edit tool on this track area */
+	/** The currently active edit tool on this track area */
 	TUniquePtr<ISequencerEditTool> EditTool;
 
 private:
diff --git a/Engine/Source/Editor/Sequencer/Private/SSequencerTrackLane.cpp b/Engine/Source/Editor/Sequencer/Private/SSequencerTrackLane.cpp
--- a/Engine/Source/Editor/Sequencer/Private/SSequencerTrackLane.cpp
+++ b/Engine/Source/Editor/Sequencer/Private/SSequencerTrackLane.cpp
@@ -23,7 +23,6 @@ void DrawLaneRecursive(const TSharedRef<FSequencerDisplayNode>& DisplayNode, con
 {
 	static const FName BorderName("Sequencer.AnimationOutliner.DefaultBorder");
 	static const FName SelectionColorName("SelectionColor");
-	static const FName SelectionColorInactiveName("SelectionColor_Inactive");
 
 	if (DisplayNode->IsHidden())
 	{
@@ -37,11 +36,7 @@ void DrawLaneRecursive(const TSharedRef<FSequencerDisplayNode>& DisplayNode, con
 
 	if (SequencerSelection.IsSelected(DisplayNode))
 	{
-		FLinearColor SelectionColor = FEditorStyle::GetSlateColor(
-			(SequencerSelection.GetActiveSelection() == FSequencerSelection::EActiveSelection::OutlinerNode)
-				? SelectionColorName
-				: SelectionColorInactiveName
-		).GetColor(InWidgetStyle);
+		FLinearColor SelectionColor = FEditorStyle::GetSlateColor(SelectionColorName).GetColor(InWidgetStyle);
 
 		FSlateDrawElement::MakeBox(
 			OutDrawElements,
@@ -57,21 +52,37 @@ void DrawLaneRecursive(const TSharedRef<FSequencerDisplayNode>& DisplayNode, con
 		);
 	}
 
-	// draw hover border
-	else if (DisplayNode->IsHovered())
+	// draw hovered or node has keys or sections selected border
+	else
 	{
-		FSlateDrawElement::MakeBox(
-			OutDrawElements,
-			LayerId,
-			AllottedGeometry.ToPaintGeometry(
-				FVector2D(0, YOffs),
-				FVector2D(AllottedGeometry.Size.X, TotalNodeHeight)
-			),
-			FEditorStyle::GetBrush(BorderName),
-			MyClippingRect,
-			ESlateDrawEffect::None,
-			FLinearColor(1.0f, 1.0f, 1.0f, 0.05f)
-		);
+		FLinearColor HighlightColor;
+		bool bDrawHighlight = false;
+		if (SequencerSelection.NodeHasSelectedKeysOrSections(DisplayNode))
+		{
+			bDrawHighlight = true;
+			HighlightColor = FLinearColor(1.0f, 1.0f, 1.0f, 0.15f);
+		}
+		else if (DisplayNode->IsHovered())
+		{
+			bDrawHighlight = true;
+			HighlightColor = FLinearColor(1.0f, 1.0f, 1.0f, 0.05f);
+		}
+
+		if (bDrawHighlight)
+		{
+			FSlateDrawElement::MakeBox(
+				OutDrawElements,
+				LayerId,
+				AllottedGeometry.ToPaintGeometry(
+					FVector2D(0, YOffs),
+					FVector2D(AllottedGeometry.Size.X, TotalNodeHeight)
+				),
+				FEditorStyle::GetBrush(BorderName),
+				MyClippingRect,
+				ESlateDrawEffect::None,
+				HighlightColor
+			);
+		}
 	}
 
 	YOffs += TotalNodeHeight;
@@ -114,6 +125,22 @@ int32 SSequencerTrackLane::OnPaint(const FPaintArgs& Args, const FGeometry& Allo
 	return SCompoundWidget::OnPaint(Args, AllottedGeometry, MyClippingRect, OutDrawElements, LayerId + 1, InWidgetStyle, bParentEnabled);
 }
 
+void SSequencerTrackLane::Tick(const FGeometry& AllottedGeometry, const double InCurrentTime, const float InDeltaTime)
+{
+	FVector2D ThisFrameDesiredSize = GetDesiredSize();
+
+	if (LastDesiredSize.IsSet() && ThisFrameDesiredSize.Y != LastDesiredSize.GetValue().Y)
+	{
+		TSharedPtr<SSequencerTreeView> PinnedTree = TreeView.Pin();
+		if (PinnedTree.IsValid())
+		{
+			PinnedTree->RequestTreeRefresh();
+		}
+	}
+
+	LastDesiredSize = ThisFrameDesiredSize;
+}
+
 FVector2D SSequencerTrackLane::ComputeDesiredSize(float LayoutScale) const
 {
 	float Height = DisplayNode->GetNodeHeight() + DisplayNode->GetNodePadding().Combined();
diff --git a/Engine/Source/Editor/Sequencer/Private/SSequencerTrackLane.h b/Engine/Source/Editor/Sequencer/Private/SSequencerTrackLane.h
--- a/Engine/Source/Editor/Sequencer/Private/SSequencerTrackLane.h
+++ b/Engine/Source/Editor/Sequencer/Private/SSequencerTrackLane.h
@@ -20,6 +20,7 @@ public:
 	
 	/** Paint this widget */
 	virtual int32 OnPaint(const FPaintArgs& Args, const FGeometry& AllottedGeometry, const FSlateRect& MyClippingRect, FSlateWindowElementList& OutDrawElements, int32 LayerId, const FWidgetStyle& InWidgetStyle, bool bParentEnabled) const override;
+	virtual void Tick(const FGeometry& AllottedGeometry, const double InCurrentTime, const float InDeltaTime) override;
 
 	/** Get the desired physical vertical position of this track lane */
 	float GetPhysicalPosition() const;
@@ -32,4 +33,6 @@ private:
 	TSharedPtr<FSequencerDisplayNode> DisplayNode;
 	/** Pointer back to the tree view for virtual <-> physical space conversions. Important: weak ptr to avoid circular references */
 	TWeakPtr<SSequencerTreeView> TreeView;
+	/** Our desired size last frame */
+	TOptional<FVector2D> LastDesiredSize;
 };
\ No newline at end of file
diff --git a/Engine/Source/Editor/Sequencer/Private/SSequencerTreeView.cpp b/Engine/Source/Editor/Sequencer/Private/SSequencerTreeView.cpp
--- a/Engine/Source/Editor/Sequencer/Private/SSequencerTreeView.cpp
+++ b/Engine/Source/Editor/Sequencer/Private/SSequencerTreeView.cpp
@@ -5,6 +5,7 @@
 #include "SSequencerTreeView.h"
 #include "SSequencerTrackLane.h"
 #include "SSequencerTrackArea.h"
+#include "SequencerDisplayNodeDragDropOp.h"
 
 static FName TrackAreaName = "TrackArea";
 
@@ -54,7 +55,12 @@ void SSequencerTreeViewRow::Construct(const FArguments& InArgs, const TSharedRef
 	Node = InNode;
 	OnGenerateWidgetForColumn = InArgs._OnGenerateWidgetForColumn;
 
-	SMultiColumnTableRow::Construct(SMultiColumnTableRow::FArguments(), OwnerTableView);
+	SMultiColumnTableRow::Construct(
+		SMultiColumnTableRow::FArguments()
+			.OnDragDetected(this, &SSequencerTreeViewRow::OnDragDetected)
+			.OnCanAcceptDrop(this, &SSequencerTreeViewRow::OnCanAcceptDrop)
+			.OnAcceptDrop(this, &SSequencerTreeViewRow::OnAcceptDrop),
+		OwnerTableView);
 }
 
 TSharedRef<SWidget> SSequencerTreeViewRow::GenerateWidgetForColumn(const FName& ColumnId)
@@ -68,6 +74,60 @@ TSharedRef<SWidget> SSequencerTreeViewRow::GenerateWidgetForColumn(const FName&
 	return SNullWidget::NullWidget;
 }
 
+FReply SSequencerTreeViewRow::OnDragDetected( const FGeometry& InGeometry, const FPointerEvent& InPointerEvent )
+{
+	TSharedPtr<FSequencerDisplayNode> DisplayNode = Node.Pin();
+	if ( DisplayNode.IsValid() )
+	{
+		FSequencer& Sequencer = DisplayNode->GetParentTree().GetSequencer();
+		if ( Sequencer.GetSelection().GetSelectedOutlinerNodes().Num() > 0 )
+		{
+			TArray<TSharedRef<FSequencerDisplayNode> > DraggableNodes;
+			for ( const TSharedRef<FSequencerDisplayNode> SelectedNode : Sequencer.GetSelection().GetSelectedOutlinerNodes() )
+			{
+				if ( SelectedNode->CanDrag() )
+				{
+					DraggableNodes.Add(SelectedNode);
+				}
+			}
+
+			TSharedRef<FSequencerDisplayNodeDragDropOp> DragDropOp = MakeShareable( new FSequencerDisplayNodeDragDropOp( DraggableNodes ) );
+			DragDropOp->CurrentHoverText = FText::Format( NSLOCTEXT( "SequencerTreeViewRow", "DefaultDragDropFormat", "Move {0} item(s)" ), FText::AsNumber( DraggableNodes.Num() ) );
+			DragDropOp->SetupDefaults();
+			DragDropOp->Construct();
+			return FReply::Handled().BeginDragDrop( DragDropOp );
+		}
+	}
+	return FReply::Unhandled();
+}
+
+TOptional<EItemDropZone> SSequencerTreeViewRow::OnCanAcceptDrop( const FDragDropEvent& DragDropEvent, EItemDropZone ItemDropZone, FDisplayNodeRef DisplayNode )
+{
+	TSharedPtr<FSequencerDisplayNodeDragDropOp> DragDropOp = DragDropEvent.GetOperationAs<FSequencerDisplayNodeDragDropOp>();
+	if ( DragDropOp.IsValid() )
+	{
+		DragDropOp->ResetToDefaultToolTip();
+		TOptional<EItemDropZone> AllowedDropZone = DisplayNode->CanDrop( *DragDropOp, ItemDropZone );
+		if ( AllowedDropZone.IsSet() == false )
+		{
+			DragDropOp->CurrentIconBrush = FEditorStyle::GetBrush( TEXT( "Graph.ConnectorFeedback.Error" ) );
+		}
+		return AllowedDropZone;
+	}
+	return TOptional<EItemDropZone>();
+}
+
+FReply SSequencerTreeViewRow::OnAcceptDrop( const FDragDropEvent& DragDropEvent, EItemDropZone ItemDropZone, FDisplayNodeRef DisplayNode )
+{
+	TSharedPtr<FSequencerDisplayNodeDragDropOp> DragDropOp = DragDropEvent.GetOperationAs<FSequencerDisplayNodeDragDropOp>();
+	if ( DragDropOp.IsValid())
+	{
+		DisplayNode->Drop( DragDropOp->GetDraggedNodes(), ItemDropZone );
+		return FReply::Handled();
+	}
+	return FReply::Unhandled();
+}
+
 TSharedPtr<FSequencerDisplayNode> SSequencerTreeViewRow::GetDisplayNode() const
 {
 	return Node.Pin();
@@ -87,12 +147,15 @@ void SSequencerTreeView::Construct(const FArguments& InArgs, const TSharedRef<FS
 {
 	SequencerNodeTree = InNodeTree;
 	TrackArea = InTrackArea;
+	bUpdatingSequencerSelection = false;
+	bUpdatingTreeSelection = false;
+	bSequencerSelectionChangeBroadcastWasSupressed = false;
 
 	// We 'leak' these delegates (they'll get cleaned up automatically when the invocation list changes)
 	// It's not safe to attempt their removal in ~SSequencerTreeView because SequencerNodeTree->GetSequencer() may not be valid
 	FSequencer& Sequencer = InNodeTree->GetSequencer();
 	Sequencer.GetSettings()->GetOnShowCurveEditorChanged().AddSP(this, &SSequencerTreeView::UpdateTrackArea);
-	Sequencer.GetSelection().GetOnOutlinerNodeSelectionChanged().AddSP(this, &SSequencerTreeView::OnSequencerSelectionChangedExternally);
+	Sequencer.GetSelection().GetOnOutlinerNodeSelectionChanged().AddSP(this, &SSequencerTreeView::SynchronizeTreeSelectionWithSequencerSelection);
 
 	HeaderRow = SNew(SHeaderRow).Visibility(EVisibility::Collapsed);
 
@@ -102,13 +165,14 @@ void SSequencerTreeView::Construct(const FArguments& InArgs, const TSharedRef<FS
 	(
 		STreeView::FArguments()
 		.TreeItemsSource(&RootNodes)
-		.SelectionMode(ESelectionMode::None)
+		.SelectionMode(ESelectionMode::Multi)
 		.OnGenerateRow(this, &SSequencerTreeView::OnGenerateRow)
 		.OnGetChildren(this, &SSequencerTreeView::OnGetChildren)
 		.HeaderRow(HeaderRow)
 		.ExternalScrollbar(InArgs._ExternalScrollbar)
 		.OnExpansionChanged(this, &SSequencerTreeView::OnExpansionChanged)
 		.AllowOverscroll(EAllowOverscroll::No)
+		.OnContextMenuOpening( this, &SSequencerTreeView::OnContextMenuOpening )
 	);
 }
 
@@ -375,33 +439,124 @@ void SSequencerTreeView::UpdateTrackArea()
 	}
 }
 
-void SSequencerTreeView::OnSequencerSelectionChangedExternally()
+void SSequencerTreeView::SynchronizeTreeSelectionWithSequencerSelection()
 {
-	Private_ClearSelection();
+	if ( bUpdatingSequencerSelection == false )
+	{
+		bUpdatingTreeSelection = true;
+		{
+			Private_ClearSelection();
 
-	FSequencer& Sequencer = SequencerNodeTree->GetSequencer();
-	for (auto& Node : Sequencer.GetSelection().GetSelectedOutlinerNodes())
+			FSequencer& Sequencer = SequencerNodeTree->GetSequencer();
+			for ( auto& Node : Sequencer.GetSelection().GetSelectedOutlinerNodes() )
+			{
+				Private_SetItemSelection( Node, true, false );
+			}
+
+			Private_SignalSelectionChanged( ESelectInfo::Direct );
+		}
+		bUpdatingTreeSelection = false;
+	}
+}
+
+void SSequencerTreeView::Private_SetItemSelection( FDisplayNodeRef TheItem, bool bShouldBeSelected, bool bWasUserDirected )
+{
+	STreeView::Private_SetItemSelection( TheItem, bShouldBeSelected, bWasUserDirected );
+	if ( bUpdatingTreeSelection == false )
+	{
+		// Don't broadcast the sequencer selection change on individual tree changes.  Wait for signal selection changed.
+		FSequencerSelection& SequencerSelection = SequencerNodeTree->GetSequencer().GetSelection();
+		SequencerSelection.SuspendBroadcast();
+		bSequencerSelectionChangeBroadcastWasSupressed = true;
+		if ( bShouldBeSelected )
+		{
+			SequencerSelection.AddToSelection( TheItem );
+		}
+		else
+		{
+			SequencerSelection.RemoveFromSelection( TheItem );
+		}
+		SequencerSelection.ResumeBroadcast();
+	}
+}
+
+
+void SSequencerTreeView::Private_ClearSelection()
+{
+	STreeView::Private_ClearSelection();
+	if ( bUpdatingTreeSelection == false )
 	{
-		Private_SetItemSelection(Node, true, false);
+		// Don't broadcast the sequencer selection change on individual tree changes.  Wait for signal selection changed.
+		FSequencerSelection& SequencerSelection = SequencerNodeTree->GetSequencer().GetSelection();
+		SequencerSelection.SuspendBroadcast();
+		bSequencerSelectionChangeBroadcastWasSupressed = true;
+		SequencerSelection.EmptySelectedOutlinerNodes();
+		SequencerSelection.ResumeBroadcast();
 	}
+}
 
-	Private_SignalSelectionChanged(ESelectInfo::Direct);
+void SSequencerTreeView::Private_SelectRangeFromCurrentTo( FDisplayNodeRef InRangeSelectionEnd )
+{
+	STreeView::Private_SelectRangeFromCurrentTo( InRangeSelectionEnd );
+	if ( bUpdatingTreeSelection == false )
+	{
+		// Don't broadcast the sequencer selection change on individual tree changes.  Wait for signal selection changed.
+		FSequencerSelection& SequencerSelection = SequencerNodeTree->GetSequencer().GetSelection();
+		SequencerSelection.SuspendBroadcast();
+		bSequencerSelectionChangeBroadcastWasSupressed = true;
+		SynchronizeSequencerSelectionWithTreeSelection();
+		SequencerSelection.ResumeBroadcast();
+	}
 }
 
 void SSequencerTreeView::Private_SignalSelectionChanged(ESelectInfo::Type SelectInfo)
 {
-	if (SelectInfo != ESelectInfo::Direct)
+	if ( bUpdatingTreeSelection == false )
+	{
+		bUpdatingSequencerSelection = true;
+		{
+			FSequencerSelection& SequencerSelection = SequencerNodeTree->GetSequencer().GetSelection();
+			SequencerSelection.SuspendBroadcast();
+			bool bSelectionChanged = SynchronizeSequencerSelectionWithTreeSelection();
+			SequencerSelection.ResumeBroadcast();
+			if ( bSequencerSelectionChangeBroadcastWasSupressed || bSelectionChanged )
+			{
+				SequencerSelection.RequestOutlinerNodeSelectionChangedBroadcast();
+				bSequencerSelectionChangeBroadcastWasSupressed = false;
+			}
+		}
+		bUpdatingSequencerSelection = false;
+	}
+
+	STreeView::Private_SignalSelectionChanged(SelectInfo);
+}
+
+bool SSequencerTreeView::SynchronizeSequencerSelectionWithTreeSelection()
+{
+	bool bSelectionChanged = false;
+	const TSet<TSharedRef<FSequencerDisplayNode>>& SequencerSelection = SequencerNodeTree->GetSequencer().GetSelection().GetSelectedOutlinerNodes();
+	if ( SelectedItems.Num() != SequencerSelection.Num() || SelectedItems.Difference( SequencerSelection ).Num() != 0 )
 	{
 		FSequencer& Sequencer = SequencerNodeTree->GetSequencer();
 		FSequencerSelection& Selection = Sequencer.GetSelection();
 		Selection.EmptySelectedOutlinerNodes();
-		for (auto& Item : GetSelectedItems())
+		for ( auto& Item : GetSelectedItems() )
 		{
-			Selection.AddToSelection(Item);
+			Selection.AddToSelection( Item );
 		}
+		bSelectionChanged = true;
 	}
+	return bSelectionChanged;
+}
 
-	STreeView::Private_SignalSelectionChanged(SelectInfo);
+TSharedPtr<SWidget> SSequencerTreeView::OnContextMenuOpening()
+{
+	const TSet<TSharedRef<FSequencerDisplayNode>> SelectedNodes = SequencerNodeTree->GetSequencer().GetSelection().GetSelectedOutlinerNodes();
+	if ( SelectedNodes.Num() == 1 )
+	{
+		return SelectedNodes.Array()[0]->OnSummonContextMenu();
+	}
+	return TSharedPtr<SWidget>();
 }
 
 void SSequencerTreeView::Refresh()
@@ -421,6 +576,12 @@ void SSequencerTreeView::Refresh()
 		}
 	}
 
+	// Force synchronization of selected tree view items here since the tree nodes may have been rebuilt
+	// and the treeview's selection will now be invalid.
+	bUpdatingTreeSelection = true;
+	SynchronizeTreeSelectionWithSequencerSelection();
+	bUpdatingTreeSelection = false;
+
 	RequestTreeRefresh();
 }
 
diff --git a/Engine/Source/Editor/Sequencer/Private/SSequencerTreeView.h b/Engine/Source/Editor/Sequencer/Private/SSequencerTreeView.h
--- a/Engine/Source/Editor/Sequencer/Private/SSequencerTreeView.h
+++ b/Engine/Source/Editor/Sequencer/Private/SSequencerTreeView.h
@@ -109,11 +109,25 @@ protected:
 	/** Called when a node has been expanded or collapsed */
 	void OnExpansionChanged(FDisplayNodeRef InItem, bool bIsExpanded);
 
-	/** Called when the sequencer outliner selection has been changed externally */
-	void OnSequencerSelectionChangedExternally();
+	// Tree selection methods which must be overriden to maintain selection consistency with the rest of sequencer.
+	virtual void Private_SetItemSelection( FDisplayNodeRef TheItem, bool bShouldBeSelected, bool bWasUserDirected = false ) override;
+	virtual void Private_ClearSelection() override;
+	virtual void Private_SelectRangeFromCurrentTo( FDisplayNodeRef InRangeSelectionEnd ) override;
+	virtual void Private_SignalSelectionChanged( ESelectInfo::Type SelectInfo ) override;
 
-	/** Overridden to keep external selection states up to date */
-	virtual void Private_SignalSelectionChanged(ESelectInfo::Type SelectInfo) override;
+private:
+
+	/** Updates the tree selection to match the current sequencer selection. */
+	void SynchronizeTreeSelectionWithSequencerSelection();
+
+	/** 
+	 * Updates the sequencer selection to match the current tree selection.
+	 * @returns Whether or not the sequencer selection was actually modified.
+	 */
+	bool SynchronizeSequencerSelectionWithTreeSelection();
+
+	/** Handles the context menu opening when right clicking on the tree view. */
+	TSharedPtr<SWidget> OnContextMenuOpening();
 
 public:
 	
@@ -178,6 +192,18 @@ private:
 
 	/** A global highlight for the currently hovered tree node hierarchy */
 	TOptional<FHighlightRegion> HighlightRegion;
+
+	/** When true, the sequencer selection is being updated from a change in the tree seleciton. */
+	bool bUpdatingSequencerSelection;
+
+	/** When true, the tree selection is being updated from a change in the sequencer selection. */
+	bool bUpdatingTreeSelection;
+
+	/**
+	 * When true a sequencer selection change broadcast was suppressed when updating sequencer selection
+	 * due to the tree selection changing.
+	 */
+	bool bSequencerSelectionChangeBroadcastWasSupressed;
 };
 
 /** Widget that represents a row in the sequencer's tree control. */
@@ -209,6 +235,15 @@ public:
 	/** Overridden from SMultiColumnTableRow.  Generates a widget for this column of the tree row. */
 	virtual TSharedRef<SWidget> GenerateWidgetForColumn(const FName& ColumnId) override;
 
+	/** Called whenever a drag is detected by the tree view. */
+	FReply OnDragDetected(const FGeometry& InGeometry, const FPointerEvent& InPointerEvent);
+
+	/** Called to determine whether a current drag operation is valid for this row. */
+	TOptional<EItemDropZone> OnCanAcceptDrop( const FDragDropEvent& DragDropEvent, EItemDropZone ItemDropZone, FDisplayNodeRef DisplayNode);
+
+	/** Called to complete a drag and drop onto this drop. */
+	FReply OnAcceptDrop( const FDragDropEvent& DragDropEvent, EItemDropZone ItemDropZone, FDisplayNodeRef DisplayNode );
+
 private:
 
 	/** Cached reference to a track lane that we relate to. This keeps the track lane alive (it's a weak widget) as long as we are in view. */
diff --git a/Engine/Source/Editor/Sequencer/Private/SSequencerTreeViewBox.h b/Engine/Source/Editor/Sequencer/Private/SSequencerTreeViewBox.h
--- a/Engine/Source/Editor/Sequencer/Private/SSequencerTreeViewBox.h
+++ b/Engine/Source/Editor/Sequencer/Private/SSequencerTreeViewBox.h
@@ -29,17 +29,7 @@ public:
 
 			if (Selection.GetSelectedOutlinerNodes().Num())
 			{
-				if (SequencerWidget.Pin().IsValid())
-				{
-					SequencerWidget.Pin()->SetUserIsSelecting(true);
-				}
-
 				Selection.EmptySelectedOutlinerNodes();
-
-				if (SequencerWidget.Pin().IsValid())
-				{
-					SequencerWidget.Pin()->SetUserIsSelecting(false);
-				}
 				return FReply::Handled();
 			}
 		}
