Copying //UE4/Dev-Sequencer to //UE4/Main

diff --git a/Engine/Source/Editor/Sequencer/Private/Sequencer.cpp b/Engine/Source/Editor/Sequencer/Private/Sequencer.cpp
--- a/Engine/Source/Editor/Sequencer/Private/Sequencer.cpp
+++ b/Engine/Source/Editor/Sequencer/Private/Sequencer.cpp
@@ -6,6 +6,7 @@
 #include "SequencerDetailKeyframeHandler.h"
 #include "MovieSceneSequence.h"
 #include "MovieScene.h"
+#include "MovieSceneFolder.h"
 #include "Engine/LevelScriptBlueprint.h"
 #include "Editor/PropertyEditor/Public/PropertyEditorModule.h"
 #include "Editor/SequencerWidgets/Public/ITimeSlider.h"
@@ -14,6 +15,7 @@
 #include "Editor/LevelEditor/Public/ILevelEditor.h"
 #include "Editor/LevelEditor/Public/ILevelViewport.h"
 #include "Editor/WorkspaceMenuStructure/Public/WorkspaceMenuStructureModule.h"
+#include "EditorSupportDelegates.h"
 #include "SSequencer.h"
 #include "SSequencerTreeView.h"
 #include "ScopedTransaction.h"
@@ -26,6 +28,7 @@
 #include "MovieSceneTrackEditor.h"
 #include "MovieSceneToolHelpers.h"
 #include "ScopedTransaction.h"
+#include "MovieSceneBoolSection.h"
 #include "MovieSceneCameraCutSection.h"
 #include "MovieScene3DTransformSection.h"
 #include "MovieSceneSubSection.h"
@@ -56,7 +59,11 @@
 #include "ContentBrowserModule.h"
 #include "IContentBrowserSingleton.h"
 #include "PackageTools.h"
-
+#include "GroupedKeyArea.h"
+#include "VirtualTrackArea.h"
+#include "SequencerUtilities.h"
+#include "MovieSceneCinematicShotTrack.h"
+#include "ISequenceRecorder.h"
 
 #define LOCTEXT_NAMESPACE "Sequencer"
 
@@ -65,14 +72,15 @@ DEFINE_LOG_CATEGORY(LogSequencer);
 
 namespace
 {
-	TSharedPtr<FSequencer> ActiveSequencer;
+	TWeakPtr<FSequencer> ActiveSequencerPtr;
 }
 
 
 void FSequencer::InitSequencer(const FSequencerInitParams& InitParams, const TSharedRef<ISequencerObjectChangeListener>& InObjectChangeListener, const TSharedRef<IDetailKeyframeHandler>& InDetailKeyframeHandler, const TArray<FOnCreateTrackEditor>& TrackEditorDelegates)
 {
 	bIsEditingWithinLevelEditor = InitParams.bEditWithinLevelEditor;
-		
+	SilentModeCount = 0;
+
 	if (InitParams.SpawnRegister.IsValid())
 	{
 		SpawnRegister = InitParams.SpawnRegister;
@@ -86,12 +94,12 @@ void FSequencer::InitSequencer(const FSequencerInitParams& InitParams, const TSh
 	// If this sequencer edits the level, close out the existing active sequencer and mark this as the active sequencer.
 	if (bIsEditingWithinLevelEditor)
 	{
-		if (ActiveSequencer.IsValid())
+		if (ActiveSequencerPtr.IsValid())
 		{
-			ActiveSequencer->Close();
+			ActiveSequencerPtr.Pin()->Close();
 		}
 
-		ActiveSequencer = SharedThis(this);
+		ActiveSequencerPtr = SharedThis(this);
 
 		// Register for saving the level so that the state of the scene can be restored before saving and updated after saving.
 		FEditorDelegates::PreSaveWorld.AddSP(this, &FSequencer::OnPreSaveWorld);
@@ -101,6 +109,8 @@ void FSequencer::InitSequencer(const FSequencerInitParams& InitParams, const TSh
 
 		FEditorDelegates::NewCurrentLevel.AddSP(this, &FSequencer::OnNewCurrentLevel);
 		FEditorDelegates::OnMapOpened.AddSP(this, &FSequencer::OnMapOpened);
+
+		USelection::SelectionChangedEvent.AddSP( this, &FSequencer::OnActorSelectionChanged );
 	}
 
 	Settings = USequencerSettingsContainer::GetOrCreate<USequencerSettings>(*InitParams.ViewParams.UniqueName);
@@ -123,7 +133,12 @@ void FSequencer::InitSequencer(const FSequencerInitParams& InitParams, const TSh
 	SequencerWidget = SNew( SSequencer, SharedThis( this ) )
 		.ViewRange( this, &FSequencer::GetViewRange )
 		.ClampRange( this, &FSequencer::GetClampRange )
+		.InOutRange( this, &FSequencer::GetInOutRange )
 		.PlaybackRange( this, &FSequencer::GetPlaybackRange )
+		.PlaybackStatus( this, &FSequencer::GetPlaybackStatus )
+		.OnInOutRangeChanged( this, &FSequencer::SetInOutRange )
+		.OnBeginInOutRangeDrag( this, &FSequencer::OnBeginInOutRangeDrag )
+		.OnEndInOutRangeDrag( this, &FSequencer::OnEndInOutRangeDrag )
 		.OnPlaybackRangeChanged( this, &FSequencer::SetPlaybackRange )
 		.OnBeginPlaybackRangeDrag( this, &FSequencer::OnBeginPlaybackRangeDrag )
 		.OnEndPlaybackRangeDrag( this, &FSequencer::OnEndPlaybackRangeDrag )
@@ -188,17 +203,21 @@ void FSequencer::InitSequencer(const FSequencerInitParams& InitParams, const TSh
 
 FSequencer::FSequencer()
 	: SequencerCommandBindings( new FUICommandList )
+	, SequencerSharedBindings( new FUICommandList )
 	, TargetViewRange(0.f, 5.f)
 	, LastViewRange(0.f, 5.f)
 	, ViewRangeBeforeZoom(TRange<float>::Empty())
 	, PlaybackState( EMovieScenePlayerStatus::Stopped )
-	, PlaybackStartRealTime ( 0.0 )
-	, NumContinuousFixedTimeStepFrames ( 0 )
 	, ScrubPosition( 0.0f )
 	, bPerspectiveViewportPossessionEnabled( true )
 	, bPerspectiveViewportCameraCutEnabled( false )
 	, bIsEditingWithinLevelEditor( false )
 	, bNeedTreeRefresh( false )
+	, bWasClosed( false )
+	, NodeTree( MakeShareable( new FSequencerNodeTree( *this ) ) )
+	, bUpdatingSequencerSelection( false )
+	, bUpdatingExternalSelection( false )
+	, OldMaxTickRate(GEngine->GetMaxFPS())
 {
 	Selection.GetOnOutlinerNodeSelectionChanged().AddRaw(this, &FSequencer::OnSelectedOutlinerNodesChanged);
 }
@@ -214,18 +233,26 @@ FSequencer::~FSequencer()
 		DetachTransportControlsFromViewports();
 	}
 
-	for (auto TrackEditor : TrackEditors)
+	// When a new level sequence asset is opened while an existing one is open, the original level sequence asset is closed, 
+	// the new level sequence asset is initialized, and then the original level sequence asset is destroyed. In that case, 
+	// the track editors should not be released a second time. However, a level sequence asset is opened and closed without 
+	// opening another, the track editors still need to be released because Close() wil not be called.
+	//
+	if (!bWasClosed)
 	{
-		TrackEditor->OnRelease();
+		for (auto TrackEditor : TrackEditors)
+		{
+			TrackEditor->OnRelease();
+		}
+		TrackEditors.Empty();
 	}
-	TrackEditors.Empty();
 	SequencerWidget.Reset();
 }
 
 
 void FSequencer::Close()
 {
-	if (ActiveSequencer.IsValid() && ActiveSequencer.Get() == this)
+	if (ActiveSequencerPtr.IsValid() && ActiveSequencerPtr.Pin().Get() == this)
 	{
 		RootMovieSceneSequenceInstance->RestoreState(*this);
 
@@ -241,7 +268,8 @@ void FSequencer::Close()
 		{
 			TrackEditor->OnRelease();
 		}
-		ActiveSequencer.Reset();
+
+		ActiveSequencerPtr.Reset();
 
 		if (GLevelEditorModeTools().IsModeActive(FSequencerEdMode::EM_SequencerMode))
 		{
@@ -260,11 +288,14 @@ void FSequencer::Close()
 	}
 	
 	SequencerWidget.Reset();
+	bWasClosed = true;
 }
 
 
 void FSequencer::Tick(float InDeltaTime)
 {
+	Selection.Tick();
+
 	if (bNeedTreeRefresh)
 	{
 		// @todo - Sequencer Will be called too often
@@ -273,7 +304,7 @@ void FSequencer::Tick(float InDeltaTime)
 		SequencerWidget->UpdateLayoutTree();
 		bNeedTreeRefresh = false;
 	}
-	
+
 	static const float AutoScrollFactor = 0.1f;
 
 	// Animate the autoscroll offset if it's set
@@ -290,9 +321,23 @@ void FSequencer::Tick(float InDeltaTime)
 		SetGlobalTimeDirectly(GetGlobalTime() + Offset);
 	}
 
-	float NewTime = GetGlobalTime() + InDeltaTime * GWorld->GetWorldSettings()->MatineeTimeDilation;
+	// override max frame rate
+	if (PlaybackState == EMovieScenePlayerStatus::Playing)
+	{
+		const float TimeSnapInterval = Settings->GetTimeSnapInterval();
+
+		if (SequencerSnapValues::IsTimeSnapIntervalFrameRate(TimeSnapInterval) && Settings->GetFixedTimeStepPlayback())
+		{
+			GEngine->SetMaxFPS(1.f / TimeSnapInterval);
+		}
+		else
+		{
+			GEngine->SetMaxFPS(OldMaxTickRate);
+		}
+	}
 
-	ConstrainFixedTimeStepFrameRate();
+	// calculate new global time
+	float NewTime = GetGlobalTime() + InDeltaTime * GWorld->GetWorldSettings()->MatineeTimeDilation;
 
 	if (PlaybackState == EMovieScenePlayerStatus::Playing ||
 		PlaybackState == EMovieScenePlayerStatus::Recording)
@@ -305,63 +350,80 @@ void FSequencer::Tick(float InDeltaTime)
 	{
 		TrackEditors[EditorIndex]->Tick(InDeltaTime);
 	}
-}
 
+	if (!IsInSilentMode())
+	{
+		PostTickRenderStateFixup();
+	}
 
-void FSequencer::ConstrainFixedTimeStepFrameRate()
-{
-	if( SequencerSnapValues::IsTimeSnapIntervalFrameRate(Settings->GetTimeSnapInterval()) && Settings->GetFixedTimeStepPlayback() )
+	ISequenceRecorder& SequenceRecorder = FModuleManager::LoadModuleChecked<ISequenceRecorder>("SequenceRecorder");
+	if(SequenceRecorder.IsRecording())
 	{
-		// NOTE: Its important that PlaybackStartRealTime and NumContinuousFixedTimeStepFrames are reset
-		//    when anything timing-related changes, like FApp::GetFixedDeltaTime() or playback direction.
+		UMovieSceneSubSection* Section = UMovieSceneSubSection::GetRecordingSection();
+		if(Section != nullptr)
+		{
+			Section->SetEndTime(Section->GetStartTime() + SequenceRecorder.GetCurrentRecordingLength());
+		}
+	}
+}
 
-		double CurRealTime = FPlatformTime::Seconds();
+void FSequencer::PostTickRenderStateFixup()
+{
+	if (PlaybackState == EMovieScenePlayerStatus::Jumping || PlaybackState == EMovieScenePlayerStatus::Stepping)
+	{
+		const bool bApplyFixup = true;
 
-		// Minor hack to handle changes to world TimeDilation.  We reset our frame rate gate state
-		// when we detect a change to time dilation.
-		static float LastTimeDilation =  GWorld->GetWorldSettings()->GetEffectiveTimeDilation();
-		if( LastTimeDilation != GWorld->GetWorldSettings()->GetEffectiveTimeDilation() )
+		if (!bApplyFixup)
 		{
-			// Looks like time dilation has changed!
-			NumContinuousFixedTimeStepFrames = 0;
-			PlaybackStartRealTime = CurRealTime;
-
-			LastTimeDilation = GWorld->GetWorldSettings()->GetEffectiveTimeDilation();
+			// Clear any delayed state changes regardless
+			SetPlaybackStatus(EMovieScenePlayerStatus::Stopped);
 		}
+		else
+		{
+			const float Time = GetGlobalTime();
+			const float FixedTickDelta = 1.f / 60.f;
+			const float StepDelta = (PlaybackState == EMovieScenePlayerStatus::Jumping) ? 0.f : FixedTickDelta;
+			
+			// Find a time for our previous frame, if there is one
+			float TimeLowerBound = GetPlaybackRange().GetLowerBoundValue();
 
-		// How long should have it taken to get to the current frame?
-		const float PlaybackSpeed = 1.f;
-		const double ExpectedPlaybackTime =
-			NumContinuousFixedTimeStepFrames * FApp::GetFixedDeltaTime() * PlaybackSpeed;
+			if (UMovieSceneSequence* MovieSequence = GetFocusedMovieSceneSequence())
+			{
+				UMovieSceneSubTrack* SubTrack = Cast<UMovieSceneSubTrack>(MovieSequence->GetMovieScene()->FindMasterTrack(UMovieSceneCinematicShotTrack::StaticClass()));
+				if (!SubTrack)
+				{
+					SubTrack = Cast<UMovieSceneSubTrack>(MovieSequence->GetMovieScene()->FindMasterTrack(UMovieSceneSubTrack::StaticClass()));
+				}
 
-		// How long has it been (in real-time) since we started playback?
-		double RealTimeSincePlaybackStarted = CurRealTime - PlaybackStartRealTime;
+				if (SubTrack)
+				{
+					for (UMovieSceneSection* Section : SubTrack->GetAllSections())
+					{
+						if (Section->IsInfinite() || Section->IsTimeWithinSection(Time))
+						{
+							UMovieSceneSubSection* SubSection = CastChecked<UMovieSceneSubSection>(Section);
 
-		// If we're way ahead of schedule (more than 5 ms), then we'll perform a long sleep
-		float WaitTime = ExpectedPlaybackTime - RealTimeSincePlaybackStarted;
-		if( WaitTime > 5 / 1000.0f )
-		{
-			FPlatformProcess::Sleep( WaitTime - 3 / 1000.0f );
+							TimeLowerBound = FMath::Max(TimeLowerBound, SubSection->GetStartTime());
+							break;
+						}
+					}
+				}
+			}
 
-			// Update timing info after our little snooze
-			CurRealTime = FPlatformTime::Seconds();
-			RealTimeSincePlaybackStarted = CurRealTime - PlaybackStartRealTime;
-			WaitTime = ExpectedPlaybackTime - RealTimeSincePlaybackStarted;
-		}
+			float PreviousTime = FMath::Max(TimeLowerBound, Time - StepDelta);
+			PreviousTime = FMath::Min(PreviousTime, Time);
 
-		while( RealTimeSincePlaybackStarted < ExpectedPlaybackTime )
-		{
-			// OK, we're running ahead of schedule so we need to wait a bit before the next frame
-			FPlatformProcess::Sleep( 0.0f );
+			// Force previous frame in sequence to re-render
+			SetPlaybackStatus(EMovieScenePlayerStatus::Stopped);
+			SetGlobalTimeDirectly(PreviousTime);
+			Tick(StepDelta);
 
-			// Check the time again
-			CurRealTime = FPlatformTime::Seconds();
-			RealTimeSincePlaybackStarted = CurRealTime - PlaybackStartRealTime;
-			WaitTime = ExpectedPlaybackTime - RealTimeSincePlaybackStarted;
-		}
+			GWorld->SendAllEndOfFrameUpdates();
 
-		// Increment number of continuous fixed time step frames
-		++NumContinuousFixedTimeStepFrames;
+			SetPlaybackStatus(EMovieScenePlayerStatus::Stopped);
+			SetGlobalTimeDirectly(Time);
+			Tick(StepDelta);
+		}
 	}
 }
 
@@ -460,8 +522,7 @@ void FSequencer::FocusSequenceInstance(UMovieSceneSubSection& InSubSection)
 
 	SequencerWidget->UpdateBreadcrumbs();
 
-	const bool bAllowSnappingToFrames = true;
-	SetGlobalTime(AbsoluteShotPosition, bAllowSnappingToFrames);
+	SetGlobalTime(AbsoluteShotPosition, ESnapTimeMode::STM_Interval);
 }
 
 
@@ -509,7 +570,18 @@ UObject* FSequencer::GetPlaybackContext() const
 {
 	if (bIsEditingWithinLevelEditor)
 	{
-		return GWorld;
+		if(GEditor && GEditor->PlayWorld)
+		{
+			return GEditor->PlayWorld;
+		}
+		else if(GEditor && GEditor->EditorWorld)
+		{
+			return GEditor->EditorWorld;
+		}
+		else
+		{
+			return GWorld;	
+		}
 	}
 
 	return nullptr;
@@ -581,6 +653,25 @@ void FSequencer::GetKeyedProperties(UObject& Object, FMovieSceneSequenceInstance
 	}
 }
 
+void FSequencer::GetKeysFromSelection(TUniquePtr<ISequencerKeyCollection>& KeyCollection)
+{
+	if (!KeyCollection.IsValid())
+	{
+		KeyCollection.Reset(new FGroupedKeyCollection);
+	}
+
+	TArray<FSequencerDisplayNode*> SelectedNodes;
+	for (const TSharedRef<FSequencerDisplayNode>& Node : Selection.GetSelectedOutlinerNodes())
+	{
+		SelectedNodes.Add(&Node.Get());
+	}
+
+	// Anything within .5 pixel's worth of time is a duplicate as far as we're concerened
+	FVirtualTrackArea TrackArea = SequencerWidget->GetVirtualTrackArea();;
+	const float DuplicateThreshold = (TrackArea.PixelToTime(0.f) - TrackArea.PixelToTime(1.f)) * .5f;
+
+	KeyCollection->InitializeRecursive(SelectedNodes, DuplicateThreshold);
+}
 
 TSharedRef<FMovieSceneSequenceInstance> FSequencer::GetSequenceInstanceForSection(UMovieSceneSection& Section) const
 {
@@ -644,6 +735,8 @@ void FSequencer::DeleteSections(const TSet<TWeakObjectPtr<UMovieSceneSection>>&
 		// Full refresh required just in case the last section was removed from any track.
 		NotifyMovieSceneDataChanged();
 	}
+
+	SequencerHelpers::ValidateNodesWithSelectedKeysOrSections(*this);
 }
 
 
@@ -671,6 +764,7 @@ void FSequencer::DeleteSelectedKeys()
 	}
 
 	Selection.EmptySelectedKeys();
+	SequencerHelpers::ValidateNodesWithSelectedKeysOrSections(*this);
 }
 
 
@@ -774,7 +868,7 @@ void FSequencer::NotifyMapChanged( class UWorld* NewWorld, EMapChangeType MapCha
 	// @todo sequencer: We should only wipe/respawn puppets that are affected by the world that is being changed! (multi-UWorld support)
 	if( ( MapChangeType == EMapChangeType::LoadMap || MapChangeType == EMapChangeType::NewMap || MapChangeType == EMapChangeType::TearDownWorld) )
 	{
-		SpawnRegister->DestroyAllObjects(*this);
+		SpawnRegister->CleanUp(*this);
 
 		NotifyMovieSceneDataChanged();
 	}
@@ -795,7 +889,11 @@ void FSequencer::OnActorsDropped( const TArray<TWeakObjectPtr<AActor> >& Actors
 			FGuid ExistingGuid = SequenceInstanceStack.Top()->FindObjectId(*Actor);
 			if (!ExistingGuid.IsValid())
 			{
-				CreateBinding(*Actor, Actor->GetActorLabel());
+				FGuid PossessableGuid = CreateBinding(*Actor, Actor->GetActorLabel());
+
+				UpdateRuntimeInstances();
+
+				OnActorAddedToSequencerEvent.Broadcast(Actor, PossessableGuid);
 			}
 			bPossessableAdded = true;
 		}
@@ -811,7 +909,7 @@ void FSequencer::OnActorsDropped( const TArray<TWeakObjectPtr<AActor> >& Actors
 void FSequencer::NotifyMovieSceneDataChanged()
 {
 	LabelManager.SetMovieScene(SequenceInstanceStack.Top()->GetSequence()->GetMovieScene());
-	PlaybackState = EMovieScenePlayerStatus::Stopped;
+	SetPlaybackStatus(EMovieScenePlayerStatus::Stopped);
 	bNeedTreeRefresh = true;
 
 	UpdatePlaybackRange();
@@ -834,43 +932,93 @@ FAnimatedRange FSequencer::GetViewRange() const
 
 FAnimatedRange FSequencer::GetClampRange() const
 {
-	UMovieScene* FocussedMovieScene = GetFocusedMovieSceneSequence()->GetMovieScene();
-	return FocussedMovieScene->GetEditorData().WorkingRange;
+	UMovieScene* FocusedMovieScene = GetFocusedMovieSceneSequence()->GetMovieScene();
+	return FocusedMovieScene->GetEditorData().WorkingRange;
 }
 
+
 void FSequencer::SetClampRange(TRange<float> InNewClampRange)
 {
-	UMovieScene* FocussedMovieScene = GetFocusedMovieSceneSequence()->GetMovieScene();
-	FocussedMovieScene->GetEditorData().WorkingRange = InNewClampRange;
+	UMovieScene* FocusedMovieScene = GetFocusedMovieSceneSequence()->GetMovieScene();
+	FocusedMovieScene->GetEditorData().WorkingRange = InNewClampRange;
 }
 
-TRange<float> FSequencer::GetPlaybackRange() const
+
+TRange<float> FSequencer::GetInOutRange() const
 {
-	return GetFocusedMovieSceneSequence()->GetMovieScene()->GetPlaybackRange();
+	return GetFocusedMovieSceneSequence()->GetMovieScene()->GetInOutRange();
+}
+
+
+void FSequencer::SetInOutRange(TRange<float> Range)
+{
+	const FScopedTransaction Transaction(LOCTEXT("SetPlaybackRange_Transaction", "Set In/Out Range"));
+	UMovieScene* FocusedMovieScene = GetFocusedMovieSceneSequence()->GetMovieScene();
+
+	FocusedMovieScene->SetInOutRange(Range);
 }
 
 
-void FSequencer::SetPlaybackRange(TRange<float> InNewPlaybackRange)
+void FSequencer::SetInOutRangeEnd()
 {
-	if (ensure(InNewPlaybackRange.HasLowerBound() && InNewPlaybackRange.HasUpperBound() && !InNewPlaybackRange.IsDegenerate()))
+	const float GlobalTime = GetGlobalTime();
+
+	if (GetInOutRange().GetLowerBoundValue() >= GlobalTime)
+	{
+		SetInOutRange(TRange<float>(GlobalTime));
+	}
+	else
 	{
-		const FScopedTransaction Transaction(LOCTEXT("SetPlaybackRange_Transaction", "Set playback range"));
+		SetInOutRange(TRange<float>(GetInOutRange().GetLowerBoundValue(), GlobalTime));
+	}
+}
+
 
-		UMovieScene* FocussedMovieScene = GetFocusedMovieSceneSequence()->GetMovieScene();
-		FocussedMovieScene->SetPlaybackRange(InNewPlaybackRange.GetLowerBoundValue(), InNewPlaybackRange.GetUpperBoundValue());
+void FSequencer::SetInOutRangeStart()
+{
+	const float GlobalTime = GetGlobalTime();
+
+	if (GetInOutRange().GetUpperBoundValue() <= GlobalTime)
+	{
+		SetInOutRange(TRange<float>(GlobalTime));
+	}
+	else
+	{
+		SetInOutRange(TRange<float>(GlobalTime, GetInOutRange().GetUpperBoundValue()));
 	}
 }
 
 
-void FSequencer::SetStartPlaybackRange()
+TRange<float> FSequencer::GetPlaybackRange() const
+{
+	return GetFocusedMovieSceneSequence()->GetMovieScene()->GetPlaybackRange();
+}
+
+
+void FSequencer::SetPlaybackRange(TRange<float> Range)
 {
-	SetPlaybackRange(TRange<float>(GetGlobalTime(), GetPlaybackRange().GetUpperBoundValue()));
+	if (ensure(Range.HasLowerBound() && Range.HasUpperBound() && !Range.IsDegenerate()))
+	{
+		const FScopedTransaction Transaction(LOCTEXT("SetPlaybackRange_Transaction", "Set Playback Range"));
+
+		UMovieScene* FocusedMovieScene = GetFocusedMovieSceneSequence()->GetMovieScene();
+		FocusedMovieScene->SetPlaybackRange(Range.GetLowerBoundValue(), Range.GetUpperBoundValue());
+	}
 }
 
 
-void FSequencer::SetEndPlaybackRange()
+void FSequencer::ResetViewRange()
 {
-	SetPlaybackRange(TRange<float>( GetPlaybackRange().GetLowerBoundValue(), GetGlobalTime()));
+	float InRange = GetPlaybackRange().GetLowerBoundValue();
+	float OutRange = GetPlaybackRange().GetUpperBoundValue();
+	const float OutputViewSize = OutRange - InRange;
+	const float OutputChange = OutputViewSize * 0.1f;
+
+	InRange -= OutputChange;
+	OutRange += OutputChange;
+
+	SetClampRange(TRange<float>(InRange, OutRange));	
+	SetViewRange(TRange<float>(InRange, OutRange), EViewRangeInterpolation::Animated);
 }
 
 
@@ -888,7 +1036,7 @@ void FSequencer::ZoomViewRange(float InZoomDelta)
 
 	if (NewViewOutputMin < NewViewOutputMax)
 	{
-		SetViewRange(TRange<float>(NewViewOutputMin, NewViewOutputMax), EViewRangeInterpolation::Immediate);
+		SetViewRange(TRange<float>(NewViewOutputMin, NewViewOutputMax), EViewRangeInterpolation::Animated);
 	}
 }
 
@@ -922,7 +1070,9 @@ void FSequencer::UpdatePlaybackRange()
 											FMath::Max(NewBounds.GetUpperBoundValue(), MovieSceneSection->GetRange().GetUpperBoundValue()) );
 			}
 
-			FocusedMovieScene->SetPlaybackRange(NewBounds.GetLowerBoundValue(), NewBounds.GetUpperBoundValue());
+			// When the playback range is determined by the section bounds, don't mark the change in the playback range otherwise the scene will be marked dirty
+			const bool bAlwaysMarkDirty = false;
+			FocusedMovieScene->SetPlaybackRange(NewBounds.GetLowerBoundValue(), NewBounds.GetUpperBoundValue(), bAlwaysMarkDirty);
 		}
 	}
 	else
@@ -995,6 +1145,18 @@ void FSequencer::SetKeyInterpolation(EMovieSceneKeyInterpolation InKeyInterpolat
 }
 
 
+bool FSequencer::GetInfiniteKeyAreas() const
+{
+	return Settings->GetInfiniteKeyAreas();
+}
+
+
+void FSequencer::SetInfiniteKeyAreas(bool bInfiniteKeyAreas)
+{
+	Settings->SetInfiniteKeyAreas(bInfiniteKeyAreas);
+}
+
+
 bool FSequencer::IsRecordingLive() const 
 {
 	return PlaybackState == EMovieScenePlayerStatus::Recording && GIsPlayInEditorWorld;
@@ -1009,13 +1171,13 @@ float FSequencer::GetCurrentLocalTime( UMovieSceneSequence& InMovieSceneSequence
 }
 
 
-float FSequencer::GetGlobalTime()
+float FSequencer::GetGlobalTime() const
 {
 	return ScrubPosition;
 }
 
 
-void FSequencer::SetGlobalTime( float NewTime, const bool& bAllowSnappingToFrames )
+void FSequencer::SetGlobalTime( float NewTime, ESnapTimeMode SnapTimeMode )
 {
 	if (IsAutoScrollEnabled())
 	{
@@ -1033,19 +1195,24 @@ void FSequencer::SetGlobalTime( float NewTime, const bool& bAllowSnappingToFrame
 		}
 	}
 
-	SetGlobalTimeDirectly(NewTime, bAllowSnappingToFrames);
+	SetGlobalTimeDirectly(NewTime, SnapTimeMode);
 }
 
 
-void FSequencer::SetGlobalTimeDirectly( float NewTime, const bool& bAllowSnappingToFrames  )
+void FSequencer::SetGlobalTimeDirectly( float NewTime, ESnapTimeMode SnapTimeMode )
 {
 	float LastTime = ScrubPosition;
 
-	if (bAllowSnappingToFrames && Settings->GetIsSnapEnabled())
+	if ((SnapTimeMode & ESnapTimeMode::STM_Interval) && Settings->GetIsSnapEnabled())
 	{
 		NewTime = Settings->SnapTimeToInterval(NewTime);
 	}
 
+	if ((SnapTimeMode & ESnapTimeMode::STM_Keys) && (Settings->GetSnapPlayTimeToKeys() || FSlateApplication::Get().GetModifierKeys().IsShiftDown()))
+	{
+		NewTime = FindNearestKey(NewTime);
+	}
+
 	// Update the position
 	ScrubPosition = NewTime;
 
@@ -1057,15 +1224,19 @@ void FSequencer::SetGlobalTimeDirectly( float NewTime, const bool& bAllowSnappin
 	{
 		if (!GetPlaybackRange().Contains(NewTime))
 		{
-			SpawnRegister->DestroyAllOwnedObjects(*this);
+			SpawnRegister->OnMasterSequenceExpired(*this);
 		}
 	}
 
 	// If realtime is off, this needs to be called to update the pivot location when scrubbing.
 	GUnrealEd->UpdatePivotLocationForSelection();
-	GEditor->RedrawLevelEditingViewports();
 
-	OnGlobalTimeChangedDelegate.Broadcast();
+	if (!IsInSilentMode())
+	{
+		// Redraw
+		FEditorSupportDelegates::RedrawAllViewports.Broadcast();
+		OnGlobalTimeChangedDelegate.Broadcast();
+	}
 }
 
 
@@ -1147,6 +1318,11 @@ void FSequencer::SetPerspectiveViewportCameraCutEnabled(bool bEnabled)
 	bPerspectiveViewportCameraCutEnabled = bEnabled;
 }
 
+ISequencer::FOnActorAddedToSequencer& FSequencer::OnActorAddedToSequencer()
+{
+	return OnActorAddedToSequencerEvent;
+}
+
 ISequencer::FOnCameraCut& FSequencer::OnCameraCut()
 {
 	return OnCameraCutEvent;
@@ -1162,9 +1338,9 @@ TSharedRef<INumericTypeInterface<float>> FSequencer::GetZeroPadNumericTypeInterf
 	return SequencerWidget->GetZeroPadNumericTypeInterface();
 }
 
-TSharedRef<SWidget> FSequencer::MakeTransportControls()
+TSharedRef<SWidget> FSequencer::MakeTimeRange(const TSharedRef<SWidget>& InnerContent, bool bShowWorkingRange, bool bShowViewRange, bool bShowPlaybackRange)
 {
-	return SequencerWidget->MakeTransportControls();
+	return SequencerWidget->MakeTimeRange(InnerContent, bShowWorkingRange, bShowViewRange, bShowPlaybackRange);
 }
 
 /** Attempt to find an object binding ID that relates to an unspawned spawnable object */
@@ -1263,6 +1439,7 @@ FGuid FSequencer::GetHandleToObject( UObject* Object, bool bCreateHandleIfMissin
 		AActor* PossessedActor = Cast<AActor>(Object);
 
 		ObjectGuid = CreateBinding(*Object, PossessedActor != nullptr ? PossessedActor->GetActorLabel() : Object->GetName());
+
 		NotifyMovieSceneDataChanged();
 	}
 	
@@ -1276,7 +1453,7 @@ ISequencerObjectChangeListener& FSequencer::GetObjectChangeListener()
 }
 
 
-void FSequencer::GetRuntimeObjects( TSharedRef<FMovieSceneSequenceInstance> MovieSceneInstance, const FGuid& ObjectHandle, TArray< UObject*>& OutObjects ) const
+void FSequencer::GetRuntimeObjects( TSharedRef<FMovieSceneSequenceInstance> MovieSceneInstance, const FGuid& ObjectHandle, TArray<TWeakObjectPtr<UObject>>& OutObjects ) const
 {
 	UObject* FoundObject = MovieSceneInstance->FindObject(ObjectHandle, *this);
 	if (FoundObject)
@@ -1286,9 +1463,16 @@ void FSequencer::GetRuntimeObjects( TSharedRef<FMovieSceneSequenceInstance> Movi
 }
 
 
-void FSequencer::UpdateCameraCut(UObject* CameraObject, UObject* UnlockIfCameraObject) const
+void FSequencer::UpdateCameraCut(UObject* CameraObject, UObject* UnlockIfCameraObject, bool bJumpCut) const
 {
-	OnCameraCutEvent.Broadcast(CameraObject);
+	OnCameraCutEvent.Broadcast(CameraObject, bJumpCut);
+
+	// tell the camera we cut
+	UCameraComponent* const CamComp = MovieSceneHelpers::CameraComponentFromRuntimeObject(CameraObject);
+	if (CamComp)
+	{
+		CamComp->NotifyCameraCut();
+	}
 
 	if (!IsPerspectiveViewportCameraCutEnabled())
 	{
@@ -1306,7 +1490,7 @@ void FSequencer::UpdateCameraCut(UObject* CameraObject, UObject* UnlockIfCameraO
 
 		if ((CameraObject != nullptr) || LevelVC->IsLockedToActor(UnlockIfCameraActor))
 		{
-			UpdatePreviewLevelViewportClientFromCameraCut(*LevelVC, CameraObject);
+			UpdatePreviewLevelViewportClientFromCameraCut(*LevelVC, CameraObject, bJumpCut);
 		}
 	}
 }
@@ -1373,6 +1557,16 @@ EMovieScenePlayerStatus::Type FSequencer::GetPlaybackStatus() const
 void FSequencer::SetPlaybackStatus(EMovieScenePlayerStatus::Type InPlaybackStatus)
 {
 	PlaybackState = InPlaybackStatus;
+
+	// backup or restore tick rate
+	if (InPlaybackStatus == EMovieScenePlayerStatus::Playing)
+	{
+		OldMaxTickRate = GEngine->GetMaxFPS();
+	}
+	else
+	{
+		GEngine->SetMaxFPS(OldMaxTickRate);
+	}
 }
 
 
@@ -1389,7 +1583,7 @@ void FSequencer::RemoveMovieSceneInstance( UMovieSceneSection& MovieSceneSection
 {
 	UMovieSceneSequence* Sequence = InstanceToRemove->GetSequence();
 
-	SpawnRegister->DestroyObjectsSpawnedByInstance(*InstanceToRemove, *this);
+	SpawnRegister->CleanUpSequence(*InstanceToRemove, *this);
 
 	SequenceInstanceBySection.Remove( &MovieSceneSection );
 }
@@ -1430,13 +1624,200 @@ void FSequencer::UpdateRuntimeInstances()
 	{
 		if (!GetPlaybackRange().Contains(ScrubPosition))
 		{
-			SpawnRegister->DestroyAllOwnedObjects(*this);
+			SpawnRegister->OnMasterSequenceExpired(*this);
 		}
 	}
 
 	// If realtime is off, this needs to be called to update the pivot location when scrubbing.
 	GUnrealEd->UpdatePivotLocationForSelection();
-	GEditor->RedrawLevelEditingViewports();
+		
+	// Redraw
+	FEditorSupportDelegates::RedrawAllViewports.Broadcast();
+}
+
+TSharedRef<SWidget> FSequencer::MakeTransportControls(bool bExtended)
+{
+	FEditorWidgetsModule& EditorWidgetsModule = FModuleManager::Get().LoadModuleChecked<FEditorWidgetsModule>( "EditorWidgets" );
+
+	FTransportControlArgs TransportControlArgs;
+	{
+		TransportControlArgs.OnBackwardEnd.BindSP( this, &FSequencer::OnStepToBeginning );
+		TransportControlArgs.OnBackwardStep.BindSP( this, &FSequencer::OnStepBackward );
+		TransportControlArgs.OnForwardPlay.BindSP( this, &FSequencer::OnPlay, true );
+		TransportControlArgs.OnForwardStep.BindSP( this, &FSequencer::OnStepForward );
+		TransportControlArgs.OnForwardEnd.BindSP( this, &FSequencer::OnStepToEnd );
+		TransportControlArgs.OnToggleLooping.BindSP( this, &FSequencer::OnToggleLooping );
+		TransportControlArgs.OnGetLooping.BindSP( this, &FSequencer::IsLooping );
+		TransportControlArgs.OnGetPlaybackMode.BindSP( this, &FSequencer::GetPlaybackMode );
+
+		if(bExtended)
+		{
+			TransportControlArgs.WidgetsToCreate.Add(FTransportControlWidget(FOnMakeTransportWidget::CreateSP(this, &FSequencer::OnCreateTransportSetPlaybackStart)));
+		}
+		TransportControlArgs.WidgetsToCreate.Add(FTransportControlWidget(ETransportControlWidgetType::BackwardEnd));
+		if(bExtended)
+		{
+			TransportControlArgs.WidgetsToCreate.Add(FTransportControlWidget(FOnMakeTransportWidget::CreateSP(this, &FSequencer::OnCreateTransportJumpToPreviousKey)));
+		}
+		TransportControlArgs.WidgetsToCreate.Add(FTransportControlWidget(ETransportControlWidgetType::BackwardStep));
+		TransportControlArgs.WidgetsToCreate.Add(FTransportControlWidget(ETransportControlWidgetType::ForwardPlay));
+		TransportControlArgs.WidgetsToCreate.Add(FTransportControlWidget(FOnMakeTransportWidget::CreateSP(this, &FSequencer::OnCreateTransportRecord)));
+		TransportControlArgs.WidgetsToCreate.Add(FTransportControlWidget(ETransportControlWidgetType::ForwardStep));
+		if(bExtended)
+		{
+			TransportControlArgs.WidgetsToCreate.Add(FTransportControlWidget(FOnMakeTransportWidget::CreateSP(this, &FSequencer::OnCreateTransportJumpToNextKey)));
+		}
+		TransportControlArgs.WidgetsToCreate.Add(FTransportControlWidget(ETransportControlWidgetType::ForwardEnd));
+		if(bExtended)
+		{
+			TransportControlArgs.WidgetsToCreate.Add(FTransportControlWidget(FOnMakeTransportWidget::CreateSP(this, &FSequencer::OnCreateTransportSetPlaybackEnd)));
+		}
+		TransportControlArgs.WidgetsToCreate.Add(FTransportControlWidget(ETransportControlWidgetType::Loop));
+	}
+
+	return EditorWidgetsModule.CreateTransportControl( TransportControlArgs );
+}
+
+TSharedRef<SWidget> FSequencer::OnCreateTransportSetPlaybackStart()
+{
+	return SNew(SButton)
+		.OnClicked(this, &FSequencer::SetPlaybackStart)
+		.ToolTipText(LOCTEXT("SetPlayStart_Tooltip", "Set playback start to the current position"))
+		.ButtonStyle(FEditorStyle::Get(), "Sequencer.Transport.SetPlayStart")
+		.ContentPadding(2.0f);
+}
+
+TSharedRef<SWidget> FSequencer::OnCreateTransportJumpToPreviousKey()
+{
+	return SNew(SButton)
+		.OnClicked(this, &FSequencer::JumpToPreviousKey)
+		.ToolTipText(LOCTEXT("JumpToPreviousKey_Tooltip", "Jump to the previous key in the selected track(s)"))
+		.ButtonStyle(FEditorStyle::Get(), "Sequencer.Transport.JumpToPreviousKey")
+		.ContentPadding(2.0f);
+}
+
+TSharedRef<SWidget> FSequencer::OnCreateTransportJumpToNextKey()
+{
+	return SNew(SButton)
+		.OnClicked(this, &FSequencer::JumpToNextKey)
+		.ToolTipText(LOCTEXT("JumpToNextKey_Tooltip", "Jump to the next key in the selected track(s)"))
+		.ButtonStyle(FEditorStyle::Get(), "Sequencer.Transport.JumpToNextKey")
+		.ContentPadding(2.0f);
+}
+
+TSharedRef<SWidget> FSequencer::OnCreateTransportSetPlaybackEnd()
+{
+	return SNew(SButton)
+		.OnClicked(this, &FSequencer::SetPlaybackEnd)
+		.ToolTipText(LOCTEXT("SetPlayEnd_Tooltip", "Set playback end to the current position"))
+		.ButtonStyle(FEditorStyle::Get(), "Sequencer.Transport.SetPlayEnd")
+		.ContentPadding(2.0f);
+}
+
+TSharedRef<SWidget> FSequencer::OnCreateTransportRecord()
+{
+	ISequenceRecorder& SequenceRecorder = FModuleManager::LoadModuleChecked<ISequenceRecorder>("SequenceRecorder");
+
+	return SNew(SButton)
+		.OnClicked(this, &FSequencer::OnRecord)
+		.ToolTipText_Lambda([&](){ return SequenceRecorder.IsRecording() ? LOCTEXT("StopRecord_Tooltip", "Stop recording current sub-track.") : LOCTEXT("Record_Tooltip", "Record the primed sequence sub-track."); })
+		.ButtonStyle(FEditorStyle::Get(), "Animation.Record")
+		.ContentPadding(2.0f)
+		.Visibility(this, &FSequencer::GetRecordButtonVisibility);
+}
+
+namespace FaderConstants
+{	
+	/** The opacity when we are hovered */
+	const float HoveredOpacity = 1.0f;
+	/** The opacity when we are not hovered */
+	const float NonHoveredOpacity = 0.75f;
+	/** The amount of time spent actually fading in or out */
+	const float FadeTime = 0.15f;
+}
+
+/** Wrapper widget allowing us to fade widgets in and out on hover state */
+class SFader : public SBorder
+{
+public:
+	SLATE_BEGIN_ARGS(SFader)
+		: _Content()
+	{}
+
+	SLATE_DEFAULT_SLOT(FArguments, Content)
+
+	SLATE_END_ARGS()
+
+	void Construct(const FArguments& InArgs)
+	{
+		FadeInSequence = FCurveSequence(0.0f, FaderConstants::FadeTime);
+		FadeOutSequence = FCurveSequence(0.0f, FaderConstants::FadeTime);
+		FadeOutSequence.JumpToEnd();
+
+		SBorder::Construct(SBorder::FArguments()
+			.BorderImage(FCoreStyle::Get().GetBrush("NoBorder"))
+			.Padding(0.0f)
+			.VAlign(VAlign_Center)
+			.ColorAndOpacity(this, &SFader::GetColorAndOpacity)
+			.Content()
+			[
+				InArgs._Content.Widget
+			]);
+	}
+
+	FLinearColor GetColorAndOpacity() const
+	{
+		FLinearColor Color = FLinearColor::White;
+	
+		if(FadeOutSequence.IsPlaying() || !bIsHovered)
+		{
+			Color.A = FMath::Lerp(FaderConstants::HoveredOpacity, FaderConstants::NonHoveredOpacity, FadeOutSequence.GetLerp());
+		}
+		else
+		{
+			Color.A = FMath::Lerp(FaderConstants::NonHoveredOpacity, FaderConstants::HoveredOpacity, FadeInSequence.GetLerp());
+		}
+
+		return Color;
+	}
+
+	virtual void OnMouseEnter(const FGeometry& MyGeometry, const FPointerEvent& MouseEvent) override
+	{
+		if(!FSlateApplication::Get().IsUsingHighPrecisionMouseMovment())
+		{
+			bIsHovered = true;
+			if(FadeOutSequence.IsPlaying())
+			{
+				// Fade out is already playing so just force the fade in curve to the end so we don't have a "pop" 
+				// effect from quickly resetting the alpha
+				FadeInSequence.JumpToEnd();
+			}
+			else
+			{
+				FadeInSequence.Play(AsShared());
+			}
+		}
+	}
+
+	virtual void OnMouseLeave(const FPointerEvent& MouseEvent) override
+	{
+		if(!FSlateApplication::Get().IsUsingHighPrecisionMouseMovment())
+		{
+			bIsHovered = false;
+			FadeOutSequence.Play(AsShared());
+		}
+	}
+
+private:
+	/** Curve sequence for fading out the widget */
+	FCurveSequence FadeOutSequence;
+	/** Curve sequence for fading in the widget */
+	FCurveSequence FadeInSequence;
+};
+
+void FSequencer::SetViewportTransportControlsVisibility(bool bVisible)
+{
+	Settings->SetShowViewportTransportControls(bVisible);
 }
 
 void FSequencer::AttachTransportControlsToViewports()
@@ -1448,26 +1829,21 @@ void FSequencer::AttachTransportControlsToViewports()
 		return;
 	}
 
+	// register level editor viewport menu extenders
+	ViewMenuExtender = FLevelEditorModule::FLevelEditorMenuExtender::CreateSP(this, &FSequencer::OnExtendLevelEditorViewMenu);
+	Module->GetAllLevelViewportOptionsMenuExtenders().Add(ViewMenuExtender);		
+	LevelEditorExtenderDelegateHandle = Module->GetAllLevelViewportOptionsMenuExtenders().Last().GetHandle();
+
 	TSharedPtr<ILevelEditor> LevelEditor = Module->GetFirstLevelEditor();
 	const TArray< TSharedPtr<ILevelViewport> >& LevelViewports = LevelEditor->GetViewports();
 		
 	FEditorWidgetsModule& EditorWidgetsModule = FModuleManager::Get().LoadModuleChecked<FEditorWidgetsModule>( "EditorWidgets" );
 
-	FTransportControlArgs TransportControlArgs;
-	TransportControlArgs.OnForwardPlay.BindSP(this, &FSequencer::OnPlay, true);
-	TransportControlArgs.OnForwardStep.BindSP(this, &FSequencer::OnStepForward);
-	TransportControlArgs.OnBackwardStep.BindSP(this, &FSequencer::OnStepBackward);
-	TransportControlArgs.OnForwardEnd.BindSP(this, &FSequencer::OnStepToEnd);
-	TransportControlArgs.OnBackwardEnd.BindSP(this, &FSequencer::OnStepToBeginning);
-	TransportControlArgs.OnToggleLooping.BindSP(this, &FSequencer::OnToggleLooping);
-	TransportControlArgs.OnGetLooping.BindSP(this, &FSequencer::IsLooping);
-	TransportControlArgs.OnGetPlaybackMode.BindSP(this, &FSequencer::GetPlaybackMode);
-
 	for (int32 i = 0; i < LevelViewports.Num(); ++i)
 	{
 		const TSharedPtr<ILevelViewport>& LevelViewport = LevelViewports[i];
 			
-		TSharedPtr<SWidget> TransportControl =
+		TSharedRef<SHorizontalBox> TransportControl =
 			SNew(SHorizontalBox)
 				.Visibility(EVisibility::SelfHitTestInvisible)
 
@@ -1477,18 +1853,40 @@ void FSequencer::AttachTransportControlsToViewports()
 				.VAlign(VAlign_Bottom)
 				.Padding(4.f)
 				[
-					SNew(SBorder)
+					SNew(SFader)
+					.Content()
+					[
+						SNew(SBorder)
 						.Padding(4.f)
 						.Cursor( EMouseCursor::Default )
-						.BorderImage( FEditorStyle::GetBrush( "FilledBorder" ) )
+						.BorderImage( FEditorStyle::GetBrush( "ToolPanel.GroupBorder" ) )
 						.Visibility(this, &FSequencer::GetTransportControlVisibility, LevelViewport)
 						.Content()
 						[
-							EditorWidgetsModule.CreateTransportControl(TransportControlArgs)
+							SNew(SHorizontalBox)
+							+SHorizontalBox::Slot()
+							.AutoWidth()
+							.HAlign(HAlign_Center)
+							.VAlign(VAlign_Center)
+							[
+								MakeTransportControls(false)
+							]
+							+SHorizontalBox::Slot()
+							.AutoWidth()
+							.HAlign(HAlign_Center)
+							.VAlign(VAlign_Top)
+							[
+								SNew(SButton)
+								.ButtonStyle(&FEditorStyle::Get().GetWidgetStyle< FButtonStyle >("Sequencer.Transport.CloseButton"))
+								.ToolTipText(LOCTEXT("CloseTransportControlsToolTip", "Hide the transport controls. You can re-enable transport controls from the viewport menu."))
+								.OnClicked_Lambda([this]() { Settings->SetShowViewportTransportControls(!Settings->GetShowViewportTransportControls()); return FReply::Handled(); })
+								
+							]
 						]
+					]
 				];
 
-		LevelViewport->AddOverlayWidget(TransportControl.ToSharedRef());
+		LevelViewport->AddOverlayWidget(TransportControl);
 
 		TransportControls.Add(LevelViewport, TransportControl);
 	}
@@ -1516,14 +1914,45 @@ void FSequencer::DetachTransportControlsFromViewports()
 				}
 			}
 		}
+
+		Module->GetAllLevelViewportOptionsMenuExtenders().RemoveAll([&](const FLevelEditorModule::FLevelEditorMenuExtender& Delegate) {
+			return Delegate.GetHandle() == LevelEditorExtenderDelegateHandle;
+		});
 	}
 }
 
+TSharedRef<FExtender> FSequencer::OnExtendLevelEditorViewMenu(const TSharedRef<FUICommandList> CommandList)
+{
+	TSharedRef<FExtender> Extender(new FExtender());
+
+	Extender->AddMenuExtension(
+		"LevelViewportViewportOptions2",
+		EExtensionHook::First,
+		NULL,
+		FMenuExtensionDelegate::CreateSP(this, &FSequencer::CreateTransportToggleMenuEntry));
+
+	return Extender;
+}
+
+void FSequencer::CreateTransportToggleMenuEntry(FMenuBuilder& MenuBuilder)
+{
+	MenuBuilder.AddMenuEntry(
+		LOCTEXT("ShowTransportControls", "Show Transport Controls"),
+		LOCTEXT("ShowTransportControlsToolTip", "Show or hide the Sequencer transport controls when a sequence is active."),
+		FSlateIcon(),
+		FUIAction(
+			FExecuteAction::CreateLambda([this](){ Settings->SetShowViewportTransportControls(!Settings->GetShowViewportTransportControls()); }),
+			FCanExecuteAction(), 
+			FGetActionCheckState::CreateLambda([this](){ return Settings->GetShowViewportTransportControls() ? ECheckBoxState::Checked : ECheckBoxState::Unchecked; })),
+		NAME_None,
+		EUserInterfaceActionType::ToggleButton);
+}
 
 EVisibility FSequencer::GetTransportControlVisibility(TSharedPtr<ILevelViewport> LevelViewport) const
 {
 	FLevelEditorViewportClient& ViewportClient = LevelViewport->GetLevelViewportClient();
-	return ViewportClient.ViewportType == LVT_Perspective &&
+	return Settings->GetShowViewportTransportControls() &&
+		ViewportClient.ViewportType == LVT_Perspective &&
 		ViewportClient.AllowsCinematicPreview() ? EVisibility::Visible : EVisibility::Collapsed;
 }
 
@@ -1532,11 +1961,10 @@ FReply FSequencer::OnPlay(bool bTogglePlay)
 	if( (PlaybackState == EMovieScenePlayerStatus::Playing ||
 		 PlaybackState == EMovieScenePlayerStatus::Recording) && bTogglePlay )
 	{
-		PlaybackState = EMovieScenePlayerStatus::Stopped;
+		SetPlaybackStatus(EMovieScenePlayerStatus::Stopped);
 
 		// Snap time when stopping play
-		const bool bAllowSnappingToFrames = true;
-		SetGlobalTime(ScrubPosition, bAllowSnappingToFrames);
+		SetGlobalTime(ScrubPosition, ESnapTimeMode::STM_Interval);
 
 		// Update on stop (cleans up things like sounds that are playing)
 		EMovieSceneUpdateData UpdateData(ScrubPosition, ScrubPosition);
@@ -1544,12 +1972,8 @@ FReply FSequencer::OnPlay(bool bTogglePlay)
 	}
 	else
 	{
-		PlaybackState = EMovieScenePlayerStatus::Playing;
+		SetPlaybackStatus(EMovieScenePlayerStatus::Playing);
 
-		// Reset our playback start time so fixed time step playback can gate frame rate properly
-		PlaybackStartRealTime = FPlatformTime::Seconds();
-		NumContinuousFixedTimeStepFrames = 0;
-		
 		// Make sure Slate ticks during playback
 		SequencerWidget->RegisterActiveTimerForPlayback();
 	}
@@ -1558,59 +1982,98 @@ FReply FSequencer::OnPlay(bool bTogglePlay)
 }
 
 
+EVisibility FSequencer::GetRecordButtonVisibility() const
+{
+	return UMovieSceneSubSection::IsSetAsRecording() ? EVisibility::Visible : EVisibility::Collapsed;
+}
+
+
 FReply FSequencer::OnRecord()
 {
-	if (PlaybackState != EMovieScenePlayerStatus::Recording)
-	{
-		PlaybackState = EMovieScenePlayerStatus::Recording;
-		
-		// Make sure Slate ticks during playback
-		SequencerWidget->RegisterActiveTimerForPlayback();
+	ISequenceRecorder& SequenceRecorder = FModuleManager::LoadModuleChecked<ISequenceRecorder>("SequenceRecorder");
 
-		// @todo sequencer livecapture: Ideally we would support fixed timestep capture from simulation
-		//			Basically we need to run the PIE world at a fixed time step, capturing key frames every frame. 
-		//			The editor world would still be ticked at the normal throttle-real-time rate
+	if(UMovieSceneSubSection::IsSetAsRecording() && !SequenceRecorder.IsRecording())
+	{
+		UWorld* PlaybackContext = Cast<UWorld>(GetPlaybackContext());
+		const FString& ActorToRecord = UMovieSceneSubSection::GetActorToRecord();
+		const FString& PathToRecordTo = UMovieSceneSubSection::GetRecordingSection()->GetTargetPathToRecordTo();
+		const FString& SequenceName = UMovieSceneSubSection::GetRecordingSection()->GetTargetSequenceName();
+		SequenceRecorder.StartRecording(PlaybackContext, ActorToRecord, FOnRecordingStarted::CreateSP(this, &FSequencer::HandleRecordingStarted), FOnRecordingFinished::CreateSP(this, &FSequencer::HandleRecordingFinished), PathToRecordTo, SequenceName);
 	}
-	else
+	else if(SequenceRecorder.IsRecording())
 	{
-		PlaybackState = EMovieScenePlayerStatus::Stopped;
+		SequenceRecorder.StopRecording();
 	}
 
 	return FReply::Handled();
 }
 
+void FSequencer::HandleRecordingStarted(UMovieSceneSequence* Sequence)
+{
+	OnPlay(false);
+		
+	// Make sure Slate ticks during playback
+	SequencerWidget->RegisterActiveTimerForPlayback();
+
+	// sync recording section to start
+	UMovieSceneSubSection* Section = UMovieSceneSubSection::GetRecordingSection();
+	if(Section != nullptr)
+	{
+		Section->SetStartTime(ScrubPosition);
+		Section->SetEndTime(ScrubPosition + Settings->GetTimeSnapInterval());
+	}
+}
+
+void FSequencer::HandleRecordingFinished(UMovieSceneSequence* Sequence)
+{
+	// toggle us to no playing if we are still playing back
+	// as the post processing takes such a long time we don't really care if the sequence doesnt carry on
+	if(PlaybackState == EMovieScenePlayerStatus::Playing)
+	{
+		OnPlay(true);
+	}
+
+	// now patchup the section that was recorded to
+	UMovieSceneSubSection* Section = UMovieSceneSubSection::GetRecordingSection();
+	if(Section != nullptr)
+	{
+		Section->SetAsRecording(false);
+		Section->SetSequence(Sequence);
+		Section->SetEndTime(Section->GetStartTime() + Sequence->GetMovieScene()->GetPlaybackRange().Size<float>());
+	}
+
+	bNeedTreeRefresh = true;
+}
 
 FReply FSequencer::OnStepForward()
 {
-	PlaybackState = EMovieScenePlayerStatus::Stopped;
+	SetPlaybackStatus(EMovieScenePlayerStatus::Stepping);
 	float NewPosition = ScrubPosition + Settings->GetTimeSnapInterval();
-	const bool bAllowSnappingToFrames = true;
-	SetGlobalTime(NewPosition, bAllowSnappingToFrames);
+	SetGlobalTime(NewPosition, ESnapTimeMode::STM_Interval);
 	return FReply::Handled();
 }
 
 
 FReply FSequencer::OnStepBackward()
 {
-	PlaybackState = EMovieScenePlayerStatus::Stopped;
+	SetPlaybackStatus(EMovieScenePlayerStatus::Stepping);
 	float NewPosition = ScrubPosition - Settings->GetTimeSnapInterval();
 	const bool bAllowSnappingToFrames = true;
-	SetGlobalTime(NewPosition, bAllowSnappingToFrames);
+	SetGlobalTime(NewPosition, ESnapTimeMode::STM_Interval);
 	return FReply::Handled();
 }
 
 
 FReply FSequencer::OnStepToEnd()
 {
-	PlaybackState = EMovieScenePlayerStatus::Stopped;
+	SetPlaybackStatus(EMovieScenePlayerStatus::Stepping);
 	SetGlobalTime(GetPlaybackRange().GetUpperBoundValue());
 	return FReply::Handled();
 }
 
-
 FReply FSequencer::OnStepToBeginning()
 {
-	PlaybackState = EMovieScenePlayerStatus::Stopped;
+	SetPlaybackStatus(EMovieScenePlayerStatus::Stepping);
 	SetGlobalTime(GetPlaybackRange().GetLowerBoundValue());
 	return FReply::Handled();
 }
@@ -1623,7 +2086,69 @@ FReply FSequencer::OnToggleLooping()
 }
 
 
-bool FSequencer::IsLooping() const
+FReply FSequencer::SetPlaybackEnd()
+{
+	const UMovieSceneSequence* FocusedSequence = GetFocusedMovieSceneSequence();
+	if (FocusedSequence)
+	{
+		TRange<float> CurrentRange = FocusedSequence->GetMovieScene()->GetPlaybackRange();
+		float NewPos = FMath::Max(GetGlobalTime(), CurrentRange.GetLowerBoundValue());
+		SetPlaybackRange(TRange<float>(CurrentRange.GetLowerBoundValue(), NewPos));
+	}
+
+	return FReply::Handled();
+}
+
+FReply FSequencer::SetPlaybackStart()
+{
+	const UMovieSceneSequence* FocusedSequence = GetFocusedMovieSceneSequence();
+	if (FocusedSequence)
+	{
+		TRange<float> CurrentRange = FocusedSequence->GetMovieScene()->GetPlaybackRange();
+		float NewPos = FMath::Max(GetGlobalTime(), CurrentRange.GetUpperBoundValue());
+		SetPlaybackRange(TRange<float>(NewPos, CurrentRange.GetUpperBoundValue()));	
+	}
+
+	return FReply::Handled();
+}
+
+FReply FSequencer::JumpToPreviousKey()
+{
+	TUniquePtr<ISequencerKeyCollection> ActiveKeyCollection;
+	GetKeysFromSelection(ActiveKeyCollection);
+	if (ActiveKeyCollection.IsValid())
+	{
+		TRange<float> FindRange(TRange<float>::BoundsType(), GetGlobalTime());
+
+		TOptional<float> NewTime = ActiveKeyCollection->FindFirstKeyInRange(FindRange, EFindKeyDirection::Backwards);
+		if (NewTime.IsSet())
+		{
+			SetGlobalTime(NewTime.GetValue());
+		}
+	}
+
+	return FReply::Handled();
+}
+
+FReply FSequencer::JumpToNextKey()
+{
+	TUniquePtr<ISequencerKeyCollection> ActiveKeyCollection;
+	GetKeysFromSelection(ActiveKeyCollection);
+	if (ActiveKeyCollection.IsValid())
+	{
+		TRange<float> FindRange(GetGlobalTime(), TRange<float>::BoundsType());
+
+		TOptional<float> NewTime = ActiveKeyCollection->FindFirstKeyInRange(FindRange, EFindKeyDirection::Forwards);
+		if (NewTime.IsSet())
+		{
+			SetGlobalTime(NewTime.GetValue());
+		}
+	}
+
+	return FReply::Handled();
+}
+
+bool FSequencer::IsLooping() const
 {
 	return Settings->IsLooping();
 }
@@ -1653,7 +2178,7 @@ void FSequencer::SetGlobalTimeLooped(float InTime)
 		if (bReachedEnd)
 		{
 			InTime = TimeBounds.GetUpperBoundValue();
-			PlaybackState = EMovieScenePlayerStatus::Stopped;
+			SetPlaybackStatus(EMovieScenePlayerStatus::Stopped);
 		}
 		// Constrain to the play range if necessary
 		else if (Settings->ShouldKeepCursorInPlayRange())
@@ -1671,7 +2196,7 @@ void FSequencer::SetGlobalTimeLooped(float InTime)
 		else if (!WorkingRange.Contains(InTime))
 		{
 			InTime = FMath::Clamp(InTime, WorkingRange.GetLowerBoundValue(), WorkingRange.GetUpperBoundValue());
-			PlaybackState = EMovieScenePlayerStatus::Stopped;
+			SetPlaybackStatus(EMovieScenePlayerStatus::Stopped);
 		}
 	}
 
@@ -1797,6 +2322,43 @@ void FSequencer::OnClampRangeChanged( TRange<float> NewClampRange )
 	}
 }
 
+float FSequencer::FindNearestKey(float NewScrubPosition)
+{
+	TUniquePtr<ISequencerKeyCollection> ActiveKeyCollection;
+	GetKeysFromSelection(ActiveKeyCollection);
+
+	if (ActiveKeyCollection.IsValid())
+	{
+		TRange<float> FindRangeBackwards(TRange<float>::BoundsType(), NewScrubPosition);
+		TOptional<float> NewTimeBackwards = ActiveKeyCollection->FindFirstKeyInRange(FindRangeBackwards, EFindKeyDirection::Backwards);
+
+		TRange<float> FindRangeForwards(NewScrubPosition, TRange<float>::BoundsType());
+		TOptional<float> NewTimeForwards = ActiveKeyCollection->FindFirstKeyInRange(FindRangeForwards, EFindKeyDirection::Forwards);
+		if (NewTimeForwards.IsSet())
+		{
+			if (NewTimeBackwards.IsSet())
+			{
+				if (FMath::Abs(NewTimeForwards.GetValue() - NewScrubPosition) < FMath::Abs(NewTimeBackwards.GetValue() - NewScrubPosition))
+				{
+					NewScrubPosition = NewTimeForwards.GetValue();
+				}
+				else
+				{
+					NewScrubPosition = NewTimeBackwards.GetValue();
+				}
+			}
+			else
+			{
+				NewScrubPosition = NewTimeForwards.GetValue();
+			}
+		}
+		else if (NewTimeBackwards.IsSet())
+		{
+			NewScrubPosition = NewTimeBackwards.GetValue();
+		}
+	}
+	return NewScrubPosition;
+}
 
 void FSequencer::OnScrubPositionChanged( float NewScrubPosition, bool bScrubbing )
 {
@@ -1828,20 +2390,20 @@ void FSequencer::OnScrubPositionChanged( float NewScrubPosition, bool bScrubbing
 		NewScrubPosition = FMath::Clamp(NewScrubPosition, TargetViewRange.GetLowerBoundValue(), TargetViewRange.GetUpperBoundValue());
 	}
 
-	SetGlobalTimeDirectly( NewScrubPosition );
+	SetGlobalTimeDirectly( NewScrubPosition, Settings->GetSnapKeyTimesToKeys() ? ESnapTimeMode::STM_Keys : ESnapTimeMode::STM_None );
 }
 
 
 void FSequencer::OnBeginScrubbing()
 {
-	PlaybackState = EMovieScenePlayerStatus::Scrubbing;
+	SetPlaybackStatus(EMovieScenePlayerStatus::Scrubbing);
 	SequencerWidget->RegisterActiveTimerForPlayback();
 }
 
 
 void FSequencer::OnEndScrubbing()
 {
-	PlaybackState = EMovieScenePlayerStatus::Stopped;
+	SetPlaybackStatus(EMovieScenePlayerStatus::Stepping);
 	AutoscrubOffset.Reset();
 	StopAutoscroll();
 }
@@ -1859,6 +2421,18 @@ void FSequencer::OnEndPlaybackRangeDrag()
 }
 
 
+void FSequencer::OnBeginInOutRangeDrag()
+{
+	GEditor->BeginTransaction(LOCTEXT("SetSelectionRange_Transaction", "Set selection range"));
+}
+
+
+void FSequencer::OnEndInOutRangeDrag()
+{
+	GEditor->EndTransaction();
+}
+
+
 void FSequencer::StartAutoscroll(float UnitsPerS)
 {
 	AutoscrollOffset = UnitsPerS;
@@ -1981,21 +2555,19 @@ FGuid FSequencer::MakeNewSpawnable( UObject& Object )
 		DefaultTransform.Scale = FVector(1.0f, 1.0f, 1.0f);
 		DefaultTransform.bValid = true;
 
-		// Make spawnable cameras owned by the master sequence by default so that they can persist outside of the shot. 
-		// This allows thumbnails and the look through camera button to function properly since the spawnable camera
-		// will exist.
-		UCameraComponent* CameraComponent = MovieSceneHelpers::CameraComponentFromActor(SpawnedActor);
-		if (CameraComponent)
-		{
-			FMovieSceneSpawnable* NewSpawnable = GetFocusedMovieSceneSequence()->GetMovieScene()->FindSpawnable(NewGuid);
-			if (NewSpawnable)
-			{
-				NewSpawnable->SetSpawnOwnership(ESpawnOwnership::MasterSequence);
-			}
-		}
+		OnActorAddedToSequencerEvent.Broadcast(SpawnedActor, NewGuid);
+
+		const bool bNotifySelectionChanged = true;
+		const bool bDeselectBSP = true;
+		const bool bWarnAboutTooManyActors = false;
+		const bool bSelectEvenIfHidden = false;
+
+		GEditor->SelectNone( bNotifySelectionChanged, bDeselectBSP, bWarnAboutTooManyActors );
+		GEditor->SelectActor( SpawnedActor, true, bNotifySelectionChanged, bSelectEvenIfHidden );
 	}
 
 	SetupDefaultsForSpawnable(NewGuid, DefaultTransform);
+
 	return NewGuid;
 }
 
@@ -2013,7 +2585,10 @@ void FSequencer::SetupDefaultsForSpawnable( const FGuid& Guid, const FTransformD
 
 	if (SpawnTrack)
 	{
-		SpawnTrack->AddSection(*SpawnTrack->CreateNewSection());
+		UMovieSceneBoolSection* SpawnSection = Cast<UMovieSceneBoolSection>(SpawnTrack->CreateNewSection());
+		SpawnSection->SetDefault(true);
+		SpawnSection->SetIsInfinite(GetInfiniteKeyAreas());
+		SpawnTrack->AddSection(*SpawnSection);
 		SpawnTrack->SetObjectId(Guid);
 	}
 	
@@ -2053,6 +2628,8 @@ void FSequencer::SetupDefaultsForSpawnable( const FGuid& Guid, const FTransformD
 				TransformSection->SetDefault( FTransformKey( EKey3DTransformChannel::Scale, EAxis::X, DefaultTransform.Scale.X, bUnwindRotation ) );
 				TransformSection->SetDefault( FTransformKey( EKey3DTransformChannel::Scale, EAxis::Y, DefaultTransform.Scale.Y, bUnwindRotation ) );
 				TransformSection->SetDefault( FTransformKey( EKey3DTransformChannel::Scale, EAxis::Z, DefaultTransform.Scale.Z, bUnwindRotation ) );
+
+				TransformSection->SetIsInfinite(GetInfiniteKeyAreas());
 			}
 		}
 	}
@@ -2128,8 +2705,33 @@ void FSequencer::OnRequestNodeDeleted( TSharedRef<const FSequencerDisplayNode> N
 	UMovieSceneSequence* Sequence = MovieSceneInstance->GetSequence();
 	UMovieScene* OwnerMovieScene = Sequence->GetMovieScene();
 
-	// Only object nodes or section areas can be deleted
-	if (NodeToBeDeleted->GetType() == ESequencerNode::Object)
+	if ( NodeToBeDeleted->GetType() == ESequencerNode::Folder )
+	{
+		// Delete Children
+		for ( const TSharedRef<FSequencerDisplayNode>& ChildNode : NodeToBeDeleted->GetChildNodes() )
+		{
+			OnRequestNodeDeleted( ChildNode );
+		}
+
+		// Delete from parent, or root.
+		TSharedRef<const FSequencerFolderNode> FolderToBeDeleted = StaticCastSharedRef<const FSequencerFolderNode>(NodeToBeDeleted);
+		if ( NodeToBeDeleted->GetParent().IsValid() )
+		{
+			TSharedPtr<FSequencerFolderNode> ParentFolder = StaticCastSharedPtr<FSequencerFolderNode>( NodeToBeDeleted->GetParent() );
+			ParentFolder->GetFolder().Modify();
+			ParentFolder->GetFolder().RemoveChildFolder( &FolderToBeDeleted->GetFolder() );
+		}
+		else
+		{
+			UMovieScene* FocusedMovieScene = GetFocusedMovieSceneSequence()->GetMovieScene();
+			FocusedMovieScene->Modify();
+			FocusedMovieScene->GetRootFolders().Remove( &FolderToBeDeleted->GetFolder() );
+		}
+
+		bAnythingRemoved = true;
+
+	}
+	else if (NodeToBeDeleted->GetType() == ESequencerNode::Object)
 	{
 		// Delete any child object bindings
 		for (const TSharedRef<FSequencerDisplayNode>& ChildNode : NodeToBeDeleted->GetChildNodes())
@@ -2141,6 +2743,14 @@ void FSequencer::OnRequestNodeDeleted( TSharedRef<const FSequencerDisplayNode> N
 		}
 
 		const FGuid& BindingToRemove = StaticCastSharedRef<const FSequencerObjectBindingNode>( NodeToBeDeleted )->GetObjectBinding();
+
+		// Remove from a parent folder if necessary.
+		if ( NodeToBeDeleted->GetParent().IsValid() && NodeToBeDeleted->GetParent()->GetType() == ESequencerNode::Folder )
+		{
+			TSharedPtr<FSequencerFolderNode> ParentFolder = StaticCastSharedPtr<FSequencerFolderNode>( NodeToBeDeleted->GetParent() );
+			ParentFolder->GetFolder().Modify();
+			ParentFolder->GetFolder().RemoveChildObjectBinding( BindingToRemove );
+		}
 		
 		// Try to remove as a spawnable first
 		if (OwnerMovieScene->RemoveSpawnable(BindingToRemove))
@@ -2166,6 +2776,14 @@ void FSequencer::OnRequestNodeDeleted( TSharedRef<const FSequencerDisplayNode> N
 		TSharedRef<const FSequencerTrackNode> SectionAreaNode = StaticCastSharedRef<const FSequencerTrackNode>( NodeToBeDeleted );
 		UMovieSceneTrack* Track = SectionAreaNode->GetTrack();
 
+		// Remove from a parent folder if necessary.
+		if ( NodeToBeDeleted->GetParent().IsValid() && NodeToBeDeleted->GetParent()->GetType() == ESequencerNode::Folder )
+		{
+			TSharedPtr<FSequencerFolderNode> ParentFolder = StaticCastSharedPtr<FSequencerFolderNode>( NodeToBeDeleted->GetParent() );
+			ParentFolder->GetFolder().Modify();
+			ParentFolder->GetFolder().RemoveChildMasterTrack( Track );
+		}
+
 		if (Track != nullptr)
 		{
 			if (OwnerMovieScene->IsAMasterTrack(*Track))
@@ -2302,7 +2920,7 @@ void FSequencer::ActivateSequencerEditorMode()
 	GLevelEditorModeTools().ActivateMode( FSequencerEdMode::EM_SequencerMode );
 
 	FSequencerEdMode* SequencerEdMode = (FSequencerEdMode*)(GLevelEditorModeTools().GetActiveMode(FSequencerEdMode::EM_SequencerMode));
-	SequencerEdMode->SetSequencer(ActiveSequencer.Get());
+	SequencerEdMode->SetSequencer(ActiveSequencerPtr.Pin().Get());
 }
 
 
@@ -2362,7 +2980,7 @@ void FSequencer::OnPropertyEditorOpened()
 }
 
 
-void FSequencer::UpdatePreviewLevelViewportClientFromCameraCut(FLevelEditorViewportClient& InViewportClient, UObject* InCameraObject) const
+void FSequencer::UpdatePreviewLevelViewportClientFromCameraCut(FLevelEditorViewportClient& InViewportClient, UObject* InCameraObject, bool bJumpCut) const
 {
 	AActor* CameraActor = Cast<AActor>(InCameraObject);
 
@@ -2370,12 +2988,12 @@ void FSequencer::UpdatePreviewLevelViewportClientFromCameraCut(FLevelEditorViewp
 	{
 		InViewportClient.SetViewLocation(CameraActor->GetActorLocation());
 		InViewportClient.SetViewRotation(CameraActor->GetActorRotation());
-		InViewportClient.bEditorCameraCut = !InViewportClient.IsLockedToActor(CameraActor);
+		InViewportClient.bEditorCameraCut = !InViewportClient.IsLockedToActor(CameraActor) || bJumpCut;
 	}
 	else
 	{
 		InViewportClient.ViewFOV = InViewportClient.FOVAngle;
-		InViewportClient.bEditorCameraCut = false;
+		InViewportClient.bEditorCameraCut = bJumpCut;
 	}
 
 	// Set the actor lock.
@@ -2561,7 +3179,12 @@ void FSequencer::AddSelectedObjects()
 			FGuid ExistingGuid = SequenceInstanceStack.Top()->FindObjectId(*Actor);
 			if (!ExistingGuid.IsValid())
 			{
-				CreateBinding(*Actor, Actor->GetActorLabel());
+				FGuid PossessableGuid = CreateBinding(*Actor, Actor->GetActorLabel());
+
+				UpdateRuntimeInstances();
+
+				OnActorAddedToSequencerEvent.Broadcast(Actor, PossessableGuid);
+
 				bPossessableAdded = true;
 			}
 		}
@@ -2581,19 +3204,138 @@ void FSequencer::OnSectionSelectionChanged()
 
 void FSequencer::OnSelectedOutlinerNodesChanged()
 {
-	if (IsLevelEditorSequencer())
+	if ( bUpdatingSequencerSelection == false )
+	{
+		bUpdatingExternalSelection = true;
+		SynchronizeExternalSelectionWithSequencerSelection();
+		bUpdatingExternalSelection = false;
+	}
+}
+
+
+void FSequencer::SynchronizeExternalSelectionWithSequencerSelection()
+{
+	if ( !IsLevelEditorSequencer() )
+	{
+		return;
+	}
+
+	TSet<AActor*> SelectedSequencerActors;
+	for ( TSharedRef<FSequencerDisplayNode> SelectedOutlinerNode : Selection.GetSelectedOutlinerNodes() )
+	{
+		// Get the root object binding node.
+		TSharedPtr<FSequencerDisplayNode> CurrentNode = SelectedOutlinerNode;
+		TSharedPtr<FSequencerObjectBindingNode> ObjectBindingNode;
+		while ( CurrentNode.IsValid() )
+		{
+			if ( CurrentNode->GetType() == ESequencerNode::Object )
+			{
+				ObjectBindingNode = StaticCastSharedPtr<FSequencerObjectBindingNode>(CurrentNode);
+			}
+			CurrentNode = CurrentNode->GetParent();
+		}
+
+		// If the root node is an object node, try to get the actor nodes from it.
+		if ( ObjectBindingNode.IsValid() )
+		{
+			TArray<TWeakObjectPtr<UObject>> RuntimeObjects;
+			GetRuntimeObjects( GetFocusedMovieSceneSequenceInstance(), ObjectBindingNode->GetObjectBinding(), RuntimeObjects );
+			
+			for (auto RuntimeObject : RuntimeObjects )
+			{
+				AActor* Actor = Cast<AActor>(RuntimeObject.Get());
+				if ( Actor != nullptr )
+				{
+					SelectedSequencerActors.Add( Actor );
+				}
+			}
+		}
+	}
+
+	const FScopedTransaction Transaction( NSLOCTEXT( "Sequencer", "UpdatingActorSelection", "Select Actors" ) );
+
+	const bool bNotifySelectionChanged = false;
+	const bool bDeselectBSP = true;
+	const bool bWarnAboutTooManyActors = false;
+	const bool bSelectEvenIfHidden = true;
+
+	GEditor->GetSelectedActors()->Modify();
+	GEditor->GetSelectedActors()->BeginBatchSelectOperation();
+
+	GEditor->SelectNone( bNotifySelectionChanged, bDeselectBSP, bWarnAboutTooManyActors );
+
+	for ( AActor* SelectedSequencerActor : SelectedSequencerActors )
+	{
+		GEditor->SelectActor( SelectedSequencerActor, true, bNotifySelectionChanged, bSelectEvenIfHidden );
+	}
+
+	GEditor->GetSelectedActors()->EndBatchSelectOperation();
+	GEditor->NoteSelectionChange();
+}
+
+
+void FSequencer::OnActorSelectionChanged( UObject* )
+{
+	if ( bUpdatingExternalSelection == false )
 	{
-		// Clear the selection if no sequencer display nodes were clicked on and there are selected actors
-		if (Selection.GetSelectedOutlinerNodes().Num() == 0 && GEditor->GetSelectedActors()->Num() != 0 && SequencerWidget->UserIsSelecting())
+		bUpdatingSequencerSelection = true;
+		SynchronizeSequencerSelectionWithExternalSelection();
+		bUpdatingSequencerSelection = false;
+	}
+}
+
+
+void GetRootObjectBindingNodes(const TArray<TSharedRef<FSequencerDisplayNode>>& DisplayNodes, TArray<TSharedRef<FSequencerObjectBindingNode>>& RootObjectBindings )
+{
+	for ( TSharedRef<FSequencerDisplayNode> DisplayNode : DisplayNodes )
+	{
+		switch ( DisplayNode->GetType() )
 		{
-			const bool bNotifySelectionChanged = true;
-			const bool bDeselectBSP = true;
-			const bool bWarnAboutTooManyActors = false;
+		case ESequencerNode::Folder:
+			GetRootObjectBindingNodes( DisplayNode->GetChildNodes(), RootObjectBindings );
+			break;
+		case ESequencerNode::Object:
+			RootObjectBindings.Add( StaticCastSharedRef<FSequencerObjectBindingNode>( DisplayNode ) );
+			break;
+		}
+	}
+}
+
+
+void FSequencer::SynchronizeSequencerSelectionWithExternalSelection()
+{
+	if ( !IsLevelEditorSequencer() )
+	{
+		return;
+	}
+
+	TArray<TSharedRef<FSequencerObjectBindingNode>> RootObjectBindingNodes;
+	GetRootObjectBindingNodes( NodeTree->GetRootNodes(), RootObjectBindingNodes );
 
-			const FScopedTransaction Transaction(NSLOCTEXT("UnrealEd", "ClickingOnActors", "Clicking on Actors"));
-			GEditor->SelectNone(bNotifySelectionChanged, bDeselectBSP, bWarnAboutTooManyActors);
+	TSet<TSharedRef<FSequencerDisplayNode>> NodesToSelect;
+	for ( TSharedRef<FSequencerObjectBindingNode> ObjectBindingNode : RootObjectBindingNodes )
+	{
+		TArray<TWeakObjectPtr<UObject>> RuntimeObjects;
+		GetRuntimeObjects( GetFocusedMovieSceneSequenceInstance(), ObjectBindingNode->GetObjectBinding(), RuntimeObjects );
+
+		for ( TWeakObjectPtr<UObject> RuntimeObjectPtr : RuntimeObjects )
+		{
+			UObject* RuntimeObject = RuntimeObjectPtr.Get();
+			if ( RuntimeObject != nullptr && GEditor->GetSelectedActors()->IsSelected( RuntimeObject ) )
+			{
+				NodesToSelect.Add( ObjectBindingNode );
+			}
 		}
 	}
+
+	Selection.SuspendBroadcast();
+	Selection.EmptySelectedOutlinerNodes();
+	for ( TSharedRef<FSequencerDisplayNode> NodeToSelect : NodesToSelect)
+	{
+		Selection.AddToSelection( NodeToSelect );
+	}
+	Selection.ResumeBroadcast();
+	Selection.GetOnOutlinerNodeSelectionChanged().Broadcast();
 }
 
 
@@ -2662,15 +3404,19 @@ float FSequencer::GetOverlayFadeCurve() const
 
 void FSequencer::DeleteSelectedItems()
 {
-	if (Selection.GetActiveSelection() == FSequencerSelection::EActiveSelection::KeyAndSection)
+	if (Selection.GetSelectedKeys().Num())
 	{
-		FScopedTransaction DeleteKeysTransaction( NSLOCTEXT("Sequencer", "DeleteKeysAndSections_Transaction", "Delete Keys and Sections") );
+		FScopedTransaction DeleteKeysTransaction( NSLOCTEXT("Sequencer", "DeleteKeys_Transaction", "Delete Keys") );
 		
 		DeleteSelectedKeys();
-
+	}
+	else if (Selection.GetSelectedSections().Num())
+	{
+		FScopedTransaction DeleteSectionsTransaction( NSLOCTEXT("Sequencer", "DeleteSections_Transaction", "Delete Sections") );
+	
 		DeleteSections(Selection.GetSelectedSections());
 	}
-	else if (Selection.GetActiveSelection() == FSequencerSelection::EActiveSelection::OutlinerNode)
+	else if (Selection.GetSelectedOutlinerNodes().Num())
 	{
 		DeleteSelectedNodes();
 	}
@@ -2774,8 +3520,35 @@ void FSequencer::DoAssignActor(AActor*const* InActors, int32 NumActors, FGuid In
 		OwnerMovieScene->ReplacePossessable(InObjectBinding, ParentGuid, NewActorLabel);
 	}
 
+	auto UpdateComponent = [&]( FGuid OldComponentGuid, UActorComponent* NewComponent )
+	{
+		// Get the object guid to assign, remove the binding if it already exists
+		FGuid NewComponentGuid = SequenceInstanceStack.Top()->FindObjectId( *NewComponent );
+		FString NewComponentLabel = NewComponent->GetName();
+		if ( NewComponentGuid.IsValid() )
+		{
+			OwnerMovieScene->RemovePossessable( NewComponentGuid );
+			OwnerSequence->UnbindPossessableObjects( NewComponentGuid );
+		}
+
+		// Add this object
+		FMovieScenePossessable NewPossessable( NewComponentLabel, NewComponent->GetClass() );
+		NewComponentGuid = NewPossessable.GetGuid();
+		OwnerSequence->BindPossessableObject( NewComponentGuid, *NewComponent, Actor );
+
+		// Replace
+		OwnerMovieScene->ReplacePossessable( OldComponentGuid, NewComponentGuid, NewComponentLabel );
+
+		FMovieScenePossessable* ThisPossessable = OwnerMovieScene->FindPossessable( NewComponentGuid );
+		if ( ensure( ThisPossessable ) )
+		{
+			ThisPossessable->SetParent( ParentGuid );
+		}
+	};
+
 	// Handle components
-	if (AActor* ActorToReplace = Cast<AActor>(RuntimeObject))
+	AActor* ActorToReplace = Cast<AActor>(RuntimeObject);
+	if (ActorToReplace != nullptr && ActorToReplace->IsActorBeingDestroyed() == false)
 	{
 		for (UActorComponent* ComponentToReplace : ActorToReplace->GetComponents())
 		{
@@ -2787,34 +3560,55 @@ void FSequencer::DoAssignActor(AActor*const* InActors, int32 NumActors, FGuid In
 					UActorComponent* NewComponent = Actor->GetComponentByClass(ComponentToReplace->GetClass());
 					if (NewComponent)
 					{
-						// Get the object guid to assign, remove the binding if it already exists
-						FGuid NewComponentGuid = SequenceInstanceStack.Top()->FindObjectId(*NewComponent);
-						FString NewComponentLabel = NewComponent->GetName();
-						if (NewComponentGuid.IsValid())
-						{
-							OwnerMovieScene->RemovePossessable(NewComponentGuid);
-							OwnerSequence->UnbindPossessableObjects(NewComponentGuid);
-						}
-
-						// Add this object
-						FMovieScenePossessable NewPossessable( NewComponentLabel, NewComponent->GetClass());
-						NewComponentGuid = NewPossessable.GetGuid();
-						OwnerSequence->BindPossessableObject(NewComponentGuid, *NewComponent, Actor);
-
-						// Replace
-						OwnerMovieScene->ReplacePossessable(ComponentGuid, NewComponentGuid, NewComponentLabel);
-
-						FMovieScenePossessable* ThisPossessable = OwnerMovieScene->FindPossessable(NewComponentGuid);
-						if (ensure(ThisPossessable))
-						{
-							ThisPossessable->SetParent(ParentGuid);
-						}
+						UpdateComponent( ComponentGuid, NewComponent );
 					}
 				}
 			}
 		}
 	}
-			
+	else // If the actor didn't exist, try to find components who's parent guids were the previous actors guid.
+	{
+		TMap<FString, UActorComponent*> ComponentNameToComponent;
+		for ( UActorComponent* Component : Actor->GetComponents() )
+		{
+			ComponentNameToComponent.Add( Component->GetName(), Component );
+		}
+		for ( int32 i = 0; i < OwnerMovieScene->GetPossessableCount(); i++ )
+		{
+			FMovieScenePossessable& OldPossessable = OwnerMovieScene->GetPossessable(i);
+			if ( OldPossessable.GetParent() == InObjectBinding )
+			{
+				UActorComponent** ComponentPtr = ComponentNameToComponent.Find( OldPossessable.GetName() );
+				if ( ComponentPtr != nullptr )
+				{
+					UpdateComponent( OldPossessable.GetGuid(), *ComponentPtr );
+				}
+			}
+		}
+	}
+
+	// Try to fix up folders.
+	TArray<UMovieSceneFolder*> FoldersToCheck;
+	FoldersToCheck.Append(GetFocusedMovieSceneSequence()->GetMovieScene()->GetRootFolders());
+	bool bFolderFound = false;
+	while ( FoldersToCheck.Num() > 0 && bFolderFound == false )
+	{
+		UMovieSceneFolder* Folder = FoldersToCheck[0];
+		FoldersToCheck.RemoveAt(0);
+		if ( Folder->GetChildObjectBindings().Contains( InObjectBinding ) )
+		{
+			Folder->Modify();
+			Folder->RemoveChildObjectBinding( InObjectBinding );
+			Folder->AddChildObjectBinding( ParentGuid );
+			bFolderFound = true;
+		}
+
+		for ( UMovieSceneFolder* ChildFolder : Folder->GetChildFolders() )
+		{
+			FoldersToCheck.Add( ChildFolder );
+		}
+	}
+
 	RootMovieSceneSequenceInstance->RestoreState(*this);
 	NotifyMovieSceneDataChanged();
 }
@@ -2938,7 +3732,7 @@ bool FSequencer::IsNodeLocked() const
 
 void FSequencer::ConvertToSpawnable(TSharedRef<FSequencerObjectBindingNode> NodeToBeConverted)
 {
-	const FScopedTransaction Transaction( LOCTEXT("ConvertSelectedSpawnable", "Convert node to spawnables") );
+	const FScopedTransaction Transaction( LOCTEXT("ConvertSelectedSpawnable", "Convert Node to Spawnables") );
 
 	// Ensure we're in a non-possessed state
 	RootMovieSceneSequenceInstance->RestoreState(*this);
@@ -2950,7 +3744,7 @@ void FSequencer::ConvertToSpawnable(TSharedRef<FSequencerObjectBindingNode> Node
 void FSequencer::ConvertSelectedNodesToSpawnables()
 {
 	// @todo sequencer: Undo doesn't seem to be working at all
-	const FScopedTransaction Transaction( LOCTEXT("ConvertSelectedSpawnable", "Convert selected nodes to spawnables") );
+	const FScopedTransaction Transaction( LOCTEXT("ConvertSelectedSpawnable", "Convert Selected Nodes to Spawnables") );
 
 	UMovieScene* MovieScene = GetFocusedMovieSceneSequence()->GetMovieScene();
 
@@ -2975,7 +3769,7 @@ void FSequencer::ConvertSelectedNodesToSpawnables()
 		}
 	}
 
-	FScopedSlowTask SlowTask(ObjectBindingNodes.Num(), LOCTEXT("ConvertSpawnableProgress", "Converting selected possessable nodes to spawnables"));
+	FScopedSlowTask SlowTask(ObjectBindingNodes.Num(), LOCTEXT("ConvertSpawnableProgress", "Converting Selected Possessable Nodes to Spawnables"));
 	SlowTask.MakeDialog(true);
 
 	for (const TSharedRef<FSequencerObjectBindingNode>& ObjectBindingNode : ObjectBindingNodes)
@@ -3004,8 +3798,11 @@ void FSequencer::ConvertToSpawnableInternal(TSharedRef<FSequencerObjectBindingNo
 	}
 
 	// Find the object in the environment
-	UObject* FoundObject = GetFocusedMovieSceneSequenceInstance()->FindObject(NodeToBeConverted->GetObjectBinding(), *this);
-	if (!FoundObject)
+	TSharedRef<FMovieSceneSequenceInstance> FocusedSequenceInstance = SequenceInstanceStack.Top();
+	UMovieSceneSequence* FocusedSequence = FocusedSequenceInstance->GetSequence();
+
+	UObject* FoundObject = FocusedSequenceInstance->FindObject(NodeToBeConverted->GetObjectBinding(), *this);
+	if (!FoundObject || !FocusedSequence)
 	{
 		return;
 	}
@@ -3034,8 +3831,18 @@ void FSequencer::ConvertToSpawnableInternal(TSharedRef<FSequencerObjectBindingNo
 		for (int32 Index = 0; Index < MovieScene->GetPossessableCount(); ++Index)
 		{
 			FMovieScenePossessable& MovieScenePossessable = MovieScene->GetPossessable(Index);
-			if (MovieScenePossessable.GetParent() == PersistentGuid)
+
+			bool bBelongsToNewSpawnable = MovieScenePossessable.GetParent() == PersistentGuid;
+			if (!bBelongsToNewSpawnable)
 			{
+				// This is potentially slow, but fixes old data where we didn't assign parent object guids
+				UObject* PotentialChildObject = FocusedSequenceInstance->FindObject(MovieScenePossessable.GetGuid(), *this);
+				bBelongsToNewSpawnable = PotentialChildObject && FocusedSequence->GetParentObject(PotentialChildObject) == FoundObject;
+			}
+			
+			if (bBelongsToNewSpawnable)
+			{
+				MovieScenePossessable.SetParent(PersistentGuid);
 				Spawnable->AddChildPossessable(MovieScenePossessable.GetGuid());
 			}
 		}
@@ -3063,6 +3870,64 @@ void FSequencer::ConvertToSpawnableInternal(TSharedRef<FSequencerObjectBindingNo
 }
 
 
+void FSequencer::OnAddFolder()
+{
+	// Check if a folder, or child of a folder is currently selected.
+	TArray<UMovieSceneFolder*> SelectedParentFolders;
+	if ( Selection.GetSelectedOutlinerNodes().Num() > 0 )
+	{
+		for ( TSharedRef<FSequencerDisplayNode> SelectedNode : Selection.GetSelectedOutlinerNodes() )
+		{
+			TSharedPtr<FSequencerDisplayNode> CurrentNode = SelectedNode;
+			while ( CurrentNode.IsValid() && CurrentNode->GetType() != ESequencerNode::Folder )
+			{
+				CurrentNode = CurrentNode->GetParent();
+			}
+			if ( CurrentNode.IsValid() )
+			{
+				SelectedParentFolders.Add( &StaticCastSharedPtr<FSequencerFolderNode>( CurrentNode )->GetFolder() );
+			}
+		}
+	}
+
+	TArray<FName> ExistingFolderNames;
+	UMovieScene* FocusedMovieScene = GetFocusedMovieSceneSequence()->GetMovieScene();
+	
+	// If there is a folder selected the existing folder names are the sibling folders.
+	if ( SelectedParentFolders.Num() == 1 )
+	{
+		for ( UMovieSceneFolder* SiblingFolder : SelectedParentFolders[0]->GetChildFolders() )
+		{
+			ExistingFolderNames.Add( SiblingFolder->GetFolderName() );
+		}
+	}
+	// Otherwise use the root folders.
+	else
+	{
+		for ( UMovieSceneFolder* MovieSceneFolder : FocusedMovieScene->GetRootFolders() )
+		{
+			ExistingFolderNames.Add( MovieSceneFolder->GetFolderName() );
+		}
+	}
+
+	FName UniqueName = FSequencerUtilities::GetUniqueName(FName("New Folder"), ExistingFolderNames);
+	UMovieSceneFolder* NewFolder = NewObject<UMovieSceneFolder>( FocusedMovieScene, NAME_None, RF_Transactional );
+	NewFolder->SetFolderName( UniqueName );
+
+	if ( SelectedParentFolders.Num() == 1 )
+	{
+		SelectedParentFolders[0]->Modify();
+		SelectedParentFolders[0]->AddChildFolder( NewFolder );
+	}
+	else
+	{
+		FocusedMovieScene->GetRootFolders().Add( NewFolder );
+	}
+
+	NotifyMovieSceneDataChanged();
+}
+
+
 void FSequencer::TogglePlay()
 {
 	OnPlay();
@@ -3294,6 +4159,8 @@ void FSequencer::SelectTrackKeys(TWeakObjectPtr<UMovieSceneSection> Section, flo
 			}
 		}
 	}
+
+	SequencerHelpers::ValidateNodesWithSelectedKeysOrSections(*this);
 }
 
 
@@ -3385,41 +4252,140 @@ void FSequencer::DiscardChanges()
 		return;
 	}
 
-	const UMovieSceneSequence* EditedSequence = SequenceInstanceStack.Top()->GetSequence();
+	TSharedPtr<IToolkitHost> MyToolkitHost = GetToolkitHost();
+
+	if (!MyToolkitHost.IsValid())
+	{
+		return;
+	}
+
+	UMovieSceneSequence* EditedSequence = SequenceInstanceStack.Top()->GetSequence();
 
 	if (EditedSequence == nullptr)
 	{
 		return;
 	}
 
+	if (FMessageDialog::Open(EAppMsgType::YesNo, LOCTEXT("RevertConfirm", "Are you sure you want to discard your current changes?")) != EAppReturnType::Yes)
+	{
+		return;
+	}
+
+	FAssetEditorManager& AssetEditorManager = FAssetEditorManager::Get();
+	UClass* SequenceClass = EditedSequence->GetClass();
 	FString SequencePath = EditedSequence->GetPathName();
-/*
-	// @todo sequencer: the following is not going to work, because the objects in the unloaded
-	// package are referenced elsewhere and cannot be destroyed. LoadObject will simply load the
-	// already loaded objects (and sub-objects, such as tracks and sections) again.
+	UPackage* SequencePackage = EditedSequence->GetOutermost();
 
-	// unload asset
-	SequenceInstanceStack.Pop();
-	PackageTools::UnloadPackages(TArrayBuilder<UPackage*>().Add(EditedSequence->GetOutermost()), false);
+	// close asset editor
+	AssetEditorManager.CloseAllEditorsForAsset(EditedSequence);
 
-	// reload asset
-	UMovieSceneSequence* ReloadedSequence = LoadObject<UMovieSceneSequence>(nullptr, *SequencePath);
+	// move existing objects into transient package
+	UPackage* const TransientPackage = GetTransientPackage();
+	TMap<FString, UObject*> MovedObjects;
 
-	if (ReloadedSequence == nullptr)
+	ForEachObjectWithOuter(SequencePackage, [&](UObject* Object) {
+		const FName OldName = Object->GetFName();
+		const FName NewName = MakeUniqueObjectName(TransientPackage, Object->GetClass(), OldName);
+		Object->Rename(*NewName.ToString(), TransientPackage, REN_DontCreateRedirectors | REN_DoNotDirty);
+		Object->SetFlags(RF_Transient);
+		Object->ClearFlags(RF_Standalone | RF_Transactional);
+		MovedObjects.Add(Object->GetPathName(), Object);
+	}, true);
+
+	// unload package
+	TArray<UPackage*> PackagesToUnload;
 	{
-		return;
+		PackagesToUnload.Add(SequencePackage);
 	}
 
-	// refresh sequencer
-	if (SequenceInstanceStack.Num() == 0)
+	FText PackageUnloadError;
+	PackageTools::UnloadPackages(PackagesToUnload, PackageUnloadError);
+
+	if (!PackageUnloadError.IsEmpty())
 	{
-		RootMovieSceneSequenceInstance = nullptr;
-		ResetToNewRootSequence(*ReloadedSequence);
+		ResetLoaders(SequencePackage);
+		SequencePackage->ClearFlags(RF_WasLoaded);
+		SequencePackage->bHasBeenFullyLoaded = false;
+		SequencePackage->GetMetaData()->RemoveMetaDataOutsidePackage();
 	}
-	else
+
+	// reload package
+	TMap<UObject*, UObject*> MovedToReloadedObjectMap;
+
+	for (const auto MovedObject : MovedObjects)
+	{
+		UObject* ReloadedObject = StaticLoadObject(MovedObject.Value->GetClass(), nullptr, *MovedObject.Key);
+		check(ReloadedObject != nullptr);
+		MovedToReloadedObjectMap.Add(MovedObject.Value, ReloadedObject);
+	}
+
+	for (TObjectIterator<UObject> It; It; ++It)
+	{
+		// @todo sequencer: only process objects that actually reference the package?
+		FArchiveReplaceObjectRef<UObject> Ar(*It, MovedToReloadedObjectMap, false, false, false, false);
+	}
+
+	auto ReloadedSequence = Cast<UMovieSceneSequence>(StaticLoadObject(SequenceClass, nullptr, *SequencePath));
+	SequencePackage->SetDirtyFlag(false);
+
+	// release transient objects
+	for (auto MovedObject : MovedObjects)
+	{
+		MovedObject.Value->RemoveFromRoot();
+		MovedObject.Value->MarkPendingKill();
+	}
+
+	// clear undo buffer
+	if (true) // @todo sequencer: check whether objects are actually referenced in undo buffer
+	{
+		GEditor->Trans->Reset(LOCTEXT("UnloadedSequence", "Unloaded Sequence"));
+	}
+
+	// reopen asset editor
+	TArray<UObject*> AssetsToReopen;
+	AssetsToReopen.Add(ReloadedSequence);
+	AssetEditorManager.OpenEditorForAssets(AssetsToReopen, EToolkitMode::Standalone, MyToolkitHost.ToSharedRef());
+}
+
+
+void FSequencer::FixActorReferences()
+{
+	UMovieScene* FocusedMovieScene = GetFocusedMovieSceneSequence()->GetMovieScene();
+
+	TMap<FString, AActor*> ActorNameToActorMap;
+	for ( TActorIterator<AActor> ActorItr( GWorld ); ActorItr; ++ActorItr )
+	{
+		// Same as with the Object Iterator, access the subclass instance with the * or -> operators.
+		AActor *Actor = *ActorItr;
+		ActorNameToActorMap.Add( Actor->GetActorLabel(), Actor);
+	}
+
+	// Cache the possessables to fix up first since the bindings will change as the fix ups happen.
+	TArray<FMovieScenePossessable> ActorsPossessablesToFix;
+	for ( int32 i = 0; i < FocusedMovieScene->GetPossessableCount(); i++ )
+	{
+		FMovieScenePossessable& Possessable = FocusedMovieScene->GetPossessable( i );
+		// Possesbles with parents are components so ignore them.
+		if ( Possessable.GetParent().IsValid() == false )
+		{
+			TArray<TWeakObjectPtr<UObject>> RuntimeObjects;
+			GetRuntimeObjects( GetFocusedMovieSceneSequenceInstance(), Possessable.GetGuid(), RuntimeObjects );
+			if ( RuntimeObjects.Num() == 0 )
+			{
+				ActorsPossessablesToFix.Add( Possessable );
+			}
+		}
+	}
+
+	// For the posseables to fix, look up the actors by name and reassign them if found.
+	for ( const FMovieScenePossessable& ActorPossessableToFix : ActorsPossessablesToFix )
 	{
-		FocusSequenceInstance(MakeShareable(new FMovieSceneSequenceInstance(*ReloadedSequence)));
-	}*/
+		AActor** ActorPtr = ActorNameToActorMap.Find( ActorPossessableToFix.GetName() );
+		if ( ActorPtr != nullptr )
+		{
+			DoAssignActor( ActorPtr, 1, ActorPossessableToFix.GetGuid() );
+		}
+	}
 }
 
 
@@ -3468,57 +4434,25 @@ void FSequencer::SplitSection()
 	NotifyMovieSceneDataChanged();
 }
 
+const ISequencerEditTool* FSequencer::GetEditTool() const
+{
+	return SequencerWidget->GetEditTool();
+}
 
-/** Get the current active hotspot */
 TSharedPtr<ISequencerHotspot> FSequencer::GetHotspot() const
 {
 	return Hotspot;
 }
 
-
-/** Set the hotspot to something else */
 void FSequencer::SetHotspot(TSharedPtr<ISequencerHotspot> NewHotspot)
 {
 	Hotspot = MoveTemp(NewHotspot);
 }
 
-
 void FSequencer::BindSequencerCommands()
 {
 	const FSequencerCommands& Commands = FSequencerCommands::Get();
 
-	SequencerCommandBindings->MapAction(
-		FGenericCommands::Get().Delete,
-		FExecuteAction::CreateSP( this, &FSequencer::DeleteSelectedItems ),
-		FCanExecuteAction::CreateSP(this, &FSequencer::IsSequencerWidgetFocused ));
-
-	SequencerCommandBindings->MapAction(
-		Commands.TogglePlay,
-		FExecuteAction::CreateSP( this, &FSequencer::TogglePlay ),
-		FCanExecuteAction::CreateSP(this, &FSequencer::IsSequencerWidgetFocused ));
-
-	SequencerCommandBindings->MapAction(
-		Commands.PlayForward,
-		FExecuteAction::CreateSP( this, &FSequencer::PlayForward ),
-		FCanExecuteAction::CreateSP(this, &FSequencer::IsSequencerWidgetFocused ) );
-
-	SequencerCommandBindings->MapAction(
-		Commands.Rewind,
-		FExecuteAction::CreateSP( this, &FSequencer::Rewind ),
-		FCanExecuteAction::CreateSP(this, &FSequencer::IsSequencerWidgetFocused ) );
-
-	SequencerCommandBindings->MapAction(
-		Commands.StepForward,
-		FExecuteAction::CreateSP( this, &FSequencer::StepForward ),
-		FCanExecuteAction::CreateSP(this, &FSequencer::IsSequencerWidgetFocused ),
-		EUIActionRepeatMode::RepeatEnabled );
-
-	SequencerCommandBindings->MapAction(
-		Commands.StepBackward,
-		FExecuteAction::CreateSP( this, &FSequencer::StepBackward ),
-		FCanExecuteAction::CreateSP(this, &FSequencer::IsSequencerWidgetFocused ),
-		EUIActionRepeatMode::RepeatEnabled );
-
 	SequencerCommandBindings->MapAction(
 		Commands.StepToNextKey,
 		FExecuteAction::CreateSP( this, &FSequencer::StepToNextKey ) );
@@ -3537,7 +4471,11 @@ void FSequencer::BindSequencerCommands()
 
 	SequencerCommandBindings->MapAction(
 		Commands.SetStartPlaybackRange,
-		FExecuteAction::CreateSP( this, &FSequencer::SetStartPlaybackRange ) );
+		FExecuteAction::CreateSP( this, &FSequencer::SetPlaybackRangeStart ) );
+
+	SequencerCommandBindings->MapAction(
+		Commands.ResetViewRange,
+		FExecuteAction::CreateSP( this, &FSequencer::ResetViewRange ) );
 
 	SequencerCommandBindings->MapAction(
 		Commands.ZoomInViewRange,
@@ -3553,7 +4491,7 @@ void FSequencer::BindSequencerCommands()
 
 	SequencerCommandBindings->MapAction(
 		Commands.SetEndPlaybackRange,
-		FExecuteAction::CreateSP( this, &FSequencer::SetEndPlaybackRange ) );
+		FExecuteAction::CreateSP( this, &FSequencer::SetPlaybackRangeEnd ) );
 
 	SequencerCommandBindings->MapAction(
 		Commands.ExpandAllNodesAndDescendants,
@@ -3672,6 +4610,14 @@ void FSequencer::BindSequencerCommands()
 		FCanExecuteAction::CreateLambda( []{ return true; } ),
 		FIsActionChecked::CreateLambda( [this]{ return Settings->GetShowCombinedKeyframes(); } ) );
 
+	SequencerCommandBindings->MapAction(
+		Commands.ToggleChannelColors,
+		FExecuteAction::CreateLambda( [this]{
+			Settings->SetShowChannelColors( !Settings->GetShowChannelColors() );
+		} ),
+		FCanExecuteAction::CreateLambda( []{ return true; } ),
+		FIsActionChecked::CreateLambda( [this]{ return Settings->GetShowChannelColors(); } ) );
+
 	SequencerCommandBindings->MapAction(
 		Commands.ToggleLabelBrowser,
 		FExecuteAction::CreateLambda( [this]{
@@ -3686,10 +4632,6 @@ void FSequencer::BindSequencerCommands()
 		FCanExecuteAction::CreateSP( this, &FSequencer::CanShowFrameNumbers ),
 		FIsActionChecked::CreateLambda( [this]{ return Settings->GetShowFrameNumbers(); } ) );
 
-	SequencerCommandBindings->MapAction(
-		Commands.ToggleShowGotoBox,
-		FExecuteAction::CreateLambda( [this]{ Settings->SetShowGotoBox( !Settings->GetShowGotoBox() ); } ) );
-
 	SequencerCommandBindings->MapAction(
 		Commands.ToggleShowRangeSlider,
 		FExecuteAction::CreateLambda( [this]{ Settings->SetShowRangeSlider( !Settings->GetShowRangeSlider() ); } ),
@@ -3726,6 +4668,12 @@ void FSequencer::BindSequencerCommands()
 		FCanExecuteAction::CreateLambda( []{ return true; } ),
 		FIsActionChecked::CreateLambda( [this]{ return Settings->GetSnapSectionTimesToSections(); } ) );
 
+	SequencerCommandBindings->MapAction(
+		Commands.ToggleSnapPlayTimeToKeys,
+		FExecuteAction::CreateLambda( [this]{ Settings->SetSnapPlayTimeToKeys( !Settings->GetSnapPlayTimeToKeys() ); } ),
+		FCanExecuteAction::CreateLambda( []{ return true; } ),
+		FIsActionChecked::CreateLambda( [this]{ return Settings->GetSnapPlayTimeToKeys(); } ) );
+
 	SequencerCommandBindings->MapAction(
 		Commands.ToggleSnapPlayTimeToInterval,
 		FExecuteAction::CreateLambda( [this]{ Settings->SetSnapPlayTimeToInterval( !Settings->GetSnapPlayTimeToInterval() ); } ),
@@ -3763,11 +4711,6 @@ void FSequencer::BindSequencerCommands()
 		FIsActionChecked::CreateLambda( [this]{ return Settings->GetLinkCurveEditorTimeRange(); } ) );
 
 	auto CanCutOrCopy = [this]{
-		if (!IsSequencerWidgetFocused())
-		{
-			return false;
-		}
-		
 		UMovieSceneTrack* Track = nullptr;
 		for (FSequencerSelectedKey Key : Selection.GetSelectedKeys())
 		{
@@ -3844,23 +4787,72 @@ void FSequencer::BindSequencerCommands()
 			}
 
 			UMovieSceneSequence* EditedSequence = SequenceInstanceStack.Top()->GetSequence();
+			UPackage* EditedPackage = EditedSequence->GetOutermost();
 
-			return ((EditedSequence != nullptr) && EditedSequence->GetOutermost()->IsDirty());
+			return ((EditedSequence != nullptr) && (EditedPackage->FileSize != 0) && EditedPackage->IsDirty());
 		})
 	);
 
-
-	SequencerWidget->BindCommands(SequencerCommandBindings);
+	SequencerCommandBindings->MapAction(
+		Commands.FixActorReferences,
+		FExecuteAction::CreateSP( this, &FSequencer::FixActorReferences ),
+		FCanExecuteAction::CreateLambda( []{ return true; } ) );
 
 	for (int32 i = 0; i < TrackEditors.Num(); ++i)
 	{
 		TrackEditors[i]->BindCommands(SequencerCommandBindings);
 	}
+
+	// copy subset of sequencer commands to shared commands
+	*SequencerSharedBindings = *SequencerCommandBindings;
+
+	// Sequencer-only bindings
+	SequencerCommandBindings->MapAction(
+		FGenericCommands::Get().Delete,
+		FExecuteAction::CreateSP( this, &FSequencer::DeleteSelectedItems ));
+
+	SequencerCommandBindings->MapAction(
+		Commands.TogglePlay,
+		FExecuteAction::CreateSP( this, &FSequencer::TogglePlay ));
+
+	SequencerCommandBindings->MapAction(
+		Commands.PlayForward,
+		FExecuteAction::CreateSP( this, &FSequencer::PlayForward ));
+
+	SequencerCommandBindings->MapAction(
+		Commands.Rewind,
+		FExecuteAction::CreateSP( this, &FSequencer::Rewind ));
+
+	SequencerCommandBindings->MapAction(
+		Commands.StepForward,
+		FExecuteAction::CreateSP( this, &FSequencer::StepForward ),
+		EUIActionRepeatMode::RepeatEnabled );
+
+	SequencerCommandBindings->MapAction(
+		Commands.StepBackward,
+		FExecuteAction::CreateSP( this, &FSequencer::StepBackward ),
+		EUIActionRepeatMode::RepeatEnabled );
+
+	SequencerCommandBindings->MapAction(
+		Commands.SetInOutEnd,
+		FExecuteAction::CreateLambda([this]{ SetInOutRangeEnd(); }));
+
+	SequencerCommandBindings->MapAction(
+		Commands.SetInOutStart,
+		FExecuteAction::CreateLambda([this]{ SetInOutRangeStart(); }));
+
+	SequencerWidget->BindCommands(SequencerCommandBindings);
 }
 
 
 void FSequencer::BuildAddTrackMenu(class FMenuBuilder& MenuBuilder)
 {
+	MenuBuilder.AddMenuEntry(
+		LOCTEXT( "AddFolder", "Add Folder" ),
+		LOCTEXT( "AddFolderToolTip", "Adds a new folder." ),
+		FSlateIcon( FEditorStyle::GetStyleSetName(), "ContentBrowser.AssetTreeFolderOpen" ),
+		FUIAction( FExecuteAction::CreateRaw( this, &FSequencer::OnAddFolder ) ) );
+
 	for (int32 i = 0; i < TrackEditors.Num(); ++i)
 	{
 		TrackEditors[i]->BuildAddTrackMenu(MenuBuilder);
@@ -3886,11 +4878,4 @@ void FSequencer::BuildObjectBindingEditButtons(TSharedPtr<SHorizontalBox> EditBo
 }
 
 
-bool FSequencer::IsSequencerWidgetFocused() const
-{
-	// Only perform certain operations if the key was pressed in the sequencer widget and not in, for example, the viewport.
-	return SequencerWidget->HasFocusedDescendants() || SequencerWidget->HasKeyboardFocus();
-}
-
-
 #undef LOCTEXT_NAMESPACE
