Copying //UE4/Dev-Sequencer to //UE4/Main

diff --git a/Engine/Source/Editor/Sequencer/Private/Sequencer.h b/Engine/Source/Editor/Sequencer/Private/Sequencer.h
--- a/Engine/Source/Editor/Sequencer/Private/Sequencer.h
+++ b/Engine/Source/Editor/Sequencer/Private/Sequencer.h
@@ -6,7 +6,7 @@
 #include "EditorUndoClient.h"
 #include "MovieSceneClipboard.h"
 #include "SequencerLabelManager.h"
-
+#include "LevelEditor.h"
 
 class FMenuBuilder;
 class FMovieSceneSequenceInstance;
@@ -22,10 +22,12 @@ class UMovieSceneSection;
 class UMovieSceneSequence;
 class UWorld;
 class IMovieSceneSpawnRegister;
+class FSequencerNodeTree;
 
 struct ISequencerHotspot;
 struct FTransformData;
 
+
 /**
  * Sequencer is the editing tool for MovieScene assets.
  */
@@ -63,11 +65,79 @@ public:
 	FAnimatedRange GetClampRange() const;
 	void SetClampRange(TRange<float> InNewClampRange);
 
+public:
+
+	/**
+	 * Get the in/out selection range.
+	 *
+	 * @return The in/out range.
+	 * @see SetInOutRange, SetInOutRangeEnd, SetInOutRangeStart
+	 */
+	TRange<float> GetInOutRange() const;
+
+	/**
+	 * Set the in/out selection range.
+	 *
+	 * @param Range The new range to set.
+	 * @see GetInOutRange, SetInOutRangeEnd, SetInOutRangeStart
+	 */
+	void SetInOutRange(TRange<float> Range);
+	
+	/**
+	 * Set the in/out range's end position to the current global time.
+	 *
+	 * @see GetInOutRange, SetInOutRange, SetInOutRangeStart
+	 */
+	void SetInOutRangeEnd();
+	
+	/**
+	 * Set the in/out range's start position to the current global time.
+	 *
+	 * @see GetInOutRange, SetInOutRange, SetInOutRangeEnd
+	 */
+	void SetInOutRangeStart();
+
+public:
+
+	/**
+	 * Get the playback range.
+	 *
+	 * @return Playback range.
+	 * @see SetPlaybackRange, SetPlaybackRangeEnd, SetPlaybackRangeStart
+	 */
 	TRange<float> GetPlaybackRange() const;
-	void SetPlaybackRange(TRange<float> InNewPlaybackRange);
-	void SetStartPlaybackRange();
-	void SetEndPlaybackRange();
 
+	/**
+	 * Set the playback range.
+	 *
+	 * @param Range The new range to set.
+	 * @see GetPlaybackRange, SetPlaybackRangeEnd, SetPlaybackRangeStart
+	 */
+	void SetPlaybackRange(TRange<float> Range);
+	
+	/**
+	 * Set the playback range's end position to the current global time.
+	 *
+	 * @see GetPlaybackRange, SetPlaybackRange, SetPlaybackRangeStart
+	 */
+	void SetPlaybackRangeEnd()
+	{
+		SetPlaybackRange(TRange<float>(GetPlaybackRange().GetLowerBoundValue(), GetGlobalTime()));
+	}
+
+	/**
+	 * Set the playback range's start position to the current global time.
+	 *
+	 * @see GetPlaybackRange, SetPlaybackRange, SetPlaybackRangeStart
+	 */
+	void SetPlaybackRangeStart()
+	{
+		SetPlaybackRange(TRange<float>(GetGlobalTime(), GetPlaybackRange().GetUpperBoundValue()));
+	}
+
+public:
+
+	void ResetViewRange();
 	void ZoomViewRange(float InZoomDelta);
 	void ZoomInViewRange();
 	void ZoomOutViewRange();
@@ -80,6 +150,12 @@ public:
 		return Settings;
 	}
 
+	/** Gets the tree of nodes which is used to populate the animation outliner. */
+	TSharedRef<FSequencerNodeTree> GetNodeTree()
+	{
+		return NodeTree;
+	}
+
 	bool IsPerspectiveViewportPossessionEnabled() const override
 	{
 		return bPerspectiveViewportPossessionEnabled;
@@ -189,7 +265,15 @@ public:
 	float GetOverlayFadeCurve() const;
 
 	/** Gets the command bindings for the sequencer */
-	virtual TSharedPtr<FUICommandList> GetCommandBindings() const override { return SequencerCommandBindings; }
+	virtual TSharedPtr<FUICommandList> GetCommandBindings(ESequencerCommandBindings Type = ESequencerCommandBindings::Sequencer) const override
+	{
+		if (Type == ESequencerCommandBindings::Sequencer)
+		{
+			return SequencerCommandBindings;
+		}
+
+		return SequencerSharedBindings;
+	}
 
 	/**
 	 * Builds up the sequencer's "Add Track" menu.
@@ -223,10 +307,19 @@ public:
 	void AttachTransportControlsToViewports();
 	/** Purges all transport control widgets from the main level editor viewports */
 	void DetachTransportControlsFromViewports();
-	
+
+	/** Gets the color and opacity of a transport control */
+	FLinearColor GetTransportControlsColorAndOpacity(TSharedRef<SWidget> TransportWidget) const;
+
 	/** Gets the visibility of a transport control based on the viewport it's attached to */
 	EVisibility GetTransportControlVisibility(TSharedPtr<ILevelViewport> LevelViewport) const;
 
+	/** Extend the level editor viewport menu */
+	TSharedRef<FExtender> OnExtendLevelEditorViewMenu(const TSharedRef<FUICommandList> CommandList);
+
+	/** Create a menu entry we can use to toggle the transport controls */
+	void CreateTransportToggleMenuEntry(FMenuBuilder& MenuBuilder);
+
 	/** Functions to push on to the transport controls we use */
 	FReply OnPlay(bool bTogglePlay=true);
 	FReply OnRecord();
@@ -235,6 +328,19 @@ public:
 	FReply OnStepToEnd();
 	FReply OnStepToBeginning();
 	FReply OnToggleLooping();
+	FReply SetPlaybackEnd();
+	FReply SetPlaybackStart();
+	FReply JumpToPreviousKey();
+	FReply JumpToNextKey();
+
+	/** Get the visibility of the record button */
+	EVisibility GetRecordButtonVisibility() const;
+
+	/** Delegate handler for recording starting */
+	void HandleRecordingStarted(UMovieSceneSequence* Sequence);
+
+	/** Delegate handler for recording finishing */
+	void HandleRecordingFinished(UMovieSceneSequence* Sequence);
 
 	bool IsLooping() const;
 
@@ -290,6 +396,12 @@ public:
 	/** Select keys belonging to a section at the key time */
 	void SelectTrackKeys(TWeakObjectPtr<UMovieSceneSection> Section, float KeyTime, bool bAddToSelection, bool bToggleSelection);
 
+	/** Updates the external selection to match the current sequencer selection. */
+	void SynchronizeExternalSelectionWithSequencerSelection();
+
+	/** Updates the sequencer selection to match the current external selection. */
+	void SynchronizeSequencerSelectionWithExternalSelection();
+
 public:
 
 	/** Copy the selected keys to the clipboard */
@@ -307,7 +419,13 @@ public:
 	/** Discard all changes to the current movie scene. */
 	void DiscardChanges();
 
+	/** Attempts to automatically fix up broken actor references in the current scene. */
+	void FixActorReferences();
+
 public:
+	
+	/** Access the currently active track area edit tool */
+	const ISequencerEditTool* GetEditTool() const;
 
 	/** Get the current active hotspot */
 	TSharedPtr<ISequencerHotspot> GetHotspot() const;
@@ -367,13 +485,18 @@ public:
 	virtual void SetKeyInterpPropertiesOnly(bool bKeyInterpPropertiesOnly) override;
 	virtual EMovieSceneKeyInterpolation GetKeyInterpolation() const override;
 	virtual void SetKeyInterpolation(EMovieSceneKeyInterpolation) override;
+	virtual bool GetInfiniteKeyAreas() const override;
+	virtual void SetInfiniteKeyAreas(bool bInfiniteKeyAreas) override;
 	virtual bool IsRecordingLive() const override;
 	virtual float GetCurrentLocalTime(UMovieSceneSequence& InMovieSceneSequence) override;
-	virtual float GetGlobalTime() override;
-	virtual void SetGlobalTime(float Time, const bool& bAllowSnappingToFrames = false) override;
-	virtual void SetGlobalTimeDirectly(float Time, const bool& bAllowSnappingToFrames = false) override;
+	virtual float GetGlobalTime() const override;
+	virtual void SetGlobalTime(float Time, ESnapTimeMode SnapTimeMode = ESnapTimeMode::STM_None) override;
+	virtual void SetGlobalTimeDirectly(float Time, ESnapTimeMode SnapTimeMode = ESnapTimeMode::STM_None) override;
 	virtual void SetPerspectiveViewportPossessionEnabled(bool bEnabled) override;
 	virtual void SetPerspectiveViewportCameraCutEnabled(bool bEnabled) override;
+	virtual void EnterSilentMode() override { ++SilentModeCount; }
+	virtual void ExitSilentMode() override { --SilentModeCount; ensure(SilentModeCount >= 0); }
+	virtual bool IsInSilentMode() const override { return SilentModeCount != 0; }
 	virtual FGuid GetHandleToObject(UObject* Object, bool bCreateHandleIfMissing = true) override;
 	virtual ISequencerObjectChangeListener& GetObjectChangeListener() override;
 	virtual void NotifyMovieSceneDataChanged() override;
@@ -389,17 +512,20 @@ public:
 	virtual FGuid CreateBinding(UObject& InObject, const FString& InName) override;
 	virtual UObject* GetPlaybackContext() const override;
 	virtual void GetAllKeyedProperties(UObject& Object, TSet<UProperty*>& OutProperties) override;
+	virtual FOnActorAddedToSequencer& OnActorAddedToSequencer() override;
 	virtual FOnCameraCut& OnCameraCut() override;
 	virtual TSharedRef<INumericTypeInterface<float>> GetNumericTypeInterface() override;
 	virtual TSharedRef<INumericTypeInterface<float>> GetZeroPadNumericTypeInterface() override;
-	virtual TSharedRef<SWidget> MakeTransportControls() override;
-
+	virtual TSharedRef<SWidget> MakeTransportControls(bool bExtended) override;
+	virtual TSharedRef<SWidget> MakeTimeRange(const TSharedRef<SWidget>& InnerContent, bool bShowWorkingRange, bool bShowViewRange, bool bShowPlaybackRange) override;
+	virtual void SetViewportTransportControlsVisibility(bool bVisible) override;
+	
 public:
 
 	// IMovieScenePlayer interface
 
-	virtual void GetRuntimeObjects( TSharedRef<FMovieSceneSequenceInstance> MovieSceneInstance, const FGuid& ObjectHandle, TArray< UObject* >& OutObjects ) const override;
-	virtual void UpdateCameraCut(UObject* CameraObject, UObject* UnlockIfCameraObject) const override;
+	virtual void GetRuntimeObjects(TSharedRef<FMovieSceneSequenceInstance> MovieSceneInstance, const FGuid& ObjectHandle, TArray<TWeakObjectPtr<UObject>>& OutObjects) const override;
+	virtual void UpdateCameraCut(UObject* CameraObject, UObject* UnlockIfCameraObject, bool bJumpCut) const override;
 	virtual void SetViewportSettings(const TMap<FViewportClient*, EMovieSceneViewportParams>& ViewportParamsMap) override;
 	virtual void GetViewportSettings(TMap<FViewportClient*, EMovieSceneViewportParams>& ViewportParamsMap) const override;
 	virtual EMovieScenePlayerStatus::Type GetPlaybackStatus() const override;
@@ -407,7 +533,7 @@ public:
 	virtual void AddOrUpdateMovieSceneInstance(UMovieSceneSection& MovieSceneSection, TSharedRef<FMovieSceneSequenceInstance> InstanceToAdd) override;
 	virtual void RemoveMovieSceneInstance(UMovieSceneSection& MovieSceneSection, TSharedRef<FMovieSceneSequenceInstance> InstanceToRemove) override;
 	virtual IMovieSceneSpawnRegister& GetSpawnRegister() override { return *SpawnRegister; }
-
+	virtual bool IsPreview() const override { return SilentModeCount != 0; }
 protected:
 
 	/** Reset data about a movie scene when pushing or popping a movie scene. */
@@ -435,9 +561,6 @@ protected:
 	/** @return The current scrub position */
 	virtual float OnGetScrubPosition() const { return ScrubPosition; }
 
-	/** Constrains the maximum frame rate to the fixed time step rate when playing back in that mode */
-	void ConstrainFixedTimeStepFrameRate();
-
 	/**
 	 * Called when the clamp range is changed by the user
 	 *
@@ -459,12 +582,18 @@ protected:
 	/** Called when the user has finished scrubbing */
 	void OnEndScrubbing();
 
-	/** Called when the user has begun dragging the play range */
+	/** Called when the user has begun dragging the playback range */
 	void OnBeginPlaybackRangeDrag();
 
-	/** Called when the user has finished dragging the play range */
+	/** Called when the user has finished dragging the playback range */
 	void OnEndPlaybackRangeDrag();
 
+	/** Called when the user has begun dragging the selection range */
+	void OnBeginInOutRangeDrag();
+
+	/** Called when the user has finished dragging the selection range */
+	void OnEndInOutRangeDrag();
+
 protected:
 
 	/**
@@ -513,6 +642,12 @@ protected:
 	 */
 	void GetKeyedProperties(UObject& Object, FMovieSceneSequenceInstance& Instance, TSet<UProperty*>& OutProperties);
 
+	/** Get all the keys for the current sequencer selection */
+	virtual void GetKeysFromSelection(TUniquePtr<ISequencerKeyCollection>& KeyCollection) override;
+
+	/** Find the nearest keyframe */
+	float FindNearestKey(float NewScrubPosition) override;
+
 protected:
 
 	/** Called via UEditorEngine::GetActorRecordingStateEvent to check to see whether we need to record actor state */
@@ -593,24 +728,47 @@ protected:
 	void OnEndPIE(bool bIsSimulating);
 
 	/** Updates a viewport client from camera cut data */
-	void UpdatePreviewLevelViewportClientFromCameraCut(FLevelEditorViewportClient& InViewportClient, UObject* InCameraObject) const;
-
-	/** Is the sequencer widget focused? */
-	bool IsSequencerWidgetFocused() const;
+	void UpdatePreviewLevelViewportClientFromCameraCut(FLevelEditorViewportClient& InViewportClient, UObject* InCameraObject, bool bJumpCut) const;
 
 	/** Internal conversion function that doesn't perform expensive reset/update tasks */
 	void ConvertToSpawnableInternal(TSharedRef<FSequencerObjectBindingNode> NodeToBeConverted);
-	
+
+	/** Handles adding a new folder to the outliner tree. */
+	void OnAddFolder();
+
+	/** Handles the actor selection changing externally .*/
+	void OnActorSelectionChanged( UObject* );
+
+	/** Create set playback start transport control */
+	TSharedRef<SWidget> OnCreateTransportSetPlaybackStart();
+
+	/** Create jump to previous key transport control */
+	TSharedRef<SWidget> OnCreateTransportJumpToPreviousKey();
+
+	/** Create jump to next key transport control */
+	TSharedRef<SWidget> OnCreateTransportJumpToNextKey();
+
+	/** Create set playback end transport control */
+	TSharedRef<SWidget> OnCreateTransportSetPlaybackEnd();
+
+	/** Create record transport control */
+	TSharedRef<SWidget> OnCreateTransportRecord();
+
 private:
+	/** Performs any post-tick rendering work needed when moving through scenes */
+	void PostTickRenderStateFixup();
 
 	/** User-supplied settings object for this sequencer */
 	USequencerSettings* Settings;
 
 	TMap< TWeakObjectPtr<UMovieSceneSection>, TSharedRef<FMovieSceneSequenceInstance> > SequenceInstanceBySection;
 
-	/** Command list for sequencer commands */
+	/** Command list for sequencer commands (Sequencer widgets only). */
 	TSharedRef<FUICommandList> SequencerCommandBindings;
 
+	/** Command list for sequencer commands (shared by non-Sequencer). */
+	TSharedRef<FUICommandList> SequencerSharedBindings;
+
 	/** List of tools we own */
 	TArray<TSharedPtr<ISequencerTrackEditor>> TrackEditors;
 
@@ -666,13 +824,6 @@ private:
 	/** Whether we are playing, recording, etc. */
 	EMovieScenePlayerStatus::Type PlaybackState;
 
-	/** The real-time that we started playback last */
-	double PlaybackStartRealTime;
-
-	/** Number of continuous fixed time step frames we've played so far without any change in play back state,
-	    such as time step, reverse mode, etc. */
-	uint32 NumContinuousFixedTimeStepFrames;
-
 	/** The current scrub position */
 	// @todo sequencer: Should use FTimespan or "double" for Time Cursor Position! (cascades)
 	float ScrubPosition;
@@ -691,15 +842,37 @@ private:
 		the MovieScene data can change many times per frame.) */
 	bool bNeedTreeRefresh;
 
+	/** A flag which indicates Close() was called on this Sequencer. */ 
+	bool bWasClosed;
+
 	FSequencerLabelManager LabelManager;
 	FSequencerSelection Selection;
 	FSequencerSelectionPreview SelectionPreview;
 
+	/** Represents the tree of nodes to display in the animation outliner. */
+	TSharedRef<FSequencerNodeTree> NodeTree;
+
 	/** A delegate which is called any time the global time changes. */
 	FOnGlobalTimeChanged OnGlobalTimeChangedDelegate;
 
 	/** A map of all the transport controls to viewports that this sequencer has made */
 	TMap< TSharedPtr<class ILevelViewport>, TSharedPtr<class SWidget> > TransportControls;
 
+	FOnActorAddedToSequencer OnActorAddedToSequencerEvent;
 	FOnCameraCut OnCameraCutEvent;
+
+	int32 SilentModeCount;
+	
+	/** Menu extender for level viewports */
+	FLevelEditorModule::FLevelEditorMenuExtender ViewMenuExtender;
+	FDelegateHandle LevelEditorExtenderDelegateHandle;
+
+	/** When true the sequencer selection is being updated from changes to the external seleciton. */
+	bool bUpdatingSequencerSelection;
+
+	/** When true the external selection is being updated from changes to the sequencer selection. */
+	bool bUpdatingExternalSelection;
+
+	/** The maximum tick rate prior to playing (used for overriding delta time during playback). */
+	double OldMaxTickRate;
 };
diff --git a/Engine/Source/Editor/Sequencer/Private/SequencerCommands.cpp b/Engine/Source/Editor/Sequencer/Private/SequencerCommands.cpp
--- a/Engine/Source/Editor/Sequencer/Private/SequencerCommands.cpp
+++ b/Engine/Source/Editor/Sequencer/Private/SequencerCommands.cpp
@@ -18,16 +18,20 @@ void FSequencerCommands::RegisterCommands()
 	UI_COMMAND( StepToPreviousCameraKey, "Step to Previous Camera Key", "Step to the previous camera key", EUserInterfaceActionType::Button, FInputChord(EModifierKey::Alt, EKeys::Comma) );
 	UI_COMMAND( SetStartPlaybackRange, "Set Start Playback Range", "Set the start playback range", EUserInterfaceActionType::Button, FInputChord(EKeys::LeftBracket) );
 	UI_COMMAND( SetEndPlaybackRange, "Set End Playback Range", "Set the end playback range", EUserInterfaceActionType::Button, FInputChord(EKeys::RightBracket) );
+	UI_COMMAND( ResetViewRange, "Reset View Range", "Reset view range to the playback range", EUserInterfaceActionType::Button, FInputChord(EKeys::Home) );
 	UI_COMMAND( ZoomInViewRange, "Zoom into the View Range", "Zoom into the view range", EUserInterfaceActionType::Button, FInputChord(EKeys::Equals) );
 	UI_COMMAND( ZoomOutViewRange, "Zoom out of the View Range", "Zoom out of the view range", EUserInterfaceActionType::Button, FInputChord(EKeys::Hyphen) );
-	UI_COMMAND( ToggleKeepCursorInPlaybackRange, "Keep Cursor In Playback Range", "When checked, the cursor will be constrained to the current playback range during playback", EUserInterfaceActionType::ToggleButton, FInputChord() );
-	UI_COMMAND( ToggleKeepPlaybackRangeInSectionBounds, "Keep Playback Range In Section Bounds", "When checked, the playback range will be synchronized to the section bounds", EUserInterfaceActionType::ToggleButton, FInputChord() );
+	UI_COMMAND( ToggleKeepCursorInPlaybackRange, "Keep Cursor in Playback Range", "When checked, the cursor will be constrained to the current playback range during playback", EUserInterfaceActionType::ToggleButton, FInputChord() );
+	UI_COMMAND( ToggleKeepPlaybackRangeInSectionBounds, "Keep Playback Range in Section Bounds", "When checked, the playback range will be synchronized to the section bounds", EUserInterfaceActionType::ToggleButton, FInputChord() );
 
 	UI_COMMAND( ExpandAllNodesAndDescendants, "Expand All Nodes", "Expand all nodes", EUserInterfaceActionType::Button, FInputChord() );
 	UI_COMMAND( CollapseAllNodesAndDescendants, "Collapse All Nodes", "Collapse all selected nodes", EUserInterfaceActionType::Button, FInputChord() );
-	UI_COMMAND( ToggleExpandCollapseNodes, "Expand/Collapse Nodes", "Toggle expand or collapse selected nodes", EUserInterfaceActionType::Button, FInputChord(EKeys::O) );
+	UI_COMMAND( ToggleExpandCollapseNodes, "Expand/Collapse Nodes", "Toggle expand or collapse selected nodes", EUserInterfaceActionType::Button, FInputChord(EModifierKey::Control, EKeys::O) );
 	UI_COMMAND( ToggleExpandCollapseNodesAndDescendants, "Expand/Collapse Nodes and Descendants", "Toggle expand or collapse selected nodes and descendants", EUserInterfaceActionType::Button, FInputChord(EModifierKey::Shift, EKeys::O) );
 
+	UI_COMMAND( SetInOutEnd, "Set selection end", "Sets the end of the in/out selection", EUserInterfaceActionType::Button, FInputChord(EKeys::O) );
+	UI_COMMAND( SetInOutStart, "Set selection start", "Sets the start of the in/out selection", EUserInterfaceActionType::Button, FInputChord(EKeys::I) );
+
 	UI_COMMAND( SetKey, "Set Key", "Sets a key on the selected tracks", EUserInterfaceActionType::Button, FInputChord(EKeys::Enter) );
 	UI_COMMAND( SetInterpolationCubicAuto, "Set Key Auto", "Cubic interpolation - Automatic tangents", EUserInterfaceActionType::Button, FInputChord(EKeys::One));
 	UI_COMMAND( SetInterpolationCubicUser, "Set Key User", "Cubic interpolation - User flat tangents", EUserInterfaceActionType::Button, FInputChord(EKeys::Two));
@@ -43,12 +47,12 @@ void FSequencerCommands::RegisterCommands()
 	UI_COMMAND( SetAutoKeyModeAnimated, "Auto-key Animated", "Enables auto keying for properties with existing animations.", EUserInterfaceActionType::ToggleButton, FInputChord());
 	UI_COMMAND( SetAutoKeyModeNone, "Disable Auto-key", "Disables auto keying.", EUserInterfaceActionType::ToggleButton, FInputChord());
 
-	UI_COMMAND(ToggleKeyAllEnabled, "Key All", "Enables and disables key all", EUserInterfaceActionType::ToggleButton, FInputChord());
+	UI_COMMAND( ToggleKeyAllEnabled, "Key All", "Enables and disables key all", EUserInterfaceActionType::ToggleButton, FInputChord());
 
 	UI_COMMAND( ToggleAutoScroll, "Auto Scroll", "Toggle auto-scroll: When enabled, automatically scrolls the sequencer view to keep the current time visible", EUserInterfaceActionType::ToggleButton, FInputChord(EModifierKey::Shift, EKeys::S) );
 
 	UI_COMMAND( ToggleShowFrameNumbers, "Show Frame Numbers", "Enables and disables showing frame numbers", EUserInterfaceActionType::ToggleButton, FInputChord() );
-	UI_COMMAND( ToggleShowGotoBox, "Go to time...", "Go to a particular point on the time line", EUserInterfaceActionType::Button, FInputChord(EModifierKey::Control, EKeys::G) );
+	UI_COMMAND( ToggleShowGotoBox, "Go to Time...", "Go to a particular point on the time line", EUserInterfaceActionType::Button, FInputChord(EModifierKey::Control, EKeys::G) );
 	UI_COMMAND( ToggleShowRangeSlider, "Show Range Slider", "Enables and disables showing the time range slider", EUserInterfaceActionType::ToggleButton, FInputChord() );
 	UI_COMMAND( ToggleIsSnapEnabled, "Enable Snapping", "Enables and disables snapping", EUserInterfaceActionType::ToggleButton, FInputChord() );
 
@@ -58,6 +62,7 @@ void FSequencerCommands::RegisterCommands()
 	UI_COMMAND( ToggleSnapSectionTimesToInterval, "Snap to the Interval", "Snap sections to the interval", EUserInterfaceActionType::ToggleButton, FInputChord() );
 	UI_COMMAND( ToggleSnapSectionTimesToSections, "Snap to Other Sections", "Snap sections to other sections", EUserInterfaceActionType::ToggleButton, FInputChord() );
 
+	UI_COMMAND( ToggleSnapPlayTimeToKeys, "Snap to Keys While Scrubbing", "Snap the time to keys while scrubbing", EUserInterfaceActionType::ToggleButton, FInputChord() );
 	UI_COMMAND( ToggleSnapPlayTimeToInterval, "Snap to the Interval While Scrubbing", "Snap the play time to the interval while scrubbing", EUserInterfaceActionType::ToggleButton, FInputChord() );
 	UI_COMMAND( ToggleSnapPlayTimeToDraggedKey, "Snap to the Dragged Key", "Snap the play time to the dragged key", EUserInterfaceActionType::ToggleButton, FInputChord() );
 
@@ -67,6 +72,7 @@ void FSequencerCommands::RegisterCommands()
 
 	UI_COMMAND( FindInContentBrowser, "Find in Content Browser", "Find the viewed sequence asset in the content browser", EUserInterfaceActionType::Button, FInputChord() );
 	UI_COMMAND( ToggleCombinedKeyframes, "Combined Keyframes", "Show/hide the combined keyframes at the top node level", EUserInterfaceActionType::ToggleButton, FInputChord() );
+	UI_COMMAND( ToggleChannelColors, "Channel Colors", "Show/hide the channel colors in the track area", EUserInterfaceActionType::ToggleButton, FInputChord() );
 	UI_COMMAND( ToggleLabelBrowser, "Label Browser", "Show/hide the track label browser", EUserInterfaceActionType::ToggleButton, FInputChord() );
 	
 	UI_COMMAND( ToggleShowCurveEditor, "Curve Editor", "Show the animation keys in a curve editor", EUserInterfaceActionType::ToggleButton, FInputChord() );
@@ -81,7 +87,9 @@ void FSequencerCommands::RegisterCommands()
 
 	UI_COMMAND( ConvertToSpawnable, "Convert to Spawnable", "Make the specified possessed objects spawnable from sequencer. This will allow sequencer to have control over the lifetime of the object.", EUserInterfaceActionType::Button, FInputChord() );
 
-	UI_COMMAND(DiscardChanges, "Discard all changes", "Revert the currently edited movie scene to its last saved state.", EUserInterfaceActionType::Button, FInputChord());
+	UI_COMMAND( DiscardChanges, "Discard All Changes", "Revert the currently edited movie scene to its last saved state.", EUserInterfaceActionType::Button, FInputChord());
+
+	UI_COMMAND( FixActorReferences, "Fix Actor References", "Try to automatically fix up broken actor bindings.", EUserInterfaceActionType::Button, FInputChord() );
 }
 
 #undef LOCTEXT_NAMESPACE
\ No newline at end of file
diff --git a/Engine/Source/Editor/Sequencer/Private/SequencerCommonHelpers.cpp b/Engine/Source/Editor/Sequencer/Private/SequencerCommonHelpers.cpp
--- a/Engine/Source/Editor/Sequencer/Private/SequencerCommonHelpers.cpp
+++ b/Engine/Source/Editor/Sequencer/Private/SequencerCommonHelpers.cpp
@@ -4,6 +4,8 @@
 #include "SequencerCommonHelpers.h"
 #include "MovieSceneSection.h"
 #include "SequencerHotspots.h"
+#include "SSequencer.h"
+#include "SSequencerTreeView.h"
 #include "VirtualTrackArea.h"
 #include "SequencerContextMenus.h"
 #include "MovieSceneTrack.h"
@@ -95,15 +97,22 @@ void SequencerHelpers::GetAllSections(TSharedRef<FSequencerDisplayNode> DisplayN
 
 bool SequencerHelpers::FindObjectBindingNode(TSharedRef<FSequencerDisplayNode> DisplayNode, TSharedRef<FSequencerDisplayNode>& ObjectBindingNode)
 {
-	if (DisplayNode->GetType() == ESequencerNode::Object)
+	TArray< TSharedPtr<FSequencerDisplayNode> > ParentNodes;
+	ParentNodes.Add(DisplayNode);
+
+	while (DisplayNode->GetParent().IsValid())
 	{
-		ObjectBindingNode = DisplayNode;
-		return true;
+		ParentNodes.Add(DisplayNode->GetParent());
+		DisplayNode = DisplayNode->GetParent().ToSharedRef();
 	}
 
-	if (DisplayNode->GetParent().IsValid())
+	for (int32 ParentNodeIndex = ParentNodes.Num() - 1; ParentNodeIndex >= 0; --ParentNodeIndex)
 	{
-		return FindObjectBindingNode(DisplayNode->GetParent().ToSharedRef(), ObjectBindingNode);
+		if (ParentNodes[ParentNodeIndex]->GetType() == ESequencerNode::Object)
+		{
+			ObjectBindingNode = ParentNodes[ParentNodeIndex].ToSharedRef();
+			return true;
+		}
 	}
 
 	return false;
@@ -120,6 +129,112 @@ float SequencerHelpers::FrameToTime(int32 Frame, float FrameRate)
 	return Frame / FrameRate;
 }
 
+bool IsSectionSelectedInNode(FSequencer& Sequencer, TSharedRef<FSequencerDisplayNode> InNode)
+{
+	if (InNode->GetType() == ESequencerNode::Track)
+	{
+		TSharedRef<FSequencerTrackNode> TrackNode = StaticCastSharedRef<FSequencerTrackNode>(InNode);
+
+		for (auto Section : TrackNode->GetSections())
+		{
+			if (Sequencer.GetSelection().IsSelected(Section->GetSectionObject()))
+			{
+				return true;
+			}
+		}
+	}
+	return false;
+}
+
+bool AreKeysSelectedInNode(FSequencer& Sequencer, TSharedRef<FSequencerDisplayNode> InNode)
+{
+	TSharedRef<FSequencerTrackNode> TrackNode = StaticCastSharedRef<FSequencerTrackNode>(InNode);
+
+	TSet<TSharedPtr<IKeyArea>> KeyAreas;
+	SequencerHelpers::GetAllKeyAreas(InNode, KeyAreas);
+	for (auto KeyArea : KeyAreas)
+	{
+		for (auto KeyHandle : KeyArea->GetUnsortedKeyHandles())
+		{
+			FSequencerSelectedKey TestKey(*KeyArea->GetOwningSection(), KeyArea, KeyHandle);
+			if (Sequencer.GetSelection().IsSelected(TestKey))
+			{
+				return true;
+			}
+		}
+	}
+
+	for (int32 KeyGroupingIndex = 0; KeyGroupingIndex < InNode->GetNumKeyGroupings(); ++KeyGroupingIndex)
+	{			
+		TSharedRef<FGroupedKeyArea> KeyGrouping = InNode->GetKeyGrouping(KeyGroupingIndex);
+
+		for (auto KeyHandle : KeyGrouping->GetUnsortedKeyHandles())
+		{
+			FSequencerSelectedKey TestKey(*KeyGrouping->GetOwningSection(), KeyGrouping, KeyHandle);
+			if (Sequencer.GetSelection().IsSelected(TestKey))
+			{
+				return true;
+			}
+		}
+	}
+	return false;
+}
+
+void SequencerHelpers::ValidateNodesWithSelectedKeysOrSections(FSequencer& Sequencer)
+{
+	for (auto Node : Sequencer.GetSelection().GetNodesWithSelectedKeysOrSections())
+	{
+		if (!IsSectionSelectedInNode(Sequencer, Node) && !AreKeysSelectedInNode(Sequencer, Node))
+		{
+			Sequencer.GetSelection().RemoveFromNodesWithSelectedKeysOrSections(Node);
+		}
+	}
+}
+
+void SequencerHelpers::UpdateHoveredNodeFromSelectedSections(FSequencer& Sequencer)
+{
+	FSequencerSelection& Selection = Sequencer.GetSelection();
+
+	TSharedRef<SSequencer> SequencerWidget = StaticCastSharedRef<SSequencer>(Sequencer.GetSequencerWidget());
+	TSharedPtr<FSequencerDisplayNode> HoveredNode = SequencerWidget->GetTreeView()->GetNodeTree()->GetHoveredNode();
+	if (!HoveredNode.IsValid())
+	{
+		return;
+	}
+
+	if (IsSectionSelectedInNode(Sequencer, HoveredNode.ToSharedRef()))
+	{
+		Selection.AddToNodesWithSelectedKeysOrSections(HoveredNode.ToSharedRef());
+	}
+	else
+	{
+		Selection.RemoveFromNodesWithSelectedKeysOrSections(HoveredNode.ToSharedRef());
+	}
+}
+
+
+void SequencerHelpers::UpdateHoveredNodeFromSelectedKeys(FSequencer& Sequencer)
+{
+	FSequencerSelection& Selection = Sequencer.GetSelection();
+
+	TSharedRef<SSequencer> SequencerWidget = StaticCastSharedRef<SSequencer>(Sequencer.GetSequencerWidget());
+	TSharedPtr<FSequencerDisplayNode> HoveredNode = SequencerWidget->GetTreeView()->GetNodeTree()->GetHoveredNode();
+	if (!HoveredNode.IsValid())
+	{
+		return;
+	}
+
+	if (AreKeysSelectedInNode(Sequencer, HoveredNode.ToSharedRef()))
+	{
+		Selection.AddToNodesWithSelectedKeysOrSections(HoveredNode.ToSharedRef());
+	}
+	else
+	{
+		Selection.RemoveFromNodesWithSelectedKeysOrSections(HoveredNode.ToSharedRef());
+	}
+}
+
+
 void SequencerHelpers::PerformDefaultSelection(FSequencer& Sequencer, const FPointerEvent& MouseEvent)
 {
 	FSequencerSelection& Selection = Sequencer.GetSelection();
@@ -127,9 +242,10 @@ void SequencerHelpers::PerformDefaultSelection(FSequencer& Sequencer, const FPoi
 	// @todo: selection in transactions
 	auto ConditionallyClearSelection = [&]{
 		if (!MouseEvent.IsShiftDown() && !MouseEvent.IsControlDown())
-		{
+		{			
 			Selection.EmptySelectedSections();
 			Selection.EmptySelectedKeys();
+			Selection.EmptyNodesWithSelectedKeysOrSections();
 		}
 	};
 
@@ -161,7 +277,25 @@ void SequencerHelpers::PerformDefaultSelection(FSequencer& Sequencer, const FPoi
 				Selection.AddToSelection(Section);
 			}
 		}
+		else if (Hotspot->GetType() == ESequencerHotspot::SectionResize_L || Hotspot->GetType() == ESequencerHotspot::SectionResize_R)
+		{
+			UMovieSceneSection* Section = static_cast<FSectionResizeHotspot*>(Hotspot.Get())->Section.GetSectionObject();
+			if (!Selection.IsSelected(Section))
+			{
+				ConditionallyClearSelection();
+				Selection.AddToSelection(Section);
+			}
+		}
 
+		if (Hotspot->GetType() == ESequencerHotspot::Key)
+		{
+			UpdateHoveredNodeFromSelectedKeys(Sequencer);
+		}
+		else
+		{
+			UpdateHoveredNodeFromSelectedSections(Sequencer);
+		}
+		
 		return;
 	}
 
@@ -169,7 +303,7 @@ void SequencerHelpers::PerformDefaultSelection(FSequencer& Sequencer, const FPoi
 	ConditionallyClearSelection();
 
 	bool bForceSelect = !MouseEvent.IsControlDown();
-
+		
 	if (Hotspot->GetType() == ESequencerHotspot::Key)
 	{
 		FSequencerSelectedKey Key = static_cast<FKeyHotspot*>(Hotspot.Get())->Key;
@@ -185,14 +319,28 @@ void SequencerHelpers::PerformDefaultSelection(FSequencer& Sequencer, const FPoi
 	else if (Hotspot->GetType() == ESequencerHotspot::Section)
 	{
 		UMovieSceneSection* Section = static_cast<FSectionHotspot*>(Hotspot.Get())->Section.GetSectionObject();
-		if (bForceSelect || !Selection.IsSelected(Section))
-		{
-			Selection.AddToSelection(Section);
-		}
-		else
+
+		// Never allow infinite sections to be selected through normal click (they're only selectable through right click)
+		if (!Section->IsInfinite())
 		{
-			Selection.RemoveFromSelection(Section);
+			if (bForceSelect || !Selection.IsSelected(Section))
+			{
+				Selection.AddToSelection(Section);
+			}
+			else
+			{
+				Selection.RemoveFromSelection(Section);
+			}
 		}
+	}	
+
+	if (Hotspot->GetType() == ESequencerHotspot::Key)
+	{
+		UpdateHoveredNodeFromSelectedKeys(Sequencer);
+	}
+	else
+	{
+		UpdateHoveredNodeFromSelectedSections(Sequencer);
 	}
 }
 
diff --git a/Engine/Source/Editor/Sequencer/Private/SequencerCommonHelpers.h b/Engine/Source/Editor/Sequencer/Private/SequencerCommonHelpers.h
--- a/Engine/Source/Editor/Sequencer/Private/SequencerCommonHelpers.h
+++ b/Engine/Source/Editor/Sequencer/Private/SequencerCommonHelpers.h
@@ -16,30 +16,46 @@ public:
 	static void GetAllKeyAreas(TSharedPtr<FSequencerDisplayNode> DisplayNode, TSet<TSharedPtr<IKeyArea>>& KeyAreas);
 
 	/**
-	* Get descendant nodes
-	*/
+	 * Get descendant nodes
+	 */
 	static void GetDescendantNodes(TSharedRef<FSequencerDisplayNode> DisplayNode, TSet<TSharedRef<FSequencerDisplayNode> >& Nodes);
 
 	/**
-	* Gets all sections from the requested node
-	*/
+	 * Gets all sections from the requested node
+	 */
 	static void GetAllSections(TSharedRef<FSequencerDisplayNode> DisplayNode, TSet<TWeakObjectPtr<UMovieSceneSection>>& Sections);
 
 	/**
-	* Find object binding node from the given display node
-	*/
+	 * Find object binding node from the given display node
+	 */
 	static bool FindObjectBindingNode(TSharedRef<FSequencerDisplayNode> DisplayNode, TSharedRef<FSequencerDisplayNode>& ObjectBindingNode);
 
 	/**
-	* Convert time to frame
-	*/
+	 * Convert time to frame
+	 */
 	static int32 TimeToFrame(float Time, float FrameRate);
 
 	/**
-	* Convert frame to time
-	*/
+	 * Convert frame to time
+	 */
 	static float FrameToTime(int32 Frame, float FrameRate);
 
+	/**
+	 * Validate that the nodes with selected keys or sections actually are true
+	 */
+
+	static void ValidateNodesWithSelectedKeysOrSections(FSequencer& Sequencer);
+
+	/**
+	 * Update the nodes with selected sections from the hovered node
+	 */
+	static void UpdateHoveredNodeFromSelectedSections(FSequencer& Sequencer);
+
+	/**
+	 * Update the nodes with selected keys from the hovered node
+	 */
+	static void UpdateHoveredNodeFromSelectedKeys(FSequencer& Sequencer);
+
 	/**
 	 * Perform default selection for the specified mouse event, based on the current hotspot
 	 */
@@ -111,7 +127,7 @@ public:
 			FrameRateSnapValues.Add( SNumericDropDown<float>::FNamedValue( 0.033334f, LOCTEXT( "Snap_30Fps", "30 fps" ), LOCTEXT( "SnapDescription_30Fps", "Set snap to 30 fps" ) ) );
 			FrameRateSnapValues.Add( SNumericDropDown<float>::FNamedValue( 0.020834f, LOCTEXT( "Snap_48Fps", "48 fps" ), LOCTEXT( "SnapDescription_48Fps", "Set snap to 48 fps" ) ) );
 			FrameRateSnapValues.Add( SNumericDropDown<float>::FNamedValue( 0.02f, LOCTEXT( "Snap_50Fps", "50 fps (PAL/50)" ), LOCTEXT( "SnapDescription_50Fps", "Set snap to 50 fps" ) ) );
-			FrameRateSnapValues.Add( SNumericDropDown<float>::FNamedValue( 0.016683f, LOCTEXT( "Snap_50.94Fps", "50.94 fps (NTSC/60)" ), LOCTEXT( "SnapDescription_50.94Fps", "Set snap to 50.94 fps" ) ) );
+			FrameRateSnapValues.Add( SNumericDropDown<float>::FNamedValue( 0.016683f, LOCTEXT( "Snap_59.94Fps", "59.94 fps (NTSC/60)" ), LOCTEXT( "SnapDescription_59.94Fps", "Set snap to 59.94 fps" ) ) );
 			FrameRateSnapValues.Add( SNumericDropDown<float>::FNamedValue( 0.016667f, LOCTEXT( "Snap_60Fps", "60 fps" ), LOCTEXT( "SnapDescription_60Fps", "Set snap to 60 fps" ) ) );
 			FrameRateSnapValues.Add( SNumericDropDown<float>::FNamedValue( 0.008334f, LOCTEXT( "Snap_120Fps", "120 fps" ), LOCTEXT( "SnapDescription_120Fps", "Set snap to 120 fps" ) ) );
 
diff --git a/Engine/Source/Editor/Sequencer/Private/SequencerContextMenus.cpp b/Engine/Source/Editor/Sequencer/Private/SequencerContextMenus.cpp
--- a/Engine/Source/Editor/Sequencer/Private/SequencerContextMenus.cpp
+++ b/Engine/Source/Editor/Sequencer/Private/SequencerContextMenus.cpp
@@ -17,7 +17,9 @@
 #include "IStructureDetailsView.h"
 #include "PropertyEditorModule.h"
 #include "MovieSceneSequence.h"
-
+#include "IntegralKeyDetailsCustomization.h"
+#include "MovieSceneSubSection.h"
+#include "MovieSceneCinematicShotSection.h"
 
 #define LOCTEXT_NAMESPACE "SequencerContextMenus"
 
@@ -34,18 +36,21 @@ void FKeyContextMenu::PopulateMenu(FMenuBuilder& MenuBuilder)
 	FSequencer* SequencerPtr = &Sequencer.Get();
 	TSharedRef<FKeyContextMenu> Shared = AsShared();
 
-	MenuBuilder.AddSubMenu(
-		LOCTEXT("KeyProperties", "Properties"),
-		LOCTEXT("KeyPropertiesTooltip", "Modify the key properties"),
-		FNewMenuDelegate::CreateLambda([=](FMenuBuilder& SubMenuBuilder){ Shared->AddPropertiesMenu(SubMenuBuilder); }),
-		FUIAction (
-			FExecuteAction(),
-			// @todo sequencer: only one struct per structure view supported right now :/
-			FCanExecuteAction::CreateLambda([=]{ return (Shared->Sequencer->GetSelection().GetSelectedKeys().Num() == 1); })
-		),
-		NAME_None,
-		EUserInterfaceActionType::Button
-	);
+	if(CanAddPropertiesMenu())
+	{
+		MenuBuilder.AddSubMenu(
+			LOCTEXT("KeyProperties", "Properties"),
+			LOCTEXT("KeyPropertiesTooltip", "Modify the key properties"),
+			FNewMenuDelegate::CreateLambda([=](FMenuBuilder& SubMenuBuilder){ Shared->AddPropertiesMenu(SubMenuBuilder); }),
+			FUIAction (
+				FExecuteAction(),
+				// @todo sequencer: only one struct per structure view supported right now :/
+				FCanExecuteAction::CreateLambda([=]{ return (Shared->Sequencer->GetSelection().GetSelectedKeys().Num() == 1); })
+			),
+			NAME_None,
+			EUserInterfaceActionType::Button
+		);
+	}
 
 	MenuBuilder.BeginSection("SequencerKeyEdit", LOCTEXT("EditMenu", "Edit"));
 	{
@@ -155,6 +160,23 @@ void FKeyContextMenu::PopulateMenu(FMenuBuilder& MenuBuilder)
 	MenuBuilder.EndSection(); // SequencerKeys
 }
 
+bool FKeyContextMenu::CanAddPropertiesMenu() const
+{
+	for (const FSequencerSelectedKey& Key : Sequencer->GetSelection().GetSelectedKeys())
+	{
+		if (Key.KeyArea.IsValid() && Key.KeyHandle.IsSet())
+		{
+			TSharedPtr<FStructOnScope> KeyStruct = Key.KeyArea->GetKeyStruct(Key.KeyHandle.GetValue());
+
+			if (KeyStruct.IsValid())
+			{
+				return true;
+			}
+		}
+	}
+
+	return false;
+}
 
 void FKeyContextMenu::AddPropertiesMenu(FMenuBuilder& MenuBuilder)
 {
@@ -179,7 +201,7 @@ void FKeyContextMenu::AddPropertiesMenu(FMenuBuilder& MenuBuilder)
 		StructureViewArgs.bShowInterfaces = false;
 	}
 
-	TArray<TSharedPtr<FStructOnScope>> Keys;
+	TArray<TPair<TSharedPtr<FStructOnScope>, const FSequencerSelectedKey&>> Keys;
 	{
 		for (const FSequencerSelectedKey& Key : Sequencer->GetSelection().GetSelectedKeys())
 		{
@@ -189,7 +211,7 @@ void FKeyContextMenu::AddPropertiesMenu(FMenuBuilder& MenuBuilder)
 
 				if (KeyStruct.IsValid())
 				{
-					Keys.Add(KeyStruct);
+					Keys.Add(TPairInitializer<TSharedPtr<FStructOnScope>, const FSequencerSelectedKey&>(KeyStruct, Key));
 				}
 			}
 		}
@@ -201,13 +223,17 @@ void FKeyContextMenu::AddPropertiesMenu(FMenuBuilder& MenuBuilder)
 		// @todo sequencer: only one struct per structure view supported right now :/
 		if (Keys.Num() == 1)
 		{
-			TSharedPtr<FStructOnScope>& Key = Keys[0];
-			StructureDetailsView->SetStructureData(Key);
+			TPair<TSharedPtr<FStructOnScope>, const FSequencerSelectedKey&>& Key = Keys[0];
+
+			// register details customizations for this instance
+			StructureDetailsView->GetDetailsView().RegisterInstancedCustomPropertyLayout(FIntegralKey::StaticStruct(), FOnGetDetailCustomizationInstance::CreateStatic(&FIntegralKeyDetailsCustomization::MakeInstance, TWeakObjectPtr<const UMovieSceneSection>(Key.Value.Section)));
+
+			StructureDetailsView->SetStructureData(Key.Key);
 			StructureDetailsView->GetOnFinishedChangingPropertiesDelegate().AddLambda(
 				[=](const FPropertyChangedEvent& ChangeEvent) {
-					if (Key->GetStruct()->IsChildOf(FMovieSceneKeyStruct::StaticStruct()))
+					if (Key.Key->GetStruct()->IsChildOf(FMovieSceneKeyStruct::StaticStruct()))
 					{
-						((FMovieSceneKeyStruct*)Key->GetStructMemory())->PropagateChanges(ChangeEvent);
+						((FMovieSceneKeyStruct*)Key.Key->GetStructMemory())->PropagateChanges(ChangeEvent);
 					}
 					Sequencer->GetFocusedMovieSceneSequence()->Modify();
 					Sequencer->UpdateRuntimeInstances();
@@ -278,6 +304,21 @@ void FSectionContextMenu::PopulateMenu(FMenuBuilder& MenuBuilder)
 
 	MenuBuilder.BeginSection("SequencerSections", LOCTEXT("SectionsMenu", "Sections"));
 	{
+		if (CanPrimeForRecording())
+		{
+			MenuBuilder.AddMenuEntry(
+				LOCTEXT("PrimeForRecording", "Primed For Recording"),
+				LOCTEXT("PrimeForRecordingTooltip", "Prime this track for recording a new sequence."),
+				FSlateIcon(),
+				FUIAction(
+					FExecuteAction::CreateLambda([=]{ return Shared->TogglePrimeForRecording(); }),
+					FCanExecuteAction(),
+					FGetActionCheckState::CreateLambda([=]{ return Shared->IsPrimedForRecording() ? ECheckBoxState::Checked : ECheckBoxState::Unchecked; })),
+				NAME_None,
+				EUserInterfaceActionType::ToggleButton
+			);
+		}
+
 		if (CanSelectAllKeys())
 		{
 			MenuBuilder.AddMenuEntry(
@@ -533,6 +574,55 @@ void FSectionContextMenu::SelectAllKeys()
 }
 
 
+void FSectionContextMenu::TogglePrimeForRecording() const
+{
+	TArray<FSectionHandle> SelectedSections = StaticCastSharedRef<SSequencer>(Sequencer->GetSequencerWidget())->GetSectionHandles(Sequencer->GetSelection().GetSelectedSections());
+	if(SelectedSections.Num() > 0)
+	{
+		const FSectionHandle& Handle = SelectedSections[0];
+		UMovieSceneSubSection* SubSection = Cast<UMovieSceneSubSection>(Handle.GetSectionObject());
+		if (SubSection)
+		{
+			SubSection->SetAsRecording(SubSection != UMovieSceneSubSection::GetRecordingSection());
+		}	
+	}
+}
+
+
+bool FSectionContextMenu::IsPrimedForRecording() const
+{
+	TArray<FSectionHandle> SelectedSections = StaticCastSharedRef<SSequencer>(Sequencer->GetSequencerWidget())->GetSectionHandles(Sequencer->GetSelection().GetSelectedSections());
+	if(SelectedSections.Num() > 0)
+	{
+		const FSectionHandle& Handle = SelectedSections[0];
+		UMovieSceneSubSection* SubSection = Cast<UMovieSceneSubSection>(Handle.GetSectionObject());
+		if (SubSection)
+		{
+			return SubSection == UMovieSceneSubSection::GetRecordingSection();
+		}	
+	}
+
+	return false;
+}
+
+bool FSectionContextMenu::CanPrimeForRecording() const
+{
+	TArray<FSectionHandle> SelectedSections = StaticCastSharedRef<SSequencer>(Sequencer->GetSequencerWidget())->GetSectionHandles(Sequencer->GetSelection().GetSelectedSections());
+	if(SelectedSections.Num() > 0)
+	{
+		const FSectionHandle& Handle = SelectedSections[0];
+		UMovieSceneSubSection* SubSection = Cast<UMovieSceneSubSection>(Handle.GetSectionObject());
+		UMovieSceneCinematicShotSection* CinematicShotSection = Cast<UMovieSceneCinematicShotSection>(Handle.GetSectionObject());
+		if (SubSection && !CinematicShotSection)
+		{
+			return true;
+		}	
+	}
+
+	return false;
+}
+
+
 bool FSectionContextMenu::CanSelectAllKeys() const
 {
 	TArray<FSectionHandle> SelectedSections = StaticCastSharedRef<SSequencer>(Sequencer->GetSequencerWidget())->GetSectionHandles(Sequencer->GetSelection().GetSelectedSections());
@@ -966,7 +1056,7 @@ void FSectionContextMenu::BringForward()
 				return A.GetOverlapPriority() < B.GetOverlapPriority();
 			});
 
-			for (int32 SectionIndex = Row.Sections.Num() - 1; SectionIndex >= 0; --SectionIndex)
+			for (int32 SectionIndex = Row.Sections.Num() - 1; SectionIndex > 0; --SectionIndex)
 			{
 				UMovieSceneSection* ThisSection = Row.Sections[SectionIndex];
 				if (Row.SectionToReOrder.Contains(ThisSection))
diff --git a/Engine/Source/Editor/Sequencer/Private/SequencerContextMenus.h b/Engine/Source/Editor/Sequencer/Private/SequencerContextMenus.h
--- a/Engine/Source/Editor/Sequencer/Private/SequencerContextMenus.h
+++ b/Engine/Source/Editor/Sequencer/Private/SequencerContextMenus.h
@@ -43,6 +43,12 @@ private:
 
 	bool CanSelectAllKeys() const;
 
+	void TogglePrimeForRecording() const;
+
+	bool IsPrimedForRecording() const;
+
+	bool CanPrimeForRecording() const;
+
 	bool CanSetExtrapolationMode() const;
 
 	void TrimSection(bool bTrimLeft);
@@ -202,6 +208,9 @@ private:
 	/** Add the Properties sub-menu. */
 	void AddPropertiesMenu(FMenuBuilder& MenuBuilder);
 
+	/** Check if we can add the key properties menu */
+	bool CanAddPropertiesMenu() const;
+
 	void PopulateMenu(FMenuBuilder& MenuBuilder);
 
 	/** The sequencer */
diff --git a/Engine/Source/Editor/Sequencer/Private/SequencerCurveOwner.cpp b/Engine/Source/Editor/Sequencer/Private/SequencerCurveOwner.cpp
--- a/Engine/Source/Editor/Sequencer/Private/SequencerCurveOwner.cpp
+++ b/Engine/Source/Editor/Sequencer/Private/SequencerCurveOwner.cpp
@@ -182,4 +182,23 @@ void FSequencerCurveOwner::OnCurveChanged( const TArray<FRichCurveEditInfo>& Cha
 bool FSequencerCurveOwner::IsValidCurve( FRichCurveEditInfo CurveInfo )
 {
 	return EditInfoToSectionMap.Contains(CurveInfo);
+}
+
+FLinearColor FSequencerCurveOwner::GetCurveColor( FRichCurveEditInfo CurveInfo ) const
+{
+	FString CurveName = CurveInfo.CurveName.ToString();
+	if (CurveName.EndsWith(TEXT("- X")) || CurveName.EndsWith(TEXT("- Red")))
+	{
+		return FLinearColor(1.0f, 0.0f, 0.0f);
+	}
+	else if (CurveName.EndsWith(TEXT("- Y")) || CurveName.EndsWith(TEXT("- Green")))
+	{
+		return FLinearColor(0.0f, 1.0f, 0.0f);
+	}
+	else if (CurveName.EndsWith(TEXT("- Z")) || CurveName.EndsWith(TEXT("- Blue")))
+	{
+		return FLinearColor(0.05f, 0.05f, 1.0f);
+	}
+
+	return FCurveOwnerInterface::GetCurveColor(CurveInfo);
 }
\ No newline at end of file
diff --git a/Engine/Source/Editor/Sequencer/Private/SequencerCurveOwner.h b/Engine/Source/Editor/Sequencer/Private/SequencerCurveOwner.h
--- a/Engine/Source/Editor/Sequencer/Private/SequencerCurveOwner.h
+++ b/Engine/Source/Editor/Sequencer/Private/SequencerCurveOwner.h
@@ -20,6 +20,7 @@ public:
 	virtual void MakeTransactional() override;
 	virtual void OnCurveChanged( const TArray<FRichCurveEditInfo>& ChangedCurveEditInfos ) override;
 	virtual bool IsValidCurve( FRichCurveEditInfo CurveInfo ) override;
+	virtual FLinearColor GetCurveColor( FRichCurveEditInfo CurveInfo) const override;
 
 private:
 	/** The FSequencerNodeTree used to build the curve owner. */
diff --git a/Engine/Source/Editor/Sequencer/Private/SequencerDisplayNodeDragDropOp.cpp b/Engine/Source/Editor/Sequencer/Private/SequencerDisplayNodeDragDropOp.cpp
new file mode 100644
--- /dev/null
+++ b/Engine/Source/Editor/Sequencer/Private/SequencerDisplayNodeDragDropOp.cpp
@@ -0,0 +1,16 @@
+// Copyright 1998-2016 Epic Games, Inc. All Rights Reserved.
+
+#include "SequencerPrivatePCH.h"
+#include "SequencerDisplayNodeDragDropOp.h"
+
+
+FSequencerDisplayNodeDragDropOp::FSequencerDisplayNodeDragDropOp( TArray<TSharedRef<FSequencerDisplayNode>>& InDraggedNodes )
+{
+	DraggedNodes = InDraggedNodes;
+}
+
+
+TArray<TSharedRef<FSequencerDisplayNode>>& FSequencerDisplayNodeDragDropOp::GetDraggedNodes()
+{
+	return DraggedNodes;
+}
\ No newline at end of file
diff --git a/Engine/Source/Editor/Sequencer/Private/SequencerDisplayNodeDragDropOp.h b/Engine/Source/Editor/Sequencer/Private/SequencerDisplayNodeDragDropOp.h
new file mode 100644
--- /dev/null
+++ b/Engine/Source/Editor/Sequencer/Private/SequencerDisplayNodeDragDropOp.h
@@ -0,0 +1,24 @@
+// Copyright 1998-2016 Epic Games, Inc. All Rights Reserved.
+
+#pragma once
+
+#include "DecoratedDragDropOp.h"
+
+class FSequencerDisplayNode;
+
+/** A decorated drag drop operation object for dragging sequencer display nodes. */
+class FSequencerDisplayNodeDragDropOp : public FDecoratedDragDropOp
+{
+public:
+	DRAG_DROP_OPERATOR_TYPE( FSequencerDisplayNodeDragDropOp, FDecoratedDragDropOp )
+
+public:
+	FSequencerDisplayNodeDragDropOp( TArray<TSharedRef<FSequencerDisplayNode>>& InDraggedNodes );
+
+	/** Gets the nodes which are currently being dragged. */
+	TArray<TSharedRef<FSequencerDisplayNode>>& GetDraggedNodes();
+
+private:
+	/** The nodes currently being dragged. */
+	TArray<TSharedRef<FSequencerDisplayNode>> DraggedNodes;
+};
\ No newline at end of file
diff --git a/Engine/Source/Editor/Sequencer/Private/SequencerEdMode.cpp b/Engine/Source/Editor/Sequencer/Private/SequencerEdMode.cpp
--- a/Engine/Source/Editor/Sequencer/Private/SequencerEdMode.cpp
+++ b/Engine/Source/Editor/Sequencer/Private/SequencerEdMode.cpp
@@ -51,7 +51,7 @@ bool FSequencerEdMode::InputKey( FEditorViewportClient* ViewportClient, FViewpor
 	{
 		FModifierKeysState KeyState = FSlateApplication::Get().GetModifierKeys();
 
-		if (Sequencer->GetCommandBindings().Get()->ProcessCommandBindings(Key, KeyState, (Event == IE_Repeat) ))
+		if (Sequencer->GetCommandBindings(ESequencerCommandBindings::Shared).Get()->ProcessCommandBindings(Key, KeyState, (Event == IE_Repeat) ))
 		{
 			return true;
 		}
@@ -113,7 +113,7 @@ void GetLocationAtTime(UMovieScene3DTransformSection* TransformSection, float Ke
 	TransformSection->EvalRotation(KeyTime, KeyRot);
 }
 
-void DrawTransformTrack(const FSceneView* View, FPrimitiveDrawInterface* PDI, UMovieScene3DTransformTrack* TransformTrack, const TArray<UObject*>& BoundObjects)
+void DrawTransformTrack(const FSceneView* View, FPrimitiveDrawInterface* PDI, UMovieScene3DTransformTrack* TransformTrack, const TArray<TWeakObjectPtr<UObject>>& BoundObjects)
 {
 	const bool bHitTesting = PDI->IsHitTesting();
 	FLinearColor TrackColor = FLinearColor::Yellow; //@todo - customizable per track
@@ -191,7 +191,7 @@ void DrawTransformTrack(const FSceneView* View, FPrimitiveDrawInterface* PDI, UM
 
 				for (auto BoundObject : BoundObjects)
 				{
-					FTransform RefTM = GetRefFrame(BoundObject);
+					FTransform RefTM = GetRefFrame(BoundObject.Get());
 
 					FVector OldPos_G = RefTM.TransformPosition(OldPos);
 					FVector NewKeyPos_G = RefTM.TransformPosition(NewKeyPos);
@@ -255,7 +255,7 @@ void DrawTransformTrack(const FSceneView* View, FPrimitiveDrawInterface* PDI, UM
 
 			for (auto BoundObject : BoundObjects)
 			{
-				FTransform RefTM = GetRefFrame(BoundObject);
+				FTransform RefTM = GetRefFrame(BoundObject.Get());
 
 				FColor KeyColor = bKeySelected ? SequencerEdMode_Draw3D::KeySelectedColor : TrackColor.ToFColor(true);
 
@@ -339,7 +339,17 @@ void FSequencerEdMode::DrawTracks3D(const FSceneView* View, FPrimitiveDrawInterf
 {
 	TSet<TSharedRef<FSequencerDisplayNode> > ObjectBindingNodes;
 
-	// Look for the object binding nodes from the selected nodes
+	// Look for the object binding nodes from the nodes with selected keys or sections
+	for (auto OutlinerNode : Sequencer->GetSelection().GetNodesWithSelectedKeysOrSections())
+	{
+		TSharedRef<FSequencerDisplayNode> ObjectBindingNode = OutlinerNode;
+		if (SequencerHelpers::FindObjectBindingNode(OutlinerNode, ObjectBindingNode))
+		{
+			ObjectBindingNodes.Add(ObjectBindingNode);
+		}
+	}
+
+	// And also look for the nodes that are directly selected
 	for (auto OutlinerNode : Sequencer->GetSelection().GetSelectedOutlinerNodes())
 	{
 		TSharedRef<FSequencerDisplayNode> ObjectBindingNode = OutlinerNode;
@@ -358,7 +368,7 @@ void FSequencerEdMode::DrawTracks3D(const FSceneView* View, FPrimitiveDrawInterf
 
 		FGuid ObjectBinding = StaticCastSharedRef<FSequencerObjectBindingNode>(ObjectBindingNode)->GetObjectBinding();
 
-		TArray<UObject*> BoundObjects;
+		TArray<TWeakObjectPtr<UObject>> BoundObjects;
 		Sequencer->GetRuntimeObjects(Sequencer->GetFocusedMovieSceneSequenceInstance(), ObjectBinding, BoundObjects);
 
 		for (auto DisplayNode : AllNodes)
diff --git a/Engine/Source/Editor/Sequencer/Private/SequencerHotspots.cpp b/Engine/Source/Editor/Sequencer/Private/SequencerHotspots.cpp
--- a/Engine/Source/Editor/Sequencer/Private/SequencerHotspots.cpp
+++ b/Engine/Source/Editor/Sequencer/Private/SequencerHotspots.cpp
@@ -51,6 +51,11 @@ TSharedPtr<ISequencerEditToolDragOperation> FSectionResizeHotspot::InitiateDrag(
 	
 	if (!SelectedSections.Contains(Section.GetSectionObject()))
 	{
+		Sequencer.GetSelection().Empty();
+		Sequencer.GetSelection().AddToSelection(Section.GetSectionObject());
+		SequencerHelpers::UpdateHoveredNodeFromSelectedSections(static_cast<FSequencer&>(Sequencer));
+
+		SectionHandles.Empty();
 		SectionHandles.Add(Section);
 	}
 	return MakeShareable( new FResizeSection(static_cast<FSequencer&>(Sequencer), SectionHandles, HandleType == Right) );
diff --git a/Engine/Source/Editor/Sequencer/Private/SequencerNodeTree.cpp b/Engine/Source/Editor/Sequencer/Private/SequencerNodeTree.cpp
--- a/Engine/Source/Editor/Sequencer/Private/SequencerNodeTree.cpp
+++ b/Engine/Source/Editor/Sequencer/Private/SequencerNodeTree.cpp
@@ -9,6 +9,7 @@
 #include "Sequencer.h"
 #include "ScopedTransaction.h"
 #include "MovieScene.h"
+#include "MovieSceneFolder.h"
 #include "MovieSceneTrackEditor.h"
 #include "SequencerSectionLayoutBuilder.h"
 #include "ISequencerSection.h"
@@ -25,6 +26,75 @@ void FSequencerNodeTree::Empty()
 }
 
 
+int32 NodeTypeToFolderSortId(ESequencerNode::Type NodeType)
+{
+	switch ( NodeType )
+	{
+	case ESequencerNode::Folder:
+		return 0;
+	case ESequencerNode::Track:
+		return 1;
+	case ESequencerNode::Object:
+		return 2;
+	default:
+		return 3;
+	}
+}
+
+
+int32 NodeTypeToObjectSortId( ESequencerNode::Type NodeType )
+{
+	switch ( NodeType )
+	{
+	case ESequencerNode::Object:
+		return 0;
+	case ESequencerNode::Track:
+		return 1;
+	default:
+		return 2;
+	}
+}
+
+
+struct FDisplayNodeSorter
+{
+	bool operator()( const TSharedRef<FSequencerDisplayNode>& A, const TSharedRef<FSequencerDisplayNode>& B ) const
+	{
+		TSharedPtr<FSequencerDisplayNode> ParentNode = A->GetParent();
+		
+		// If the nodes are root nodes, or in folders and they are the same type, sort by name.
+		if ( (ParentNode.IsValid() == false || ParentNode->GetType() == ESequencerNode::Folder) && A->GetType() == B->GetType() )
+		{
+			return A->GetDisplayName().ToString() < B->GetDisplayName().ToString();
+		}
+
+		int32 SortIdA;
+		int32 SortIdB;
+
+		// Otherwise if they are root nodes or in folders use the folder sort id.
+		if ( ParentNode.IsValid() == false || ParentNode->GetType() == ESequencerNode::Folder )
+		{
+			SortIdA = NodeTypeToFolderSortId( A->GetType() );
+			SortIdB = NodeTypeToFolderSortId( B->GetType() );
+		}
+		// Otherwise if they are in an object node use the object node sort id.
+		else if ( ParentNode->GetType() == ESequencerNode::Object )
+		{
+			SortIdA = NodeTypeToObjectSortId( A->GetType() );
+			SortIdB = NodeTypeToObjectSortId( B->GetType() );
+		}
+		// Otherwise they are equal, and in a stable sort shouldn't change position.
+		else
+		{
+			SortIdA = 0;
+			SortIdB = 0;
+		}
+
+		return SortIdA < SortIdB;
+	}
+};
+
+
 void FSequencerNodeTree::Update()
 {
 	HoveredNode = nullptr;
@@ -35,10 +105,9 @@ void FSequencerNodeTree::Update()
 	UMovieScene* MovieScene = Sequencer.GetFocusedMovieSceneSequence()->GetMovieScene();
 	UMovieSceneCinematicShotTrack* CinematicShotTrack = MovieScene->FindMasterTrack<UMovieSceneCinematicShotTrack>();
 
-	TArray<TSharedRef<FSequencerDisplayNode>> NewRootNodes;
-
 	// Get the master tracks  so we can get sections from them
 	const TArray<UMovieSceneTrack*>& MasterTracks = MovieScene->GetMasterTracks();
+	TArray<TSharedRef<FSequencerTrackNode>> MasterTrackNodes;
 
 	for (UMovieSceneTrack* Track : MasterTracks)
 	{
@@ -46,8 +115,8 @@ void FSequencerNodeTree::Update()
 		{
 			UMovieSceneTrack& TrackRef = *Track;
 
-			TSharedRef<FSequencerTrackNode> SectionNode = MakeShareable(new FSequencerTrackNode(TrackRef, *FindOrAddTypeEditor(TrackRef), nullptr, *this));
-			NewRootNodes.Add(SectionNode);
+			TSharedRef<FSequencerTrackNode> SectionNode = MakeShareable(new FSequencerTrackNode(TrackRef, *FindOrAddTypeEditor(TrackRef), true, nullptr, *this));
+			MasterTrackNodes.Add(SectionNode);
 	
 			MakeSectionInterfaces(TrackRef, SectionNode);
 		}
@@ -62,10 +131,10 @@ void FSequencerNodeTree::Update()
 	}
 
 	// Make nodes for all object bindings
-	TArray<TSharedRef<FSequencerDisplayNode>> NewObjectNodes;
+	TArray<TSharedRef<FSequencerObjectBindingNode>> ObjectNodes;
 	for( const FMovieSceneBinding& Binding : Bindings )
 	{
-		TSharedRef<FSequencerObjectBindingNode> ObjectBindingNode = AddObjectBinding( Binding.GetName(), Binding.GetObjectGuid(), GuidToBindingMap, NewObjectNodes );
+		TSharedRef<FSequencerObjectBindingNode> ObjectBindingNode = AddObjectBinding( Binding.GetName(), Binding.GetObjectGuid(), GuidToBindingMap, ObjectNodes );
 
 		const TArray<UMovieSceneTrack*>& Tracks = Binding.GetTracks();
 
@@ -78,42 +147,10 @@ void FSequencerNodeTree::Update()
 	}
 
 
-	struct FObjectNodeSorter
-	{
-		bool operator()( const TSharedRef<FSequencerDisplayNode>& A, const TSharedRef<FSequencerDisplayNode>& B ) const
-		{
-			if (A->GetType() == ESequencerNode::Object && B->GetType() != ESequencerNode::Object)
-			{
-				return true;
-			}
-
-			if (A->GetType() != ESequencerNode::Object && B->GetType() == ESequencerNode::Object)
-			{
-				return false;
-			}
-
-			if ( A->GetType() == ESequencerNode::Object && B->GetType() == ESequencerNode::Object )
-			{
-				return A->GetDisplayName().ToString() < B->GetDisplayName().ToString();
-			}
-
-			return 0;
-		}
-	};
-
-
-	NewObjectNodes.Sort( FObjectNodeSorter() );
-	for (TSharedRef<FSequencerDisplayNode> NewObjectNode : NewObjectNodes)
-	{
-		NewObjectNode->SortChildNodes(FObjectNodeSorter());
-	}
-
-	NewRootNodes.Append(NewObjectNodes);
-
 	// Cinematic shot track always comes first
 	if (CinematicShotTrack)
 	{
-		TSharedRef<FSequencerTrackNode> SectionNode = MakeShareable(new FSequencerTrackNode(*CinematicShotTrack, *FindOrAddTypeEditor(*CinematicShotTrack), nullptr, *this));
+		TSharedRef<FSequencerTrackNode> SectionNode = MakeShareable(new FSequencerTrackNode(*CinematicShotTrack, *FindOrAddTypeEditor(*CinematicShotTrack), false, nullptr, *this));
 
 		RootNodes.Add(SectionNode);
 		MakeSectionInterfaces(*CinematicShotTrack, SectionNode);
@@ -124,15 +161,24 @@ void FSequencerNodeTree::Update()
 	
 	if (CameraCutTrack)
 	{
-		TSharedRef<FSequencerTrackNode> SectionNode = MakeShareable(new FSequencerTrackNode(*CameraCutTrack, *FindOrAddTypeEditor(*CameraCutTrack), nullptr, *this));
+		TSharedRef<FSequencerTrackNode> SectionNode = MakeShareable(new FSequencerTrackNode(*CameraCutTrack, *FindOrAddTypeEditor(*CameraCutTrack), false, nullptr, *this));
 
 		RootNodes.Add(SectionNode);
 		MakeSectionInterfaces(*CameraCutTrack, SectionNode);
 	}
 
 	// Add all other nodes after the camera cut track
-	RootNodes.Append(NewRootNodes);
+	TArray<TSharedRef<FSequencerDisplayNode>> FolderMasterTrackAndObjectNodes;
+	CreateAndPopulateFolderNodes( MasterTrackNodes, ObjectNodes, MovieScene->GetRootFolders(), FolderMasterTrackAndObjectNodes );
+	
+	// Sort the created nodes.
+	FolderMasterTrackAndObjectNodes.Sort(FDisplayNodeSorter());
+	for ( TSharedRef<FSequencerDisplayNode> Node : FolderMasterTrackAndObjectNodes )
+	{
+		Node->SortChildNodes(FDisplayNodeSorter());
+	}
 
+	RootNodes.Append( FolderMasterTrackAndObjectNodes );
 
 	RootNodes.Reserve(RootNodes.Num()*2);
 	for (int32 Index = 0; Index < RootNodes.Num(); Index += 2)
@@ -218,7 +264,7 @@ const TArray<TSharedRef<FSequencerDisplayNode>>& FSequencerNodeTree::GetRootNode
 }
 
 
-TSharedRef<FSequencerObjectBindingNode> FSequencerNodeTree::AddObjectBinding(const FString& ObjectName, const FGuid& ObjectBinding, TMap<FGuid, const FMovieSceneBinding*>& GuidToBindingMap, TArray<TSharedRef<FSequencerDisplayNode>>& OutNodeList)
+TSharedRef<FSequencerObjectBindingNode> FSequencerNodeTree::AddObjectBinding(const FString& ObjectName, const FGuid& ObjectBinding, TMap<FGuid, const FMovieSceneBinding*>& GuidToBindingMap, TArray<TSharedRef<FSequencerObjectBindingNode>>& OutNodeList)
 {
 	TSharedPtr<FSequencerObjectBindingNode> ObjectNode;
 	TSharedPtr<FSequencerObjectBindingNode>* FoundObjectNode = ObjectBindingMap.Find(ObjectBinding);
@@ -300,6 +346,81 @@ TSharedRef<FSequencerObjectBindingNode> FSequencerNodeTree::AddObjectBinding(con
 }
 
 
+TSharedRef<FSequencerDisplayNode> CreateFolderNode(
+	UMovieSceneFolder& MovieSceneFolder, FSequencerNodeTree& NodeTree, 
+	TMap<UMovieSceneTrack*, TSharedRef<FSequencerTrackNode>>& MasterTrackToDisplayNodeMap,
+	TMap<FGuid, TSharedRef<FSequencerObjectBindingNode>>& ObjectGuidToDisplayNodeMap )
+{
+	TSharedRef<FSequencerFolderNode> FolderNode( new FSequencerFolderNode( MovieSceneFolder, TSharedPtr<FSequencerDisplayNode>(), NodeTree ) );
+
+	for ( UMovieSceneFolder* ChildFolder : MovieSceneFolder.GetChildFolders() )
+	{
+		FolderNode->AddChildNode( CreateFolderNode( *ChildFolder, NodeTree, MasterTrackToDisplayNodeMap, ObjectGuidToDisplayNodeMap ) );
+	}
+
+	for ( UMovieSceneTrack* MasterTrack : MovieSceneFolder.GetChildMasterTracks() )
+	{
+		TSharedRef<FSequencerTrackNode>* TrackNodePtr = MasterTrackToDisplayNodeMap.Find( MasterTrack );
+		if ( TrackNodePtr != nullptr)
+		{
+			// TODO: Log this.
+			FolderNode->AddChildNode( *TrackNodePtr );
+			MasterTrackToDisplayNodeMap.Remove( MasterTrack );
+		}
+	}
+
+	for (const FGuid& ObjectGuid : MovieSceneFolder.GetChildObjectBindings() )
+	{
+		TSharedRef<FSequencerObjectBindingNode>* ObjectNodePtr = ObjectGuidToDisplayNodeMap.Find( ObjectGuid );
+		if ( ObjectNodePtr != nullptr )
+		{
+			// TODO: Log this.
+			FolderNode->AddChildNode( *ObjectNodePtr );
+			ObjectGuidToDisplayNodeMap.Remove( ObjectGuid );
+		}
+	}
+
+	return FolderNode;
+}
+
+
+void FSequencerNodeTree::CreateAndPopulateFolderNodes( 
+	TArray<TSharedRef<FSequencerTrackNode>>& MasterTrackNodes, TArray<TSharedRef<FSequencerObjectBindingNode>>& ObjectNodes,
+	TArray<UMovieSceneFolder*>& MovieSceneFolders, TArray<TSharedRef<FSequencerDisplayNode>>& GroupedNodes )
+{
+	TMap<UMovieSceneTrack*, TSharedRef<FSequencerTrackNode>> MasterTrackToDisplayNodeMap;
+	for ( TSharedRef<FSequencerTrackNode> MasterTrackNode : MasterTrackNodes )
+	{
+		MasterTrackToDisplayNodeMap.Add( MasterTrackNode->GetTrack(), MasterTrackNode );
+	}
+
+	TMap<FGuid, TSharedRef<FSequencerObjectBindingNode>> ObjectGuidToDisplayNodeMap;
+	for ( TSharedRef<FSequencerObjectBindingNode> ObjectBindingNode : ObjectNodes )
+	{
+		ObjectGuidToDisplayNodeMap.Add( ObjectBindingNode->GetObjectBinding(), ObjectBindingNode );
+	}
+
+	for ( UMovieSceneFolder* MovieSceneFolder : MovieSceneFolders )
+	{
+		GroupedNodes.Add( CreateFolderNode( *MovieSceneFolder, *this, MasterTrackToDisplayNodeMap, ObjectGuidToDisplayNodeMap ) );	
+	}
+
+	TArray<TSharedRef<FSequencerTrackNode>> NonFolderTrackNodes;
+	MasterTrackToDisplayNodeMap.GenerateValueArray( NonFolderTrackNodes );
+	for ( TSharedRef<FSequencerTrackNode> NonFolderTrackNode : NonFolderTrackNodes )
+	{
+		GroupedNodes.Add( NonFolderTrackNode );
+	}
+
+	TArray<TSharedRef<FSequencerObjectBindingNode>> NonFolderObjectNodes;
+	ObjectGuidToDisplayNodeMap.GenerateValueArray( NonFolderObjectNodes );
+	for ( TSharedRef<FSequencerObjectBindingNode> NonFolderObjectNode : NonFolderObjectNodes )
+	{
+		GroupedNodes.Add( NonFolderObjectNode );
+	}
+}
+
+
 void FSequencerNodeTree::SaveExpansionState(const FSequencerDisplayNode& Node, bool bExpanded)
 {	
 	// @todo Sequencer - This should be moved to the sequence level
diff --git a/Engine/Source/Editor/Sequencer/Private/SequencerNodeTree.h b/Engine/Source/Editor/Sequencer/Private/SequencerNodeTree.h
--- a/Engine/Source/Editor/Sequencer/Private/SequencerNodeTree.h
+++ b/Engine/Source/Editor/Sequencer/Private/SequencerNodeTree.h
@@ -7,8 +7,11 @@ class ISequencerTrackEditor;
 class UMovieSceneSection;
 class FSequencerDisplayNode;
 class FSequencerObjectBindingNode;
+class FSequencerTrackNode;
+class FSequencerFolderNode;
 class UMovieSceneTrack;
 struct FMovieSceneBinding;
+class UMovieSceneFolder;
 
 
 /**
@@ -109,7 +112,14 @@ private:
 	/**
 	 * Creates a new object binding node and any parent binding nodes.
 	 */
-	TSharedRef<FSequencerObjectBindingNode> AddObjectBinding( const FString& ObjectName, const FGuid& ObjectBinding, TMap<FGuid, const FMovieSceneBinding*>& GuidToBindingMap, TArray< TSharedRef<FSequencerDisplayNode> >& OutNodeList );
+	TSharedRef<FSequencerObjectBindingNode> AddObjectBinding( const FString& ObjectName, const FGuid& ObjectBinding, TMap<FGuid, const FMovieSceneBinding*>& GuidToBindingMap, TArray< TSharedRef<FSequencerObjectBindingNode> >& OutNodeList );
+
+	/**
+	 * Creates the tree of folder nodes and populates it with object and track nodes. 
+	 */
+	void CreateAndPopulateFolderNodes( TArray<TSharedRef<FSequencerTrackNode>>& MasterTrackNodes, TArray<TSharedRef<FSequencerObjectBindingNode>>& ObjectNodes,
+		TArray<UMovieSceneFolder*>& MovieSceneFolders, TArray<TSharedRef<FSequencerDisplayNode>>& GroupedNodes );
+
 private:
 	/** Tools for building movie scene section layouts.  One tool for each track */
 	TMap< UMovieSceneTrack*, TSharedPtr<ISequencerTrackEditor> > EditorMap;
diff --git a/Engine/Source/Editor/Sequencer/Private/SequencerObjectChangeListener.cpp b/Engine/Source/Editor/Sequencer/Private/SequencerObjectChangeListener.cpp
--- a/Engine/Source/Editor/Sequencer/Private/SequencerObjectChangeListener.cpp
+++ b/Engine/Source/Editor/Sequencer/Private/SequencerObjectChangeListener.cpp
@@ -293,7 +293,7 @@ void FSequencerObjectChangeListener::OnObjectPreEditChange( UObject* Object, con
 				FPropertyListenerSettings Settings;
 				// Ignore array and object properties
 				Settings.bIgnoreArrayProperties = true;
-				Settings.bIgnoreObjectProperties = true;
+				Settings.bIgnoreObjectProperties = false;
 				// Property flags which must be on the property
 				Settings.RequiredPropertyFlags = 0;
 				// Property flags which cannot be on the property
diff --git a/Engine/Source/Editor/Sequencer/Private/SequencerPrivatePCH.h b/Engine/Source/Editor/Sequencer/Private/SequencerPrivatePCH.h
--- a/Engine/Source/Editor/Sequencer/Private/SequencerPrivatePCH.h
+++ b/Engine/Source/Editor/Sequencer/Private/SequencerPrivatePCH.h
@@ -15,6 +15,7 @@
 #include "SequencerCommands.h"
 #include "SequencerNodeTree.h"
 #include "SequencerDisplayNode.h"
+#include "SequencerFolderNode.h"
 #include "SequencerObjectBindingNode.h"
 #include "SequencerSectionCategoryNode.h"
 #include "SequencerSectionKeyAreaNode.h"
diff --git a/Engine/Source/Editor/Sequencer/Private/SequencerSectionPainter.cpp b/Engine/Source/Editor/Sequencer/Private/SequencerSectionPainter.cpp
--- a/Engine/Source/Editor/Sequencer/Private/SequencerSectionPainter.cpp
+++ b/Engine/Source/Editor/Sequencer/Private/SequencerSectionPainter.cpp
@@ -13,19 +13,22 @@ FSequencerSectionPainter::FSequencerSectionPainter(FSlateWindowElementList& OutD
 {
 }
 
-int32 FSequencerSectionPainter::PaintSectionBackground(const FColor& Tint)
-{
-	FSequencerSectionTintRegion Region(Tint);
-	return PaintSectionBackground(&Region, 1);
-}
-
 int32 FSequencerSectionPainter::PaintSectionBackground()
 {
-	FSequencerSectionTintRegion Region(GetTrack()->GetColorTint());
-	return PaintSectionBackground(&Region, 1);
+	return PaintSectionBackground(GetTrack()->GetColorTint());
 }
 
 UMovieSceneTrack* FSequencerSectionPainter::GetTrack() const
 {
 	return Section.GetTypedOuter<UMovieSceneTrack>();
+}
+
+FLinearColor FSequencerSectionPainter::BlendColor(FLinearColor InColor)
+{
+	static FLinearColor BaseColor(FColor(71,71,71));
+
+	const float Alpha = InColor.A;
+	InColor.A = 1.f;
+	
+	return BaseColor * (1.f - Alpha) + InColor * Alpha;
 }
\ No newline at end of file
diff --git a/Engine/Source/Editor/Sequencer/Private/SequencerSelection.cpp b/Engine/Source/Editor/Sequencer/Private/SequencerSelection.cpp
--- a/Engine/Source/Editor/Sequencer/Private/SequencerSelection.cpp
+++ b/Engine/Source/Editor/Sequencer/Private/SequencerSelection.cpp
@@ -5,9 +5,9 @@
 #include "MovieSceneSection.h"
 
 FSequencerSelection::FSequencerSelection()
+	: SuspendBroadcastCount(0)
+	, bOutlinerNodeSelectionChangedBroadcastPending(false)
 {
-	ActiveSelection = EActiveSelection::None;
-	bSuspendBroadcast = false;
 }
 
 const TSet<FSequencerSelectedKey>& FSequencerSelection::GetSelectedKeys() const
@@ -25,9 +25,9 @@ const TSet<TSharedRef<FSequencerDisplayNode>>& FSequencerSelection::GetSelectedO
 	return SelectedOutlinerNodes;
 }
 
-FSequencerSelection::EActiveSelection FSequencerSelection::GetActiveSelection() const
+const TSet<TSharedRef<FSequencerDisplayNode>>& FSequencerSelection::GetNodesWithSelectedKeysOrSections() const
 {
-	return ActiveSelection;
+	return NodesWithSelectedKeysOrSections;
 }
 
 FSequencerSelection::FOnSelectionChanged& FSequencerSelection::GetOnKeySelectionChanged()
@@ -45,41 +45,71 @@ FSequencerSelection::FOnSelectionChanged& FSequencerSelection::GetOnOutlinerNode
 	return OnOutlinerNodeSelectionChanged;
 }
 
+FSequencerSelection::FOnSelectionChanged& FSequencerSelection::GetOnNodesWithSelectedKeysOrSectionsChanged()
+{
+	return OnNodesWithSelectedKeysOrSectionsChanged;
+}
+
 void FSequencerSelection::AddToSelection(FSequencerSelectedKey Key)
 {
 	SelectedKeys.Add(Key);
-	ActiveSelection = EActiveSelection::KeyAndSection;
-	if (!bSuspendBroadcast)
+	if ( IsBroadcasting() )
 	{
 		OnKeySelectionChanged.Broadcast();
 	}
+
+	// Deselect any outliner nodes that aren't within the trunk of this key
+	if (Key.KeyArea.IsValid())
+	{
+		EmptySelectedOutlinerNodesWithoutSection(Key.KeyArea->GetOwningSection());
+	}
+
+	EmptySelectedSections();
 }
 
 void FSequencerSelection::AddToSelection(UMovieSceneSection* Section)
 {
 	SelectedSections.Add(Section);
-	ActiveSelection = EActiveSelection::KeyAndSection;
-	if (!bSuspendBroadcast)
+	if ( IsBroadcasting() )
 	{
 		OnSectionSelectionChanged.Broadcast();
 	}
+
+	// Deselect any outliner nodes that aren't within the trunk of this section
+	if (Section)
+	{
+		EmptySelectedOutlinerNodesWithoutSection(Section);
+	}
+
+	EmptySelectedKeys();
 }
 
 void FSequencerSelection::AddToSelection(TSharedRef<FSequencerDisplayNode> OutlinerNode)
 {
 	SelectedOutlinerNodes.Add(OutlinerNode);
-	ActiveSelection = EActiveSelection::OutlinerNode;
-	if (!bSuspendBroadcast)
+	if ( IsBroadcasting() )
 	{
 		OnOutlinerNodeSelectionChanged.Broadcast();
 	}
+	EmptySelectedKeys();
+	EmptySelectedSections();
+	EmptyNodesWithSelectedKeysOrSections();
+}
+
+void FSequencerSelection::AddToNodesWithSelectedKeysOrSections(TSharedRef<FSequencerDisplayNode> OutlinerNode)
+{
+	NodesWithSelectedKeysOrSections.Add(OutlinerNode);
+	if ( IsBroadcasting() )
+	{
+		OnNodesWithSelectedKeysOrSectionsChanged.Broadcast();
+	}
 }
 
+
 void FSequencerSelection::RemoveFromSelection(FSequencerSelectedKey Key)
 {
 	SelectedKeys.Remove(Key);
-	ActiveSelection = EActiveSelection::KeyAndSection;
-	if (!bSuspendBroadcast)
+	if ( IsBroadcasting() )
 	{
 		OnKeySelectionChanged.Broadcast();
 	}
@@ -88,8 +118,7 @@ void FSequencerSelection::RemoveFromSelection(FSequencerSelectedKey Key)
 void FSequencerSelection::RemoveFromSelection(UMovieSceneSection* Section)
 {
 	SelectedSections.Remove(Section);
-	ActiveSelection = EActiveSelection::KeyAndSection;
-	if (!bSuspendBroadcast)
+	if ( IsBroadcasting() )
 	{
 		OnSectionSelectionChanged.Broadcast();
 	}
@@ -98,13 +127,21 @@ void FSequencerSelection::RemoveFromSelection(UMovieSceneSection* Section)
 void FSequencerSelection::RemoveFromSelection(TSharedRef<FSequencerDisplayNode> OutlinerNode)
 {
 	SelectedOutlinerNodes.Remove(OutlinerNode);
-	ActiveSelection = EActiveSelection::OutlinerNode;
-	if (!bSuspendBroadcast)
+	if ( IsBroadcasting() )
 	{
 		OnOutlinerNodeSelectionChanged.Broadcast();
 	}
 }
 
+void FSequencerSelection::RemoveFromNodesWithSelectedKeysOrSections(TSharedRef<FSequencerDisplayNode> OutlinerNode)
+{
+	NodesWithSelectedKeysOrSections.Remove(OutlinerNode);
+	if ( IsBroadcasting() )
+	{
+		OnNodesWithSelectedKeysOrSectionsChanged.Broadcast();
+	}
+}
+
 bool FSequencerSelection::IsSelected(FSequencerSelectedKey Key) const
 {
 	return SelectedKeys.Contains(Key);
@@ -120,17 +157,28 @@ bool FSequencerSelection::IsSelected(TSharedRef<FSequencerDisplayNode> OutlinerN
 	return SelectedOutlinerNodes.Contains(OutlinerNode);
 }
 
+bool FSequencerSelection::NodeHasSelectedKeysOrSections(TSharedRef<FSequencerDisplayNode> OutlinerNode) const
+{
+	return NodesWithSelectedKeysOrSections.Contains(OutlinerNode);
+}
+
 void FSequencerSelection::Empty()
 {
 	EmptySelectedKeys();
 	EmptySelectedSections();
 	EmptySelectedOutlinerNodes();
+	EmptyNodesWithSelectedKeysOrSections();
 }
 
 void FSequencerSelection::EmptySelectedKeys()
 {
+	if (!SelectedKeys.Num())
+	{
+		return;
+	}
+
 	SelectedKeys.Empty();
-	if (!bSuspendBroadcast)
+	if ( IsBroadcasting() )
 	{
 		OnKeySelectionChanged.Broadcast();
 	}
@@ -138,8 +186,13 @@ void FSequencerSelection::EmptySelectedKeys()
 
 void FSequencerSelection::EmptySelectedSections()
 {
+	if (!SelectedSections.Num())
+	{
+		return;
+	}
+
 	SelectedSections.Empty();
-	if (!bSuspendBroadcast)
+	if ( IsBroadcasting() )
 	{
 		OnSectionSelectionChanged.Broadcast();
 	}
@@ -147,9 +200,94 @@ void FSequencerSelection::EmptySelectedSections()
 
 void FSequencerSelection::EmptySelectedOutlinerNodes()
 {
+	if (!SelectedOutlinerNodes.Num())
+	{
+		return;
+	}
+
 	SelectedOutlinerNodes.Empty();
-	if (!bSuspendBroadcast)
+	if ( IsBroadcasting() )
 	{
 		OnOutlinerNodeSelectionChanged.Broadcast();
 	}
-}
\ No newline at end of file
+}
+
+void FSequencerSelection::EmptyNodesWithSelectedKeysOrSections()
+{
+	if (!NodesWithSelectedKeysOrSections.Num())
+	{
+		return;
+	}
+
+	NodesWithSelectedKeysOrSections.Empty();
+	if ( IsBroadcasting() )
+	{
+		OnNodesWithSelectedKeysOrSectionsChanged.Broadcast();
+	}
+}
+
+/** Suspend or resume broadcast of selection changing  */
+void FSequencerSelection::SuspendBroadcast()
+{
+	SuspendBroadcastCount++;
+}
+
+void FSequencerSelection::ResumeBroadcast()
+{
+	SuspendBroadcastCount--;
+	checkf(SuspendBroadcastCount >= 0, TEXT("Suspend/Resume broadcast mismatch!"));
+}
+
+bool FSequencerSelection::IsBroadcasting() 
+{
+	return SuspendBroadcastCount == 0;
+}
+
+void FSequencerSelection::EmptySelectedOutlinerNodesWithoutSection(UMovieSceneSection* Section)
+{
+	for (auto SelectedOutlinerNode : SelectedOutlinerNodes)
+	{
+		TSet<TSharedRef<FSequencerDisplayNode> > TrunkNodes;
+		TrunkNodes.Add(SelectedOutlinerNode);
+		SequencerHelpers::GetDescendantNodes(SelectedOutlinerNode, TrunkNodes);
+		
+		bool bFoundMatch = false;
+		for (auto TrunkIt = TrunkNodes.CreateConstIterator(); TrunkIt && !bFoundMatch; ++TrunkIt)
+		{
+			TSet<TWeakObjectPtr<UMovieSceneSection>> AllSections;
+			SequencerHelpers::GetAllSections(*TrunkIt, AllSections);
+
+			for (auto SectionIt = AllSections.CreateConstIterator(); SectionIt && !bFoundMatch; ++SectionIt)
+			{
+				if (*SectionIt == Section)
+				{
+					bFoundMatch = true;
+					break;
+				}
+			}
+		}
+
+		if (!bFoundMatch)
+		{
+			RemoveFromSelection(SelectedOutlinerNode);
+		}
+	}
+}
+
+void FSequencerSelection::RequestOutlinerNodeSelectionChangedBroadcast()
+{
+	if ( IsBroadcasting() )
+	{
+		bOutlinerNodeSelectionChangedBroadcastPending = true;
+	}
+}
+
+void FSequencerSelection::Tick()
+{
+	if ( bOutlinerNodeSelectionChangedBroadcastPending && IsBroadcasting() )
+	{
+		bOutlinerNodeSelectionChangedBroadcastPending = false;
+		OnOutlinerNodeSelectionChanged.Broadcast();
+	}
+}
+
diff --git a/Engine/Source/Editor/Sequencer/Private/SequencerSelection.h b/Engine/Source/Editor/Sequencer/Private/SequencerSelection.h
--- a/Engine/Source/Editor/Sequencer/Private/SequencerSelection.h
+++ b/Engine/Source/Editor/Sequencer/Private/SequencerSelection.h
@@ -16,14 +16,6 @@ public:
 
 	DECLARE_MULTICAST_DELEGATE(FOnSelectionChanged)
 
-	/** Options for the current active selection. */
-	enum class EActiveSelection
-	{
-		KeyAndSection,
-		OutlinerNode,
-		None
-	};
-
 public:
 
 	FSequencerSelection();
@@ -37,11 +29,8 @@ public:
 	/** Gets a set of the selected outliner nodes. */
 	const TSet<TSharedRef<FSequencerDisplayNode>>& GetSelectedOutlinerNodes() const;
 
-	/** 
-	 * Gets the currently active selection.  This is used to determine which selection actions
-	 * like delete should act on.
-	 */
-	EActiveSelection GetActiveSelection() const;
+	/** Gets a set of the outliner nodes that have selected keys or sections */
+	const TSet<TSharedRef<FSequencerDisplayNode>>& GetNodesWithSelectedKeysOrSections() const;
 
 	/** Adds a key to the selection */
 	void AddToSelection(FSequencerSelectedKey Key);
@@ -52,6 +41,9 @@ public:
 	/** Adds an outliner node to the selection */
 	void AddToSelection(TSharedRef<FSequencerDisplayNode> OutlinerNode);
 
+	/** Adds an outliner node that has selected keys or sections */
+	void AddToNodesWithSelectedKeysOrSections(TSharedRef<FSequencerDisplayNode> OutlinerNode);
+
 	/** Removes a key from the selection */
 	void RemoveFromSelection(FSequencerSelectedKey Key);
 
@@ -61,6 +53,12 @@ public:
 	/** Removes an outliner node from the selection */
 	void RemoveFromSelection(TSharedRef<FSequencerDisplayNode> OutlinerNode);
 
+	/** Removes an outliner node that has selected keys or sections */
+	void RemoveFromNodesWithSelectedKeysOrSections(TSharedRef<FSequencerDisplayNode> OutlinerNode);
+
+	/** Removes any outliner nodes from the selection that do not relate to the given section */
+	void EmptySelectedOutlinerNodesWithoutSection(UMovieSceneSection* Section);
+
 	/** Returns whether or not the key is selected. */
 	bool IsSelected(FSequencerSelectedKey Key) const;
 
@@ -70,6 +68,9 @@ public:
 	/** Returns whether or not the outliner node is selected. */
 	bool IsSelected(TSharedRef<FSequencerDisplayNode> OutlinerNode) const;
 
+	/** Returns whether or not the outliner node has keys or sections selected. */
+	bool NodeHasSelectedKeysOrSections(TSharedRef<FSequencerDisplayNode> OutlinerNode) const;
+
 	/** Empties all selections. */
 	void Empty();
 
@@ -82,6 +83,9 @@ public:
 	/** Empties the outliner node selection. */
 	void EmptySelectedOutlinerNodes();
 
+	/** Empties the outliner nodes with selected keys or sections. */
+	void EmptyNodesWithSelectedKeysOrSections();
+
 	/** Gets a multicast delegate which is called when the key selection changes. */
 	FOnSelectionChanged& GetOnKeySelectionChanged();
 
@@ -91,21 +95,40 @@ public:
 	/** Gets a multicast delegate which is called when the outliner node selection changes. */
 	FOnSelectionChanged& GetOnOutlinerNodeSelectionChanged();
 
-	/** Suspend or resume broadcast of selection changing  */
-	void SuspendBroadcast() { bSuspendBroadcast = true; }
-	void ResumeBroadcast() { bSuspendBroadcast = false; }
+	/** Gets a multicast delegate which is called when the outliner node with selected keys or sections changes. */
+	FOnSelectionChanged& GetOnNodesWithSelectedKeysOrSectionsChanged();
+
+	/** Suspend the broadcast of selection change notifications.  */
+	void SuspendBroadcast();
+
+	/** Resume the broadcast of selection change notifications.  */
+	void ResumeBroadcast();
+
+	/** Requests that the outliner node selection changed delegate be broadcast on the next update. */
+	void RequestOutlinerNodeSelectionChangedBroadcast();
+
+	/** Updates the selection once per frame.  This is required for deferred selection broadcasts. */
+	void Tick();
+
+private:
+	/** When true, selection change notifications should be broadcasted. */
+	bool IsBroadcasting();
 	
 private:
 
 	TSet<FSequencerSelectedKey> SelectedKeys;
 	TSet<TWeakObjectPtr<UMovieSceneSection>> SelectedSections;
 	TSet<TSharedRef<FSequencerDisplayNode>> SelectedOutlinerNodes;
+	TSet<TSharedRef<FSequencerDisplayNode>> NodesWithSelectedKeysOrSections;
 
 	FOnSelectionChanged OnKeySelectionChanged;
 	FOnSelectionChanged OnSectionSelectionChanged;
 	FOnSelectionChanged OnOutlinerNodeSelectionChanged;
+	FOnSelectionChanged OnNodesWithSelectedKeysOrSectionsChanged;
 
-	EActiveSelection ActiveSelection;
+	/** The number of times the broadcasting of selection change notifications has been suspended. */
+	int32 SuspendBroadcastCount;
 
-	bool bSuspendBroadcast;
+	/** When true there is a pending outliner node selection change which will be broadcast next tick. */
+	bool bOutlinerNodeSelectionChangedBroadcastPending;
 };
diff --git a/Engine/Source/Editor/Sequencer/Private/SequencerSelectionPreview.cpp b/Engine/Source/Editor/Sequencer/Private/SequencerSelectionPreview.cpp
--- a/Engine/Source/Editor/Sequencer/Private/SequencerSelectionPreview.cpp
+++ b/Engine/Source/Editor/Sequencer/Private/SequencerSelectionPreview.cpp
@@ -28,6 +28,18 @@ void FSequencerSelectionPreview::SetSelectionState(UMovieSceneSection* Section,
 	}
 }
 
+void FSequencerSelectionPreview::SetSelectionState(TSharedRef<FSequencerDisplayNode> OutlinerNode, ESelectionPreviewState InState)
+{
+	if (InState == ESelectionPreviewState::Undefined)
+	{
+		DefinedOutlinerNodeStates.Remove(OutlinerNode);
+	}
+	else
+	{
+		DefinedOutlinerNodeStates.Add(OutlinerNode, InState);
+	}
+}
+
 ESelectionPreviewState FSequencerSelectionPreview::GetSelectionState(FSequencerSelectedKey Key) const
 {
 	if (auto* State = DefinedKeyStates.Find(Key))
@@ -46,10 +58,20 @@ ESelectionPreviewState FSequencerSelectionPreview::GetSelectionState(UMovieScene
 	return ESelectionPreviewState::Undefined;
 }
 
+ESelectionPreviewState FSequencerSelectionPreview::GetSelectionState(TSharedRef<FSequencerDisplayNode> OutlinerNode) const
+{
+	if (auto* State = DefinedOutlinerNodeStates.Find(OutlinerNode))
+	{
+		return *State;
+	}
+	return ESelectionPreviewState::Undefined;
+}
+
 void FSequencerSelectionPreview::Empty()
 {
 	EmptyDefinedKeyStates();
 	EmptyDefinedSectionStates();
+	EmptyDefinedOutlinerNodeStates();
 }
 
 void FSequencerSelectionPreview::EmptyDefinedKeyStates()
@@ -60,4 +82,9 @@ void FSequencerSelectionPreview::EmptyDefinedKeyStates()
 void FSequencerSelectionPreview::EmptyDefinedSectionStates()
 {
 	DefinedSectionStates.Reset();
+}
+
+void FSequencerSelectionPreview::EmptyDefinedOutlinerNodeStates()
+{
+	DefinedOutlinerNodeStates.Reset();
 }
\ No newline at end of file
diff --git a/Engine/Source/Editor/Sequencer/Private/SequencerSelectionPreview.h b/Engine/Source/Editor/Sequencer/Private/SequencerSelectionPreview.h
--- a/Engine/Source/Editor/Sequencer/Private/SequencerSelectionPreview.h
+++ b/Engine/Source/Editor/Sequencer/Private/SequencerSelectionPreview.h
@@ -28,18 +28,27 @@ public:
 	/** Access the defined section states */
 	const TMap<TWeakObjectPtr<UMovieSceneSection>, ESelectionPreviewState>& GetDefinedSectionStates() const { return DefinedSectionStates; }
 
+	/** Access the defined display node states */
+	const TMap<TSharedRef<FSequencerDisplayNode>, ESelectionPreviewState>& GetDefinedOutlinerNodeStates() const { return DefinedOutlinerNodeStates; }
+
 	/** Adds a key to the selection */
 	void SetSelectionState(FSequencerSelectedKey Key, ESelectionPreviewState InState);
 
 	/** Adds a section to the selection */
 	void SetSelectionState(UMovieSceneSection* Section, ESelectionPreviewState InState);
 
+	/** Adds an outliner node to the selection */
+	void SetSelectionState(TSharedRef<FSequencerDisplayNode> OutlinerNode, ESelectionPreviewState InState);
+
 	/** Returns the selection state for the the specified key. */
 	ESelectionPreviewState GetSelectionState(FSequencerSelectedKey Key) const;
 
 	/** Returns the selection state for the the specified section. */
 	ESelectionPreviewState GetSelectionState(UMovieSceneSection* Section) const;
 
+	/** Returns the selection state for the the specified section. */
+	ESelectionPreviewState GetSelectionState(TSharedRef<FSequencerDisplayNode> OutlinerNode) const;
+
 	/** Empties all selections. */
 	void Empty();
 
@@ -49,8 +58,12 @@ public:
 	/** Empties the section selection. */
 	void EmptyDefinedSectionStates();
 
+	/** Empties the outliner node selection. */
+	void EmptyDefinedOutlinerNodeStates();
+
 private:
 
 	TMap<FSequencerSelectedKey, ESelectionPreviewState> DefinedKeyStates;
 	TMap<TWeakObjectPtr<UMovieSceneSection>, ESelectionPreviewState> DefinedSectionStates;
+	TMap<TSharedRef<FSequencerDisplayNode>, ESelectionPreviewState> DefinedOutlinerNodeStates;
 };
diff --git a/Engine/Source/Editor/Sequencer/Private/SequencerSettings.cpp b/Engine/Source/Editor/Sequencer/Private/SequencerSettings.cpp
--- a/Engine/Source/Editor/Sequencer/Private/SequencerSettings.cpp
+++ b/Engine/Source/Editor/Sequencer/Private/SequencerSettings.cpp
@@ -18,9 +18,10 @@ USequencerSettings::USequencerSettings( const FObjectInitializer& ObjectInitiali
 	bSnapKeyTimesToKeys = true;
 	bSnapSectionTimesToInterval = true;
 	bSnapSectionTimesToSections = true;
+	bSnapPlayTimeToKeys = false;
 	bSnapPlayTimeToInterval = true;
 	bSnapPlayTimeToDraggedKey = false;
-	bFixedTimeStepPlayback = false;
+	bFixedTimeStepPlayback = true;
 	CurveValueSnapInterval = 10.0f;
 	bSnapCurveValueToInterval = true;
 	bDetailsViewVisible = false;
@@ -34,6 +35,9 @@ USequencerSettings::USequencerSettings( const FObjectInitializer& ObjectInitiali
 	bKeepPlayRangeInSectionBounds = true;
 	ZeroPadFrames = 0;
 	bShowCombinedKeyframes = true;
+	bInfiniteKeyAreas = false;
+	bShowChannelColors = false;
+	bShowViewportTransportControls = true;
 }
 
 EAutoKeyMode USequencerSettings::GetAutoKeyMode() const
@@ -120,20 +124,6 @@ void USequencerSettings::SetShowFrameNumbers(bool InbShowFrameNumbers)
 	}
 }
 
-bool USequencerSettings::GetShowGotoBox() const
-{
-	return bShowGotoBox;
-}
-
-void USequencerSettings::SetShowGotoBox(bool InbShowGotoBox)
-{
-	if ( bShowGotoBox != InbShowGotoBox )
-	{
-		bShowGotoBox = InbShowGotoBox;
-		SaveConfig();
-	}
-}
-
 bool USequencerSettings::GetShowRangeSlider() const
 {
 	return bShowRangeSlider;
@@ -232,6 +222,20 @@ void USequencerSettings::SetSnapSectionTimesToSections( bool InbSnapSectionTimes
 	}
 }
 
+bool USequencerSettings::GetSnapPlayTimeToKeys() const
+{
+	return bSnapPlayTimeToKeys;
+}
+
+void USequencerSettings::SetSnapPlayTimeToKeys(bool InbSnapPlayTimeToKeys)
+{
+	if ( bSnapPlayTimeToKeys != InbSnapPlayTimeToKeys )
+	{
+		bSnapPlayTimeToKeys = InbSnapPlayTimeToKeys;
+		SaveConfig();
+	}
+}
+
 bool USequencerSettings::GetSnapPlayTimeToInterval() const
 {
 	return bSnapPlayTimeToInterval;
@@ -461,6 +465,49 @@ void USequencerSettings::SetShowCombinedKeyframes(bool InbShowCombinedKeyframes)
 }
 
 
+bool USequencerSettings::GetInfiniteKeyAreas() const
+{
+	return bInfiniteKeyAreas;
+}
+
+void USequencerSettings::SetInfiniteKeyAreas(bool InbInfiniteKeyAreas)
+{
+	if (bInfiniteKeyAreas != InbInfiniteKeyAreas)
+	{
+		bInfiniteKeyAreas = InbInfiniteKeyAreas;
+		SaveConfig();
+	}
+}
+
+
+bool USequencerSettings::GetShowChannelColors() const
+{
+	return bShowChannelColors;
+}
+
+void USequencerSettings::SetShowChannelColors(bool InbShowChannelColors)
+{
+	if (bShowChannelColors != InbShowChannelColors)
+	{
+		bShowChannelColors = InbShowChannelColors;
+		SaveConfig();
+	}
+}
+
+bool USequencerSettings::GetShowViewportTransportControls() const
+{
+	return bShowViewportTransportControls;
+}
+
+void USequencerSettings::SetShowViewportTransportControls(bool bVisible)
+{
+	if (bShowViewportTransportControls != bVisible)
+	{
+		bShowViewportTransportControls = bVisible;
+		SaveConfig();
+	}
+}
+
 float USequencerSettings::SnapTimeToInterval( float InTimeValue ) const
 {
 	return TimeSnapInterval > 0
diff --git a/Engine/Source/Editor/Sequencer/Private/SequencerSettings.h b/Engine/Source/Editor/Sequencer/Private/SequencerSettings.h
--- a/Engine/Source/Editor/Sequencer/Private/SequencerSettings.h
+++ b/Engine/Source/Editor/Sequencer/Private/SequencerSettings.h
@@ -89,11 +89,6 @@ public:
 	/** Sets whether or not to show frame numbers. */
 	void SetShowFrameNumbers(bool InbShowFrameNumbers);
 
-	/** Gets whether or not to show the goto box. */
-	bool GetShowGotoBox() const;
-	/** Sets whether or not to show the goto box. */
-	void SetShowGotoBox(bool InbShowGotoBox);
-
 	/** Gets whether or not to show the time range slider. */
 	bool GetShowRangeSlider() const;
 	/** Sets whether or not to show frame numbers. */
@@ -129,6 +124,11 @@ public:
 	/** sets whether or not to snap sections to other sections. */
 	void SetSnapSectionTimesToSections( bool InbSnapSectionTimesToSections );
 
+	/** Gets whether or not to snap the play time to keys while scrubbing. */
+	bool GetSnapPlayTimeToKeys() const;
+	/** Sets whether or not to snap the play time to keys while scrubbing. */
+	void SetSnapPlayTimeToKeys(bool InbSnapPlayTimeToKeys);
+
 	/** Gets whether or not to snap the play time to the interval while scrubbing. */
 	bool GetSnapPlayTimeToInterval() const;
 	/** Sets whether or not to snap the play time to the interval while scrubbing. */
@@ -209,7 +209,22 @@ public:
 	/** @return true if showing combined keyframes at the top node */
 	bool GetShowCombinedKeyframes() const;
 	/** Set whether to show combined keyframes at the top node */ 
-	void SetShowCombinedKeyframes(bool InShowCombinedKeyframes);
+	void SetShowCombinedKeyframes(bool bInShowCombinedKeyframes);
+
+	/** @return true if key areas are infinite */
+	bool GetInfiniteKeyAreas() const;
+	/** Set whether to show channel colors */
+	void SetInfiniteKeyAreas(bool bInInfiniteKeyAreas);
+
+	/** @return true if showing channel colors */
+	bool GetShowChannelColors() const;
+	/** Set whether to show channel colors */
+	void SetShowChannelColors(bool bInShowChannelColors);
+
+	/** @return true if showing transport controls in level editor viewports */
+	bool GetShowViewportTransportControls() const;
+	/** Toggle whether to show transport controls in level editor viewports */
+	void SetShowViewportTransportControls(bool bVisible);
 
 	/** Snaps a time value in seconds to the currently selected interval. */
 	float SnapTimeToInterval(float InTimeValue) const;
@@ -237,9 +252,6 @@ protected:
 	UPROPERTY( config )
 	bool bShowFrameNumbers;
 
-	UPROPERTY( config )
-	bool bShowGotoBox;
-
 	UPROPERTY( config )
 	bool bShowRangeSlider;
 
@@ -261,6 +273,9 @@ protected:
 	UPROPERTY( config )
 	bool bSnapSectionTimesToSections;
 
+	UPROPERTY( config )
+	bool bSnapPlayTimeToKeys;
+
 	UPROPERTY( config )
 	bool bSnapPlayTimeToInterval;
 
@@ -309,5 +324,14 @@ protected:
 	UPROPERTY( config )
 	bool bShowCombinedKeyframes;
 
+	UPROPERTY( config )
+	bool bInfiniteKeyAreas;
+
+	UPROPERTY( config )
+	bool bShowChannelColors;
+
+	UPROPERTY( config )
+	bool bShowViewportTransportControls;
+
 	FOnShowCurveEditorChanged OnShowCurveEditorChanged;
 };
