Copying //UE4/Dev-Sequencer to //UE4/Main

diff --git a/Engine/Source/Editor/Sequencer/Private/SequencerTimeSliderController.cpp b/Engine/Source/Editor/Sequencer/Private/SequencerTimeSliderController.cpp
--- a/Engine/Source/Editor/Sequencer/Private/SequencerTimeSliderController.cpp
+++ b/Engine/Source/Editor/Sequencer/Private/SequencerTimeSliderController.cpp
@@ -258,43 +258,49 @@ int32 FSequencerTimeSliderController::OnPaintTimeSlider( bool bMirrorLabels, con
 	if ( LocalSequenceLength > 0)
 	{
 		FScrubRangeToScreen RangeToScreen( LocalViewRange, AllottedGeometry.Size );
-	
+
+		// draw tick marks
 		const float MajorTickHeight = 9.0f;
 	
 		FDrawTickArgs Args;
-		Args.AllottedGeometry = AllottedGeometry;
-		Args.bMirrorLabels = bMirrorLabels;
-		Args.bOnlyDrawMajorTicks = false;
-		Args.TickColor = FLinearColor::White;
-		Args.ClippingRect = MyClippingRect;
-		Args.DrawEffects = DrawEffects;
-		Args.StartLayer = LayerId;
-		Args.TickOffset = bMirrorLabels ? 0.0f : FMath::Abs( AllottedGeometry.Size.Y - MajorTickHeight );
-		Args.MajorTickHeight = MajorTickHeight;
+		{
+			Args.AllottedGeometry = AllottedGeometry;
+			Args.bMirrorLabels = bMirrorLabels;
+			Args.bOnlyDrawMajorTicks = false;
+			Args.TickColor = FLinearColor::White;
+			Args.ClippingRect = MyClippingRect;
+			Args.DrawEffects = DrawEffects;
+			Args.StartLayer = LayerId;
+			Args.TickOffset = bMirrorLabels ? 0.0f : FMath::Abs( AllottedGeometry.Size.Y - MajorTickHeight );
+			Args.MajorTickHeight = MajorTickHeight;
+		}
 
 		DrawTicks( OutDrawElements, RangeToScreen, Args );
 
+		// draw playback & in/out range
 		FPaintPlaybackRangeArgs PlaybackRangeArgs(
 			bMirrorLabels ? FEditorStyle::GetBrush("Sequencer.Timeline.PlayRange_Bottom_L") : FEditorStyle::GetBrush("Sequencer.Timeline.PlayRange_Top_L"),
 			bMirrorLabels ? FEditorStyle::GetBrush("Sequencer.Timeline.PlayRange_Bottom_R") : FEditorStyle::GetBrush("Sequencer.Timeline.PlayRange_Top_R"),
 			6.f
 		);
+
 		LayerId = DrawPlaybackRange(AllottedGeometry, MyClippingRect, OutDrawElements, LayerId, RangeToScreen, PlaybackRangeArgs);
+		LayerId = DrawInOutRange(AllottedGeometry, MyClippingRect, OutDrawElements, LayerId, RangeToScreen, PlaybackRangeArgs);
 
-		float HalfSize = FMath::CeilToFloat(ScrubHandleSize/2.0f);
+		float HalfSize = FMath::CeilToFloat(ScrubHandleSize / 2.0f);
 
 		// Draw the scrub handle
 		float XPos = RangeToScreen.InputToLocalX( TimeSliderArgs.ScrubPosition.Get() );
-
-		// Should draw above the text
 		const int32 ArrowLayer = LayerId + 2;
 		FPaintGeometry MyGeometry =	AllottedGeometry.ToPaintGeometry( FVector2D( XPos-HalfSize, 0 ), FVector2D( ScrubHandleSize, AllottedGeometry.Size.Y ) );
 		FLinearColor ScrubColor = InWidgetStyle.GetColorAndOpacityTint();
+		{
+			// @todo Sequencer this color should be specified in the style
+			ScrubColor.A = ScrubColor.A * 0.75f;
+			ScrubColor.B *= 0.1f;
+			ScrubColor.G *= 0.2f;
+		}
 
-		// @todo Sequencer this color should be specified in the style
-		ScrubColor.A = ScrubColor.A*0.75f;
-		ScrubColor.B *= 0.1f;
-		ScrubColor.G *= 0.2f;
 		FSlateDrawElement::MakeBox( 
 			OutDrawElements,
 			ArrowLayer, 
@@ -303,20 +309,20 @@ int32 FSequencerTimeSliderController::OnPaintTimeSlider( bool bMirrorLabels, con
 			MyClippingRect, 
 			DrawEffects, 
 			ScrubColor
-			);
+		);
 
 		// Draw the current time next to the scrub handle
 		float Time = TimeSliderArgs.ScrubPosition.Get();
-
 		FString FrameString;
+
 		if (SequencerSnapValues::IsTimeSnapIntervalFrameRate(TimeSliderArgs.Settings->GetTimeSnapInterval()) && TimeSliderArgs.Settings->GetShowFrameNumbers())
 		{
 			float FrameRate = 1.0f/TimeSliderArgs.Settings->GetTimeSnapInterval();
 			float FrameTime = Time * FrameRate;
 			int32 Frame = SequencerHelpers::TimeToFrame(Time, FrameRate);
-					
 			const float FrameTolerance = 0.001f;
-			if (FMath::IsNearlyEqual(FrameTime, (float)Frame, FrameTolerance))
+
+			if (FMath::IsNearlyEqual(FrameTime, (float)Frame, FrameTolerance) || TimeSliderArgs.PlaybackStatus.Get() == EMovieScenePlayerStatus::Playing)
 			{
 				FrameString = FString::Printf( TEXT("%d"), TimeToFrame(Time));
 			}
@@ -383,6 +389,52 @@ int32 FSequencerTimeSliderController::OnPaintTimeSlider( bool bMirrorLabels, con
 	return LayerId;
 }
 
+
+int32 FSequencerTimeSliderController::DrawInOutRange(const FGeometry& AllottedGeometry, const FSlateRect& MyClippingRect, FSlateWindowElementList& OutDrawElements, int32 LayerId, const FScrubRangeToScreen& RangeToScreen, const FPaintPlaybackRangeArgs& Args) const
+{
+	const TRange<float> InOutRange = TimeSliderArgs.InOutRange.Get();
+
+	if (!InOutRange.IsEmpty())
+	{
+		const float InOutRangeL = RangeToScreen.InputToLocalX(InOutRange.GetLowerBoundValue()) - 1;
+		const float InOutRangeR = RangeToScreen.InputToLocalX(InOutRange.GetUpperBoundValue()) + 1;
+		const auto DrawColor = FLinearColor(FColor(32, 32, 128));
+
+		FSlateDrawElement::MakeBox(
+			OutDrawElements,
+			LayerId + 1,
+			AllottedGeometry.ToPaintGeometry(FVector2D(InOutRangeL, 0.f), FVector2D(InOutRangeR - InOutRangeL, AllottedGeometry.Size.Y)),
+			FEditorStyle::GetBrush("WhiteBrush"),
+			MyClippingRect,
+			ESlateDrawEffect::None,
+			DrawColor.CopyWithNewOpacity(0.2f)
+		);
+
+		FSlateDrawElement::MakeBox(
+			OutDrawElements,
+			LayerId + 1,
+			AllottedGeometry.ToPaintGeometry(FVector2D(InOutRangeL, 0.f), FVector2D(Args.BrushWidth, AllottedGeometry.Size.Y)),
+			Args.StartBrush,
+			MyClippingRect,
+			ESlateDrawEffect::None,
+			DrawColor
+		);
+
+		FSlateDrawElement::MakeBox(
+			OutDrawElements,
+			LayerId + 1,
+			AllottedGeometry.ToPaintGeometry(FVector2D(InOutRangeR - Args.BrushWidth, 0.f), FVector2D(Args.BrushWidth, AllottedGeometry.Size.Y)),
+			Args.EndBrush,
+			MyClippingRect,
+			ESlateDrawEffect::None,
+			DrawColor
+		);
+	}
+
+	return LayerId + 1;
+}
+
+
 int32 FSequencerTimeSliderController::DrawPlaybackRange(const FGeometry& AllottedGeometry, const FSlateRect& MyClippingRect, FSlateWindowElementList& OutDrawElements, int32 LayerId, const FScrubRangeToScreen& RangeToScreen, const FPaintPlaybackRangeArgs& Args) const
 {
 	if (!TimeSliderArgs.PlaybackRange.IsSet())
@@ -390,10 +442,9 @@ int32 FSequencerTimeSliderController::DrawPlaybackRange(const FGeometry& Allotte
 		return LayerId;
 	}
 
-	TRange<float> PlaybackRange = TimeSliderArgs.PlaybackRange.Get();
-
-	float PlaybackRangeL = RangeToScreen.InputToLocalX(PlaybackRange.GetLowerBoundValue()) - 1;
-	float PlaybackRangeR = RangeToScreen.InputToLocalX(PlaybackRange.GetUpperBoundValue()) + 1;
+	const TRange<float> PlaybackRange = TimeSliderArgs.PlaybackRange.Get();
+	const float PlaybackRangeL = RangeToScreen.InputToLocalX(PlaybackRange.GetLowerBoundValue()) - 1;
+	const float PlaybackRangeR = RangeToScreen.InputToLocalX(PlaybackRange.GetUpperBoundValue()) + 1;
 
 	FSlateDrawElement::MakeBox(
 		OutDrawElements,
@@ -435,6 +486,7 @@ int32 FSequencerTimeSliderController::DrawPlaybackRange(const FGeometry& Allotte
 		ESlateDrawEffect::None,
 		FLinearColor::Black.CopyWithNewOpacity(0.3f)
 	);
+
 	return LayerId + 1;
 }
 
@@ -509,11 +561,18 @@ FReply FSequencerTimeSliderController::OnMouseButtonUp( SWidget& WidgetOwner, co
 		{
 			TimeSliderArgs.OnEndPlaybackRangeDrag.ExecuteIfBound();
 		}
-		// Set the end range time?
 		else if (MouseDragType == DRAG_END_RANGE)
 		{
 			TimeSliderArgs.OnEndPlaybackRangeDrag.ExecuteIfBound();
 		}
+		else if (MouseDragType == DRAG_IN_RANGE)
+		{
+			TimeSliderArgs.OnEndInOutRangeDrag.ExecuteIfBound();
+		}
+		else if (MouseDragType == DRAG_OUT_RANGE)
+		{
+			TimeSliderArgs.OnEndInOutRangeDrag.ExecuteIfBound();
+		}
 		else if (MouseDragType == DRAG_SETTING_RANGE)
 		{
 			FScrubRangeToScreen RangeToScreen( TimeSliderArgs.ViewRange.Get(), MyGeometry.Size );
@@ -585,124 +644,170 @@ FReply FSequencerTimeSliderController::OnMouseButtonUp( SWidget& WidgetOwner, co
 	return FReply::Unhandled();
 }
 
+
 FReply FSequencerTimeSliderController::OnMouseMove( SWidget& WidgetOwner, const FGeometry& MyGeometry, const FPointerEvent& MouseEvent )
 {
-	if ( WidgetOwner.HasMouseCapture() )
+	if (!WidgetOwner.HasMouseCapture())
 	{
-		if (MouseEvent.IsMouseButtonDown( EKeys::RightMouseButton ))
+		return FReply::Unhandled();
+	}
+
+	if (MouseEvent.IsMouseButtonDown( EKeys::RightMouseButton ))
+	{
+		if (!bPanning)
 		{
-			if (!bPanning)
+			DistanceDragged += FMath::Abs( MouseEvent.GetCursorDelta().X );
+			if ( DistanceDragged > FSlateApplication::Get().GetDragTriggerDistance() )
 			{
-				DistanceDragged += FMath::Abs( MouseEvent.GetCursorDelta().X );
-				if ( DistanceDragged > FSlateApplication::Get().GetDragTriggerDistance() )
-				{
-					bPanning = true;
-				}
+				bPanning = true;
 			}
-			else
-			{
-				TRange<float> LocalViewRange = TimeSliderArgs.ViewRange.Get();
-				float LocalViewRangeMin = LocalViewRange.GetLowerBoundValue();
-				float LocalViewRangeMax = LocalViewRange.GetUpperBoundValue();
+		}
+		else
+		{
+			TRange<float> LocalViewRange = TimeSliderArgs.ViewRange.Get();
+			float LocalViewRangeMin = LocalViewRange.GetLowerBoundValue();
+			float LocalViewRangeMax = LocalViewRange.GetUpperBoundValue();
+
+			FScrubRangeToScreen ScaleInfo( LocalViewRange, MyGeometry.Size );
+			FVector2D ScreenDelta = MouseEvent.GetCursorDelta();
+			FVector2D InputDelta;
+			InputDelta.X = ScreenDelta.X/ScaleInfo.PixelsPerInput;
+
+			float NewViewOutputMin = LocalViewRangeMin - InputDelta.X;
+			float NewViewOutputMax = LocalViewRangeMax - InputDelta.X;
+
+			ClampViewRange(NewViewOutputMin, NewViewOutputMax);
+			SetViewRange(NewViewOutputMin, NewViewOutputMax, EViewRangeInterpolation::Immediate);
+		}
+	}
+	else if (MouseEvent.IsMouseButtonDown( EKeys::LeftMouseButton ))
+	{
+		TRange<float> LocalViewRange = TimeSliderArgs.ViewRange.Get();
+		DistanceDragged += FMath::Abs( MouseEvent.GetCursorDelta().X );
 
-				FScrubRangeToScreen ScaleInfo( LocalViewRange, MyGeometry.Size );
-				FVector2D ScreenDelta = MouseEvent.GetCursorDelta();
-				FVector2D InputDelta;
-				InputDelta.X = ScreenDelta.X/ScaleInfo.PixelsPerInput;
+		if ( MouseDragType == DRAG_NONE )
+		{
+			if ( DistanceDragged > FSlateApplication::Get().GetDragTriggerDistance() )
+			{
+				FScrubRangeToScreen RangeToScreen(LocalViewRange, MyGeometry.Size);
+				const float ScrubPosition = TimeSliderArgs.ScrubPosition.Get();
 
-				float NewViewOutputMin = LocalViewRangeMin - InputDelta.X;
-				float NewViewOutputMax = LocalViewRangeMax - InputDelta.X;
+				TRange<float> InOutRange = TimeSliderArgs.InOutRange.Get();
+				TRange<float> PlaybackRange = TimeSliderArgs.PlaybackRange.Get();
+				float LocalMouseDownPos = RangeToScreen.InputToLocalX(MouseDownRange[0]);
 
-				ClampViewRange(NewViewOutputMin, NewViewOutputMax);
-				SetViewRange(NewViewOutputMin, NewViewOutputMax, EViewRangeInterpolation::Immediate);
+				if (HitTestScrubberEnd(RangeToScreen, InOutRange, LocalMouseDownPos, ScrubPosition))
+				{
+					// in/out range end scrubber
+					MouseDragType = DRAG_OUT_RANGE;
+					TimeSliderArgs.OnBeginInOutRangeDrag.ExecuteIfBound();
+				}
+				else if (HitTestScrubberStart(RangeToScreen, InOutRange, LocalMouseDownPos, ScrubPosition))
+				{
+					// in/out range start scrubber
+					MouseDragType = DRAG_IN_RANGE;
+					TimeSliderArgs.OnBeginInOutRangeDrag.ExecuteIfBound();
+				}
+				else if (HitTestScrubberEnd(RangeToScreen, PlaybackRange, LocalMouseDownPos, ScrubPosition))
+				{
+					// playback range end scrubber
+					MouseDragType = DRAG_END_RANGE;
+					TimeSliderArgs.OnBeginPlaybackRangeDrag.ExecuteIfBound();
+				}
+				else if (HitTestScrubberStart(RangeToScreen, PlaybackRange, LocalMouseDownPos, ScrubPosition))
+				{
+					// playback range start scrubber
+					MouseDragType = DRAG_START_RANGE;
+					TimeSliderArgs.OnBeginPlaybackRangeDrag.ExecuteIfBound();
+				}
+				else if (FSlateApplication::Get().GetModifierKeys().AreModifersDown(EModifierKey::Control))
+				{
+					MouseDragType = DRAG_SETTING_RANGE;
+				}
+				else
+				{
+					MouseDragType = DRAG_SCRUBBING_TIME;
+					TimeSliderArgs.OnBeginScrubberMovement.ExecuteIfBound();
+				}
 			}
 		}
-		else if (MouseEvent.IsMouseButtonDown( EKeys::LeftMouseButton ))
+		else
 		{
-			TRange<float> LocalViewRange = TimeSliderArgs.ViewRange.Get();
-			DistanceDragged += FMath::Abs( MouseEvent.GetCursorDelta().X );
+			FScrubRangeToScreen RangeToScreen( TimeSliderArgs.ViewRange.Get(), MyGeometry.Size );
+			FVector2D CursorPos = MyGeometry.AbsoluteToLocal( MouseEvent.GetLastScreenSpacePosition() );
+			float NewValue = RangeToScreen.LocalXToInput( CursorPos.X );
 
-			if ( MouseDragType == DRAG_NONE )
+
+			// Set the start range time?
+			if (MouseDragType == DRAG_START_RANGE)
 			{
-				if ( DistanceDragged > FSlateApplication::Get().GetDragTriggerDistance() )
+				if (TimeSliderArgs.Settings->GetIsSnapEnabled())
 				{
-					FScrubRangeToScreen RangeToScreen(LocalViewRange, MyGeometry.Size);
-					const float ScrubPosition = TimeSliderArgs.ScrubPosition.Get();
-
-					TRange<float> PlaybackRange = TimeSliderArgs.PlaybackRange.Get();
-					float LocalMouseDownPos = RangeToScreen.InputToLocalX(MouseDownRange[0]);
-
-					// Favor dragging the end position
-					if (HitTestPlaybackEnd(RangeToScreen, PlaybackRange, LocalMouseDownPos, ScrubPosition))
-					{
-						MouseDragType = DRAG_END_RANGE;
-						TimeSliderArgs.OnBeginPlaybackRangeDrag.ExecuteIfBound();
-					}
-					else if (HitTestPlaybackStart(RangeToScreen, PlaybackRange, LocalMouseDownPos, ScrubPosition))
-					{
-						MouseDragType = DRAG_START_RANGE;
-						TimeSliderArgs.OnBeginPlaybackRangeDrag.ExecuteIfBound();
-					}
-					else if (FSlateApplication::Get().GetModifierKeys().AreModifersDown(EModifierKey::Control))
-					{
-						MouseDragType = DRAG_SETTING_RANGE;
-					}
-					else
-					{
-						MouseDragType = DRAG_SCRUBBING_TIME;
-						TimeSliderArgs.OnBeginScrubberMovement.ExecuteIfBound();
-					}
+					NewValue = TimeSliderArgs.Settings->SnapTimeToInterval(NewValue);
 				}
+
+				SetPlaybackRangeStart(NewValue);
 			}
-			else
+			// Set the end range time?
+			else if(MouseDragType == DRAG_END_RANGE)
 			{
-				FScrubRangeToScreen RangeToScreen( TimeSliderArgs.ViewRange.Get(), MyGeometry.Size );
-				FVector2D CursorPos = MyGeometry.AbsoluteToLocal( MouseEvent.GetLastScreenSpacePosition() );
-				float NewValue = RangeToScreen.LocalXToInput( CursorPos.X );
+				if (TimeSliderArgs.Settings->GetIsSnapEnabled())
+				{
+					NewValue = TimeSliderArgs.Settings->SnapTimeToInterval(NewValue);
+				}
+					
+				SetPlaybackRangeEnd(NewValue);
+			}
+			else if (MouseDragType == DRAG_IN_RANGE)
+			{
+				if (TimeSliderArgs.Settings->GetIsSnapEnabled())
+				{
+					NewValue = TimeSliderArgs.Settings->SnapTimeToInterval(NewValue);
+				}
 
+				TRange<float> InOutRange = TimeSliderArgs.InOutRange.Get();
 
-				// Set the start range time?
-				if (MouseDragType == DRAG_START_RANGE)
+				if (NewValue <= InOutRange.GetUpperBoundValue())
 				{
-					if (TimeSliderArgs.Settings->GetIsSnapEnabled())
-					{
-						NewValue = TimeSliderArgs.Settings->SnapTimeToInterval(NewValue);
-					}
-
-					SetPlaybackRangeStart(NewValue);
+					TimeSliderArgs.OnInOutRangeChanged.ExecuteIfBound(TRange<float>(NewValue, InOutRange.GetUpperBoundValue()));
 				}
-				// Set the end range time?
-				else if(MouseDragType == DRAG_END_RANGE)
+			}
+			// Set the end range time?
+			else if(MouseDragType == DRAG_OUT_RANGE)
+			{
+				if (TimeSliderArgs.Settings->GetIsSnapEnabled())
 				{
-					if (TimeSliderArgs.Settings->GetIsSnapEnabled())
-					{
-						NewValue = TimeSliderArgs.Settings->SnapTimeToInterval(NewValue);
-					}
-					
-					SetPlaybackRangeEnd(NewValue);
+					NewValue = TimeSliderArgs.Settings->SnapTimeToInterval(NewValue);
 				}
-				else if (MouseDragType == DRAG_SCRUBBING_TIME)
-				{
-					if ( TimeSliderArgs.Settings->GetIsSnapEnabled() && TimeSliderArgs.Settings->GetSnapPlayTimeToInterval() )
-					{
-						NewValue = TimeSliderArgs.Settings->SnapTimeToInterval(NewValue);
-					}
 					
-					// Delegate responsibility for clamping to the current viewrange to the client
-					CommitScrubPosition( NewValue, /*bIsScrubbing=*/true );
+				TRange<float> InOutRange = TimeSliderArgs.InOutRange.Get();
+
+				if (NewValue >= InOutRange.GetLowerBoundValue())
+				{
+					TimeSliderArgs.OnInOutRangeChanged.ExecuteIfBound(TRange<float>(InOutRange.GetLowerBoundValue(), NewValue));
 				}
-				else if (MouseDragType == DRAG_SETTING_RANGE)
+			}
+			else if (MouseDragType == DRAG_SCRUBBING_TIME)
+			{
+				if ( TimeSliderArgs.Settings->GetIsSnapEnabled() && TimeSliderArgs.Settings->GetSnapPlayTimeToInterval() )
 				{
-					MouseDownRange[1] = NewValue;
+					NewValue = TimeSliderArgs.Settings->SnapTimeToInterval(NewValue);
 				}
+					
+				// Delegate responsibility for clamping to the current viewrange to the client
+				CommitScrubPosition( NewValue, /*bIsScrubbing=*/true );
+			}
+			else if (MouseDragType == DRAG_SETTING_RANGE)
+			{
+				MouseDownRange[1] = NewValue;
 			}
 		}
-		return FReply::Handled();
 	}
 
-	return FReply::Unhandled();
+	return FReply::Handled();
 }
 
+
 void FSequencerTimeSliderController::CommitScrubPosition( float NewValue, bool bIsScrubbing )
 {
 	// Manage the scrub position ourselves if its not bound to a delegate
@@ -740,15 +845,19 @@ FCursorReply FSequencerTimeSliderController::OnCursorQuery( TSharedRef<const SWi
 {
 	FScrubRangeToScreen RangeToScreen(TimeSliderArgs.ViewRange.Get(), MyGeometry.Size);
 	TRange<float> PlaybackRange = TimeSliderArgs.PlaybackRange.Get();
+	TRange<float> InOutRange = TimeSliderArgs.InOutRange.Get();
 	const float ScrubPosition = TimeSliderArgs.ScrubPosition.Get();
 
 	// Use L/R resize cursor if we're dragging or hovering a playback range bound
 	float HitTestPosition = MyGeometry.AbsoluteToLocal(CursorEvent.GetScreenSpacePosition()).X;
-	if (MouseDragType == DRAG_END_RANGE ||
-		MouseDragType == DRAG_START_RANGE ||
-		HitTestPlaybackStart(RangeToScreen, PlaybackRange, HitTestPosition, ScrubPosition) ||
-		HitTestPlaybackEnd(RangeToScreen, PlaybackRange, HitTestPosition, ScrubPosition)
-		)
+	if ((MouseDragType == DRAG_END_RANGE) ||
+		(MouseDragType == DRAG_START_RANGE) ||
+		(MouseDragType == DRAG_IN_RANGE) ||
+		(MouseDragType == DRAG_OUT_RANGE) ||
+		HitTestScrubberStart(RangeToScreen, PlaybackRange, HitTestPosition, ScrubPosition) ||
+		HitTestScrubberEnd(RangeToScreen, PlaybackRange, HitTestPosition, ScrubPosition) ||
+		HitTestScrubberStart(RangeToScreen, InOutRange, HitTestPosition, ScrubPosition) ||
+		HitTestScrubberEnd(RangeToScreen, InOutRange, HitTestPosition, ScrubPosition))
 	{
 		return FCursorReply::Cursor(EMouseCursor::ResizeLeftRight);
 	}
@@ -772,6 +881,7 @@ int32 FSequencerTimeSliderController::OnPaintSectionView( const FGeometry& Allot
 	if (Args.PlaybackRangeArgs.IsSet())
 	{
 		LayerId = DrawPlaybackRange(AllottedGeometry, MyClippingRect, OutDrawElements, LayerId, RangeToScreen, Args.PlaybackRangeArgs.GetValue());
+		LayerId = DrawInOutRange(AllottedGeometry, MyClippingRect, OutDrawElements, LayerId, RangeToScreen, Args.PlaybackRangeArgs.GetValue());
 	}
 
 	if( Args.bDisplayTickLines )
@@ -780,17 +890,19 @@ int32 FSequencerTimeSliderController::OnPaintSectionView( const FGeometry& Allot
 		
 		// Draw major tick lines in the section area
 		FDrawTickArgs DrawTickArgs;
-		DrawTickArgs.AllottedGeometry = AllottedGeometry;
-		DrawTickArgs.bMirrorLabels = false;
-		DrawTickArgs.bOnlyDrawMajorTicks = true;
-		DrawTickArgs.TickColor = TickColor;
-		DrawTickArgs.ClippingRect = MyClippingRect;
-		DrawTickArgs.DrawEffects = DrawEffects;
-		// Draw major ticks under sections
-		DrawTickArgs.StartLayer = LayerId-1;
-		// Draw the tick the entire height of the section area
-		DrawTickArgs.TickOffset = 0.0f;
-		DrawTickArgs.MajorTickHeight = AllottedGeometry.Size.Y;
+		{
+			DrawTickArgs.AllottedGeometry = AllottedGeometry;
+			DrawTickArgs.bMirrorLabels = false;
+			DrawTickArgs.bOnlyDrawMajorTicks = true;
+			DrawTickArgs.TickColor = TickColor;
+			DrawTickArgs.ClippingRect = MyClippingRect;
+			DrawTickArgs.DrawEffects = DrawEffects;
+			// Draw major ticks under sections
+			DrawTickArgs.StartLayer = LayerId-1;
+			// Draw the tick the entire height of the section area
+			DrawTickArgs.TickOffset = 0.0f;
+			DrawTickArgs.MajorTickHeight = AllottedGeometry.Size.Y;
+		}
 
 		DrawTicks( OutDrawElements, RangeToScreen, DrawTickArgs );
 	}
@@ -799,9 +911,11 @@ int32 FSequencerTimeSliderController::OnPaintSectionView( const FGeometry& Allot
 	{
 		// Draw a line for the scrub position
 		TArray<FVector2D> LinePoints;
-		LinePoints.AddUninitialized(2);
-		LinePoints[0] = FVector2D( 0.0f, 0.0f );
-		LinePoints[1] = FVector2D( 0.0f, FMath::RoundToFloat( AllottedGeometry.Size.Y ) );
+		{
+			LinePoints.AddUninitialized(2);
+			LinePoints[0] = FVector2D( 0.0f, 0.0f );
+			LinePoints[1] = FVector2D( 0.0f, FMath::RoundToFloat( AllottedGeometry.Size.Y ) );
+		}
 
 		FSlateDrawElement::MakeLines(
 			OutDrawElements,
@@ -812,7 +926,7 @@ int32 FSequencerTimeSliderController::OnPaintSectionView( const FGeometry& Allot
 			DrawEffects,
 			FLinearColor::White,
 			false
-			);
+		);
 	}
 
 	return LayerId;
@@ -923,6 +1037,19 @@ void FSequencerTimeSliderController::SetClampRange( float NewRangeMin, float New
 	}
 }
 
+void FSequencerTimeSliderController::SetPlayRange( float NewRangeMin, float NewRangeMax )
+{
+	const TRange<float> NewRange(NewRangeMin, NewRangeMax);
+
+	TimeSliderArgs.OnPlaybackRangeChanged.ExecuteIfBound(NewRange);
+
+	if( !TimeSliderArgs.PlaybackRange.IsBound() )
+	{	
+		// The  output is not bound to a delegate so we'll manage the value ourselves (no animation)
+		TimeSliderArgs.PlaybackRange.Set(NewRange);
+	}
+}
+
 bool FSequencerTimeSliderController::ZoomByDelta( float InDelta, float MousePositionFraction )
 {
 	TRange<float> LocalViewRange = TimeSliderArgs.ViewRange.Get().GetAnimationTarget();
@@ -961,7 +1088,8 @@ void FSequencerTimeSliderController::PanByDelta( float InDelta )
 	SetViewRange(NewViewOutputMin, NewViewOutputMax, EViewRangeInterpolation::Animated);
 }
 
-bool FSequencerTimeSliderController::HitTestPlaybackStart(const FScrubRangeToScreen& RangeToScreen, const TRange<float>& PlaybackRange, float LocalHitPositionX, float ScrubPosition) const
+
+bool FSequencerTimeSliderController::HitTestScrubberStart(const FScrubRangeToScreen& RangeToScreen, const TRange<float>& PlaybackRange, float LocalHitPositionX, float ScrubPosition) const
 {
 	static float BrushSizeInStateUnits = 6.f, DragToleranceSlateUnits = 2.f;
 	float LocalPlaybackStartPos = RangeToScreen.InputToLocalX(PlaybackRange.GetLowerBoundValue());
@@ -977,7 +1105,7 @@ bool FSequencerTimeSliderController::HitTestPlaybackStart(const FScrubRangeToScr
 		LocalHitPositionX <= LocalPlaybackStartPos + BrushSizeInStateUnits + DragToleranceSlateUnits;
 }
 
-bool FSequencerTimeSliderController::HitTestPlaybackEnd(const FScrubRangeToScreen& RangeToScreen, const TRange<float>& PlaybackRange, float LocalHitPositionX, float ScrubPosition) const
+bool FSequencerTimeSliderController::HitTestScrubberEnd(const FScrubRangeToScreen& RangeToScreen, const TRange<float>& PlaybackRange, float LocalHitPositionX, float ScrubPosition) const
 {
 	static float BrushSizeInStateUnits = 6.f, DragToleranceSlateUnits = 2.f;
 	float LocalPlaybackEndPos = RangeToScreen.InputToLocalX(PlaybackRange.GetUpperBoundValue());
diff --git a/Engine/Source/Editor/Sequencer/Private/SequencerTimeSliderController.h b/Engine/Source/Editor/Sequencer/Private/SequencerTimeSliderController.h
--- a/Engine/Source/Editor/Sequencer/Private/SequencerTimeSliderController.h
+++ b/Engine/Source/Editor/Sequencer/Private/SequencerTimeSliderController.h
@@ -61,6 +61,9 @@ public:
 	/** Get the current clamp range for this controller */
 	virtual FAnimatedRange GetClampRange() const override { return TimeSliderArgs.ClampRange.Get(); }
 
+	/** Get the current play range for this controller */
+	virtual TRange<float> GetPlayRange() const override { return TimeSliderArgs.PlaybackRange.Get(TRange<float>()); }
+
 	/** Convert time to frame */
 	virtual int32 TimeToFrame(float Time) const override;
 
@@ -92,6 +95,14 @@ public:
 	 */
 	virtual void SetClampRange( float NewRangeMin, float NewRangeMax ) override;
 
+	/**
+	 * Set a new playback range based on a min, max
+	 * 
+	 * @param NewRangeMin		The new lower bound of the playback range
+	 * @param NewRangeMax		The new upper bound of the playback range
+	 */
+	virtual void SetPlayRange( float NewRangeMin, float NewRangeMax ) override;
+
 	/**
 	 * Zoom the range by a given delta.
 	 * 
@@ -122,7 +133,7 @@ private:
 	void CommitScrubPosition( float NewValue, bool bIsScrubbing );
 
 	/**
-	 * Draws time tick marks
+	 * Draw time tick marks
 	 *
 	 * @param OutDrawElements	List to add draw elements to
 	 * @param RangeToScreen		Time range to screen space converter
@@ -131,21 +142,30 @@ private:
 	void DrawTicks( FSlateWindowElementList& OutDrawElements, const struct FScrubRangeToScreen& RangeToScreen, struct FDrawTickArgs& InArgs ) const;
 
 	/**
-	 * Draws the playback range
+	 * Draw the in/out selection range.
+	 *
+	 * @return The new layer ID.
+	 */
+	int32 DrawInOutRange(const FGeometry& AllottedGeometry, const FSlateRect& MyClippingRect, FSlateWindowElementList& OutDrawElements, int32 LayerId, const FScrubRangeToScreen& RangeToScreen, const FPaintPlaybackRangeArgs& Args) const;
+
+	/**
+	 * Draw the playback range.
+	 *
 	 * @return the new layer ID
 	 */
 	int32 DrawPlaybackRange(const FGeometry& AllottedGeometry, const FSlateRect& MyClippingRect, FSlateWindowElementList& OutDrawElements, int32 LayerId, const FScrubRangeToScreen& RangeToScreen, const FPaintPlaybackRangeArgs& Args) const;
 
 private:
+
 	/**
 	 * Hit test the lower bound of the playback range
 	 */
-	bool HitTestPlaybackStart(const FScrubRangeToScreen& RangeToScreen, const TRange<float>& PlaybackRange, float LocalHitPositionX, float ScrubPosition) const;
+	bool HitTestScrubberStart(const FScrubRangeToScreen& RangeToScreen, const TRange<float>& PlaybackRange, float LocalHitPositionX, float ScrubPosition) const;
 
 	/**
 	 * Hit test the upper bound of the playback range
 	 */
-	bool HitTestPlaybackEnd(const FScrubRangeToScreen& RangeToScreen, const TRange<float>& PlaybackRange, float LocalHitPositionX, float ScrubPosition) const;
+	bool HitTestScrubberEnd(const FScrubRangeToScreen& RangeToScreen, const TRange<float>& PlaybackRange, float LocalHitPositionX, float ScrubPosition) const;
 
 	void SetPlaybackRangeStart(float NewStart);
 	void SetPlaybackRangeEnd(float NewEnd);
@@ -169,6 +189,8 @@ private:
 		DRAG_SETTING_RANGE,
 		DRAG_START_RANGE,
 		DRAG_END_RANGE,
+		DRAG_IN_RANGE,
+		DRAG_OUT_RANGE,
 		DRAG_NONE
 	};
 	DragType MouseDragType;
diff --git a/Engine/Source/Editor/Sequencer/Private/SequencerUtilities.cpp b/Engine/Source/Editor/Sequencer/Private/SequencerUtilities.cpp
--- a/Engine/Source/Editor/Sequencer/Private/SequencerUtilities.cpp
+++ b/Engine/Source/Editor/Sequencer/Private/SequencerUtilities.cpp
@@ -63,3 +63,23 @@ TSharedRef<SWidget> FSequencerUtilities::MakeAddButton(FText HoverText, FOnGetCo
 	return ComboButton;
 }
 
+FName FSequencerUtilities::GetUniqueName( FName CandidateName, const TArray<FName>& ExistingNames )
+{
+	FString CandidateNameString = CandidateName.ToString();
+	FString BaseNameString = CandidateNameString;
+	if ( CandidateNameString.Len() >= 3 && CandidateNameString.Right(3).IsNumeric() )
+	{
+		BaseNameString = CandidateNameString.Left( CandidateNameString.Len() - 3 );
+	}
+
+	FName UniqueName = FName(*BaseNameString);
+	int32 NameIndex = 1;
+	while ( ExistingNames.Contains( UniqueName ) )
+	{
+		UniqueName = FName( *FString::Printf(TEXT("%s%03i"), *BaseNameString, NameIndex ) );
+		NameIndex++;
+	}
+
+	return UniqueName;
+}
+
diff --git a/Engine/Source/Editor/Sequencer/Private/Tools/EditToolDragOperations.cpp b/Engine/Source/Editor/Sequencer/Private/Tools/EditToolDragOperations.cpp
--- a/Engine/Source/Editor/Sequencer/Private/Tools/EditToolDragOperations.cpp
+++ b/Engine/Source/Editor/Sequencer/Private/Tools/EditToolDragOperations.cpp
@@ -160,10 +160,28 @@ void FResizeSection::OnBeginDrag(const FPointerEvent& MouseEvent, FVector2D Loca
 	DraggedKeyHandles.Empty();
 	SectionInitTimes.Empty();
 
+	bool bIsDilating = MouseEvent.IsControlDown();
+
 	for (auto& Handle : Sections)
 	{
 		UMovieSceneSection* Section = Handle.GetSectionObject();
 
+		for (auto SequencerSection : Handle.TrackNode->GetSections())
+		{
+			if (SequencerSection->GetSectionObject() == Section)
+			{
+				if (bIsDilating)
+				{
+					SequencerSection->BeginDilateSection();
+				}
+				else
+				{
+					SequencerSection->BeginResizeSection();
+				}
+				break;
+			}
+		}
+
 		Section->GetKeyHandles(DraggedKeyHandles);
 		SectionInitTimes.Add(Section, bDraggingByEnd ? Section->GetEndTime() : Section->GetStartTime());
 	}
@@ -217,48 +235,58 @@ void FResizeSection::OnDrag(const FPointerEvent& MouseEvent, FVector2D LocalMous
 	for (auto& Handle : Sections)
 	{
 		UMovieSceneSection* Section = Handle.GetSectionObject();
-		float NewTime = SectionInitTimes[Section] + DeltaTime;
 
-		if( bDraggingByEnd )
+		// Find the corresponding sequencer section to this movie scene section
+		for (auto SequencerSection : Handle.TrackNode->GetSections())
 		{
-			// Dragging the end of a section
-			// Ensure we aren't shrinking past the start time
-			NewTime = FMath::Max( NewTime, Section->GetStartTime() );
-
-			if (bIsDilating)
+			if (SequencerSection->GetSectionObject() == Section)
 			{
-				float NewSize = NewTime - Section->GetStartTime();
-				float DilationFactor = NewSize / Section->GetTimeSize();
-				Section->DilateSection(DilationFactor, Section->GetStartTime(), DraggedKeyHandles);
-			}
-			else
-			{
-				Section->SetEndTime( NewTime );
-			}
-		}
-		else
-		{
-			// Dragging the start of a section
-			// Ensure we arent expanding past the end time
-			NewTime = FMath::Min( NewTime, Section->GetEndTime() );
+				float NewTime = SectionInitTimes[Section] + DeltaTime;
 
-			if (bIsDilating)
-			{
-				float NewSize = Section->GetEndTime() - NewTime;
-				float DilationFactor = NewSize / Section->GetTimeSize();
-				Section->DilateSection(DilationFactor, Section->GetEndTime(), DraggedKeyHandles);
-			}
-			else
-			{
-				Section->SetStartTime( NewTime );
-			}
-		}
+				if( bDraggingByEnd )
+				{
+					// Dragging the end of a section
+					// Ensure we aren't shrinking past the start time
+					NewTime = FMath::Max( NewTime, Section->GetStartTime() );
+
+					if (bIsDilating)
+					{
+						float NewSize = NewTime - Section->GetStartTime();
+						float DilationFactor = NewSize / Section->GetTimeSize();
+						SequencerSection->DilateSection(DilationFactor, Section->GetStartTime(), DraggedKeyHandles);
+					}
+					else
+					{
+						SequencerSection->ResizeSection( SSRM_TrailingEdge, NewTime );
+					}
+				}
+				else
+				{
+					// Dragging the start of a section
+					// Ensure we arent expanding past the end time
+					NewTime = FMath::Min( NewTime, Section->GetEndTime() );
+
+					if (bIsDilating)
+					{
+						float NewSize = Section->GetEndTime() - NewTime;
+						float DilationFactor = NewSize / Section->GetTimeSize();
+						SequencerSection->DilateSection(DilationFactor, Section->GetEndTime(), DraggedKeyHandles);
+					}
+					else
+					{
+						SequencerSection->ResizeSection( SSRM_LeadingEdge, NewTime );
+					}
+				}
 
-		UMovieSceneTrack* OuterTrack = Section->GetTypedOuter<UMovieSceneTrack>();
-		if (OuterTrack)
-		{
-			OuterTrack->Modify();
-			OuterTrack->OnSectionMoved(*Section);
+				UMovieSceneTrack* OuterTrack = Section->GetTypedOuter<UMovieSceneTrack>();
+				if (OuterTrack)
+				{
+					OuterTrack->Modify();
+					OuterTrack->OnSectionMoved(*Section);
+				}
+
+				break;
+			}
 		}
 	}
 
@@ -267,12 +295,24 @@ void FResizeSection::OnDrag(const FPointerEvent& MouseEvent, FVector2D LocalMous
 
 FMoveSection::FMoveSection( FSequencer& InSequencer, TArray<FSectionHandle> InSections )
 	: FEditToolDragOperation( InSequencer )
-	, Sections( MoveTemp(InSections) )
 {
+	// Only allow sections that are not infinite to be movable.
+	for (auto InSection : InSections)
+	{
+		if (!InSection.GetSectionObject()->IsInfinite())
+		{
+			Sections.Add(InSection);
+		}
+	}
 }
 
 void FMoveSection::OnBeginDrag(const FPointerEvent& MouseEvent, FVector2D LocalMousePos, const FVirtualTrackArea& VirtualTrackArea)
 {
+	if (!Sections.Num())
+	{
+		return;
+	}
+
 	BeginTransaction( Sections, NSLOCTEXT("Sequencer", "MoveSectionTransaction", "Move Section") );
 
 	// Construct a snap field of unselected sections
@@ -295,6 +335,11 @@ void FMoveSection::OnBeginDrag(const FPointerEvent& MouseEvent, FVector2D LocalM
 
 void FMoveSection::OnEndDrag(const FPointerEvent& MouseEvent, FVector2D LocalMousePos, const FVirtualTrackArea& VirtualTrackArea)
 {
+	if (!Sections.Num())
+	{
+		return;
+	}
+
 	DraggedKeyHandles.Empty();
 
 	for (auto& Handle : Sections)
@@ -316,6 +361,11 @@ void FMoveSection::OnEndDrag(const FPointerEvent& MouseEvent, FVector2D LocalMou
 
 void FMoveSection::OnDrag(const FPointerEvent& MouseEvent, FVector2D LocalMousePos, const FVirtualTrackArea& VirtualTrackArea)
 {
+	if (!Sections.Num())
+	{
+		return;
+	}
+
 	LocalMousePos.Y = FMath::Clamp(LocalMousePos.Y, 0.f, VirtualTrackArea.GetPhysicalSize().Y);
 
 	// Convert the current mouse position to a time
diff --git a/Engine/Source/Editor/Sequencer/Private/Tools/SequencerEditTool_Movement.cpp b/Engine/Source/Editor/Sequencer/Private/Tools/SequencerEditTool_Movement.cpp
--- a/Engine/Source/Editor/Sequencer/Private/Tools/SequencerEditTool_Movement.cpp
+++ b/Engine/Source/Editor/Sequencer/Private/Tools/SequencerEditTool_Movement.cpp
@@ -7,8 +7,10 @@
 #include "SequencerHotspots.h"
 #include "EditToolDragOperations.h"
 
+
 const FName FSequencerEditTool_Movement::Identifier = "Movement";
 
+
 FSequencerEditTool_Movement::FSequencerEditTool_Movement(FSequencer& InSequencer)
 	: FSequencerEditTool(InSequencer)
 	, SequencerWidget(StaticCastSharedRef<SSequencer>(InSequencer.GetSequencerWidget()))
@@ -57,7 +59,8 @@ FReply FSequencerEditTool_Movement::OnMouseMove(SWidget& OwnerWidget, const FGeo
 			// If we're already dragging, just update the drag op if it exists
 			if (DragOperation.IsValid())
 			{
-				DragOperation->OnDrag(MouseEvent, MyGeometry.AbsoluteToLocal(MouseEvent.GetScreenSpacePosition()), VirtualTrackArea);
+				DragPosition = MyGeometry.AbsoluteToLocal(MouseEvent.GetScreenSpacePosition());
+				DragOperation->OnDrag(MouseEvent, DragPosition, VirtualTrackArea);
 			}
 		}
 		// Otherwise we can attempt a new drag
@@ -79,11 +82,47 @@ FReply FSequencerEditTool_Movement::OnMouseMove(SWidget& OwnerWidget, const FGeo
 	return FReply::Unhandled();
 }
 
+bool FSequencerEditTool_Movement::GetHotspotTime(float& HotspotTime) const
+{
+	if (DelayedDrag->Hotspot.IsValid())
+	{
+		auto HotspotType = DelayedDrag->Hotspot->GetType();
+		if (HotspotType == ESequencerHotspot::Key)
+		{
+			FSequencerSelectedKey& ThisKey = StaticCastSharedPtr<FKeyHotspot>(DelayedDrag->Hotspot)->Key;
+			HotspotTime = ThisKey.KeyArea->GetKeyTime(ThisKey.KeyHandle.GetValue());
+			return true;
+		}
+		else if (HotspotType == ESequencerHotspot::Section || HotspotType == ESequencerHotspot::SectionResize_L)
+		{
+			FSectionHandle& ThisHandle = StaticCastSharedPtr<FSectionHotspot>(DelayedDrag->Hotspot)->Section;
+			UMovieSceneSection* ThisSection = ThisHandle.GetSectionObject();
+			if (ThisSection)
+			{
+				HotspotTime = ThisSection->GetStartTime();
+				return true;
+			}
+		}
+		else if (HotspotType == ESequencerHotspot::SectionResize_R)
+		{
+			FSectionHandle& ThisHandle = StaticCastSharedPtr<FSectionHotspot>(DelayedDrag->Hotspot)->Section;
+			UMovieSceneSection* ThisSection = ThisHandle.GetSectionObject();
+			if (ThisSection)
+			{
+				HotspotTime = ThisSection->GetEndTime();
+				return true;
+			}
+		}
+	}
+	return false;
+}
 
 TSharedPtr<ISequencerEditToolDragOperation> FSequencerEditTool_Movement::CreateDrag(const FPointerEvent& MouseEvent)
 {
 	FSequencerSelection& Selection = Sequencer.GetSelection();
 
+	GetHotspotTime(OriginalHotspotTime);
+
 	if (DelayedDrag->Hotspot.IsValid())
 	{
 		// Let the hotspot start a drag first, if it wants to
@@ -111,7 +150,9 @@ TSharedPtr<ISequencerEditToolDragOperation> FSequencerEditTool_Movement::CreateD
 			{
 				Selection.EmptySelectedKeys();
 				Selection.EmptySelectedSections();
+				Selection.EmptyNodesWithSelectedKeysOrSections();
 				Selection.AddToSelection(ThisSection);
+				SequencerHelpers::UpdateHoveredNodeFromSelectedSections(Sequencer);
 				SectionHandles.Add(ThisHandle);
 			}
 			return MakeShareable( new FMoveSection( Sequencer, SectionHandles ) );
@@ -126,7 +167,9 @@ TSharedPtr<ISequencerEditToolDragOperation> FSequencerEditTool_Movement::CreateD
 			{
 				Selection.EmptySelectedKeys();
 				Selection.EmptySelectedSections();
+				Selection.EmptyNodesWithSelectedKeysOrSections();
 				Selection.AddToSelection(ThisKey);
+				SequencerHelpers::UpdateHoveredNodeFromSelectedKeys(Sequencer);
 			}
 
 			// @todo sequencer: Make this a customizable UI command modifier?
@@ -134,10 +177,8 @@ TSharedPtr<ISequencerEditToolDragOperation> FSequencerEditTool_Movement::CreateD
 			{
 				return MakeShareable( new FDuplicateKeys( Sequencer, Selection.GetSelectedKeys() ) );
 			}
-			else
-			{
-				return MakeShareable( new FMoveKeys( Sequencer, Selection.GetSelectedKeys() ) );
-			}
+
+			return MakeShareable( new FMoveKeys( Sequencer, Selection.GetSelectedKeys() ) );
 		}
 	}
 	// If we're not dragging a hotspot, sections take precedence over keys
@@ -199,14 +240,128 @@ void FSequencerEditTool_Movement::OnMouseCaptureLost()
 }
 
 
-FCursorReply FSequencerEditTool_Movement::OnCursorQuery(const FGeometry& MyGeometry, const FPointerEvent& CursorEvent) const
+int32 FSequencerEditTool_Movement::OnPaint(const FGeometry& AllottedGeometry, const FSlateRect& MyClippingRect, FSlateWindowElementList& OutDrawElements, int32 LayerId) const
 {
-	TSharedPtr<ISequencerHotspot> Hotspot = Sequencer.GetHotspot();
-	if (DelayedDrag.IsSet())
+	if (DelayedDrag.IsSet() && DelayedDrag->IsDragging())
 	{
-		Hotspot = DelayedDrag->Hotspot;
+		const TSharedPtr<ISequencerHotspot>& Hotspot = DelayedDrag->Hotspot;
+
+		if (Hotspot.IsValid())
+		{
+			float CurrentTime = 0.0f;
+
+			if (GetHotspotTime(CurrentTime))
+			{
+				const FSlateFontInfo SmallLayoutFont(FPaths::EngineContentDir() / TEXT("Slate/Fonts/Roboto-Bold.ttf"), 10);			
+				const TSharedRef< FSlateFontMeasure > FontMeasureService = FSlateApplication::Get().GetRenderer()->GetFontMeasureService();
+				const FLinearColor DrawColor = FEditorStyle::GetSlateColor("SelectionColor").GetColor(FWidgetStyle());
+				const FVector2D BoxPadding = FVector2D(4.0f, 2.0f);
+				const float MousePadding = 20.0f;
+
+				// calculate draw position
+				const FVirtualTrackArea VirtualTrackArea = SequencerWidget.Pin()->GetVirtualTrackArea();
+				const float HorizontalDelta = DragPosition.X - DelayedDrag->GetInitialPosition().X;
+				const float InitialY = DelayedDrag->GetInitialPosition().Y;
+
+				const FVector2D OldPos = FVector2D(VirtualTrackArea.TimeToPixel(OriginalHotspotTime), InitialY);
+				const FVector2D NewPos = FVector2D(VirtualTrackArea.TimeToPixel(CurrentTime), InitialY);
+
+				TArray<FVector2D> LinePoints;
+				{
+					LinePoints.AddUninitialized(2);
+					LinePoints[0] = FVector2D(0.0f, 0.0f);
+					LinePoints[1] = FVector2D(0.0f, VirtualTrackArea.GetPhysicalSize().Y);
+				}
+
+				// draw old position vertical
+				FSlateDrawElement::MakeLines(
+					OutDrawElements,
+					LayerId + 1,
+					AllottedGeometry.ToPaintGeometry(FVector2D(OldPos.X, 0.0f), FVector2D(1.0f, 1.0f)),
+					LinePoints,
+					MyClippingRect,
+					ESlateDrawEffect::None,
+					FLinearColor::White.CopyWithNewOpacity(0.5f),
+					false
+				);
+
+				// draw new position vertical
+				FSlateDrawElement::MakeLines(
+					OutDrawElements,
+					LayerId + 1,
+					AllottedGeometry.ToPaintGeometry(FVector2D(NewPos.X, 0.0f), FVector2D(1.0f, 1.0f)),
+					LinePoints,
+					MyClippingRect,
+					ESlateDrawEffect::None,
+					DrawColor,
+					false
+				);
+
+				// draw time string
+				const FString TimeString = TimeToString(CurrentTime, false);
+				const FVector2D TimeStringSize = FontMeasureService->Measure(TimeString, SmallLayoutFont);
+				const FVector2D TimePos = FVector2D(NewPos.X - MousePadding - TimeStringSize.X, NewPos.Y - 0.5f * TimeStringSize.Y);
+
+				FSlateDrawElement::MakeBox( 
+					OutDrawElements,
+					LayerId + 2, 
+					AllottedGeometry.ToPaintGeometry(TimePos - BoxPadding, TimeStringSize + 2.0f * BoxPadding),
+					FEditorStyle::GetBrush("WhiteBrush"),
+					MyClippingRect, 
+					ESlateDrawEffect::None, 
+					FLinearColor::Black.CopyWithNewOpacity(0.5f)
+				);
+
+				FSlateDrawElement::MakeText(
+					OutDrawElements,
+					LayerId + 3,
+					AllottedGeometry.ToPaintGeometry(TimePos, TimeStringSize),
+					TimeString,
+					SmallLayoutFont,
+					MyClippingRect,
+					ESlateDrawEffect::None,
+					DrawColor
+				);
+
+				// draw offset string
+				const FString OffsetString = TimeToString(CurrentTime - OriginalHotspotTime, true);
+				const FVector2D OffsetStringSize = FontMeasureService->Measure(OffsetString, SmallLayoutFont);
+				const FVector2D OffsetPos = FVector2D(NewPos.X + MousePadding, NewPos.Y - 0.5f * OffsetStringSize.Y);
+
+				FSlateDrawElement::MakeBox( 
+					OutDrawElements,
+					LayerId + 2, 
+					AllottedGeometry.ToPaintGeometry(OffsetPos - BoxPadding, OffsetStringSize + 2.0f * BoxPadding),
+					FEditorStyle::GetBrush("WhiteBrush"),
+					MyClippingRect, 
+					ESlateDrawEffect::None, 
+					FLinearColor::Black.CopyWithNewOpacity(0.5f)
+				);
+
+				FSlateDrawElement::MakeText(
+					OutDrawElements,
+					LayerId + 3,
+					AllottedGeometry.ToPaintGeometry(OffsetPos, TimeStringSize),
+					OffsetString,
+					SmallLayoutFont,
+					MyClippingRect,
+					ESlateDrawEffect::None,
+					DrawColor
+				);
+			}
+		}
 	}
 
+	return LayerId;
+}
+
+
+FCursorReply FSequencerEditTool_Movement::OnCursorQuery(const FGeometry& MyGeometry, const FPointerEvent& CursorEvent) const
+{
+	TSharedPtr<ISequencerHotspot> Hotspot = DelayedDrag.IsSet()
+		? DelayedDrag->Hotspot
+		: Sequencer.GetHotspot();
+
 	if (Hotspot.IsValid())
 	{
 		FCursorReply Reply = Hotspot->GetCursor();
@@ -219,12 +374,38 @@ FCursorReply FSequencerEditTool_Movement::OnCursorQuery(const FGeometry& MyGeome
 	return FCursorReply::Cursor(EMouseCursor::CardinalCross);
 }
 
+
 FName FSequencerEditTool_Movement::GetIdentifier() const
 {
 	return Identifier;
 }
 
+
 bool FSequencerEditTool_Movement::CanDeactivate() const
 {
 	return !DelayedDrag.IsSet();
+}
+
+
+FString FSequencerEditTool_Movement::TimeToString(float Time, bool IsDelta) const
+{
+	USequencerSettings* Settings = Sequencer.GetSettings();
+
+	if ((Settings != nullptr) && Settings->GetShowFrameNumbers())
+	{
+		if (SequencerSnapValues::IsTimeSnapIntervalFrameRate(Settings->GetTimeSnapInterval()))
+		{
+			const float FrameRate = 1.0f / Settings->GetTimeSnapInterval();
+			const int32 Frame = SequencerHelpers::TimeToFrame(Time, FrameRate);
+
+			return FString::Printf(IsDelta ? TEXT("[%+d]") : TEXT("%d"), Frame);
+		}
+	}
+
+	return FString::Printf(IsDelta ? TEXT("[%+.3f]") : TEXT("%.3f"), Time);
+}
+
+const ISequencerHotspot* FSequencerEditTool_Movement::GetDragHotspot() const
+{
+	return DelayedDrag.IsSet() ? DelayedDrag->Hotspot.Get() : nullptr;
 }
\ No newline at end of file
diff --git a/Engine/Source/Editor/Sequencer/Private/Tools/SequencerEditTool_Movement.h b/Engine/Source/Editor/Sequencer/Private/Tools/SequencerEditTool_Movement.h
--- a/Engine/Source/Editor/Sequencer/Private/Tools/SequencerEditTool_Movement.h
+++ b/Engine/Source/Editor/Sequencer/Private/Tools/SequencerEditTool_Movement.h
@@ -28,14 +28,22 @@ public:
 	virtual FReply OnMouseButtonUp(SWidget& OwnerWidget, const FGeometry& MyGeometry, const FPointerEvent& MouseEvent) override;
 	virtual FReply OnMouseMove(SWidget& OwnerWidget, const FGeometry& MyGeometry, const FPointerEvent& MouseEvent) override;
 	virtual void OnMouseCaptureLost() override;
+	virtual int32 OnPaint(const FGeometry& AllottedGeometry, const FSlateRect& MyClippingRect, FSlateWindowElementList& OutDrawElements, int32 LayerId) const override;
 	virtual FCursorReply OnCursorQuery(const FGeometry& MyGeometry, const FPointerEvent& CursorEvent) const override;
 	virtual FName GetIdentifier() const override;
 	virtual bool CanDeactivate() const override;
+	virtual const ISequencerHotspot* GetDragHotspot() const override;
+
+protected:
+
+	FString TimeToString(float Time, bool IsDelta) const;
 
 private:
 
 	TSharedPtr<ISequencerEditToolDragOperation> CreateDrag(const FPointerEvent& MouseEvent);
 
+	bool GetHotspotTime(float& HotspotTime) const;
+
 	/** Sequencer widget */
 	TWeakPtr<SSequencer> SequencerWidget;
 
@@ -54,4 +62,10 @@ private:
 
 	/** Current drag operation if any */
 	TSharedPtr<ISequencerEditToolDragOperation> DragOperation;
+
+	/** Current local position the mouse is dragged to. */
+	FVector2D DragPosition;
+
+	/** The hotspot's time before dragging started. */
+	float OriginalHotspotTime;
 };
diff --git a/Engine/Source/Editor/Sequencer/Private/Tools/SequencerEditTool_Selection.cpp b/Engine/Source/Editor/Sequencer/Private/Tools/SequencerEditTool_Selection.cpp
--- a/Engine/Source/Editor/Sequencer/Private/Tools/SequencerEditTool_Selection.cpp
+++ b/Engine/Source/Editor/Sequencer/Private/Tools/SequencerEditTool_Selection.cpp
@@ -20,34 +20,48 @@ struct FSelectionPreviewVisitor
 		, SetStateTo(InSetStateTo)
 	{}
 
-	virtual void VisitKey(FKeyHandle KeyHandle, float KeyTime, const TSharedPtr<IKeyArea>& KeyArea, UMovieSceneSection* Section) const override
+	virtual void VisitKey(FKeyHandle KeyHandle, float KeyTime, const TSharedPtr<IKeyArea>& KeyArea, UMovieSceneSection* Section, TSharedRef<FSequencerDisplayNode> Node) const override
 	{
 		FSequencerSelectedKey Key(*Section, KeyArea, KeyHandle);
 
-		bool bResetSectionSelection = false;
-
-		// If we're trying to change this key's selection state, we reset the selection state of the section
+		// If we're trying to change this key's selection state, we go into 'key selection mode', thus we reset the selection state of any nodes that weren't selected by keys
 		bool bKeyIsSelected = ExistingSelection.IsSelected(Key);
 		if ( (bKeyIsSelected && SetStateTo == ESelectionPreviewState::NotSelected) ||
 			(!bKeyIsSelected && SetStateTo == ESelectionPreviewState::Selected) )
 		{
-			SelectionPreview.SetSelectionState(Section, ESelectionPreviewState::Undefined);
-			SectionsWithKeysSelected.Add(Section);
+			// Clear selected nodes
+			for (TSharedRef<FSequencerDisplayNode> SelectedNode : NodesSelectedBySections)
+			{
+				if (!NodesSelectedByKeys.Contains(SelectedNode))
+				{
+					SelectionPreview.SetSelectionState(SelectedNode, ESelectionPreviewState::Undefined);
+				}
+			}
+			
+			// Clear selected sections
+			SelectionPreview.EmptyDefinedSectionStates();
 		}
 
 		SelectionPreview.SetSelectionState(Key, SetStateTo);
+		SelectionPreview.SetSelectionState(Node, SetStateTo);
+		NodesSelectedByKeys.Add(Node);
 	}
 
-	virtual void VisitSection(UMovieSceneSection* Section) const
+	virtual void VisitSection(UMovieSceneSection* Section, TSharedRef<FSequencerDisplayNode> Node) const
 	{
-		if (!SectionsWithKeysSelected.Contains(Section))
+		// Never select a combination of sections and keys
+		// Never allow infinite sections to be selected (they're only selectable through right click)
+		if (SelectionPreview.GetDefinedKeyStates().Num() == 0 && !Section->IsInfinite())
 		{
 			SelectionPreview.SetSelectionState(Section, SetStateTo);
+			SelectionPreview.SetSelectionState(Node, SetStateTo);
+			NodesSelectedBySections.Add(Node);
 		}
 	}
 
 private:
-	mutable TSet<UMovieSceneSection*> SectionsWithKeysSelected;
+	mutable TSet<TSharedRef<FSequencerDisplayNode>> NodesSelectedBySections;
+	mutable TSet<TSharedRef<FSequencerDisplayNode>> NodesSelectedByKeys;
 
 	FSequencerSelectionPreview& SelectionPreview;
 	FSequencerSelection& ExistingSelection;
@@ -195,6 +209,7 @@ public:
 		for (const auto& Pair : SelectionPreview.GetDefinedSectionStates())
 		{
 			UMovieSceneSection* Section = Pair.Key.Get();
+
 			if (Pair.Value == ESelectionPreviewState::Selected)
 			{
 				// Select it in the main selection
@@ -206,8 +221,21 @@ public:
 			}
 		}
 
+		for (const auto& Pair : SelectionPreview.GetDefinedOutlinerNodeStates())
+		{
+			if (Pair.Value == ESelectionPreviewState::Selected)
+			{
+				Selection.AddToNodesWithSelectedKeysOrSections(Pair.Key);
+			}
+			else
+			{
+				Selection.RemoveFromNodesWithSelectedKeysOrSections(Pair.Key);
+			}
+		}
+
 		// We're done with this now
 		SelectionPreview.Empty();
+		SequencerHelpers::ValidateNodesWithSelectedKeysOrSections(Sequencer);
 	}
 
 	virtual int32 OnPaint(const FGeometry& AllottedGeometry, const FSlateRect& MyClippingRect, FSlateWindowElementList& OutDrawElements, int32 LayerId) const
@@ -432,6 +460,11 @@ bool FSequencerEditTool_Selection::CanDeactivate() const
 	return !DelayedDrag.IsSet();
 }
 
+const ISequencerHotspot* FSequencerEditTool_Selection::GetDragHotspot() const
+{
+	return DelayedDrag.IsSet() ? DelayedDrag->Hotspot.Get() : nullptr;
+}
+
 void FSequencerEditTool_Selection::UpdateCursor(const FGeometry& MyGeometry, const FPointerEvent& MouseEvent)
 {
 	MousePosition = MyGeometry.AbsoluteToLocal(MouseEvent.GetScreenSpacePosition());
diff --git a/Engine/Source/Editor/Sequencer/Private/Tools/SequencerEditTool_Selection.h b/Engine/Source/Editor/Sequencer/Private/Tools/SequencerEditTool_Selection.h
--- a/Engine/Source/Editor/Sequencer/Private/Tools/SequencerEditTool_Selection.h
+++ b/Engine/Source/Editor/Sequencer/Private/Tools/SequencerEditTool_Selection.h
@@ -33,6 +33,7 @@ public:
 	virtual FCursorReply OnCursorQuery(const FGeometry& MyGeometry, const FPointerEvent& CursorEvent) const override;
 	virtual FName GetIdentifier() const override;
 	virtual bool CanDeactivate() const override;
+	virtual const ISequencerHotspot* GetDragHotspot() const override;
 	
 private:
 
diff --git a/Engine/Source/Editor/Sequencer/Private/Tools/SequencerEntityVisitor.cpp b/Engine/Source/Editor/Sequencer/Private/Tools/SequencerEntityVisitor.cpp
--- a/Engine/Source/Editor/Sequencer/Private/Tools/SequencerEntityVisitor.cpp
+++ b/Engine/Source/Editor/Sequencer/Private/Tools/SequencerEntityVisitor.cpp
@@ -18,25 +18,44 @@ FSequencerEntityRange::FSequencerEntityRange(FVector2D TopLeft, FVector2D Bottom
 {
 }
 
-bool FSequencerEntityRange::IntersectSection(const UMovieSceneSection* InSection) const
+bool FSequencerEntityRange::IntersectSection(const UMovieSceneSection* InSection, const TSharedRef<FSequencerTrackNode>& InTrackNode, int32 MaxRowIndex) const
 {
-	return InSection->GetStartTime() <= EndTime && InSection->GetEndTime() >= StartTime;
+	// Test horizontal bounds
+	if (!InSection->IsInfinite() && (InSection->GetStartTime() > EndTime || InSection->GetEndTime() < StartTime))
+	{
+		return false;
+	}
+	// Test vertical bounds
+	else if (MaxRowIndex > 0 && VerticalTop.IsSet())
+	{
+		const float NodeTop = InTrackNode->GetVirtualTop();
+		const float NodeBottom = InTrackNode->GetVirtualBottom();
+
+		const float RowHeight = (NodeBottom - NodeTop) / (MaxRowIndex + 1);
+		const float RowTop = NodeTop + InSection->GetRowIndex() * RowHeight;
+
+		return RowTop <= VerticalBottom.GetValue() && RowTop + RowHeight >= VerticalTop.GetValue();
+	}
+	else
+	{
+		return true;
+	}
 }
 
-bool FSequencerEntityRange::IntersectNode(const FSequencerDisplayNode& InNode) const
+bool FSequencerEntityRange::IntersectNode(TSharedRef<FSequencerDisplayNode> InNode) const
 {
 	if (VerticalTop.IsSet())
 	{
-		return InNode.GetVirtualTop() <= VerticalBottom.GetValue() && InNode.GetVirtualBottom() >= VerticalTop.GetValue();
+		return InNode->GetVirtualTop() <= VerticalBottom.GetValue() && InNode->GetVirtualBottom() >= VerticalTop.GetValue();
 	}
 	return true;
 }
 
-bool FSequencerEntityRange::IntersectKeyArea(const FSequencerDisplayNode& InNode, float VirtualKeyHeight) const
+bool FSequencerEntityRange::IntersectKeyArea(TSharedRef<FSequencerDisplayNode> InNode, float VirtualKeyHeight) const
 {
 	if (VerticalTop.IsSet())
 	{
-		const float NodeCenter = InNode.GetVirtualTop() + (InNode.GetVirtualBottom() - InNode.GetVirtualTop())/2;
+		const float NodeCenter = InNode->GetVirtualTop() + (InNode->GetVirtualBottom() - InNode->GetVirtualTop())/2;
 		return NodeCenter + VirtualKeyHeight/2 > VerticalTop.GetValue() && NodeCenter - VirtualKeyHeight/2 < VerticalBottom.GetValue();
 	}
 	return true;
@@ -53,108 +72,137 @@ void FSequencerEntityWalker::Traverse(const ISequencerEntityVisitor& Visitor, co
 	{
 		if (!Child->IsHidden())
 		{
-			HandleNode(Visitor, *Child);
+			HandleNode(Visitor, Child);
 		}
 	}
 }
 
-void FSequencerEntityWalker::HandleNode(const ISequencerEntityVisitor& Visitor, FSequencerDisplayNode& InNode)
+void FSequencerEntityWalker::HandleNode(const ISequencerEntityVisitor& Visitor, const TSharedRef<FSequencerDisplayNode>& InNode)
 {
-	if (InNode.GetType() == ESequencerNode::Track)
+	if (InNode->GetType() == ESequencerNode::Track)
 	{
-		HandleNode(Visitor, InNode, static_cast<FSequencerTrackNode&>(InNode).GetSections());
+		HandleTrackNode(Visitor, StaticCastSharedRef<FSequencerTrackNode>(InNode));
 	}
 
-	if (InNode.IsExpanded())
+	if (InNode->IsExpanded())
 	{
-		for (auto& Child : InNode.GetChildNodes())
+		for (auto& Child : InNode->GetChildNodes())
 		{
 			if (!Child->IsHidden())
 			{
-				HandleNode(Visitor, *Child);
+				HandleNode(Visitor, Child);
 			}
 		}
 	}
 }
 
-void FSequencerEntityWalker::HandleNode(const ISequencerEntityVisitor& Visitor, FSequencerDisplayNode& InNode, TArray<TSharedRef<ISequencerSection>> InSections)
+void FSequencerEntityWalker::HandleTrackNode(const ISequencerEntityVisitor& Visitor, const TSharedRef<FSequencerTrackNode>& InTrackNode)
 {
-	if (Range.IntersectNode(InNode))
+	TArray<TSharedRef<ISequencerSection>> Sections = InTrackNode->GetSections();
+
+	if (Range.IntersectNode(InTrackNode))
 	{
+		const int32 MaxRowIndex = InTrackNode->GetMaxRowIndex();
+
 		// Prune the selections to anything that is in the range, visiting if necessary
-		for (int32 SectionIndex = 0; SectionIndex < InSections.Num();)
+		for (int32 SectionIndex = 0; SectionIndex < Sections.Num();)
 		{
-			UMovieSceneSection* Section = InSections[SectionIndex]->GetSectionObject();
-			if (!Range.IntersectSection(Section))
+			UMovieSceneSection* Section = Sections[SectionIndex]->GetSectionObject();
+			if (!Range.IntersectSection(Section, InTrackNode, MaxRowIndex))
 			{
-				InSections.RemoveAtSwap(SectionIndex, 1, false);
+				Sections.RemoveAtSwap(SectionIndex, 1, false);
 				continue;
 			}
 
 			if (Visitor.CheckEntityMask(ESequencerEntity::Section))
 			{
-				Visitor.VisitSection(Section);
+				Visitor.VisitSection(Section, InTrackNode);
 			}
 
 			++SectionIndex;
 		}
 
-		bool bNodeHasKeyArea = false;
-		if (InNode.GetType() == ESequencerNode::KeyArea)
-		{
-			HandleKeyAreaNode(Visitor, static_cast<FSequencerSectionKeyAreaNode&>(InNode), InNode, InSections);
-			bNodeHasKeyArea = true;
-		}
-		else if (InNode.GetType() == ESequencerNode::Track)
+		HandleSingleNode(Visitor, InTrackNode, Sections);
+	}
+
+	if (InTrackNode->IsExpanded())
+	{
+		// Handle Children
+		for (auto& Child : InTrackNode->GetChildNodes())
 		{
-			TSharedPtr<FSequencerSectionKeyAreaNode> SectionKeyNode = static_cast<FSequencerTrackNode&>(InNode).GetTopLevelKeyNode();
-			if (SectionKeyNode.IsValid())
+			if (!Child->IsHidden())
 			{
-				HandleKeyAreaNode(Visitor, *SectionKeyNode, InNode, InSections);
-				bNodeHasKeyArea = true;
+				HandleChildNode(Visitor, Child, Sections);
 			}
 		}
+	}
+}
 
-		// As a fallback, we need to handle:
-		//  - Key groupings on collapsed parents
-		//  - Sections that have no key areas
-		const bool bIterateKeyGroupings = Visitor.CheckEntityMask(ESequencerEntity::Key) &&
-			!bNodeHasKeyArea &&
-			(!InNode.IsExpanded() || InNode.GetChildNodes().Num() == 0);
+void FSequencerEntityWalker::HandleChildNode(const ISequencerEntityVisitor& Visitor, const TSharedRef<FSequencerDisplayNode>& InNode, const TArray<TSharedRef<ISequencerSection>>& InSections)
+{
+	if (Range.IntersectNode(InNode))
+	{
+		HandleSingleNode(Visitor, InNode, InSections);
+	}
 
-		if (bIterateKeyGroupings)
+	if (InNode->IsExpanded())
+	{
+		// Handle Children
+		for (auto& Child : InNode->GetChildNodes())
 		{
-			for (int32 SectionIndex = 0; SectionIndex < InSections.Num(); ++SectionIndex)
+			if (!Child->IsHidden())
 			{
-				UMovieSceneSection* Section = InSections[SectionIndex]->GetSectionObject();
-
-				if (Visitor.CheckEntityMask(ESequencerEntity::Key))
-				{
-					// Only handle grouped keys if we actually have children
-					if (InNode.GetChildNodes().Num() != 0 && Range.IntersectKeyArea(InNode, VirtualKeySize.X))
-					{
-						TSharedRef<IKeyArea> KeyArea = InNode.UpdateKeyGrouping(SectionIndex);
-						HandleKeyArea(Visitor, KeyArea, Section);
-					}
-				}
+				HandleChildNode(Visitor, Child, InSections);
 			}
 		}
 	}
+}
 
-	if (InNode.IsExpanded())
+void FSequencerEntityWalker::HandleSingleNode(const ISequencerEntityVisitor& Visitor, const TSharedRef<FSequencerDisplayNode>& InNode, const TArray<TSharedRef<ISequencerSection>>& InSections)
+{
+	bool bNodeHasKeyArea = false;
+	if (InNode->GetType() == ESequencerNode::KeyArea)
 	{
-		// Handle Children
-		for (auto& Child : InNode.GetChildNodes())
+		HandleKeyAreaNode(Visitor, StaticCastSharedRef<FSequencerSectionKeyAreaNode>(InNode), InNode, InSections);
+		bNodeHasKeyArea = true;
+	}
+	else if (InNode->GetType() == ESequencerNode::Track)
+	{
+		TSharedPtr<FSequencerSectionKeyAreaNode> SectionKeyNode = StaticCastSharedRef<FSequencerTrackNode>(InNode)->GetTopLevelKeyNode();
+		if (SectionKeyNode.IsValid())
 		{
-			if (!Child->IsHidden())
+			HandleKeyAreaNode(Visitor, SectionKeyNode.ToSharedRef(), InNode, InSections);
+			bNodeHasKeyArea = true;
+		}
+	}
+
+	// As a fallback, we need to handle:
+	//  - Key groupings on collapsed parents
+	//  - Sections that have no key areas
+	const bool bIterateKeyGroupings = Visitor.CheckEntityMask(ESequencerEntity::Key) &&
+		!bNodeHasKeyArea &&
+		(!InNode->IsExpanded() || InNode->GetChildNodes().Num() == 0);
+
+	if (bIterateKeyGroupings)
+	{
+		for (int32 SectionIndex = 0; SectionIndex < InSections.Num(); ++SectionIndex)
+		{
+			UMovieSceneSection* Section = InSections[SectionIndex]->GetSectionObject();
+
+			if (Visitor.CheckEntityMask(ESequencerEntity::Key))
 			{
-				HandleNode(Visitor, *Child, InSections);
+				// Only handle grouped keys if we actually have children
+				if (InNode->GetChildNodes().Num() != 0 && Range.IntersectKeyArea(InNode, VirtualKeySize.X))
+				{
+					TSharedRef<IKeyArea> KeyArea = InNode->UpdateKeyGrouping(SectionIndex);
+					HandleKeyArea(Visitor, KeyArea, Section, InNode);
+				}
 			}
 		}
 	}
 }
 
-void FSequencerEntityWalker::HandleKeyAreaNode(const ISequencerEntityVisitor& Visitor, FSequencerSectionKeyAreaNode& InKeyAreaNode, FSequencerDisplayNode& InOwnerNode, const TArray<TSharedRef<ISequencerSection>>& InSections)
+void FSequencerEntityWalker::HandleKeyAreaNode(const ISequencerEntityVisitor& Visitor, const TSharedRef<FSequencerSectionKeyAreaNode>& InKeyAreaNode, const TSharedRef<FSequencerDisplayNode>& InOwnerNode, const TArray<TSharedRef<ISequencerSection>>& InSections)
 {
 	for( int32 SectionIndex = 0; SectionIndex < InSections.Num(); ++SectionIndex )
 	{
@@ -163,14 +211,14 @@ void FSequencerEntityWalker::HandleKeyAreaNode(const ISequencerEntityVisitor& Vi
 		{
 			if (Range.IntersectKeyArea(InOwnerNode, VirtualKeySize.X))
 			{
-				TSharedRef<IKeyArea> KeyArea = InKeyAreaNode.GetKeyArea(SectionIndex);
-				HandleKeyArea(Visitor, KeyArea, InSections[SectionIndex]->GetSectionObject());
+				TSharedRef<IKeyArea> KeyArea = InKeyAreaNode->GetKeyArea(SectionIndex);
+				HandleKeyArea(Visitor, KeyArea, InSections[SectionIndex]->GetSectionObject(), InOwnerNode);
 			}
 		}
 	}
 }
 
-void FSequencerEntityWalker::HandleKeyArea(const ISequencerEntityVisitor& Visitor, const TSharedPtr<IKeyArea>& KeyArea, UMovieSceneSection* Section)
+void FSequencerEntityWalker::HandleKeyArea(const ISequencerEntityVisitor& Visitor, const TSharedPtr<IKeyArea>& KeyArea, UMovieSceneSection* Section, const TSharedRef<FSequencerDisplayNode>& InNode)
 {
 	for (FKeyHandle KeyHandle : KeyArea->GetUnsortedKeyHandles())
 	{
@@ -178,7 +226,7 @@ void FSequencerEntityWalker::HandleKeyArea(const ISequencerEntityVisitor& Visito
 		if (KeyPosition + VirtualKeySize.X/2 > Range.StartTime &&
 			KeyPosition - VirtualKeySize.X/2 < Range.EndTime)
 		{
-			Visitor.VisitKey(KeyHandle, KeyPosition, KeyArea, Section);
+			Visitor.VisitKey(KeyHandle, KeyPosition, KeyArea, Section, InNode);
 		}
 	}
 }
\ No newline at end of file
diff --git a/Engine/Source/Editor/Sequencer/Private/Tools/SequencerEntityVisitor.h b/Engine/Source/Editor/Sequencer/Private/Tools/SequencerEntityVisitor.h
--- a/Engine/Source/Editor/Sequencer/Private/Tools/SequencerEntityVisitor.h
+++ b/Engine/Source/Editor/Sequencer/Private/Tools/SequencerEntityVisitor.h
@@ -25,8 +25,8 @@ struct ISequencerEntityVisitor
 {
 	ISequencerEntityVisitor(uint32 InEntityMask = ESequencerEntity::Everything) : EntityMask(InEntityMask) {}
 
-	virtual void VisitKey(FKeyHandle KeyHandle, float KeyTime, const TSharedPtr<IKeyArea>& KeyArea, UMovieSceneSection* Section) const { }
-	virtual void VisitSection(UMovieSceneSection* Section) const { }
+	virtual void VisitKey(FKeyHandle KeyHandle, float KeyTime, const TSharedPtr<IKeyArea>& KeyArea, UMovieSceneSection* Section, TSharedRef<FSequencerDisplayNode>) const { }
+	virtual void VisitSection(UMovieSceneSection* Section, TSharedRef<FSequencerDisplayNode>) const { }
 	
 	/** Check if the specified type of entity is applicable to this visitor */
 	bool CheckEntityMask(ESequencerEntity::Type Type) const { return (EntityMask & Type) != 0; }
@@ -44,13 +44,13 @@ struct FSequencerEntityRange
 	FSequencerEntityRange(FVector2D TopLeft, FVector2D BottomRight);
 
 	/** Check whether the specified section intersects this range */
-	bool IntersectSection(const UMovieSceneSection* InSection) const;
+	bool IntersectSection(const UMovieSceneSection* InSection, const TSharedRef<FSequencerTrackNode>& InTrackNode, int32 MaxRowIndex) const;
 
 	/** Check whether the specified node intersects this range */
-	bool IntersectNode(const FSequencerDisplayNode& InNode) const;
+	bool IntersectNode(TSharedRef<FSequencerDisplayNode> InNode) const;
 
 	/** Check whether the specified node's key area intersects this range */
-	bool IntersectKeyArea(const FSequencerDisplayNode& InNode, float VirtualKeyHeight) const;
+	bool IntersectKeyArea(TSharedRef<FSequencerDisplayNode> InNode, float VirtualKeyHeight) const;
 
 	/** Start/end times */
 	float StartTime, EndTime;
@@ -71,13 +71,17 @@ struct FSequencerEntityWalker
 private:
 
 	/** Handle visitation of a particular node */
-	void HandleNode(const ISequencerEntityVisitor& Visitor, FSequencerDisplayNode& InNode);
-	/** Handle visitation of a particular node, along with a set of sections */
-	void HandleNode(const ISequencerEntityVisitor& Visitor, FSequencerDisplayNode& InNode, TArray<TSharedRef<ISequencerSection>> InSections);
+	void HandleNode(const ISequencerEntityVisitor& Visitor, const TSharedRef<FSequencerDisplayNode>& InNode);
+	/** Handle visitation of a track node, along with a set of sections */
+	void HandleTrackNode(const ISequencerEntityVisitor& Visitor, const TSharedRef<FSequencerTrackNode>& InNode);
+	/** Handle visitation of an arbitrary child node, with a filtered set of sections */
+	void HandleChildNode(const ISequencerEntityVisitor& Visitor, const TSharedRef<FSequencerDisplayNode>& InNode, const TArray<TSharedRef<ISequencerSection>>& InSections);
+	/** Handle a single node that is known to be within the range */
+	void HandleSingleNode(const ISequencerEntityVisitor& Visitor, const TSharedRef<FSequencerDisplayNode>& InNode, const TArray<TSharedRef<ISequencerSection>>& InSections);
 	/** Handle visitation of a key area node */
-	void HandleKeyAreaNode(const ISequencerEntityVisitor& Visitor, FSequencerSectionKeyAreaNode& InKeyAreaNode, FSequencerDisplayNode& InOwnerNode, const TArray<TSharedRef<ISequencerSection>>& InSections);
+	void HandleKeyAreaNode(const ISequencerEntityVisitor& Visitor, const TSharedRef<FSequencerSectionKeyAreaNode>& InKeyAreaNode, const TSharedRef<FSequencerDisplayNode>& InOwnerNode, const TArray<TSharedRef<ISequencerSection>>& InSections);
 	/** Handle visitation of a key area */
-	void HandleKeyArea(const ISequencerEntityVisitor& Visitor, const TSharedPtr<IKeyArea>& KeyArea, UMovieSceneSection* Section);
+	void HandleKeyArea(const ISequencerEntityVisitor& Visitor, const TSharedPtr<IKeyArea>& KeyArea, UMovieSceneSection* Section, const TSharedRef<FSequencerDisplayNode>& InNode);
 
 	/** The bounds of the range */
 	FSequencerEntityRange Range;
diff --git a/Engine/Source/Editor/Sequencer/Private/Tools/SequencerSnapField.cpp b/Engine/Source/Editor/Sequencer/Private/Tools/SequencerSnapField.cpp
--- a/Engine/Source/Editor/Sequencer/Private/Tools/SequencerSnapField.cpp
+++ b/Engine/Source/Editor/Sequencer/Private/Tools/SequencerSnapField.cpp
@@ -17,14 +17,14 @@ struct FSnapGridVisitor : ISequencerEntityVisitor
 		, Candidate(InCandidate)
 	{}
 
-	virtual void VisitKey(FKeyHandle KeyHandle, float KeyTime, const TSharedPtr<IKeyArea>& KeyArea, UMovieSceneSection* Section) const
+	virtual void VisitKey(FKeyHandle KeyHandle, float KeyTime, const TSharedPtr<IKeyArea>& KeyArea, UMovieSceneSection* Section, TSharedRef<FSequencerDisplayNode>) const
 	{
 		if (Candidate.IsKeyApplicable(KeyHandle, KeyArea, Section))
 		{
 			Snaps.Add(FSequencerSnapPoint{ FSequencerSnapPoint::Key, KeyTime });
 		}
 	}
-	virtual void VisitSection(UMovieSceneSection* Section) const
+	virtual void VisitSection(UMovieSceneSection* Section, TSharedRef<FSequencerDisplayNode> Node) const
 	{
 		if (Candidate.AreSectionBoundsApplicable(Section))
 		{
@@ -69,6 +69,9 @@ FSequencerSnapField::FSequencerSnapField(const ISequencer& InSequencer, ISequenc
 	Visitor.Snaps.Add(FSequencerSnapPoint{ FSequencerSnapPoint::PlaybackRange, PlaybackRange.GetLowerBoundValue() });
 	Visitor.Snaps.Add(FSequencerSnapPoint{ FSequencerSnapPoint::PlaybackRange, PlaybackRange.GetUpperBoundValue() });
 
+	// Add the current time as a potential snap candidate
+	Visitor.Snaps.Add(FSequencerSnapPoint{ FSequencerSnapPoint::CurrentTime, InSequencer.GetGlobalTime() });
+
 	// Sort
 	Visitor.Snaps.Sort([](const FSequencerSnapPoint& A, const FSequencerSnapPoint& B){
 		return A.Time < B.Time;
diff --git a/Engine/Source/Editor/Sequencer/Private/Tools/SequencerSnapField.h b/Engine/Source/Editor/Sequencer/Private/Tools/SequencerSnapField.h
--- a/Engine/Source/Editor/Sequencer/Private/Tools/SequencerSnapField.h
+++ b/Engine/Source/Editor/Sequencer/Private/Tools/SequencerSnapField.h
@@ -13,7 +13,7 @@ class UMovieSceneSection;
 /** Structure defining a point to snap to in the sequencer */
 struct FSequencerSnapPoint
 {
-	enum ESnapType { Key, SectionBounds, CustomSection, PlaybackRange };
+	enum ESnapType { Key, SectionBounds, CustomSection, PlaybackRange, CurrentTime };
 
 	/** The type of snap */
 	ESnapType Type;
diff --git a/Engine/Source/Editor/Sequencer/Public/ISequencer.h b/Engine/Source/Editor/Sequencer/Public/ISequencer.h
--- a/Engine/Source/Editor/Sequencer/Public/ISequencer.h
+++ b/Engine/Source/Editor/Sequencer/Public/ISequencer.h
@@ -15,6 +15,7 @@ class UAnimSequence;
 class UMovieScene;
 class UMovieSceneSection;
 class UMovieSceneSubSection;
+class ISequencerKeyCollection;
 enum class EMapChangeType : uint8;
 
 
@@ -35,6 +36,37 @@ enum class EAutoKeyMode
 };
 
 
+/**
+ * Enumerates types of UI Command bindings.
+ */
+enum class ESequencerCommandBindings
+{
+	/** Bindings that are used by Sequencer widgets only. */
+	Sequencer,
+
+	/** Bindings that are shared between Sequencer and non-Sequencer widgets (subset of Sequencer commands). */
+	Shared
+};
+
+
+/*
+ * Allowable snapping modes when setting global time
+ */
+enum ESnapTimeMode
+{
+	/** No snapping */
+	STM_None = 0x00000000,
+
+	/** Snap to the time interval. */
+	STM_Interval = 0x00000001,
+
+	/** Snap to keys. */
+	STM_Keys = 0x00000002,
+
+	/** All snapping */
+	STM_All = STM_Interval | STM_Keys
+};
+
 /**
  * Interface for sequencers.
  */
@@ -124,6 +156,12 @@ public:
 	/** Set default key interpolation */
 	virtual void SetKeyInterpolation(EMovieSceneKeyInterpolation) = 0;
 
+	/** @return Returns whether key sections are infinite by default when created */
+	virtual bool GetInfiniteKeyAreas() const = 0;
+
+	/** Set infinite key area default */
+	virtual void SetInfiniteKeyAreas(bool bInfiniteKeyAreas) = 0;
+
 	/** @return Returns whether sequencer is currently recording live data from simulated actors */
 	virtual bool IsRecordingLive() const = 0;
 
@@ -140,19 +178,19 @@ public:
 	 * @return Global time in seconds.
 	 * @see SetGlobalTime
 	 */
-	virtual float GetGlobalTime() = 0;
+	virtual float GetGlobalTime() const = 0;
 
 	/**
 	 * Sets the global position to the time.
 	 *
 	 * @param Time The global time to set.
-	 * @param bAllowSnappingToFrames Whether or not to allow snapping to frames if.
+	 * @param SnapTimeMode The type of time snapping allowed.
 	 * @see GetGlobalTime
 	 */
-	virtual void SetGlobalTime(float Time, const bool& bAllowSnappingToFrames = false) = 0;
+	virtual void SetGlobalTime(float Time, ESnapTimeMode SnapTimeMode = ESnapTimeMode::STM_None) = 0;
 
 	/** Set the global time directly, without performing any auto-scroll */
-	virtual void SetGlobalTimeDirectly(float Time, const bool& bAllowSnappingToFrames = false) = 0;
+	virtual void SetGlobalTimeDirectly(float Time, ESnapTimeMode SnapTimeMode = ESnapTimeMode::STM_None) = 0;
 
 	/** @return The current view range */
 	virtual FAnimatedRange GetViewRange() const
@@ -191,7 +229,25 @@ public:
 	 */ 
 	virtual bool IsPerspectiveViewportCameraCutEnabled() const { return true; }
 
-	DECLARE_EVENT_OneParam(ISequencer, FOnCameraCut, UObject*)
+	/*
+	 * Puts sequencer in a silent state (whereby it will not redraw viewports, or attempt to update external state besides the sequence itself)
+	 */
+	virtual void EnterSilentMode() = 0;
+
+	/*
+	 * Leaves a silent state (see above)
+	 */
+	virtual void ExitSilentMode() = 0;
+
+	/*
+	 * Checks whether we're in silent mode or not
+	 */
+	virtual bool IsInSilentMode() const = 0;
+
+	DECLARE_EVENT_TwoParams(ISequencer, FOnActorAddedToSequencer, AActor*, const FGuid)
+	virtual FOnActorAddedToSequencer& OnActorAddedToSequencer() = 0;
+
+	DECLARE_EVENT_TwoParams(ISequencer, FOnCameraCut, UObject*, bool)
 	virtual FOnCameraCut& OnCameraCut() = 0;
 
 	/**
@@ -223,6 +279,11 @@ public:
 
 	virtual void UpdatePlaybackRange() = 0;
 
+	/** Get all the keys for the current sequencer selection */
+	virtual void GetKeysFromSelection(TUniquePtr<ISequencerKeyCollection>& KeyCollection) = 0;
+
+	virtual float FindNearestKey(float NewScrubPosition) = 0;
+
 	virtual FSequencerSelection& GetSelection() = 0;
 	virtual FSequencerSelectionPreview& GetSelectionPreview() = 0;
 
@@ -238,8 +299,19 @@ public:
 	virtual TSharedRef<INumericTypeInterface<float>> GetZeroPadNumericTypeInterface() = 0;
 
 	/** @return the command bindings for this sequencer */
-	virtual TSharedPtr<FUICommandList> GetCommandBindings() const = 0;
+	virtual TSharedPtr<FUICommandList> GetCommandBindings(ESequencerCommandBindings Type = ESequencerCommandBindings::Sequencer) const = 0;
 
 	/** @return Returns a widget containing the sequencer's playback controls */
-	virtual TSharedRef<SWidget> MakeTransportControls() = 0;
+	virtual TSharedRef<SWidget> MakeTransportControls(bool bExtended) = 0;
+
+	/** Turn viewport transport controls on or off */
+	virtual void SetViewportTransportControlsVisibility(bool bVisible) = 0;
+
+	/**
+	 * Create a widget containing the spinboxes for setting the working and playback range
+	 * 
+	 * @param InnerContent		Inner content to be inserted to the middle of the widget (inbetween the playback range spinboxes)
+	 * @return the widget
+	 */
+	virtual TSharedRef<SWidget> MakeTimeRange(const TSharedRef<SWidget>& InnerContent, bool bShowWorkingRange, bool bShowViewRange, bool bShowPlaybackRange) = 0;
 };
diff --git a/Engine/Source/Editor/Sequencer/Public/ISequencerEditTool.h b/Engine/Source/Editor/Sequencer/Public/ISequencerEditTool.h
--- a/Engine/Source/Editor/Sequencer/Public/ISequencerEditTool.h
+++ b/Engine/Source/Editor/Sequencer/Public/ISequencerEditTool.h
@@ -68,6 +68,7 @@ public:
 	virtual void OnMouseLeave(SWidget& OwnerWidget, const FPointerEvent& MouseEvent) = 0;
 	virtual FName GetIdentifier() const = 0;
 	virtual bool CanDeactivate() const = 0;
+	virtual const ISequencerHotspot* GetDragHotspot() const = 0;
 	
 public:
 
diff --git a/Engine/Source/Editor/Sequencer/Public/ISequencerKeyCollection.h b/Engine/Source/Editor/Sequencer/Public/ISequencerKeyCollection.h
new file mode 100644
--- /dev/null
+++ b/Engine/Source/Editor/Sequencer/Public/ISequencerKeyCollection.h
@@ -0,0 +1,41 @@
+// Copyright 1998-2016 Epic Games, Inc. All Rights Reserved.
+
+#pragma once
+
+class FSequencerDisplayNode;
+
+enum class EFindKeyDirection
+{
+	Forwards, Backwards
+};
+
+/**
+ * A collection of keys gathered recursively from a particular node or nodes
+ */
+class ISequencerKeyCollection
+{
+public:
+	virtual ~ISequencerKeyCollection() {}
+
+public:
+
+	/** Iterate the keys contained within this collection */
+	virtual void IterateKeys(const TFunctionRef<bool(float)>& Iter) = 0;
+
+	/** Get a value specifying how close keys need to be in order to be considered equal by this collection */
+	virtual float GetKeyGroupingThreshold() const = 0;
+
+	/** Find the first key in the given range */
+	virtual TOptional<float> FindFirstKeyInRange(const TRange<float>& InRange, EFindKeyDirection Direction) const = 0;
+
+public:
+
+	/** Initialize this key collection from the specified nodes. Only gathers keys from those explicitly specified. */
+	virtual void InitializeExplicit(const TArray<FSequencerDisplayNode*>& InNodes, float DuplicateThreshold = SMALL_NUMBER) = 0;
+
+	/** Initialize this key collection from the specified nodes. Gathers keys from all child nodes. */
+	virtual void InitializeRecursive(const TArray<FSequencerDisplayNode*>& InNodes, float DuplicateThreshold = SMALL_NUMBER) = 0;
+
+	/** Initialize this key collection from the specified node and section index. */
+	virtual void InitializeRecursive(FSequencerDisplayNode& InNode, int32 InSectionIndex, float DuplicateThreshold = SMALL_NUMBER) = 0;
+};
\ No newline at end of file
diff --git a/Engine/Source/Editor/Sequencer/Public/ISequencerSection.h b/Engine/Source/Editor/Sequencer/Public/ISequencerSection.h
--- a/Engine/Source/Editor/Sequencer/Public/ISequencerSection.h
+++ b/Engine/Source/Editor/Sequencer/Public/ISequencerSection.h
@@ -3,8 +3,18 @@
 #pragma once
 
 #include "SequencerSectionPainter.h"
+#include "MovieSceneSection.h"
 
 class IKeyArea;
+class UMovieSceneSection;
+
+/** Enumerates which edge is being resized */
+UENUM()
+enum ESequencerSectionResizeMode
+{
+	SSRM_LeadingEdge,
+	SSRM_TrailingEdge
+};
 
 namespace SequencerSectionConstants
 {
@@ -58,6 +68,14 @@ public:
 	 */
 	virtual const FSlateBrush* GetKeyBrush(FKeyHandle KeyHandle) const { return nullptr; }
 
+	/** When a section overrides the brush to use, this allows it to set the scale origin tfor use
+	 * when the brush is scaled to create border and selection effects
+	 *
+	 * @param KeyHandle the handle of the key to get a brush origin for.
+	 * @return A FVector2D describing the custom origin, in slate units
+	 */
+	virtual FVector2D GetKeyBrushOrigin( FKeyHandle KeyHandle ) const { return FVector2D(0.0f, 0.0f); }
+
 	/**
 	 * Called when the section is double clicked
 	 *
@@ -131,10 +149,29 @@ public:
 	virtual bool RequestDeleteCategory( const TArray<FName>& CategoryNamePath ) { return false; }
 
 	/**
-	* Called when the user requests that a key area from this section be deleted.
-	*
-	* @param KeyAreaNamePath An array of names representing the path of to the key area to delete, starting with any categories which contain the key area.
-	* @returns Whether or not the key area was deleted.
-	*/
+	 * Called when the user requests that a key area from this section be deleted.
+	 *
+	 * @param KeyAreaNamePath An array of names representing the path of to the key area to delete, starting with any categories which contain the key area.
+	 * @returns Whether or not the key area was deleted.
+	 */
 	virtual bool RequestDeleteKeyArea( const TArray<FName>& KeyAreaNamePath ) { return false; }
+
+	/**
+	 * Resize the section 
+	 *
+	 * @param ResizeMode Resize either the leading or the trailing edge of the section
+	 * @param ResizeTime The time to resize to
+	 */
+	virtual void BeginResizeSection() {}
+	virtual void ResizeSection(ESequencerSectionResizeMode ResizeMode, float ResizeTime) { ResizeMode == ESequencerSectionResizeMode::SSRM_LeadingEdge ? GetSectionObject()->SetStartTime(ResizeTime) : GetSectionObject()->SetEndTime(ResizeTime); }
+
+	/**
+	 * Dilates the section by a specific factor
+	 *
+	 * @param DilationFactor The multiplier which scales this section
+	 * @param bFromStart Whether to dilate from the beginning or end (whichever stays put)
+	 * @param KeyHandles The key handles to operate on
+	 */
+	virtual void BeginDilateSection() {}
+	virtual void DilateSection(float DilationFactor, float Origin, TSet<FKeyHandle>& KeyHandles) { GetSectionObject()->DilateSection(DilationFactor, Origin, KeyHandles); }
 };
\ No newline at end of file
diff --git a/Engine/Source/Editor/Sequencer/Public/SequencerCommands.h b/Engine/Source/Editor/Sequencer/Public/SequencerCommands.h
--- a/Engine/Source/Editor/Sequencer/Public/SequencerCommands.h
+++ b/Engine/Source/Editor/Sequencer/Public/SequencerCommands.h
@@ -48,6 +48,9 @@ public:
 	/** Set end playback range */
 	TSharedPtr< FUICommandInfo > SetEndPlaybackRange;
 
+	/** Reset the view range to the playback range */
+	TSharedPtr< FUICommandInfo > ResetViewRange;
+
 	/** Zoom into the view range */
 	TSharedPtr< FUICommandInfo > ZoomInViewRange;
 
@@ -72,6 +75,12 @@ public:
 	/** Expand/collapse nodes and descendants */
 	TSharedPtr< FUICommandInfo > ToggleExpandCollapseNodesAndDescendants;
 
+	/** Sets the upper bound of the in/out section */
+	TSharedPtr< FUICommandInfo > SetInOutEnd;
+
+	/** Sets the lower bound of the in/out section */
+	TSharedPtr< FUICommandInfo > SetInOutStart;
+
 	/** Sets a key at the current time for the selected actor */
 	TSharedPtr< FUICommandInfo > SetKey;
 
@@ -138,6 +147,9 @@ public:
 	/** Toggles whether or not sections should snap to other sections. */
 	TSharedPtr< FUICommandInfo > ToggleSnapSectionTimesToSections;
 
+	/** Toggles whether or not snap to key times while scrubbing. */
+	TSharedPtr< FUICommandInfo > ToggleSnapPlayTimeToKeys;
+
 	/** Toggles whether or not the play time should snap to the selected interval. */
 	TSharedPtr< FUICommandInfo > ToggleSnapPlayTimeToInterval;
 
@@ -156,6 +168,9 @@ public:
 	/** Toggles whether to show combined keys at the top node level. */
 	TSharedPtr< FUICommandInfo > ToggleCombinedKeyframes;
 
+	/** Toggles whether to show channel colors in the track area. */
+	TSharedPtr< FUICommandInfo > ToggleChannelColors;
+
 	/** Toggles whether the label browser is enabled in the level editor. */
 	TSharedPtr< FUICommandInfo > ToggleLabelBrowser;
 
@@ -186,6 +201,10 @@ public:
 	/** Discard all changes to the current movie scene. */
 	TSharedPtr< FUICommandInfo > DiscardChanges;
 
+	/** Attempts to fix broken actor references. */
+	TSharedPtr< FUICommandInfo > FixActorReferences;
+
+
 	/**
 	 * Initialize commands
 	 */
diff --git a/Engine/Source/Editor/Sequencer/Public/SequencerSectionPainter.h b/Engine/Source/Editor/Sequencer/Public/SequencerSectionPainter.h
--- a/Engine/Source/Editor/Sequencer/Public/SequencerSectionPainter.h
+++ b/Engine/Source/Editor/Sequencer/Public/SequencerSectionPainter.h
@@ -6,23 +6,6 @@ struct FTimeToPixel;
 class UMovieSceneTrack;
 class UMovieSceneSection;
 
-/** Structure representing a specific region on a section to tint */
-struct FSequencerSectionTintRegion
-{
-	FSequencerSectionTintRegion(const FColor& InColor)
-		: Color(InColor)
-	{}
-
-	FSequencerSectionTintRegion(const FColor& InColor, float InStartTime, float InEndTime)
-		: Color(InColor), Range(TRange<float>(InStartTime, InEndTime))
-	{}
-
-	/** The color of the tint */
-	FColor Color;
-
-	/** The time range to tint */
-	TOptional<TRange<float>> Range;
-};
 
 /** Class that wraps up common section painting functionality */
 class SEQUENCER_API FSequencerSectionPainter
@@ -32,7 +15,7 @@ public:
 	FSequencerSectionPainter(FSlateWindowElementList& OutDrawElements, const FGeometry& InSectionGeometry, UMovieSceneSection& Section);
 
 	/** Paint the section background with the specified tint override */
-	int32 PaintSectionBackground(const FColor& Tint);
+	virtual int32 PaintSectionBackground(const FLinearColor& Tint) = 0;
 
 	/** Paint the section background with the tint stored on the track */
 	int32 PaintSectionBackground();
@@ -40,10 +23,10 @@ public:
 	/** Get the track that this painter is painting sections for */
 	UMovieSceneTrack* GetTrack() const;
 
-public:
+	/** Blend the specified color with the default track color */
+	static FLinearColor BlendColor(FLinearColor InColor);
 
-	/** Paint the section background with the specified tints */
-	virtual int32 PaintSectionBackground(const FSequencerSectionTintRegion* Tints, uint32 NumTints) = 0;
+public:
 
 	/** Get a time-to-pixel converter for the section */
 	virtual const FTimeToPixel& GetTimeConverter() const = 0;
diff --git a/Engine/Source/Editor/Sequencer/Public/SequencerSpawnRegister.h b/Engine/Source/Editor/Sequencer/Public/SequencerSpawnRegister.h
new file mode 100644
--- /dev/null
+++ b/Engine/Source/Editor/Sequencer/Public/SequencerSpawnRegister.h
@@ -0,0 +1,33 @@
+// Copyright 1998-2016 Epic Games, Inc. All Rights Reserved.
+
+#pragma once
+
+#include "MovieScene.h"
+#include "MovieSceneSpawnable.h"
+#include "IMovieSceneSpawnRegister.h"
+
+/** Spawn register used by sequencer to override default behaviour of other implementations */
+template<typename MixinType>
+struct TSequencerSpawnRegister : public MixinType
+{
+	/** Implement CanDestroyExpiredObject to take into account FMovieSceneSpawnable::ShouldIgnoreOwnershipInEditor */
+	virtual bool CanDestroyExpiredObject(const FGuid& ObjectId, ESpawnOwnership Ownership, FMovieSceneSequenceInstance& Instance) override
+	{
+		if (!MixinType::CanDestroyExpiredObject(ObjectId, Ownership, Instance))
+		{
+			return false;
+		}
+		
+		// Find the spawnable
+		UMovieSceneSequence* Sequence = Instance.GetSequence();
+		UMovieScene* MovieScene = Sequence ? Sequence->GetMovieScene() : nullptr;
+
+		if (!MovieScene)
+		{
+			return true;
+		}
+
+		const FMovieSceneSpawnable* Spawnable = MovieScene->FindSpawnable(ObjectId);
+		return !Spawnable || !Spawnable->ShouldIgnoreOwnershipInEditor();
+	}
+};
diff --git a/Engine/Source/Editor/Sequencer/Public/SequencerUtilities.h b/Engine/Source/Editor/Sequencer/Public/SequencerUtilities.h
--- a/Engine/Source/Editor/Sequencer/Public/SequencerUtilities.h
+++ b/Engine/Source/Editor/Sequencer/Public/SequencerUtilities.h
@@ -5,4 +5,10 @@
 struct SEQUENCER_API FSequencerUtilities
 {
 	static TSharedRef<SWidget> MakeAddButton(FText HoverText, FOnGetContent MenuContent, const TAttribute<bool>& HoverState);
+
+	/** 
+	 * Generates a unique FName from a candidate name given a set of already existing names.  
+	 * The name is made unique by appending a number to the end.
+	 */
+	static FName GetUniqueName(FName CandidateName, const TArray<FName>& ExistingNames);
 };
\ No newline at end of file
diff --git a/Engine/Source/Editor/Sequencer/Sequencer.Build.cs b/Engine/Source/Editor/Sequencer/Sequencer.Build.cs
--- a/Engine/Source/Editor/Sequencer/Sequencer.Build.cs
+++ b/Engine/Source/Editor/Sequencer/Sequencer.Build.cs
@@ -41,7 +41,8 @@ public class Sequencer : ModuleRules
                 "ContentBrowser",
 				"PropertyEditor",
 				"Kismet",
-				"SceneOutliner"
+				"SceneOutliner",
+				"SequenceRecorder"
 			}
 		);
 
@@ -50,7 +51,8 @@ public class Sequencer : ModuleRules
 				"PropertyEditor",
 				"LevelEditor",
 				"SceneOutliner",
-				"WorkspaceMenuStructure"
+				"WorkspaceMenuStructure",
+				"SequenceRecorder"
 			}
 		);
 
