Copying //UE4/Dev-Sequencer to //UE4/Main (Source: //UE4/Dev-Sequencer @ 2945541)

diff --git a/Engine/Source/Editor/Sequencer/Private/SSequencer.cpp b/Engine/Source/Editor/Sequencer/Private/SSequencer.cpp
--- a/Engine/Source/Editor/Sequencer/Private/SSequencer.cpp
+++ b/Engine/Source/Editor/Sequencer/Private/SSequencer.cpp
@@ -113,23 +113,23 @@ void SSequencer::Construct(const FArguments& InArgs, TSharedRef<FSequencer> InSe
 
 	ISequencerWidgetsModule& SequencerWidgets = FModuleManager::Get().LoadModuleChecked<ISequencerWidgetsModule>( "SequencerWidgets" );
 
-	OnBeginInOutRangeDrag = InArgs._OnBeginInOutRangeDrag;
-	OnEndInOutRangeDrag = InArgs._OnEndInOutRangeDrag;
-	OnBeginPlaybackRangeDrag = InArgs._OnBeginPlaybackRangeDrag;
-	OnEndPlaybackRangeDrag = InArgs._OnEndPlaybackRangeDrag;
+	OnPlaybackRangeBeginDrag = InArgs._OnPlaybackRangeBeginDrag;
+	OnPlaybackRangeEndDrag = InArgs._OnPlaybackRangeEndDrag;
+	OnSelectionRangeBeginDrag = InArgs._OnSelectionRangeBeginDrag;
+	OnSelectionRangeEndDrag = InArgs._OnSelectionRangeEndDrag;
 
 	FTimeSliderArgs TimeSliderArgs;
 	{
 		TimeSliderArgs.ViewRange = InArgs._ViewRange;
 		TimeSliderArgs.ClampRange = InArgs._ClampRange;
-		TimeSliderArgs.InOutRange = InArgs._InOutRange;
 		TimeSliderArgs.PlaybackRange = InArgs._PlaybackRange;
-		TimeSliderArgs.OnInOutRangeChanged = InArgs._OnInOutRangeChanged;
-		TimeSliderArgs.OnBeginInOutRangeDrag = OnBeginInOutRangeDrag;
-		TimeSliderArgs.OnEndInOutRangeDrag = OnEndInOutRangeDrag;
+		TimeSliderArgs.SelectionRange = InArgs._SelectionRange;
 		TimeSliderArgs.OnPlaybackRangeChanged = InArgs._OnPlaybackRangeChanged;
-		TimeSliderArgs.OnBeginPlaybackRangeDrag = OnBeginPlaybackRangeDrag;
-		TimeSliderArgs.OnEndPlaybackRangeDrag = OnEndPlaybackRangeDrag;
+		TimeSliderArgs.OnPlaybackRangeBeginDrag = OnPlaybackRangeBeginDrag;
+		TimeSliderArgs.OnPlaybackRangeEndDrag = OnPlaybackRangeEndDrag;
+		TimeSliderArgs.OnSelectionRangeChanged = InArgs._OnSelectionRangeChanged;
+		TimeSliderArgs.OnSelectionRangeBeginDrag = OnSelectionRangeBeginDrag;
+		TimeSliderArgs.OnSelectionRangeEndDrag = OnSelectionRangeEndDrag;
 		TimeSliderArgs.OnViewRangeChanged = InArgs._OnViewRangeChanged;
 		TimeSliderArgs.OnClampRangeChanged = InArgs._OnClampRangeChanged;
 		TimeSliderArgs.ScrubPosition = InArgs._ScrubPosition;
@@ -686,7 +686,9 @@ TSharedRef<SWidget> SSequencer::MakeToolBar()
 				FSlateIcon(FEditorStyle::GetStyleSetName(), "Sequencer.SaveAs")
 			);
 
-			ToolBarBuilder.AddToolBarButton( FSequencerCommands::Get().DiscardChanges );
+			//ToolBarBuilder.AddToolBarButton( FSequencerCommands::Get().DiscardChanges );
+			ToolBarBuilder.AddToolBarButton( FSequencerCommands::Get().FindInContentBrowser );
+			ToolBarBuilder.AddToolBarButton( FSequencerCommands::Get().CreateCamera );
 			ToolBarBuilder.AddToolBarButton( FSequencerCommands::Get().RenderMovie );
 			ToolBarBuilder.AddSeparator();
 		}
@@ -765,7 +767,7 @@ TSharedRef<SWidget> SSequencer::MakeToolBar()
 		ToolBarBuilder.AddSeparator();
 		ToolBarBuilder.AddWidget(
 			SNew( SImage )
-				.Image(FEditorStyle::GetBrush("Sequencer.Time")) );
+				.Image(FEditorStyle::GetBrush("Sequencer.Time.Small")) );
 
 		ToolBarBuilder.AddWidget(
 			SNew( SBox )
@@ -820,10 +822,12 @@ TSharedRef<SWidget> SSequencer::MakeAddMenu()
 TSharedRef<SWidget> SSequencer::MakeGeneralMenu()
 {
 	FMenuBuilder MenuBuilder( true, SequencerPtr.Pin()->GetCommandBindings() );
+	TSharedPtr<FSequencer> Sequencer = SequencerPtr.Pin();
 
+	// view options
 	MenuBuilder.BeginSection( "ViewOptions", LOCTEXT( "ViewMenuHeader", "View" ) );
 	{
-		if (SequencerPtr.Pin()->IsLevelEditorSequencer())
+		if (Sequencer->IsLevelEditorSequencer())
 		{
 			MenuBuilder.AddMenuEntry( FSequencerCommands::Get().ToggleLabelBrowser );
 		}
@@ -831,7 +835,7 @@ TSharedRef<SWidget> SSequencer::MakeGeneralMenu()
 		MenuBuilder.AddMenuEntry( FSequencerCommands::Get().ToggleCombinedKeyframes );
 		MenuBuilder.AddMenuEntry( FSequencerCommands::Get().ToggleChannelColors );
 
-		if (SequencerPtr.Pin()->IsLevelEditorSequencer())
+		if (Sequencer->IsLevelEditorSequencer())
 		{
 			MenuBuilder.AddMenuEntry( FSequencerCommands::Get().FindInContentBrowser );
 		}
@@ -843,41 +847,41 @@ TSharedRef<SWidget> SSequencer::MakeGeneralMenu()
 	}
 	MenuBuilder.EndSection();
 
-	MenuBuilder.BeginSection( "Ranges", LOCTEXT( "RangesHeader", "Playback Range" ) );
+	// playback range options
+	MenuBuilder.BeginSection("PlaybackRange", LOCTEXT("PlaybackRangeHeader", "Playback Range"));
 	{
-		TSharedPtr<FSequencer> Sequencer = SequencerPtr.Pin();
-		
 		// Menu entry for the start position
 		auto OnStartChanged = [=](float NewValue){
 			float Upper = Sequencer->GetPlaybackRange().GetUpperBoundValue();
 			Sequencer->SetPlaybackRange(TRange<float>(FMath::Min(NewValue, Upper), Upper));
 		};
+
 		MenuBuilder.AddWidget(
 			SNew(SHorizontalBox)
 			+ SHorizontalBox::Slot()
-			[
-				SNew(SSpacer)
-			]
+				[
+					SNew(SSpacer)
+				]
 			+ SHorizontalBox::Slot()
-			.AutoWidth()
-			[
-				SNew(SSpinBox<float>)
-				.TypeInterface(NumericTypeInterface)
-				.Style(&FEditorStyle::GetWidgetStyle<FSpinBoxStyle>("Sequencer.HyperlinkSpinBox"))
-				.OnValueCommitted_Lambda([=](float Value, ETextCommit::Type){ OnStartChanged(Value); })
-				.OnValueChanged_Lambda(OnStartChanged)
-				.OnBeginSliderMovement(OnBeginPlaybackRangeDrag)
-				.OnEndSliderMovement_Lambda([=](float Value){ OnStartChanged(Value); OnEndPlaybackRangeDrag.ExecuteIfBound(); })
-				.MinValue_Lambda([=]() -> float {
-					return Sequencer->GetClampRange().GetLowerBoundValue(); 
-				})
-				.MaxValue_Lambda([=]() -> float {
-					return Sequencer->GetPlaybackRange().GetUpperBoundValue(); 
-				})
-				.Value_Lambda([=]() -> float {
-					return Sequencer->GetPlaybackRange().GetLowerBoundValue();
-				})
-			],
+				.AutoWidth()
+				[
+					SNew(SSpinBox<float>)
+						.TypeInterface(NumericTypeInterface)
+						.Style(&FEditorStyle::GetWidgetStyle<FSpinBoxStyle>("Sequencer.HyperlinkSpinBox"))
+						.OnValueCommitted_Lambda([=](float Value, ETextCommit::Type){ OnStartChanged(Value); })
+						.OnValueChanged_Lambda(OnStartChanged)
+						.OnBeginSliderMovement(OnPlaybackRangeBeginDrag)
+						.OnEndSliderMovement_Lambda([=](float Value){ OnStartChanged(Value); OnPlaybackRangeEndDrag.ExecuteIfBound(); })
+						.MinValue_Lambda([=]() -> float {
+							return Sequencer->GetClampRange().GetLowerBoundValue(); 
+						})
+						.MaxValue_Lambda([=]() -> float {
+							return Sequencer->GetPlaybackRange().GetUpperBoundValue(); 
+						})
+						.Value_Lambda([=]() -> float {
+							return Sequencer->GetPlaybackRange().GetLowerBoundValue();
+						})
+				],
 			LOCTEXT("PlaybackStartLabel", "Start"));
 
 		// Menu entry for the end position
@@ -885,32 +889,33 @@ TSharedRef<SWidget> SSequencer::MakeGeneralMenu()
 			float Lower = Sequencer->GetPlaybackRange().GetLowerBoundValue();
 			Sequencer->SetPlaybackRange(TRange<float>(Lower, FMath::Max(NewValue, Lower)));
 		};
+
 		MenuBuilder.AddWidget(
 			SNew(SHorizontalBox)
 			+ SHorizontalBox::Slot()
-			[
-				SNew(SSpacer)
-			]
+				[
+					SNew(SSpacer)
+				]
 			+ SHorizontalBox::Slot()
-			.AutoWidth()
-			[
-				SNew(SSpinBox<float>)
-				.TypeInterface(NumericTypeInterface)
-				.Style(&FEditorStyle::GetWidgetStyle<FSpinBoxStyle>("Sequencer.HyperlinkSpinBox"))
-				.OnValueCommitted_Lambda([=](float Value, ETextCommit::Type){ OnEndChanged(Value); })
-				.OnValueChanged_Lambda(OnEndChanged)
-				.OnBeginSliderMovement(OnBeginPlaybackRangeDrag)
-				.OnEndSliderMovement_Lambda([=](float Value){ OnEndChanged(Value); OnEndPlaybackRangeDrag.ExecuteIfBound(); })
-				.MinValue_Lambda([=]() -> float {
-					return Sequencer->GetPlaybackRange().GetLowerBoundValue(); 
-				})
-				.MaxValue_Lambda([=]() -> float {
-					return Sequencer->GetClampRange().GetUpperBoundValue(); 
-				})
-				.Value_Lambda([=]() -> float {
-					return Sequencer->GetPlaybackRange().GetUpperBoundValue();
-				})
-			],
+				.AutoWidth()
+				[
+					SNew(SSpinBox<float>)
+						.TypeInterface(NumericTypeInterface)
+						.Style(&FEditorStyle::GetWidgetStyle<FSpinBoxStyle>("Sequencer.HyperlinkSpinBox"))
+						.OnValueCommitted_Lambda([=](float Value, ETextCommit::Type){ OnEndChanged(Value); })
+						.OnValueChanged_Lambda(OnEndChanged)
+						.OnBeginSliderMovement(OnPlaybackRangeBeginDrag)
+						.OnEndSliderMovement_Lambda([=](float Value){ OnEndChanged(Value); OnPlaybackRangeEndDrag.ExecuteIfBound(); })
+						.MinValue_Lambda([=]() -> float {
+							return Sequencer->GetPlaybackRange().GetLowerBoundValue(); 
+						})
+						.MaxValue_Lambda([=]() -> float {
+							return Sequencer->GetClampRange().GetUpperBoundValue(); 
+						})
+						.Value_Lambda([=]() -> float {
+							return Sequencer->GetPlaybackRange().GetUpperBoundValue();
+						})
+				],
 			LOCTEXT("PlaybackStartEnd", "End"));
 
 		MenuBuilder.AddMenuEntry( FSequencerCommands::Get().ToggleKeepCursorInPlaybackRange );
@@ -921,39 +926,51 @@ TSharedRef<SWidget> SSequencer::MakeGeneralMenu()
 		auto OnZeroPadChanged = [=](uint8 NewValue){
 			Settings->SetZeroPadFrames(NewValue);
 		};
+
 		MenuBuilder.AddWidget(
-			SNew(SHorizontalBox)
+			SNew(SHorizontalBox)	
 			+ SHorizontalBox::Slot()
-			[
-				SNew(SSpacer)
-			]
+				[
+					SNew(SSpacer)
+				]
 			+ SHorizontalBox::Slot()
-			.AutoWidth()
-			[
-				SNew(SSpinBox<uint8>)
-				.Style(&FEditorStyle::GetWidgetStyle<FSpinBoxStyle>("Sequencer.HyperlinkSpinBox"))
-				.OnValueCommitted_Lambda([=](uint8 Value, ETextCommit::Type){ OnZeroPadChanged(Value); })
-				.OnValueChanged_Lambda(OnZeroPadChanged)
-				.MinValue(0)
-				.MaxValue(8)
-				.Value_Lambda([=]() -> uint8 {
-					return Settings->GetZeroPadFrames();
-				})
-			],
+				.AutoWidth()
+				[
+					SNew(SSpinBox<uint8>)
+					.Style(&FEditorStyle::GetWidgetStyle<FSpinBoxStyle>("Sequencer.HyperlinkSpinBox"))
+					.OnValueCommitted_Lambda([=](uint8 Value, ETextCommit::Type){ OnZeroPadChanged(Value); })
+					.OnValueChanged_Lambda(OnZeroPadChanged)
+					.MinValue(0)
+					.MaxValue(8)
+					.Value_Lambda([=]() -> uint8 {
+						return Settings->GetZeroPadFrames();
+					})
+				],
 			LOCTEXT("ZeroPaddingText", "Zero Pad Frame Numbers"));
+	}
+	MenuBuilder.AddMenuEntry(FSequencerCommands::Get().ToggleShowGotoBox);
+	MenuBuilder.EndSection();
 
-		MenuBuilder.AddMenuEntry( FSequencerCommands::Get().ToggleShowGotoBox );
+	// selection range options
+	MenuBuilder.BeginSection("SelectionRange", LOCTEXT("SelectionRangeHeader", "Selection Range"));
+	{
+		MenuBuilder.AddMenuEntry(FSequencerCommands::Get().SetSelectionRangeStart);
+		MenuBuilder.AddMenuEntry(FSequencerCommands::Get().SetSelectionRangeEnd);
+		MenuBuilder.AddMenuEntry(FSequencerCommands::Get().SelectKeysInSelectionRange);
+	}
+	MenuBuilder.EndSection();
 
-		if (SequencerPtr.Pin()->IsLevelEditorSequencer())
-		{
-			MenuBuilder.AddMenuSeparator();
-			MenuBuilder.AddMenuEntry( FSequencerCommands::Get().FixActorReferences );
-		}
+	// other options
+	if (SequencerPtr.Pin()->IsLevelEditorSequencer())
+	{
+		MenuBuilder.AddMenuSeparator();
+		MenuBuilder.AddMenuEntry(FSequencerCommands::Get().FixActorReferences);
 	}
 
 	return MenuBuilder.MakeWidget();
 }
 
+
 TSharedRef<SWidget> SSequencer::MakeSnapMenu()
 {
 	FMenuBuilder MenuBuilder( false, SequencerPtr.Pin()->GetCommandBindings() );
@@ -1090,6 +1107,8 @@ void RestoreSelectionState(const TArray<TSharedRef<FSequencerDisplayNode>>& Disp
 
 void SSequencer::UpdateLayoutTree()
 {
+	TrackArea->Empty();
+
 	TSharedPtr<FSequencer> Sequencer = SequencerPtr.Pin();
 	if ( Sequencer.IsValid() )
 	{
@@ -1106,6 +1125,7 @@ void SSequencer::UpdateLayoutTree()
 
 		// Suspend broadcasting selection changes because we don't want unnecessary rebuilds.
 		Sequencer->GetSelection().SuspendBroadcast();
+		Sequencer->GetSelection().Empty();
 
 		// Update the node tree
 		Sequencer->GetNodeTree()->Update();
diff --git a/Engine/Source/Editor/Sequencer/Private/SSequencer.h b/Engine/Source/Editor/Sequencer/Private/SSequencer.h
--- a/Engine/Source/Editor/Sequencer/Private/SSequencer.h
+++ b/Engine/Source/Editor/Sequencer/Private/SSequencer.h
@@ -33,7 +33,7 @@ struct FPaintPlaybackRangeArgs;
 namespace SequencerLayoutConstants
 {
 	/** The amount to indent child nodes of the layout tree */
-	const float IndentAmount = 8.0f;
+	const float IndentAmount = 10.0f;
 
 	/** Height of each folder node */
 	const float FolderNodeHeight = 20.0f;
@@ -100,12 +100,12 @@ public:
 		/** The current clamp range (seconds) */
 		SLATE_ATTRIBUTE( FAnimatedRange, ClampRange )
 
-		/** The in/out selection range */
-		SLATE_ATTRIBUTE( TRange<float>, InOutRange )
-
 		/** The playback range */
 		SLATE_ATTRIBUTE( TRange<float>, PlaybackRange )
 
+		/** The selection selection range */
+		SLATE_ATTRIBUTE(TRange<float>, SelectionRange)
+
 		/** The playback status */
 		SLATE_ATTRIBUTE( EMovieScenePlayerStatus::Type, PlaybackStatus )
 
@@ -115,17 +115,23 @@ public:
 		/** Called when the user has begun dragging the selection range */
 		SLATE_EVENT( FSimpleDelegate, OnBeginInOutRangeDrag )
 
-		/** Called when the user has finished dragging the selection range */
-		SLATE_EVENT( FSimpleDelegate, OnEndInOutRangeDrag )
-
 		/** Called when the user changes the playback range */
 		SLATE_EVENT( FOnRangeChanged, OnPlaybackRangeChanged )
 
 		/** Called when the user has begun dragging the playback range */
-		SLATE_EVENT( FSimpleDelegate, OnBeginPlaybackRangeDrag )
+		SLATE_EVENT( FSimpleDelegate, OnPlaybackRangeBeginDrag )
 
 		/** Called when the user has finished dragging the playback range */
-		SLATE_EVENT( FSimpleDelegate, OnEndPlaybackRangeDrag )
+		SLATE_EVENT( FSimpleDelegate, OnPlaybackRangeEndDrag )
+
+		/** Called when the user changes the selection range */
+		SLATE_EVENT( FOnRangeChanged, OnSelectionRangeChanged )
+
+		/** Called when the user has begun dragging the selection range */
+		SLATE_EVENT( FSimpleDelegate, OnSelectionRangeBeginDrag )
+
+		/** Called when the user has finished dragging the selection range */
+		SLATE_EVENT( FSimpleDelegate, OnSelectionRangeEndDrag )
 
 		/** The current scrub position in (seconds) */
 		SLATE_ATTRIBUTE( float, ScrubPosition )
@@ -440,17 +446,17 @@ private:
 
 	FOnGetAddMenuContent OnGetAddMenuContent;
 
-	/** Called when the user has begun dragging the in/out selection range */
-	FSimpleDelegate OnBeginInOutRangeDrag;
+	/** Called when the user has begun dragging the selection selection range */
+	FSimpleDelegate OnSelectionRangeBeginDrag;
 
-	/** Called when the user has finished dragging the in/out selection range */
-	FSimpleDelegate OnEndInOutRangeDrag;
+	/** Called when the user has finished dragging the selection selection range */
+	FSimpleDelegate OnSelectionRangeEndDrag;
 
 	/** Called when the user has begun dragging the playback range */
-	FSimpleDelegate OnBeginPlaybackRangeDrag;
+	FSimpleDelegate OnPlaybackRangeBeginDrag;
 
 	/** Called when the user has finished dragging the playback range */
-	FSimpleDelegate OnEndPlaybackRangeDrag;
+	FSimpleDelegate OnPlaybackRangeEndDrag;
 
 	/** Cached clamp and view range for unlinking the curve editor time range */
 	TRange<float> CachedClampRange;
diff --git a/Engine/Source/Editor/Sequencer/Private/SSequencerCurveEditor.cpp b/Engine/Source/Editor/Sequencer/Private/SSequencerCurveEditor.cpp
--- a/Engine/Source/Editor/Sequencer/Private/SSequencerCurveEditor.cpp
+++ b/Engine/Source/Editor/Sequencer/Private/SSequencerCurveEditor.cpp
@@ -29,7 +29,7 @@ void SSequencerCurveEditor::Construct( const FArguments& InArgs, TSharedRef<FSeq
 		.ShowCurveSelector( false )
 		.ShowZoomButtons( false )
 		.ShowInputGridNumbers( false )
-		.SnappingEnabled( this, &SSequencerCurveEditor::GetCurveSnapEnabled )
+		.InputSnappingEnabled(this, &SSequencerCurveEditor::GetInputCurveSnapEnabled)
 		.InputSnap( this, &SSequencerCurveEditor::GetCurveTimeSnapInterval )
 		.OutputSnap( this, &SSequencerCurveEditor::GetCurveValueSnapInterval )
 		.TimelineLength( 0.0f )
@@ -145,7 +145,7 @@ void SSequencerCurveEditor::UpdateCurveOwner()
 	UpdateCurveViewModelSelection();
 }
 
-bool SSequencerCurveEditor::GetCurveSnapEnabled() const
+bool SSequencerCurveEditor::GetInputCurveSnapEnabled() const
 {
 	return SequencerSettings->GetIsSnapEnabled();
 }
diff --git a/Engine/Source/Editor/Sequencer/Private/SSequencerCurveEditor.h b/Engine/Source/Editor/Sequencer/Private/SSequencerCurveEditor.h
--- a/Engine/Source/Editor/Sequencer/Private/SSequencerCurveEditor.h
+++ b/Engine/Source/Editor/Sequencer/Private/SSequencerCurveEditor.h
@@ -38,8 +38,8 @@ private:
 	/** Builds and assigns a new curve owner to the FCurveEditor. */
 	void UpdateCurveOwner();
 
-	/** Gets whether or not snapping is enabled for the curve editor. */
-	bool GetCurveSnapEnabled() const;
+	/** Gets whether or not snapping is enabled for the input of the curve. */
+	bool GetInputCurveSnapEnabled() const;
 
 	/** Gets time snapping interval. */
 	float GetCurveTimeSnapInterval() const;
diff --git a/Engine/Source/Editor/Sequencer/Private/SSequencerCurveEditorToolBar.cpp b/Engine/Source/Editor/Sequencer/Private/SSequencerCurveEditorToolBar.cpp
--- a/Engine/Source/Editor/Sequencer/Private/SSequencerCurveEditorToolBar.cpp
+++ b/Engine/Source/Editor/Sequencer/Private/SSequencerCurveEditorToolBar.cpp
@@ -21,9 +21,11 @@ void SSequencerCurveEditorToolBar::Construct( const FArguments& InArgs, TSharedR
 		TAttribute<FSlateIcon>(),
 		true );
 
+	ToolBarBuilder.AddToolBarButton(FRichCurveEditorCommands::Get().ToggleOutputSnapping);
+
 	ToolBarBuilder.AddWidget(
 		SNew( SImage )
-			.Image(FEditorStyle::GetBrush("Sequencer.Value")) );
+			.Image(FEditorStyle::GetBrush("Sequencer.Value.Small")) );
 
 	ToolBarBuilder.AddWidget(
 		SNew( SBox )
diff --git a/Engine/Source/Editor/Sequencer/Private/SSequencerTrackArea.cpp b/Engine/Source/Editor/Sequencer/Private/SSequencerTrackArea.cpp
--- a/Engine/Source/Editor/Sequencer/Private/SSequencerTrackArea.cpp
+++ b/Engine/Source/Editor/Sequencer/Private/SSequencerTrackArea.cpp
@@ -57,6 +57,11 @@ void SSequencerTrackArea::SetTreeView(const TSharedPtr<SSequencerTreeView>& InTr
 	TreeView = InTreeView;
 }
 
+void SSequencerTrackArea::Empty()
+{
+	TrackSlots.Empty();
+	Children.Empty();
+}
 
 void SSequencerTrackArea::AddTrackSlot(const TSharedRef<FSequencerDisplayNode>& InNode, const TSharedPtr<SSequencerTrackLane>& InSlot)
 {
@@ -281,7 +286,7 @@ void SSequencerTrackArea::UpdateHoverStates( const FGeometry& MyGeometry, const
 				{
 					// Activate selection mode if the section has keys, otherwise just move it
 					TSet<FKeyHandle> KeyHandles;
-					Section->GetKeyHandles(KeyHandles);
+					Section->GetKeyHandles(KeyHandles, Section->GetRange());
 
 					if (KeyHandles.Num())
 					{
diff --git a/Engine/Source/Editor/Sequencer/Private/SSequencerTrackArea.h b/Engine/Source/Editor/Sequencer/Private/SSequencerTrackArea.h
--- a/Engine/Source/Editor/Sequencer/Private/SSequencerTrackArea.h
+++ b/Engine/Source/Editor/Sequencer/Private/SSequencerTrackArea.h
@@ -51,6 +51,9 @@ public:
 
 public:
 
+	/** Empty the track area */
+	void Empty();
+
 	/** Add a new track slot to this area for the given node. The slot will be automatically cleaned up when all external references to the supplied slot are removed. */
 	void AddTrackSlot(const TSharedRef<FSequencerDisplayNode>& InNode, const TSharedPtr<SSequencerTrackLane>& InSlot);
 
diff --git a/Engine/Source/Editor/Sequencer/Private/Sequencer.cpp b/Engine/Source/Editor/Sequencer/Private/Sequencer.cpp
--- a/Engine/Source/Editor/Sequencer/Private/Sequencer.cpp
+++ b/Engine/Source/Editor/Sequencer/Private/Sequencer.cpp
@@ -22,12 +22,14 @@
 #include "Kismet2/BlueprintEditorUtils.h"
 #include "BlueprintEditorModule.h"
 #include "MovieSceneTrack.h"
+#include "MovieScene3DAttachTrack.h"
 #include "MovieSceneCameraCutTrack.h"
 #include "MovieSceneAudioTrack.h"
 #include "MovieSceneSkeletalAnimationTrack.h"
 #include "MovieSceneTrackEditor.h"
 #include "MovieSceneToolHelpers.h"
 #include "ScopedTransaction.h"
+#include "MovieScene3DAttachSection.h"
 #include "MovieSceneBoolSection.h"
 #include "MovieSceneCameraCutSection.h"
 #include "MovieScene3DTransformSection.h"
@@ -64,6 +66,11 @@
 #include "SequencerUtilities.h"
 #include "MovieSceneCinematicShotTrack.h"
 #include "ISequenceRecorder.h"
+#include "GameDelegates.h"
+#include "CineCameraActor.h"
+#include "CameraRig_Rail.h"
+#include "CameraRig_Crane.h"
+#include "Components/SplineComponent.h"
 
 #define LOCTEXT_NAMESPACE "Sequencer"
 
@@ -106,9 +113,14 @@ void FSequencer::InitSequencer(const FSequencerInitParams& InitParams, const TSh
 		FEditorDelegates::PostSaveWorld.AddSP(this, &FSequencer::OnPostSaveWorld);
 		FEditorDelegates::PreBeginPIE.AddSP(this, &FSequencer::OnPreBeginPIE);
 		FEditorDelegates::EndPIE.AddSP(this, &FSequencer::OnEndPIE);
+		FGameDelegates::Get().GetEndPlayMapDelegate().AddSP(this, &FSequencer::OnEndPlayMap);
+
+		FWorldDelegates::LevelAddedToWorld.AddSP(this, &FSequencer::OnLevelAdded);
+		FWorldDelegates::LevelRemovedFromWorld.AddSP(this, &FSequencer::OnLevelRemoved);
 
 		FEditorDelegates::NewCurrentLevel.AddSP(this, &FSequencer::OnNewCurrentLevel);
 		FEditorDelegates::OnMapOpened.AddSP(this, &FSequencer::OnMapOpened);
+		FEditorDelegates::OnNewActorsDropped.AddSP(this, &FSequencer::OnNewActorsDropped);
 
 		USelection::SelectionChangedEvent.AddSP( this, &FSequencer::OnActorSelectionChanged );
 	}
@@ -133,15 +145,15 @@ void FSequencer::InitSequencer(const FSequencerInitParams& InitParams, const TSh
 	SequencerWidget = SNew( SSequencer, SharedThis( this ) )
 		.ViewRange( this, &FSequencer::GetViewRange )
 		.ClampRange( this, &FSequencer::GetClampRange )
-		.InOutRange( this, &FSequencer::GetInOutRange )
 		.PlaybackRange( this, &FSequencer::GetPlaybackRange )
 		.PlaybackStatus( this, &FSequencer::GetPlaybackStatus )
-		.OnInOutRangeChanged( this, &FSequencer::SetInOutRange )
-		.OnBeginInOutRangeDrag( this, &FSequencer::OnBeginInOutRangeDrag )
-		.OnEndInOutRangeDrag( this, &FSequencer::OnEndInOutRangeDrag )
+		.SelectionRange( this, &FSequencer::GetSelectionRange )
 		.OnPlaybackRangeChanged( this, &FSequencer::SetPlaybackRange )
-		.OnBeginPlaybackRangeDrag( this, &FSequencer::OnBeginPlaybackRangeDrag )
-		.OnEndPlaybackRangeDrag( this, &FSequencer::OnEndPlaybackRangeDrag )
+		.OnPlaybackRangeBeginDrag( this, &FSequencer::OnPlaybackRangeBeginDrag )
+		.OnPlaybackRangeEndDrag( this, &FSequencer::OnPlaybackRangeEndDrag )
+		.OnSelectionRangeChanged( this, &FSequencer::SetSelectionRange )
+		.OnSelectionRangeBeginDrag( this, &FSequencer::OnSelectionRangeBeginDrag )
+		.OnSelectionRangeEndDrag( this, &FSequencer::OnSelectionRangeEndDrag )
 		.ScrubPosition( this, &FSequencer::OnGetScrubPosition )
 		.OnBeginScrubbing( this, &FSequencer::OnBeginScrubbing )
 		.OnEndScrubbing( this, &FSequencer::OnEndScrubbing )
@@ -186,7 +198,7 @@ void FSequencer::InitSequencer(const FSequencerInitParams& InitParams, const TSh
 
 	// NOTE: Could fill in asset editor commands here!
 
-	BindSequencerCommands();
+	BindCommands();
 	ActivateSequencerEditorMode();
 
 	if( bIsEditingWithinLevelEditor )
@@ -198,6 +210,8 @@ void FSequencer::InitSequencer(const FSequencerInitParams& InitParams, const TSh
 	{
 		TrackEditor->OnInitialize();
 	}
+
+	OnActivateSequenceEvent.Broadcast(*RootMovieSceneSequenceInstance);
 }
 
 
@@ -213,6 +227,7 @@ FSequencer::FSequencer()
 	, bPerspectiveViewportCameraCutEnabled( false )
 	, bIsEditingWithinLevelEditor( false )
 	, bNeedTreeRefresh( false )
+	, StoredPlaybackState( EMovieScenePlayerStatus::Stopped )
 	, bWasClosed( false )
 	, NodeTree( MakeShareable( new FSequencerNodeTree( *this ) ) )
 	, bUpdatingSequencerSelection( false )
@@ -260,9 +275,14 @@ void FSequencer::Close()
 		FEditorDelegates::PostSaveWorld.RemoveAll(this);
 		FEditorDelegates::PreBeginPIE.RemoveAll(this);
 		FEditorDelegates::EndPIE.RemoveAll(this);
+		FGameDelegates::Get().GetEndPlayMapDelegate().RemoveAll(this);
+
+		FWorldDelegates::LevelAddedToWorld.RemoveAll(this);
+		FWorldDelegates::LevelRemovedFromWorld.RemoveAll(this);
 
 		FEditorDelegates::NewCurrentLevel.RemoveAll(this);
 		FEditorDelegates::OnMapOpened.RemoveAll(this);
+		FEditorDelegates::OnNewActorsDropped.RemoveAll(this);
 		
 		for (auto TrackEditor : TrackEditors)
 		{
@@ -298,11 +318,15 @@ void FSequencer::Tick(float InDeltaTime)
 
 	if (bNeedTreeRefresh)
 	{
+		SelectionPreview.Empty();
+
 		// @todo - Sequencer Will be called too often
 		UpdateRuntimeInstances();
 
 		SequencerWidget->UpdateLayoutTree();
 		bNeedTreeRefresh = false;
+
+		SetPlaybackStatus(StoredPlaybackState);
 	}
 
 	static const float AutoScrollFactor = 0.1f;
@@ -471,6 +495,8 @@ void FSequencer::ResetToNewRootSequence(UMovieSceneSequence& NewSequence)
 
 	NotifyMovieSceneDataChanged();
 	UpdateTimeBoundsToFocusedMovieScene();
+
+	OnActivateSequenceEvent.Broadcast(*SequenceInstanceStack.Top());
 }
 
 
@@ -523,6 +549,8 @@ void FSequencer::FocusSequenceInstance(UMovieSceneSubSection& InSubSection)
 	SequencerWidget->UpdateBreadcrumbs();
 
 	SetGlobalTime(AbsoluteShotPosition, ESnapTimeMode::STM_Interval);
+
+	OnActivateSequenceEvent.Broadcast(*SequenceInstance);
 }
 
 
@@ -697,9 +725,25 @@ void FSequencer::PopToSequenceInstance(TSharedRef<FMovieSceneSequenceInstance> S
 	
 		check( SequenceInstanceStack.Num() > 0 );
 
+		// Pop out of any potentially locked cameras from the shot and toggle on camera cuts
+		for (int32 i = 0; i < GEditor->LevelViewportClients.Num(); ++i)
+		{		
+			FLevelEditorViewportClient* LevelVC = GEditor->LevelViewportClients[i];
+			if (LevelVC && LevelVC->IsPerspective() && LevelVC->AllowsCinematicPreview() && LevelVC->GetViewMode() != VMI_Unknown)
+			{
+				LevelVC->SetActorLock(nullptr);
+				LevelVC->bLockedCameraView = false;
+				LevelVC->UpdateViewForLockedActor();
+				LevelVC->Invalidate();
+			}
+		}
+		SetPerspectiveViewportCameraCutEnabled(true);
+
 		ResetPerMovieSceneData();
 		NotifyMovieSceneDataChanged();
 		UpdateTimeBoundsToFocusedMovieScene();
+
+		OnActivateSequenceEvent.Broadcast(*SequenceInstanceStack.Top());
 	}
 }
 
@@ -877,38 +921,14 @@ void FSequencer::NotifyMapChanged( class UWorld* NewWorld, EMapChangeType MapCha
 
 void FSequencer::OnActorsDropped( const TArray<TWeakObjectPtr<AActor> >& Actors )
 {
-	// @todo sequencer: Undo doesn't seem to be working at all
-	const FScopedTransaction Transaction(LOCTEXT("UndoPossessingObject", "Possess Object in Sequencer"));
-	GetFocusedMovieSceneSequence()->Modify();
-
-	bool bPossessableAdded = false;
-	for (TWeakObjectPtr<AActor> WeakActor : Actors)
-	{
-		if (AActor* Actor = WeakActor.Get())
-		{
-			FGuid ExistingGuid = SequenceInstanceStack.Top()->FindObjectId(*Actor);
-			if (!ExistingGuid.IsValid())
-			{
-				FGuid PossessableGuid = CreateBinding(*Actor, Actor->GetActorLabel());
-
-				UpdateRuntimeInstances();
-
-				OnActorAddedToSequencerEvent.Broadcast(Actor, PossessableGuid);
-			}
-			bPossessableAdded = true;
-		}
-	}
-	
-	if (bPossessableAdded)
-	{
-		NotifyMovieSceneDataChanged();
-	}
+	AddActors(Actors);
 }
 
 
 void FSequencer::NotifyMovieSceneDataChanged()
 {
 	LabelManager.SetMovieScene(SequenceInstanceStack.Top()->GetSequence()->GetMovieScene());
+	StoredPlaybackState = GetPlaybackStatus();
 	SetPlaybackStatus(EMovieScenePlayerStatus::Stopped);
 	bNeedTreeRefresh = true;
 
@@ -944,51 +964,110 @@ void FSequencer::SetClampRange(TRange<float> InNewClampRange)
 }
 
 
-TRange<float> FSequencer::GetInOutRange() const
+TRange<float> FSequencer::GetSelectionRange() const
 {
-	return GetFocusedMovieSceneSequence()->GetMovieScene()->GetInOutRange();
+	return GetFocusedMovieSceneSequence()->GetMovieScene()->GetSelectionRange();
 }
 
 
-void FSequencer::SetInOutRange(TRange<float> Range)
+void FSequencer::SetSelectionRange(TRange<float> Range)
 {
-	const FScopedTransaction Transaction(LOCTEXT("SetInOutRange_Transaction", "Set In/Out Range"));
-	UMovieScene* FocusedMovieScene = GetFocusedMovieSceneSequence()->GetMovieScene();
+	const FScopedTransaction Transaction(LOCTEXT("SetSelectionRange_Transaction", "Set Selection Range"));
+	UMovieScene* FocussedMovieScene = GetFocusedMovieSceneSequence()->GetMovieScene();
 
-	FocusedMovieScene->SetInOutRange(Range);
+	FocussedMovieScene->SetSelectionRange(Range);
 }
 
 
-void FSequencer::SetInOutRangeEnd()
+void FSequencer::SetSelectionRangeEnd()
 {
 	const float GlobalTime = GetGlobalTime();
 
-	if (GetInOutRange().GetLowerBoundValue() >= GlobalTime)
+	if (GetSelectionRange().GetLowerBoundValue() >= GlobalTime)
 	{
-		SetInOutRange(TRange<float>(GlobalTime));
+		SetSelectionRange(TRange<float>(GlobalTime));
 	}
 	else
 	{
-		SetInOutRange(TRange<float>(GetInOutRange().GetLowerBoundValue(), GlobalTime));
+		SetSelectionRange(TRange<float>(GetSelectionRange().GetLowerBoundValue(), GlobalTime));
 	}
 }
 
 
-void FSequencer::SetInOutRangeStart()
+void FSequencer::SetSelectionRangeStart()
 {
 	const float GlobalTime = GetGlobalTime();
 
-	if (GetInOutRange().GetUpperBoundValue() <= GlobalTime)
+	if (GetSelectionRange().GetUpperBoundValue() <= GlobalTime)
 	{
-		SetInOutRange(TRange<float>(GlobalTime));
+		SetSelectionRange(TRange<float>(GlobalTime));
 	}
 	else
 	{
-		SetInOutRange(TRange<float>(GlobalTime, GetInOutRange().GetUpperBoundValue()));
+		SetSelectionRange(TRange<float>(GlobalTime, GetSelectionRange().GetUpperBoundValue()));
 	}
 }
 
 
+void FSequencer::SelectKeysInSelectionRange(const TSharedRef<FSequencerDisplayNode>& DisplayNode, const TRange<float>& SelectionRange)
+{
+	if (DisplayNode->GetType() == ESequencerNode::Track)
+	{
+		TSet<TSharedPtr<IKeyArea>> OutKeyAreas;
+		SequencerHelpers::GetAllKeyAreas(DisplayNode, OutKeyAreas);
+
+		for (auto KeyArea : OutKeyAreas)
+		{
+			if (!KeyArea.IsValid())
+			{
+				continue;
+			}
+
+			UMovieSceneSection* Section = KeyArea->GetOwningSection();
+
+			if (Section == nullptr)
+			{
+				continue;
+			}
+
+			TSet<FKeyHandle> OutKeyHandles;
+			Section->GetKeyHandles(OutKeyHandles, SelectionRange);
+
+			for (auto KeyHandle : OutKeyHandles)
+			{
+				Selection.AddToSelection(FSequencerSelectedKey(*Section, KeyArea, KeyHandle));
+			}
+		}
+	}
+
+	for (const auto& ChildNode : DisplayNode->GetChildNodes())
+	{
+		SelectKeysInSelectionRange(ChildNode, SelectionRange);
+	}
+}
+
+
+void FSequencer::SelectKeysInSelectionRange()
+{
+	UMovieSceneSequence* Sequence = GetFocusedMovieSceneSequence();
+	UMovieScene* MovieScene = Sequence->GetMovieScene();
+	TRange<float> SelectionRange = MovieScene->GetSelectionRange();
+
+	Selection.Empty();
+
+	for (const TSharedRef<FSequencerDisplayNode>& DisplayNode : NodeTree->GetRootNodes())
+	{
+		SelectKeysInSelectionRange(DisplayNode, SelectionRange);
+	}
+}
+
+
+void FSequencer::DeleteSelectionRange()
+{
+
+}
+
+
 TRange<float> FSequencer::GetPlaybackRange() const
 {
 	return GetFocusedMovieSceneSequence()->GetMovieScene()->GetPlaybackRange();
@@ -1323,6 +1402,16 @@ ISequencer::FOnActorAddedToSequencer& FSequencer::OnActorAddedToSequencer()
 	return OnActorAddedToSequencerEvent;
 }
 
+ISequencer::FOnPreSave& FSequencer::OnPreSave()
+{
+	return OnPreSaveEvent;
+}
+
+ISequencer::FOnActivateSequence& FSequencer::OnActivateSequence()
+{
+	return OnActivateSequenceEvent;
+}
+
 ISequencer::FOnCameraCut& FSequencer::OnCameraCut()
 {
 	return OnCameraCutEvent;
@@ -1348,20 +1437,17 @@ FGuid FindUnspawnedObjectGuid(UObject& InObject, UMovieSceneSequence& Sequence)
 {
 	UMovieScene* MovieScene = Sequence.GetMovieScene();
 
-	// If the object is a CDO, the it relates to an unspawned spawnable.
+	// If the object is an archetype, the it relates to an unspawned spawnable.
 	UObject* ParentObject = Sequence.GetParentObject(&InObject);
-	if (ParentObject && ParentObject->HasAnyFlags(RF_ClassDefaultObject))
+	if (ParentObject && ParentObject->HasAnyFlags(RF_ArchetypeObject))
 	{
-		// Find the unspawned spawnable object's CDO
-		UClass* ObjectClass = ParentObject->GetClass();
-
 		FMovieSceneSpawnable* ParentSpawnable = MovieScene->FindSpawnable([&](FMovieSceneSpawnable& InSpawnable){
-			return InSpawnable.GetClass() == ObjectClass;
+			return InSpawnable.GetObjectTemplate() == ParentObject;
 		});
 
 		if (ParentSpawnable)
 		{
-			UObject* ParentContext = ObjectClass->GetDefaultObject();
+			UObject* ParentContext = ParentSpawnable->GetObjectTemplate();
 
 			// The only way to find the object now is to resolve all the child bindings, and see if they are the same
 			for (const FGuid& ChildGuid : ParentSpawnable->GetChildPossessables())
@@ -1374,17 +1460,15 @@ FGuid FindUnspawnedObjectGuid(UObject& InObject, UMovieSceneSequence& Sequence)
 			}
 		}
 	}
-	else if (InObject.HasAnyFlags(RF_ClassDefaultObject))
+	else if (InObject.HasAnyFlags(RF_ArchetypeObject))
 	{
-		UClass* ObjectClass = InObject.GetClass();
-
-		FMovieSceneSpawnable* SpawnableByClass = MovieScene->FindSpawnable([&](FMovieSceneSpawnable& InSpawnable){
-			return InSpawnable.GetClass() == ObjectClass;
+		FMovieSceneSpawnable* SpawnableByArchetype = MovieScene->FindSpawnable([&](FMovieSceneSpawnable& InSpawnable){
+			return InSpawnable.GetObjectTemplate() == &InObject;
 		});
 
-		if (SpawnableByClass)
+		if (SpawnableByArchetype)
 		{
-			return SpawnableByClass->GetGuid();
+			return SpawnableByArchetype->GetGuid();
 		}
 	}
 
@@ -1619,7 +1703,7 @@ void FSequencer::UpdateRuntimeInstances()
 	EMovieSceneUpdateData UpdateData(ScrubPosition, ScrubPosition);
 	SequenceInstanceStack.Top()->Update(UpdateData, *this);
 
-	// Ensure any relevant spawned objects are cleaned up if we're playing back the master sequence
+	// Ensure any relevant spawned objects are cleaned up if we're playing back the master sequenceg
 	if (SequenceInstanceStack.Num() == 1)
 	{
 		if (!GetPlaybackRange().Contains(ScrubPosition))
@@ -1718,12 +1802,32 @@ TSharedRef<SWidget> FSequencer::OnCreateTransportRecord()
 {
 	ISequenceRecorder& SequenceRecorder = FModuleManager::LoadModuleChecked<ISequenceRecorder>("SequenceRecorder");
 
-	return SNew(SButton)
+	TSharedRef<SButton> RecordButton = SNew(SButton)
 		.OnClicked(this, &FSequencer::OnRecord)
+		.ButtonStyle( FEditorStyle::Get(), "NoBorder" )
 		.ToolTipText_Lambda([&](){ return SequenceRecorder.IsRecording() ? LOCTEXT("StopRecord_Tooltip", "Stop recording current sub-track.") : LOCTEXT("Record_Tooltip", "Record the primed sequence sub-track."); })
-		.ButtonStyle(FEditorStyle::Get(), "Animation.Record")
-		.ContentPadding(2.0f)
-		.Visibility(this, &FSequencer::GetRecordButtonVisibility);
+		.Visibility(this, &FSequencer::GetRecordButtonVisibility)
+		.ContentPadding(2.0f);
+
+	TWeakPtr<SButton> WeakButton = RecordButton;
+
+	RecordButton->SetContent(SNew(SImage)
+		.Image_Lambda([&SequenceRecorder, WeakButton]()
+		{
+			if (SequenceRecorder.IsRecording())
+			{
+				return WeakButton.IsValid() && WeakButton.Pin()->IsPressed() ? 
+					&FEditorStyle::Get().GetWidgetStyle<FButtonStyle>("Animation.Recording").Pressed : 
+					&FEditorStyle::Get().GetWidgetStyle<FButtonStyle>("Animation.Recording").Normal;
+			}
+
+			return WeakButton.IsValid() && WeakButton.Pin()->IsPressed() ? 
+				&FEditorStyle::Get().GetWidgetStyle<FButtonStyle>("Animation.Record").Pressed : 
+				&FEditorStyle::Get().GetWidgetStyle<FButtonStyle>("Animation.Record").Normal;
+		})
+	);
+
+	return RecordButton;
 }
 
 namespace FaderConstants
@@ -1921,6 +2025,45 @@ void FSequencer::DetachTransportControlsFromViewports()
 	}
 }
 
+UObject* FSequencer::FindSpawnedObjectOrTemplate(const FGuid& BindingId) const
+{
+	if (UObject* Object = GetFocusedMovieSceneSequenceInstance()->FindObject(BindingId, *this))
+	{
+		return Object;
+	}
+
+	UMovieSceneSequence* Sequence = GetFocusedMovieSceneSequence();
+	if (!Sequence)
+	{
+		return nullptr;
+	}
+
+	UMovieScene* FocusedMovieScene = Sequence->GetMovieScene();
+
+	FMovieScenePossessable* Possessable = FocusedMovieScene->FindPossessable(BindingId);
+	// If we're a possessable with a parent spawnable and we don't have the object, we look the object up within the default object of the spawnable
+	if (Possessable && Possessable->GetParent().IsValid())
+	{
+		// If we're a spawnable and we don't have the object, use the default object to build up the track menu
+		FMovieSceneSpawnable* ParentSpawnable = FocusedMovieScene->FindSpawnable(Possessable->GetParent());
+		if (ParentSpawnable)
+		{
+			UObject* ParentObject = ParentSpawnable->GetObjectTemplate();
+			if (ParentObject)
+			{
+				return Sequence->FindPossessableObject(BindingId, ParentObject);
+			}
+		}
+	}
+	// If we're a spawnable and we don't have the object, use the default object to build up the track menu
+	else if (FMovieSceneSpawnable* Spawnable = FocusedMovieScene->FindSpawnable(BindingId))
+	{
+		return Spawnable->GetObjectTemplate();
+	}
+
+	return nullptr;
+}
+
 TSharedRef<FExtender> FSequencer::OnExtendLevelEditorViewMenu(const TSharedRef<FUICommandList> CommandList)
 {
 	TSharedRef<FExtender> Extender(new FExtender());
@@ -2105,7 +2248,7 @@ FReply FSequencer::SetPlaybackStart()
 	if (FocusedSequence)
 	{
 		TRange<float> CurrentRange = FocusedSequence->GetMovieScene()->GetPlaybackRange();
-		float NewPos = FMath::Max(GetGlobalTime(), CurrentRange.GetUpperBoundValue());
+		float NewPos = FMath::Min(GetGlobalTime(), CurrentRange.GetUpperBoundValue());
 		SetPlaybackRange(TRange<float>(NewPos, CurrentRange.GetUpperBoundValue()));	
 	}
 
@@ -2123,6 +2266,7 @@ FReply FSequencer::JumpToPreviousKey()
 		TOptional<float> NewTime = ActiveKeyCollection->FindFirstKeyInRange(FindRange, EFindKeyDirection::Backwards);
 		if (NewTime.IsSet())
 		{
+			SetPlaybackStatus(EMovieScenePlayerStatus::Stepping);
 			SetGlobalTime(NewTime.GetValue());
 		}
 	}
@@ -2141,6 +2285,7 @@ FReply FSequencer::JumpToNextKey()
 		TOptional<float> NewTime = ActiveKeyCollection->FindFirstKeyInRange(FindRange, EFindKeyDirection::Forwards);
 		if (NewTime.IsSet())
 		{
+			SetPlaybackStatus(EMovieScenePlayerStatus::Stepping);
 			SetGlobalTime(NewTime.GetValue());
 		}
 	}
@@ -2409,25 +2554,25 @@ void FSequencer::OnEndScrubbing()
 }
 
 
-void FSequencer::OnBeginPlaybackRangeDrag()
+void FSequencer::OnPlaybackRangeBeginDrag()
 {
 	GEditor->BeginTransaction(LOCTEXT("SetPlaybackRange_Transaction", "Set Playback Range"));
 }
 
 
-void FSequencer::OnEndPlaybackRangeDrag()
+void FSequencer::OnPlaybackRangeEndDrag()
 {
 	GEditor->EndTransaction();
 }
 
 
-void FSequencer::OnBeginInOutRangeDrag()
+void FSequencer::OnSelectionRangeBeginDrag()
 {
 	GEditor->BeginTransaction(LOCTEXT("SetSelectionRange_Transaction", "Set Selection Range"));
 }
 
 
-void FSequencer::OnEndInOutRangeDrag()
+void FSequencer::OnSelectionRangeEndDrag()
 {
 	GEditor->EndTransaction();
 }
@@ -2519,7 +2664,7 @@ FGuid FSequencer::AddSpawnable(UObject& Object)
 		NewSpawnable.Name += UniqueString;
 	}
 
-	FGuid NewGuid = OwnerMovieScene->AddSpawnable(NewSpawnable.Name, NewSpawnable.Blueprint);
+	FGuid NewGuid = OwnerMovieScene->AddSpawnable(NewSpawnable.Name, *NewSpawnable.ObjectTemplate);
 
 	UpdateRuntimeInstances();
 
@@ -2763,11 +2908,6 @@ void FSequencer::OnRequestNodeDeleted( TSharedRef<const FSequencerDisplayNode> N
 			Sequence->Modify();
 			Sequence->UnbindPossessableObjects( BindingToRemove );
 		}
-		else
-		{
-			// If this check fails there was an invalid guid being stored on a node
-			checkf(false, TEXT("Guid was not associated with a spawnable or possessable"));
-		}
 
 		bAnythingRemoved = true;
 	}
@@ -2909,6 +3049,185 @@ void FSequencer::OnMapOpened(const FString& Filename, bool bLoadAsTemplate)
 	ActivateSequencerEditorMode();
 }
 
+void FSequencer::OnLevelAdded(ULevel* InLevel, UWorld* InWorld)
+{
+	NotifyMovieSceneDataChanged();
+}
+
+void FSequencer::OnLevelRemoved(ULevel* InLevel, UWorld* InWorld)
+{
+	NotifyMovieSceneDataChanged();
+}
+
+void FSequencer::OnNewActorsDropped(const TArray<UObject*>& DroppedObjects, const TArray<AActor*>& DroppedActors)
+{
+	bool bAddSpawnable = FSlateApplication::Get().GetModifierKeys().IsShiftDown();
+	bool bAddPossessable = FSlateApplication::Get().GetModifierKeys().IsControlDown();
+
+	if (bAddSpawnable || bAddPossessable)
+	{
+		TArray<AActor*> SpawnedActors;
+
+		const FScopedTransaction Transaction(LOCTEXT("UndoAddActors", "Add Actors to Sequencer"));
+		
+		UMovieSceneSequence* Sequence = GetFocusedMovieSceneSequence();
+		UMovieScene* OwnerMovieScene = Sequence->GetMovieScene();
+
+		Sequence->Modify();
+
+		for ( AActor* Actor : DroppedActors )
+		{
+			AActor* NewActor = Actor;
+			bool bCreateAndAttachCamera = false;
+			if (NewActor->GetClass() == ACameraRig_Rail::StaticClass() ||
+				NewActor->GetClass() == ACameraRig_Crane::StaticClass())
+			{
+				bCreateAndAttachCamera = true;
+			}
+
+			FGuid PossessableGuid = CreateBinding(*NewActor, NewActor->GetActorLabel());
+			FGuid NewGuid = PossessableGuid;
+
+			if (bAddSpawnable)
+			{
+				FMovieSceneSpawnable* Spawnable = ConvertToSpawnableInternal(PossessableGuid);
+
+				UpdateRuntimeInstances();
+
+				UObject* RuntimeObject = GetFocusedMovieSceneSequenceInstance()->FindObject(Spawnable->GetGuid(), *this);
+
+				AActor* SpawnedActor = Cast<AActor>(RuntimeObject);
+				if (SpawnedActor)
+				{
+					SpawnedActors.Add(SpawnedActor);
+				}
+
+				NewGuid = Spawnable->GetGuid();
+				NewActor = SpawnedActor;
+			}
+			else
+			{
+				UpdateRuntimeInstances();
+
+				OnActorAddedToSequencerEvent.Broadcast(NewActor, PossessableGuid);
+			}
+
+			if (bCreateAndAttachCamera)
+			{
+				ACameraRig_Rail* RailActor = nullptr;
+				if (Actor->GetClass() == ACameraRig_Rail::StaticClass())
+				{
+					RailActor = Cast<ACameraRig_Rail>(NewActor);
+				}
+
+				// Create a cine camera actor
+				ACineCameraActor* NewCamera = GWorld->SpawnActor<ACineCameraActor>();
+				FGuid NewCameraGuid = CreateBinding(*NewCamera, NewCamera->GetActorLabel());
+
+				if (RailActor)
+				{
+					NewCamera->SetActorRotation(FRotator(0.f, -90.f, 0.f));
+				}
+
+				if (bAddSpawnable)
+				{
+					FMovieSceneSpawnable* Spawnable = ConvertToSpawnableInternal(NewCameraGuid);
+
+					UpdateRuntimeInstances();
+
+					UObject* RuntimeObject = GetFocusedMovieSceneSequenceInstance()->FindObject(Spawnable->GetGuid(), *this);
+
+					NewCamera = Cast<ACineCameraActor>(RuntimeObject);
+					NewCameraGuid = Spawnable->GetGuid();
+
+					// Create an attach track
+					UMovieScene3DAttachTrack* AttachTrack = Cast<UMovieScene3DAttachTrack>(OwnerMovieScene->AddTrack(UMovieScene3DAttachTrack::StaticClass(), NewCameraGuid));
+					AttachTrack->AddConstraint(GetPlaybackRange().GetLowerBoundValue(), GetPlaybackRange().GetUpperBoundValue(), NAME_None, NAME_None, NewGuid);
+				}
+				else
+				{
+					OnActorAddedToSequencerEvent.Broadcast(NewCamera, NewCameraGuid);
+
+					// Parent it
+					NewCamera->AttachToActor(NewActor, FAttachmentTransformRules::KeepRelativeTransform);
+				}
+
+				if (RailActor)
+				{
+					// Extend the rail a bit
+					if (RailActor->GetRailSplineComponent()->GetNumberOfSplinePoints() == 2)
+					{
+						FVector SplinePoint1 = RailActor->GetRailSplineComponent()->GetLocationAtSplinePoint(0, ESplineCoordinateSpace::Local);
+						FVector SplinePoint2 = RailActor->GetRailSplineComponent()->GetLocationAtSplinePoint(1, ESplineCoordinateSpace::Local);
+						FVector SplineDirection = SplinePoint2 - SplinePoint1;
+						SplineDirection.Normalize();
+
+						float DefaultRailDistance = 650.f;
+						SplinePoint2 = SplinePoint1 + SplineDirection* DefaultRailDistance;
+						RailActor->GetRailSplineComponent()->SetLocationAtSplinePoint(1, SplinePoint2, ESplineCoordinateSpace::Local);
+						RailActor->GetRailSplineComponent()->bSplineHasBeenEdited = true;
+					}
+
+					// Create a track for the CurrentPositionOnRail
+					TArray<UProperty*> PropertyPath;
+					UProperty* CurrentPositionOnRailProperty = RailActor->GetClass()->FindPropertyByName(TEXT("CurrentPositionOnRail"));
+					PropertyPath.Add(CurrentPositionOnRailProperty);
+
+					FKeyPropertyParams KeyPropertyParams(TArrayBuilder<UObject*>().Add(RailActor), PropertyPath);
+					{
+						KeyPropertyParams.KeyParams.bCreateTrackIfMissing = true;
+						KeyPropertyParams.KeyParams.bCreateHandleIfMissing = true;
+						KeyPropertyParams.KeyParams.bCreateKeyIfUnchanged = false;
+						KeyPropertyParams.KeyParams.bCreateKeyIfEmpty = true;
+						KeyPropertyParams.KeyParams.bCreateKeyOnlyWhenAutoKeying = false;
+					}
+
+					float OriginalTime = GetGlobalTime();
+
+					SetGlobalTime(GetPlaybackRange().GetLowerBoundValue());
+					RailActor->CurrentPositionOnRail = 0.f;
+					KeyProperty(KeyPropertyParams);
+
+					SetGlobalTime(GetPlaybackRange().GetUpperBoundValue());
+					RailActor->CurrentPositionOnRail = 1.f;
+					KeyProperty(KeyPropertyParams);
+
+					SetGlobalTime(OriginalTime);
+				}
+
+				// New camera added, don't lock the view to the camera because we want to see where the camera rig was placed
+				const bool bLockToCamera = false;
+				NewCameraAdded(NewCamera, NewCameraGuid, bLockToCamera);
+			}
+		}
+
+		if (SpawnedActors.Num())
+		{
+			const bool bNotifySelectionChanged = true;
+			const bool bDeselectBSP = true;
+			const bool bWarnAboutTooManyActors = false;
+			const bool bSelectEvenIfHidden = false;
+	
+			GEditor->GetSelectedActors()->Modify();
+			GEditor->GetSelectedActors()->BeginBatchSelectOperation();
+			GEditor->SelectNone( bNotifySelectionChanged, bDeselectBSP, bWarnAboutTooManyActors );
+			for (auto SpawnedActor : SpawnedActors)
+			{
+				GEditor->SelectActor( SpawnedActor, true, bNotifySelectionChanged, bSelectEvenIfHidden );
+			}
+			GEditor->GetSelectedActors()->EndBatchSelectOperation();
+			GEditor->NoteSelectionChange();
+		}
+
+		SequencerWidget->UpdateLayoutTree();
+
+		bUpdatingSequencerSelection = true;
+		SynchronizeSequencerSelectionWithExternalSelection();
+		bUpdatingSequencerSelection = false;
+
+		NotifyMovieSceneDataChanged();
+	}
+}
 
 void FSequencer::ActivateSequencerEditorMode()
 {
@@ -2930,12 +3249,18 @@ void FSequencer::OnPreBeginPIE(bool bIsSimulating)
 }
 
 
-void FSequencer::OnEndPIE(bool bIsSimulating)
+void FSequencer::OnEndPlayMap()
 {
 	UpdateRuntimeInstances();
 }
 
 
+void FSequencer::OnEndPIE(bool bIsSimulating)
+{
+	OnEndPlayMap();
+}
+
+
 void FSequencer::ActivateDetailKeyframeHandler()
 {			
 	// Add sequencer detail keyframe handler
@@ -3061,6 +3386,8 @@ void GetDescendantMovieScenes(UMovieSceneSequence* InSequence, TArray<UMovieScen
 
 void FSequencer::SaveCurrentMovieScene()
 {
+	OnPreSaveEvent.Broadcast(*this);
+
 	// Capture thumbnail
 	// Convert UObject* array to FAssetData array
 	TArray<FAssetData> AssetDataList;
@@ -3120,6 +3447,12 @@ void FSequencer::SaveCurrentMovieScene()
 	const bool bPromptToSave = PackagesToSave.Num() > 1;
 
 	FEditorFileUtils::PromptForCheckoutAndSave( PackagesToSave, bCheckDirty, bPromptToSave );
+
+	// Ensure the sequence is up-to-date
+	EMovieSceneUpdateData UpdateData(ScrubPosition, ScrubPosition);
+	SequenceInstanceStack.Top()->Update(UpdateData, *this);
+
+	FEditorSupportDelegates::RedrawAllViewports.Broadcast();
 }
 
 
@@ -3147,34 +3480,16 @@ void FSequencer::SaveCurrentMovieSceneAs()
 }
 
 
-void FSequencer::AddSelectedObjects()
+void FSequencer::AddActors(const TArray<TWeakObjectPtr<AActor> >& InActors)
 {
-	UMovieSceneSequence* OwnerSequence = GetFocusedMovieSceneSequence();
-	UMovieScene* OwnerMovieScene = OwnerSequence->GetMovieScene();
-	USelection* CurrentSelection = GEditor->GetSelectedActors();
-
-	TArray<AActor*> SelectedActors;
-	TArray<UObject*> SelectedObjects;
-
-	CurrentSelection->GetSelectedObjects( AActor::StaticClass(), SelectedObjects );
-
-	for (TArray<UObject*>::TIterator It(SelectedObjects); It; ++It)
-	{
-		AActor* Actor = CastChecked<AActor>(*It);
-
-		if (Actor != nullptr)
-		{
-			SelectedActors.Add(Actor);
-		}
-	}
+	// @todo sequencer: Undo doesn't seem to be working at all
+	const FScopedTransaction Transaction(LOCTEXT("UndoPossessingObject", "Possess Object in Sequencer"));
+	GetFocusedMovieSceneSequence()->Modify();
 
 	bool bPossessableAdded = false;
-	if (SelectedActors.Num() != 0)
+	for (TWeakObjectPtr<AActor> WeakActor : InActors)
 	{
-		const FScopedTransaction Transaction(LOCTEXT("UndoPossessingObject", "Possess Object in Sequencer"));
-		OwnerSequence->Modify();
-
-		for (auto Actor : SelectedActors)
+		if (AActor* Actor = WeakActor.Get())
 		{
 			FGuid ExistingGuid = SequenceInstanceStack.Top()->FindObjectId(*Actor);
 			if (!ExistingGuid.IsValid())
@@ -3184,14 +3499,19 @@ void FSequencer::AddSelectedObjects()
 				UpdateRuntimeInstances();
 
 				OnActorAddedToSequencerEvent.Broadcast(Actor, PossessableGuid);
-
-				bPossessableAdded = true;
 			}
+			bPossessableAdded = true;
 		}
 	}
 
 	if (bPossessableAdded)
 	{
+		SequencerWidget->UpdateLayoutTree();
+
+		bUpdatingSequencerSelection = true;
+		SynchronizeSequencerSelectionWithExternalSelection();
+		bUpdatingSequencerSelection = false;
+
 		NotifyMovieSceneDataChanged();
 	}
 }
@@ -3587,7 +3907,7 @@ void FSequencer::DoAssignActor(AActor*const* InActors, int32 NumActors, FGuid In
 		}
 	}
 
-	// Try to fix up folders.
+	// Try to fix up folders
 	TArray<UMovieSceneFolder*> FoldersToCheck;
 	FoldersToCheck.Append(GetFocusedMovieSceneSequence()->GetMovieScene()->GetRootFolders());
 	bool bFolderFound = false;
@@ -3737,8 +4057,12 @@ void FSequencer::ConvertToSpawnable(TSharedRef<FSequencerObjectBindingNode> Node
 	// Ensure we're in a non-possessed state
 	RootMovieSceneSequenceInstance->RestoreState(*this);
 	GetFocusedMovieSceneSequence()->GetMovieScene()->Modify();
-	ConvertToSpawnableInternal(NodeToBeConverted);
-	NotifyMovieSceneDataChanged();
+	FMovieScenePossessable* Possessable = GetFocusedMovieSceneSequence()->GetMovieScene()->FindPossessable(NodeToBeConverted->GetObjectBinding());
+	if (Possessable)
+	{
+		ConvertToSpawnableInternal(Possessable->GetGuid());
+		NotifyMovieSceneDataChanged();
+	}
 }
 
 void FSequencer::ConvertSelectedNodesToSpawnables()
@@ -3772,10 +4096,23 @@ void FSequencer::ConvertSelectedNodesToSpawnables()
 	FScopedSlowTask SlowTask(ObjectBindingNodes.Num(), LOCTEXT("ConvertSpawnableProgress", "Converting Selected Possessable Nodes to Spawnables"));
 	SlowTask.MakeDialog(true);
 
+	TArray<AActor*> SpawnedActors;
 	for (const TSharedRef<FSequencerObjectBindingNode>& ObjectBindingNode : ObjectBindingNodes)
 	{
 		SlowTask.EnterProgressFrame();
-		ConvertToSpawnableInternal(ObjectBindingNode);
+	
+		FMovieScenePossessable* Possessable = MovieScene->FindPossessable(ObjectBindingNode->GetObjectBinding());
+		if (Possessable)
+		{
+			FMovieSceneSpawnable* Spawnable = ConvertToSpawnableInternal(Possessable->GetGuid());
+
+			UpdateRuntimeInstances();
+
+			UObject* RuntimeObject = GetFocusedMovieSceneSequenceInstance()->FindObject(Spawnable->GetGuid(), *this);
+
+			AActor* SpawnedActor = Cast<AActor>(RuntimeObject);
+			SpawnedActors.Add(SpawnedActor);
+		}
 
 		if (GWarn->ReceivedUserCancel())
 		{
@@ -3783,28 +4120,41 @@ void FSequencer::ConvertSelectedNodesToSpawnables()
 		}
 	}
 
-	NotifyMovieSceneDataChanged();
+	if (SpawnedActors.Num())
+	{
+		const bool bNotifySelectionChanged = true;
+		const bool bDeselectBSP = true;
+		const bool bWarnAboutTooManyActors = false;
+		const bool bSelectEvenIfHidden = false;
+
+		GEditor->GetSelectedActors()->Modify();
+		GEditor->GetSelectedActors()->BeginBatchSelectOperation();
+		GEditor->SelectNone(bNotifySelectionChanged, bDeselectBSP, bWarnAboutTooManyActors);
+		for (auto SpawnedActor : SpawnedActors)
+		{
+			GEditor->SelectActor(SpawnedActor, true, bNotifySelectionChanged, bSelectEvenIfHidden);
+		}
+		GEditor->GetSelectedActors()->EndBatchSelectOperation();
+		GEditor->NoteSelectionChange();
+
+		NotifyMovieSceneDataChanged();
+	}
 }
 
-void FSequencer::ConvertToSpawnableInternal(TSharedRef<FSequencerObjectBindingNode> NodeToBeConverted)
+FMovieSceneSpawnable* FSequencer::ConvertToSpawnableInternal(FGuid PossessableGuid)
 {
 	UMovieSceneSequence* Sequence = GetFocusedMovieSceneSequence();
 	UMovieScene* MovieScene = Sequence->GetMovieScene();
-	FMovieScenePossessable* Possessable = MovieScene->FindPossessable(NodeToBeConverted->GetObjectBinding());
-
-	if (!Possessable)
-	{
-		return;
-	}
 
 	// Find the object in the environment
 	TSharedRef<FMovieSceneSequenceInstance> FocusedSequenceInstance = SequenceInstanceStack.Top();
 	UMovieSceneSequence* FocusedSequence = FocusedSequenceInstance->GetSequence();
+	FMovieScenePossessable* Possessable = MovieScene->FindPossessable(PossessableGuid);
 
-	UObject* FoundObject = FocusedSequenceInstance->FindObject(NodeToBeConverted->GetObjectBinding(), *this);
+	UObject* FoundObject = FocusedSequenceInstance->FindObject(PossessableGuid, *this);
 	if (!FoundObject || !FocusedSequence)
 	{
-		return;
+		return nullptr;
 	}
 
 	Sequence->Modify();
@@ -3812,7 +4162,7 @@ void FSequencer::ConvertToSpawnableInternal(TSharedRef<FSequencerObjectBindingNo
 	FMovieSceneSpawnable* Spawnable = MovieScene->FindSpawnable(AddSpawnable(*FoundObject));
 	if (!Spawnable)
 	{
-		return;
+		return nullptr;
 	}
 
 	// Swap the guids, so the possessable's tracks now belong to the spawnable
@@ -3850,7 +4200,7 @@ void FSequencer::ConvertToSpawnableInternal(TSharedRef<FSequencerObjectBindingNo
 	
 	FTransformData DefaultTransform;
 
-	// @todo: Where should this code live (not here)?
+	// @todo: Where should this code live (not here)? Probably in an editor-only implementation inside FLevelSequenceSpawnRegister
 	AActor* OldActor = Cast<AActor>(FoundObject);
 	if (OldActor)
 	{
@@ -3859,14 +4209,14 @@ void FSequencer::ConvertToSpawnableInternal(TSharedRef<FSequencerObjectBindingNo
 		DefaultTransform.Scale = OldActor->GetActorScale();
 		DefaultTransform.bValid = true;
 
-		AActor* CDO = CastChecked<AActor>(Spawnable->GetClass()->GetDefaultObject());
-
-		EditorUtilities::CopyActorProperties(OldActor, CDO, EditorUtilities::ECopyOptions::Default);
-		GWorld->DestroyActor(OldActor);
+		GEditor->SelectActor(OldActor, false, true);
+		GWorld->EditorDestroyActor(OldActor, true);
 	}
 
 	SetupDefaultsForSpawnable(Spawnable->GetGuid(), DefaultTransform);
 	SetGlobalTimeDirectly(ScrubPosition);
+
+	return Spawnable;
 }
 
 
@@ -4279,25 +4629,41 @@ void FSequencer::DiscardChanges()
 	// close asset editor
 	AssetEditorManager.CloseAllEditorsForAsset(EditedSequence);
 
-	// move existing objects into transient package
-	UPackage* const TransientPackage = GetTransientPackage();
+	// collect objects to be unloaded
 	TMap<FString, UObject*> MovedObjects;
 
 	ForEachObjectWithOuter(SequencePackage, [&](UObject* Object) {
-		const FName OldName = Object->GetFName();
-		const FName NewName = MakeUniqueObjectName(TransientPackage, Object->GetClass(), OldName);
-		Object->Rename(*NewName.ToString(), TransientPackage, REN_DontCreateRedirectors | REN_DoNotDirty);
-		Object->SetFlags(RF_Transient);
-		Object->ClearFlags(RF_Standalone | RF_Transactional);
 		MovedObjects.Add(Object->GetPathName(), Object);
 	}, true);
 
-	// unload package
-	TArray<UPackage*> PackagesToUnload;
+	// move objects into transient package
+	UPackage* const TransientPackage = GetTransientPackage();
+
+	for (auto MovedObject : MovedObjects)
+	{
+		UObject* Object = MovedObject.Value;
+
+		const FString OldName = Object->GetName();
+		const FString NewName = FString::Printf(TEXT("UNLOADING_%s"), *OldName);
+		const FName UniqueName = MakeUniqueObjectName(TransientPackage, Object->GetClass(), FName(*NewName));
+		UObject* NewOuter = (Object->GetOuter() == SequencePackage) ? TransientPackage : Object->GetOuter();
+	
+		Object->Rename(*UniqueName.ToString(), NewOuter, REN_DontCreateRedirectors | REN_DoNotDirty | REN_NonTransactional);
+		Object->SetFlags(RF_Transient);
+		Object->ClearFlags(RF_Standalone | RF_Transactional);
+	}
+
+	for (auto MovedObject : MovedObjects)
 	{
-		PackagesToUnload.Add(SequencePackage);
+		GLog->Logf(TEXT("Moved %s ---------> %s"), *MovedObject.Key, *MovedObject.Value->GetPathName());
 	}
 
+	// unload package
+	SequencePackage->SetDirtyFlag(false);
+
+	TArray<UPackage*> PackagesToUnload;
+	PackagesToUnload.Add(SequencePackage);
+
 	FText PackageUnloadError;
 	PackageTools::UnloadPackages(PackagesToUnload, PackageUnloadError);
 
@@ -4314,8 +4680,7 @@ void FSequencer::DiscardChanges()
 
 	for (const auto MovedObject : MovedObjects)
 	{
-		UObject* ReloadedObject = StaticLoadObject(MovedObject.Value->GetClass(), nullptr, *MovedObject.Key);
-		check(ReloadedObject != nullptr);
+		UObject* ReloadedObject = StaticLoadObject(MovedObject.Value->GetClass(), nullptr, *MovedObject.Key, nullptr);//, LOAD_NoWarn);
 		MovedToReloadedObjectMap.Add(MovedObject.Value, ReloadedObject);
 	}
 
@@ -4325,8 +4690,7 @@ void FSequencer::DiscardChanges()
 		FArchiveReplaceObjectRef<UObject> Ar(*It, MovedToReloadedObjectMap, false, false, false, false);
 	}
 
-	auto ReloadedSequence = Cast<UMovieSceneSequence>(StaticLoadObject(SequenceClass, nullptr, *SequencePath));
-	SequencePackage->SetDirtyFlag(false);
+	auto ReloadedSequence = Cast<UMovieSceneSequence>(StaticLoadObject(SequenceClass, nullptr, *SequencePath, nullptr));//, LOAD_NoWarn));
 
 	// release transient objects
 	for (auto MovedObject : MovedObjects)
@@ -4335,6 +4699,8 @@ void FSequencer::DiscardChanges()
 		MovedObject.Value->MarkPendingKill();
 	}
 
+//	CollectGarbage(GARBAGE_COLLECTION_KEEPFLAGS);
+
 	// clear undo buffer
 	if (true) // @todo sequencer: check whether objects are actually referenced in undo buffer
 	{
@@ -4344,10 +4710,96 @@ void FSequencer::DiscardChanges()
 	// reopen asset editor
 	TArray<UObject*> AssetsToReopen;
 	AssetsToReopen.Add(ReloadedSequence);
+
 	AssetEditorManager.OpenEditorForAssets(AssetsToReopen, EToolkitMode::Standalone, MyToolkitHost.ToSharedRef());
 }
 
 
+void FSequencer::CreateCamera()
+{
+	if (!GCurrentLevelEditingViewportClient)
+	{
+		return;
+	}
+
+	const FScopedTransaction Transaction(NSLOCTEXT("Sequencer", "CreateCameraHere", "Create Camera Here"));
+
+	// Set new camera to match viewport
+	ACineCameraActor* NewCamera = GCurrentLevelEditingViewportClient->GetWorld()->SpawnActor<ACineCameraActor>();
+	if (!NewCamera)
+	{
+		return;
+	}
+	NewCamera->SetActorLocation( GCurrentLevelEditingViewportClient->GetViewLocation(), false );
+	NewCamera->SetActorRotation( GCurrentLevelEditingViewportClient->GetViewRotation() );
+	//pNewCamera->CameraComponent->FieldOfView = ViewportClient->ViewFOV; //@todo set the focal length from this field of view
+	FGuid PossessableGuid = CreateBinding(*NewCamera, NewCamera->GetActorLabel());
+	if (!PossessableGuid.IsValid())
+	{
+		return;
+	}
+	OnActorAddedToSequencerEvent.Broadcast(NewCamera, PossessableGuid);
+
+	const bool bLockToCamera = true;
+	NewCameraAdded(NewCamera, PossessableGuid, bLockToCamera);
+}
+
+void FSequencer::NewCameraAdded(ACineCameraActor* NewCamera, FGuid PossessableGuid, bool bLockToCamera)
+{
+	SetPerspectiveViewportCameraCutEnabled(false);
+
+	// Lock the viewport to this camera
+	if (bLockToCamera)
+	{
+		GCurrentLevelEditingViewportClient->SetMatineeActorLock(nullptr);
+		GCurrentLevelEditingViewportClient->SetActorLock(NewCamera);
+		GCurrentLevelEditingViewportClient->bLockedCameraView = true;
+		GCurrentLevelEditingViewportClient->UpdateViewForLockedActor();
+		GCurrentLevelEditingViewportClient->Invalidate();
+	}
+
+	UMovieSceneSequence* Sequence = GetFocusedMovieSceneSequence();
+	UMovieScene* OwnerMovieScene = Sequence->GetMovieScene();
+
+	// If there's a cinematic shot track, no need to set this camera to a shot
+	UMovieSceneTrack* CinematicShotTrack = OwnerMovieScene->FindMasterTrack(UMovieSceneCinematicShotTrack::StaticClass());
+	if (CinematicShotTrack)
+	{
+		return;
+	}
+
+	// If there's a camera cut track, create or set the camera section to this new camera
+	UMovieSceneTrack* CameraCutTrack = OwnerMovieScene->GetCameraCutTrack();
+
+	if (!CameraCutTrack)
+	{
+		CameraCutTrack = OwnerMovieScene->AddCameraCutTrack(UMovieSceneCameraCutTrack::StaticClass());
+	}
+
+	if (CameraCutTrack)
+	{
+		UMovieSceneSection* Section = MovieSceneHelpers::FindSectionAtTime(CameraCutTrack->GetAllSections(), GetGlobalTime());
+		UMovieSceneCameraCutSection* CameraCutSection = Cast<UMovieSceneCameraCutSection>(Section);
+
+		if (CameraCutSection)
+		{
+			CameraCutSection->Modify();
+			CameraCutSection->SetCameraGuid(PossessableGuid);
+		}
+		else
+		{
+			CameraCutTrack->Modify();
+
+			UMovieSceneCameraCutSection* NewSection = Cast<UMovieSceneCameraCutSection>(CameraCutTrack->CreateNewSection());
+			NewSection->SetStartTime(GetPlaybackRange().GetLowerBoundValue());
+			NewSection->SetEndTime(GetPlaybackRange().GetUpperBoundValue());
+			NewSection->SetCameraGuid(PossessableGuid);
+			CameraCutTrack->AddSection(*NewSection);
+		}
+	}
+}
+
+
 void FSequencer::FixActorReferences()
 {
 	UMovieScene* FocusedMovieScene = GetFocusedMovieSceneSequence()->GetMovieScene();
@@ -4449,7 +4901,7 @@ void FSequencer::SetHotspot(TSharedPtr<ISequencerHotspot> NewHotspot)
 	Hotspot = MoveTemp(NewHotspot);
 }
 
-void FSequencer::BindSequencerCommands()
+void FSequencer::BindCommands()
 {
 	const FSequencerCommands& Commands = FSequencerCommands::Get();
 
@@ -4513,26 +4965,6 @@ void FSequencer::BindSequencerCommands()
 		Commands.SetKey,
 		FExecuteAction::CreateSP( this, &FSequencer::SetKey ) );
 
-	SequencerCommandBindings->MapAction(
-		Commands.SetInterpolationCubicAuto,
-		FExecuteAction::CreateSP( this, &FSequencer::SetInterpTangentMode, ERichCurveInterpMode::RCIM_Cubic, ERichCurveTangentMode::RCTM_Auto ) );
-
-	SequencerCommandBindings->MapAction(
-		Commands.SetInterpolationCubicUser,
-		FExecuteAction::CreateSP( this, &FSequencer::SetInterpTangentMode, ERichCurveInterpMode::RCIM_Cubic, ERichCurveTangentMode::RCTM_User ) );
-
-	SequencerCommandBindings->MapAction(
-		Commands.SetInterpolationCubicBreak,
-		FExecuteAction::CreateSP( this, &FSequencer::SetInterpTangentMode, ERichCurveInterpMode::RCIM_Cubic, ERichCurveTangentMode::RCTM_Break ) );
-
-	SequencerCommandBindings->MapAction(
-		Commands.SetInterpolationLinear,
-		FExecuteAction::CreateSP( this, &FSequencer::SetInterpTangentMode, ERichCurveInterpMode::RCIM_Linear, ERichCurveTangentMode::RCTM_Auto ) );
-
-	SequencerCommandBindings->MapAction(
-		Commands.SetInterpolationConstant,
-		FExecuteAction::CreateSP( this, &FSequencer::SetInterpTangentMode, ERichCurveInterpMode::RCIM_Constant, ERichCurveTangentMode::RCTM_Auto ) );
-
 	SequencerCommandBindings->MapAction(
 		Commands.TrimSectionLeft,
 		FExecuteAction::CreateSP( this, &FSequencer::TrimSection, true ) );
@@ -4777,6 +5209,11 @@ void FSequencer::BindSequencerCommands()
 		})
 	);
 
+	SequencerCommandBindings->MapAction(
+		Commands.CreateCamera,
+		FExecuteAction::CreateSP( this, &FSequencer::CreateCamera ),
+		FCanExecuteAction::CreateLambda( []{ return true; } ) );
+
 	SequencerCommandBindings->MapAction(
 		Commands.DiscardChanges,
 		FExecuteAction::CreateSP(this, &FSequencer::DiscardChanges),
@@ -4807,6 +5244,26 @@ void FSequencer::BindSequencerCommands()
 	*SequencerSharedBindings = *SequencerCommandBindings;
 
 	// Sequencer-only bindings
+	SequencerCommandBindings->MapAction(
+		Commands.SetInterpolationCubicAuto,
+		FExecuteAction::CreateSP(this, &FSequencer::SetInterpTangentMode, ERichCurveInterpMode::RCIM_Cubic, ERichCurveTangentMode::RCTM_Auto));
+
+	SequencerCommandBindings->MapAction(
+		Commands.SetInterpolationCubicUser,
+		FExecuteAction::CreateSP(this, &FSequencer::SetInterpTangentMode, ERichCurveInterpMode::RCIM_Cubic, ERichCurveTangentMode::RCTM_User));
+
+	SequencerCommandBindings->MapAction(
+		Commands.SetInterpolationCubicBreak,
+		FExecuteAction::CreateSP(this, &FSequencer::SetInterpTangentMode, ERichCurveInterpMode::RCIM_Cubic, ERichCurveTangentMode::RCTM_Break));
+
+	SequencerCommandBindings->MapAction(
+		Commands.SetInterpolationLinear,
+		FExecuteAction::CreateSP(this, &FSequencer::SetInterpTangentMode, ERichCurveInterpMode::RCIM_Linear, ERichCurveTangentMode::RCTM_Auto));
+
+	SequencerCommandBindings->MapAction(
+		Commands.SetInterpolationConstant,
+		FExecuteAction::CreateSP(this, &FSequencer::SetInterpTangentMode, ERichCurveInterpMode::RCIM_Constant, ERichCurveTangentMode::RCTM_Auto));
+
 	SequencerCommandBindings->MapAction(
 		FGenericCommands::Get().Delete,
 		FExecuteAction::CreateSP( this, &FSequencer::DeleteSelectedItems ));
@@ -4834,13 +5291,22 @@ void FSequencer::BindSequencerCommands()
 		EUIActionRepeatMode::RepeatEnabled );
 
 	SequencerCommandBindings->MapAction(
-		Commands.SetInOutEnd,
-		FExecuteAction::CreateLambda([this]{ SetInOutRangeEnd(); }));
+		Commands.SetSelectionRangeEnd,
+		FExecuteAction::CreateLambda([this]{ SetSelectionRangeEnd(); }));
+
+	SequencerCommandBindings->MapAction(
+		Commands.SetSelectionRangeStart,
+		FExecuteAction::CreateLambda([this]{ SetSelectionRangeStart(); }));
+
+	SequencerCommandBindings->MapAction(
+		Commands.SelectKeysInSelectionRange,
+		FExecuteAction::CreateLambda([this]{ SelectKeysInSelectionRange(); }));
 
 	SequencerCommandBindings->MapAction(
-		Commands.SetInOutStart,
-		FExecuteAction::CreateLambda([this]{ SetInOutRangeStart(); }));
+		Commands.DeleteSelectionRange,
+		FExecuteAction::CreateLambda([this]{ DeleteSelectionRange(); }));
 
+	// bind widget specific commands
 	SequencerWidget->BindCommands(SequencerCommandBindings);
 }
 
diff --git a/Engine/Source/Editor/Sequencer/Private/Sequencer.h b/Engine/Source/Editor/Sequencer/Private/Sequencer.h
--- a/Engine/Source/Editor/Sequencer/Private/Sequencer.h
+++ b/Engine/Source/Editor/Sequencer/Private/Sequencer.h
@@ -8,6 +8,7 @@
 #include "SequencerLabelManager.h"
 #include "LevelEditor.h"
 
+class ACineCameraActor;
 class FMenuBuilder;
 class FMovieSceneSequenceInstance;
 class IDetailKeyframeHandler;
@@ -68,34 +69,40 @@ public:
 public:
 
 	/**
-	 * Get the in/out selection range.
+	 * Get the selection range.
 	 *
-	 * @return The in/out range.
-	 * @see SetInOutRange, SetInOutRangeEnd, SetInOutRangeStart
+	 * @return The selection range.
+	 * @see SetSelectionRange, SetSelectionRangeEnd, SetSelectionRangeStart
 	 */
-	TRange<float> GetInOutRange() const;
+	TRange<float> GetSelectionRange() const;
 
 	/**
-	 * Set the in/out selection range.
+	 * Set the selection selection range.
 	 *
 	 * @param Range The new range to set.
-	 * @see GetInOutRange, SetInOutRangeEnd, SetInOutRangeStart
+	 * @see GetSelectionRange, SetSelectionRangeEnd, SetSelectionRangeStart
 	 */
-	void SetInOutRange(TRange<float> Range);
+	void SetSelectionRange(TRange<float> Range);
 	
 	/**
-	 * Set the in/out range's end position to the current global time.
+	 * Set the selection range's end position to the current global time.
 	 *
-	 * @see GetInOutRange, SetInOutRange, SetInOutRangeStart
+	 * @see GetSelectionRange, SetSelectionRange, SetSelectionRangeStart
 	 */
-	void SetInOutRangeEnd();
+	void SetSelectionRangeEnd();
 	
 	/**
-	 * Set the in/out range's start position to the current global time.
+	 * Set the selection range's start position to the current global time.
 	 *
-	 * @see GetInOutRange, SetInOutRange, SetInOutRangeEnd
+	 * @see GetSelectionRange, SetSelectionRange, SetSelectionRangeEnd
 	 */
-	void SetInOutRangeStart();
+	void SetSelectionRangeStart();
+
+	/** Select all keys that fall into the current selection range. */
+	void SelectKeysInSelectionRange();
+
+	/** Delete the time range and all the keys that fall into the selection range. */
+	void DeleteSelectionRange();
 
 public:
 
@@ -364,9 +371,6 @@ public:
 	/** Called to save the current movie scene under a new name */
 	void SaveCurrentMovieSceneAs();
 
-	/** Called to add selected objects to the movie scene */
-	void AddSelectedObjects();
-
 	/** Called when a user executes the assign actor to track menu item */
 	void AssignActor(FMenuBuilder& MenuBuilder, FGuid ObjectBinding);
 	void DoAssignActor(AActor*const* InActors, int32 NumActors, FGuid ObjectBinding);
@@ -419,6 +423,12 @@ public:
 	/** Discard all changes to the current movie scene. */
 	void DiscardChanges();
 
+	/** Create camera and set it as the current camera cut. */
+	void CreateCamera();
+
+	/** Called when a new camera is added */
+	void NewCameraAdded(ACineCameraActor* NewCamera, FGuid PossessableGuid, bool bLockToCamera);
+
 	/** Attempts to automatically fix up broken actor references in the current scene. */
 	void FixActorReferences();
 
@@ -502,6 +512,7 @@ public:
 	virtual void NotifyMovieSceneDataChanged() override;
 	virtual void UpdateRuntimeInstances() override;
 	virtual void UpdatePlaybackRange() override;
+	virtual void AddActors(const TArray<TWeakObjectPtr<AActor> >& InActors) override;
 	virtual void AddSubSequence(UMovieSceneSequence* Sequence) override;
 	virtual bool CanKeyProperty(FCanKeyPropertyParams CanKeyPropertyParams) const override;
 	virtual void KeyProperty(FKeyPropertyParams KeyPropertyParams) override;
@@ -513,12 +524,15 @@ public:
 	virtual UObject* GetPlaybackContext() const override;
 	virtual void GetAllKeyedProperties(UObject& Object, TSet<UProperty*>& OutProperties) override;
 	virtual FOnActorAddedToSequencer& OnActorAddedToSequencer() override;
+	virtual FOnPreSave& OnPreSave() override;
+	virtual FOnActivateSequence& OnActivateSequence() override;
 	virtual FOnCameraCut& OnCameraCut() override;
 	virtual TSharedRef<INumericTypeInterface<float>> GetNumericTypeInterface() override;
 	virtual TSharedRef<INumericTypeInterface<float>> GetZeroPadNumericTypeInterface() override;
 	virtual TSharedRef<SWidget> MakeTransportControls(bool bExtended) override;
 	virtual TSharedRef<SWidget> MakeTimeRange(const TSharedRef<SWidget>& InnerContent, bool bShowWorkingRange, bool bShowViewRange, bool bShowPlaybackRange) override;
 	virtual void SetViewportTransportControlsVisibility(bool bVisible) override;
+	virtual UObject* FindSpawnedObjectOrTemplate(const FGuid& BindingId) const override;
 	
 public:
 
@@ -583,16 +597,16 @@ protected:
 	void OnEndScrubbing();
 
 	/** Called when the user has begun dragging the playback range */
-	void OnBeginPlaybackRangeDrag();
+	void OnPlaybackRangeBeginDrag();
 
 	/** Called when the user has finished dragging the playback range */
-	void OnEndPlaybackRangeDrag();
+	void OnPlaybackRangeEndDrag();
 
 	/** Called when the user has begun dragging the selection range */
-	void OnBeginInOutRangeDrag();
+	void OnSelectionRangeBeginDrag();
 
 	/** Called when the user has finished dragging the selection range */
-	void OnEndInOutRangeDrag();
+	void OnSelectionRangeEndDrag();
 
 protected:
 
@@ -692,7 +706,7 @@ protected:
 	void SplitSection();
 
 	/** Generates command bindings for UI commands */
-	void BindSequencerCommands();
+	void BindCommands();
 
 	void ActivateSequencerEditorMode();
 
@@ -715,23 +729,35 @@ protected:
 	/** Called after the world has been saved. The sequencer updates to the animated state. */
 	void OnPostSaveWorld(uint32 SaveFlags, UWorld* World, bool bSuccess);
 
+	/** Called after a level has been added */
+	void OnLevelAdded(ULevel* InLevel, UWorld* InWorld);
+
+	/** Called after a level has been removed */
+	void OnLevelRemoved(ULevel* InLevel, UWorld* InWorld);
+
 	/** Called after a new level has been created. The sequencer editor mode needs to be enabled. */
 	void OnNewCurrentLevel();
 
 	/** Called after a map has been opened. The sequencer editor mode needs to be enabled. */
 	void OnMapOpened(const FString& Filename, bool bLoadAsTemplate);
 
+	/** Called when new actors are dropped in the viewport. */
+	void OnNewActorsDropped(const TArray<UObject*>& DroppedObjects, const TArray<AActor*>& DroppedActors);
+
 	/** Called before a PIE session begins. */
 	void OnPreBeginPIE(bool bIsSimulating);
 
 	/** Called after a PIE session ends. */
 	void OnEndPIE(bool bIsSimulating);
 
+	/** Called after PIE session ends and maps have been cleaned up */
+	void OnEndPlayMap();
+
 	/** Updates a viewport client from camera cut data */
 	void UpdatePreviewLevelViewportClientFromCameraCut(FLevelEditorViewportClient& InViewportClient, UObject* InCameraObject, bool bJumpCut) const;
 
 	/** Internal conversion function that doesn't perform expensive reset/update tasks */
-	void ConvertToSpawnableInternal(TSharedRef<FSequencerObjectBindingNode> NodeToBeConverted);
+	FMovieSceneSpawnable* ConvertToSpawnableInternal(FGuid PossessableGuid);
 
 	/** Handles adding a new folder to the outliner tree. */
 	void OnAddFolder();
@@ -751,6 +777,8 @@ protected:
 	/** Create set playback end transport control */
 	TSharedRef<SWidget> OnCreateTransportSetPlaybackEnd();
 
+	/** Select all keys in a display node that fall into the current selection range. */
+	void SelectKeysInSelectionRange(const TSharedRef<FSequencerDisplayNode>& DisplayNode, const TRange<float>& SelectionRange);
 	/** Create record transport control */
 	TSharedRef<SWidget> OnCreateTransportRecord();
 
@@ -842,6 +870,9 @@ private:
 		the MovieScene data can change many times per frame.) */
 	bool bNeedTreeRefresh;
 
+	/** Stores the playback status to be restored on refresh. */
+	EMovieScenePlayerStatus::Type StoredPlaybackState;
+
 	/** A flag which indicates Close() was called on this Sequencer. */ 
 	bool bWasClosed;
 
@@ -860,6 +891,8 @@ private:
 
 	FOnActorAddedToSequencer OnActorAddedToSequencerEvent;
 	FOnCameraCut OnCameraCutEvent;
+	FOnPreSave OnPreSaveEvent;
+	FOnActivateSequence OnActivateSequenceEvent;
 
 	int32 SilentModeCount;
 	
diff --git a/Engine/Source/Editor/Sequencer/Private/SequencerCommands.cpp b/Engine/Source/Editor/Sequencer/Private/SequencerCommands.cpp
--- a/Engine/Source/Editor/Sequencer/Private/SequencerCommands.cpp
+++ b/Engine/Source/Editor/Sequencer/Private/SequencerCommands.cpp
@@ -29,8 +29,10 @@ void FSequencerCommands::RegisterCommands()
 	UI_COMMAND( ToggleExpandCollapseNodes, "Expand/Collapse Nodes", "Toggle expand or collapse selected nodes", EUserInterfaceActionType::Button, FInputChord(EModifierKey::Control, EKeys::O) );
 	UI_COMMAND( ToggleExpandCollapseNodesAndDescendants, "Expand/Collapse Nodes and Descendants", "Toggle expand or collapse selected nodes and descendants", EUserInterfaceActionType::Button, FInputChord(EModifierKey::Shift, EKeys::O) );
 
-	UI_COMMAND( SetInOutEnd, "Set selection end", "Sets the end of the in/out selection", EUserInterfaceActionType::Button, FInputChord(EKeys::O) );
-	UI_COMMAND( SetInOutStart, "Set selection start", "Sets the start of the in/out selection", EUserInterfaceActionType::Button, FInputChord(EKeys::I) );
+	UI_COMMAND( SetSelectionRangeEnd, "Set Selection End", "Sets the end of the selection range", EUserInterfaceActionType::Button, FInputChord(EKeys::O) );
+	UI_COMMAND( SetSelectionRangeStart, "Set Selection Start", "Sets the start of the selection range", EUserInterfaceActionType::Button, FInputChord(EKeys::I) );
+	UI_COMMAND( SelectKeysInSelectionRange, "Select in Selection Range", "Select all keys that fall into the selection range", EUserInterfaceActionType::Button, FInputChord() );
+	UI_COMMAND( DeleteSelectionRange, "Delete Selection Range", "Delete the time range and all the keys that fall into the selection range", EUserInterfaceActionType::Button, FInputChord() );
 
 	UI_COMMAND( SetKey, "Set Key", "Sets a key on the selected tracks", EUserInterfaceActionType::Button, FInputChord(EKeys::Enter) );
 	UI_COMMAND( SetInterpolationCubicAuto, "Set Key Auto", "Cubic interpolation - Automatic tangents", EUserInterfaceActionType::Button, FInputChord(EKeys::One));
@@ -82,6 +84,7 @@ void FSequencerCommands::RegisterCommands()
 	UI_COMMAND( MarqueeTool, "Marquee", "Activates the marquee selection tool", EUserInterfaceActionType::ToggleButton, FInputChord( EKeys::M ) );
 
 	UI_COMMAND( RenderMovie, "Render Movie", "Render this movie to a video, or image frame sequence", EUserInterfaceActionType::Button, FInputChord() );
+	UI_COMMAND( CreateCamera, "Create Camera", "Create a new camera and set it as the current camera cut", EUserInterfaceActionType::Button, FInputChord() );
 
 	UI_COMMAND( PasteFromHistory, "Paste From History", "Paste from the sequencer clipboard history", EUserInterfaceActionType::Button, FInputChord(EModifierKey::Control | EModifierKey::Shift, EKeys::V) );
 
diff --git a/Engine/Source/Editor/Sequencer/Private/SequencerContextMenus.cpp b/Engine/Source/Editor/Sequencer/Private/SequencerContextMenus.cpp
--- a/Engine/Source/Editor/Sequencer/Private/SequencerContextMenus.cpp
+++ b/Engine/Source/Editor/Sequencer/Private/SequencerContextMenus.cpp
@@ -514,7 +514,10 @@ void FSectionContextMenu::AddPropertiesMenu(FMenuBuilder& MenuBuilder)
 	{
 		for (auto Section : Sequencer->GetSelection().GetSelectedSections())
 		{
-			Sections.Add(Section);
+			if (Section.IsValid())
+			{
+				Sections.Add(Section);
+			}
 		}
 	}
 
@@ -774,9 +777,12 @@ void FSectionContextMenu::ToggleSectionActive()
 
 	for (auto Section : Sequencer->GetSelection().GetSelectedSections())
 	{
-		bAnythingChanged = true;
-		Section->Modify();
-		Section->SetIsActive(bIsActive);
+		if (Section.IsValid())
+		{
+			bAnythingChanged = true;
+			Section->Modify();
+			Section->SetIsActive(bIsActive);
+		}
 	}
 
 	if (bAnythingChanged)
@@ -795,7 +801,7 @@ bool FSectionContextMenu::IsSectionActive() const
 	// Active only if all are active
 	for (auto Section : Sequencer->GetSelection().GetSelectedSections())
 	{
-		if (!Section->IsActive())
+		if (Section.IsValid() && !Section->IsActive())
 		{
 			return false;
 		}
@@ -813,9 +819,12 @@ void FSectionContextMenu::ToggleSectionLocked()
 
 	for (auto Section : Sequencer->GetSelection().GetSelectedSections())
 	{
-		bAnythingChanged = true;
-		Section->Modify();
-		Section->SetIsLocked(bIsLocked);
+		if (Section.IsValid())
+		{
+			bAnythingChanged = true;
+			Section->Modify();
+			Section->SetIsLocked(bIsLocked);
+		}
 	}
 
 	if (bAnythingChanged)
@@ -834,7 +843,7 @@ bool FSectionContextMenu::IsSectionLocked() const
 	// Locked only if all are locked
 	for (auto Section : Sequencer->GetSelection().GetSelectedSections())
 	{
-		if (!Section->IsLocked())
+		if (Section.IsValid() && !Section->IsLocked())
 		{
 			return false;
 		}
diff --git a/Engine/Source/Editor/Sequencer/Private/SequencerCurveOwner.cpp b/Engine/Source/Editor/Sequencer/Private/SequencerCurveOwner.cpp
--- a/Engine/Source/Editor/Sequencer/Private/SequencerCurveOwner.cpp
+++ b/Engine/Source/Editor/Sequencer/Private/SequencerCurveOwner.cpp
@@ -68,6 +68,21 @@ FName BuildCurveName( TSharedPtr<FSequencerDisplayNode> KeyAreaNode )
 	return FName(*FString::Join(NameParts, TEXT(" - ")));
 }
 
+bool IsNodeOrAncestorSelected(TSharedPtr<FSequencerDisplayNode> DisplayNode)
+{
+	if (!DisplayNode.IsValid())
+	{
+		return false;
+	}
+	
+	if (DisplayNode->GetSequencer().GetSelection().IsSelected(DisplayNode.ToSharedRef()))
+	{
+		return true;
+	}
+
+	return IsNodeOrAncestorSelected(DisplayNode->GetParent());
+}
+
 FSequencerCurveOwner::FSequencerCurveOwner( TSharedPtr<FSequencerNodeTree> InSequencerNodeTree, ECurveEditorCurveVisibility::Type CurveVisibility )
 {
 	SequencerNodeTree = InSequencerNodeTree;
@@ -86,7 +101,7 @@ FSequencerCurveOwner::FSequencerCurveOwner( TSharedPtr<FSequencerNodeTree> InSeq
 				bAddCurve = true;
 				break;
 			case ECurveEditorCurveVisibility::SelectedCurves:
-				bAddCurve = DisplayNodeAndKeyArea.DisplayNode->GetSequencer().GetSelection().IsSelected( DisplayNodeAndKeyArea.DisplayNode.ToSharedRef() );
+				bAddCurve = IsNodeOrAncestorSelected(DisplayNodeAndKeyArea.DisplayNode);
 				break;
 			case ECurveEditorCurveVisibility::AnimatedCurves:
 				bAddCurve = RichCurve->GetNumKeys() > 0;
@@ -113,9 +128,9 @@ TArray<FRichCurve*> FSequencerCurveOwner::GetSelectedCurves() const
 	for ( const FDisplayNodeAndKeyArea& DisplayNodeAndKeyArea : DisplayNodesAndKeyAreas )
 	{
 		FRichCurve* RichCurve = DisplayNodeAndKeyArea.KeyArea->GetRichCurve();
-		if ( RichCurve != nullptr )
+		if ( RichCurve != nullptr && SelectedCurves.Find(RichCurve) == INDEX_NONE )
 		{
-			if (DisplayNodeAndKeyArea.DisplayNode->GetSequencer().GetSelection().IsSelected(DisplayNodeAndKeyArea.DisplayNode.ToSharedRef()))
+			if (IsNodeOrAncestorSelected(DisplayNodeAndKeyArea.DisplayNode))
 			{
 				SelectedCurves.Add(RichCurve);
 			}
diff --git a/Engine/Source/Editor/Sequencer/Private/SequencerNodeTree.cpp b/Engine/Source/Editor/Sequencer/Private/SequencerNodeTree.cpp
--- a/Engine/Source/Editor/Sequencer/Private/SequencerNodeTree.cpp
+++ b/Engine/Source/Editor/Sequencer/Private/SequencerNodeTree.cpp
@@ -23,6 +23,7 @@ void FSequencerNodeTree::Empty()
 	Sequencer.GetSelection().EmptySelectedOutlinerNodes();
 	EditorMap.Empty();
 	FilteredNodes.Empty();
+	HoveredNode = nullptr;
 }
 
 
diff --git a/Engine/Source/Editor/Sequencer/Private/SequencerTimeSliderController.cpp b/Engine/Source/Editor/Sequencer/Private/SequencerTimeSliderController.cpp
--- a/Engine/Source/Editor/Sequencer/Private/SequencerTimeSliderController.cpp
+++ b/Engine/Source/Editor/Sequencer/Private/SequencerTimeSliderController.cpp
@@ -277,7 +277,7 @@ int32 FSequencerTimeSliderController::OnPaintTimeSlider( bool bMirrorLabels, con
 
 		DrawTicks( OutDrawElements, RangeToScreen, Args );
 
-		// draw playback & in/out range
+		// draw playback & selection range
 		FPaintPlaybackRangeArgs PlaybackRangeArgs(
 			bMirrorLabels ? FEditorStyle::GetBrush("Sequencer.Timeline.PlayRange_Bottom_L") : FEditorStyle::GetBrush("Sequencer.Timeline.PlayRange_Top_L"),
 			bMirrorLabels ? FEditorStyle::GetBrush("Sequencer.Timeline.PlayRange_Bottom_R") : FEditorStyle::GetBrush("Sequencer.Timeline.PlayRange_Top_R"),
@@ -285,7 +285,7 @@ int32 FSequencerTimeSliderController::OnPaintTimeSlider( bool bMirrorLabels, con
 		);
 
 		LayerId = DrawPlaybackRange(AllottedGeometry, MyClippingRect, OutDrawElements, LayerId, RangeToScreen, PlaybackRangeArgs);
-		LayerId = DrawInOutRange(AllottedGeometry, MyClippingRect, OutDrawElements, LayerId, RangeToScreen, PlaybackRangeArgs);
+		LayerId = DrawSelectionRange(AllottedGeometry, MyClippingRect, OutDrawElements, LayerId, RangeToScreen, PlaybackRangeArgs);
 
 		float HalfSize = FMath::CeilToFloat(ScrubHandleSize / 2.0f);
 
@@ -390,30 +390,30 @@ int32 FSequencerTimeSliderController::OnPaintTimeSlider( bool bMirrorLabels, con
 }
 
 
-int32 FSequencerTimeSliderController::DrawInOutRange(const FGeometry& AllottedGeometry, const FSlateRect& MyClippingRect, FSlateWindowElementList& OutDrawElements, int32 LayerId, const FScrubRangeToScreen& RangeToScreen, const FPaintPlaybackRangeArgs& Args) const
+int32 FSequencerTimeSliderController::DrawSelectionRange(const FGeometry& AllottedGeometry, const FSlateRect& MyClippingRect, FSlateWindowElementList& OutDrawElements, int32 LayerId, const FScrubRangeToScreen& RangeToScreen, const FPaintPlaybackRangeArgs& Args) const
 {
-	const TRange<float> InOutRange = TimeSliderArgs.InOutRange.Get();
+	const TRange<float> SelectionRange = TimeSliderArgs.SelectionRange.Get();
 
-	if (!InOutRange.IsEmpty())
+	if (!SelectionRange.IsEmpty())
 	{
-		const float InOutRangeL = RangeToScreen.InputToLocalX(InOutRange.GetLowerBoundValue()) - 1;
-		const float InOutRangeR = RangeToScreen.InputToLocalX(InOutRange.GetUpperBoundValue()) + 1;
-		const auto DrawColor = FLinearColor(FColor(32, 32, 128));
+		const float SelectionRangeL = RangeToScreen.InputToLocalX(SelectionRange.GetLowerBoundValue()) - 1;
+		const float SelectionRangeR = RangeToScreen.InputToLocalX(SelectionRange.GetUpperBoundValue()) + 1;
+		const auto DrawColor = FLinearColor(FColor(72, 72, 255));
 
 		FSlateDrawElement::MakeBox(
 			OutDrawElements,
 			LayerId + 1,
-			AllottedGeometry.ToPaintGeometry(FVector2D(InOutRangeL, 0.f), FVector2D(InOutRangeR - InOutRangeL, AllottedGeometry.Size.Y)),
+			AllottedGeometry.ToPaintGeometry(FVector2D(SelectionRangeL, 0.f), FVector2D(SelectionRangeR - SelectionRangeL, AllottedGeometry.Size.Y)),
 			FEditorStyle::GetBrush("WhiteBrush"),
 			MyClippingRect,
 			ESlateDrawEffect::None,
-			DrawColor.CopyWithNewOpacity(0.2f)
+			DrawColor.CopyWithNewOpacity(0.1f)
 		);
 
 		FSlateDrawElement::MakeBox(
 			OutDrawElements,
 			LayerId + 1,
-			AllottedGeometry.ToPaintGeometry(FVector2D(InOutRangeL, 0.f), FVector2D(Args.BrushWidth, AllottedGeometry.Size.Y)),
+			AllottedGeometry.ToPaintGeometry(FVector2D(SelectionRangeL, 0.f), FVector2D(Args.BrushWidth, AllottedGeometry.Size.Y)),
 			Args.StartBrush,
 			MyClippingRect,
 			ESlateDrawEffect::None,
@@ -423,7 +423,7 @@ int32 FSequencerTimeSliderController::DrawInOutRange(const FGeometry& AllottedGe
 		FSlateDrawElement::MakeBox(
 			OutDrawElements,
 			LayerId + 1,
-			AllottedGeometry.ToPaintGeometry(FVector2D(InOutRangeR - Args.BrushWidth, 0.f), FVector2D(Args.BrushWidth, AllottedGeometry.Size.Y)),
+			AllottedGeometry.ToPaintGeometry(FVector2D(SelectionRangeR - Args.BrushWidth, 0.f), FVector2D(Args.BrushWidth, AllottedGeometry.Size.Y)),
 			Args.EndBrush,
 			MyClippingRect,
 			ESlateDrawEffect::None,
@@ -557,21 +557,21 @@ FReply FSequencerTimeSliderController::OnMouseButtonUp( SWidget& WidgetOwner, co
 	}
 	else if ( bHandleLeftMouseButton )
 	{
-		if (MouseDragType == DRAG_START_RANGE)
+		if (MouseDragType == DRAG_PLAYBACK_START)
 		{
-			TimeSliderArgs.OnEndPlaybackRangeDrag.ExecuteIfBound();
+			TimeSliderArgs.OnPlaybackRangeEndDrag.ExecuteIfBound();
 		}
-		else if (MouseDragType == DRAG_END_RANGE)
+		else if (MouseDragType == DRAG_PLAYBACK_END)
 		{
-			TimeSliderArgs.OnEndPlaybackRangeDrag.ExecuteIfBound();
+			TimeSliderArgs.OnPlaybackRangeEndDrag.ExecuteIfBound();
 		}
-		else if (MouseDragType == DRAG_IN_RANGE)
+		else if (MouseDragType == DRAG_SELECTION_START)
 		{
-			TimeSliderArgs.OnEndInOutRangeDrag.ExecuteIfBound();
+			TimeSliderArgs.OnSelectionRangeEndDrag.ExecuteIfBound();
 		}
-		else if (MouseDragType == DRAG_OUT_RANGE)
+		else if (MouseDragType == DRAG_SELECTION_END)
 		{
-			TimeSliderArgs.OnEndInOutRangeDrag.ExecuteIfBound();
+			TimeSliderArgs.OnSelectionRangeEndDrag.ExecuteIfBound();
 		}
 		else if (MouseDragType == DRAG_SETTING_RANGE)
 		{
@@ -692,33 +692,33 @@ FReply FSequencerTimeSliderController::OnMouseMove( SWidget& WidgetOwner, const
 				FScrubRangeToScreen RangeToScreen(LocalViewRange, MyGeometry.Size);
 				const float ScrubPosition = TimeSliderArgs.ScrubPosition.Get();
 
-				TRange<float> InOutRange = TimeSliderArgs.InOutRange.Get();
+				TRange<float> SelectionRange = TimeSliderArgs.SelectionRange.Get();
 				TRange<float> PlaybackRange = TimeSliderArgs.PlaybackRange.Get();
 				float LocalMouseDownPos = RangeToScreen.InputToLocalX(MouseDownRange[0]);
 
-				if (HitTestScrubberEnd(RangeToScreen, InOutRange, LocalMouseDownPos, ScrubPosition))
+				if (HitTestScrubberEnd(RangeToScreen, SelectionRange, LocalMouseDownPos, ScrubPosition))
 				{
-					// in/out range end scrubber
-					MouseDragType = DRAG_OUT_RANGE;
-					TimeSliderArgs.OnBeginInOutRangeDrag.ExecuteIfBound();
+					// selection range end scrubber
+					MouseDragType = DRAG_SELECTION_END;
+					TimeSliderArgs.OnSelectionRangeBeginDrag.ExecuteIfBound();
 				}
-				else if (HitTestScrubberStart(RangeToScreen, InOutRange, LocalMouseDownPos, ScrubPosition))
+				else if (HitTestScrubberStart(RangeToScreen, SelectionRange, LocalMouseDownPos, ScrubPosition))
 				{
-					// in/out range start scrubber
-					MouseDragType = DRAG_IN_RANGE;
-					TimeSliderArgs.OnBeginInOutRangeDrag.ExecuteIfBound();
+					// selection range start scrubber
+					MouseDragType = DRAG_SELECTION_START;
+					TimeSliderArgs.OnSelectionRangeBeginDrag.ExecuteIfBound();
 				}
 				else if (HitTestScrubberEnd(RangeToScreen, PlaybackRange, LocalMouseDownPos, ScrubPosition))
 				{
 					// playback range end scrubber
-					MouseDragType = DRAG_END_RANGE;
-					TimeSliderArgs.OnBeginPlaybackRangeDrag.ExecuteIfBound();
+					MouseDragType = DRAG_PLAYBACK_END;
+					TimeSliderArgs.OnPlaybackRangeBeginDrag.ExecuteIfBound();
 				}
 				else if (HitTestScrubberStart(RangeToScreen, PlaybackRange, LocalMouseDownPos, ScrubPosition))
 				{
 					// playback range start scrubber
-					MouseDragType = DRAG_START_RANGE;
-					TimeSliderArgs.OnBeginPlaybackRangeDrag.ExecuteIfBound();
+					MouseDragType = DRAG_PLAYBACK_START;
+					TimeSliderArgs.OnPlaybackRangeBeginDrag.ExecuteIfBound();
 				}
 				else if (FSlateApplication::Get().GetModifierKeys().AreModifersDown(EModifierKey::Control))
 				{
@@ -739,7 +739,7 @@ FReply FSequencerTimeSliderController::OnMouseMove( SWidget& WidgetOwner, const
 
 
 			// Set the start range time?
-			if (MouseDragType == DRAG_START_RANGE)
+			if (MouseDragType == DRAG_PLAYBACK_START)
 			{
 				if (TimeSliderArgs.Settings->GetIsSnapEnabled())
 				{
@@ -749,7 +749,7 @@ FReply FSequencerTimeSliderController::OnMouseMove( SWidget& WidgetOwner, const
 				SetPlaybackRangeStart(NewValue);
 			}
 			// Set the end range time?
-			else if(MouseDragType == DRAG_END_RANGE)
+			else if(MouseDragType == DRAG_PLAYBACK_END)
 			{
 				if (TimeSliderArgs.Settings->GetIsSnapEnabled())
 				{
@@ -758,33 +758,33 @@ FReply FSequencerTimeSliderController::OnMouseMove( SWidget& WidgetOwner, const
 					
 				SetPlaybackRangeEnd(NewValue);
 			}
-			else if (MouseDragType == DRAG_IN_RANGE)
+			else if (MouseDragType == DRAG_SELECTION_START)
 			{
 				if (TimeSliderArgs.Settings->GetIsSnapEnabled())
 				{
 					NewValue = TimeSliderArgs.Settings->SnapTimeToInterval(NewValue);
 				}
 
-				TRange<float> InOutRange = TimeSliderArgs.InOutRange.Get();
+				TRange<float> SelectionRange = TimeSliderArgs.SelectionRange.Get();
 
-				if (NewValue <= InOutRange.GetUpperBoundValue())
+				if (NewValue <= SelectionRange.GetUpperBoundValue())
 				{
-					TimeSliderArgs.OnInOutRangeChanged.ExecuteIfBound(TRange<float>(NewValue, InOutRange.GetUpperBoundValue()));
+					TimeSliderArgs.OnSelectionRangeChanged.ExecuteIfBound(TRange<float>(NewValue, SelectionRange.GetUpperBoundValue()));
 				}
 			}
 			// Set the end range time?
-			else if(MouseDragType == DRAG_OUT_RANGE)
+			else if(MouseDragType == DRAG_SELECTION_END)
 			{
 				if (TimeSliderArgs.Settings->GetIsSnapEnabled())
 				{
 					NewValue = TimeSliderArgs.Settings->SnapTimeToInterval(NewValue);
 				}
 					
-				TRange<float> InOutRange = TimeSliderArgs.InOutRange.Get();
+				TRange<float> SelectionRange = TimeSliderArgs.SelectionRange.Get();
 
-				if (NewValue >= InOutRange.GetLowerBoundValue())
+				if (NewValue >= SelectionRange.GetLowerBoundValue())
 				{
-					TimeSliderArgs.OnInOutRangeChanged.ExecuteIfBound(TRange<float>(InOutRange.GetLowerBoundValue(), NewValue));
+					TimeSliderArgs.OnSelectionRangeChanged.ExecuteIfBound(TRange<float>(SelectionRange.GetLowerBoundValue(), NewValue));
 				}
 			}
 			else if (MouseDragType == DRAG_SCRUBBING_TIME)
@@ -845,19 +845,19 @@ FCursorReply FSequencerTimeSliderController::OnCursorQuery( TSharedRef<const SWi
 {
 	FScrubRangeToScreen RangeToScreen(TimeSliderArgs.ViewRange.Get(), MyGeometry.Size);
 	TRange<float> PlaybackRange = TimeSliderArgs.PlaybackRange.Get();
-	TRange<float> InOutRange = TimeSliderArgs.InOutRange.Get();
+	TRange<float> SelectionRange = TimeSliderArgs.SelectionRange.Get();
 	const float ScrubPosition = TimeSliderArgs.ScrubPosition.Get();
 
 	// Use L/R resize cursor if we're dragging or hovering a playback range bound
 	float HitTestPosition = MyGeometry.AbsoluteToLocal(CursorEvent.GetScreenSpacePosition()).X;
-	if ((MouseDragType == DRAG_END_RANGE) ||
-		(MouseDragType == DRAG_START_RANGE) ||
-		(MouseDragType == DRAG_IN_RANGE) ||
-		(MouseDragType == DRAG_OUT_RANGE) ||
+	if ((MouseDragType == DRAG_PLAYBACK_END) ||
+		(MouseDragType == DRAG_PLAYBACK_START) ||
+		(MouseDragType == DRAG_SELECTION_START) ||
+		(MouseDragType == DRAG_SELECTION_END) ||
 		HitTestScrubberStart(RangeToScreen, PlaybackRange, HitTestPosition, ScrubPosition) ||
 		HitTestScrubberEnd(RangeToScreen, PlaybackRange, HitTestPosition, ScrubPosition) ||
-		HitTestScrubberStart(RangeToScreen, InOutRange, HitTestPosition, ScrubPosition) ||
-		HitTestScrubberEnd(RangeToScreen, InOutRange, HitTestPosition, ScrubPosition))
+		HitTestScrubberStart(RangeToScreen, SelectionRange, HitTestPosition, ScrubPosition) ||
+		HitTestScrubberEnd(RangeToScreen, SelectionRange, HitTestPosition, ScrubPosition))
 	{
 		return FCursorReply::Cursor(EMouseCursor::ResizeLeftRight);
 	}
@@ -881,7 +881,7 @@ int32 FSequencerTimeSliderController::OnPaintSectionView( const FGeometry& Allot
 	if (Args.PlaybackRangeArgs.IsSet())
 	{
 		LayerId = DrawPlaybackRange(AllottedGeometry, MyClippingRect, OutDrawElements, LayerId, RangeToScreen, Args.PlaybackRangeArgs.GetValue());
-		LayerId = DrawInOutRange(AllottedGeometry, MyClippingRect, OutDrawElements, LayerId, RangeToScreen, Args.PlaybackRangeArgs.GetValue());
+		LayerId = DrawSelectionRange(AllottedGeometry, MyClippingRect, OutDrawElements, LayerId, RangeToScreen, Args.PlaybackRangeArgs.GetValue());
 	}
 
 	if( Args.bDisplayTickLines )
diff --git a/Engine/Source/Editor/Sequencer/Private/SequencerTimeSliderController.h b/Engine/Source/Editor/Sequencer/Private/SequencerTimeSliderController.h
--- a/Engine/Source/Editor/Sequencer/Private/SequencerTimeSliderController.h
+++ b/Engine/Source/Editor/Sequencer/Private/SequencerTimeSliderController.h
@@ -142,11 +142,11 @@ private:
 	void DrawTicks( FSlateWindowElementList& OutDrawElements, const struct FScrubRangeToScreen& RangeToScreen, struct FDrawTickArgs& InArgs ) const;
 
 	/**
-	 * Draw the in/out selection range.
+	 * Draw the selection range.
 	 *
 	 * @return The new layer ID.
 	 */
-	int32 DrawInOutRange(const FGeometry& AllottedGeometry, const FSlateRect& MyClippingRect, FSlateWindowElementList& OutDrawElements, int32 LayerId, const FScrubRangeToScreen& RangeToScreen, const FPaintPlaybackRangeArgs& Args) const;
+	int32 DrawSelectionRange(const FGeometry& AllottedGeometry, const FSlateRect& MyClippingRect, FSlateWindowElementList& OutDrawElements, int32 LayerId, const FScrubRangeToScreen& RangeToScreen, const FPaintPlaybackRangeArgs& Args) const;
 
 	/**
 	 * Draw the playback range.
@@ -174,34 +174,45 @@ private:
 
 private:
 	FTimeSliderArgs TimeSliderArgs;
+
 	/** The size of the scrub handle */
 	float ScrubHandleSize;
+	
 	/** Brush for drawing an upwards facing scrub handle */
 	const FSlateBrush* ScrubHandleUp;
+	
 	/** Brush for drawing a downwards facing scrub handle */
 	const FSlateBrush* ScrubHandleDown;
+	
 	/** Total mouse delta during dragging **/
 	float DistanceDragged;
-	/** If we are dragging the scrubber or dragging to set the time range */
+	
+	/** If we are dragging a scrubber or dragging to set the time range */
 	enum DragType
 	{
 		DRAG_SCRUBBING_TIME,
 		DRAG_SETTING_RANGE,
-		DRAG_START_RANGE,
-		DRAG_END_RANGE,
-		DRAG_IN_RANGE,
-		DRAG_OUT_RANGE,
+		DRAG_PLAYBACK_START,
+		DRAG_PLAYBACK_END,
+		DRAG_SELECTION_START,
+		DRAG_SELECTION_END,
 		DRAG_NONE
 	};
+	
 	DragType MouseDragType;
+	
 	/** If we are currently panning the panel */
 	bool bPanning;
+	
 	/** Mouse down time range */
 	FVector2D MouseDownRange;
+	
 	/** Range stack */
 	TArray<FVector2D> RangeStack;
+	
 	/** When > 0, we should not show context menus */
 	int32 ContextMenuSuppression;
+	
 	friend FContextMenuSuppressor;
 };
 
diff --git a/Engine/Source/Editor/Sequencer/Private/Tools/EditToolDragOperations.cpp b/Engine/Source/Editor/Sequencer/Private/Tools/EditToolDragOperations.cpp
--- a/Engine/Source/Editor/Sequencer/Private/Tools/EditToolDragOperations.cpp
+++ b/Engine/Source/Editor/Sequencer/Private/Tools/EditToolDragOperations.cpp
@@ -182,7 +182,7 @@ void FResizeSection::OnBeginDrag(const FPointerEvent& MouseEvent, FVector2D Loca
 			}
 		}
 
-		Section->GetKeyHandles(DraggedKeyHandles);
+		Section->GetKeyHandles(DraggedKeyHandles, Section->GetRange());
 		SectionInitTimes.Add(Section, bDraggingByEnd ? Section->GetEndTime() : Section->GetStartTime());
 	}
 }
@@ -328,7 +328,7 @@ void FMoveSection::OnBeginDrag(const FPointerEvent& MouseEvent, FVector2D LocalM
 	{
 		auto* Section = Sections[Index].GetSectionObject();
 
-		Section->GetKeyHandles(DraggedKeyHandles);
+		Section->GetKeyHandles(DraggedKeyHandles, Section->GetRange());
 		RelativeOffsets.Add(FRelativeOffset{ Section->GetStartTime() - InitialPosition.X, Section->GetEndTime() - InitialPosition.X });
 	}
 }
diff --git a/Engine/Source/Editor/Sequencer/Private/Tools/SequencerSnapField.cpp b/Engine/Source/Editor/Sequencer/Private/Tools/SequencerSnapField.cpp
--- a/Engine/Source/Editor/Sequencer/Private/Tools/SequencerSnapField.cpp
+++ b/Engine/Source/Editor/Sequencer/Private/Tools/SequencerSnapField.cpp
@@ -72,6 +72,11 @@ FSequencerSnapField::FSequencerSnapField(const ISequencer& InSequencer, ISequenc
 	// Add the current time as a potential snap candidate
 	Visitor.Snaps.Add(FSequencerSnapPoint{ FSequencerSnapPoint::CurrentTime, InSequencer.GetGlobalTime() });
 
+	// Add the selection range bounds as a potential snap candidate
+	TRange<float> SelectionRange = InSequencer.GetFocusedMovieSceneSequence()->GetMovieScene()->GetSelectionRange();
+	Visitor.Snaps.Add(FSequencerSnapPoint{ FSequencerSnapPoint::InOutRange, SelectionRange.GetLowerBoundValue() });
+	Visitor.Snaps.Add(FSequencerSnapPoint{ FSequencerSnapPoint::InOutRange, SelectionRange.GetUpperBoundValue() });
+
 	// Sort
 	Visitor.Snaps.Sort([](const FSequencerSnapPoint& A, const FSequencerSnapPoint& B){
 		return A.Time < B.Time;
diff --git a/Engine/Source/Editor/Sequencer/Private/Tools/SequencerSnapField.h b/Engine/Source/Editor/Sequencer/Private/Tools/SequencerSnapField.h
--- a/Engine/Source/Editor/Sequencer/Private/Tools/SequencerSnapField.h
+++ b/Engine/Source/Editor/Sequencer/Private/Tools/SequencerSnapField.h
@@ -13,7 +13,7 @@ class UMovieSceneSection;
 /** Structure defining a point to snap to in the sequencer */
 struct FSequencerSnapPoint
 {
-	enum ESnapType { Key, SectionBounds, CustomSection, PlaybackRange, CurrentTime };
+	enum ESnapType { Key, SectionBounds, CustomSection, PlaybackRange, CurrentTime, InOutRange };
 
 	/** The type of snap */
 	ESnapType Type;
diff --git a/Engine/Source/Editor/Sequencer/Public/ISequencer.h b/Engine/Source/Editor/Sequencer/Public/ISequencer.h
--- a/Engine/Source/Editor/Sequencer/Public/ISequencer.h
+++ b/Engine/Source/Editor/Sequencer/Public/ISequencer.h
@@ -125,6 +125,13 @@ public:
 	 */
 	virtual bool HasSequenceInstanceForSection(UMovieSceneSection& Section) const = 0;
 
+	/**
+	 * Add actors as possessable objects to sequencer.
+	 * 
+	 * @param InActors The actors to add to sequencer.
+	 */
+	virtual void AddActors(const TArray<TWeakObjectPtr<AActor> >& InActors) = 0;
+
 	/**
 	 * Adds a movie scene as a section inside the current movie scene
 	 * 
@@ -250,6 +257,12 @@ public:
 	DECLARE_EVENT_TwoParams(ISequencer, FOnCameraCut, UObject*, bool)
 	virtual FOnCameraCut& OnCameraCut() = 0;
 
+	DECLARE_EVENT_OneParam(ISequencer, FOnPreSave, ISequencer&)
+	virtual FOnPreSave& OnPreSave() = 0;
+
+	DECLARE_EVENT_OneParam(ISequencer, FOnActivateSequence, FMovieSceneSequenceInstance&)
+	virtual FOnActivateSequence& OnActivateSequence() = 0;
+
 	/**
 	 * Gets a handle to runtime information about the object being manipulated by a movie scene
 	 * 
@@ -307,6 +320,9 @@ public:
 	/** Turn viewport transport controls on or off */
 	virtual void SetViewportTransportControlsVisibility(bool bVisible) = 0;
 
+	/** Attempt to find a spawned object in the currently focused movie scene, or the template object for the specified binding ID, if possible */
+	virtual UObject* FindSpawnedObjectOrTemplate(const FGuid& BindingId) const = 0;
+
 	/**
 	 * Create a widget containing the spinboxes for setting the working and playback range
 	 * 
