Copying //UE4/Release-Staging-4.12 to //UE4/Dev-Main (Source: //UE4/Release-4.12 @ 2955635)

diff --git a/Engine/Source/Editor/Sequencer/Private/SequencerTimeSliderController.cpp b/Engine/Source/Editor/Sequencer/Private/SequencerTimeSliderController.cpp
--- a/Engine/Source/Editor/Sequencer/Private/SequencerTimeSliderController.cpp
+++ b/Engine/Source/Editor/Sequencer/Private/SequencerTimeSliderController.cpp
@@ -498,7 +498,7 @@ FReply FSequencerTimeSliderController::OnMouseButtonDown( SWidget& WidgetOwner,
 	DistanceDragged = 0;
 
 	FScrubRangeToScreen RangeToScreen( TimeSliderArgs.ViewRange.Get(), MyGeometry.Size );
-	MouseDownRange[0] = RangeToScreen.LocalXToInput(MyGeometry.AbsoluteToLocal(MouseEvent.GetLastScreenSpacePosition()).X);
+	MouseDownRange[0] = RangeToScreen.LocalXToInput(MyGeometry.AbsoluteToLocal(MouseEvent.GetScreenSpacePosition()).X);
 	MouseDownRange[1] = MouseDownRange[0];
 
 	if ( bHandleLeftMouseButton )
@@ -696,13 +696,14 @@ FReply FSequencerTimeSliderController::OnMouseMove( SWidget& WidgetOwner, const
 				TRange<float> PlaybackRange = TimeSliderArgs.PlaybackRange.Get();
 				float LocalMouseDownPos = RangeToScreen.InputToLocalX(MouseDownRange[0]);
 
-				if (HitTestScrubberEnd(RangeToScreen, SelectionRange, LocalMouseDownPos, ScrubPosition))
+				// Disable selection range test if it's empty so that the playback range scrubbing gets priority
+				if (!SelectionRange.IsEmpty() && HitTestScrubberEnd(RangeToScreen, SelectionRange, LocalMouseDownPos, ScrubPosition))
 				{
 					// selection range end scrubber
 					MouseDragType = DRAG_SELECTION_END;
 					TimeSliderArgs.OnSelectionRangeBeginDrag.ExecuteIfBound();
 				}
-				else if (HitTestScrubberStart(RangeToScreen, SelectionRange, LocalMouseDownPos, ScrubPosition))
+				else if (!SelectionRange.IsEmpty() && HitTestScrubberStart(RangeToScreen, SelectionRange, LocalMouseDownPos, ScrubPosition))
 				{
 					// selection range start scrubber
 					MouseDragType = DRAG_SELECTION_START;
@@ -823,9 +824,22 @@ FReply FSequencerTimeSliderController::OnMouseWheel( SWidget& WidgetOwner, const
 {
 	TOptional<TRange<float>> NewTargetRange;
 
-	float MouseFractionX = MyGeometry.AbsoluteToLocal(MouseEvent.GetScreenSpacePosition()).X / MyGeometry.GetLocalSize().X;
 	if ( TimeSliderArgs.AllowZoom && MouseEvent.IsControlDown() )
 	{
+		float MouseFractionX = MyGeometry.AbsoluteToLocal(MouseEvent.GetScreenSpacePosition()).X / MyGeometry.GetLocalSize().X;
+
+		// If zooming on the current time, adjust mouse fractionX
+		if (TimeSliderArgs.Settings->GetZoomPosition() == ESequencerZoomPosition::SZP_CurrentTime)
+		{
+			const float ScrubPosition = TimeSliderArgs.ScrubPosition.Get();
+			if (TimeSliderArgs.ViewRange.Get().Contains(ScrubPosition))
+			{
+				FScrubRangeToScreen RangeToScreen(TimeSliderArgs.ViewRange.Get(), MyGeometry.Size);
+				float TimePosition = RangeToScreen.InputToLocalX(ScrubPosition);
+				MouseFractionX = TimePosition / MyGeometry.GetLocalSize().X;
+			}
+		}
+
 		const float ZoomDelta = -0.2f * MouseEvent.GetWheelDelta();
 		if (ZoomByDelta(ZoomDelta, MouseFractionX))
 		{
@@ -856,8 +870,8 @@ FCursorReply FSequencerTimeSliderController::OnCursorQuery( TSharedRef<const SWi
 		(MouseDragType == DRAG_SELECTION_END) ||
 		HitTestScrubberStart(RangeToScreen, PlaybackRange, HitTestPosition, ScrubPosition) ||
 		HitTestScrubberEnd(RangeToScreen, PlaybackRange, HitTestPosition, ScrubPosition) ||
-		HitTestScrubberStart(RangeToScreen, SelectionRange, HitTestPosition, ScrubPosition) ||
-		HitTestScrubberEnd(RangeToScreen, SelectionRange, HitTestPosition, ScrubPosition))
+		(!SelectionRange.IsEmpty() && HitTestScrubberStart(RangeToScreen, SelectionRange, HitTestPosition, ScrubPosition)) ||
+		(!SelectionRange.IsEmpty() && HitTestScrubberEnd(RangeToScreen, SelectionRange, HitTestPosition, ScrubPosition)))
 	{
 		return FCursorReply::Cursor(EMouseCursor::ResizeLeftRight);
 	}
@@ -1091,7 +1105,7 @@ void FSequencerTimeSliderController::PanByDelta( float InDelta )
 
 bool FSequencerTimeSliderController::HitTestScrubberStart(const FScrubRangeToScreen& RangeToScreen, const TRange<float>& PlaybackRange, float LocalHitPositionX, float ScrubPosition) const
 {
-	static float BrushSizeInStateUnits = 6.f, DragToleranceSlateUnits = 2.f;
+	static float BrushSizeInStateUnits = 6.f, DragToleranceSlateUnits = 2.f, MouseTolerance = 2.f;
 	float LocalPlaybackStartPos = RangeToScreen.InputToLocalX(PlaybackRange.GetLowerBoundValue());
 
 	// We favor hit testing the scrub bar over hit testing the playback range bounds
@@ -1101,13 +1115,13 @@ bool FSequencerTimeSliderController::HitTestScrubberStart(const FScrubRangeToScr
 	}
 
 	// Hit test against the brush region to the right of the playback start position, +/- DragToleranceSlateUnits
-	return LocalHitPositionX >= LocalPlaybackStartPos - DragToleranceSlateUnits &&
-		LocalHitPositionX <= LocalPlaybackStartPos + BrushSizeInStateUnits + DragToleranceSlateUnits;
+	return LocalHitPositionX >= LocalPlaybackStartPos - MouseTolerance - DragToleranceSlateUnits &&
+		LocalHitPositionX <= LocalPlaybackStartPos + MouseTolerance + BrushSizeInStateUnits + DragToleranceSlateUnits;
 }
 
 bool FSequencerTimeSliderController::HitTestScrubberEnd(const FScrubRangeToScreen& RangeToScreen, const TRange<float>& PlaybackRange, float LocalHitPositionX, float ScrubPosition) const
 {
-	static float BrushSizeInStateUnits = 6.f, DragToleranceSlateUnits = 2.f;
+	static float BrushSizeInStateUnits = 6.f, DragToleranceSlateUnits = 2.f, MouseTolerance = 2.f;
 	float LocalPlaybackEndPos = RangeToScreen.InputToLocalX(PlaybackRange.GetUpperBoundValue());
 	
 	// We favor hit testing the scrub bar over hit testing the playback range bounds
@@ -1117,8 +1131,8 @@ bool FSequencerTimeSliderController::HitTestScrubberEnd(const FScrubRangeToScree
 	}
 	
 	// Hit test against the brush region to the left of the playback end position, +/- DragToleranceSlateUnits
-	return LocalHitPositionX >= LocalPlaybackEndPos - BrushSizeInStateUnits - DragToleranceSlateUnits &&
-		LocalHitPositionX <= LocalPlaybackEndPos + DragToleranceSlateUnits;
+	return LocalHitPositionX >= LocalPlaybackEndPos - MouseTolerance - BrushSizeInStateUnits - DragToleranceSlateUnits &&
+		LocalHitPositionX <= LocalPlaybackEndPos + MouseTolerance + DragToleranceSlateUnits;
 }
 
 void FSequencerTimeSliderController::SetPlaybackRangeStart(float NewStart)
diff --git a/Engine/Source/Editor/Sequencer/Public/ISequencer.h b/Engine/Source/Editor/Sequencer/Public/ISequencer.h
--- a/Engine/Source/Editor/Sequencer/Public/ISequencer.h
+++ b/Engine/Source/Editor/Sequencer/Public/ISequencer.h
@@ -192,12 +192,13 @@ public:
 	 *
 	 * @param Time The global time to set.
 	 * @param SnapTimeMode The type of time snapping allowed.
+	 * @param bLooped Whether the time has been looped
 	 * @see GetGlobalTime
 	 */
-	virtual void SetGlobalTime(float Time, ESnapTimeMode SnapTimeMode = ESnapTimeMode::STM_None) = 0;
+	virtual void SetGlobalTime(float Time, ESnapTimeMode SnapTimeMode = ESnapTimeMode::STM_None, bool bLooped = false) = 0;
 
 	/** Set the global time directly, without performing any auto-scroll */
-	virtual void SetGlobalTimeDirectly(float Time, ESnapTimeMode SnapTimeMode = ESnapTimeMode::STM_None) = 0;
+	virtual void SetGlobalTimeDirectly(float Time, ESnapTimeMode SnapTimeMode = ESnapTimeMode::STM_None, bool bLooped = false) = 0;
 
 	/** @return The current view range */
 	virtual FAnimatedRange GetViewRange() const
diff --git a/Engine/Source/Editor/UMGEditor/Private/WidgetBlueprintEditorUtils.cpp b/Engine/Source/Editor/UMGEditor/Private/WidgetBlueprintEditorUtils.cpp
--- a/Engine/Source/Editor/UMGEditor/Private/WidgetBlueprintEditorUtils.cpp
+++ b/Engine/Source/Editor/UMGEditor/Private/WidgetBlueprintEditorUtils.cpp
@@ -35,7 +35,7 @@ public:
 
 	// FCustomizableTextObjectFactory implementation
 
-	virtual bool CanCreateClass(UClass* ObjectClass) const override
+	virtual bool CanCreateClass(UClass* ObjectClass, bool& bOmitSubObjs) const override
 	{
 		const bool bIsWidget = ObjectClass->IsChildOf(UWidget::StaticClass());
 		const bool bIsSlot = ObjectClass->IsChildOf(UPanelSlot::StaticClass());
diff --git a/Engine/Source/Editor/UnrealEd/Classes/Editor/EditorEngine.h b/Engine/Source/Editor/UnrealEd/Classes/Editor/EditorEngine.h
--- a/Engine/Source/Editor/UnrealEd/Classes/Editor/EditorEngine.h
+++ b/Engine/Source/Editor/UnrealEd/Classes/Editor/EditorEngine.h
@@ -561,6 +561,9 @@ public:
 
 protected:
 
+	/** Count of how many PIE instances are waiting to log in */
+	int32 PIEInstancesToLogInCount;
+
 	/* These are parameters that we need to cache for late joining */
 	FString ServerPrefix;
 	int32 PIEInstance;
@@ -2485,6 +2488,9 @@ private:
 	/** Above function but called a frame later, to stop PIE login from happening from a network callback */
 	virtual void OnLoginPIEComplete_Deferred(int32 LocalUserNum, bool bWasSuccessful, FString ErrorString, FPieLoginStruct DataStruct);
 
+	/** Called when all PIE instances have been successfully logged in */
+	virtual void OnLoginPIEAllComplete();
+
 public:
 	/**
 	 * Continue the creation of a single PIE world after a login was successful
@@ -2536,6 +2542,11 @@ private:
 	 */
 	void OnSwitchWorldsForPIE( bool bSwitchToPieWorld );
 
+	/**
+	 * Gives focus to the server or first PIE client viewport
+	 */
+	void GiveFocusToFirstClientPIEViewport();
+
 public:
 	/**
 	 * Spawns a PlayFromHere playerstart in the given world
diff --git a/Engine/Source/Editor/UnrealEd/Classes/MaterialGraph/MaterialGraphNode.h b/Engine/Source/Editor/UnrealEd/Classes/MaterialGraph/MaterialGraphNode.h
--- a/Engine/Source/Editor/UnrealEd/Classes/MaterialGraph/MaterialGraphNode.h
+++ b/Engine/Source/Editor/UnrealEd/Classes/MaterialGraph/MaterialGraphNode.h
@@ -2,6 +2,8 @@
 
 #pragma once
 
+#include "MaterialGraph.h"
+#include "MaterialGraphNode_Base.h"
 #include "MaterialGraphNode.generated.h"
 
 UCLASS(MinimalAPI)
diff --git a/Engine/Source/Editor/UnrealEd/Classes/Settings/EditorLoadingSavingSettings.h b/Engine/Source/Editor/UnrealEd/Classes/Settings/EditorLoadingSavingSettings.h
--- a/Engine/Source/Editor/UnrealEd/Classes/Settings/EditorLoadingSavingSettings.h
+++ b/Engine/Source/Editor/UnrealEd/Classes/Settings/EditorLoadingSavingSettings.h
@@ -86,16 +86,10 @@ public:
 	UPROPERTY(EditAnywhere, config, Category=Startup)
 	uint32 bRestoreOpenAssetTabsOnRestart:1;
 
-public:
-
-	/** Check that the current source control settings will play nicely with the auto-reimport feature */
-	void CheckSourceControlCompatability();
-
 private:
 
-	/** true when CheckSourceControlCompatability is enabled */
 	UPROPERTY(config)
-	bool bEnableSourceControlCompatabilityCheck;
+	bool bEnableSourceControlCompatabilityCheck_DEPRECATED;
 
 public:
 
diff --git a/Engine/Source/Editor/UnrealEd/Private/Animation/EditorCompositeSection.cpp b/Engine/Source/Editor/UnrealEd/Private/Animation/EditorCompositeSection.cpp
--- a/Engine/Source/Editor/UnrealEd/Private/Animation/EditorCompositeSection.cpp
+++ b/Engine/Source/Editor/UnrealEd/Private/Animation/EditorCompositeSection.cpp
@@ -5,7 +5,7 @@
 =============================================================================*/ 
 
 #include "UnrealEd.h"
-
+#include "Animation/AnimMetaData.h"
 
 
 UEditorCompositeSection::UEditorCompositeSection(const FObjectInitializer& ObjectInitializer)
@@ -15,6 +15,26 @@ UEditorCompositeSection::UEditorCompositeSection(const FObjectInitializer& Objec
 	SectionIndex = INDEX_NONE;
 }
 
+// since meta data is instanced, they have to copied manually with correct Outer
+// when we move between editor composite section and montage composite section
+void CopyMetaData(FCompositeSection& Source, FCompositeSection& Dest, UObject* DestOuter)
+{
+	const int32 TotalMetaData = Source.MetaData.Num();
+	Dest.MetaData.Reset(TotalMetaData);
+	Dest.MetaData.AddZeroed(TotalMetaData);
+
+	for (int32 Idx = 0; Idx < TotalMetaData; ++Idx)
+	{
+		UAnimMetaData* SourceMetaData = Source.MetaData[Idx];
+		if (SourceMetaData)
+		{
+			TSubclassOf<UAnimMetaData> SourceMetaDataClass = SourceMetaData->GetClass();
+			UAnimMetaData* NewMetaData = NewObject<UAnimMetaData>(DestOuter, SourceMetaDataClass, NAME_None, RF_NoFlags, SourceMetaData);
+			Dest.MetaData[Idx] = NewMetaData;
+		}
+	}
+}
+
 void UEditorCompositeSection::InitSection(int SectionIndexIn)
 {
 	SectionIndex = SectionIndexIn;
@@ -23,6 +43,7 @@ void UEditorCompositeSection::InitSection(int SectionIndexIn)
 		if(Montage->CompositeSections.IsValidIndex(SectionIndex))
 		{
 			CompositeSection = Montage->CompositeSections[SectionIndex];
+			CopyMetaData(Montage->CompositeSections[SectionIndex], CompositeSection, this);
 		}
 	}
 }
@@ -34,9 +55,11 @@ bool UEditorCompositeSection::ApplyChangesToMontage()
 		{
 			CompositeSection.OnChanged(CompositeSection.GetTime());
 			Montage->CompositeSections[SectionIndex] = CompositeSection;
+			CopyMetaData(CompositeSection, Montage->CompositeSections[SectionIndex], Montage);
 			return true;
 		}
 	}
 
 	return false;
 }
+
diff --git a/Engine/Source/Editor/UnrealEd/Private/EdGraphUtilities.cpp b/Engine/Source/Editor/UnrealEd/Private/EdGraphUtilities.cpp
--- a/Engine/Source/Editor/UnrealEd/Private/EdGraphUtilities.cpp
+++ b/Engine/Source/Editor/UnrealEd/Private/EdGraphUtilities.cpp
@@ -28,10 +28,14 @@ public:
 	}
 
 protected:
-	virtual bool CanCreateClass(UClass* ObjectClass) const override
+	virtual bool CanCreateClass(UClass* ObjectClass, bool& bOmitSubObjs) const override
 	{
 		if (const UEdGraphNode* DefaultNode = Cast<UEdGraphNode>(ObjectClass->GetDefaultObject()))
 		{
+			// if the root node can't be created, don't continue to check sub-
+			// objects (for like collapsed graphs, or anim state-machine nodes)
+			bOmitSubObjs = true;
+
 			if (DefaultNode->CanDuplicateNode())
 			{
 				if (DestinationGraph != NULL)
diff --git a/Engine/Source/Editor/UnrealEd/Private/EditorEngine.cpp b/Engine/Source/Editor/UnrealEd/Private/EditorEngine.cpp
--- a/Engine/Source/Editor/UnrealEd/Private/EditorEngine.cpp
+++ b/Engine/Source/Editor/UnrealEd/Private/EditorEngine.cpp
@@ -4069,7 +4069,7 @@ void UEditorEngine::OnPostSaveWorld(uint32 SaveFlags, UWorld* World, uint32 Orig
 
 					if ( MainFrameModule.GetMRUFavoritesList() )
 					{
-						MainFrameModule.GetMRUFavoritesList()->AddMRUItem( Filename );
+						MainFrameModule.GetMRUFavoritesList()->AddMRUItem(WorldPackage->GetName());
 					}
 
 					FEditorDirectories::Get().SetLastDirectory(ELastDirectory::UNR, FPaths::GetPath(Filename)); // Save path as default for next time.
diff --git a/Engine/Source/Editor/UnrealEd/Private/EditorViewportClient.cpp b/Engine/Source/Editor/UnrealEd/Private/EditorViewportClient.cpp
--- a/Engine/Source/Editor/UnrealEd/Private/EditorViewportClient.cpp
+++ b/Engine/Source/Editor/UnrealEd/Private/EditorViewportClient.cpp
@@ -3143,15 +3143,34 @@ void FEditorViewportClient::SetupViewForRendering( FSceneViewFamily& ViewFamily,
 	FPixelInspectorModule* PixelInspectorModule = &FModuleManager::LoadModuleChecked<FPixelInspectorModule>(TEXT("PixelInspectorModule"));
 	if (PixelInspectorModule != nullptr)
 	{
-		View.bUsePixelInspector = PixelInspectorModule->IsPixelInspectorEnable();
-		bool IsMouseInsideViewport = CurrentMousePos != FIntPoint(-1, -1);
-		if (View.bUsePixelInspector && IsMouseInsideViewport)
+		bool IsInspectorActive = PixelInspectorModule->IsPixelInspectorEnable();
+		View.bUsePixelInspector = IsInspectorActive;
+		FIntPoint InspectViewportPos = FIntPoint(-1, -1);
+		if (IsInspectorActive)
+		{
+			if (CurrentMousePos == FIntPoint(-1, -1))
+			{
+				uint32 CoordinateViewportId = 0;
+				PixelInspectorModule->GetCoordinatePosition(InspectViewportPos, CoordinateViewportId);
+				PixelInspectorModule->SetViewportInformation(View.State->GetViewKey(), Viewport->GetSizeXY());
+				bool IsCoordinateInViewport = InspectViewportPos.X <= Viewport->GetSizeXY().X && InspectViewportPos.Y <= Viewport->GetSizeXY().Y;
+				IsInspectorActive = IsCoordinateInViewport && (CoordinateViewportId == View.State->GetViewKey());
+			}
+			else
+			{
+				InspectViewportPos = CurrentMousePos;
+				PixelInspectorModule->SetViewportInformation(View.State->GetViewKey(), Viewport->GetSizeXY());
+				PixelInspectorModule->SetCoordinatePosition(CurrentMousePos, false);
+			}
+		}
+
+		if (IsInspectorActive)
 		{
 			// Ready to send a request
 			FSceneInterface *SceneInterface = GetScene();
-			PixelInspectorModule->CreatePixelInspectorRequest(CurrentMousePos, View.State->GetViewKey(), SceneInterface);
+			PixelInspectorModule->CreatePixelInspectorRequest(InspectViewportPos, View.State->GetViewKey(), SceneInterface);
 		}
-		else if (!View.bUsePixelInspector && IsMouseInsideViewport)
+		else if (!View.bUsePixelInspector && CurrentMousePos != FIntPoint(-1, -1))
 		{
 			//Track in case the user hit esc key to stop inspecting pixel
 			PixelInspectorRealtimeManagement(this, true);
diff --git a/Engine/Source/Editor/UnrealEd/Private/Factories/EditorFactories.cpp b/Engine/Source/Editor/UnrealEd/Private/Factories/EditorFactories.cpp
--- a/Engine/Source/Editor/UnrealEd/Private/Factories/EditorFactories.cpp
+++ b/Engine/Source/Editor/UnrealEd/Private/Factories/EditorFactories.cpp
@@ -5273,17 +5273,37 @@ void FCustomizableTextObjectFactory::ProcessBuffer(UObject* InParent, EObjectFla
 
 	FParse::Next( &Buffer );
 
+	int32 NestedDepth     = 0;
+	int32 OmittedOuterObj = 0; // zero signifies "nothing omitted"
+
 	FString StrLine;
 	while( FParse::Line(&Buffer,StrLine) )
 	{
 		const TCHAR* Str = *StrLine;
 		if( GetBEGIN(&Str,TEXT("OBJECT")) )
 		{
+			++NestedDepth;
+			if (OmittedOuterObj > 0)
+			{
+				if (NestedDepth > OmittedOuterObj)
+				{
+					continue;
+				}
+				ensure(OmittedOuterObj == NestedDepth);
+				// clear the omitted outer, we've parsed passed it
+				OmittedOuterObj = 0;
+			}
+
 			UClass* ObjClass;
 			if( ParseObject<UClass>( Str, TEXT("CLASS="), ObjClass, ANY_PACKAGE ) )
 			{
-				if (!CanCreateClass(ObjClass))
+				bool bOmitSubObjects = false;
+				if (!CanCreateClass(ObjClass, bOmitSubObjects))
 				{
+					if (bOmitSubObjects)
+					{
+						OmittedOuterObj = NestedDepth;
+					}
 					continue;
 				}
 
@@ -5343,6 +5363,10 @@ void FCustomizableTextObjectFactory::ProcessBuffer(UObject* InParent, EObjectFla
 				NewObjects.Add(CreatedObject);
 			}
 		}
+		else if (GetEND(&Str, TEXT("OBJECT")))
+		{
+			--NestedDepth;
+		}
 	}
 
 	// Apply the property text to each of the created objects
@@ -5368,28 +5392,52 @@ bool FCustomizableTextObjectFactory::CanCreateObjectsFromText( const FString& Te
 
 	FParse::Next( &Buffer );
 
+	int32 NestedDepth     = 0;
+	int32 OmittedOuterObj = 0; // zero signifies "nothing omitted"
+
 	FString StrLine;
 	while( FParse::Line(&Buffer,StrLine) )
 	{
 		const TCHAR* Str = *StrLine;
 		if( GetBEGIN(&Str,TEXT("OBJECT")) )
 		{
+			++NestedDepth;
+			if (OmittedOuterObj > 0)
+			{
+				if (NestedDepth > OmittedOuterObj)
+				{
+					continue;
+				}
+				ensure(OmittedOuterObj == NestedDepth);
+				// clear the omitted outer, we've parsed passed it
+				OmittedOuterObj = 0;
+			}
+
 			UClass* ObjClass;
-			if( ParseObject<UClass>( Str, TEXT("CLASS="), ObjClass, ANY_PACKAGE ) )
+			if (ParseObject<UClass>(Str, TEXT("CLASS="), ObjClass, ANY_PACKAGE))
 			{
-				if(CanCreateClass(ObjClass))
+				bool bOmitSubObjects = false;;
+				if (CanCreateClass(ObjClass, bOmitSubObjects))
 				{
 					bCanCreate = true;
 					break;
 				}
+				else if (bOmitSubObjects)
+				{
+					OmittedOuterObj = NestedDepth;
+				}
 			}
 		}
+		else if ( GetEND(&Str, TEXT("OBJECT")) )
+		{
+			--NestedDepth;
+		}
 	}
 	return bCanCreate;
 }
 
 /** Return true if the an object of type ObjectClass is allowed to be created; If false is returned, the object and subobjects will be ignored. */
-bool FCustomizableTextObjectFactory::CanCreateClass(UClass* ObjectClass) const
+bool FCustomizableTextObjectFactory::CanCreateClass(UClass* ObjectClass, bool& bOmitSubObjs) const
 {
 	return false;
 }
diff --git a/Engine/Source/Editor/UnrealEd/Private/Fbx/FbxFactory.cpp b/Engine/Source/Editor/UnrealEd/Private/Fbx/FbxFactory.cpp
--- a/Engine/Source/Editor/UnrealEd/Private/Fbx/FbxFactory.cpp
+++ b/Engine/Source/Editor/UnrealEd/Private/Fbx/FbxFactory.cpp
@@ -167,8 +167,11 @@ UObject* UFbxFactory::FactoryCreateBinary
 	// this one prints all messages that are stored in FFbxImporter
 	UnFbx::FFbxImporter* FbxImporter = UnFbx::FFbxImporter::GetInstance();
 	UnFbx::FBXImportOptions* ImportOptions = FbxImporter->GetImportOptions();
-	//Clean up the options
-	UnFbx::FBXImportOptions::ResetOptions(ImportOptions);
+	if (bShowOption)
+	{
+		//Clean up the options
+		UnFbx::FBXImportOptions::ResetOptions(ImportOptions);
+	}
 	
 	UnFbx::FFbxLoggerSetter Logger(FbxImporter);
 
diff --git a/Engine/Source/Editor/UnrealEd/Private/Fbx/FbxMainImport.cpp b/Engine/Source/Editor/UnrealEd/Private/Fbx/FbxMainImport.cpp
--- a/Engine/Source/Editor/UnrealEd/Private/Fbx/FbxMainImport.cpp
+++ b/Engine/Source/Editor/UnrealEd/Private/Fbx/FbxMainImport.cpp
@@ -1801,9 +1801,7 @@ void FFbxImporter::RecursiveFindRigidMesh(FbxNode* Node, TArray< TArray<FbxNode*
 		int32 ChildIndex;
 		for (ChildIndex=0; ChildIndex<Node->GetChildCount(); ++ChildIndex)
 		{
-			FbxNode* MeshNode = FindLODGroupNode(Node, ChildIndex);
-			if(MeshNode != nullptr)
-				RecursiveFindRigidMesh(MeshNode, outSkelMeshArray, SkeletonArray, ExpandLOD);
+			RecursiveFindRigidMesh(Node->GetChild(ChildIndex), outSkelMeshArray, SkeletonArray, ExpandLOD);
 		}
 	}
 }
diff --git a/Engine/Source/Editor/UnrealEd/Private/Fbx/FbxSceneImportFactory.cpp b/Engine/Source/Editor/UnrealEd/Private/Fbx/FbxSceneImportFactory.cpp
--- a/Engine/Source/Editor/UnrealEd/Private/Fbx/FbxSceneImportFactory.cpp
+++ b/Engine/Source/Editor/UnrealEd/Private/Fbx/FbxSceneImportFactory.cpp
@@ -1067,7 +1067,7 @@ USceneComponent *CreateLightComponent(AActor *ParentActor, TSharedPtr<FFbxLightI
 		UPointLightComponent *PointLightComponent = NewObject<UPointLightComponent>(ParentActor, *(LightInfo->Name));
 		PointLightComponent->SetAttenuationRadius(LightInfo->EnableFarAttenuation ? LightInfo->FarAttenuationEnd : 16384.0f);
 		LightComponent = static_cast<ULightComponent*>(PointLightComponent);
-		LightComponent->SetIntensity(LightInfo->Intensity);
+		LightComponent->SetIntensity(LightComponent->Intensity * LightInfo->Intensity / 100.0f);
 	}
 		break;
 	case 1:
@@ -1075,7 +1075,8 @@ USceneComponent *CreateLightComponent(AActor *ParentActor, TSharedPtr<FFbxLightI
 		//Directional light
 		UDirectionalLightComponent *DirectionalLightComponent = NewObject<UDirectionalLightComponent>(ParentActor, *(LightInfo->Name));
 		LightComponent = static_cast<ULightComponent*>(DirectionalLightComponent);
-		LightComponent->SetIntensity(LightInfo->Intensity/100); //In unreal this is a ratio not a lumen value
+		//We cannot convert fbx value to unreal value so we kept the default object value
+		LightComponent->SetIntensity(LightComponent->Intensity * LightInfo->Intensity / 100.0f);
 	}
 		break;
 	case 2:
@@ -1086,7 +1087,8 @@ USceneComponent *CreateLightComponent(AActor *ParentActor, TSharedPtr<FFbxLightI
 		SpotLightComponent->SetOuterConeAngle(LightInfo->OuterAngle/2.0f);
 		SpotLightComponent->SetAttenuationRadius(LightInfo->EnableFarAttenuation ? LightInfo->FarAttenuationEnd : 16384.0f);
 		LightComponent = static_cast<ULightComponent*>(SpotLightComponent);
-		LightComponent->SetIntensity(LightInfo->Intensity);
+		
+		LightComponent->SetIntensity(LightComponent->Intensity * LightInfo->Intensity / 100.0f);
 	}
 		break;
 	case 3:
diff --git a/Engine/Source/Editor/UnrealEd/Private/Fbx/FbxStaticMeshImport.cpp b/Engine/Source/Editor/UnrealEd/Private/Fbx/FbxStaticMeshImport.cpp
--- a/Engine/Source/Editor/UnrealEd/Private/Fbx/FbxStaticMeshImport.cpp
+++ b/Engine/Source/Editor/UnrealEd/Private/Fbx/FbxStaticMeshImport.cpp
@@ -954,7 +954,7 @@ UStaticMesh* UnFbx::FFbxImporter::ImportStaticMeshAsSingle(UObject* InParent, TA
 
 	// Parent package to place new meshes
 	UPackage* Package = NULL;
-	if (InParent != nullptr && InParent->IsA(UPackage::StaticClass()))
+	if (ImportOptions->bImportScene && InParent != nullptr && InParent->IsA(UPackage::StaticClass()))
 	{
 		Package = StaticCast<UPackage*>(InParent);
 	}
diff --git a/Engine/Source/Editor/UnrealEd/Private/FileHelpers.cpp b/Engine/Source/Editor/UnrealEd/Private/FileHelpers.cpp
--- a/Engine/Source/Editor/UnrealEd/Private/FileHelpers.cpp
+++ b/Engine/Source/Editor/UnrealEd/Private/FileHelpers.cpp
@@ -2175,7 +2175,7 @@ void FEditorFileUtils::LoadMap(const FString& InFilename, bool LoadAsTemplate, b
 		FMainMRUFavoritesList* MRUFavoritesList = MainFrameModule.GetMRUFavoritesList();
 		if(MRUFavoritesList)
 		{
-			MRUFavoritesList->AddMRUItem( *Filename );
+			MRUFavoritesList->AddMRUItem(LongMapPackageName);
 		}
 	}
 
diff --git a/Engine/Source/Editor/UnrealEd/Private/GroupActor.cpp b/Engine/Source/Editor/UnrealEd/Private/GroupActor.cpp
--- a/Engine/Source/Editor/UnrealEd/Private/GroupActor.cpp
+++ b/Engine/Source/Editor/UnrealEd/Private/GroupActor.cpp
@@ -100,6 +100,21 @@ bool ActorHasParentInGroup(const TArray<class AActor*> &GroupActors, const AActo
 	return false;
 }
 
+bool ActorHasParentInSelection(const AActor* Actor)
+{
+	check(Actor);
+	bool bHasParentInSelection = false;
+	AActor* ParentActor = Actor->GetAttachParentActor();
+	while (ParentActor != NULL && !bHasParentInSelection)
+	{
+		if (ParentActor->IsSelected())
+		{
+			bHasParentInSelection = true;
+		}
+		ParentActor = ParentActor->GetAttachParentActor();
+	}
+	return bHasParentInSelection;
+}
 
 void AGroupActor::GroupApplyDelta(FLevelEditorViewportClient* Viewport, const FVector& InDrag, const FRotator& InRot, const FVector& InScale )
 {
@@ -108,11 +123,11 @@ void AGroupActor::GroupApplyDelta(FLevelEditorViewportClient* Viewport, const FV
 	{
 		if( GroupActors[ActorIndex] != NULL )
 		{
-			// Check that we've not got a parent attachment within the group.
-			const bool canApplyDelta = !ActorHasParentInGroup(GroupActors, GroupActors[ActorIndex]);
-			if( canApplyDelta )
+			// Check that we've not got a parent attachment within the group/selection
+			const bool bCanApplyDelta = !ActorHasParentInGroup(GroupActors, GroupActors[ActorIndex]) && !ActorHasParentInSelection(GroupActors[ActorIndex]);
+			if(bCanApplyDelta)
 			{
-				Viewport->ApplyDeltaToActor( GroupActors[ActorIndex], InDrag, InRot, InScale );
+				Viewport->ApplyDeltaToActor(GroupActors[ActorIndex], InDrag, InRot, InScale);
 			}
 		}
 	}
@@ -130,9 +145,9 @@ void AGroupActor::GroupApplyDelta(const FVector& InDrag, const FRotator& InRot,
 {
 	for(int32 ActorIndex=0; ActorIndex<GroupActors.Num(); ++ActorIndex)
 	{
-		// Check that we've not got a parent attachment within the group.
-		const bool canApplyDelta = !ActorHasParentInGroup(GroupActors, GroupActors[ActorIndex]);
-		if( canApplyDelta && GroupActors[ActorIndex] != NULL )
+		// Check that we've not got a parent attachment within the group/selection
+		const bool bCanApplyDelta = !ActorHasParentInGroup(GroupActors, GroupActors[ActorIndex]) && !ActorHasParentInSelection(GroupActors[ActorIndex]);
+		if(bCanApplyDelta && GroupActors[ActorIndex] != NULL )
 		{
 			GEditor->ApplyDeltaToActor( GroupActors[ActorIndex], true, &InDrag, &InRot, &InScale );
 		}
diff --git a/Engine/Source/Editor/UnrealEd/Private/Kismet2/BlueprintEditorUtils.cpp b/Engine/Source/Editor/UnrealEd/Private/Kismet2/BlueprintEditorUtils.cpp
--- a/Engine/Source/Editor/UnrealEd/Private/Kismet2/BlueprintEditorUtils.cpp
+++ b/Engine/Source/Editor/UnrealEd/Private/Kismet2/BlueprintEditorUtils.cpp
@@ -5149,14 +5149,28 @@ bool FBlueprintEditorUtils::IsVariableUsed(const UBlueprint* Blueprint, const FN
 			TArray<UK2Node_Variable*> GraphNodes;
 			CurrentGraph->GetNodesOfClass(GraphNodes);
 
-			for( TArray<UK2Node_Variable*>::TConstIterator NodeIt(GraphNodes); NodeIt; ++NodeIt )
+			for (const UK2Node_Variable* CurrentNode : GraphNodes )
 			{
-				UK2Node_Variable* CurrentNode = *NodeIt;
 				if(Name == CurrentNode->GetVarName())
 				{
 					return true;
 				}
 			}
+
+			// Also consider "used" if there's a GetClassDefaults node that exposes the variable as an output pin that's connected to something.
+			TArray<UK2Node_GetClassDefaults*> ClassDefaultsNodes;
+			CurrentGraph->GetNodesOfClass(ClassDefaultsNodes);
+			for (const UK2Node_GetClassDefaults* ClassDefaultsNode : ClassDefaultsNodes)
+			{
+				if (ClassDefaultsNode->GetInputClass() == Blueprint->SkeletonGeneratedClass)
+				{
+					const UEdGraphPin* VarPin = ClassDefaultsNode->FindPin(Name.ToString());
+					if (VarPin && VarPin->Direction == EGPD_Output && VarPin->LinkedTo.Num() > 0)
+					{
+						return true;
+					}
+				}
+			}
 		}
 	}
 	return false;
diff --git a/Engine/Source/Editor/UnrealEd/Private/MRUFavoritesList.cpp b/Engine/Source/Editor/UnrealEd/Private/MRUFavoritesList.cpp
--- a/Engine/Source/Editor/UnrealEd/Private/MRUFavoritesList.cpp
+++ b/Engine/Source/Editor/UnrealEd/Private/MRUFavoritesList.cpp
@@ -38,30 +38,31 @@ void FMainMRUFavoritesList::WriteToINI() const
 }
 
 /**
- * Add a new file item to the favorites list
+ * Add a new item to the favorites list
  *
- * @param	Item	Filename of the item to add to the favorites list
+ * @param	Item	Package name of the item to add to the favorites list
  */
 void FMainMRUFavoritesList::AddFavoritesItem( const FString& Item )
 {
+	check(FPackageName::IsValidLongPackageName(Item));
+
 	// Only add the item if it isn't already a favorite!
-	const FString CleanedName = FPaths::ConvertRelativePathToFull(Item);
-	if ( !FavoriteItems.Contains( CleanedName ) )
+	if ( !FavoriteItems.Contains( Item ) )
 	{
-		FavoriteItems.Insert( CleanedName, 0 );
+		FavoriteItems.Insert( Item, 0 );
 		WriteToINI();
 	}
 }
 
 /**
- * Remove a file from the favorites list
+ * Remove an item from the favorites list
  *
- * @param	Item	Filename of the item to remove from the favorites list
+ * @param	Item	Package name of the item to remove from the favorites list
  */
 void FMainMRUFavoritesList::RemoveFavoritesItem( const FString& Item )
 {
-	const FString CleanedName = FPaths::ConvertRelativePathToFull(Item);
-	const int32 ItemIndex = FavoriteItems.Find( CleanedName );
+	check(FPackageName::IsValidLongPackageName(Item));
+	const int32 ItemIndex = FavoriteItems.Find( Item );
 	if ( ItemIndex != INDEX_NONE )
 	{
 		FavoriteItems.RemoveAt( ItemIndex );
@@ -72,27 +73,28 @@ void FMainMRUFavoritesList::RemoveFavoritesItem( const FString& Item )
 /**
  * Moves the specified favorites item to the head of the list
  *
- * @param	Item	Filename of the item to move
+ * @param	Item	Package name of the item to move
  */
 void FMainMRUFavoritesList::MoveFavoritesItemToHead(const FString& Item)
 {
-	const FString CleanedName = FPaths::ConvertRelativePathToFull(Item);
+	check(FPackageName::IsValidLongPackageName(Item));
 	if ( FavoriteItems.RemoveSingle(Item) == 1 )
 	{
-		FavoriteItems.Insert( CleanedName, 0 );
+		FavoriteItems.Insert( Item, 0 );
 		WriteToINI();
 	}
 }
 
 /**
- * Returns whether a filename is favorited or not
+ * Returns whether a package name is favorited or not
  *
- * @param	Item	Filename of the item to check
+ * @param	Item	Package name of the item to check
  *
  * @return	true if the provided item is in the favorite's list; false if it is not
  */
 bool FMainMRUFavoritesList::ContainsFavoritesItem( const FString& Item )
 {
+	check(FPackageName::IsValidLongPackageName(Item));
 	const FString CleanedName = FPaths::ConvertRelativePathToFull(Item);
 	return FavoriteItems.Contains( CleanedName );
 }
@@ -121,19 +123,21 @@ FString FMainMRUFavoritesList::GetFavoritesItem( int32 ItemIndex ) const
 bool FMainMRUFavoritesList::VerifyFavoritesFile( int32 ItemIndex )
 {
 	check( FavoriteItems.IsValidIndex( ItemIndex ) );
-	const FString& CurFileName = FavoriteItems[ ItemIndex ];
+	const FString& CurPackageName = FavoriteItems[ ItemIndex ];
+
+	FString CurFileName;
+	bool bSuccess = FPackageName::TryConvertLongPackageNameToFilename(CurPackageName, CurFileName, FPackageName::GetMapPackageExtension());
 
-	const bool bFileExists = IFileManager::Get().FileSize( *CurFileName ) != INDEX_NONE;
-	
 	// If the file doesn't exist any more, remove it from the favorites list and alert the user
-	if ( !bFileExists )
+	if ( !bSuccess || IFileManager::Get().FileSize(*CurFileName) == INDEX_NONE )
 	{
-		FNotificationInfo Info(FText::Format(NSLOCTEXT("UnrealEd", "Error_FavoritesFileDoesNotExist", "Map '{0}' does not exist - it will be removed from the Favorites list."), FText::FromString(FPaths::GetCleanFilename(CurFileName))));
+		FNotificationInfo Info(FText::Format(NSLOCTEXT("UnrealEd", "Error_FavoritesFileDoesNotExist", "Map '{0}' does not exist - it will be removed from the Favorites list."), FText::FromString(CurPackageName)));
 		Info.bUseThrobber = false;
 		Info.ExpireDuration = 8.0f;
 		FSlateNotificationManager::Get().AddNotification(Info)->SetCompletionState(SNotificationItem::CS_Fail);
-		RemoveFavoritesItem( CurFileName );
+		RemoveFavoritesItem(CurPackageName);
+		return false;
 	}
 
-	return bFileExists;
+	return true;
 }
diff --git a/Engine/Source/Editor/UnrealEd/Private/MRUList.cpp b/Engine/Source/Editor/UnrealEd/Private/MRUList.cpp
--- a/Engine/Source/Editor/UnrealEd/Private/MRUList.cpp
+++ b/Engine/Source/Editor/UnrealEd/Private/MRUList.cpp
@@ -56,11 +56,11 @@ void FMRUList::MoveToTop(int32 InItem)
 
 void FMRUList::AddMRUItem(const FString& InItem)
 {
-	FString CleanedName = FPaths::ConvertRelativePathToFull(InItem);
+	check(FPackageName::IsValidLongPackageName(InItem));
 
 	// See if the item already exists in the list.  If so,
 	// move it to the top of the list and leave.
-	const int32 ItemIndex = Items.Find( CleanedName );
+	const int32 ItemIndex = Items.Find(InItem);
 	if ( ItemIndex != INDEX_NONE )
 	{
 		MoveToTop( ItemIndex );
@@ -68,9 +68,9 @@ void FMRUList::AddMRUItem(const FString& InItem)
 	else
 	{
 		// Item is new, so add it to the bottom of the list.
-		if( CleanedName.Len() )
+		if(InItem.Len() )
 		{
-			new(Items) FString( *CleanedName );
+			new(Items) FString( *InItem );
 			MoveToTop( Items.Num()-1 );
 		}
 
@@ -83,6 +83,8 @@ void FMRUList::AddMRUItem(const FString& InItem)
 
 int32 FMRUList::FindMRUItemIdx(const FString& InItem) const
 {
+	check(FPackageName::IsValidLongPackageName(InItem));
+
 	for( int32 mru = 0 ; mru < Items.Num() ; ++mru )
 	{
 		if( Items[mru] == InItem )
@@ -114,16 +116,37 @@ void FMRUList::InternalReadINI( TArray<FString>& OutItems, const FString& INISec
 	// Clear existing items
 	OutItems.Empty();
 
+	bool bConvertedToNewFormat = false;
+
 	// Iterate over the maximum number of provided elements
 	for( int32 ItemIdx = 0 ; ItemIdx < NumElements ; ++ItemIdx )
 	{
 		// Try to find data for a key formed as "INIKeyBaseItemIdx" for the provided INI section. If found, add the data to the output item array.
 		FString CurItem;
-		if ( GConfig->GetString( *INISection, *FString::Printf( TEXT("%s%d"), *INIKeyBase, ItemIdx ), CurItem, GEditorPerProjectIni ) )
+		if (GConfig->GetString(*INISection, *FString::Printf(TEXT("%s%d"), *INIKeyBase, ItemIdx), CurItem, GEditorPerProjectIni))
 		{
-			OutItems.AddUnique( FPaths::ConvertRelativePathToFull(CurItem) );
+			if (!FPackageName::IsValidLongPackageName(CurItem))
+			{
+				FString NewItem;
+				if (FPackageName::TryConvertFilenameToLongPackageName(CurItem, NewItem))
+				{
+					CurItem = NewItem;
+					OutItems.AddUnique(CurItem);
+				}
+
+				bConvertedToNewFormat = true;
+			}
+			else
+			{
+				OutItems.AddUnique(CurItem);
+			}
 		}
 	}
+
+	if (bConvertedToNewFormat)
+	{
+		InternalWriteINI(OutItems, INISection, INIKeyBase);
+	}
 }
 
 
@@ -143,16 +166,19 @@ void FMRUList::InternalWriteINI( const TArray<FString>& InItems, const FString&
 bool FMRUList::VerifyMRUFile(int32 InItem)
 {
 	check( InItem > -1 && InItem < GetMaxItems() );
-	const FString filename = Items[InItem];
+	const FString PackageName = Items[InItem];
+
+	FString Filename;
+	bool bSuccess = FPackageName::TryConvertLongPackageNameToFilename(PackageName, Filename, FPackageName::GetMapPackageExtension());
 
 	// If the file doesn't exist, tell the user about it, remove the file from the list
-	if( IFileManager::Get().FileSize( *filename ) == -1 )
+	if( !bSuccess || IFileManager::Get().FileSize( *Filename) == INDEX_NONE )
 	{
 		FMessageLog EditorErrors("EditorErrors");
 		FFormatNamedArguments Arguments;
-		Arguments.Add(TEXT("FileName"), FText::FromString(filename));
-		EditorErrors.Warning(FText::Format( NSLOCTEXT("MRUList", "Error_FileDoesNotExist", "File does not exist : '{FileName}'.  It will be removed from the recent items list."), Arguments ) );
-		EditorErrors.Notify(NSLOCTEXT("MRUList", "Notification_FileDoesNotExist", "File does not exist! Removed from recent items list!"));
+		Arguments.Add(TEXT("PackageName"), FText::FromString(PackageName));
+		EditorErrors.Warning(FText::Format( NSLOCTEXT("MRUList", "Error_FileDoesNotExist", "Map '{PackageName}' does not exist.  It will be removed from the recent items list."), Arguments ) );
+		EditorErrors.Notify(NSLOCTEXT("MRUList", "Notification_PackageDoesNotExist", "Map does not exist! Removed from recent items list!"));
 		RemoveMRUItem( InItem );
 		WriteToINI();
 
diff --git a/Engine/Source/Editor/UnrealEd/Private/PlayLevel.cpp b/Engine/Source/Editor/UnrealEd/Private/PlayLevel.cpp
--- a/Engine/Source/Editor/UnrealEd/Private/PlayLevel.cpp
+++ b/Engine/Source/Editor/UnrealEd/Private/PlayLevel.cpp
@@ -2474,6 +2474,9 @@ void UEditorEngine::SpawnIntraProcessPIEWorlds(bool bAnyBlueprintErrors, bool bS
 
 	// Restore window settings
 	GetMultipleInstancePositions(0, NextX, NextY);	// restore cached settings
+
+	// Give focus to the first client
+	GiveFocusToFirstClientPIEViewport();
 }
 
 bool UEditorEngine::CreatePIEWorldFromLogin(FWorldContext& PieWorldContext, EPlayNetMode PlayNetMode, FPieLoginStruct& DataStruct)
@@ -2565,6 +2568,9 @@ void UEditorEngine::LoginPIEInstances(bool bAnyBlueprintErrors, bool bStartInSpe
 	const EPlayNetMode PlayNetMode = [&PlayInSettings]{ EPlayNetMode NetMode(PIE_Standalone); return (PlayInSettings->GetPlayNetMode(NetMode) ? NetMode : PIE_Standalone); }();
 	const bool CanPlayNetDedicated = [&PlayInSettings]{ bool PlayNetDedicated(false); return (PlayInSettings->GetPlayNetDedicated(PlayNetDedicated) && PlayNetDedicated); }();
 	const bool WillAutoConnectToServer = [&PlayInSettings]{ bool AutoConnectToServer(false); return (PlayInSettings->GetAutoConnectToServer(AutoConnectToServer) && AutoConnectToServer); }();
+	const int32 PlayNumberOfClients = [&PlayInSettings] { int32 NumberOfClients(0); return (PlayInSettings->GetPlayNumberOfClients(NumberOfClients) ? NumberOfClients : 0); }();
+
+	PIEInstancesToLogInCount = PlayNumberOfClients;
 
 	// Server
 	if (WillAutoConnectToServer || CanPlayNetDedicated)
@@ -2620,7 +2626,6 @@ void UEditorEngine::LoginPIEInstances(bool bAnyBlueprintErrors, bool bStartInSpe
 	}
 
 	// Clients
-	const int32 PlayNumberOfClients = [&PlayInSettings]{ int32 NumberOfClients(0); return (PlayInSettings->GetPlayNumberOfClients(NumberOfClients) ? NumberOfClients : 0); }();
 	for (; ClientNum < PlayNumberOfClients; ++ClientNum)
 	{
 		PlayInSettings->SetPlayNetMode(PlayNetMode);
@@ -2711,6 +2716,41 @@ void UEditorEngine::OnLoginPIEComplete_Deferred(int32 LocalUserNum, bool bWasSuc
 			}
 		}
 	}
+
+	PIEInstancesToLogInCount--;
+	if (PIEInstancesToLogInCount == 0)
+	{
+		OnLoginPIEAllComplete();
+	}
+}
+
+void UEditorEngine::OnLoginPIEAllComplete()
+{
+	GiveFocusToFirstClientPIEViewport();
+}
+
+void UEditorEngine::GiveFocusToFirstClientPIEViewport()
+{
+	// Find the non-dedicated server or first client window to give focus to
+	int32 LowestPIEInstance = TNumericLimits<int32>::Max();
+	UGameViewportClient* ViewportClient = nullptr;
+	for (const FWorldContext& WorldContext : WorldList)
+	{
+		if (WorldContext.WorldType == EWorldType::PIE && !WorldContext.RunAsDedicated)
+		{
+			if (WorldContext.PIEInstance < LowestPIEInstance)
+			{
+				LowestPIEInstance = WorldContext.PIEInstance;
+				ViewportClient = WorldContext.GameViewport;
+			}
+		}
+	}
+
+	// Give focus to the first client
+	if (ViewportClient && ViewportClient->GetGameViewportWidget().IsValid())
+	{
+		FSlateApplication::Get().RegisterGameViewport(ViewportClient->GetGameViewportWidget().ToSharedRef());
+	}
 }
 
 void UEditorEngine::RequestLateJoin()
diff --git a/Engine/Source/Editor/UnrealEd/Private/Settings/SettingsClasses.cpp b/Engine/Source/Editor/UnrealEd/Private/Settings/SettingsClasses.cpp
--- a/Engine/Source/Editor/UnrealEd/Private/Settings/SettingsClasses.cpp
+++ b/Engine/Source/Editor/UnrealEd/Private/Settings/SettingsClasses.cpp
@@ -126,12 +126,11 @@ void UEditorExperimentalSettings::PostEditChangeProperty( struct FPropertyChange
 
 UEditorLoadingSavingSettings::UEditorLoadingSavingSettings( const FObjectInitializer& ObjectInitializer )
 	: Super(ObjectInitializer)
-	, bEnableSourceControlCompatabilityCheck(true)
 	, bMonitorContentDirectories(true)
 	, AutoReimportThreshold(3.f)
 	, bAutoCreateAssets(true)
 	, bAutoDeleteAssets(true)
-	, bDetectChangesOnStartup(false)
+	, bDetectChangesOnStartup(true)
 	, bDeleteSourceFilesWithAssets(false)
 {
 	TextDiffToolPath.FilePath = TEXT("P4Merge.exe");
@@ -187,64 +186,6 @@ void UEditorLoadingSavingSettings::PostInitProperties()
 	Super::PostInitProperties();
 }
 
-void UEditorLoadingSavingSettings::CheckSourceControlCompatability()
-{
-	if (!bEnableSourceControlCompatabilityCheck || !bMonitorContentDirectories)
-	{
-		return;
-	}
-
-	if (ISourceControlModule::Get().IsEnabled() && bDetectChangesOnStartup)
-	{
-		// Persistent shared payload captured by the lambdas below
-		struct FPersistentPayload { TSharedPtr<SNotificationItem> Notification; };
-		TSharedRef<FPersistentPayload> Payload = MakeShareable(new FPersistentPayload);
-
-		FNotificationInfo Info(LOCTEXT("AutoReimport_NotificationTitle", "We noticed that your auto-reimport settings are set up to detect source content changes on restart.\nThis might cause unexpected behavior when starting up after getting latest from source control.\n\nWe recommend disabling this specific behavior."));
-
- 		auto OnTurnOffClicked = [=]{
-			auto* Settings = GetMutableDefault<UEditorLoadingSavingSettings>();
-			Settings->bDetectChangesOnStartup = false;
-			Settings->SaveConfig();
-
-			Payload->Notification->SetEnabled(false);
-			Payload->Notification->Fadeout();
-		};
-		Info.ButtonDetails.Emplace(LOCTEXT("AutoReimport_TurnOff", "Don't detect changes on start-up"), FText(), FSimpleDelegate::CreateLambda(OnTurnOffClicked), SNotificationItem::ECompletionState::CS_None);
-
- 		auto OnIgnoreClicked = [=]{
-
-			Payload->Notification->SetEnabled(false);
- 			Payload->Notification->Fadeout();
-		};
-		Info.ButtonDetails.Emplace(LOCTEXT("AutoReimport_Ignore", "Ignore"), FText(), FSimpleDelegate::CreateLambda(OnIgnoreClicked), SNotificationItem::ECompletionState::CS_None);
-
-		Info.bUseLargeFont = false;
-		Info.bFireAndForget = false;
-
-		Info.CheckBoxStateChanged = FOnCheckStateChanged::CreateLambda([](ECheckBoxState State){
-			auto* Settings = GetMutableDefault<UEditorLoadingSavingSettings>();
-			Settings->bEnableSourceControlCompatabilityCheck = (State != ECheckBoxState::Checked);
-			Settings->SaveConfig();
-		});
-		Info.CheckBoxText = LOCTEXT("AutoReimport_DontShowAgain", "Don't show again");
-
-		Info.Hyperlink = FSimpleDelegate::CreateLambda([]{
-			// Open Settings
-			ISettingsModule* SettingsModule = FModuleManager::GetModulePtr<ISettingsModule>("Settings");
-			if (SettingsModule != nullptr)
-			{
-				// Ensure that the advanced properties are visible
-				GConfig->SetBool(TEXT("DetailCategoriesAdvanced"), TEXT("EditorLoadingSavingSettings.AutoReimport"), true, GEditorPerProjectIni);
-				SettingsModule->ShowViewer("Editor", "General", "LoadingSaving");
-			}
-		});
-		Info.HyperlinkText = LOCTEXT("AutoReimport_OpenSettings", "Settings");
-
-		Payload->Notification = FSlateNotificationManager::Get().AddNotification(Info);
-	}
-}
-
 FAutoReimportDirectoryConfig::FParseContext::FParseContext(bool bInEnableLogging)
 	: bEnableLogging(bInEnableLogging)
 {
diff --git a/Engine/Source/Editor/UnrealEd/Private/Toolkits/AssetEditorManager.cpp b/Engine/Source/Editor/UnrealEd/Private/Toolkits/AssetEditorManager.cpp
--- a/Engine/Source/Editor/UnrealEd/Private/Toolkits/AssetEditorManager.cpp
+++ b/Engine/Source/Editor/UnrealEd/Private/Toolkits/AssetEditorManager.cpp
@@ -706,4 +706,12 @@ void FAssetEditorManager::OpenEditorsForAssets(const TArray<FString>& AssetsToOp
 	}
 }
 
+void FAssetEditorManager::OpenEditorsForAssets(const TArray<FName>& AssetsToOpen)
+{
+	for (const FName AssetName : AssetsToOpen)
+	{
+		OpenEditorForAsset(AssetName.ToString());
+	}
+}
+
 #undef LOCTEXT_NAMESPACE
\ No newline at end of file
diff --git a/Engine/Source/Editor/UnrealEd/Private/UnrealEdEngine.cpp b/Engine/Source/Editor/UnrealEd/Private/UnrealEdEngine.cpp
--- a/Engine/Source/Editor/UnrealEd/Private/UnrealEdEngine.cpp
+++ b/Engine/Source/Editor/UnrealEd/Private/UnrealEdEngine.cpp
@@ -115,8 +115,6 @@ void UUnrealEdEngine::Init(IEngineLoop* InEngineLoop)
 	{
 		AutoReimportManager = NewObject<UAutoReimportManager>();
 		AutoReimportManager->Initialize();
-		
-		GetMutableDefault<UEditorLoadingSavingSettings>()->CheckSourceControlCompatability();
 	}
 
 	// register details panel customizations
diff --git a/Engine/Source/Editor/UnrealEd/Public/Factories.h b/Engine/Source/Editor/UnrealEd/Public/Factories.h
--- a/Engine/Source/Editor/UnrealEd/Public/Factories.h
+++ b/Engine/Source/Editor/UnrealEd/Public/Factories.h
@@ -39,7 +39,7 @@ public:
 
 protected:
 	/** Return true if the an object of type ObjectClass is allowed to be created; If false is returned, the object and subobjects will be ignored. */
-	virtual bool CanCreateClass(UClass* ObjectClass) const;
+	virtual bool CanCreateClass(UClass* ObjectClass, bool& bOmitSubObjs) const;
 
 	/** This is called on each created object after the property text is imported */
 	virtual void ProcessConstructedObject(UObject* CreatedObject);
diff --git a/Engine/Source/Editor/UnrealEd/Public/MRUList.h b/Engine/Source/Editor/UnrealEd/Public/MRUList.h
--- a/Engine/Source/Editor/UnrealEd/Public/MRUList.h
+++ b/Engine/Source/Editor/UnrealEd/Public/MRUList.h
@@ -33,7 +33,7 @@ public:
 	/**
 	 * Adds an item to the MRU list, moving it to the top.
 	 *
-	 * @param	Item		The item to add.
+	 * @param	Item		The long package name of the item to add.
 	 */
 	void AddMRUItem(const FString& Item);
 	
diff --git a/Engine/Source/Editor/UnrealEd/Public/Toolkits/AssetEditorManager.h b/Engine/Source/Editor/UnrealEd/Public/Toolkits/AssetEditorManager.h
--- a/Engine/Source/Editor/UnrealEd/Public/Toolkits/AssetEditorManager.h
+++ b/Engine/Source/Editor/UnrealEd/Public/Toolkits/AssetEditorManager.h
@@ -66,6 +66,7 @@ public:
 
 	/** Opens editors for the supplied assets (via OpenEditorForAsset) */
 	void OpenEditorsForAssets(const TArray<FString>& AssetsToOpen);
+	void OpenEditorsForAssets(const TArray<FName>& AssetsToOpen);
 
 	/** Returns the primary editor if one is already open for the specified asset.
 	 * If there is one open and bFocusIfOpen is true, that editor will be brought to the foreground and focused if possible.
diff --git a/Engine/Source/Editor/UnrealEd/UnrealEd.Build.cs b/Engine/Source/Editor/UnrealEd/UnrealEd.Build.cs
--- a/Engine/Source/Editor/UnrealEd/UnrealEd.Build.cs
+++ b/Engine/Source/Editor/UnrealEd/UnrealEd.Build.cs
@@ -45,7 +45,8 @@ public class UnrealEd : ModuleRules
 				"PlacementMode",
 				"Settings",
 				"SettingsEditor",
-				"SoundClassEditor",
+                "SuperSearch",
+                "SoundClassEditor",
 				"ViewportSnapping",
 				"SourceCodeAccess",
 				"ReferenceViewer",
@@ -122,7 +123,6 @@ public class UnrealEd : ModuleRules
 				"SoundClassEditor",
 				"SoundCueEditor",
 				"SourceControlWindows",
-                "SuperSearch",
 				"StatsViewer",
 				"SwarmInterface",
 				"TargetPlatform",
@@ -182,7 +182,8 @@ public class UnrealEd : ModuleRules
 				"SettingsEditor",
 				"SessionFrontend",
 				"Sequencer",
-				"GeometryMode",
+                "SuperSearch",
+                "GeometryMode",
 				"TextureAlignMode",
 				"FoliageEdit",
 				"PackageDependencyInfo",
@@ -230,7 +231,6 @@ public class UnrealEd : ModuleRules
 			{
                 "GraphEditor",
 				"Kismet",
-                "SuperSearch",
             }
 		); 
 
diff --git a/Engine/Source/Programs/AutomationTool/GUBP/ElectricCommander.cs b/Engine/Source/Programs/AutomationTool/GUBP/ElectricCommander.cs
--- a/Engine/Source/Programs/AutomationTool/GUBP/ElectricCommander.cs
+++ b/Engine/Source/Programs/AutomationTool/GUBP/ElectricCommander.cs
@@ -464,7 +464,7 @@ namespace AutomationTool
 						JsonWriter.WriteValue("DescriptionText", NodeToDo.DescriptionText);
 						if (NodeToDo.RecipientsForFailureEmails.Length > 0)
 						{
-							JsonWriter.WriteValue("Notify", String.Join(" ", NodeToDo.RecipientsForFailureEmails));
+							JsonWriter.WriteValue("Notify", String.Join(";", NodeToDo.RecipientsForFailureEmails));
 						}
 						JsonWriter.WriteValue("IsTrigger", true);
 						JsonWriter.WriteObjectEnd();
diff --git a/Engine/Source/Programs/ShaderCompileWorker/Private/ShaderCompileWorker.cpp b/Engine/Source/Programs/ShaderCompileWorker/Private/ShaderCompileWorker.cpp
--- a/Engine/Source/Programs/ShaderCompileWorker/Private/ShaderCompileWorker.cpp
+++ b/Engine/Source/Programs/ShaderCompileWorker/Private/ShaderCompileWorker.cpp
@@ -580,8 +580,8 @@ static EShaderPlatform FormatNameToEnum(FName ShaderFormat)
 	if (ShaderFormat == NAME_VULKAN_SM4)			return SP_VULKAN_SM4;
 	if (ShaderFormat == NAME_VULKAN_SM5)			return SP_VULKAN_SM5;
 	if (ShaderFormat == NAME_VULKAN_ES3_1_ANDROID)	return SP_VULKAN_ES3_1_ANDROID;
-	if (ShaderFormat == NAME_VULKAN_ES3_1)			return SP_VULKAN_ES3_1_ANDROID;
-	if (ShaderFormat == NAME_VULKAN_ES3_1_UB)		return SP_VULKAN_ES3_1_ANDROID;
+	if (ShaderFormat == NAME_VULKAN_ES3_1)			return SP_VULKAN_PCES3_1;
+	if (ShaderFormat == NAME_VULKAN_ES3_1_UB)		return SP_VULKAN_PCES3_1;
 	if (ShaderFormat == NAME_SF_METAL_SM4)		return SP_METAL_SM4;
 	if (ShaderFormat == NAME_SF_METAL_MACES3_1)	return SP_METAL_MACES3_1;
 	return SP_NumPlatforms;
diff --git a/Engine/Source/Programs/UnrealBuildTool/Configuration/UEBuildTarget.cs b/Engine/Source/Programs/UnrealBuildTool/Configuration/UEBuildTarget.cs
--- a/Engine/Source/Programs/UnrealBuildTool/Configuration/UEBuildTarget.cs
+++ b/Engine/Source/Programs/UnrealBuildTool/Configuration/UEBuildTarget.cs
@@ -3561,6 +3561,10 @@ namespace UnrealBuildTool
 				}
 			}
 
+			// Remove any enabled plugins that are unused on the current platform. This prevents having to stage the .uplugin files, but requires that the project descriptor
+			// doesn't have a platform-neutral reference to it.
+			EnabledPlugins.RemoveAll(Plugin => !UProjectInfo.IsPluginDescriptorRequiredForProject(Plugin, ProjectDescriptor, Platform, TargetType, UEBuildConfiguration.bBuildDeveloperTools, UEBuildConfiguration.bBuildEditor));
+
 			// Set the list of plugins that should be built
 			if (Rules.bBuildAllPlugins)
 			{
diff --git a/Engine/Source/Programs/UnrealBuildTool/System/UProjectInfo.cs b/Engine/Source/Programs/UnrealBuildTool/System/UProjectInfo.cs
--- a/Engine/Source/Programs/UnrealBuildTool/System/UProjectInfo.cs
+++ b/Engine/Source/Programs/UnrealBuildTool/System/UProjectInfo.cs
@@ -320,5 +320,45 @@ namespace UnrealBuildTool
 			}
 			return bEnabled;
 		}
+
+		/// <summary>
+		/// Determine if a plugin is enabled for a given project
+		/// </summary>
+		/// <param name="Project">The project to check</param>
+		/// <param name="Plugin">Information about the plugin</param>
+		/// <param name="Platform">The target platform</param>
+		/// <returns>True if the plugin should be enabled for this project</returns>
+		public static bool IsPluginDescriptorRequiredForProject(PluginInfo Plugin, ProjectDescriptor Project, UnrealTargetPlatform Platform, TargetRules.TargetType TargetType, bool bBuildDeveloperTools, bool bBuildEditor)
+		{
+			// Check if it's referenced by name from the project descriptor. If it is, we'll need the plugin to be included with the project regardless of whether it has
+			// any platform-specific modules or content, just so the runtime can make the call.
+			if (Project != null && Project.Plugins != null)
+			{
+				foreach (PluginReferenceDescriptor PluginReference in Project.Plugins)
+				{
+					if (String.Compare(PluginReference.Name, Plugin.Name, true) == 0)
+					{
+						return PluginReference.IsEnabledForPlatform(Platform) && PluginReference.IsEnabledForTarget(TargetType);
+					}
+				}
+			}
+
+			// If the plugin contains content, it should be included for all platforms
+			if(Plugin.Descriptor.bCanContainContent)
+			{
+				return true;
+			}
+
+			// Check if the plugin has any modules for the given target
+			foreach (ModuleDescriptor Module in Plugin.Descriptor.Modules)
+			{
+				if(Module.IsCompiledInConfiguration(Platform, TargetType, bBuildDeveloperTools, bBuildEditor))
+				{
+					return true;
+				}
+			}
+
+			return false;
+		}
 	}
 }
diff --git a/Engine/Source/Programs/UnrealBuildTool/System/VCProjectFileGenerator.cs b/Engine/Source/Programs/UnrealBuildTool/System/VCProjectFileGenerator.cs
--- a/Engine/Source/Programs/UnrealBuildTool/System/VCProjectFileGenerator.cs
+++ b/Engine/Source/Programs/UnrealBuildTool/System/VCProjectFileGenerator.cs
@@ -282,13 +282,6 @@ namespace UnrealBuildTool
 						{
 							VCProjectFileFormat ProposedFormat = ProjectFileFormat;
 
-							// Visual Studio 2015 is not supported by PS4 VSI
-							if( SupportedPlatform == UnrealTargetPlatform.PS4 )
-							{
-								Log.TraceInformation("Forcing Visual Studio max version to 2013 projects for PS4 compatibility (use '-2015' to override.)");
-								ProposedFormat = VCProjectFileFormat.VisualStudio2013;
-							}
-
 							// Visual Studio 2015 is not supported by the Android debugger we currently furnish
 							if (SupportedPlatform == UnrealTargetPlatform.Android)
 							{
diff --git a/Engine/Source/Runtime/AnimGraphRuntime/Public/BoneControllers/AnimNode_TwistCorrectiveNode.h b/Engine/Source/Runtime/AnimGraphRuntime/Public/BoneControllers/AnimNode_TwistCorrectiveNode.h
--- a/Engine/Source/Runtime/AnimGraphRuntime/Public/BoneControllers/AnimNode_TwistCorrectiveNode.h
+++ b/Engine/Source/Runtime/AnimGraphRuntime/Public/BoneControllers/AnimNode_TwistCorrectiveNode.h
@@ -59,7 +59,10 @@ struct FReferenceBoneFrame
 };
 
 /**
- * This is the runtime version of a bone driven controller, which maps part of the state from one bone to another (e.g., 2 * source.x -> target.z)
+ * This is the node that apply corrective morphtarget for twist 
+ * Good example is that if you twist your neck too far right or left, you're going to see odd stretch shape of neck, 
+ * This node can detect the angle and apply morphtarget curve 
+ * This isn't the twist control node for bone twist
  */
 USTRUCT()
 struct ANIMGRAPHRUNTIME_API FAnimNode_TwistCorrectiveNode : public FAnimNode_SkeletalControlBase
diff --git a/Engine/Source/Runtime/AppFramework/Private/Widgets/Colors/SColorPicker.cpp b/Engine/Source/Runtime/AppFramework/Private/Widgets/Colors/SColorPicker.cpp
--- a/Engine/Source/Runtime/AppFramework/Private/Widgets/Colors/SColorPicker.cpp
+++ b/Engine/Source/Runtime/AppFramework/Private/Widgets/Colors/SColorPicker.cpp
@@ -1291,10 +1291,16 @@ void SColorPicker::HandleColorSpinBoxValueChanged( float NewValue, EColorPickerC
 }
 
 
-void SColorPicker::HandleEyeDropperButtonComplete()
+void SColorPicker::HandleEyeDropperButtonComplete(bool bCancelled)
 {
 	bIsInteractive = false;
 
+	if (bCancelled)
+	{
+		SetNewTargetColorHSV(OldColor, true);
+		RestoreColors();
+	}
+
 	if (bOnlyRefreshOnMouseUp || bPerfIsTooSlowToUpdate)
 	{
 		UpdateColorPick();
diff --git a/Engine/Source/Runtime/AppFramework/Private/Widgets/Colors/SEyeDropperButton.cpp b/Engine/Source/Runtime/AppFramework/Private/Widgets/Colors/SEyeDropperButton.cpp
--- a/Engine/Source/Runtime/AppFramework/Private/Widgets/Colors/SEyeDropperButton.cpp
+++ b/Engine/Source/Runtime/AppFramework/Private/Widgets/Colors/SEyeDropperButton.cpp
@@ -105,7 +105,8 @@ FReply SEyeDropperButton::OnMouseButtonUp(const FGeometry& MyGeometry, const FPo
 
 		Reply.ReleaseMouseCapture();
 
-		OnComplete.ExecuteIfBound();
+		const bool bCancelled = false;
+		OnComplete.ExecuteIfBound(bCancelled);
 	}
 	else if ( MouseEvent.GetEffectingButton() == EKeys::LeftMouseButton && bWasClicked )
 	{
@@ -167,7 +168,8 @@ FReply SEyeDropperButton::OnKeyDown(const FGeometry& MyGeometry, const FKeyEvent
 
 		FReply ReleaseReply = FReply::Handled().ReleaseMouseCapture();
 
-		OnComplete.ExecuteIfBound();
+		const bool bCancelled = true;
+		OnComplete.ExecuteIfBound(bCancelled);
 
 		return ReleaseReply;
 	}
diff --git a/Engine/Source/Runtime/AppFramework/Private/Widgets/Colors/SEyeDropperButton.h b/Engine/Source/Runtime/AppFramework/Private/Widgets/Colors/SEyeDropperButton.h
--- a/Engine/Source/Runtime/AppFramework/Private/Widgets/Colors/SEyeDropperButton.h
+++ b/Engine/Source/Runtime/AppFramework/Private/Widgets/Colors/SEyeDropperButton.h
@@ -12,7 +12,7 @@
 class SEyeDropperButton : public SButton
 {
 public:
-	DECLARE_DELEGATE(FOnDropperComplete)
+	DECLARE_DELEGATE_OneParam(FOnDropperComplete, bool)
 
 	SLATE_BEGIN_ARGS( SEyeDropperButton )
 		: _OnValueChanged()
@@ -22,7 +22,7 @@ public:
 		{}
 
 		/** Invoked when a new value is selected by the dropper */
-	SLATE_EVENT(FOnLinearColorValueChanged, OnValueChanged)
+		SLATE_EVENT(FOnLinearColorValueChanged, OnValueChanged)
 
 		/** Invoked when the dropper goes from inactive to active */
 		SLATE_EVENT(FSimpleDelegate, OnBegin)
diff --git a/Engine/Source/Runtime/AppFramework/Public/Widgets/Colors/SColorPicker.h b/Engine/Source/Runtime/AppFramework/Public/Widgets/Colors/SColorPicker.h
--- a/Engine/Source/Runtime/AppFramework/Public/Widgets/Colors/SColorPicker.h
+++ b/Engine/Source/Runtime/AppFramework/Public/Widgets/Colors/SColorPicker.h
@@ -269,7 +269,7 @@ private:
 	void HandleColorSpinBoxValueChanged( float NewValue, EColorPickerChannels Channel );
 
 	// Callback for completed eye dropper interactions.
-	void HandleEyeDropperButtonComplete();
+	void HandleEyeDropperButtonComplete(bool bCancelled);
 
 	// Callback for getting the text in the hex linear box.
 	FText HandleHexLinearBoxText() const;
diff --git a/Engine/Source/Runtime/Apple/MetalRHI/Private/MetalCommandEncoder.cpp b/Engine/Source/Runtime/Apple/MetalRHI/Private/MetalCommandEncoder.cpp
--- a/Engine/Source/Runtime/Apple/MetalRHI/Private/MetalCommandEncoder.cpp
+++ b/Engine/Source/Runtime/Apple/MetalRHI/Private/MetalCommandEncoder.cpp
@@ -74,6 +74,7 @@ FMetalCommandEncoder::~FMetalCommandEncoder(void)
 		[DepthStencilState release];
 	}
 	[RenderPassDesc release];
+	RenderPassDesc = nil;
 
 	if(DebugGroups)
 	{
@@ -206,6 +207,11 @@ bool FMetalCommandEncoder::IsBlitCommandEncoderActive(void) const
 	return BlitCommandEncoder != nil;
 }
 
+bool FMetalCommandEncoder::IsRenderPassDescriptorValid(void) const
+{
+	return (RenderPassDesc != nil);
+}
+
 id<MTLRenderCommandEncoder> FMetalCommandEncoder::GetRenderCommandEncoder(void) const
 {
 	return RenderCommandEncoder;
@@ -514,6 +520,7 @@ void FMetalCommandEncoder::SetRenderPassDescriptor(MTLRenderPassDescriptor* cons
 		}
 		RenderPassDesc = RenderPass;
 	}
+	check(RenderPassDesc);
 	
 	if (bReset)
 	{
diff --git a/Engine/Source/Runtime/Apple/MetalRHI/Private/MetalCommandEncoder.h b/Engine/Source/Runtime/Apple/MetalRHI/Private/MetalCommandEncoder.h
--- a/Engine/Source/Runtime/Apple/MetalRHI/Private/MetalCommandEncoder.h
+++ b/Engine/Source/Runtime/Apple/MetalRHI/Private/MetalCommandEncoder.h
@@ -52,6 +52,9 @@ public:
 	
 	/** @returns True if and only if there is an active blit command encoder, otherwise false. */
 	bool IsBlitCommandEncoderActive(void) const;
+	
+	/** @returns True if and only if there is valid render pass descriptor set on the encoder, otherwise false. */
+	bool IsRenderPassDescriptorValid(void) const;
 
 	/** @returns The active render command encoder or nil if there isn't one. */
 	id<MTLRenderCommandEncoder> GetRenderCommandEncoder(void) const;
diff --git a/Engine/Source/Runtime/Apple/MetalRHI/Private/MetalCommands.cpp b/Engine/Source/Runtime/Apple/MetalRHI/Private/MetalCommands.cpp
--- a/Engine/Source/Runtime/Apple/MetalRHI/Private/MetalCommands.cpp
+++ b/Engine/Source/Runtime/Apple/MetalRHI/Private/MetalCommands.cpp
@@ -886,13 +886,13 @@ void FMetalRHICommandContext::RHIClearMRT(bool bClearColor,int32 NumClearColors,
 	
 	if (Context->GetCurrentState().GetNumRenderTargets() <= 1)
 	{
-		BlendStateRHI = (bClearColor && Context->GetCurrentState().GetHasValidRenderTarget())
+		BlendStateRHI = (bClearColor && Context->GetCurrentState().GetHasValidColorTarget())
 		? TStaticBlendState<>::GetRHI()
 		: TStaticBlendState<CW_NONE>::GetRHI();
 	}
 	else
 	{
-		BlendStateRHI = (bClearColor && Context->GetCurrentState().GetHasValidRenderTarget())
+		BlendStateRHI = (bClearColor && Context->GetCurrentState().GetHasValidColorTarget())
 			? TStaticBlendState<>::GetRHI()
 			: TStaticBlendStateWriteMask<CW_NONE,CW_NONE,CW_NONE,CW_NONE,CW_NONE,CW_NONE,CW_NONE,CW_NONE>::GetRHI();
 	}
diff --git a/Engine/Source/Runtime/Apple/MetalRHI/Private/MetalContext.cpp b/Engine/Source/Runtime/Apple/MetalRHI/Private/MetalContext.cpp
--- a/Engine/Source/Runtime/Apple/MetalRHI/Private/MetalContext.cpp
+++ b/Engine/Source/Runtime/Apple/MetalRHI/Private/MetalContext.cpp
@@ -387,10 +387,23 @@ void FMetalDeviceContext::EndDrawingViewport(FMetalViewport* Viewport, bool bPre
 	// enqueue a present if desired
 	if (bPresent)
 	{
-		id<MTLDrawable> Drawable = Viewport->GetDrawable();
-		[CurrentCommandBuffer addScheduledHandler:^(id<MTLCommandBuffer>) {
-			[Drawable present];
-		 }];
+		id<CAMetalDrawable> Drawable = (id<CAMetalDrawable>)Viewport->GetDrawable();
+		if (Drawable.texture)
+		{
+#if PLATFORM_MAC
+			id<MTLBlitCommandEncoder> Blitter = GetBlitContext();
+			
+			id<MTLTexture> Src = Viewport->GetBackBuffer()->Surface.Texture;
+			id<MTLTexture> Dst = Drawable.texture;
+			
+			[Blitter copyFromTexture:Src sourceSlice:0 sourceLevel:0 sourceOrigin:MTLOriginMake(0, 0, 0) sourceSize:MTLSizeMake(Src.width, Src.height, 1) toTexture:Dst destinationSlice:0 destinationLevel:0 destinationOrigin:MTLOriginMake(0, 0, 0)];
+			
+			CommandEncoder.EndEncoding();
+#endif
+			[CurrentCommandBuffer addScheduledHandler:^(id<MTLCommandBuffer>) {
+				[Drawable present];
+			}];
+		}
 	}
 	
 	// We may be limiting our framerate to the display link
@@ -737,6 +750,9 @@ void FMetalContext::PrepareToDraw(uint32 PrimitiveType)
 
 	TRefCountPtr<FMetalBoundShaderState> CurrentBoundShaderState = StateCache.GetBoundShaderState();
 	
+	// Enforce calls to SetRenderTarget prior to issuing draw calls.
+	check(StateCache.GetHasValidRenderTarget());
+	
 	// Validate the vertex layout in debug mode, or when the validation layer is enabled for development builds.
 	// Other builds will just crash & burn if it is incorrect.
 #if UE_BUILD_DEBUG || UE_BUILD_DEVELOPMENT
@@ -795,6 +811,10 @@ void FMetalContext::PrepareToDraw(uint32 PrimitiveType)
 		StateCache.SetRenderTargetsInfo(Info, StateCache.GetVisibilityResultsBuffer(), false);
 		
 		bRestoreState = true;
+		
+		// Enforce calls to SetRenderTarget prior to issuing draw calls.
+		check(StateCache.GetHasValidRenderTarget());
+		check(CommandEncoder.IsRenderPassDescriptorValid());
 	}
 	
 	// make sure the BSS has a valid pipeline state object
@@ -812,6 +832,14 @@ void FMetalContext::PrepareToDraw(uint32 PrimitiveType)
 		
 		if(!CommandEncoder.IsRenderCommandEncoderActive())
 		{
+			if (!CommandEncoder.IsRenderPassDescriptorValid())
+			{
+				UE_LOG(LogMetal, Warning, TEXT("Re-binding the render-target because no RenderPassDescriptor was bound!"));
+				FRHISetRenderTargetsInfo Info = StateCache.GetRenderTargetsInfo();
+				StateCache.SetHasValidRenderTarget(false);
+				StateCache.SetRenderTargetsInfo(Info, StateCache.GetVisibilityResultsBuffer(), false);
+			}
+			
 			CommandEncoder.RestoreRenderCommandEncoding();
 		}
 		else if (bRestoreState)
diff --git a/Engine/Source/Runtime/Apple/MetalRHI/Private/MetalRHI.cpp b/Engine/Source/Runtime/Apple/MetalRHI/Private/MetalRHI.cpp
--- a/Engine/Source/Runtime/Apple/MetalRHI/Private/MetalRHI.cpp
+++ b/Engine/Source/Runtime/Apple/MetalRHI/Private/MetalRHI.cpp
@@ -103,7 +103,7 @@ FMetalDynamicRHI::FMetalDynamicRHI()
 		GMaxRHIShaderPlatform = SP_METAL_SM4;
 	}
 	
-	GShaderPlatformForFeatureLevel[ERHIFeatureLevel::ES2] = SP_METAL_MACES3_1;
+	GShaderPlatformForFeatureLevel[ERHIFeatureLevel::ES2] = SP_METAL_MACES2;
 	GShaderPlatformForFeatureLevel[ERHIFeatureLevel::ES3_1] = SP_METAL_MACES3_1;
 	GShaderPlatformForFeatureLevel[ERHIFeatureLevel::SM4] = SP_METAL_SM4;
 	GShaderPlatformForFeatureLevel[ERHIFeatureLevel::SM5] = (GMaxRHIFeatureLevel >= ERHIFeatureLevel::SM5) ? GMaxRHIShaderPlatform : SP_NumPlatforms;
diff --git a/Engine/Source/Runtime/Apple/MetalRHI/Private/MetalStateCache.cpp b/Engine/Source/Runtime/Apple/MetalRHI/Private/MetalStateCache.cpp
--- a/Engine/Source/Runtime/Apple/MetalRHI/Private/MetalStateCache.cpp
+++ b/Engine/Source/Runtime/Apple/MetalRHI/Private/MetalStateCache.cpp
@@ -36,6 +36,7 @@ FMetalStateCache::FMetalStateCache(FMetalCommandEncoder& InCommandEncoder)
 , FrameBufferSize(CGSizeMake(0.0, 0.0))
 , RenderTargetArraySize(1)
 , bHasValidRenderTarget(false)
+, bHasValidColorTarget(false)
 , bScissorRectEnabled(false)
 {
 	Viewport.originX = Viewport.originY = Viewport.width = Viewport.height = Viewport.znear = Viewport.zfar = 0.0;
@@ -79,6 +80,7 @@ void FMetalStateCache::Reset(void)
 	
 	FMemory::Memzero(RenderTargetsInfo);
 	bHasValidRenderTarget = false;
+	bHasValidColorTarget = false;
 	
 	FMemory::Memzero(DirtyUniformBuffers);
 	
@@ -264,6 +266,7 @@ void FMetalStateCache::SetRenderTargetsInfo(FRHISetRenderTargetsInfo const& InRe
 		PipelineDesc.SampleCount = 0;
 	
 		bHasValidRenderTarget = false;
+		bHasValidColorTarget = false;
 		
 		uint8 ArrayTargets = 0;
 		uint8 BoundTargets = 0;
@@ -302,11 +305,8 @@ void FMetalStateCache::SetRenderTargetsInfo(FRHISetRenderTargetsInfo const& InRe
 	
 				BoundTargets |= 1 << RenderTargetIndex;
 				
-	            if (Surface.Texture == nil)
-	            {
-	                PipelineDesc.SampleCount = OldCount;
-	                return;
-	            }
+				// The surface cannot be nil - we have to have a valid render-target array after this call.
+				check (Surface.Texture != nil);
 	
 				// user code generally passes -1 as a default, but we need 0
 				uint32 ArraySliceIndex = RenderTargetView.ArraySliceIndex == 0xFFFFFFFF ? 0 : RenderTargetView.ArraySliceIndex;
@@ -386,6 +386,7 @@ void FMetalStateCache::SetRenderTargetsInfo(FRHISetRenderTargetsInfo const& InRe
 				[ColorAttachment release];
 	
 				bHasValidRenderTarget = true;
+				bHasValidColorTarget = true;
 			}
 			else
 			{
@@ -552,6 +553,8 @@ void FMetalStateCache::SetRenderTargetsInfo(FRHISetRenderTargetsInfo const& InRe
 				{
 					PipelineDesc.SampleCount = DepthAttachment.texture.sampleCount;
 				}
+				
+				bHasValidRenderTarget = true;
 	
 				// and assign it
 				RenderPass.depthAttachment = DepthAttachment;
@@ -592,6 +595,8 @@ void FMetalStateCache::SetRenderTargetsInfo(FRHISetRenderTargetsInfo const& InRe
 				{
 					PipelineDesc.SampleCount = StencilAttachment.texture.sampleCount;
 				}
+				
+				bHasValidRenderTarget = true;
 	
 				// and assign it
 				RenderPass.stencilAttachment = StencilAttachment;
@@ -600,6 +605,7 @@ void FMetalStateCache::SetRenderTargetsInfo(FRHISetRenderTargetsInfo const& InRe
 				[StencilAttachment release];
 			}
 		}
+		check(bHasValidRenderTarget);
 	
 		// update hash for the depth/stencil buffer & sample count
 		PipelineDesc.SetHashValue(Offset_DepthFormat, NumBits_DepthFormat, DepthFormatKey);
@@ -629,11 +635,15 @@ void FMetalStateCache::SetRenderTargetsInfo(FRHISetRenderTargetsInfo const& InRe
 			BoundShaderState.SafeRelease();
 		}
 	}
+	
+	// Ensure that the RenderPassDesc is valid in the CommandEncoder.
+	check(CommandEncoder.IsRenderPassDescriptorValid());
 }
 
 void FMetalStateCache::SetHasValidRenderTarget(bool InHasValidRenderTarget)
 {
 	bHasValidRenderTarget = InHasValidRenderTarget;
+	bHasValidColorTarget = InHasValidRenderTarget;
 }
 
 void FMetalStateCache::SetViewport(const MTLViewport& InViewport)
@@ -748,6 +758,7 @@ void FMetalStateCache::ConditionalUpdateBackBuffer(FMetalSurface& Surface)
 			// set the texture into the backbuffer
 			Surface.GetDrawableTexture();
 		}
+		check (Surface.Texture);
 	}
 }
 
@@ -756,7 +767,7 @@ bool FMetalStateCache::NeedsToSetRenderTarget(const FRHISetRenderTargetsInfo& In
 	// see if our new Info matches our previous Info
 	
 	// basic checks
-	bool bAllChecksPassed = GetHasValidRenderTarget() && CommandEncoder.IsRenderCommandEncoderActive() && InRenderTargetsInfo.NumColorRenderTargets == RenderTargetsInfo.NumColorRenderTargets &&
+	bool bAllChecksPassed = GetHasValidRenderTarget() && CommandEncoder.IsRenderPassDescriptorValid() && CommandEncoder.IsRenderCommandEncoderActive() && InRenderTargetsInfo.NumColorRenderTargets == RenderTargetsInfo.NumColorRenderTargets &&
 		// handle the case where going from backbuffer + depth -> backbuffer + null, no need to reset RT and do a store/load
 		(InRenderTargetsInfo.DepthStencilRenderTarget.Texture == RenderTargetsInfo.DepthStencilRenderTarget.Texture || InRenderTargetsInfo.DepthStencilRenderTarget.Texture == nullptr);
 
diff --git a/Engine/Source/Runtime/Apple/MetalRHI/Private/MetalStateCache.h b/Engine/Source/Runtime/Apple/MetalRHI/Private/MetalStateCache.h
--- a/Engine/Source/Runtime/Apple/MetalRHI/Private/MetalStateCache.h
+++ b/Engine/Source/Runtime/Apple/MetalRHI/Private/MetalStateCache.h
@@ -54,8 +54,9 @@ public:
 	FMetalComputeShader* GetComputeShader() const { return ComputeShader; }
 	CGSize GetFrameBufferSize() const { return FrameBufferSize; }
 	FRHISetRenderTargetsInfo const& GetRenderTargetsInfo() const { return RenderTargetsInfo; }
-	int32 GetNumRenderTargets() { return bHasValidRenderTarget ? RenderTargetsInfo.NumColorRenderTargets : -1; }
+	int32 GetNumRenderTargets() { return bHasValidColorTarget ? RenderTargetsInfo.NumColorRenderTargets : -1; }
 	bool GetHasValidRenderTarget() const { return bHasValidRenderTarget; }
+	bool GetHasValidColorTarget() const { return bHasValidColorTarget; }
 	const MTLViewport& GetViewport() const { return Viewport; }
 	id<MTLBuffer> GetVertexBuffer(uint32 const Index) { check(Index < MaxMetalStreams); return VertexBuffers[Index]; }
 	uint32 GetVertexStride(uint32 const Index) { check(Index < MaxMetalStreams); return VertexStrides[Index]; }
@@ -102,6 +103,6 @@ private:
 	
 	FRHISetRenderTargetsInfo RenderTargetsInfo;
 	bool bHasValidRenderTarget;
-	
+	bool bHasValidColorTarget;
 	bool bScissorRectEnabled;
 };
diff --git a/Engine/Source/Runtime/Apple/MetalRHI/Private/MetalTexture.cpp b/Engine/Source/Runtime/Apple/MetalRHI/Private/MetalTexture.cpp
--- a/Engine/Source/Runtime/Apple/MetalRHI/Private/MetalTexture.cpp
+++ b/Engine/Source/Runtime/Apple/MetalRHI/Private/MetalTexture.cpp
@@ -854,7 +854,7 @@ void* FMetalSurface::Lock(uint32 MipIndex, uint32 ArrayIndex, EResourceLockMode
 				//kick the current command buffer.
 				GetMetalDeviceContext().SubmitCommandBufferAndWait();
 #endif
-				[Texture getBytes:LockedMemory[MipIndex] bytesPerRow:DestStride bytesPerImage:MipBytes fromRegion:Region mipmapLevel:MipIndex slice:ArrayIndex];
+				[Texture getBytes:[LockedMemory[MipIndex] contents] bytesPerRow:DestStride bytesPerImage:MipBytes fromRegion:Region mipmapLevel:MipIndex slice:ArrayIndex];
 			}
 			
 #if PLATFORM_MAC
diff --git a/Engine/Source/Runtime/Apple/MetalRHI/Private/MetalViewport.cpp b/Engine/Source/Runtime/Apple/MetalRHI/Private/MetalViewport.cpp
--- a/Engine/Source/Runtime/Apple/MetalRHI/Private/MetalViewport.cpp
+++ b/Engine/Source/Runtime/Apple/MetalRHI/Private/MetalViewport.cpp
@@ -80,20 +80,20 @@ FMetalViewport::FMetalViewport(void* WindowHandle, uint32 InSizeX,uint32 InSizeY
 
 FMetalViewport::~FMetalViewport()
 {
-	BackBuffer.SafeRelease();
+	BackBuffer.SafeRelease();	// when the rest of the engine releases it, its framebuffers will be released too (those the engine knows about)
 	check(!IsValidRef(BackBuffer));
 }
 
 void FMetalViewport::Resize(uint32 InSizeX, uint32 InSizeY, bool bInIsFullscreen)
 {
-	BackBuffer.SafeRelease();	// when the rest of the engine releases it, its framebuffers will be released too (those the engine knows about)
-
 	FRHIResourceCreateInfo CreateInfo;
-	BackBuffer = (FMetalTexture2D*)(FTexture2DRHIParamRef)GDynamicRHI->RHICreateTexture2D(InSizeX, InSizeY, PF_B8G8R8A8, 1, 1, TexCreate_RenderTargetable | TexCreate_Presentable, CreateInfo);
-	BackBuffer->Surface.Viewport = this;
+	BackBuffer.SafeRelease();	// when the rest of the engine releases it, its framebuffers will be released too (those the engine knows about)
+	
 #if PLATFORM_MAC // @todo zebra: ios
+	BackBuffer = (FMetalTexture2D*)(FTexture2DRHIParamRef)GDynamicRHI->RHICreateTexture2D(InSizeX, InSizeY, PF_B8G8R8A8, 1, 1, TexCreate_RenderTargetable, CreateInfo);
 	((CAMetalLayer*)View.layer).drawableSize = CGSizeMake(InSizeX, InSizeY);
 #else
+	BackBuffer = (FMetalTexture2D*)(FTexture2DRHIParamRef)GDynamicRHI->RHICreateTexture2D(InSizeX, InSizeY, PF_B8G8R8A8, 1, 1, TexCreate_RenderTargetable | TexCreate_Presentable, CreateInfo);
 	IOSAppDelegate* AppDelegate = [IOSAppDelegate GetDelegate];
 	FIOSView* GLView = AppDelegate.IOSView;
 	[GLView UpdateRenderWidth:InSizeX andHeight:InSizeY];
@@ -104,6 +104,7 @@ void FMetalViewport::Resize(uint32 InSizeX, uint32 InSizeY, bool bInIsFullscreen
 	check(FMath::TruncToInt(ScalingFactor * ViewFrame.size.width) == InSizeX &&
 		  FMath::TruncToInt(ScalingFactor * ViewFrame.size.height) == InSizeY);
 #endif
+	BackBuffer->Surface.Viewport = this;
 }
 
 id<MTLDrawable> FMetalViewport::GetDrawable()
@@ -141,7 +142,9 @@ void FMetalViewport::ReleaseDrawable()
 		[Drawable release];
 		Drawable = nil;
 	}
+#if !PLATFORM_MAC
 	BackBuffer->Surface.Texture = nil;
+#endif
 }
 
 #if PLATFORM_MAC
@@ -186,6 +189,8 @@ void FMetalDynamicRHI::RHIBeginDrawingViewport(FViewportRHIParamRef ViewportRHI,
 {
 	FMetalViewport* Viewport = ResourceCast(ViewportRHI);
 
+	check(Viewport);
+	
 	((FMetalDeviceContext*)Context)->BeginDrawingViewport(Viewport);
 
 	// Set the render target and viewport.
diff --git a/Engine/Source/Runtime/AssetRegistry/Private/AssetRegistryModule.cpp b/Engine/Source/Runtime/AssetRegistry/Private/AssetRegistryModule.cpp
--- a/Engine/Source/Runtime/AssetRegistry/Private/AssetRegistryModule.cpp
+++ b/Engine/Source/Runtime/AssetRegistry/Private/AssetRegistryModule.cpp
@@ -5,27 +5,12 @@
 #include "AssetRegistryModule.h"
 #include "AssetRegistryConsoleCommands.h"
 
-#include "Developer/MessageLog/Public/MessageLogModule.h"
-
-#define LOCTEXT_NAMESPACE "AssetRegistry"
-
 IMPLEMENT_MODULE( FAssetRegistryModule, AssetRegistry );
 
 void FAssetRegistryModule::StartupModule()
 {
 	AssetRegistry = new FAssetRegistry();
 	ConsoleCommands = new FAssetRegistryConsoleCommands(*this);
-
-#if WITH_UNREAL_DEVELOPER_TOOLS
-	// MessageLog needs EditorStyle to be loaded, and we might be being initialized before that has happened
-	FModuleManager::LoadModuleChecked<IModuleInterface>("EditorStyle");
-
-	// create a message log for the asset registry to use
-	FMessageLogModule& MessageLogModule = FModuleManager::LoadModuleChecked<FMessageLogModule>("MessageLog");
-	FMessageLogInitializationOptions InitOptions;
-	InitOptions.bShowPages = true;
-	MessageLogModule.RegisterLogListing("AssetRegistry", LOCTEXT("AssetRegistryLogLabel", "Asset Registry"), InitOptions);
-#endif
 }
 
 
@@ -42,15 +27,6 @@ void FAssetRegistryModule::ShutdownModule()
 		delete ConsoleCommands;
 		ConsoleCommands = NULL;
 	}
-
-#if WITH_UNREAL_DEVELOPER_TOOLS
-	if( FModuleManager::Get().IsModuleLoaded("MessageLog") )
-	{
-		// unregister message log
-		FMessageLogModule& MessageLogModule = FModuleManager::GetModuleChecked<FMessageLogModule>("MessageLog");
-		MessageLogModule.UnregisterLogListing("AssetRegistry");
-	}
-#endif
 }
 
 IAssetRegistry& FAssetRegistryModule::Get() const
@@ -58,5 +34,3 @@ IAssetRegistry& FAssetRegistryModule::Get() const
 	check(AssetRegistry);
 	return *AssetRegistry;
 }
-
-#undef LOCTEXT_NAMESPACE
\ No newline at end of file
diff --git a/Engine/Source/Runtime/Core/Private/Android/AndroidMisc.cpp b/Engine/Source/Runtime/Core/Private/Android/AndroidMisc.cpp
--- a/Engine/Source/Runtime/Core/Private/Android/AndroidMisc.cpp
+++ b/Engine/Source/Runtime/Core/Private/Android/AndroidMisc.cpp
@@ -601,12 +601,17 @@ const int32 TargetSignals[] =
 const int32 NumTargetSignals = ARRAY_COUNT(TargetSignals);
 
 struct sigaction PrevActions[NumTargetSignals];
+static bool PreviousSignalHandlersValid = false;
 
 static void RestorePreviousSignalHandlers()
 {
-	for (int32 i = 0; i < NumTargetSignals; ++i)
+	if (PreviousSignalHandlersValid)
 	{
-		sigaction(TargetSignals[i],	&PrevActions[i], NULL);
+		for (int32 i = 0; i < NumTargetSignals; ++i)
+		{
+			sigaction(TargetSignals[i], &PrevActions[i], NULL);
+		}
+		PreviousSignalHandlersValid = false;
 	}
 }
 
@@ -616,7 +621,6 @@ void PlatformCrashHandler(int32 Signal, siginfo* Info, void* Context)
 	// Switch to malloc crash.
 	//FGenericPlatformMallocCrash::Get().SetAsGMalloc(); @todo uncomment after verification
 
-	//fprintf(stderr, "Signal %d caught.\n", Signal);
 	FPlatformMisc::LowLevelOutputDebugStringf(TEXT("Signal %d caught!"), Signal);
 
 	// Restore system handlers so Android could catch this signal after we are done with crashreport
@@ -640,8 +644,15 @@ void FAndroidMisc::SetCrashHandler(void (* CrashHandler)(const FGenericCrashCont
 {
 	GCrashHandlerPointer = CrashHandler;
 
+	RestorePreviousSignalHandlers();
 	FMemory::Memzero(&PrevActions, sizeof(PrevActions));
 
+	// Passing -1 will leave these restored and won't trap them
+	if ((int)CrashHandler == -1)
+	{
+		return;
+	}
+
 	struct sigaction Action;
 	FMemory::Memzero(&Action, sizeof(struct sigaction));
 	Action.sa_sigaction = PlatformCrashHandler;
@@ -652,6 +663,7 @@ void FAndroidMisc::SetCrashHandler(void (* CrashHandler)(const FGenericCrashCont
 	{
 		sigaction(TargetSignals[i],	&Action, &PrevActions[i]);
 	}
+	PreviousSignalHandlersValid = true;
 }
 
 bool FAndroidMisc::GetUseVirtualJoysticks()
diff --git a/Engine/Source/Runtime/Core/Private/IOS/IOSPlatformFile.cpp b/Engine/Source/Runtime/Core/Private/IOS/IOSPlatformFile.cpp
--- a/Engine/Source/Runtime/Core/Private/IOS/IOSPlatformFile.cpp
+++ b/Engine/Source/Runtime/Core/Private/IOS/IOSPlatformFile.cpp
@@ -582,11 +582,7 @@ IFileHandle* FIOSPlatformFile::OpenRead(const TCHAR* Filename, bool bAllowWrite)
 IFileHandle* FIOSPlatformFile::OpenWrite(const TCHAR* Filename, bool bAppend, bool bAllowRead)
 {
 	int Flags = O_CREAT;
-	if (bAppend)
-	{
-		Flags |= O_APPEND;
-	}
-	else
+	if (!bAppend)
 	{
 		Flags |= O_TRUNC;
 	}
diff --git a/Engine/Source/Runtime/Core/Public/Math/Transform.h b/Engine/Source/Runtime/Core/Public/Math/Transform.h
--- a/Engine/Source/Runtime/Core/Public/Math/Transform.h
+++ b/Engine/Source/Runtime/Core/Public/Math/Transform.h
@@ -319,7 +319,8 @@ public:
 	FORCEINLINE FTransform Inverse() const
 	{
 		FQuat   InvRotation    = Rotation.Inverse();
-		FVector InvScale3D     = Scale3D.Reciprocal();
+		// this used to cause NaN if Scale contained 0 
+		FVector InvScale3D     = GetSafeScaleReciprocal(Scale3D);
 		FVector InvTranslation = InvRotation * (InvScale3D * -Translation);
 
 		return FTransform(InvRotation, InvTranslation, InvScale3D);
@@ -721,6 +722,18 @@ public:
 	 */
 	FORCEINLINE static void Multiply(FTransform* OutTransform, const FTransform* A, const FTransform* B);
 
+	/**
+	* Create a new transform: OutTransform = A * B using the matrix while keeping the scale that's given by A and B
+	* Please note that this operation is a lot more expensive than normal Multiply
+	*
+	* Order matters when composing transforms : A * B will yield a transform that logically first applies A then B to any subsequent transformation.
+	*
+	* @param  OutTransform pointer to transform that will store the result of A * B.
+	* @param  A Transform A.
+	* @param  B Transform B.
+	*/
+	FORCEINLINE static void MultiplyUsingMatrixWithScale(FTransform* OutTransform, const FTransform* A, const FTransform* B);
+
 	/**
 	 * Sets the components
 	 * @param InRotation The new value for the Rotation component
@@ -1195,14 +1208,47 @@ FORCEINLINE void FTransform::RemoveScaling(float Tolerance/*=SMALL_NUMBER*/)
 	DiagnosticCheckNaN_Scale3D();
 }
 
+FORCEINLINE void FTransform::MultiplyUsingMatrixWithScale(FTransform* OutTransform, const FTransform* A, const FTransform* B)
+{
+	// the goal of using M is to get the correct orientation
+	// but for translation, we still need scale
+	FMatrix M = A->ToMatrixWithScale() * B->ToMatrixWithScale();
+	M.RemoveScaling();
+
+	// get combined scale
+	FVector Scale3D = A->Scale3D*B->Scale3D;
+
+	// apply negative scale back to axes
+	FVector SignedScale = Scale3D.GetSignVector();
+
+	M.SetAxis(0, SignedScale.X * M.GetScaledAxis(EAxis::X));
+	M.SetAxis(1, SignedScale.Y * M.GetScaledAxis(EAxis::Y));
+	M.SetAxis(2, SignedScale.Z * M.GetScaledAxis(EAxis::Z));
+
+	// @note: if you have negative with 0 scale, this will return rotation that is identity
+	// since matrix loses that axes
+	FQuat Rotation = FQuat(M);
+	Rotation.Normalize();
+
+	// set values back to output
+	OutTransform->Scale3D = Scale3D;
+	OutTransform->Rotation = Rotation;
+	
+	// technically I could calculate this using FTransform but then it does more quat multiplication 
+	// instead of using Scale in matrix multiplication
+	// it's a question of between RemoveScaling vs using FTransform to move translation
+	OutTransform->Translation = M.GetOrigin();
+}
+
+
 /** Returns Multiplied Transform of 2 FTransforms **/
 FORCEINLINE void FTransform::Multiply(FTransform* OutTransform, const FTransform* A, const FTransform* B)
 {
 	A->DiagnosticCheckNaN_All();
 	B->DiagnosticCheckNaN_All();
 
-	checkSlow( A->IsRotationNormalized() );
-	checkSlow( B->IsRotationNormalized() );
+	checkSlow(A->IsRotationNormalized());
+	checkSlow(B->IsRotationNormalized());
 
 	//	When Q = quaternion, S = single scalar scale, and T = translation
 	//	QST(A) = Q(A), S(A), T(A), and QST(B) = Q(B), S(B), T(B)
@@ -1219,19 +1265,23 @@ FORCEINLINE void FTransform::Multiply(FTransform* OutTransform, const FTransform
 	//	S(AxB) = S(A)*S(B)
 	//	T(AxB) = Q(B)*S(B)*T(A)*-Q(B) + T(B)
 
-	FTransform Ret;
-	Ret.Rotation = B->Rotation*A->Rotation;
+	const bool bHaveNegativeScale = A->Scale3D.GetMin() < 0 || B->Scale3D.GetMin() < 0;
+	if (bHaveNegativeScale)
+	{
+		// @note, if you have 0 scale with negative, you're going to lose rotation as it can't convert back to quat
+		MultiplyUsingMatrixWithScale(OutTransform, A, B);
+	}
+	else
+	{
+		OutTransform->Rotation = B->Rotation*A->Rotation;
+		OutTransform->Scale3D = A->Scale3D*B->Scale3D;
+		OutTransform->Translation = B->Rotation*(B->Scale3D*A->Translation) + B->Translation;
+	}
 
-	Ret.Scale3D = A->Scale3D*B->Scale3D;
-	Ret.Translation = B->Rotation*(B->Scale3D*A->Translation) + B->Translation;
 	// we do not support matrix transform when non-uniform
 	// that was removed at rev 21 with UE4
-
-	*OutTransform = Ret;
-	Ret.DiagnosticCheckNaN_All();
+	OutTransform->DiagnosticCheckNaN_All();
 }
-
-
 /** 
  * Apply Scale to this transform
  */
diff --git a/Engine/Source/Runtime/Core/Public/Math/TransformVectorized.h b/Engine/Source/Runtime/Core/Public/Math/TransformVectorized.h
--- a/Engine/Source/Runtime/Core/Public/Math/TransformVectorized.h
+++ b/Engine/Source/Runtime/Core/Public/Math/TransformVectorized.h
@@ -744,6 +744,18 @@ public:
 
 	FORCEINLINE static void Multiply(FTransform* OutTransform, const FTransform* A, const FTransform* B);
 
+	/**
+	* Create a new transform: OutTransform = A * B using the matrix while keeping the scale that's given by A and B
+	* Please note that this operation is a lot more expensive than normal Multiply
+	*
+	* Order matters when composing transforms : A * B will yield a transform that logically first applies A then B to any subsequent transformation.
+	*
+	* @param  OutTransform pointer to transform that will store the result of A * B.
+	* @param  A Transform A.
+	* @param  B Transform B.
+	*/
+	FORCEINLINE static void MultiplyUsingMatrixWithScale(FTransform* OutTransform, const FTransform* A, const FTransform* B);
+
 	/**
 	 * Sets the components
 	 * @param InRotation The new value for the Rotation component
@@ -1411,14 +1423,48 @@ FORCEINLINE void FTransform::RemoveScaling(float Tolerance/*=SMALL_NUMBER*/)
 	DiagnosticCheckNaN_Scale3D();
 }
 
+FORCEINLINE void FTransform::MultiplyUsingMatrixWithScale(FTransform* OutTransform, const FTransform* A, const FTransform* B)
+{
+	// the goal of using M is to get the correct orientation
+	// but for translation, we still need scale
+	FMatrix M = A->ToMatrixWithScale() * B->ToMatrixWithScale();
+	M.RemoveScaling();
+
+	// get combined scale
+	const VectorRegister Scale3D = VectorMultiply(A->Scale3D, B->Scale3D);
+
+	// apply negative scale back to axes
+	FVector SignedScale;
+	VectorStoreFloat3(VectorSign(Scale3D), &SignedScale);
+
+	M.SetAxis(0, SignedScale.X * M.GetScaledAxis(EAxis::X));
+	M.SetAxis(1, SignedScale.Y * M.GetScaledAxis(EAxis::Y));
+	M.SetAxis(2, SignedScale.Z * M.GetScaledAxis(EAxis::Z));
+
+	// @note: if you have negative with 0 scale, this will return rotation that is identity
+	// since matrix loses that axes
+	FQuat Rotation = FQuat(M);
+	Rotation.Normalize();
+
+	// set values back to output
+	OutTransform->Scale3D = Scale3D;
+	OutTransform->Rotation = VectorLoadAligned(&Rotation);
+	
+	// technically I could calculate this using FTransform but then it does more quat multiplication 
+	// instead of using Scale in matrix multiplication
+	// it's a question of between RemoveScaling vs using FTransform to move translation
+	FVector Translation = M.GetOrigin();
+	OutTransform->Translation = VectorLoadFloat3_W0(&Translation);
+}
+
 /** Returns Multiplied Transform of 2 FTransforms **/
 FORCEINLINE void FTransform::Multiply(FTransform* OutTransform, const FTransform* A, const FTransform* B)
 {
 	A->DiagnosticCheckNaN_All();
 	B->DiagnosticCheckNaN_All();
 
-	checkSlow( A->IsRotationNormalized() );
-	checkSlow( B->IsRotationNormalized() );
+	checkSlow(A->IsRotationNormalized());
+	checkSlow(B->IsRotationNormalized());
 
 	//	When Q = quaternion, S = single scalar scale, and T = translation
 	//	QST(A) = Q(A), S(A), T(A), and QST(B) = Q(B), S(B), T(B)
@@ -1435,26 +1481,34 @@ FORCEINLINE void FTransform::Multiply(FTransform* OutTransform, const FTransform
 	//	S(AxB) = S(A)*S(B)
 	//	T(AxB) = Q(B)*S(B)*T(A)*-Q(B) + T(B)
 
-	const VectorRegister QuatA = A->Rotation;
-	const VectorRegister QuatB = B->Rotation;
-	const VectorRegister TranslateA = A->Translation;
-	const VectorRegister TranslateB = B->Translation;
-	const VectorRegister ScaleA = A->Scale3D;
-	const VectorRegister ScaleB = B->Scale3D;
-
-	// RotationResult = B.Rotation * A.Rotation
-	OutTransform->Rotation = VectorQuaternionMultiply2(QuatB, QuatA);
+	const VectorRegister MinVector = VectorMin(A->Scale3D, B->Scale3D);
+	const bool bHaveNegativeScale = VectorGetComponent(MinVector, 0) < 0.f || VectorGetComponent(MinVector, 1) < 0.f || VectorGetComponent(MinVector, 2) < 0.f;
+	if (bHaveNegativeScale)
+	{
+		// @note, if you have 0 scale with negative, you're going to lose rotation as it can't convert back to quat
+		MultiplyUsingMatrixWithScale(OutTransform, A, B);
+	}
+	else
+	{
+		const VectorRegister QuatA = A->Rotation;
+		const VectorRegister QuatB = B->Rotation;
+		const VectorRegister TranslateA = A->Translation;
+		const VectorRegister TranslateB = B->Translation;
+		const VectorRegister ScaleA = A->Scale3D;
+		const VectorRegister ScaleB = B->Scale3D;
 
-	// TranslateResult = B.Rotate(B.Scale * A.Translation) + B.Translate
-	const VectorRegister ScaledTransA = VectorMultiply(TranslateA, ScaleB);
-	const VectorRegister RotatedTranslate = VectorQuaternionRotateVector(QuatB, ScaledTransA);
-	OutTransform->Translation = VectorAdd(RotatedTranslate, TranslateB);
+		// RotationResult = B.Rotation * A.Rotation
+		OutTransform->Rotation = VectorQuaternionMultiply2(QuatB, QuatA);
 
-	// ScaleResult = Scale.B * Scale.A
-	OutTransform->Scale3D = VectorMultiply(ScaleA, ScaleB);;
+		// TranslateResult = B.Rotate(B.Scale * A.Translation) + B.Translate
+		const VectorRegister ScaledTransA = VectorMultiply(TranslateA, ScaleB);
+		const VectorRegister RotatedTranslate = VectorQuaternionRotateVector(QuatB, ScaledTransA);
+		OutTransform->Translation = VectorAdd(RotatedTranslate, TranslateB);
 
+		// ScaleResult = Scale.B * Scale.A
+		OutTransform->Scale3D = VectorMultiply(ScaleA, ScaleB);;
+	}
 }
-
 /** 
  * Apply Scale to this transform
  */
diff --git a/Engine/Source/Runtime/Core/Public/UObject/FrameworkObjectVersion.h b/Engine/Source/Runtime/Core/Public/UObject/FrameworkObjectVersion.h
--- a/Engine/Source/Runtime/Core/Public/UObject/FrameworkObjectVersion.h
+++ b/Engine/Source/Runtime/Core/Public/UObject/FrameworkObjectVersion.h
@@ -25,6 +25,10 @@ struct CORE_API FFrameworkObjectVersion
 		// Moved compressed anim data from uasset to the DDC
 		MoveCompressedAnimDataToTheDDC,
 
+		// Some graph pins created using legacy code seem to have lost the RF_Transactional flag,
+		// which causes issues with undo. Restore the flag at this version
+		FixNonTransactionalPins,
+
 		// -----<new versions can be added above this line>-------------------------------------------------
 		VersionPlusOne,
 		LatestVersion = VersionPlusOne - 1
diff --git a/Engine/Source/Runtime/CoreUObject/Private/Internationalization/PackageLocalizationManager.cpp b/Engine/Source/Runtime/CoreUObject/Private/Internationalization/PackageLocalizationManager.cpp
--- a/Engine/Source/Runtime/CoreUObject/Private/Internationalization/PackageLocalizationManager.cpp
+++ b/Engine/Source/Runtime/CoreUObject/Private/Internationalization/PackageLocalizationManager.cpp
@@ -38,6 +38,25 @@ void FDefaultPackageLocalizationCache::FindLocalizedPackages(const FString& InSo
 	}));
 }
 
+void FPackageLocalizationManager::PerformLazyInitialization()
+{
+	if (!ActiveCache.IsValid() && LazyInitFunc)
+	{
+		LazyInitFunc(*this);
+
+		if (!ActiveCache.IsValid())
+		{
+			UE_LOG(LogPackageLocalizationManager, Warning, TEXT("InitializeFromLazyCallback was bound to a callback that didn't initialize the active cache."));
+		}
+	}
+}
+
+void FPackageLocalizationManager::InitializeFromLazyCallback(FLazyInitFunc InLazyInitFunc)
+{
+	LazyInitFunc = MoveTemp(InLazyInitFunc);
+	ActiveCache.Reset();
+}
+
 void FPackageLocalizationManager::InitializeFromCache(const TSharedRef<IPackageLocalizationCache>& InCache)
 {
 	ActiveCache = InCache;
@@ -50,8 +69,10 @@ void FPackageLocalizationManager::InitializeFromDefaultCache()
 	ActiveCache->ConditionalUpdateCache();
 }
 
-FName FPackageLocalizationManager::FindLocalizedPackageName(const FName InSourcePackageName) const
+FName FPackageLocalizationManager::FindLocalizedPackageName(const FName InSourcePackageName)
 {
+	PerformLazyInitialization();
+
 	if (ActiveCache.IsValid())
 	{
 		return ActiveCache->FindLocalizedPackageName(InSourcePackageName);
@@ -62,8 +83,10 @@ FName FPackageLocalizationManager::FindLocalizedPackageName(const FName InSource
 	return FindLocalizedPackageNameNoCache(InSourcePackageName, CurrentCultureName);
 }
 
-FName FPackageLocalizationManager::FindLocalizedPackageNameForCulture(const FName InSourcePackageName, const FString& InCultureName) const
+FName FPackageLocalizationManager::FindLocalizedPackageNameForCulture(const FName InSourcePackageName, const FString& InCultureName)
 {
+	PerformLazyInitialization();
+
 	if (ActiveCache.IsValid())
 	{
 		return ActiveCache->FindLocalizedPackageNameForCulture(InSourcePackageName, InCultureName);
diff --git a/Engine/Source/Runtime/CoreUObject/Private/Internationalization/PackageLocalizationUtil.cpp b/Engine/Source/Runtime/CoreUObject/Private/Internationalization/PackageLocalizationUtil.cpp
new file mode 100644
--- /dev/null
+++ b/Engine/Source/Runtime/CoreUObject/Private/Internationalization/PackageLocalizationUtil.cpp
@@ -0,0 +1,71 @@
+// Copyright 1998-2016 Epic Games, Inc. All Rights Reserved.
+
+#include "CoreUObjectPrivate.h"
+#include "PackageLocalizationUtil.h"
+
+bool FPackageLocalizationUtil::ConvertLocalizedToSource(const FString& InLocalized, FString& OutSource)
+{
+	if (FPackageName::IsLocalizedPackage(InLocalized))
+	{
+		const int32 L10NStart = InLocalized.Find(TEXT("L10N"), ESearchCase::IgnoreCase);
+		if (L10NStart != INDEX_NONE)
+		{
+			// /Game/L10N/fr/Folder/MyAsset
+			//       ^ We matched here, so we need to walk over the L10N folder, and then walk over the culture code to find the range of characters to remove
+			int32 NumSlashesToFind = 2;
+			int32 L10NEnd = L10NStart;
+			for (; L10NEnd < InLocalized.Len() && NumSlashesToFind > 0; ++L10NEnd)
+			{
+				if (InLocalized[L10NEnd] == TEXT('/'))
+				{
+					--NumSlashesToFind;
+				}
+			}
+
+			OutSource = InLocalized;
+			OutSource.RemoveAt(L10NStart, L10NEnd - L10NStart, false);
+			return true;
+		}
+	}
+
+	return false;
+}
+
+bool FPackageLocalizationUtil::ConvertSourceToLocalized(const FString& InSource, const FString& InCulture, FString& OutLocalized)
+{
+	if (!FPackageName::IsLocalizedPackage(InSource))
+	{
+		if (InSource.Len() > 0 && InSource[0] == TEXT('/'))
+		{
+			const int32 RootPathEnd = InSource.Find(TEXT("/"), ESearchCase::IgnoreCase, ESearchDir::FromStart, 1);
+			if (RootPathEnd != INDEX_NONE)
+			{
+				OutLocalized = InSource;
+				OutLocalized.InsertAt(RootPathEnd, TEXT("/L10N") / InCulture);
+				return true;
+			}
+		}
+	}
+
+	return false;
+}
+
+bool FPackageLocalizationUtil::GetLocalizedRoot(const FString& InPath, const FString& InCulture, FString& OutLocalized)
+{
+	if (InPath.Len() > 0 && InPath[0] == TEXT('/'))
+	{
+		const int32 RootPathEnd = InPath.Find(TEXT("/"), ESearchCase::IgnoreCase, ESearchDir::FromStart, 1);
+		if (RootPathEnd != INDEX_NONE)
+		{
+			OutLocalized = InPath.Left(RootPathEnd);
+			OutLocalized /= TEXT("L10N");
+			if (InCulture.Len() > 0)
+			{
+				OutLocalized /= InCulture;
+			}
+			return true;
+		}
+	}
+
+	return false;
+}
diff --git a/Engine/Source/Runtime/CoreUObject/Private/UObject/UObjectGlobals.cpp b/Engine/Source/Runtime/CoreUObject/Private/UObject/UObjectGlobals.cpp
--- a/Engine/Source/Runtime/CoreUObject/Private/UObject/UObjectGlobals.cpp
+++ b/Engine/Source/Runtime/CoreUObject/Private/UObject/UObjectGlobals.cpp
@@ -2761,8 +2761,8 @@ void FObjectInitializer::InitProperties(UObject* Obj, UClass* DefaultsClass, UOb
 	{
 		// @TODO: temporarily disabling the PostConstructLink optimization here
 		//        since it is unknowingly causing issues in gameplay (see 
-		//        UE-28522) - we're still unclear on what specifically this is
-		//        doing wrong, but will investigate (UE-29449)
+		//        UE-29940) - we're still unclear on what specifically this is
+		//        doing wrong, but will investigate (UE-30120)
 		bCanUsePostConstructLink = false;
 
 		QUICK_SCOPE_CYCLE_COUNTER(STAT_InitProperties_Blueprint);
@@ -2838,15 +2838,24 @@ void FObjectInitializer::InitArrayPropertyFromCustomList(const UArrayProperty* A
 	FScriptArrayHelper DstArrayValueHelper(ArrayProperty, DataPtr);
 	FScriptArrayHelper SrcArrayValueHelper(ArrayProperty, DefaultDataPtr);
 
-	int32 Num = SrcArrayValueHelper.Num();
+	const int32 SrcNum = SrcArrayValueHelper.Num();
+	const int32 DstNum = DstArrayValueHelper.Num();
 
-	if (!(ArrayProperty->Inner->PropertyFlags & CPF_IsPlainOldData))
+	if (SrcNum > DstNum)
 	{
-		DstArrayValueHelper.EmptyAndAddValues(Num);
+		const int32 Count = SrcNum - DstNum;
+		if (!(ArrayProperty->Inner->PropertyFlags & CPF_IsPlainOldData))
+		{
+			DstArrayValueHelper.AddValues(Count);
+		}
+		else
+		{
+			DstArrayValueHelper.AddUninitializedValues(Count);
+		}
 	}
-	else
+	else if (SrcNum < DstNum)
 	{
-		DstArrayValueHelper.EmptyAndAddUninitializedValues(Num);
+		DstArrayValueHelper.RemoveValues(SrcNum, DstNum - SrcNum);
 	}
 
 	for (const FCustomPropertyListNode* CustomArrayPropertyListNode = InPropertyList; CustomArrayPropertyListNode; CustomArrayPropertyListNode = CustomArrayPropertyListNode->PropertyListNext)
@@ -2856,24 +2865,7 @@ void FObjectInitializer::InitArrayPropertyFromCustomList(const UArrayProperty* A
 		uint8* DstArrayItemValue = DstArrayValueHelper.GetRawPtr(ArrayIndex);
 		const uint8* SrcArrayItemValue = SrcArrayValueHelper.GetRawPtr(ArrayIndex);
 
-		// A null property value signals that we should serialize the remaining array values in full starting at this index
-		if (CustomArrayPropertyListNode->Property == nullptr)
-		{
-			int32 Size = ArrayProperty->Inner->ElementSize;
-
-			if (!(ArrayProperty->Inner->PropertyFlags & CPF_IsPlainOldData))
-			{
-				for (int32 i = 0; i < Num - ArrayIndex; i++)
-				{
-					ArrayProperty->Inner->CopyCompleteValue(DstArrayItemValue + i * Size, SrcArrayItemValue + i * Size);
-				}
-			}
-			else
-			{
-				FMemory::Memcpy(DstArrayItemValue, SrcArrayItemValue, (Num - ArrayIndex) * Size);
-			}
-		}
-		else if (const UStructProperty* InnerStructProperty = Cast<UStructProperty>(ArrayProperty->Inner))
+		if (const UStructProperty* InnerStructProperty = Cast<UStructProperty>(ArrayProperty->Inner))
 		{
 			InitPropertiesFromCustomList(CustomArrayPropertyListNode->SubPropertyList, InnerStructProperty->Struct, DstArrayItemValue, SrcArrayItemValue);
 		}
