Copying //UE4/Dev-Sequencer to //UE4/Dev-Main (Source: //UE4/Dev-Sequencer @ 3048536)

diff --git a/Engine/Source/Editor/SequenceRecorder/Private/Sections/MovieScene3DTransformSectionRecorder.cpp b/Engine/Source/Editor/SequenceRecorder/Private/Sections/MovieScene3DTransformSectionRecorder.cpp
--- a/Engine/Source/Editor/SequenceRecorder/Private/Sections/MovieScene3DTransformSectionRecorder.cpp
+++ b/Engine/Source/Editor/SequenceRecorder/Private/Sections/MovieScene3DTransformSectionRecorder.cpp
@@ -289,16 +289,19 @@ FTransform FMovieScene3DTransformSectionRecorder::GetTransformToRecord()
 	}
 	else if(AActor* Actor = Cast<AActor>(ObjectToRecord.Get()))
 	{
-		FName SocketName, ComponentName;
-		if(SequenceRecorderUtils::GetAttachment(Actor, SocketName, ComponentName) != nullptr && Actor->GetRootComponent() != nullptr)
-		{
-			bWasAttached = true;
-			return Actor->GetRootComponent()->GetRelativeTransform();
-		}
-		else
+		bool bCaptureWorldSpaceTransform = false;
+
+		USceneComponent* RootComponent = Actor->GetRootComponent();
+		USceneComponent* AttachParent = RootComponent ? RootComponent->GetAttachParent() : nullptr;
+
+		bWasAttached = AttachParent != nullptr;
+		if (AttachParent)
 		{
-			return Actor->GetTransform();
+			// We capture world space transforms for actors if they're attached, but we're not recording the attachment parent
+			bCaptureWorldSpaceTransform = !FSequenceRecorder::Get().FindRecording(AttachParent->GetOwner());
 		}
+
+		return (bCaptureWorldSpaceTransform || !RootComponent) ? Actor->ActorToWorld() : RootComponent->GetRelativeTransform();
 	}
 
 	return FTransform::Identity;
diff --git a/Engine/Source/Editor/SequenceRecorder/Private/Sections/MovieSceneMultiPropertyRecorder.cpp b/Engine/Source/Editor/SequenceRecorder/Private/Sections/MovieSceneMultiPropertyRecorder.cpp
new file mode 100644
--- /dev/null
+++ b/Engine/Source/Editor/SequenceRecorder/Private/Sections/MovieSceneMultiPropertyRecorder.cpp
@@ -0,0 +1,129 @@
+// Copyright 1998-2015 Epic Games, Inc. All Rights Reserved.
+
+#include "SequenceRecorderPrivatePCH.h"
+#include "MovieSceneMultiPropertyRecorder.h"
+#include "MovieScenePropertyRecorder.h"
+#include "SequenceRecorderSettings.h"
+
+TSharedPtr<IMovieSceneSectionRecorder> FMovieSceneMultiPropertyRecorderFactory::CreateSectionRecorder(const FActorRecordingSettings& InActorRecordingSettings) const
+{
+	return MakeShareable(new FMovieSceneMultiPropertyRecorder());
+}
+
+bool FMovieSceneMultiPropertyRecorderFactory::CanRecordObject(UObject* InObjectToRecord) const
+{
+	const USequenceRecorderSettings* Settings = GetDefault<USequenceRecorderSettings>();
+
+	for (const FPropertiesToRecordForClass& PropertiesToRecordForClass : Settings->PropertiesToRecord)
+	{
+		if (*PropertiesToRecordForClass.Class == InObjectToRecord->GetClass() && PropertiesToRecordForClass.Properties.Num() > 0)
+		{
+			return true;
+		}
+	}
+
+	return false;
+}
+
+FMovieSceneMultiPropertyRecorder::FMovieSceneMultiPropertyRecorder()
+{
+}
+
+void FMovieSceneMultiPropertyRecorder::CreateSection(UObject* InObjectToRecord, class UMovieScene* InMovieScene, const FGuid& Guid, float Time)
+{
+	ObjectToRecord = InObjectToRecord;
+
+	const FPropertiesToRecordForClass* PropertiesToRecordForClassPtr = nullptr;
+	const USequenceRecorderSettings* Settings = GetDefault<USequenceRecorderSettings>();
+	for (const FPropertiesToRecordForClass& PropertiesToRecordForClass : Settings->PropertiesToRecord)
+	{
+		if (*PropertiesToRecordForClass.Class == InObjectToRecord->GetClass() && PropertiesToRecordForClass.Properties.Num() > 0)
+		{
+			PropertiesToRecordForClassPtr = &PropertiesToRecordForClass;
+			break;
+		}
+	}
+
+	check(PropertiesToRecordForClassPtr);
+
+	// create a recorder for each property name
+	for (const FName& PropertyName : PropertiesToRecordForClassPtr->Properties)
+	{
+		FTrackInstancePropertyBindings Binding(PropertyName, PropertyName.ToString());
+		Binding.UpdateBinding(InObjectToRecord);
+		UProperty* Property = Binding.GetProperty(InObjectToRecord);
+		if (Property != nullptr)
+		{
+			if (Property->IsA<UBoolProperty>())
+			{
+				PropertyRecorders.Add(MakeShareable(new FMovieScenePropertyRecorder<bool>(Binding)));
+			}
+			else if (Property->IsA<UByteProperty>())
+			{
+				PropertyRecorders.Add(MakeShareable(new FMovieScenePropertyRecorder<uint8>(Binding)));
+			}
+			else if (Property->IsA<UFloatProperty>())
+			{
+				PropertyRecorders.Add(MakeShareable(new FMovieScenePropertyRecorder<float>(Binding)));
+			}
+			else if (UStructProperty* StructProperty = Cast<UStructProperty>(Property))
+			{
+				if (StructProperty->Struct->GetFName() == NAME_Vector)
+				{
+					PropertyRecorders.Add(MakeShareable(new FMovieScenePropertyRecorder<FVector>(Binding)));
+				}
+				else if (StructProperty->Struct->GetFName() == NAME_Color)
+				{
+					PropertyRecorders.Add(MakeShareable(new FMovieScenePropertyRecorder<FColor>(Binding)));
+				}
+			}
+
+			PropertyRecorders.Last()->Create(InObjectToRecord, InMovieScene, Guid, Time);
+		}
+	}
+}
+
+void FMovieSceneMultiPropertyRecorder::FinalizeSection()
+{
+	for (TSharedPtr<IMovieScenePropertyRecorder> PropertyRecorder : PropertyRecorders)
+	{
+		PropertyRecorder->Finalize(ObjectToRecord.Get());
+	}
+}
+
+void FMovieSceneMultiPropertyRecorder::Record(float CurrentTime)
+{
+	for (TSharedPtr<IMovieScenePropertyRecorder> PropertyRecorder : PropertyRecorders)
+	{
+		PropertyRecorder->Record(ObjectToRecord.Get(), CurrentTime);
+	}
+}
+
+bool FMovieSceneMultiPropertyRecorder::CanPropertyBeRecorded(const UProperty& InProperty)
+{
+	if (InProperty.IsA<UBoolProperty>())
+	{
+		return true;
+	}
+	else if (InProperty.IsA<UByteProperty>())
+	{
+		return true;
+	}
+	else if (InProperty.IsA<UFloatProperty>())
+	{
+		return true;
+	}
+	else if (const UStructProperty* StructProperty = Cast<const UStructProperty>(&InProperty))
+	{
+		if (StructProperty->Struct->GetFName() == NAME_Vector)
+		{
+			return true;
+		}
+		else if (StructProperty->Struct->GetFName() == NAME_Color)
+		{
+			return true;
+		}
+	}
+
+	return false;
+}
\ No newline at end of file
diff --git a/Engine/Source/Editor/SequenceRecorder/Private/Sections/MovieScenePropertyRecorder.cpp b/Engine/Source/Editor/SequenceRecorder/Private/Sections/MovieScenePropertyRecorder.cpp
new file mode 100644
--- /dev/null
+++ b/Engine/Source/Editor/SequenceRecorder/Private/Sections/MovieScenePropertyRecorder.cpp
@@ -0,0 +1,252 @@
+// Copyright 1998-2015 Epic Games, Inc. All Rights Reserved.
+
+#include "SequenceRecorderPrivatePCH.h"
+#include "MovieScenePropertyRecorder.h"
+#include "MovieSceneBoolSection.h"
+#include "MovieSceneBoolTrack.h"
+#include "MovieSceneByteSection.h"
+#include "MovieSceneByteTrack.h"
+#include "MovieSceneFloatSection.h"
+#include "MovieSceneFloatTrack.h"
+#include "MovieSceneColorSection.h"
+#include "MovieSceneColorTrack.h"
+#include "MovieSceneVectorSection.h"
+#include "MovieSceneVectorTrack.h"
+
+// current set of compiled-in property types
+
+template <>
+bool FMovieScenePropertyRecorder<bool>::ShouldAddNewKey(const bool& InNewValue) const
+{
+	return InNewValue != PreviousValue;
+}
+
+template <>
+UMovieSceneSection* FMovieScenePropertyRecorder<bool>::AddSection(UMovieScene* InMovieScene, const FGuid& InGuid, float InTime)
+{
+	UMovieSceneBoolTrack* Track = InMovieScene->AddTrack<UMovieSceneBoolTrack>(InGuid);
+	if (Track)
+	{
+		Track->SetPropertyNameAndPath(Binding.GetPropertyName(), Binding.GetPropertyPath());
+
+		UMovieSceneBoolSection* Section = Cast<UMovieSceneBoolSection>(Track->CreateNewSection());
+		Section->SetDefault(PreviousValue);
+		Section->SetStartTime(InTime);
+		Section->SetEndTime(InTime);
+		Section->AddKey(InTime, PreviousValue, EMovieSceneKeyInterpolation::Break);
+
+		Track->AddSection(*Section);
+
+		return Section;
+	}
+
+	return nullptr;
+}
+
+template <>
+void FMovieScenePropertyRecorder<bool>::AddKeyToSection(UMovieSceneSection* InSection, const FPropertyKey<bool>& InKey)
+{
+	CastChecked<UMovieSceneBoolSection>(InSection)->AddKey(InKey.Time, InKey.Value, EMovieSceneKeyInterpolation::Break);
+}
+
+template <>
+void FMovieScenePropertyRecorder<bool>::ReduceKeys(UMovieSceneSection* InSection)
+{
+}
+
+template <>
+bool FMovieScenePropertyRecorder<uint8>::ShouldAddNewKey(const uint8& InNewValue) const
+{
+	return InNewValue != PreviousValue;
+}
+
+template <>
+UMovieSceneSection* FMovieScenePropertyRecorder<uint8>::AddSection(UMovieScene* InMovieScene, const FGuid& InGuid, float InTime)
+{
+	UMovieSceneByteTrack* Track = InMovieScene->AddTrack<UMovieSceneByteTrack>(InGuid);
+	if (Track)
+	{
+		Track->SetPropertyNameAndPath(Binding.GetPropertyName(), Binding.GetPropertyPath());
+
+		UMovieSceneByteSection* Section = Cast<UMovieSceneByteSection>(Track->CreateNewSection());
+		Section->SetDefault(PreviousValue);
+		Section->SetStartTime(InTime);
+		Section->SetEndTime(InTime);
+		Section->AddKey(InTime, PreviousValue, EMovieSceneKeyInterpolation::Break);
+
+		Track->AddSection(*Section);
+
+		return Section;
+	}
+
+	return nullptr;
+}
+
+template <>
+void FMovieScenePropertyRecorder<uint8>::AddKeyToSection(UMovieSceneSection* InSection, const FPropertyKey<uint8>& InKey)
+{
+	CastChecked<UMovieSceneByteSection>(InSection)->AddKey(InKey.Time, InKey.Value, EMovieSceneKeyInterpolation::Break);
+}
+
+template <>
+void FMovieScenePropertyRecorder<uint8>::ReduceKeys(UMovieSceneSection* InSection)
+{
+}
+
+template <>
+bool FMovieScenePropertyRecorder<float>::ShouldAddNewKey(const float& InNewValue) const
+{
+	return true;
+}
+
+template <>
+UMovieSceneSection* FMovieScenePropertyRecorder<float>::AddSection(UMovieScene* InMovieScene, const FGuid& InGuid, float InTime)
+{
+	UMovieSceneFloatTrack* Track = InMovieScene->AddTrack<UMovieSceneFloatTrack>(InGuid);
+	if (Track)
+	{
+		Track->SetPropertyNameAndPath(Binding.GetPropertyName(), Binding.GetPropertyPath());
+
+		UMovieSceneFloatSection* Section = Cast<UMovieSceneFloatSection>(Track->CreateNewSection());
+		Section->SetDefault(PreviousValue);
+		Section->SetStartTime(InTime);
+		Section->SetEndTime(InTime);
+		Section->AddKey(InTime, PreviousValue, EMovieSceneKeyInterpolation::Break);
+
+		Track->AddSection(*Section);
+
+		return Section;
+	}
+
+	return nullptr;
+}
+
+template <>
+void FMovieScenePropertyRecorder<float>::AddKeyToSection(UMovieSceneSection* InSection, const FPropertyKey<float>& InKey)
+{
+	CastChecked<UMovieSceneFloatSection>(InSection)->AddKey(InKey.Time, InKey.Value, EMovieSceneKeyInterpolation::Auto);
+}
+
+template <>
+void FMovieScenePropertyRecorder<float>::ReduceKeys(UMovieSceneSection* InSection)
+{
+	CastChecked<UMovieSceneFloatSection>(InSection)->GetFloatCurve().RemoveRedundantKeys(KINDA_SMALL_NUMBER);
+}
+
+template <>
+bool FMovieScenePropertyRecorder<FColor>::ShouldAddNewKey(const FColor& InNewValue) const
+{
+	return true;
+}
+
+template <>
+UMovieSceneSection* FMovieScenePropertyRecorder<FColor>::AddSection(UMovieScene* InMovieScene, const FGuid& InGuid, float InTime)
+{
+	UMovieSceneColorTrack* Track = InMovieScene->AddTrack<UMovieSceneColorTrack>(InGuid);
+	if (Track)
+	{
+		Track->SetPropertyNameAndPath(Binding.GetPropertyName(), Binding.GetPropertyPath());
+
+		UMovieSceneColorSection* Section = Cast<UMovieSceneColorSection>(Track->CreateNewSection());
+		Section->SetDefault(FColorKey(EKeyColorChannel::Red, PreviousValue.R, false));
+		Section->SetDefault(FColorKey(EKeyColorChannel::Green, PreviousValue.G, false));
+		Section->SetDefault(FColorKey(EKeyColorChannel::Blue, PreviousValue.B, false));
+		Section->SetDefault(FColorKey(EKeyColorChannel::Alpha, PreviousValue.A, false));
+
+		Section->SetStartTime(InTime);
+		Section->SetEndTime(InTime);
+
+		Section->AddKey(InTime, FColorKey(EKeyColorChannel::Red, PreviousValue.R, false), EMovieSceneKeyInterpolation::Break);
+		Section->AddKey(InTime, FColorKey(EKeyColorChannel::Green, PreviousValue.G, false), EMovieSceneKeyInterpolation::Break);
+		Section->AddKey(InTime, FColorKey(EKeyColorChannel::Blue, PreviousValue.B, false), EMovieSceneKeyInterpolation::Break);
+		Section->AddKey(InTime, FColorKey(EKeyColorChannel::Alpha, PreviousValue.A, false), EMovieSceneKeyInterpolation::Break);
+
+		Track->AddSection(*Section);
+
+		return Section;
+	}
+
+	return nullptr;
+}
+
+template <>
+void FMovieScenePropertyRecorder<FColor>::AddKeyToSection(UMovieSceneSection* InSection, const FPropertyKey<FColor>& InKey)
+{
+	UMovieSceneColorSection* ColorSection = CastChecked<UMovieSceneColorSection>(InSection);
+	ColorSection->AddKey(InKey.Time, FColorKey(EKeyColorChannel::Red, InKey.Value.R, false), EMovieSceneKeyInterpolation::Auto);
+	ColorSection->AddKey(InKey.Time, FColorKey(EKeyColorChannel::Green, InKey.Value.G, false), EMovieSceneKeyInterpolation::Auto);
+	ColorSection->AddKey(InKey.Time, FColorKey(EKeyColorChannel::Blue, InKey.Value.B, false), EMovieSceneKeyInterpolation::Auto);
+	ColorSection->AddKey(InKey.Time, FColorKey(EKeyColorChannel::Alpha, InKey.Value.A, false), EMovieSceneKeyInterpolation::Auto);
+}
+
+template <>
+void FMovieScenePropertyRecorder<FColor>::ReduceKeys(UMovieSceneSection* InSection)
+{
+	UMovieSceneColorSection* ColorSection = CastChecked<UMovieSceneColorSection>(InSection);
+	ColorSection->GetRedCurve().RemoveRedundantKeys(KINDA_SMALL_NUMBER);
+	ColorSection->GetGreenCurve().RemoveRedundantKeys(KINDA_SMALL_NUMBER);
+	ColorSection->GetBlueCurve().RemoveRedundantKeys(KINDA_SMALL_NUMBER);
+	ColorSection->GetAlphaCurve().RemoveRedundantKeys(KINDA_SMALL_NUMBER);
+}
+
+template <>
+bool FMovieScenePropertyRecorder<FVector>::ShouldAddNewKey(const FVector& InNewValue) const
+{
+	return true;
+}
+
+template <>
+UMovieSceneSection* FMovieScenePropertyRecorder<FVector>::AddSection(UMovieScene* InMovieScene, const FGuid& InGuid, float InTime)
+{
+	UMovieSceneVectorTrack* Track = InMovieScene->AddTrack<UMovieSceneVectorTrack>(InGuid);
+	Track->SetNumChannelsUsed(3);
+	if (Track)
+	{
+		Track->SetPropertyNameAndPath(Binding.GetPropertyName(), Binding.GetPropertyPath());
+
+		UMovieSceneVectorSection* Section = Cast<UMovieSceneVectorSection>(Track->CreateNewSection());
+
+		Section->SetDefault(FVectorKey(EKeyVectorChannel::X, PreviousValue.X));
+		Section->SetDefault(FVectorKey(EKeyVectorChannel::Y, PreviousValue.Y));
+		Section->SetDefault(FVectorKey(EKeyVectorChannel::Z, PreviousValue.Z));
+
+		Section->SetStartTime(InTime);
+		Section->SetEndTime(InTime);
+		
+		Section->AddKey(InTime, FVectorKey(EKeyVectorChannel::X, PreviousValue.X), EMovieSceneKeyInterpolation::Break);
+		Section->AddKey(InTime, FVectorKey(EKeyVectorChannel::Y, PreviousValue.Y), EMovieSceneKeyInterpolation::Break);
+		Section->AddKey(InTime, FVectorKey(EKeyVectorChannel::Z, PreviousValue.Z), EMovieSceneKeyInterpolation::Break);
+
+		Track->AddSection(*Section);
+
+		return Section;
+	}
+
+	return nullptr;
+}
+
+template <>
+void FMovieScenePropertyRecorder<FVector>::AddKeyToSection(UMovieSceneSection* InSection, const FPropertyKey<FVector>& InKey)
+{
+	UMovieSceneVectorSection* VectorSection = CastChecked<UMovieSceneVectorSection>(InSection);
+
+	VectorSection->AddKey(InKey.Time, FVectorKey(EKeyVectorChannel::X, InKey.Value.X), EMovieSceneKeyInterpolation::Auto);
+	VectorSection->AddKey(InKey.Time, FVectorKey(EKeyVectorChannel::Y, InKey.Value.Y), EMovieSceneKeyInterpolation::Auto);
+	VectorSection->AddKey(InKey.Time, FVectorKey(EKeyVectorChannel::Z, InKey.Value.Z), EMovieSceneKeyInterpolation::Auto);
+}
+
+template <>
+void FMovieScenePropertyRecorder<FVector>::ReduceKeys(UMovieSceneSection* InSection)
+{
+	UMovieSceneVectorSection* VectorSection = CastChecked<UMovieSceneVectorSection>(InSection);
+
+	VectorSection->GetCurve(0).RemoveRedundantKeys(KINDA_SMALL_NUMBER);
+	VectorSection->GetCurve(1).RemoveRedundantKeys(KINDA_SMALL_NUMBER);
+	VectorSection->GetCurve(2).RemoveRedundantKeys(KINDA_SMALL_NUMBER);
+}
+
+template class FMovieScenePropertyRecorder<bool>;
+template class FMovieScenePropertyRecorder<uint8>;
+template class FMovieScenePropertyRecorder<float>;
+template class FMovieScenePropertyRecorder<FColor>;
+template class FMovieScenePropertyRecorder<FVector>;
diff --git a/Engine/Source/Editor/SequenceRecorder/Private/SequenceRecorder.cpp b/Engine/Source/Editor/SequenceRecorder/Private/SequenceRecorder.cpp
--- a/Engine/Source/Editor/SequenceRecorder/Private/SequenceRecorder.cpp
+++ b/Engine/Source/Editor/SequenceRecorder/Private/SequenceRecorder.cpp
@@ -37,6 +37,7 @@ void FSequenceRecorder::Initialize()
 	// register built-in recorders
 	IModularFeatures::Get().RegisterModularFeature("MovieSceneSectionRecorderFactory", &AnimationSectionRecorderFactory);
 	IModularFeatures::Get().RegisterModularFeature("MovieSceneSectionRecorderFactory", &TransformSectionRecorderFactory);
+	IModularFeatures::Get().RegisterModularFeature("MovieSceneSectionRecorderFactory", &MultiPropertySectionRecorder);
 }
 
 void FSequenceRecorder::Shutdown()
@@ -44,6 +45,7 @@ void FSequenceRecorder::Shutdown()
 	// unregister built-in recorders
 	IModularFeatures::Get().UnregisterModularFeature("MovieSceneSectionRecorderFactory", &AnimationSectionRecorderFactory);
 	IModularFeatures::Get().UnregisterModularFeature("MovieSceneSectionRecorderFactory", &TransformSectionRecorderFactory);
+	IModularFeatures::Get().UnregisterModularFeature("MovieSceneSectionRecorderFactory", &MultiPropertySectionRecorder);
 
 	if(CountdownTexture.IsValid())
 	{
@@ -68,7 +70,7 @@ bool FSequenceRecorder::IsRecordingQueued(AActor* Actor) const
 {
 	for (UActorRecording* QueuedRecording : QueuedRecordings)
 	{
-		if (QueuedRecording->ActorToRecord == Actor)
+		if (QueuedRecording->GetActorToRecord() == Actor)
 		{
 			return true;
 		}
@@ -81,7 +83,7 @@ UActorRecording* FSequenceRecorder::FindRecording(AActor* Actor) const
 {
 	for (UActorRecording* QueuedRecording : QueuedRecordings)
 	{
-		if (QueuedRecording->ActorToRecord == Actor)
+		if (QueuedRecording->GetActorToRecord() == Actor)
 		{
 			return QueuedRecording;
 		}
@@ -112,7 +114,7 @@ UActorRecording* FSequenceRecorder::AddNewQueuedRecording(AActor* Actor, UAnimSe
 
 	UActorRecording* ActorRecording = NewObject<UActorRecording>();
 	ActorRecording->AddToRoot();
-	ActorRecording->ActorToRecord = Actor;
+	ActorRecording->SetActorToRecord(Actor);
 	ActorRecording->TargetAnimation = AnimSequence;
 	ActorRecording->AnimationSettings.Length = Length;
 
@@ -140,7 +142,7 @@ void FSequenceRecorder::RemoveQueuedRecording(AActor* Actor)
 {
 	for (UActorRecording* QueuedRecording : QueuedRecordings)
 	{
-		if (QueuedRecording->ActorToRecord == Actor)
+		if (QueuedRecording->GetActorToRecord() == Actor)
 		{
 			QueuedRecording->RemoveFromRoot();
 			QueuedRecordings.Remove(QueuedRecording);
@@ -275,7 +277,7 @@ void FSequenceRecorder::Tick(float DeltaSeconds)
 		auto RemoveDeadActorPredicate = 
 			[&](UActorRecording* Recording)
 			{
-				if(!Recording->ActorToRecord.IsValid())
+				if(!Recording->GetActorToRecord())
 				{
 					DeadRecordings.Add(Recording);
 					return true;
@@ -425,9 +427,9 @@ bool FSequenceRecorder::StartRecordingInternal(UWorld* World)
 	UWorld* ActorWorld = nullptr;
 	if(Settings->bRecordWorldSettingsActor)
 	{
-		if(World != nullptr || (QueuedRecordings.Num() > 0 && QueuedRecordings[0]->ActorToRecord.IsValid()))
+		if(World != nullptr || (QueuedRecordings.Num() > 0 && QueuedRecordings[0]->GetActorToRecord() != nullptr))
 		{
-			ActorWorld = World != nullptr ? World : QueuedRecordings[0]->ActorToRecord->GetWorld();
+			ActorWorld = World != nullptr ? World : QueuedRecordings[0]->GetActorToRecord()->GetWorld();
 			if(ActorWorld)
 			{
 				AWorldSettings* WorldSettings = ActorWorld->GetWorldSettings();
@@ -474,9 +476,9 @@ bool FSequenceRecorder::StartRecordingInternal(UWorld* World)
 		// register for spawning delegate in the world(s) of recorded actors
 		for(UActorRecording* Recording : QueuedRecordings)
 		{
-			if(Recording->ActorToRecord.IsValid())
+			if(Recording->GetActorToRecord() != nullptr)
 			{
-				UWorld* ActorToRecordWorld = Recording->ActorToRecord->GetWorld();
+				UWorld* ActorToRecordWorld = Recording->GetActorToRecord()->GetWorld();
 				if(ActorToRecordWorld != nullptr)
 				{
 					FDelegateHandle* FoundHandle = ActorSpawningDelegateHandles.Find(ActorToRecordWorld);
@@ -699,7 +701,7 @@ bool FSequenceRecorder::IsActorValidForRecording(AActor* Actor)
 
 		for(UActorRecording* Recording : QueuedRecordings)
 		{
-			if(AActor* OtherActor = Recording->ActorToRecord.Get())
+			if(AActor* OtherActor = Recording->GetActorToRecord())
 			{
 				if(OtherActor != Actor)
 				{
@@ -745,7 +747,7 @@ void FSequenceRecorder::HandleActorDespawned(AActor* Actor)
 		for(int32 Index = 0; Index < QueuedRecordings.Num(); ++Index)
 		{
 			UActorRecording* Recording = QueuedRecordings[Index];
-			if(Recording->ActorToRecord.Get() == Actor)
+			if(Recording->GetActorToRecord() == Actor)
 			{
 				Recording->InvalidateObjectToRecord();
 				DeadRecordings.Add(Recording);
diff --git a/Engine/Source/Editor/SequenceRecorder/Private/SequenceRecorder.h b/Engine/Source/Editor/SequenceRecorder/Private/SequenceRecorder.h
--- a/Engine/Source/Editor/SequenceRecorder/Private/SequenceRecorder.h
+++ b/Engine/Source/Editor/SequenceRecorder/Private/SequenceRecorder.h
@@ -5,6 +5,7 @@
 #include "ISequenceRecorder.h"
 #include "MovieSceneAnimationSectionRecorder.h"
 #include "MovieScene3DTransformSectionRecorder.h"
+#include "MovieSceneMultiPropertyRecorder.h"
 
 struct FSequenceRecorder
 {
@@ -152,4 +153,7 @@ private:
 
 	/** Built-in transform recorder factory */
 	FMovieScene3DTransformSectionRecorderFactory TransformSectionRecorderFactory;
+
+	/** Built in multi-property recorder */
+	FMovieSceneMultiPropertyRecorderFactory MultiPropertySectionRecorder;
 };
diff --git a/Engine/Source/Editor/SequenceRecorder/Private/SequenceRecorderModule.cpp b/Engine/Source/Editor/SequenceRecorder/Private/SequenceRecorderModule.cpp
--- a/Engine/Source/Editor/SequenceRecorder/Private/SequenceRecorderModule.cpp
+++ b/Engine/Source/Editor/SequenceRecorder/Private/SequenceRecorderModule.cpp
@@ -11,6 +11,7 @@
 #include "NotificationManager.h"
 #include "SNotificationList.h"
 #include "ActorRecordingDetailsCustomization.h"
+#include "PropertiesToRecordForClassDetailsCustomization.h"
 
 #define LOCTEXT_NAMESPACE "SequenceRecorder"
 
@@ -83,12 +84,6 @@ class FSequenceRecorderModule : public ISequenceRecorder, private FSelfRegisteri
 			FLevelEditorModule& LevelEditorModule = FModuleManager::LoadModuleChecked<FLevelEditorModule>(TEXT("LevelEditor"));
 			LevelEditorModule.OnCaptureSingleFrameAnimSequence().BindStatic(&FSequenceRecorderModule::HandleCaptureSingleFrameAnimSequence);
 
-			// register level editor menu extender
-			LevelEditorMenuExtenderDelegate = FLevelEditorModule::FLevelViewportMenuExtender_SelectedActors::CreateStatic(&FSequenceRecorderModule::ExtendLevelViewportContextMenu);
-			auto& MenuExtenders = LevelEditorModule.GetAllLevelViewportContextMenuExtenders();
-			MenuExtenders.Add(LevelEditorMenuExtenderDelegate);
-			LevelEditorExtenderDelegateHandle = MenuExtenders.Last().GetHandle();
-
 			// register standalone UI
 			FGlobalTabmanager::Get()->RegisterNomadTabSpawner(SequenceRecorderTabName, FOnSpawnTab::CreateStatic(&FSequenceRecorderModule::SpawnSequenceRecorderTab))
 				.SetGroup(WorkspaceMenu::GetMenuStructure().GetLevelEditorCategory())
@@ -97,11 +92,12 @@ class FSequenceRecorderModule : public ISequenceRecorder, private FSelfRegisteri
 				.SetIcon(FSlateIcon(FEditorStyle::GetStyleSetName(), "SequenceRecorder.TabIcon"));
 
 			// register for debug drawing
-			DrawDebugDelegateHandle = UDebugDrawService::Register(TEXT("Game"),  FDebugDrawDelegate::CreateStatic(&FSequenceRecorderModule::DrawDebug));
+			DrawDebugDelegateHandle = UDebugDrawService::Register(TEXT("Decals"), FDebugDrawDelegate::CreateStatic(&FSequenceRecorderModule::DrawDebug));
 
 			// register details customization
 			FPropertyEditorModule& PropertyModule = FModuleManager::LoadModuleChecked<FPropertyEditorModule>("PropertyEditor");
 			PropertyModule.RegisterCustomClassLayout(UActorRecording::StaticClass()->GetFName(), FOnGetDetailCustomizationInstance::CreateStatic(&FActorRecordingDetailsCustomization::MakeInstance));
+			PropertyModule.RegisterCustomPropertyTypeLayout(FPropertiesToRecordForClass::StaticStruct()->GetFName(), FOnGetPropertyTypeCustomizationInstance::CreateStatic(&FPropertiesToRecordForClassDetailsCustomization::MakeInstance));
 		}
 #endif
 	}
@@ -125,9 +121,6 @@ class FSequenceRecorderModule : public ISequenceRecorder, private FSelfRegisteri
 			{
 				FLevelEditorModule& LevelEditorModule = FModuleManager::GetModuleChecked<FLevelEditorModule>(TEXT("LevelEditor"));
 				LevelEditorModule.OnCaptureSingleFrameAnimSequence().Unbind();
-				LevelEditorModule.GetAllLevelViewportContextMenuExtenders().RemoveAll([&](const FLevelEditorModule::FLevelViewportMenuExtender_SelectedActors& Delegate) {
-					return Delegate.GetHandle() == LevelEditorExtenderDelegateHandle;
-				});
 			}
 
 			if (FModuleManager::Get().IsModuleLoaded(TEXT("Persona")))
@@ -147,6 +140,8 @@ class FSequenceRecorderModule : public ISequenceRecorder, private FSelfRegisteri
 				if (UObjectInitialized())
 				{
 					PropertyModule.UnregisterCustomClassLayout(UActorRecording::StaticClass()->GetFName());
+					PropertyModule.UnregisterCustomClassLayout(USequenceRecorderSettings::StaticClass()->GetFName());
+					PropertyModule.UnregisterCustomPropertyTypeLayout(FPropertiesToRecordForClass::StaticStruct()->GetFName());
 				}
 			}
 		}
@@ -237,7 +232,7 @@ class FSequenceRecorderModule : public ISequenceRecorder, private FSelfRegisteri
 				{
 					if (Actor)
 					{
-						if (Actor->GetName() == ActorNameStr)
+						if (Actor->GetActorLabel() == ActorNameStr)
 						{
 							return Actor;
 						}
@@ -260,7 +255,7 @@ class FSequenceRecorderModule : public ISequenceRecorder, private FSelfRegisteri
 					{
 						if (Actor)
 						{
-							if (Actor->GetName().ToLower().StartsWith(FuzzyActorNameStr))
+							if (Actor->GetActorLabel().ToLower().StartsWith(FuzzyActorNameStr))
 							{
 								return Actor;
 							}
@@ -486,28 +481,27 @@ class FSequenceRecorderModule : public ISequenceRecorder, private FSelfRegisteri
 		return CurrentSequence.IsValid() ? CurrentSequence->GetMovieScene()->GetPlaybackRange().Size<float>() : 0.0f;
 	}
 
-	virtual bool StartRecording(UWorld* World, const FString& ActorNameToRecord, const FOnRecordingStarted& OnRecordingStarted, const FOnRecordingFinished& OnRecordingFinished, const FString& PathToRecordTo, const FString& SequenceName) override
+	virtual bool StartRecording(TFixedSizeArrayView<AActor*> ActorsToRecord, const FOnRecordingStarted& OnRecordingStarted, const FOnRecordingFinished& OnRecordingFinished, const FString& PathToRecordTo, const FString& SequenceName) override
 	{
-		if(ActorNameToRecord.Len() > 0)
+		if(ActorsToRecord.Num() != 0)
 		{
-			AActor* Actor = FindActor(ActorNameToRecord, World, true);
-			if(Actor != nullptr)
+			FSequenceRecorder::Get().ClearQueuedRecordings();
+			for (AActor* Actor : ActorsToRecord)
 			{
-				FSequenceRecorder::Get().ClearQueuedRecordings();
 				FSequenceRecorder::Get().AddNewQueuedRecording(Actor);
 			}
-			else
+		}
+		else
+		{
+			if(FSlateApplication::IsInitialized())
 			{
-				if(FSlateApplication::IsInitialized())
-				{
-					FNotificationInfo Info(FText::Format(LOCTEXT("SequenceRecordingErrorActor", "Couldnt find an actor to record matching name \"{0}\""), FText::FromString(ActorNameToRecord)));
-					Info.bUseLargeFont = false;
+				FNotificationInfo Info(LOCTEXT("SequenceRecordingErrorActor", "Couldn't find actor to record"));
+				Info.bUseLargeFont = false;
 
-					FSlateNotificationManager::Get().AddNotification(Info);
-				}
-
-				UE_LOG(LogAnimation, Display, TEXT("Couldnt find an actor to record matching name \"%s\""), *ActorNameToRecord);
+				FSlateNotificationManager::Get().AddNotification(Info);
 			}
+
+			UE_LOG(LogAnimation, Display, TEXT("Couldn't find actor to record"));
 		}
 
 		return FSequenceRecorder::Get().StartRecording(OnRecordingStarted, OnRecordingFinished, PathToRecordTo, SequenceName);
@@ -601,110 +595,6 @@ class FSequenceRecorderModule : public ISequenceRecorder, private FSelfRegisteri
 		return MajorTab;
 	}
 
-	static TSharedRef<FExtender> ExtendLevelViewportContextMenu(const TSharedRef<FUICommandList> CommandList, const TArray<AActor*> SelectedActors)
-	{
-		TSharedRef<FExtender> Extender(new FExtender());
-
-		// check if actors are already registered for recording
-		bool bCanAddRecording = false;
-		bool bCanRemoveRecording = false;
-
-		FSequenceRecorder& SequenceRecorder = FSequenceRecorder::Get();
-
-		for (AActor* Actor : SelectedActors)
-		{
-			const bool bIsQueued = SequenceRecorder.IsRecordingQueued(Actor);
-			bCanRemoveRecording |= bIsQueued;
-			bCanAddRecording |= !bIsQueued;
-		}
-
-		if (bCanAddRecording || bCanRemoveRecording)
-		{
-			// Add the sequence recorder sub-menu extender
-			Extender->AddMenuExtension(
-				"ActorSelectVisibilityLevels",
-				EExtensionHook::After,
-				nullptr,
-				FMenuExtensionDelegate::CreateStatic(&FSequenceRecorderModule::CreateLevelViewportContextMenuEntries, bCanAddRecording, bCanRemoveRecording));
-		}
-
-		return Extender;
-	}
-
-	static void CreateLevelViewportContextMenuEntries(FMenuBuilder& MenuBuilder, bool bCanAddRecording, bool bCanRemoveRecording)
-	{
-		MenuBuilder.BeginSection("SequenceRecording", LOCTEXT("SequenceRecordingLevelEditorHeading", "Sequence Recording"));
-
-		if (bCanAddRecording)
-		{
-			FUIAction Action_AddQuickRecording(FExecuteAction::CreateStatic(&FSequenceRecorderModule::AddActorsForQuickRecording));
-
-			MenuBuilder.AddMenuEntry(
-				LOCTEXT("MenuExtensionAddQuickRecording", "Trigger Actor Recording"),
-				LOCTEXT("MenuExtensionAddQuickRecording_Tooltip", "Set up the selected actors for recording and trigger recording immediately"),
-				FSlateIcon(),
-				Action_AddQuickRecording,
-				NAME_None,
-				EUserInterfaceActionType::Button);
-
-			FUIAction Action_AddRecording(FExecuteAction::CreateStatic(&FSequenceRecorderModule::AddActorsForRecording));
-
-			MenuBuilder.AddMenuEntry(
-				LOCTEXT("MenuExtensionAddRecording", "Queue Actor Recording"),
-				LOCTEXT("MenuExtensionAddRecording_Tooltip", "Queue up the selected actors for recording"),
-				FSlateIcon(),
-				Action_AddRecording,
-				NAME_None,
-				EUserInterfaceActionType::Button);
-		}
-
-		if (bCanRemoveRecording)
-		{
-			FUIAction Action_RemoveRecording(FExecuteAction::CreateStatic(&FSequenceRecorderModule::RemoveActorsForRecording));
-
-			MenuBuilder.AddMenuEntry(
-				LOCTEXT("MenuExtensionRemoveRecording", "Removed Queued Recording"),
-				LOCTEXT("MenuExtensionRemoveRecording_Tooltip", "Remove the selected actors from the recording queue"),
-				FSlateIcon(),
-				Action_RemoveRecording,
-				NAME_None,
-				EUserInterfaceActionType::Button);
-		}
-
-		MenuBuilder.EndSection();
-	}
-
-	static void AddActorsForRecording()
-	{
-		FSequenceRecorder& SequenceRecorder = FSequenceRecorder::Get();
-
-		TArray<AActor*> SelectedActors;
-		GEditor->GetSelectedActors()->GetSelectedObjects<AActor>(SelectedActors);
-		for(AActor* Actor : SelectedActors)
-		{
-			SequenceRecorder.AddNewQueuedRecording(Actor, nullptr, GetDefault<USequenceRecorderSettings>()->SequenceLength);
-		}
-	}
-
-	static void AddActorsForQuickRecording()
-	{
-		AddActorsForRecording();
-
-		FSequenceRecorder::Get().StartRecording();
-	}
-
-	static void RemoveActorsForRecording()
-	{
-		FSequenceRecorder& SequenceRecorder = FSequenceRecorder::Get();
-
-		TArray<AActor*> SelectedActors;
-		GEditor->GetSelectedActors()->GetSelectedObjects<AActor>(SelectedActors);
-		for (AActor* Actor : SelectedActors)
-		{
-			SequenceRecorder.RemoveQueuedRecording(Actor);
-		}
-	}
-
 	static void DrawDebug(UCanvas* InCanvas, APlayerController* InPlayerController)
 	{
 		FSequenceRecorder::Get().DrawDebug(InCanvas, InPlayerController);
@@ -712,10 +602,6 @@ class FSequenceRecorderModule : public ISequenceRecorder, private FSelfRegisteri
 #endif
 	FDelegateHandle PostEditorTickHandle;
 
-	FLevelEditorModule::FLevelViewportMenuExtender_SelectedActors LevelEditorMenuExtenderDelegate;
-
-	FDelegateHandle LevelEditorExtenderDelegateHandle;
-
 	FDelegateHandle DrawDebugDelegateHandle;
 };
 
diff --git a/Engine/Source/Editor/SequenceRecorder/Private/SequenceRecorderSettings.cpp b/Engine/Source/Editor/SequenceRecorder/Private/SequenceRecorderSettings.cpp
--- a/Engine/Source/Editor/SequenceRecorder/Private/SequenceRecorderSettings.cpp
+++ b/Engine/Source/Editor/SequenceRecorder/Private/SequenceRecorderSettings.cpp
@@ -20,6 +20,7 @@ USequenceRecorderSettings::USequenceRecorderSettings(const FObjectInitializer& O
 	ComponentClassesToRecord.Add(UStaticMeshComponent::StaticClass());
 	ComponentClassesToRecord.Add(UParticleSystemComponent::StaticClass());
 	ComponentClassesToRecord.Add(ULightComponent::StaticClass());
+	ComponentClassesToRecord.Add(UCameraComponent::StaticClass());
 }
 
 void USequenceRecorderSettings::PostEditChangeChainProperty(struct FPropertyChangedChainEvent& PropertyChangedEvent)
diff --git a/Engine/Source/Editor/SequenceRecorder/Public/ISequenceRecorder.h b/Engine/Source/Editor/SequenceRecorder/Public/ISequenceRecorder.h
--- a/Engine/Source/Editor/SequenceRecorder/Public/ISequenceRecorder.h
+++ b/Engine/Source/Editor/SequenceRecorder/Public/ISequenceRecorder.h
@@ -3,6 +3,7 @@
 #pragma once
 
 #include "ModuleInterface.h"
+#include "FixedSizeArrayView.h"
 
 DECLARE_DELEGATE_OneParam(FOnRecordingStarted, class UMovieSceneSequence* /*Sequence*/);
 
@@ -30,15 +31,28 @@ public:
 
 	/**
 	 * Start a recording, possibly with some delay (specified by the sequence recording settings).
-	 * @param	World				The world we use to record actors
-	 * @param	ActorNameToRecord	The name of the actor to record (fuzzy-matched, precise name not required).
+	 * @param	ActorsToRecord		Actors to record.
 	 * @param	OnRecordingStarted	Delegate fired when recording has commenced.
 	 * @param	OnRecordingFinished	Delegate fired when recording has finished.
 	 * @param	PathToRecordTo		Optional path to a sequence to record to. If none is specified we use the defaults in the settings.
 	 * @param	SequenceName		Optional name of a sequence to record to. If none is specified we use the defaults in the settings.
 	 * @return true if recording was successfully started
 	*/
-	virtual bool StartRecording(UWorld* World, const FString& ActorNameToRecord, const FOnRecordingStarted& OnRecordingStarted, const FOnRecordingFinished& OnRecordingFinished, const FString& PathToRecordTo = FString(), const FString& SequenceName = FString()) = 0;
+	virtual bool StartRecording(TFixedSizeArrayView<AActor*> ActorsToRecord, const FOnRecordingStarted& OnRecordingStarted, const FOnRecordingFinished& OnRecordingFinished, const FString& PathToRecordTo = FString(), const FString& SequenceName = FString()) = 0;
+
+	/**
+	 * Start a recording, possibly with some delay (specified by the sequence recording settings).
+	 * @param	ActorToRecord		Actor to record.
+	 * @param	OnRecordingStarted	Delegate fired when recording has commenced.
+	 * @param	OnRecordingFinished	Delegate fired when recording has finished.
+	 * @param	PathToRecordTo		Optional path to a sequence to record to. If none is specified we use the defaults in the settings.
+	 * @param	SequenceName		Optional name of a sequence to record to. If none is specified we use the defaults in the settings.
+	 * @return true if recording was successfully started
+	*/
+	bool StartRecording(AActor* ActorToRecord, const FOnRecordingStarted& OnRecordingStarted, const FOnRecordingFinished& OnRecordingFinished, const FString& PathToRecordTo = FString(), const FString& SequenceName = FString())
+	{
+		return StartRecording(TFixedSizeArrayView<AActor*>(&ActorToRecord, 1), OnRecordingStarted, OnRecordingFinished, PathToRecordTo, SequenceName);
+	}
 
 	/**
 	 * Notify we should start recording an actor - usually used for 'actor pooling' implementations
diff --git a/Engine/Source/Editor/SequenceRecorder/Public/SequenceRecorderSettings.h b/Engine/Source/Editor/SequenceRecorder/Public/SequenceRecorderSettings.h
--- a/Engine/Source/Editor/SequenceRecorder/Public/SequenceRecorderSettings.h
+++ b/Engine/Source/Editor/SequenceRecorder/Public/SequenceRecorderSettings.h
@@ -6,6 +6,20 @@
 #include "SequenceRecorderActorFilter.h"
 #include "SequenceRecorderSettings.generated.h"
 
+USTRUCT()
+struct FPropertiesToRecordForClass
+{
+	GENERATED_BODY()
+
+	/** The class of the object we can record */
+	UPROPERTY(Config, EditAnywhere, AdvancedDisplay, Category = "Sequence Recording")
+	TSubclassOf<UObject> Class;
+
+	/** List of properties we want to record for this class */
+	UPROPERTY(Config, EditAnywhere, AdvancedDisplay, Category = "Sequence Recording")
+	TArray<FName> Properties;
+};
+
 UCLASS(config=Editor)
 class SEQUENCERECORDER_API USequenceRecorderSettings : public UObject
 {
@@ -69,4 +83,8 @@ public:
 	/** Whether to record actors that are spawned by sequencer itself (this is usually disabled as results can be unexpected) */
 	UPROPERTY(Config, EditAnywhere, AdvancedDisplay, Category = "Sequence Recording")
 	bool bRecordSequencerSpawnedActors;
+
+	/** The properties to record for specified classes */
+	UPROPERTY(Config, EditAnywhere, Category = "General Property Recording")
+	TArray<FPropertiesToRecordForClass> PropertiesToRecord;
 };
\ No newline at end of file
diff --git a/Engine/Source/Editor/Sequencer/Private/DisplayNodes/SequencerFolderNode.cpp b/Engine/Source/Editor/Sequencer/Private/DisplayNodes/SequencerFolderNode.cpp
--- a/Engine/Source/Editor/Sequencer/Private/DisplayNodes/SequencerFolderNode.cpp
+++ b/Engine/Source/Editor/Sequencer/Private/DisplayNodes/SequencerFolderNode.cpp
@@ -261,7 +261,7 @@ void FSequencerFolderNode::Drop( const TArray<TSharedRef<FSequencerDisplayNode>>
 		}
 	}
 	SetExpansionState( true );
-	ParentTree.GetSequencer().NotifyMovieSceneDataChanged();
+	ParentTree.GetSequencer().NotifyMovieSceneDataChanged( EMovieSceneDataChangeType::MovieSceneStructureItemsChanged );
 }
 
 void FSequencerFolderNode::BuildContextMenu(FMenuBuilder& MenuBuilder)
diff --git a/Engine/Source/Editor/Sequencer/Private/DisplayNodes/SequencerObjectBindingNode.cpp b/Engine/Source/Editor/Sequencer/Private/DisplayNodes/SequencerObjectBindingNode.cpp
--- a/Engine/Source/Editor/Sequencer/Private/DisplayNodes/SequencerObjectBindingNode.cpp
+++ b/Engine/Source/Editor/Sequencer/Private/DisplayNodes/SequencerObjectBindingNode.cpp
@@ -581,7 +581,22 @@ void FSequencerObjectBindingNode::HandleAddTrackSubMenuNew(FMenuBuilder& AddTrac
 
 void FSequencerObjectBindingNode::HandleLabelsSubMenuCreate(FMenuBuilder& MenuBuilder)
 {
-	MenuBuilder.AddWidget(SNew(SSequencerLabelEditor, GetSequencer(), ObjectBinding), FText::GetEmpty(), true);
+	const TSet< TSharedRef<FSequencerDisplayNode> >& SelectedNodes = GetSequencer().GetSelection().GetSelectedOutlinerNodes();
+	TArray<FGuid> ObjectBindingIds;
+	for (TSharedRef<const FSequencerDisplayNode> SelectedNode : SelectedNodes )
+	{
+		if (SelectedNode->GetType() == ESequencerNode::Object)
+		{
+			TSharedRef<const FSequencerObjectBindingNode> ObjectBindingNode = StaticCastSharedRef<const FSequencerObjectBindingNode>(SelectedNode);
+			FGuid ObjectBindingId = ObjectBindingNode->GetObjectBinding();
+			if (ObjectBindingId.IsValid())
+			{
+				ObjectBindingIds.Add(ObjectBindingId);
+			}
+		}
+	}
+
+	MenuBuilder.AddWidget(SNew(SSequencerLabelEditor, GetSequencer(), ObjectBindingIds), FText::GetEmpty(), true);
 }
 
 
diff --git a/Engine/Source/Editor/Sequencer/Private/SAnimationOutlinerTreeNode.cpp b/Engine/Source/Editor/Sequencer/Private/SAnimationOutlinerTreeNode.cpp
--- a/Engine/Source/Editor/Sequencer/Private/SAnimationOutlinerTreeNode.cpp
+++ b/Engine/Source/Editor/Sequencer/Private/SAnimationOutlinerTreeNode.cpp
@@ -224,6 +224,7 @@ void SAnimationOutlinerTreeNode::Construct( const FArguments& InArgs, TSharedRef
 				.ContentPadding(0)
 				.VAlign(VAlign_Fill)
 				.HasDownArrow(false)
+				.IsFocusable(false)
 				.ButtonStyle(FEditorStyle::Get(), "Sequencer.AnimationOutliner.ColorStrip")
 				.OnGetMenuContent(this, &SAnimationOutlinerTreeNode::OnGetColorPicker)
 				.OnMenuOpenChanged_Lambda([](bool bIsOpen){
diff --git a/Engine/Source/Editor/Sequencer/Private/SKeyNavigationButtons.h b/Engine/Source/Editor/Sequencer/Private/SKeyNavigationButtons.h
--- a/Engine/Source/Editor/Sequencer/Private/SKeyNavigationButtons.h
+++ b/Engine/Source/Editor/Sequencer/Private/SKeyNavigationButtons.h
@@ -48,6 +48,7 @@ public:
 					.OnClicked(this, &SKeyNavigationButtons::OnPreviousKeyClicked)
 					.ForegroundColor( FSlateColor::UseForeground() )
 					.ContentPadding(0)
+					.IsFocusable(false)
 					[
 						SNew(STextBlock)
 						.Font(FEditorStyle::Get().GetFontStyle("FontAwesome.7"))
@@ -71,6 +72,7 @@ public:
 					.OnClicked(this, &SKeyNavigationButtons::OnAddKeyClicked)
 					.ForegroundColor( FSlateColor::UseForeground() )
 					.ContentPadding(0)
+					.IsFocusable(false)
 					[
 						SNew(STextBlock)
 						.Font(FEditorStyle::Get().GetFontStyle("FontAwesome.7"))
@@ -94,6 +96,7 @@ public:
 					.OnClicked(this, &SKeyNavigationButtons::OnNextKeyClicked)
 					.ContentPadding(0)
 					.ForegroundColor( FSlateColor::UseForeground() )
+					.IsFocusable(false)
 					[
 						SNew(STextBlock)
 						.Font(FEditorStyle::Get().GetFontStyle("FontAwesome.7"))
diff --git a/Engine/Source/Editor/Sequencer/Private/SSequencer.cpp b/Engine/Source/Editor/Sequencer/Private/SSequencer.cpp
--- a/Engine/Source/Editor/Sequencer/Private/SSequencer.cpp
+++ b/Engine/Source/Editor/Sequencer/Private/SSequencer.cpp
@@ -33,7 +33,6 @@
 #include "SSequencerShotFilterOverlay.h"
 #include "GenericCommands.h"
 #include "SequencerContextMenus.h"
-#include "SSequencerTreeViewBox.h"
 #include "NumericTypeInterface.h"
 #include "NumericUnitTypeInterface.inl"
 #include "SNumericEntryBox.h"
@@ -348,7 +347,7 @@ void SSequencer::Construct(const FArguments& InArgs, TSharedRef<FSequencer> InSe
 									+ SHorizontalBox::Slot()
 									.FillWidth( FillCoefficient_0 )
 									[
-										SNew(SSequencerTreeViewBox, InSequencer, SharedThis(this))
+										SNew(SBox)
 										[
 											TreeView.ToSharedRef()
 										]
@@ -969,18 +968,27 @@ TSharedRef<SWidget> SSequencer::MakeGeneralMenu()
 	{
 		MenuBuilder.AddMenuEntry(FSequencerCommands::Get().SetSelectionRangeStart);
 		MenuBuilder.AddMenuEntry(FSequencerCommands::Get().SetSelectionRangeEnd);
+		MenuBuilder.AddMenuEntry(FSequencerCommands::Get().ResetSelectionRange);
 		MenuBuilder.AddMenuEntry(FSequencerCommands::Get().SelectKeysInSelectionRange);
 	}
 	MenuBuilder.EndSection();
 
 	// other options
 	MenuBuilder.AddMenuSeparator();
+
 	if (SequencerPtr.Pin()->IsLevelEditorSequencer())
 	{
 		MenuBuilder.AddMenuEntry(FSequencerCommands::Get().FixActorReferences);
 	}
 	MenuBuilder.AddMenuEntry(FSequencerCommands::Get().FixFrameTiming);
 
+	MenuBuilder.AddMenuSeparator();
+
+	if ( SequencerPtr.Pin()->IsLevelEditorSequencer() )
+	{
+		MenuBuilder.AddMenuEntry( FSequencerCommands::Get().ExportSceneAndSequence );
+	}
+
 	return MenuBuilder.MakeWidget();
 }
 
@@ -1366,7 +1374,7 @@ void SSequencer::OnAssetsDropped( const FAssetDragDropOp& DragDropOp )
 	if( bObjectAdded )
 	{
 		// Update the sequencers view of the movie scene data when any object is added
-		SequencerRef.UpdateRuntimeInstances();
+		SequencerRef.NotifyMovieSceneDataChanged( EMovieSceneDataChangeType::MovieSceneStructureItemAdded );
 
 		// Update the tree and synchronize selection
 		UpdateLayoutTree();
@@ -1388,12 +1396,6 @@ void SSequencer::OnClassesDropped( const FClassDragDropOp& DragDropOp )
 			UObject* Object = Class->GetDefaultObject();
 
 			FGuid NewGuid = SequencerRef.MakeNewSpawnable( *Object );
-
-			if (NewGuid.IsValid())
-			{
-				// Update the sequencers view of the movie scene data
-				SequencerRef.NotifyMovieSceneDataChanged();
-			}
 		}
 	}
 }
@@ -1440,10 +1442,6 @@ void SSequencer::OnUnloadedClassesDropped( const FUnloadedClassDragDropOp& DragD
 		if( Object != nullptr )
 		{
 			FGuid NewGuid = SequencerRef.MakeNewSpawnable( *Object );
-			if (NewGuid.IsValid())
-			{
-				SequencerRef.NotifyMovieSceneDataChanged();
-			}
 		}
 	}
 }
diff --git a/Engine/Source/Editor/Sequencer/Private/SSequencerLabelEditor.cpp b/Engine/Source/Editor/Sequencer/Private/SSequencerLabelEditor.cpp
--- a/Engine/Source/Editor/Sequencer/Private/SSequencerLabelEditor.cpp
+++ b/Engine/Source/Editor/Sequencer/Private/SSequencerLabelEditor.cpp
@@ -11,10 +11,10 @@
 /* SSequencerLabelEditor interface
  *****************************************************************************/
 
-void SSequencerLabelEditor::Construct(const FArguments& InArgs, FSequencer& InSequencer, const FGuid& InObjectId)
+void SSequencerLabelEditor::Construct(const FArguments& InArgs, FSequencer& InSequencer, const TArray<FGuid>& InObjectIds)
 {
 	Sequencer = &InSequencer;
-	ObjectId = InObjectId;
+	ObjectIds = InObjectIds;
 
 	ChildSlot
 	[
@@ -81,7 +81,10 @@ void SSequencerLabelEditor::Construct(const FArguments& InArgs, FSequencer& InSe
 
 void SSequencerLabelEditor::CreateLabelFromFilterText()
 {
-	Sequencer->GetLabelManager().AddObjectLabel(ObjectId, FilterBox->GetText().ToString());
+	for (auto ObjectId : ObjectIds)
+	{
+		Sequencer->GetLabelManager().AddObjectLabel(ObjectId, FilterBox->GetText().ToString());
+	}
 	FilterBox->SetText(FText::GetEmpty());
 	ReloadLabelList(true);
 }
@@ -179,11 +182,17 @@ void SSequencerLabelEditor::HandleLabelListViewRowCheckedStateChanged(ECheckBoxS
 
 	if (State == ECheckBoxState::Checked)
 	{
-		LabelManager.AddObjectLabel(ObjectId, *Label);
+		for (auto ObjectId : ObjectIds)
+		{
+			LabelManager.AddObjectLabel(ObjectId, *Label);
+		}
 	}
 	else
 	{
-		LabelManager.RemoveObjectLabel(ObjectId, *Label);
+		for (auto ObjectId : ObjectIds)
+		{
+			LabelManager.RemoveObjectLabel(ObjectId, *Label);
+		}
 	}
 }
 
@@ -196,10 +205,25 @@ FText SSequencerLabelEditor::HandleLabelListViewRowHighlightText() const
 
 ECheckBoxState SSequencerLabelEditor::HandleLabelListViewRowIsChecked(TSharedPtr<FString> Label) const
 {
-	const FMovieSceneTrackLabels* Labels = Sequencer->GetLabelManager().GetObjectLabels(ObjectId);
-	return ((Labels != nullptr) && Labels->Strings.Contains(*Label))
-		? ECheckBoxState::Checked
-		: ECheckBoxState::Unchecked;
+	int32 NumChecked = 0;
+	for (auto ObjectId : ObjectIds)
+	{
+		const FMovieSceneTrackLabels* Labels = Sequencer->GetLabelManager().GetObjectLabels(ObjectId);
+		if (Labels != nullptr)
+		{
+			if (Labels->Strings.Contains(*Label))
+			{
+				++NumChecked;
+			}
+		}
+	}
+
+	if (NumChecked == ObjectIds.Num())
+	{
+		return ECheckBoxState::Checked;
+	}
+
+	return ECheckBoxState::Unchecked;
 }
 
 
diff --git a/Engine/Source/Editor/Sequencer/Private/SSequencerLabelEditor.h b/Engine/Source/Editor/Sequencer/Private/SSequencerLabelEditor.h
--- a/Engine/Source/Editor/Sequencer/Private/SSequencerLabelEditor.h
+++ b/Engine/Source/Editor/Sequencer/Private/SSequencerLabelEditor.h
@@ -18,9 +18,9 @@ public:
 	 * Construct this widget.
 	 *
 	 * @param InArgs The construction arguments.
-	 * @param InObjectId The unique identifier of the object whose labels are being edited.
+	 * @param InObjectIds The unique identifier of the objects whose labels are being edited.
 	 */
-	void Construct(const FArguments& InArgs, FSequencer& InSequencer, const FGuid& InObjectId);
+	void Construct(const FArguments& InArgs, FSequencer& InSequencer, const TArray<FGuid>& InObjectIds);
 
 protected:
 
@@ -74,8 +74,8 @@ private:
 	/** Holds the list view for filtered track labels. */
 	TSharedPtr<SListView<TSharedPtr<FString>>> LabelListView;
 
-	/** The identifier of the object being edited. */
-	FGuid ObjectId;
+	/** The identifiers of the objects being edited. */
+	TArray<FGuid> ObjectIds;
 
 	/** The sequencer object. */
 	FSequencer* Sequencer;
diff --git a/Engine/Source/Editor/Sequencer/Private/SSequencerTreeView.cpp b/Engine/Source/Editor/Sequencer/Private/SSequencerTreeView.cpp
--- a/Engine/Source/Editor/Sequencer/Private/SSequencerTreeView.cpp
+++ b/Engine/Source/Editor/Sequencer/Private/SSequencerTreeView.cpp
@@ -552,7 +552,7 @@ bool SSequencerTreeView::SynchronizeSequencerSelectionWithTreeSelection()
 TSharedPtr<SWidget> SSequencerTreeView::OnContextMenuOpening()
 {
 	const TSet<TSharedRef<FSequencerDisplayNode>> SelectedNodes = SequencerNodeTree->GetSequencer().GetSelection().GetSelectedOutlinerNodes();
-	if ( SelectedNodes.Num() == 1 )
+	if ( SelectedNodes.Num() > 0 )
 	{
 		return SelectedNodes.Array()[0]->OnSummonContextMenu();
 	}
diff --git a/Engine/Source/Editor/Sequencer/Private/Sequencer.cpp b/Engine/Source/Editor/Sequencer/Private/Sequencer.cpp
--- a/Engine/Source/Editor/Sequencer/Private/Sequencer.cpp
+++ b/Engine/Source/Editor/Sequencer/Private/Sequencer.cpp
@@ -35,6 +35,7 @@
 #include "MovieScene3DTransformSection.h"
 #include "MovieSceneSubSection.h"
 #include "MovieSceneSubTrack.h"
+#include "MovieSceneCinematicShotSection.h"
 #include "ISequencerSection.h"
 #include "MovieSceneSequenceInstance.h"
 #include "IKeyArea.h"
@@ -72,6 +73,9 @@
 #include "CameraRig_Rail.h"
 #include "CameraRig_Crane.h"
 #include "Components/SplineComponent.h"
+#include "MainFrame.h"
+#include "DesktopPlatformModule.h"
+#include "FbxExporter.h"
 
 #define LOCTEXT_NAMESPACE "Sequencer"
 
@@ -99,6 +103,12 @@ void FSequencer::InitSequencer(const FSequencerInitParams& InitParams, const TSh
 		SpawnRegister = MakeShareable(new FNullMovieSceneSpawnRegister);
 	}
 
+	EventContextsAttribute = InitParams.EventContexts;
+	if (EventContextsAttribute.IsSet())
+	{
+		CachedEventContexts = EventContextsAttribute.Get();
+	}
+
 	// If this sequencer edits the level, close out the existing active sequencer and mark this as the active sequencer.
 	if (bIsEditingWithinLevelEditor)
 	{
@@ -124,6 +134,8 @@ void FSequencer::InitSequencer(const FSequencerInitParams& InitParams, const TSh
 		FEditorDelegates::OnNewActorsDropped.AddSP(this, &FSequencer::OnNewActorsDropped);
 
 		USelection::SelectionChangedEvent.AddSP( this, &FSequencer::OnActorSelectionChanged );
+
+		AddLevelViewportMenuExtender();
 	}
 
 	Settings = USequencerSettingsContainer::GetOrCreate<USequencerSettings>(*InitParams.ViewParams.UniqueName);
@@ -140,6 +152,8 @@ void FSequencer::InitSequencer(const FSequencerInitParams& InitParams, const TSh
 	ToolkitHost = InitParams.ToolkitHost;
 
 	ScrubPosition = InitParams.ViewParams.InitialScrubPosition;
+	PlayRate = 1.f;
+	ShuttleMultiplier = 0;
 	ObjectChangeListener = InObjectChangeListener;
 	DetailKeyframeHandler = InDetailKeyframeHandler;
 
@@ -203,7 +217,7 @@ void FSequencer::InitSequencer(const FSequencerInitParams& InitParams, const TSh
 	OverlayCurve = OverlayAnimation.AddCurve(0.f, 0.2f, ECurveEaseFunction::QuadIn);
 
 	// Update initial movie scene data
-	NotifyMovieSceneDataChanged();
+	NotifyMovieSceneDataChanged( EMovieSceneDataChangeType::ActiveMovieSceneChanged );
 	UpdateTimeBoundsToFocusedMovieScene();
 
 	// NOTE: Could fill in asset editor commands here!
@@ -258,6 +272,7 @@ FSequencer::~FSequencer()
 	if( bIsEditingWithinLevelEditor )
 	{
 		DetachTransportControlsFromViewports();
+		RemoveLevelViewportMenuExtender();
 	}
 
 	// When a new level sequence asset is opened while an existing one is open, the original level sequence asset is closed, 
@@ -317,6 +332,7 @@ void FSequencer::Close()
 
 	if (bIsEditingWithinLevelEditor)
 	{
+		RemoveLevelViewportMenuExtender();
 		DeactivateDetailKeyframeHandler();
 			
 		FPropertyEditorModule& EditModule = FModuleManager::Get().GetModuleChecked<FPropertyEditorModule>("PropertyEditor");			
@@ -332,11 +348,15 @@ void FSequencer::Close()
 
 void FSequencer::Tick(float InDeltaTime)
 {
+	if (EventContextsAttribute.IsBound())
+	{
+		CachedEventContexts = EventContextsAttribute.Get();
+	}
+
 	Selection.Tick();
 
 	if ( bNeedInstanceRefresh )
 	{
-		// @todo - Sequencer Will be called too often
 		UpdateRuntimeInstances();
 		bNeedInstanceRefresh = false;
 	}
@@ -388,8 +408,15 @@ void FSequencer::Tick(float InDeltaTime)
 		}
 	}
 
+	float TimeDilation = 1.f;
+	UWorld* PlaybackContext = Cast<UWorld>(GetPlaybackContext());
+	if (PlaybackContext)
+	{
+		TimeDilation = PlaybackContext->GetWorldSettings()->MatineeTimeDilation;
+	}
+
 	// calculate new global time
-	float NewTime = GetGlobalTime() + InDeltaTime * GWorld->GetWorldSettings()->MatineeTimeDilation;
+	float NewTime = GetGlobalTime() + InDeltaTime * TimeDilation * PlayRate;
 
 	if (PlaybackState == EMovieScenePlayerStatus::Playing ||
 		PlaybackState == EMovieScenePlayerStatus::Recording)
@@ -400,7 +427,7 @@ void FSequencer::Tick(float InDeltaTime)
 	// Tick all the tools we own as well
 	for (int32 EditorIndex = 0; EditorIndex < TrackEditors.Num(); ++EditorIndex)
 	{
-		TrackEditors[EditorIndex]->Tick(InDeltaTime);
+		TrackEditors[EditorIndex]->Tick(InDeltaTime * PlayRate);
 	}
 
 	if (!IsInSilentMode())
@@ -417,10 +444,31 @@ void FSequencer::Tick(float InDeltaTime)
 			Section->SetEndTime(Section->GetStartTime() + SequenceRecorder.GetCurrentRecordingLength());
 		}
 	}
+
+	// Reset any player controllers that we were possessing, if we're not possessing them any more
+	if (!IsPerspectiveViewportCameraCutEnabled() && PrePossessionViewTargets.Num())
+	{
+		for (const FCachedViewTarget& CachedView : PrePossessionViewTargets)
+		{
+			APlayerController* PlayerController = CachedView.PlayerController.Get();
+			AActor* ViewTarget = CachedView.ViewTarget.Get();
+
+			if (PlayerController && ViewTarget)
+			{
+				PlayerController->SetViewTarget(ViewTarget);
+			}
+		}
+		PrePossessionViewTargets.Reset();
+	}
 }
 
 void FSequencer::PostTickRenderStateFixup()
 {
+	if ( !IsLevelEditorSequencer() )
+	{
+		return;
+	}
+
 	if (PlaybackState == EMovieScenePlayerStatus::Jumping || PlaybackState == EMovieScenePlayerStatus::Stepping)
 	{
 		const bool bApplyFixup = true;
@@ -471,7 +519,11 @@ void FSequencer::PostTickRenderStateFixup()
 			SetGlobalTimeDirectly(PreviousTime);
 			Tick(StepDelta);
 
-			GWorld->SendAllEndOfFrameUpdates();
+			UWorld* PlaybackContext = Cast<UWorld>(GetPlaybackContext());
+			if (PlaybackContext)
+			{
+				PlaybackContext->SendAllEndOfFrameUpdates();
+			}
 
 			SetPlaybackStatus(EMovieScenePlayerStatus::Stopped);
 			SetGlobalTimeDirectly(Time);
@@ -641,6 +693,10 @@ UObject* FSequencer::GetPlaybackContext() const
 	return nullptr;
 }
 
+TArray<UObject*> FSequencer::GetEventContexts() const
+{
+	return CachedEventContexts;
+}
 
 void FSequencer::GetAllKeyedProperties(UObject& Object, TSet<UProperty*>& OutProperties)
 {
@@ -802,7 +858,7 @@ void FSequencer::DeleteSections(const TSet<TWeakObjectPtr<UMovieSceneSection>>&
 	if (bAnythingRemoved)
 	{
 		// Full refresh required just in case the last section was removed from any track.
-		NotifyMovieSceneDataChanged();
+		NotifyMovieSceneDataChanged( EMovieSceneDataChangeType::MovieSceneStructureItemRemoved );
 	}
 
 	SequencerHelpers::ValidateNodesWithSelectedKeysOrSections(*this);
@@ -829,7 +885,7 @@ void FSequencer::DeleteSelectedKeys()
 
 	if (bAnythingRemoved)
 	{
-		UpdateRuntimeInstances();
+		NotifyMovieSceneDataChanged( EMovieSceneDataChangeType::TrackValueChanged );
 	}
 
 	Selection.EmptySelectedKeys();
@@ -863,7 +919,7 @@ void FSequencer::SetInterpTangentMode(ERichCurveInterpMode InterpMode, ERichCurv
 
 	if (bAnythingChanged)
 	{
-		UpdateRuntimeInstances();
+		NotifyMovieSceneDataChanged( EMovieSceneDataChangeType::TrackValueChanged );
 	}
 }
 
@@ -919,7 +975,7 @@ void FSequencer::SnapToFrame()
 
 	if (bAnythingChanged)
 	{
-		UpdateRuntimeInstances();
+		NotifyMovieSceneDataChanged( EMovieSceneDataChangeType::TrackValueChanged );
 	}
 }
 
@@ -939,7 +995,7 @@ void FSequencer::NotifyMapChanged( class UWorld* NewWorld, EMapChangeType MapCha
 	{
 		SpawnRegister->CleanUp(*this);
 
-		NotifyMovieSceneDataChanged();
+		UpdateRuntimeInstances();
 	}
 }
 
@@ -950,15 +1006,49 @@ void FSequencer::OnActorsDropped( const TArray<TWeakObjectPtr<AActor> >& Actors
 }
 
 
-void FSequencer::NotifyMovieSceneDataChanged()
+void FSequencer::NotifyMovieSceneDataChangedInternal()
+{
+	NotifyMovieSceneDataChanged( EMovieSceneDataChangeType::Unknown );
+}
+
+
+void FSequencer::NotifyMovieSceneDataChanged( EMovieSceneDataChangeType DataChangeType )
 {
-	LabelManager.SetMovieScene(SequenceInstanceStack.Top()->GetSequence()->GetMovieScene());
+	if ( DataChangeType == EMovieSceneDataChangeType::ActiveMovieSceneChanged ||
+		DataChangeType == EMovieSceneDataChangeType::Unknown )
+	{
+		LabelManager.SetMovieScene( SequenceInstanceStack.Top()->GetSequence()->GetMovieScene() );
+	}
+
 	StoredPlaybackState = GetPlaybackStatus();
-	SetPlaybackStatus(EMovieScenePlayerStatus::Stopped);
-	bNeedTreeRefresh = true;
-	bNeedInstanceRefresh = true;
+
+	if ( DataChangeType == EMovieSceneDataChangeType::MovieSceneStructureItemRemoved ||
+		DataChangeType == EMovieSceneDataChangeType::MovieSceneStructureItemsChanged ||
+		DataChangeType == EMovieSceneDataChangeType::Unknown )
+	{
+		// When structure items are removed, or we don't know what may have changed, refresh the tree and instances immediately so that the data
+		// is in a consistent state when the UI is updated during the next tick.
+		SetPlaybackStatus( EMovieScenePlayerStatus::Stopped );
+		UpdateRuntimeInstances();
+		SelectionPreview.Empty();
+		SequencerWidget->UpdateLayoutTree();
+		bNeedInstanceRefresh = false;
+		bNeedTreeRefresh = false;
+		SetPlaybackStatus( StoredPlaybackState );
+	}
+	else
+	{
+		if ( DataChangeType != EMovieSceneDataChangeType::TrackValueChanged )
+		{
+			// All changes types except for track value changes require refreshing the outliner tree.
+			SetPlaybackStatus( EMovieScenePlayerStatus::Stopped );
+			bNeedTreeRefresh = true;
+		}
+		bNeedInstanceRefresh = true;
+	}
 
 	UpdatePlaybackRange();
+	OnMovieSceneDataChangedDelegate.Broadcast();
 }
 
 
@@ -1072,6 +1162,10 @@ void FSequencer::SelectKeysInSelectionRange(const TSharedRef<FSequencerDisplayNo
 	}
 }
 
+void FSequencer::ResetSelectionRange()
+{
+	SetSelectionRange(TRange<float>::Empty());
+}
 
 void FSequencer::SelectKeysInSelectionRange()
 {
@@ -1088,12 +1182,6 @@ void FSequencer::SelectKeysInSelectionRange()
 }
 
 
-void FSequencer::DeleteSelectionRange()
-{
-
-}
-
-
 TRange<float> FSequencer::GetPlaybackRange() const
 {
 	return GetFocusedMovieSceneSequence()->GetMovieScene()->GetPlaybackRange();
@@ -1282,7 +1370,7 @@ float FSequencer::GetGlobalTime() const
 }
 
 
-void FSequencer::SetGlobalTime( float NewTime, ESnapTimeMode SnapTimeMode, bool bLooped )
+void FSequencer::SetGlobalTime( float NewTime, ESnapTimeMode SnapTimeMode, bool bRestarted )
 {
 	if (IsAutoScrollEnabled())
 	{
@@ -1300,11 +1388,11 @@ void FSequencer::SetGlobalTime( float NewTime, ESnapTimeMode SnapTimeMode, bool
 		}
 	}
 
-	SetGlobalTimeDirectly(NewTime, SnapTimeMode, bLooped);
+	SetGlobalTimeDirectly(NewTime, SnapTimeMode, bRestarted);
 }
 
 
-void FSequencer::SetGlobalTimeDirectly( float NewTime, ESnapTimeMode SnapTimeMode, bool bLooped )
+void FSequencer::SetGlobalTimeDirectly( float NewTime, ESnapTimeMode SnapTimeMode, bool bRestarted )
 {
 	float LastTime = ScrubPosition;
 
@@ -1322,14 +1410,25 @@ void FSequencer::SetGlobalTimeDirectly( float NewTime, ESnapTimeMode SnapTimeMod
 	ScrubPosition = NewTime;
 
 	UMovieScene* MovieScene = SequenceInstanceStack.Top()->GetSequence()->GetMovieScene();
+
+	if (LastTime == GetPlaybackRange().GetLowerBoundValue())
+	{
+		bRestarted = true;
+	}
+
 	if ( MovieScene->GetForceFixedFrameIntervalPlayback() && MovieScene->GetFixedFrameInterval() > 0 )
 	{
 		NewTime = UMovieScene::CalculateFixedFrameTime( NewTime, MovieScene->GetFixedFrameInterval() );
 		LastTime = UMovieScene::CalculateFixedFrameTime( LastTime, MovieScene->GetFixedFrameInterval() );
 	}
 
+	// Force last time to be the lower bound of the playback range when restarting or looping around 
+	if (bRestarted)
+	{
+		LastTime = GetPlaybackRange().GetLowerBoundValue();
+	}
+
 	EMovieSceneUpdateData UpdateData(NewTime, LastTime);
-	UpdateData.bLooped = bLooped;
 	SequenceInstanceStack.Top()->Update(UpdateData, *this);
 
 	// Ensure any relevant spawned objects are cleaned up if we're playing back the master sequence
@@ -1504,7 +1603,7 @@ FGuid FindUnspawnedObjectGuid(UObject& InObject, UMovieSceneSequence& Sequence)
 
 	// If the object is an archetype, the it relates to an unspawned spawnable.
 	UObject* ParentObject = Sequence.GetParentObject(&InObject);
-	if (ParentObject && ParentObject->HasAnyFlags(RF_ArchetypeObject))
+	if (ParentObject && FMovieSceneSpawnable::IsSpawnableTemplate(*ParentObject))
 	{
 		FMovieSceneSpawnable* ParentSpawnable = MovieScene->FindSpawnable([&](FMovieSceneSpawnable& InSpawnable){
 			return InSpawnable.GetObjectTemplate() == ParentObject;
@@ -1525,7 +1624,7 @@ FGuid FindUnspawnedObjectGuid(UObject& InObject, UMovieSceneSequence& Sequence)
 			}
 		}
 	}
-	else if (InObject.HasAnyFlags(RF_ArchetypeObject))
+	else if (FMovieSceneSpawnable::IsSpawnableTemplate(InObject))
 	{
 		FMovieSceneSpawnable* SpawnableByArchetype = MovieScene->FindSpawnable([&](FMovieSceneSpawnable& InSpawnable){
 			return InSpawnable.GetObjectTemplate() == &InObject;
@@ -1589,7 +1688,7 @@ FGuid FSequencer::GetHandleToObject( UObject* Object, bool bCreateHandleIfMissin
 
 		ObjectGuid = CreateBinding(*Object, PossessedActor != nullptr ? PossessedActor->GetActorLabel() : Object->GetName());
 
-		NotifyMovieSceneDataChanged();
+		NotifyMovieSceneDataChanged( EMovieSceneDataChangeType::MovieSceneStructureItemAdded );
 	}
 	
 	return ObjectGuid;
@@ -1611,8 +1710,74 @@ void FSequencer::GetRuntimeObjects( TSharedRef<FMovieSceneSequenceInstance> Movi
 	}
 }
 
+void FSequencer::PossessPIEViewports(UObject* CameraObject, UObject* UnlockIfCameraObject, bool bJumpCut)
+{
+	for (const FWorldContext& WorldContext : GEngine->GetWorldContexts())
+	{
+		UWorld* World = WorldContext.World();
+		if (!World || WorldContext.WorldType != EWorldType::PIE)
+		{
+			continue;
+		}
+		APlayerController* PC = World->GetGameInstance()->GetFirstLocalPlayerController();
+		if (PC == nullptr)
+		{
+			continue;
+		}
+
+		TWeakObjectPtr<APlayerController> WeakPC = PC;
+		auto FindViewTarget = [=](const FCachedViewTarget& In){ return In.PlayerController == WeakPC; };
+
+		// skip same view target
+		AActor* ViewTarget = PC->GetViewTarget();
+
+		// save the last view target so that it can be restored when the camera object is null
+		if (!PrePossessionViewTargets.ContainsByPredicate(FindViewTarget))
+		{
+			PrePossessionViewTargets.Add(FCachedViewTarget{ PC, ViewTarget });
+		}
+
+		if (CameraObject == ViewTarget)
+		{
+			if ( bJumpCut && PC->PlayerCameraManager )
+			{
+				PC->PlayerCameraManager->bGameCameraCutThisFrame = true;
+			}
+			continue;
+		}
+
+		// skip unlocking if the current view target differs
+		AActor* UnlockIfCameraActor = Cast<AActor>(UnlockIfCameraObject);
+
+		if ((CameraObject == nullptr) && (UnlockIfCameraActor != ViewTarget))
+		{
+			return;
+		}
+
+		// override the player controller's view target
+		AActor* CameraActor = Cast<AActor>(CameraObject);
+
+		// if the camera object is null, use the last view target so that it is restored to the state before the sequence takes control
+		if (CameraActor == nullptr)
+		{
+			if (const FCachedViewTarget* CachedTarget = PrePossessionViewTargets.FindByPredicate(FindViewTarget))
+			{
+				CameraActor = CachedTarget->ViewTarget.Get();
+			}
+		}
+
+		FViewTargetTransitionParams TransitionParams;
+		PC->SetViewTarget(CameraActor, TransitionParams);
 
-void FSequencer::UpdateCameraCut(UObject* CameraObject, UObject* UnlockIfCameraObject, bool bJumpCut) const
+		if (PC->PlayerCameraManager)
+		{
+			PC->PlayerCameraManager->bClientSimulatingViewTarget = (CameraActor != nullptr);
+			PC->PlayerCameraManager->bGameCameraCutThisFrame = true;
+		}
+	}
+}
+
+void FSequencer::UpdateCameraCut(UObject* CameraObject, UObject* UnlockIfCameraObject, bool bJumpCut)
 {
 	OnCameraCutEvent.Broadcast(CameraObject, bJumpCut);
 
@@ -1628,6 +1793,11 @@ void FSequencer::UpdateCameraCut(UObject* CameraObject, UObject* UnlockIfCameraO
 		return;
 	}
 
+	if (Settings->ShouldAllowPossessionOfPIEViewports())
+	{
+		PossessPIEViewports(CameraObject, UnlockIfCameraObject, bJumpCut);
+	}
+
 	AActor* UnlockIfCameraActor = Cast<AActor>(UnlockIfCameraObject);
 
 	for (FLevelEditorViewportClient* LevelVC : GEditor->LevelViewportClients)
@@ -1708,7 +1878,7 @@ void FSequencer::SetPlaybackStatus(EMovieScenePlayerStatus::Type InPlaybackStatu
 	PlaybackState = InPlaybackStatus;
 
 	// Inform the renderer when Sequencer is in a 'paused' state for the sake of inter-frame effects
-	const bool bIsPaused = (InPlaybackStatus == EMovieScenePlayerStatus::Stopped || InPlaybackStatus == EMovieScenePlayerStatus::Scrubbing);
+	const bool bIsPaused = (InPlaybackStatus == EMovieScenePlayerStatus::Stopped || InPlaybackStatus == EMovieScenePlayerStatus::Scrubbing || InPlaybackStatus == EMovieScenePlayerStatus::Stepping);
 
 	for (FLevelEditorViewportClient* LevelVC : GEditor->LevelViewportClients)
 	{
@@ -1726,6 +1896,9 @@ void FSequencer::SetPlaybackStatus(EMovieScenePlayerStatus::Type InPlaybackStatu
 	else
 	{
 		GEngine->SetMaxFPS(OldMaxTickRate);
+
+		PlayRate = 1.f;
+		ShuttleMultiplier = 0;
 	}
 }
 
@@ -1802,6 +1975,159 @@ void FSequencer::UpdateRuntimeInstances()
 	FEditorSupportDelegates::RedrawAllViewports.Broadcast();
 }
 
+void FSequencer::AddLevelViewportMenuExtender()
+{
+	FLevelEditorModule& LevelEditorModule = FModuleManager::Get().LoadModuleChecked<FLevelEditorModule>("LevelEditor");
+	auto& MenuExtenders = LevelEditorModule.GetAllLevelViewportContextMenuExtenders();
+
+	if (LevelViewportExtenderHandle.IsValid())
+	{
+		typedef FLevelEditorModule::FLevelViewportMenuExtender_SelectedActors DelegateType;
+		MenuExtenders.RemoveAll([=](const DelegateType& In){ return In.GetHandle() == LevelViewportExtenderHandle; });
+	}
+
+	MenuExtenders.Add(FLevelEditorModule::FLevelViewportMenuExtender_SelectedActors::CreateSP(this, &FSequencer::GetLevelViewportExtender));
+	LevelViewportExtenderHandle = MenuExtenders.Last().GetHandle();
+}
+
+void FSequencer::RemoveLevelViewportMenuExtender()
+{
+	if (LevelViewportExtenderHandle.IsValid())
+	{
+		FLevelEditorModule* LevelEditorModule = FModuleManager::Get().GetModulePtr<FLevelEditorModule>("LevelEditor");
+		if (LevelEditorModule)
+		{
+			typedef FLevelEditorModule::FLevelViewportMenuExtender_SelectedActors DelegateType;
+			LevelEditorModule->GetAllLevelViewportContextMenuExtenders().RemoveAll([=](const DelegateType& In){ return In.GetHandle() == LevelViewportExtenderHandle; });
+		}
+	}
+}
+
+TSharedRef<FExtender> FSequencer::GetLevelViewportExtender(const TSharedRef<FUICommandList> CommandList, const TArray<AActor*> InActors)
+{
+	TSharedRef<FExtender> Extender = MakeShareable(new FExtender);
+
+	FText ActorName;
+	if (InActors.Num() == 1)
+	{
+		ActorName = FText::Format(LOCTEXT("ActorNameSingular", "\"{0}\""), FText::FromString(InActors[0]->GetActorLabel()));
+	}
+	else if (InActors.Num() > 1)
+	{
+		ActorName = FText::Format(LOCTEXT("ActorNamePlural", "{0} Actors"), FText::AsNumber(InActors.Num()));
+	}
+
+	Extender->AddMenuExtension("ActorControl", EExtensionHook::After, CommandList, FMenuExtensionDelegate::CreateLambda(
+		[this, ActorName](FMenuBuilder& MenuBuilder){
+			MenuBuilder.BeginSection("SequenceRecording", LOCTEXT("SequenceRecordingHeading", "Sequence Recording"));
+
+			MenuBuilder.AddMenuEntry(
+				FText::Format(LOCTEXT("RecordSelectedActorsText", "Record {0} In Sequencer"), ActorName),
+				LOCTEXT("RecordSelectedActors_Tooltip", "Start recording the specified actors in sequencer."),
+				FSlateIcon(FEditorStyle::GetStyleSetName(), "SequenceRecorder.TabIcon"),
+				FUIAction(
+					FExecuteAction::CreateLambda([this]{ RecordSelectedActors(); })
+				)
+			);
+
+			MenuBuilder.EndSection();
+		}
+	));
+
+	return Extender;
+}
+
+void FSequencer::RecordSelectedActors()
+{
+	ISequenceRecorder& SequenceRecorder = FModuleManager::LoadModuleChecked<ISequenceRecorder>("SequenceRecorder");
+	if (SequenceRecorder.IsRecording())
+	{
+		FNotificationInfo Info(LOCTEXT("UnableToRecord_AlreadyRecording", "Cannot start a new recording while one is already in progress."));
+		Info.bUseLargeFont = false;
+		FSlateNotificationManager::Get().AddNotification(Info);
+		return;
+	}
+
+	if (Settings->ShouldRewindOnRecord())
+	{
+		Rewind();
+	}
+	
+	TArray<ACameraActor*> SelectedCameras;
+	TArray<AActor*> EntireSelection;
+
+	GEditor->GetSelectedActors()->GetSelectedObjects(SelectedCameras);
+	GEditor->GetSelectedActors()->GetSelectedObjects(EntireSelection);
+
+	UMovieScene* MovieScene = GetFocusedMovieSceneSequence()->GetMovieScene();
+
+	// Figure out what we're recording into - a sub track, or a camera cut track, or a shot track
+	UMovieSceneTrack* DestinationTrack = nullptr;
+	if (SelectedCameras.Num())
+	{
+		DestinationTrack = MovieScene->FindMasterTrack<UMovieSceneCinematicShotTrack>();
+		if (!DestinationTrack)
+		{
+			DestinationTrack = MovieScene->AddMasterTrack<UMovieSceneCinematicShotTrack>();
+		}
+	}
+	else if (EntireSelection.Num())
+	{
+		DestinationTrack = MovieScene->FindMasterTrack<UMovieSceneSubTrack>();
+		if (!DestinationTrack)
+		{
+			DestinationTrack = MovieScene->AddMasterTrack<UMovieSceneSubTrack>();
+		}
+	}
+	else
+	{
+		FNotificationInfo Info(LOCTEXT("UnableToRecordNoSelection", "Unable to start recording because no actors are selected"));
+		Info.bUseLargeFont = false;
+		FSlateNotificationManager::Get().AddNotification(Info);
+		return;
+	}
+
+	if (!DestinationTrack)
+	{
+		FNotificationInfo Info(LOCTEXT("UnableToRecord", "Unable to start recording because a valid sub track could not be found or created"));
+		Info.bUseLargeFont = false;
+		FSlateNotificationManager::Get().AddNotification(Info);
+		return;
+	}
+
+	int32 MaxRow = 0;
+	for (UMovieSceneSection* Section : DestinationTrack->GetAllSections())
+	{
+		MaxRow = FMath::Max(Section->GetRowIndex(), MaxRow);
+	}
+	// @todo: Get row at current time
+	UMovieSceneSubSection* NewSection = CastChecked<UMovieSceneSubSection>(DestinationTrack->CreateNewSection());
+	NewSection->SetRowIndex(MaxRow + 1);
+	DestinationTrack->AddSection(*NewSection);
+	NewSection->SetAsRecording(true);
+
+	NotifyMovieSceneDataChanged(EMovieSceneDataChangeType::MovieSceneStructureItemAdded);
+
+	if (UMovieSceneSubSection::IsSetAsRecording())
+	{
+		TArray<AActor*> ActorsToRecord;
+		for (AActor* Actor : EntireSelection)
+		{
+			AActor* CounterpartActor = EditorUtilities::GetSimWorldCounterpartActor(Actor);
+			ActorsToRecord.Add(CounterpartActor ? CounterpartActor : Actor);
+		}
+
+		const FString& PathToRecordTo = UMovieSceneSubSection::GetRecordingSection()->GetTargetPathToRecordTo();
+		const FString& SequenceName = UMovieSceneSubSection::GetRecordingSection()->GetTargetSequenceName();
+		SequenceRecorder.StartRecording(
+			ActorsToRecord,
+			FOnRecordingStarted::CreateSP(this, &FSequencer::HandleRecordingStarted),
+			FOnRecordingFinished::CreateSP(this, &FSequencer::HandleRecordingFinished),
+			PathToRecordTo,
+			SequenceName);
+	}
+}
+
 TSharedRef<SWidget> FSequencer::MakeTransportControls(bool bExtended)
 {
 	FEditorWidgetsModule& EditorWidgetsModule = FModuleManager::Get().LoadModuleChecked<FEditorWidgetsModule>( "EditorWidgets" );
@@ -1810,7 +2136,8 @@ TSharedRef<SWidget> FSequencer::MakeTransportControls(bool bExtended)
 	{
 		TransportControlArgs.OnBackwardEnd.BindSP( this, &FSequencer::OnStepToBeginning );
 		TransportControlArgs.OnBackwardStep.BindSP( this, &FSequencer::OnStepBackward );
-		TransportControlArgs.OnForwardPlay.BindSP( this, &FSequencer::OnPlay, true );
+		TransportControlArgs.OnForwardPlay.BindSP( this, &FSequencer::OnPlay, true, 1.f );
+		TransportControlArgs.OnBackwardPlay.BindSP( this, &FSequencer::OnPlay, true, -1.f );
 		TransportControlArgs.OnForwardStep.BindSP( this, &FSequencer::OnStepForward );
 		TransportControlArgs.OnForwardEnd.BindSP( this, &FSequencer::OnStepToEnd );
 		TransportControlArgs.OnToggleLooping.BindSP( this, &FSequencer::OnToggleLooping );
@@ -1827,6 +2154,7 @@ TSharedRef<SWidget> FSequencer::MakeTransportControls(bool bExtended)
 			TransportControlArgs.WidgetsToCreate.Add(FTransportControlWidget(FOnMakeTransportWidget::CreateSP(this, &FSequencer::OnCreateTransportJumpToPreviousKey)));
 		}
 		TransportControlArgs.WidgetsToCreate.Add(FTransportControlWidget(ETransportControlWidgetType::BackwardStep));
+		TransportControlArgs.WidgetsToCreate.Add(FTransportControlWidget(ETransportControlWidgetType::BackwardPlay));
 		TransportControlArgs.WidgetsToCreate.Add(FTransportControlWidget(ETransportControlWidgetType::ForwardPlay));
 		TransportControlArgs.WidgetsToCreate.Add(FTransportControlWidget(FOnMakeTransportWidget::CreateSP(this, &FSequencer::OnCreateTransportRecord)));
 		TransportControlArgs.WidgetsToCreate.Add(FTransportControlWidget(ETransportControlWidgetType::ForwardStep));
@@ -2182,22 +2510,17 @@ EVisibility FSequencer::GetTransportControlVisibility(TSharedPtr<ILevelViewport>
 		ViewportClient.AllowsCinematicPreview() ? EVisibility::Visible : EVisibility::Collapsed;
 }
 
-FReply FSequencer::OnPlay(bool bTogglePlay)
+FReply FSequencer::OnPlay(bool bTogglePlay, float InPlayRate)
 {
 	if( (PlaybackState == EMovieScenePlayerStatus::Playing ||
-		 PlaybackState == EMovieScenePlayerStatus::Recording) && bTogglePlay )
+		 PlaybackState == EMovieScenePlayerStatus::Recording) && bTogglePlay && (FMath::Sign(InPlayRate) == FMath::Sign(PlayRate) ) )
 	{
-		SetPlaybackStatus(EMovieScenePlayerStatus::Stopped);
-
-		// Snap time when stopping play
-		SetGlobalTime(ScrubPosition, ESnapTimeMode::STM_Interval);
-
-		// Update on stop (cleans up things like sounds that are playing)
-		EMovieSceneUpdateData UpdateData(ScrubPosition, ScrubPosition);
-		SequenceInstanceStack.Top()->Update(UpdateData, *this );
+		Pause();
 	}
 	else
 	{
+		PlayRate = InPlayRate;
+
 		SetPlaybackStatus(EMovieScenePlayerStatus::Playing);
 
 		// Make sure Slate ticks during playback
@@ -2220,11 +2543,19 @@ FReply FSequencer::OnRecord()
 
 	if(UMovieSceneSubSection::IsSetAsRecording() && !SequenceRecorder.IsRecording())
 	{
-		UWorld* PlaybackContext = Cast<UWorld>(GetPlaybackContext());
-		const FString& ActorToRecord = UMovieSceneSubSection::GetActorToRecord();
+		AActor* ActorToRecord = UMovieSceneSubSection::GetActorToRecord();
+		if (ActorToRecord != nullptr)
+		{
+			AActor* OutActor = EditorUtilities::GetSimWorldCounterpartActor(ActorToRecord);
+			if (OutActor != nullptr)
+			{
+				ActorToRecord = OutActor;
+			}
+		}
+
 		const FString& PathToRecordTo = UMovieSceneSubSection::GetRecordingSection()->GetTargetPathToRecordTo();
 		const FString& SequenceName = UMovieSceneSubSection::GetRecordingSection()->GetTargetSequenceName();
-		SequenceRecorder.StartRecording(PlaybackContext, ActorToRecord, FOnRecordingStarted::CreateSP(this, &FSequencer::HandleRecordingStarted), FOnRecordingFinished::CreateSP(this, &FSequencer::HandleRecordingFinished), PathToRecordTo, SequenceName);
+		SequenceRecorder.StartRecording(ActorToRecord, FOnRecordingStarted::CreateSP(this, &FSequencer::HandleRecordingStarted), FOnRecordingFinished::CreateSP(this, &FSequencer::HandleRecordingFinished), PathToRecordTo, SequenceName);
 	}
 	else if(SequenceRecorder.IsRecording())
 	{
@@ -2266,6 +2597,24 @@ void FSequencer::HandleRecordingFinished(UMovieSceneSequence* Sequence)
 		Section->SetAsRecording(false);
 		Section->SetSequence(Sequence);
 		Section->SetEndTime(Section->GetStartTime() + Sequence->GetMovieScene()->GetPlaybackRange().Size<float>());
+	
+		if (Section->IsA<UMovieSceneCinematicShotSection>())
+		{
+			const FMovieSceneSpawnable* SpawnedCamera = Sequence->GetMovieScene()->FindSpawnable(
+				[](FMovieSceneSpawnable& InSpawnable){
+					return InSpawnable.GetObjectTemplate() && InSpawnable.GetObjectTemplate()->IsA<ACameraActor>();
+				}
+			);
+
+			if (SpawnedCamera && !Sequence->GetMovieScene()->GetCameraCutTrack())
+			{
+				UMovieSceneTrack* CameraCutTrack = Sequence->GetMovieScene()->AddCameraCutTrack(UMovieSceneCameraCutTrack::StaticClass());
+				UMovieSceneCameraCutSection* CameraCutSection = Cast<UMovieSceneCameraCutSection>(CameraCutTrack->CreateNewSection());
+				CameraCutSection->SetCameraGuid(SpawnedCamera->GetGuid());
+				CameraCutSection->SetRange(Sequence->GetMovieScene()->GetPlaybackRange());
+				CameraCutTrack->AddSection(*CameraCutSection);
+			}
+		}
 	}
 
 	bNeedTreeRefresh = true;
@@ -2385,16 +2734,17 @@ bool FSequencer::IsLooping() const
 
 void FSequencer::SetGlobalTimeLooped(float InTime)
 {
-	bool bLooped = false;
+	bool bRestarted = false;
 	if (Settings->IsLooping())
 	{
 		const UMovieSceneSequence* FocusedSequence = GetFocusedMovieSceneSequence();
 		if (FocusedSequence)
 		{
-			if (InTime > FocusedSequence->GetMovieScene()->GetPlaybackRange().GetUpperBoundValue())
+			if (InTime <= FocusedSequence->GetMovieScene()->GetPlaybackRange().GetLowerBoundValue() ||
+				InTime >= FocusedSequence->GetMovieScene()->GetPlaybackRange().GetUpperBoundValue())
 			{
-				InTime = FocusedSequence->GetMovieScene()->GetPlaybackRange().GetLowerBoundValue();
-				bLooped = true;
+				InTime = PlayRate > 0 ? FocusedSequence->GetMovieScene()->GetPlaybackRange().GetLowerBoundValue() : FocusedSequence->GetMovieScene()->GetPlaybackRange().GetUpperBoundValue();
+				bRestarted = true;
 			}
 		}
 	}
@@ -2403,25 +2753,32 @@ void FSequencer::SetGlobalTimeLooped(float InTime)
 		TRange<float> TimeBounds = GetTimeBounds();
 		TRange<float> WorkingRange = GetClampRange();
 
-		const bool bReachedEnd = GetGlobalTime() < TimeBounds.GetUpperBoundValue() && InTime >= TimeBounds.GetUpperBoundValue();
+		bool bReachedEnd = false;
+		if (PlayRate > 0)
+		{
+			bReachedEnd = GetGlobalTime() < TimeBounds.GetUpperBoundValue() && InTime >= TimeBounds.GetUpperBoundValue();
+		}
+		else
+		{
+			bReachedEnd = GetGlobalTime() > TimeBounds.GetLowerBoundValue() && InTime <= TimeBounds.GetLowerBoundValue();
+		}
 
 		// Stop if we hit the playback range end
 		if (bReachedEnd)
 		{
-			InTime = TimeBounds.GetUpperBoundValue();
+			InTime = PlayRate > 0 ? TimeBounds.GetUpperBoundValue() : TimeBounds.GetLowerBoundValue();
 			SetPlaybackStatus(EMovieScenePlayerStatus::Stopped);
 		}
 		// Constrain to the play range if necessary
 		else if (Settings->ShouldKeepCursorInPlayRange())
 		{
-			// Clamp to lower bound
-			InTime = FMath::Max(InTime, TimeBounds.GetLowerBoundValue());
-
-			// Jump back if necessary
-			if (InTime >= TimeBounds.GetUpperBoundValue())
+			// Clamp to bound or jumpp back if necessary
+			if (InTime <= TimeBounds.GetLowerBoundValue() || 
+				InTime >= TimeBounds.GetUpperBoundValue())
 			{
-				InTime = TimeBounds.GetLowerBoundValue();
-			}
+				InTime = PlayRate > 0 ? TimeBounds.GetLowerBoundValue() : TimeBounds.GetUpperBoundValue();					
+				bRestarted = true;
+			}			
 		}
 		// Ensure the time is within the working range
 		else if (!WorkingRange.Contains(InTime))
@@ -2431,7 +2788,7 @@ void FSequencer::SetGlobalTimeLooped(float InTime)
 		}
 	}
 
-	SetGlobalTime(InTime, ESnapTimeMode::STM_None, bLooped);
+	SetGlobalTime(InTime, ESnapTimeMode::STM_None, bRestarted);
 }
 
 
@@ -2443,9 +2800,23 @@ bool FSequencer::CanShowFrameNumbers() const
 
 EPlaybackMode::Type FSequencer::GetPlaybackMode() const
 {
-	return PlaybackState == EMovieScenePlayerStatus::Playing ? EPlaybackMode::PlayingForward :
-		PlaybackState == EMovieScenePlayerStatus::Recording ? EPlaybackMode::Recording :
-		EPlaybackMode::Stopped;
+	if (PlaybackState == EMovieScenePlayerStatus::Playing)
+	{
+		if (PlayRate > 0)
+		{
+			return EPlaybackMode::PlayingForward;
+		}
+		else
+		{
+			return EPlaybackMode::PlayingReverse;
+		}
+	}
+	else if (PlaybackState == EMovieScenePlayerStatus::Recording)
+	{
+		return EPlaybackMode::Recording;
+	}
+		
+	return EPlaybackMode::Stopped;
 }
 
 
@@ -2753,7 +3124,7 @@ FGuid FSequencer::AddSpawnable(UObject& Object)
 	FGuid NewGuid = OwnerMovieScene->AddSpawnable(NewSpawnable.Name, *NewSpawnable.ObjectTemplate);
 
 	UpdateRuntimeInstances();
-
+	
 	return NewGuid;
 }
 
@@ -2928,7 +3299,7 @@ void GetParentTrackNodeAndNamePath(TSharedRef<const FSequencerDisplayNode> Displ
 }
 
 
-void FSequencer::OnRequestNodeDeleted( TSharedRef<const FSequencerDisplayNode> NodeToBeDeleted )
+bool FSequencer::OnRequestNodeDeleted( TSharedRef<const FSequencerDisplayNode> NodeToBeDeleted )
 {
 	bool bAnythingRemoved = false;
 	
@@ -2960,7 +3331,6 @@ void FSequencer::OnRequestNodeDeleted( TSharedRef<const FSequencerDisplayNode> N
 		}
 
 		bAnythingRemoved = true;
-
 	}
 	else if (NodeToBeDeleted->GetType() == ESequencerNode::Object)
 	{
@@ -3055,11 +3425,7 @@ void FSequencer::OnRequestNodeDeleted( TSharedRef<const FSequencerDisplayNode> N
 		}
 	}
 
-	if( bAnythingRemoved )
-	{
-		SequencerWidget->UpdateLayoutTree();
-		UpdateRuntimeInstances();
-	}
+	return bAnythingRemoved;
 }
 
 
@@ -3106,8 +3472,7 @@ void FSequencer::GetActorRecordingState( bool& bIsRecording /* In+Out */ ) const
 
 void FSequencer::PostUndo(bool bSuccess)
 {
-	SequencerWidget->UpdateLayoutTree();
-	UpdateRuntimeInstances();
+	NotifyMovieSceneDataChanged( EMovieSceneDataChangeType::Unknown );
 	OnActivateSequenceEvent.Broadcast( *SequenceInstanceStack.Top() );
 }
 
@@ -3206,7 +3571,8 @@ void FSequencer::OnNewActorsDropped(const TArray<UObject*>& DroppedObjects, cons
 				}
 
 				// Create a cine camera actor
-				ACineCameraActor* NewCamera = GWorld->SpawnActor<ACineCameraActor>();
+				UWorld* PlaybackContext = Cast<UWorld>(GetPlaybackContext());
+				ACineCameraActor* NewCamera = PlaybackContext->SpawnActor<ACineCameraActor>();
 				FGuid NewCameraGuid = CreateBinding(*NewCamera, NewCamera->GetActorLabel());
 
 				if (RailActor)
@@ -3304,13 +3670,11 @@ void FSequencer::OnNewActorsDropped(const TArray<UObject*>& DroppedObjects, cons
 			GEditor->NoteSelectionChange();
 		}
 
-		SequencerWidget->UpdateLayoutTree();
+		NotifyMovieSceneDataChanged( EMovieSceneDataChangeType::MovieSceneStructureItemsChanged );
 
 		bUpdatingSequencerSelection = true;
 		SynchronizeSequencerSelectionWithExternalSelection();
 		bUpdatingSequencerSelection = false;
-
-		NotifyMovieSceneDataChanged();
 	}
 }
 
@@ -3331,6 +3695,7 @@ void FSequencer::ActivateSequencerEditorMode()
 void FSequencer::OnPreBeginPIE(bool bIsSimulating)
 {
 	RootMovieSceneSequenceInstance->RestoreState(*this);
+	PrePossessionViewTargets.Reset();
 }
 
 
@@ -3398,12 +3763,12 @@ void FSequencer::UpdatePreviewLevelViewportClientFromCameraCut(FLevelEditorViewp
 	{
 		InViewportClient.SetViewLocation(CameraActor->GetActorLocation());
 		InViewportClient.SetViewRotation(CameraActor->GetActorRotation());
-		InViewportClient.bEditorCameraCut = !InViewportClient.IsLockedToActor(CameraActor) || bJumpCut;
+		InViewportClient.SetIsCameraCut(!InViewportClient.IsLockedToActor(CameraActor) || bJumpCut);
 	}
 	else
 	{
 		InViewportClient.ViewFOV = InViewportClient.FOVAngle;
-		InViewportClient.bEditorCameraCut = bJumpCut;
+		InViewportClient.SetIsCameraCut(bJumpCut);
 	}
 
 	// Set the actor lock.
@@ -3596,8 +3961,6 @@ void FSequencer::AddActors(const TArray<TWeakObjectPtr<AActor> >& InActors)
 		bUpdatingSequencerSelection = true;
 		SynchronizeSequencerSelectionWithExternalSelection();
 		bUpdatingSequencerSelection = false;
-
-		NotifyMovieSceneDataChanged();
 	}
 }
 
@@ -3721,6 +4084,11 @@ void FSequencer::SynchronizeSequencerSelectionWithExternalSelection()
 	TArray<TSharedRef<FSequencerObjectBindingNode>> RootObjectBindingNodes;
 	GetRootObjectBindingNodes( NodeTree->GetRootNodes(), RootObjectBindingNodes );
 
+	// If all nodes are already selected, do nothing. This ensures that when an undo event happens, 
+	// nodes are not cleared and reselected, which can cause issues with the curve editor auto-fitting 
+	// based on selection.
+	bool bAllAlreadySelected = true;
+
 	TSet<TSharedRef<FSequencerDisplayNode>> NodesToSelect;
 	for ( TSharedRef<FSequencerObjectBindingNode> ObjectBindingNode : RootObjectBindingNodes )
 	{
@@ -3733,18 +4101,46 @@ void FSequencer::SynchronizeSequencerSelectionWithExternalSelection()
 			if ( RuntimeObject != nullptr && GEditor->GetSelectedActors()->IsSelected( RuntimeObject ) )
 			{
 				NodesToSelect.Add( ObjectBindingNode );
+
+				if (bAllAlreadySelected)
+				{
+					bool bAlreadySelected = Selection.IsSelected(ObjectBindingNode);
+
+					if (!bAlreadySelected)
+					{
+						TSet<TSharedRef<FSequencerDisplayNode> > DescendantNodes;
+						SequencerHelpers::GetDescendantNodes(ObjectBindingNode, DescendantNodes);
+
+						for (auto DescendantNode : DescendantNodes)
+						{
+							if (Selection.IsSelected(DescendantNode) || Selection.NodeHasSelectedKeysOrSections(DescendantNode))
+							{
+								bAlreadySelected = true;
+								break;
+							}
+						}
+					}
+
+					if (!bAlreadySelected)
+					{
+						bAllAlreadySelected = false;
+					}
+				}
 			}
 		}
 	}
 
-	Selection.SuspendBroadcast();
-	Selection.EmptySelectedOutlinerNodes();
-	for ( TSharedRef<FSequencerDisplayNode> NodeToSelect : NodesToSelect)
+	if (!bAllAlreadySelected || NodesToSelect.Num() == 0)
 	{
-		Selection.AddToSelection( NodeToSelect );
+		Selection.SuspendBroadcast();
+		Selection.EmptySelectedOutlinerNodes();
+		for ( TSharedRef<FSequencerDisplayNode> NodeToSelect : NodesToSelect)
+		{
+			Selection.AddToSelection( NodeToSelect );
+		}
+		Selection.ResumeBroadcast();
+		Selection.GetOnOutlinerNodeSelectionChanged().Broadcast();
 	}
-	Selection.ResumeBroadcast();
-	Selection.GetOnOutlinerNodeSelectionChanged().Broadcast();
 }
 
 
@@ -4019,8 +4415,7 @@ void FSequencer::DoAssignActor(AActor*const* InActors, int32 NumActors, FGuid In
 
 	RootMovieSceneSequenceInstance->RestoreState(*this);
 
-	SequencerWidget->UpdateLayoutTree();
-	UpdateRuntimeInstances();
+	NotifyMovieSceneDataChanged( EMovieSceneDataChangeType::MovieSceneStructureItemsChanged );
 }
 
 void FSequencer::ImportFBX(FGuid InObjectBinding)
@@ -4028,7 +4423,7 @@ void FSequencer::ImportFBX(FGuid InObjectBinding)
 	UMovieScene* MovieScene = GetFocusedMovieSceneSequence()->GetMovieScene();
 	if (MovieSceneToolHelpers::ImportFBX(MovieScene, InObjectBinding))
 	{
-		NotifyMovieSceneDataChanged();
+		NotifyMovieSceneDataChanged( EMovieSceneDataChangeType::MovieSceneStructureItemAdded );
 	}
 }
 
@@ -4042,7 +4437,11 @@ void FSequencer::DeleteNode(TSharedRef<FSequencerDisplayNode> NodeToBeDeleted)
 	else
 	{
 		const FScopedTransaction Transaction( NSLOCTEXT("Sequencer", "UndoDeletingObject", "Delete Node") );
-		OnRequestNodeDeleted(NodeToBeDeleted);
+		bool bAnythingDeleted = OnRequestNodeDeleted(NodeToBeDeleted);
+		if ( bAnythingDeleted )
+		{
+			NotifyMovieSceneDataChanged( EMovieSceneDataChangeType::MovieSceneStructureItemRemoved );
+		}
 	}
 }
 
@@ -4058,15 +4457,22 @@ void FSequencer::DeleteSelectedNodes()
 
 	const FScopedTransaction Transaction( NSLOCTEXT("Sequencer", "UndoDeletingObject", "Delete Node") );
 
+	bool bAnythingDeleted = false;
+
 	for( const TSharedRef<FSequencerDisplayNode>& SelectedNode : SelectedNodesCopy )
 	{
 		if( !SelectedNode->IsHidden() )
 		{
 			// Delete everything in the entire node
 			TSharedRef<const FSequencerDisplayNode> NodeToBeDeleted = StaticCastSharedRef<const FSequencerDisplayNode>(SelectedNode);
-			OnRequestNodeDeleted( NodeToBeDeleted );
+			bAnythingDeleted |= OnRequestNodeDeleted( NodeToBeDeleted );
 		}
 	}
+
+	if ( bAnythingDeleted )
+	{
+		NotifyMovieSceneDataChanged( EMovieSceneDataChangeType::MovieSceneStructureItemRemoved );
+	}
 }
 
 
@@ -4159,7 +4565,7 @@ void FSequencer::ConvertToSpawnable(TSharedRef<FSequencerObjectBindingNode> Node
 	if (Possessable)
 	{
 		ConvertToSpawnableInternal(Possessable->GetGuid());
-		NotifyMovieSceneDataChanged();
+		NotifyMovieSceneDataChanged( EMovieSceneDataChangeType::MovieSceneStructureItemsChanged );
 	}
 }
 
@@ -4234,9 +4640,9 @@ void FSequencer::ConvertSelectedNodesToSpawnables()
 		}
 		GEditor->GetSelectedActors()->EndBatchSelectOperation();
 		GEditor->NoteSelectionChange();
-
-		NotifyMovieSceneDataChanged();
 	}
+
+	NotifyMovieSceneDataChanged( EMovieSceneDataChangeType::MovieSceneStructureItemsChanged );
 }
 
 FMovieSceneSpawnable* FSequencer::ConvertToSpawnableInternal(FGuid PossessableGuid)
@@ -4308,7 +4714,8 @@ FMovieSceneSpawnable* FSequencer::ConvertToSpawnableInternal(FGuid PossessableGu
 		DefaultTransform.bValid = true;
 
 		GEditor->SelectActor(OldActor, false, true);
-		GWorld->EditorDestroyActor(OldActor, true);
+		UWorld* PlaybackContext = Cast<UWorld>(GetPlaybackContext());
+		PlaybackContext->EditorDestroyActor(OldActor, true);
 	}
 
 	SetupDefaultsForSpawnable(Spawnable->GetGuid(), DefaultTransform);
@@ -4320,6 +4727,8 @@ FMovieSceneSpawnable* FSequencer::ConvertToSpawnableInternal(FGuid PossessableGu
 
 void FSequencer::OnAddFolder()
 {
+	FScopedTransaction AddFolderTransaction( NSLOCTEXT("Sequencer", "AddFolder_Transaction", "Add Folder") );
+
 	// Check if a folder, or child of a folder is currently selected.
 	TArray<UMovieSceneFolder*> SelectedParentFolders;
 	if ( Selection.GetSelectedOutlinerNodes().Num() > 0 )
@@ -4369,10 +4778,11 @@ void FSequencer::OnAddFolder()
 	}
 	else
 	{
+		FocusedMovieScene->Modify();
 		FocusedMovieScene->GetRootFolders().Add( NewFolder );
 	}
 
-	NotifyMovieSceneDataChanged();
+	NotifyMovieSceneDataChanged( EMovieSceneDataChangeType::MovieSceneStructureItemAdded );
 }
 
 
@@ -4393,6 +4803,49 @@ void FSequencer::Rewind()
 	OnStepToBeginning();
 }
 
+void FSequencer::ShuttleForward()
+{
+	float NewPlayRate = PlayRate;
+	if (ShuttleMultiplier == 0 || PlayRate < 0) 
+	{
+		ShuttleMultiplier = 2.f;
+		NewPlayRate = 1.f;
+	}
+	else
+	{
+		NewPlayRate *= ShuttleMultiplier;
+	}
+
+	OnPlay(false, NewPlayRate);
+}
+
+void FSequencer::ShuttleBackward()
+{
+	float NewPlayRate = PlayRate;
+	if (ShuttleMultiplier == 0 || PlayRate > 0)
+	{
+		ShuttleMultiplier = 2.f;
+		NewPlayRate = -1.f;
+	}
+	else
+	{
+		NewPlayRate *= ShuttleMultiplier;
+	}
+
+	OnPlay(false, NewPlayRate);
+}
+
+void FSequencer::Pause()
+{
+	SetPlaybackStatus(EMovieScenePlayerStatus::Stopped);
+
+	// Snap time when stopping play
+	SetGlobalTime(ScrubPosition, ESnapTimeMode::STM_Interval);
+
+	// Update on stop (cleans up things like sounds that are playing)
+	EMovieSceneUpdateData UpdateData(ScrubPosition, ScrubPosition);
+	SequenceInstanceStack.Top()->Update(UpdateData, *this );
+}
 
 void FSequencer::StepForward()
 {
@@ -4838,7 +5291,8 @@ void FSequencer::CreateCamera()
 		UObject* SpawnedCamera = FindSpawnedObjectOrTemplate(CameraGuid);
 		if (SpawnedCamera)
 		{
-			GWorld->EditorDestroyActor(NewCamera, true);
+			UWorld* PlaybackContext = Cast<UWorld>(GetPlaybackContext());
+			PlaybackContext->EditorDestroyActor(NewCamera, true);
 			NewCamera = Cast<ACineCameraActor>(SpawnedCamera);
 		}
 	}
@@ -4923,7 +5377,8 @@ void FSequencer::FixActorReferences()
 	UMovieScene* FocusedMovieScene = GetFocusedMovieSceneSequence()->GetMovieScene();
 
 	TMap<FString, AActor*> ActorNameToActorMap;
-	for ( TActorIterator<AActor> ActorItr( GWorld ); ActorItr; ++ActorItr )
+	UWorld* PlaybackContext = Cast<UWorld>(GetPlaybackContext());
+	for ( TActorIterator<AActor> ActorItr( PlaybackContext ); ActorItr; ++ActorItr )
 	{
 		// Same as with the Object Iterator, access the subclass instance with the * or -> operators.
 		AActor *Actor = *ActorItr;
@@ -5094,6 +5549,67 @@ void FSequencer::FixFrameTiming()
 }
 
 
+void FSequencer::ExportSceneAndSequence()
+{
+	void* ParentWindowWindowHandle = nullptr;
+	IMainFrameModule& MainFrameModule = FModuleManager::LoadModuleChecked<IMainFrameModule>( TEXT( "MainFrame" ) );
+	const TSharedPtr<SWindow>& MainFrameParentWindow = MainFrameModule.GetParentWindow();
+	if ( MainFrameParentWindow.IsValid() && MainFrameParentWindow->GetNativeWindow().IsValid() )
+	{
+		ParentWindowWindowHandle = MainFrameParentWindow->GetNativeWindow()->GetOSWindowHandle();
+	}
+
+	TArray<FString> SaveFilenames;
+	IDesktopPlatform* DesktopPlatform = FDesktopPlatformModule::Get();
+	bool bExportFileNamePicked = false;
+	if ( DesktopPlatform != NULL )
+	{
+		bExportFileNamePicked = DesktopPlatform->SaveFileDialog(
+			ParentWindowWindowHandle,
+			LOCTEXT( "ExportLevelSequence", "Export Level Sequence" ).ToString(),
+			*( FEditorDirectories::Get().GetLastDirectory( ELastDirectory::GENERIC_EXPORT ) ),
+			TEXT( "" ),
+			TEXT( "FBX document|*.fbx" ),
+			EFileDialogFlags::None,
+			SaveFilenames );
+	}
+
+	if ( bExportFileNamePicked )
+	{
+		FString ExportFilename = SaveFilenames[0];
+		FEditorDirectories::Get().SetLastDirectory( ELastDirectory::GENERIC_EXPORT, FPaths::GetPath( ExportFilename ) ); // Save path as default for next time.
+
+		UnFbx::FFbxExporter* Exporter = UnFbx::FFbxExporter::GetInstance();
+
+		Exporter->CreateDocument();
+		Exporter->SetTrasformBaking( true );
+		Exporter->SetKeepHierarchy( true );
+
+		const bool bSelectedOnly = false;
+
+		// Export the persistent level and all of it's actors
+		UWorld* World = Cast<UWorld>( GetPlaybackContext() );
+		Exporter->ExportLevelMesh( World->PersistentLevel, nullptr, bSelectedOnly );
+
+		// Export streaming levels and actors
+		for ( int32 CurLevelIndex = 0; CurLevelIndex < World->GetNumLevels(); ++CurLevelIndex )
+		{
+			ULevel* CurLevel = World->GetLevel( CurLevelIndex );
+			if ( CurLevel != NULL && CurLevel != ( World->PersistentLevel ) )
+			{
+				Exporter->ExportLevelMesh( CurLevel, nullptr, bSelectedOnly );
+			}
+		}
+
+		// Export the movie scene data.
+		Exporter->ExportLevelSequence( GetFocusedMovieSceneSequence()->GetMovieScene(), this );
+
+		// Save to disk
+		Exporter->WriteToFile( *ExportFilename );
+	}
+}
+
+
 void FSequencer::GenericTextEntryModeless(const FText& DialogText, const FText& DefaultText, FOnTextCommitted OnTextComitted)
 {
 	TSharedRef<STextEntryPopup> TextEntryPopup = 
@@ -5128,7 +5644,7 @@ void FSequencer::TrimSection(bool bTrimLeft)
 {
 	FScopedTransaction TrimSectionTransaction( NSLOCTEXT("Sequencer", "TrimSection_Transaction", "Trim Section") );
 	MovieSceneToolHelpers::TrimSection(Selection.GetSelectedSections(), GetGlobalTime(), bTrimLeft);
-	NotifyMovieSceneDataChanged();
+	NotifyMovieSceneDataChanged( EMovieSceneDataChangeType::TrackValueChanged );
 }
 
 
@@ -5136,7 +5652,7 @@ void FSequencer::SplitSection()
 {
 	FScopedTransaction SplitSectionTransaction( NSLOCTEXT("Sequencer", "SplitSection_Transaction", "Split Section") );
 	MovieSceneToolHelpers::SplitSection(Selection.GetSelectedSections(), GetGlobalTime());
-	NotifyMovieSceneDataChanged();
+	NotifyMovieSceneDataChanged( EMovieSceneDataChangeType::MovieSceneStructureItemAdded );
 }
 
 const ISequencerEditTool* FSequencer::GetEditTool() const
@@ -5455,7 +5971,7 @@ void FSequencer::BindCommands()
 
 	SequencerCommandBindings->MapAction(
 		Commands.ToggleKeepPlaybackRangeInSectionBounds,
-		FExecuteAction::CreateLambda( [this]{ Settings->SetKeepPlayRangeInSectionBounds( !Settings->ShouldKeepPlayRangeInSectionBounds() ); NotifyMovieSceneDataChanged(); } ),
+		FExecuteAction::CreateLambda( [this]{ Settings->SetKeepPlayRangeInSectionBounds( !Settings->ShouldKeepPlayRangeInSectionBounds() ); NotifyMovieSceneDataChanged( EMovieSceneDataChangeType::TrackValueChanged ); } ),
 		FCanExecuteAction::CreateLambda( []{ return true; } ),
 		FIsActionChecked::CreateLambda( [this]{ return Settings->ShouldKeepPlayRangeInSectionBounds(); } ) );
 
@@ -5495,6 +6011,11 @@ void FSequencer::BindCommands()
 		FExecuteAction::CreateSP( this, &FSequencer::FixFrameTiming ),
 		FCanExecuteAction::CreateLambda( [] { return true; } ) );
 
+	SequencerCommandBindings->MapAction(
+		Commands.ExportSceneAndSequence,
+		FExecuteAction::CreateSP( this, &FSequencer::ExportSceneAndSequence ),
+		FCanExecuteAction::CreateLambda( [] { return true; } ) );
+
 	for (int32 i = 0; i < TrackEditors.Num(); ++i)
 	{
 		TrackEditors[i]->BindCommands(SequencerCommandBindings);
@@ -5540,6 +6061,18 @@ void FSequencer::BindCommands()
 		Commands.Rewind,
 		FExecuteAction::CreateSP( this, &FSequencer::Rewind ));
 
+	SequencerCommandBindings->MapAction(
+		Commands.ShuttleForward,
+		FExecuteAction::CreateSP( this, &FSequencer::ShuttleForward ));
+
+	SequencerCommandBindings->MapAction(
+		Commands.ShuttleBackward,
+		FExecuteAction::CreateSP( this, &FSequencer::ShuttleBackward ));
+
+	SequencerCommandBindings->MapAction(
+		Commands.Pause,
+		FExecuteAction::CreateSP( this, &FSequencer::Pause ));
+
 	SequencerCommandBindings->MapAction(
 		Commands.StepForward,
 		FExecuteAction::CreateSP( this, &FSequencer::StepForward ),
@@ -5558,13 +6091,17 @@ void FSequencer::BindCommands()
 		Commands.SetSelectionRangeStart,
 		FExecuteAction::CreateLambda([this]{ SetSelectionRangeStart(); }));
 
+	SequencerCommandBindings->MapAction(
+		Commands.ResetSelectionRange,
+		FExecuteAction::CreateLambda([this]{ ResetSelectionRange(); }));
+
 	SequencerCommandBindings->MapAction(
 		Commands.SelectKeysInSelectionRange,
 		FExecuteAction::CreateLambda([this]{ SelectKeysInSelectionRange(); }));
 
 	SequencerCommandBindings->MapAction(
-		Commands.DeleteSelectionRange,
-		FExecuteAction::CreateLambda([this]{ DeleteSelectionRange(); }));
+		Commands.RecordSelectedActors,
+		FExecuteAction::CreateLambda([this]{ RecordSelectedActors(); }));
 
 	// bind widget specific commands
 	SequencerWidget->BindCommands(SequencerCommandBindings);
