Copying //UE4/Dev-Sequencer to //UE4/Dev-Main (Source: //UE4/Dev-Sequencer @ 3057646)

diff --git a/Engine/Source/Editor/MovieSceneTools/Private/TrackEditors/CinematicShotTrackEditor.cpp b/Engine/Source/Editor/MovieSceneTools/Private/TrackEditors/CinematicShotTrackEditor.cpp
--- a/Engine/Source/Editor/MovieSceneTools/Private/TrackEditors/CinematicShotTrackEditor.cpp
+++ b/Engine/Source/Editor/MovieSceneTools/Private/TrackEditors/CinematicShotTrackEditor.cpp
@@ -23,6 +23,9 @@
 #include "MovieSceneCaptureModule.h"
 #include "LevelSequence.h"
 #include "AssetToolsModule.h"
+#include "TrackEditorThumbnailPool.h"
+#include "MovieSceneToolsProjectSettings.h"
+
 
 #define LOCTEXT_NAMESPACE "FCinematicShotTrackEditor"
 
@@ -42,11 +45,13 @@ TSharedRef<ISequencerTrackEditor> FCinematicShotTrackEditor::CreateTrackEditor(T
 	return MakeShareable(new FCinematicShotTrackEditor(InSequencer));
 }
 
+
 void FCinematicShotTrackEditor::OnInitialize()
 {
 	OnCameraCutHandle = GetSequencer()->OnCameraCut().AddSP(this, &FCinematicShotTrackEditor::OnUpdateCameraCut);
 }
-	 
+
+
 void FCinematicShotTrackEditor::OnRelease()
 {
 	if (OnCameraCutHandle.IsValid() && GetSequencer().IsValid())
@@ -79,6 +84,7 @@ void FCinematicShotTrackEditor::BuildAddTrackMenu(FMenuBuilder& MenuBuilder)
 	);
 }
 
+
 TSharedPtr<SWidget> FCinematicShotTrackEditor::BuildOutlinerEditWidget(const FGuid& ObjectBinding, UMovieSceneTrack* Track, const FBuildEditWidgetParams& Params)
 {
 	// Create a container edit box
@@ -121,6 +127,7 @@ TSharedRef<ISequencerSection> FCinematicShotTrackEditor::MakeSectionInterface(UM
 	return MakeShareable(new FCinematicShotSection(GetSequencer(), ThumbnailPool, SectionObject, SharedThis(this)));
 }
 
+
 bool FCinematicShotTrackEditor::HandleAssetAdded(UObject* Asset, const FGuid& TargetObjectGuid)
 {
 	UMovieSceneSequence* Sequence = Cast<UMovieSceneSequence>(Asset);
@@ -149,6 +156,7 @@ bool FCinematicShotTrackEditor::HandleAssetAdded(UObject* Asset, const FGuid& Ta
 	return false;
 }
 
+
 bool FCinematicShotTrackEditor::SupportsType(TSubclassOf<UMovieSceneTrack> Type) const
 {
 	return (Type == UMovieSceneCinematicShotTrack::StaticClass());
@@ -180,6 +188,7 @@ void FCinematicShotTrackEditor::Tick(float DeltaTime)
 	}
 }
 
+
 void FCinematicShotTrackEditor::BuildTrackContextMenu( FMenuBuilder& MenuBuilder, UMovieSceneTrack* Track )
 {
 	MenuBuilder.AddMenuEntry(
@@ -197,11 +206,13 @@ void FCinematicShotTrackEditor::BuildTrackContextMenu( FMenuBuilder& MenuBuilder
 		FExecuteAction::CreateRaw(this, &FCinematicShotTrackEditor::ExportEDL )));
 }
 
+
 const FSlateBrush* FCinematicShotTrackEditor::GetIconBrush() const
 {
 	return FEditorStyle::GetBrush("Sequencer.Tracks.CinematicShot");
 }
 
+
 UMovieSceneSubSection* FCinematicShotTrackEditor::CreateShotInternal(FString& NewShotName, float NewShotStartTime, UMovieSceneCinematicShotSection* ShotToDuplicate, const bool& bInsertShot)
 {
 	FString NewShotPath = MovieSceneToolHelpers::GenerateNewShotPath(GetSequencer()->GetFocusedMovieSceneSequence()->GetMovieScene(), NewShotName);
@@ -316,6 +327,7 @@ void FCinematicShotTrackEditor::InsertFillerAtCurrentTime()
 	GetSequencer()->NotifyMovieSceneDataChanged( EMovieSceneDataChangeType::MovieSceneStructureItemAdded );
 }
 
+
 void FCinematicShotTrackEditor::DuplicateShot(UMovieSceneCinematicShotSection* Section)
 {
 	const FScopedTransaction Transaction(LOCTEXT("DuplicateShot_Transaction", "Duplicate Shot"));
@@ -338,6 +350,7 @@ void FCinematicShotTrackEditor::DuplicateShot(UMovieSceneCinematicShotSection* S
 	}
 }
 
+
 void FCinematicShotTrackEditor::RenderShot(UMovieSceneCinematicShotSection* Section)
 {
 	GetSequencer()->RenderMovie(Section);
@@ -452,6 +465,7 @@ void FCinematicShotTrackEditor::HandleAddCinematicShotTrackMenuEntryExecute()
 	FindOrCreateCinematicShotTrack();
 }
 
+
 TSharedRef<SWidget> FCinematicShotTrackEditor::HandleAddCinematicShotComboButtonGetMenuContent()
 {
 	FMenuBuilder MenuBuilder(true, nullptr);
@@ -507,6 +521,7 @@ void FCinematicShotTrackEditor::HandleAddCinematicShotComboButtonMenuEntryExecut
 	}
 }
 
+
 bool FCinematicShotTrackEditor::AddKeyInternal(float KeyTime, UMovieSceneSequence* InMovieSceneSequence)
 {	
 	if (CanAddSubSequence(*InMovieSceneSequence))
@@ -521,6 +536,7 @@ bool FCinematicShotTrackEditor::AddKeyInternal(float KeyTime, UMovieSceneSequenc
 	return false;
 }
 
+
 UMovieSceneCinematicShotTrack* FCinematicShotTrackEditor::FindOrCreateCinematicShotTrack()
 {
 	UMovieScene* FocusedMovieScene = GetFocusedMovieScene();
@@ -547,6 +563,7 @@ UMovieSceneCinematicShotTrack* FCinematicShotTrackEditor::FindOrCreateCinematicS
 	return NewTrack;
 }
 
+
 ECheckBoxState FCinematicShotTrackEditor::AreShotsLocked() const
 {
 	if (GetSequencer()->IsPerspectiveViewportCameraCutEnabled())
@@ -594,6 +611,7 @@ FText FCinematicShotTrackEditor::GetLockShotsToolTip() const
 		LOCTEXT("LockShots", "Lock Viewport to Shots");
 }
 
+
 bool FCinematicShotTrackEditor::CanAddSubSequence(const UMovieSceneSequence& Sequence) const
 {
 	// prevent adding ourselves and ensure we have a valid movie scene
@@ -624,12 +642,14 @@ bool FCinematicShotTrackEditor::CanAddSubSequence(const UMovieSceneSequence& Seq
 	return !SequenceCinematicShotTrack->ContainsSequence(*FocusedSequence, true);
 }
 
+
 void FCinematicShotTrackEditor::OnUpdateCameraCut(UObject* CameraObject, bool bJumpCut)
 {
 	// Keep track of the camera when it switches so that the thumbnail can be drawn with the correct camera
 	CinematicShotCamera = Cast<AActor>(CameraObject);
 }
 
+
 bool FCinematicShotTrackEditor::HandleSequenceAdded(float KeyTime, UMovieSceneSequence* Sequence)
 {
 	auto CinematicShotTrack = FindOrCreateCinematicShotTrack();
@@ -639,6 +659,7 @@ bool FCinematicShotTrackEditor::HandleSequenceAdded(float KeyTime, UMovieSceneSe
 	return true;
 }
 
+
 void FCinematicShotTrackEditor::ImportEDL()
 {
 	UMovieSceneSequence* FocusedSequence = GetSequencer()->GetFocusedMovieSceneSequence();
@@ -675,6 +696,7 @@ void FCinematicShotTrackEditor::ImportEDL()
 	}
 }
 
+
 void FCinematicShotTrackEditor::ExportEDL()
 {
 	UMovieSceneSequence* FocusedSequence = GetSequencer()->GetFocusedMovieSceneSequence();
@@ -708,4 +730,5 @@ void FCinematicShotTrackEditor::ExportEDL()
 	MovieSceneToolHelpers::ShowExportEDLDialog(MovieScene, FrameRate, SaveDirectory);
 }
 
+
 #undef LOCTEXT_NAMESPACE
diff --git a/Engine/Source/Editor/MovieSceneTools/Private/TrackEditors/PropertyTrackEditors/ActorReferencePropertyTrackEditor.h b/Engine/Source/Editor/MovieSceneTools/Private/TrackEditors/PropertyTrackEditors/ActorReferencePropertyTrackEditor.h
--- a/Engine/Source/Editor/MovieSceneTools/Private/TrackEditors/PropertyTrackEditors/ActorReferencePropertyTrackEditor.h
+++ b/Engine/Source/Editor/MovieSceneTools/Private/TrackEditors/PropertyTrackEditors/ActorReferencePropertyTrackEditor.h
@@ -4,6 +4,11 @@
 
 #include "MovieSceneActorReferenceTrack.h"
 #include "MovieSceneActorReferenceSection.h"
+#include "PropertyTrackEditor.h"
+
+
+class ISequencer;
+
 
 /**
  * A property track editor for actor references.
@@ -14,26 +19,26 @@ class FActorReferencePropertyTrackEditor
 public:
 
 	/**
-	 * Constructor
+	 * Constructor.
 	 *
-	 * @param InSequencer The sequencer instance to be used by this tool
+	 * @param InSequencer The sequencer instance to be used by this tool.
 	 */
-	FActorReferencePropertyTrackEditor( TSharedRef<ISequencer> InSequencer )
-		: FPropertyTrackEditor( InSequencer, NAME_ObjectProperty )
+	FActorReferencePropertyTrackEditor(TSharedRef<ISequencer> InSequencer)
+		: FPropertyTrackEditor(InSequencer, NAME_ObjectProperty)
 	{ }
 
 	/**
-	 * Creates an instance of this class.  Called by a sequencer 
+	 * Creates an instance of this class (called by a sequencer).
 	 *
-	 * @param OwningSequencer The sequencer instance to be used by this tool
-	 * @return The new instance of this class
+	 * @param OwningSequencer The sequencer instance to be used by this tool.
+	 * @return The new instance of this class.
 	 */
-	static TSharedRef<ISequencerTrackEditor> CreateTrackEditor( TSharedRef<ISequencer> OwningSequencer );
+	static TSharedRef<ISequencerTrackEditor> CreateTrackEditor(TSharedRef<ISequencer> OwningSequencer);
 
 protected:
 
-	// FPropertyTrackEditor interface
+	//~ FPropertyTrackEditor interface
 
-	virtual TSharedRef<FPropertySection> MakePropertySectionInterface( UMovieSceneSection& SectionObject, UMovieSceneTrack& Track ) override;
-	virtual void GenerateKeysFromPropertyChanged( const FPropertyChangedParams& PropertyChangedParams, TArray<FGuid>& NewGeneratedKeys, TArray<FGuid>& DefaultGeneratedKeys ) override;
+	virtual TSharedRef<FPropertySection> MakePropertySectionInterface(UMovieSceneSection& SectionObject, UMovieSceneTrack& Track) override;
+	virtual void GenerateKeysFromPropertyChanged(const FPropertyChangedParams& PropertyChangedParams, TArray<FGuid>& NewGeneratedKeys, TArray<FGuid>& DefaultGeneratedKeys) override;
 };
diff --git a/Engine/Source/Editor/MovieSceneTools/Private/TrackEditors/PropertyTrackEditors/BoolPropertyTrackEditor.h b/Engine/Source/Editor/MovieSceneTools/Private/TrackEditors/PropertyTrackEditors/BoolPropertyTrackEditor.h
--- a/Engine/Source/Editor/MovieSceneTools/Private/TrackEditors/PropertyTrackEditors/BoolPropertyTrackEditor.h
+++ b/Engine/Source/Editor/MovieSceneTools/Private/TrackEditors/PropertyTrackEditors/BoolPropertyTrackEditor.h
@@ -4,34 +4,41 @@
 
 #include "MovieSceneBoolSection.h"
 #include "MovieSceneBoolTrack.h"
+#include "PropertyTrackEditor.h"
+
+
+class ISequencer;
+
 
 /**
-* A property track editor for bools.
-*/
+ * A property track editor for Booleans.
+ */
 class FBoolPropertyTrackEditor
 	: public FPropertyTrackEditor<UMovieSceneBoolTrack, UMovieSceneBoolSection, bool>
 {
 public:
+
 	/**
-	 * Constructor
+	 * Constructor.
 	 *
-	 * @param InSequencer The sequencer instance to be used by this tool
+	 * @param InSequencer The sequencer instance to be used by this tool.
 	 */
 	FBoolPropertyTrackEditor( TSharedRef<ISequencer> InSequencer )
 		: FPropertyTrackEditor( InSequencer, NAME_BoolProperty )
 	{ }
 
 	/**
-	 * Creates an instance of this class.  Called by a sequencer 
+	 * Creates an instance of this class (called by a sequencer).
 	 *
-	 * @param OwningSequencer The sequencer instance to be used by this tool
-	 * @return The new instance of this class
+	 * @param OwningSequencer The sequencer instance to be used by this tool.
+	 * @return The new instance of this class.
 	 */
 	static TSharedRef<ISequencerTrackEditor> CreateTrackEditor( TSharedRef<ISequencer> OwningSequencer );
 
 protected:
 
-	// FPropertyTrackEditor interface
+	//~ FPropertyTrackEditor interface
+
 	virtual TSharedRef<FPropertySection> MakePropertySectionInterface( UMovieSceneSection& SectionObject, UMovieSceneTrack& Track ) override;
 	virtual void GenerateKeysFromPropertyChanged( const FPropertyChangedParams& PropertyChangedParams, TArray<bool>& NewGeneratedKeys, TArray<bool>& DefaultGeneratedKeys ) override;
 };
diff --git a/Engine/Source/Editor/MovieSceneTools/Private/TrackEditors/PropertyTrackEditors/BytePropertyTrackEditor.h b/Engine/Source/Editor/MovieSceneTools/Private/TrackEditors/PropertyTrackEditors/BytePropertyTrackEditor.h
--- a/Engine/Source/Editor/MovieSceneTools/Private/TrackEditors/PropertyTrackEditors/BytePropertyTrackEditor.h
+++ b/Engine/Source/Editor/MovieSceneTools/Private/TrackEditors/PropertyTrackEditors/BytePropertyTrackEditor.h
@@ -4,10 +4,14 @@
 
 #include "MovieSceneByteTrack.h"
 #include "MovieSceneByteSection.h"
+#include "PropertyTrackEditor.h"
+
+
+class ISequencer;
 
 
 /**
-* A property track editor for byte and enums.
+* A property track editor for byte and enumerations.
 */
 class FBytePropertyTrackEditor
 	: public FPropertyTrackEditor<UMovieSceneByteTrack, UMovieSceneByteSection, uint8>
@@ -17,30 +21,30 @@ public:
 	/**
 	 * Constructor.
 	 *
-	 * @param InSequencer The sequencer instance to be used by this tool
+	 * @param InSequencer The sequencer instance to be used by this tool.
 	 */
-	FBytePropertyTrackEditor( TSharedRef<ISequencer> InSequencer)
-		: FPropertyTrackEditor( InSequencer, NAME_ByteProperty )
+	FBytePropertyTrackEditor(TSharedRef<ISequencer> InSequencer)
+		: FPropertyTrackEditor(InSequencer, NAME_ByteProperty)
 	{ }
 
 	/**
-	 * Creates an instance of this class.  Called by a sequencer 
+	 * Creates an instance of this class (called by a sequencer).
 	 *
-	 * @param OwningSequencer The sequencer instance to be used by this tool
-	 * @return The new instance of this class
+	 * @param OwningSequencer The sequencer instance to be used by this tool.
+	 * @return The new instance of this class.
 	 */
-	static TSharedRef<ISequencerTrackEditor> CreateTrackEditor( TSharedRef<ISequencer> OwningSequencer );
+	static TSharedRef<ISequencerTrackEditor> CreateTrackEditor(TSharedRef<ISequencer> OwningSequencer);
 
 public:
 
-	// ISequencerTrackEditor interface
+	//~ ISequencerTrackEditor interface
 
-	virtual UMovieSceneTrack* AddTrack( UMovieScene* FocusedMovieScene, const FGuid& ObjectHandle, TSubclassOf<class UMovieSceneTrack> TrackClass, FName UniqueTypeName ) override;
+	virtual UMovieSceneTrack* AddTrack(UMovieScene* FocusedMovieScene, const FGuid& ObjectHandle, TSubclassOf<class UMovieSceneTrack> TrackClass, FName UniqueTypeName) override;
 
 protected:
 
-	// FPropertyTrackEditor interface
+	//~ FPropertyTrackEditor interface
 
-	virtual TSharedRef<FPropertySection> MakePropertySectionInterface( UMovieSceneSection& SectionObject, UMovieSceneTrack& Track ) override;
-	virtual void GenerateKeysFromPropertyChanged( const FPropertyChangedParams& PropertyChangedParams, TArray<uint8>& NewGeneratedKeys, TArray<uint8>& DefaultGeneratedKeys ) override;
+	virtual TSharedRef<FPropertySection> MakePropertySectionInterface(UMovieSceneSection& SectionObject, UMovieSceneTrack& Track) override;
+	virtual void GenerateKeysFromPropertyChanged(const FPropertyChangedParams& PropertyChangedParams, TArray<uint8>& NewGeneratedKeys, TArray<uint8>& DefaultGeneratedKeys) override;
 };
diff --git a/Engine/Source/Editor/MovieSceneTools/Private/TrackEditors/PropertyTrackEditors/ColorPropertyTrackEditor.h b/Engine/Source/Editor/MovieSceneTools/Private/TrackEditors/PropertyTrackEditors/ColorPropertyTrackEditor.h
--- a/Engine/Source/Editor/MovieSceneTools/Private/TrackEditors/PropertyTrackEditors/ColorPropertyTrackEditor.h
+++ b/Engine/Source/Editor/MovieSceneTools/Private/TrackEditors/PropertyTrackEditors/ColorPropertyTrackEditor.h
@@ -2,8 +2,12 @@
 
 #pragma once
 
-#include "MovieSceneColorTrack.h"
 #include "MovieSceneColorSection.h"
+#include "MovieSceneColorTrack.h"
+#include "PropertyTrackEditor.h"
+
+
+class ISequencer;
 
 
 /**
@@ -15,34 +19,37 @@ class FColorPropertyTrackEditor
 public:
 
 	/**
-	 * Constructor
+	 * Constructor.
 	 *
-	 * @param InSequencer The sequencer instance to be used by this tool
+	 * @param InSequencer The sequencer instance to be used by this tool.
 	 */
-	FColorPropertyTrackEditor( TSharedRef<ISequencer> InSequencer )
-		: FPropertyTrackEditor<UMovieSceneColorTrack, UMovieSceneColorSection, FColorKey>( InSequencer, NAME_Color, NAME_LinearColor, "SlateColor" )
+	FColorPropertyTrackEditor(TSharedRef<ISequencer> InSequencer)
+		: FPropertyTrackEditor<UMovieSceneColorTrack, UMovieSceneColorSection, FColorKey>(InSequencer, NAME_Color, NAME_LinearColor, "SlateColor")
 	{ }
 
 	/**
-	 * Creates an instance of this class.  Called by a sequencer 
+	 * Creates an instance of this class (called by a sequencer).
 	 *
-	 * @param OwningSequencer The sequencer instance to be used by this tool
-	 * @return The new instance of this class
+	 * @param OwningSequencer The sequencer instance to be used by this tool.
+	 * @return The new instance of this class.
 	 */
-	static TSharedRef<ISequencerTrackEditor> CreateTrackEditor( TSharedRef<ISequencer> OwningSequencer );
+	static TSharedRef<ISequencerTrackEditor> CreateTrackEditor(TSharedRef<ISequencer> OwningSequencer);
 
 protected:
 
-	// FPropertyTrackEditor interface
+	//~ FPropertyTrackEditor interface
 
-	virtual TSharedRef<FPropertySection> MakePropertySectionInterface( UMovieSceneSection& SectionObject, UMovieSceneTrack& Track ) override;
-	virtual void GenerateKeysFromPropertyChanged( const FPropertyChangedParams& PropertyChangedParams, TArray<FColorKey>& NewGeneratedKeys, TArray<FColorKey>& DefaultGeneratedKeys ) override;
+	virtual TSharedRef<FPropertySection> MakePropertySectionInterface(UMovieSceneSection& SectionObject, UMovieSceneTrack& Track) override;
+	virtual void GenerateKeysFromPropertyChanged(const FPropertyChangedParams& PropertyChangedParams, TArray<FColorKey>& NewGeneratedKeys, TArray<FColorKey>& DefaultGeneratedKeys) override;
 
-	// ISequencerTrackEditor interface
+protected:
+
+	//~ ISequencerTrackEditor interface
 
-	virtual void BuildTrackContextMenu( FMenuBuilder& MenuBuilder, UMovieSceneTrack* Track ) override;
+	virtual void BuildTrackContextMenu(FMenuBuilder& MenuBuilder, UMovieSceneTrack* Track) override;
 
 private:
+
 	static FName RedName;
 	static FName GreenName;
 	static FName BlueName;
diff --git a/Engine/Source/Editor/MovieSceneTools/Private/TrackEditors/PropertyTrackEditors/FloatPropertyTrackEditor.h b/Engine/Source/Editor/MovieSceneTools/Private/TrackEditors/PropertyTrackEditors/FloatPropertyTrackEditor.h
--- a/Engine/Source/Editor/MovieSceneTools/Private/TrackEditors/PropertyTrackEditors/FloatPropertyTrackEditor.h
+++ b/Engine/Source/Editor/MovieSceneTools/Private/TrackEditors/PropertyTrackEditors/FloatPropertyTrackEditor.h
@@ -4,6 +4,11 @@
 
 #include "MovieSceneFloatTrack.h"
 #include "MovieSceneFloatSection.h"
+#include "PropertyTrackEditor.h"
+
+
+class ISequencer;
+
 
 /**
  * A property track editor for floats.
@@ -14,29 +19,32 @@ class FFloatPropertyTrackEditor
 public:
 
 	/**
-	 * Constructor
+	 * Constructor.
 	 *
-	 * @param InSequencer The sequencer instance to be used by this tool
+	 * @param InSequencer The sequencer instance to be used by this tool.
 	 */
-	FFloatPropertyTrackEditor( TSharedRef<ISequencer> InSequencer )
-		: FPropertyTrackEditor( InSequencer, NAME_FloatProperty )
+	FFloatPropertyTrackEditor(TSharedRef<ISequencer> InSequencer)
+		: FPropertyTrackEditor(InSequencer, NAME_FloatProperty)
 	{ }
 
 	/**
-	 * Creates an instance of this class.  Called by a sequencer 
+	 * Creates an instance of this class (called by a sequencer).
 	 *
 	 * @param OwningSequencer The sequencer instance to be used by this tool
 	 * @return The new instance of this class
 	 */
-	static TSharedRef<ISequencerTrackEditor> CreateTrackEditor( TSharedRef<ISequencer> OwningSequencer );
+	static TSharedRef<ISequencerTrackEditor> CreateTrackEditor(TSharedRef<ISequencer> OwningSequencer);
+
+public:
+
+	//~ FMovieSceneTrackEditor interface
 
-	// FMovieSceneTrackEditor interface.
-	virtual void BuildTrackContextMenu( FMenuBuilder& MenuBuilder, UMovieSceneTrack* Track ) override;
+	virtual void BuildTrackContextMenu(FMenuBuilder& MenuBuilder, UMovieSceneTrack* Track) override;
 
 protected:
 
-	// FPropertyTrackEditor interface
+	//~ FPropertyTrackEditor interface
 
-	virtual TSharedRef<FPropertySection> MakePropertySectionInterface( UMovieSceneSection& SectionObject, UMovieSceneTrack& Track ) override;
-	virtual void GenerateKeysFromPropertyChanged( const FPropertyChangedParams& PropertyChangedParams, TArray<float>& NewGeneratedKeys, TArray<float>& DefaultGeneratedKeys ) override;
+	virtual TSharedRef<FPropertySection> MakePropertySectionInterface(UMovieSceneSection& SectionObject, UMovieSceneTrack& Track) override;
+	virtual void GenerateKeysFromPropertyChanged(const FPropertyChangedParams& PropertyChangedParams, TArray<float>& NewGeneratedKeys, TArray<float>& DefaultGeneratedKeys) override;
 };
diff --git a/Engine/Source/Editor/MovieSceneTools/Private/TrackEditors/PropertyTrackEditors/StringPropertyTrackEditor.cpp b/Engine/Source/Editor/MovieSceneTools/Private/TrackEditors/PropertyTrackEditors/StringPropertyTrackEditor.cpp
new file mode 100644
--- /dev/null
+++ b/Engine/Source/Editor/MovieSceneTools/Private/TrackEditors/PropertyTrackEditors/StringPropertyTrackEditor.cpp
@@ -0,0 +1,35 @@
+// Copyright 1998-2016 Epic Games, Inc. All Rights Reserved.
+
+#include "MovieSceneToolsPrivatePCH.h"
+#include "StringPropertyTrackEditor.h"
+#include "StringPropertySection.h"
+
+
+TSharedRef<ISequencerTrackEditor> FStringPropertyTrackEditor::CreateTrackEditor( TSharedRef<ISequencer> OwningSequencer )
+{
+	return MakeShareable(new FStringPropertyTrackEditor(OwningSequencer));
+}
+
+
+TSharedRef<FPropertySection> FStringPropertyTrackEditor::MakePropertySectionInterface( UMovieSceneSection& SectionObject, UMovieSceneTrack& Track )
+{
+	return MakeShareable(new FStringPropertySection(SectionObject, Track.GetDisplayName()));
+}
+
+
+void FStringPropertyTrackEditor::GenerateKeysFromPropertyChanged( const FPropertyChangedParams& PropertyChangedParams, TArray<FString>& NewGeneratedKeys, TArray<FString>& DefaultGeneratedKeys )
+{
+	void* CurrentObject = PropertyChangedParams.ObjectsThatChanged[0];
+	void* PropertyValue = nullptr;
+	for (int32 i = 0; i < PropertyChangedParams.PropertyPath.Num(); i++)
+	{
+		CurrentObject = PropertyChangedParams.PropertyPath[i]->ContainerPtrToValuePtr<FString>(CurrentObject, 0);
+	}
+
+	const UStrProperty* StrProperty = Cast<const UStrProperty>( PropertyChangedParams.PropertyPath.Last() );
+	if ( StrProperty )
+	{
+		FString StrPropertyValue = StrProperty->GetPropertyValue(CurrentObject);
+		NewGeneratedKeys.Add(StrPropertyValue);
+	}
+}
diff --git a/Engine/Source/Editor/MovieSceneTools/Private/TrackEditors/PropertyTrackEditors/StringPropertyTrackEditor.h b/Engine/Source/Editor/MovieSceneTools/Private/TrackEditors/PropertyTrackEditors/StringPropertyTrackEditor.h
new file mode 100644
--- /dev/null
+++ b/Engine/Source/Editor/MovieSceneTools/Private/TrackEditors/PropertyTrackEditors/StringPropertyTrackEditor.h
@@ -0,0 +1,44 @@
+// Copyright 1998-2016 Epic Games, Inc. All Rights Reserved.
+
+#pragma once
+
+#include "MovieSceneStringTrack.h"
+#include "MovieSceneStringSection.h"
+#include "PropertyTrackEditor.h"
+
+
+class ISequencer;
+
+
+/**
+ * A property track editor for strings.
+ */
+class FStringPropertyTrackEditor
+	: public FPropertyTrackEditor<UMovieSceneStringTrack, UMovieSceneStringSection, FString>
+{
+public:
+
+	/**
+	 * Constructor.
+	 *
+	 * @param InSequencer The sequencer instance to be used by this tool.
+	 */
+	FStringPropertyTrackEditor(TSharedRef<ISequencer> InSequencer)
+		: FPropertyTrackEditor(InSequencer, NAME_StrProperty)
+	{ }
+
+	/**
+	 * Creates an instance of this class (called by a sequencer).
+	 *
+	 * @param OwningSequencer The sequencer instance to be used by this tool
+	 * @return The new instance of this class
+	 */
+	static TSharedRef<ISequencerTrackEditor> CreateTrackEditor(TSharedRef<ISequencer> OwningSequencer);
+
+protected:
+
+	//~ FPropertyTrackEditor interface
+
+	virtual TSharedRef<FPropertySection> MakePropertySectionInterface(UMovieSceneSection& SectionObject, UMovieSceneTrack& Track) override;
+	virtual void GenerateKeysFromPropertyChanged(const FPropertyChangedParams& PropertyChangedParams, TArray<FString>& NewGeneratedKeys, TArray<FString>& DefaultGeneratedKeys) override;
+};
diff --git a/Engine/Source/Editor/MovieSceneTools/Private/TrackEditors/PropertyTrackEditors/VectorPropertyTrackEditor.h b/Engine/Source/Editor/MovieSceneTools/Private/TrackEditors/PropertyTrackEditors/VectorPropertyTrackEditor.h
--- a/Engine/Source/Editor/MovieSceneTools/Private/TrackEditors/PropertyTrackEditors/VectorPropertyTrackEditor.h
+++ b/Engine/Source/Editor/MovieSceneTools/Private/TrackEditors/PropertyTrackEditors/VectorPropertyTrackEditor.h
@@ -2,8 +2,12 @@
 
 #pragma once
 
-#include "MovieSceneVectorTrack.h"
 #include "MovieSceneVectorSection.h"
+#include "MovieSceneVectorTrack.h"
+#include "PropertyTrackEditor.h"
+
+
+class ISequencer;
 
 
 /**
@@ -15,30 +19,32 @@ class FVectorPropertyTrackEditor
 public:
 
 	/**
-	 * Constructor
+	 * Constructor.
 	 *
-	 * @param InSequencer The sequencer instance to be used by this tool
+	 * @param InSequencer The sequencer instance to be used by this tool.
 	 */
-	FVectorPropertyTrackEditor( TSharedRef<ISequencer> InSequencer )
-		: FPropertyTrackEditor<UMovieSceneVectorTrack, UMovieSceneVectorSection, FVectorKey>( InSequencer, NAME_Vector, NAME_Vector4, NAME_Vector2D )
+	FVectorPropertyTrackEditor(TSharedRef<ISequencer> InSequencer)
+		: FPropertyTrackEditor<UMovieSceneVectorTrack, UMovieSceneVectorSection, FVectorKey>(InSequencer, NAME_Vector, NAME_Vector4, NAME_Vector2D)
 	{ }
 
 	/**
-	 * Creates an instance of this class.  Called by a sequencer 
+	 * Creates an instance of this class (called by a sequence).
 	 *
-	 * @param OwningSequencer The sequencer instance to be used by this tool
-	 * @return The new instance of this class
+	 * @param OwningSequencer The sequencer instance to be used by this tool.
+	 * @return The new instance of this class.
 	 */
-	static TSharedRef<ISequencerTrackEditor> CreateTrackEditor( TSharedRef<ISequencer> OwningSequencer );
+	static TSharedRef<ISequencerTrackEditor> CreateTrackEditor(TSharedRef<ISequencer> OwningSequencer);
 
 protected:
 
-	// FPropertyTrackEditor interface
-	virtual TSharedRef<FPropertySection> MakePropertySectionInterface( UMovieSceneSection& SectionObject, UMovieSceneTrack& Track ) override;
-	virtual void GenerateKeysFromPropertyChanged( const FPropertyChangedParams& PropertyChangedParams, TArray<FVectorKey>& NewGeneratedKeys, TArray<FVectorKey>& DefaultGeneratedKeys ) override;
-	virtual void InitializeNewTrack( UMovieSceneVectorTrack* NewTrack, FPropertyChangedParams PropertyChangedParams ) override;
+	//~ FPropertyTrackEditor interface
+
+	virtual TSharedRef<FPropertySection> MakePropertySectionInterface(UMovieSceneSection& SectionObject, UMovieSceneTrack& Track) override;
+	virtual void GenerateKeysFromPropertyChanged(const FPropertyChangedParams& PropertyChangedParams, TArray<FVectorKey>& NewGeneratedKeys, TArray<FVectorKey>& DefaultGeneratedKeys) override;
+	virtual void InitializeNewTrack(UMovieSceneVectorTrack* NewTrack, FPropertyChangedParams PropertyChangedParams) override;
 
 private:
+
 	static FName XName;
 	static FName YName;
 	static FName ZName;
diff --git a/Engine/Source/Editor/MovieSceneTools/Private/TrackEditors/PropertyTrackEditors/VisibilityPropertyTrackEditor.h b/Engine/Source/Editor/MovieSceneTools/Private/TrackEditors/PropertyTrackEditors/VisibilityPropertyTrackEditor.h
--- a/Engine/Source/Editor/MovieSceneTools/Private/TrackEditors/PropertyTrackEditors/VisibilityPropertyTrackEditor.h
+++ b/Engine/Source/Editor/MovieSceneTools/Private/TrackEditors/PropertyTrackEditors/VisibilityPropertyTrackEditor.h
@@ -2,8 +2,12 @@
 
 #pragma once
 
-#include "MovieSceneVisibilityTrack.h"
 #include "MovieSceneVisibilitySection.h"
+#include "MovieSceneVisibilityTrack.h"
+#include "PropertyTrackEditor.h"
+
+
+class ISequencer;
 
 
 /**
@@ -16,9 +20,9 @@ class FVisibilityPropertyTrackEditor
 public:
 
 	/** Constructor. */
-	FVisibilityPropertyTrackEditor( TSharedRef<ISequencer> InSequencer)
+	FVisibilityPropertyTrackEditor(TSharedRef<ISequencer> InSequencer)
 		// Don't supply any property type names to watch since the FBoolPropertyTrackEditor is already watching for bool property changes.
-		: FPropertyTrackEditor<UMovieSceneVisibilityTrack, UMovieSceneVisibilitySection, bool>( InSequencer, NAME_BoolProperty )
+		: FPropertyTrackEditor<UMovieSceneVisibilityTrack, UMovieSceneVisibilitySection, bool>(InSequencer, NAME_BoolProperty)
 	{ }
 
 	/**
@@ -27,12 +31,13 @@ public:
 	 * @param OwningSequencer The sequencer instance to be used by this tool
 	 * @return The new instance of this class
 	 */
-	static TSharedRef<ISequencerTrackEditor> CreateTrackEditor( TSharedRef<ISequencer> OwningSequencer );
+	static TSharedRef<ISequencerTrackEditor> CreateTrackEditor(TSharedRef<ISequencer> OwningSequencer);
 
 protected:
 
-	// FPropertyTrackEditor interface
-	virtual TSharedRef<FPropertySection> MakePropertySectionInterface( UMovieSceneSection& SectionObject, UMovieSceneTrack& Track ) override;
-	virtual void GenerateKeysFromPropertyChanged( const FPropertyChangedParams& PropertyChangedParams, TArray<bool>& NewGeneratedKeys, TArray<bool>& DefaultGeneratedKeys ) override;
+	//~ FPropertyTrackEditor interface
+
+	virtual TSharedRef<FPropertySection> MakePropertySectionInterface(UMovieSceneSection& SectionObject, UMovieSceneTrack& Track) override;
+	virtual void GenerateKeysFromPropertyChanged(const FPropertyChangedParams& PropertyChangedParams, TArray<bool>& NewGeneratedKeys, TArray<bool>& DefaultGeneratedKeys) override;
 	virtual bool ForCustomizedUseOnly() { return true; }
 };
diff --git a/Engine/Source/Editor/MovieSceneTools/Public/BoolKeyArea.h b/Engine/Source/Editor/MovieSceneTools/Public/BoolKeyArea.h
--- a/Engine/Source/Editor/MovieSceneTools/Public/BoolKeyArea.h
+++ b/Engine/Source/Editor/MovieSceneTools/Public/BoolKeyArea.h
@@ -3,10 +3,13 @@
 #pragma once
 
 #include "IntegralKeyArea.h"
-#include "ClipboardTypes.h"
+
+
+class UMovieSceneSection;
+
 
 /**
- * A key area for displaying and editing integral curves representing bools.
+ * A key area for displaying and editing integral curves representing Booleans.
  */
 class FBoolKeyArea
 	: public FIntegralKeyArea<bool>
@@ -20,7 +23,7 @@ public:
 
 public:
 
-	// FIntegralKeyArea overrides
+	//~ FIntegralKeyArea interface
 
 	virtual bool CanCreateKeyEditor() override;
 	virtual TSharedRef<SWidget> CreateKeyEditor(ISequencer* Sequencer) override;
diff --git a/Engine/Source/Editor/MovieSceneTools/Public/ByteKeyArea.h b/Engine/Source/Editor/MovieSceneTools/Public/ByteKeyArea.h
--- a/Engine/Source/Editor/MovieSceneTools/Public/ByteKeyArea.h
+++ b/Engine/Source/Editor/MovieSceneTools/Public/ByteKeyArea.h
@@ -3,9 +3,14 @@
 #pragma once
 
 #include "IntegralKeyArea.h"
-#include "ClipboardTypes.h"
 
-/** A key area for displaying and editing byte curves. */
+
+class UMovieSceneSection;
+
+
+/**
+ * A key area for displaying and editing byte curves.
+ */
 class FByteKeyArea
 	: public FIntegralKeyArea<uint8>
 {
@@ -18,7 +23,7 @@ public:
 
 public:
 
-	// FIntegralKeyArea overrides
+	//~ FIntegralKeyArea interface
 
 	virtual bool CanCreateKeyEditor() override;
 	virtual TSharedRef<SWidget> CreateKeyEditor(ISequencer* Sequencer) override;
diff --git a/Engine/Source/Editor/MovieSceneTools/Public/ClipboardTypes.h b/Engine/Source/Editor/MovieSceneTools/Public/ClipboardTypes.h
--- a/Engine/Source/Editor/MovieSceneTools/Public/ClipboardTypes.h
+++ b/Engine/Source/Editor/MovieSceneTools/Public/ClipboardTypes.h
@@ -33,6 +33,11 @@ namespace MovieSceneClipboard
 		static FName Name("Bool");
 		return Name;
 	}
+	template<> inline FName GetKeyTypeName<FString>()
+	{
+		static FName Name("String");
+		return Name;
+	}
 }
 
 #endif
\ No newline at end of file
diff --git a/Engine/Source/Editor/MovieSceneTools/Public/EnumKeyArea.h b/Engine/Source/Editor/MovieSceneTools/Public/EnumKeyArea.h
--- a/Engine/Source/Editor/MovieSceneTools/Public/EnumKeyArea.h
+++ b/Engine/Source/Editor/MovieSceneTools/Public/EnumKeyArea.h
@@ -5,6 +5,12 @@
 #include "ByteKeyArea.h"
 
 
+struct FIntegralCurve;
+class SWidget;
+class UEnum;
+class UMovieSceneSection;
+
+
 /**
  * A key area for displaying and editing integral curves representing enums.
  */
@@ -21,7 +27,7 @@ public:
 
 public:
 
-	// FByteKeyArea overrides
+	//~ FByteKeyArea interface
 
 	virtual bool CanCreateKeyEditor() override;
 	virtual TSharedRef<SWidget> CreateKeyEditor(ISequencer* Sequencer) override;
diff --git a/Engine/Source/Editor/MovieSceneTools/Public/FloatCurveKeyArea.h b/Engine/Source/Editor/MovieSceneTools/Public/FloatCurveKeyArea.h
--- a/Engine/Source/Editor/MovieSceneTools/Public/FloatCurveKeyArea.h
+++ b/Engine/Source/Editor/MovieSceneTools/Public/FloatCurveKeyArea.h
@@ -6,8 +6,12 @@
 #include "ClipboardTypes.h"
 
 
+struct FRichCurve;
+class UMovieSceneSection;
+
+
 /**
- * A key area for float keys
+ * A key area for float keys.
  */
 class MOVIESCENETOOLS_API FFloatCurveKeyArea
 	: public FNamedKeyArea
@@ -34,7 +38,7 @@ public:
 
 public:
 
-	// IKeyArea interface
+	//` IKeyArea interface
 
 	virtual TArray<FKeyHandle> AddKeyUnique(float Time, EMovieSceneKeyInterpolation InKeyInterpolation, float TimeToCopyFrom = FLT_MAX) override;
 	virtual TOptional<FKeyHandle> DuplicateKey(FKeyHandle KeyToDuplicate) override;
diff --git a/Engine/Source/Editor/MovieSceneTools/Public/IntegralKeyArea.h b/Engine/Source/Editor/MovieSceneTools/Public/IntegralKeyArea.h
--- a/Engine/Source/Editor/MovieSceneTools/Public/IntegralKeyArea.h
+++ b/Engine/Source/Editor/MovieSceneTools/Public/IntegralKeyArea.h
@@ -2,10 +2,15 @@
 
 #pragma once
 
+#include "ClipboardTypes.h"
 #include "NamedKeyArea.h"
 #include "MovieSceneClipboard.h"
 #include "SequencerClipboardReconciler.h"
 
+
+struct FIntegralCurve;
+
+
 /**
  * Abstract base class for integral curve key areas.
  */
@@ -22,7 +27,7 @@ public:
 
 public:
 
-	// IKeyArea interface
+	//~ IKeyArea interface
 
 	virtual TArray<FKeyHandle> AddKeyUnique( float Time, EMovieSceneKeyInterpolation InKeyInterpolation, float TimeToCopyFrom = FLT_MAX ) override;
 	virtual TOptional<FKeyHandle> DuplicateKey(FKeyHandle KeyToDuplicate) override;
diff --git a/Engine/Source/Editor/MovieSceneTools/Public/NameCurveKeyArea.h b/Engine/Source/Editor/MovieSceneTools/Public/NameCurveKeyArea.h
--- a/Engine/Source/Editor/MovieSceneTools/Public/NameCurveKeyArea.h
+++ b/Engine/Source/Editor/MovieSceneTools/Public/NameCurveKeyArea.h
@@ -5,6 +5,11 @@
 #include "NamedKeyArea.h"
 #include "ClipboardTypes.h"
 
+
+struct FNameCurve;
+class UMovieSceneSection;
+
+
 /**
  * A key area for FName curves.
  */
@@ -26,7 +31,7 @@ public:
 
 public:
 
-	// IKeyArea interface
+	//~ IKeyArea interface
 
 	virtual TArray<FKeyHandle> AddKeyUnique(float Time, EMovieSceneKeyInterpolation InKeyInterpolation, float TimeToCopyFrom = FLT_MAX) override;
 	virtual TOptional<FKeyHandle> DuplicateKey(FKeyHandle KeyToDuplicate) override;
diff --git a/Engine/Source/Editor/MovieSceneTools/Public/StringCurveKeyArea.h b/Engine/Source/Editor/MovieSceneTools/Public/StringCurveKeyArea.h
new file mode 100644
--- /dev/null
+++ b/Engine/Source/Editor/MovieSceneTools/Public/StringCurveKeyArea.h
@@ -0,0 +1,82 @@
+// Copyright 1998-2016 Epic Games, Inc. All Rights Reserved.
+
+#pragma once
+
+#include "NamedKeyArea.h"
+#include "ClipboardTypes.h"
+
+
+struct FStringCurve;
+class UMovieSceneSection;
+
+
+/**
+ * A key area for string keys.
+ */
+class MOVIESCENETOOLS_API FStringCurveKeyArea
+	: public FNamedKeyArea
+{
+public:
+
+	FStringCurveKeyArea(FStringCurve* InCurve, UMovieSceneSection* InOwningSection, TOptional<FLinearColor> InColor = TOptional<FLinearColor>())
+		: Color(InColor)
+		, Curve(InCurve)
+		, OwningSection(InOwningSection)
+	{ }
+
+public:
+
+	void ClearIntermediateValue()
+	{
+		IntermediateValue.Reset();
+	}
+
+	void SetIntermediateValue(FString InIntermediateValue)
+	{
+		IntermediateValue = InIntermediateValue;
+	}
+
+public:
+
+	//` IKeyArea interface
+
+	virtual TArray<FKeyHandle> AddKeyUnique(float Time, EMovieSceneKeyInterpolation InKeyInterpolation, float TimeToCopyFrom = FLT_MAX) override;
+	virtual TOptional<FKeyHandle> DuplicateKey(FKeyHandle KeyToDuplicate) override;
+	virtual bool CanCreateKeyEditor() override;
+	virtual TSharedRef<SWidget> CreateKeyEditor(ISequencer* Sequencer) override;
+	virtual void DeleteKey(FKeyHandle KeyHandle) override;
+	virtual TOptional<FLinearColor> GetColor() override;
+	virtual ERichCurveExtrapolation GetExtrapolationMode(bool bPreInfinity) const override;
+	virtual ERichCurveInterpMode GetKeyInterpMode(FKeyHandle KeyHandle) const override;
+	virtual TSharedPtr<FStructOnScope> GetKeyStruct(FKeyHandle KeyHandle) override;
+	virtual ERichCurveTangentMode GetKeyTangentMode(FKeyHandle KeyHandle) const override;
+	virtual float GetKeyTime(FKeyHandle KeyHandle) const override;
+	virtual UMovieSceneSection* GetOwningSection() override;
+	virtual FRichCurve* GetRichCurve() override;
+	virtual TArray<FKeyHandle> GetUnsortedKeyHandles() const override;
+	virtual FKeyHandle MoveKey(FKeyHandle KeyHandle, float DeltaPosition) override;
+	virtual void SetExtrapolationMode(ERichCurveExtrapolation ExtrapMode, bool bPreInfinity) override;
+	virtual bool CanSetExtrapolationMode() const override;
+	virtual void SetKeyInterpMode(FKeyHandle KeyHandle, ERichCurveInterpMode InterpMode) override;
+	virtual void SetKeyTangentMode(FKeyHandle KeyHandle, ERichCurveTangentMode TangentMode) override;
+	virtual void SetKeyTime(FKeyHandle KeyHandle, float NewKeyTime) const override;
+	virtual void CopyKeys(FMovieSceneClipboardBuilder& ClipboardBuilder, const TFunctionRef<bool(FKeyHandle, const IKeyArea&)>& KeyMask) const override;
+	virtual void PasteKeys(const FMovieSceneClipboardKeyTrack& KeyTrack, const FMovieSceneClipboardEnvironment& SrcEnvironment, const FSequencerPasteEnvironment& DstEnvironment) override;
+
+protected:
+
+	void OnValueChanged(FString InValue);
+
+private:
+
+	/** The key area's color. */
+	TOptional<FLinearColor> Color;
+
+	/** The curve which provides the keys for this key area. */
+	FStringCurve* Curve;
+
+	/** The section that owns this key area. */
+	UMovieSceneSection* OwningSection;
+
+	TOptional<FString> IntermediateValue;
+};
diff --git a/Engine/Source/Editor/SequenceRecorder/Private/ActorRecording.cpp b/Engine/Source/Editor/SequenceRecorder/Private/ActorRecording.cpp
--- a/Engine/Source/Editor/SequenceRecorder/Private/ActorRecording.cpp
+++ b/Engine/Source/Editor/SequenceRecorder/Private/ActorRecording.cpp
@@ -52,9 +52,9 @@ bool UActorRecording::IsRelevantForRecording(AActor* Actor)
 	TInlineComponentArray<USceneComponent*> SceneComponents(Actor);
 	for(USceneComponent* SceneComponent : SceneComponents)
 	{
-		for (TSubclassOf<USceneComponent> ComponentClass : Settings->ComponentClassesToRecord)
+		for (const FPropertiesToRecordForClass& PropertiesToRecordForClass : Settings->ClassesAndPropertiesToRecord)
 		{
-			if (SceneComponent->IsA(ComponentClass))
+			if (SceneComponent->IsA(PropertiesToRecordForClass.Class))
 			{
 				return true;
 			}
@@ -207,9 +207,9 @@ bool UActorRecording::ValidComponent(USceneComponent* SceneComponent) const
 	if(SceneComponent != nullptr)
 	{
 		const USequenceRecorderSettings* Settings = GetDefault<USequenceRecorderSettings>();
-		for (TSubclassOf<USceneComponent> ComponentClass : Settings->ComponentClassesToRecord)
+		for (const FPropertiesToRecordForClass& PropertiesToRecordForClass : Settings->ClassesAndPropertiesToRecord)
 		{			
-			if (ComponentClass != nullptr && SceneComponent->IsA(ComponentClass))
+			if (PropertiesToRecordForClass.Class != nullptr && SceneComponent->IsA(PropertiesToRecordForClass.Class))
 			{
 				return true;
 			}
@@ -326,7 +326,7 @@ void UActorRecording::StartRecordingActorProperties(ULevelSequence* CurrentSeque
 			for(USceneComponent* SceneComponent : ValidSceneComponents)
 			{
 				TSharedPtr<FMovieSceneAnimationSectionRecorder> AnimRecorder = StartRecordingComponentProperties(SceneComponent->GetFName(), SceneComponent, GetActorToRecord(), CurrentSequence, CurrentSequenceTime);
-				if(!FirstAnimRecorder.IsValid() && AnimRecorder.IsValid())
+				if(!FirstAnimRecorder.IsValid() && AnimRecorder.IsValid() && GetActorToRecord()->IsA<ACharacter>())
 				{
 					FirstAnimRecorder = AnimRecorder;
 				}
diff --git a/Engine/Source/Editor/SequenceRecorder/Private/PropertiesToRecordForClassDetailsCustomization.cpp b/Engine/Source/Editor/SequenceRecorder/Private/PropertiesToRecordForClassDetailsCustomization.cpp
--- a/Engine/Source/Editor/SequenceRecorder/Private/PropertiesToRecordForClassDetailsCustomization.cpp
+++ b/Engine/Source/Editor/SequenceRecorder/Private/PropertiesToRecordForClassDetailsCustomization.cpp
@@ -11,11 +11,6 @@ void FPropertiesToRecordForClassDetailsCustomization::CustomizeHeader(TSharedRef
 	NameContent()
 	[
 		PropertyHandle->CreatePropertyNameWidget()
-	]
-	.ValueContent()
-	.MaxDesiredWidth(250.f)
-	[
-		PropertyHandle->CreatePropertyValueWidget()
 	];
 }
 
diff --git a/Engine/Source/Editor/SequenceRecorder/Private/SSequenceRecorder.cpp b/Engine/Source/Editor/SequenceRecorder/Private/SSequenceRecorder.cpp
--- a/Engine/Source/Editor/SequenceRecorder/Private/SSequenceRecorder.cpp
+++ b/Engine/Source/Editor/SequenceRecorder/Private/SSequenceRecorder.cpp
@@ -164,28 +164,39 @@ void SSequenceRecorder::Construct(const FArguments& Args)
 					SNew(SOverlay)
 					+SOverlay::Slot()
 					[
-						SAssignNew(ListView, SListView<UActorRecording*>)
-						.ListItemsSource(&FSequenceRecorder::Get().GetQueuedRecordings())
-						.SelectionMode(ESelectionMode::SingleToggle)
-						.OnGenerateRow(this, &SSequenceRecorder::MakeListViewWidget)
-						.OnSelectionChanged(this, &SSequenceRecorder::OnSelectionChanged)
-						.HeaderRow
-						(
-						SNew(SHeaderRow)
-						+ SHeaderRow::Column(ActorColumnName)
-						.FillWidth(43.0f)
-						.DefaultLabel(LOCTEXT("ActorHeaderName", "Actor"))
-						+ SHeaderRow::Column(AnimationColumnName)
-						.FillWidth(43.0f)
-						.DefaultLabel(LOCTEXT("AnimationHeaderName", "Animation"))
-						+ SHeaderRow::Column(LengthColumnName)
-						.FillWidth(14.0f)
-						.DefaultLabel(LOCTEXT("LengthHeaderName", "Length"))
-						)
+						SNew(SVerticalBox)
+						+SVerticalBox::Slot()
+						.FillHeight(1.0f)
+						[
+							SAssignNew(ListView, SListView<UActorRecording*>)
+							.ListItemsSource(&FSequenceRecorder::Get().GetQueuedRecordings())
+							.SelectionMode(ESelectionMode::SingleToggle)
+							.OnGenerateRow(this, &SSequenceRecorder::MakeListViewWidget)
+							.OnSelectionChanged(this, &SSequenceRecorder::OnSelectionChanged)
+							.HeaderRow
+							(
+								SNew(SHeaderRow)
+								+ SHeaderRow::Column(ActorColumnName)
+								.FillWidth(43.0f)
+								.DefaultLabel(LOCTEXT("ActorHeaderName", "Actor"))
+								+ SHeaderRow::Column(AnimationColumnName)
+								.FillWidth(43.0f)
+								.DefaultLabel(LOCTEXT("AnimationHeaderName", "Animation"))
+								+ SHeaderRow::Column(LengthColumnName)
+								.FillWidth(14.0f)
+								.DefaultLabel(LOCTEXT("LengthHeaderName", "Length"))
+							)
+						]
+						+SVerticalBox::Slot()
+						.AutoHeight()
+						[
+							SNew(STextBlock)
+							.Text(this, &SSequenceRecorder::GetTargetSequenceName)
+						]
 					]
 					+SOverlay::Slot()
 					[
-						SNew(SVerticalBox)
+						SNew(SVerticalBox)	
 						+SVerticalBox::Slot()
 						.VAlign(VAlign_Bottom)
 						.MaxHeight(2.0f)
@@ -285,7 +296,7 @@ bool SSequenceRecorder::CanRecord() const
 
 bool SSequenceRecorder::IsRecordVisible() const
 {
-	return !FAnimationRecorderManager::Get().IsRecording() && !FSequenceRecorder::Get().IsDelaying();
+	return !FSequenceRecorder::Get().IsRecording() && !FAnimationRecorderManager::Get().IsRecording() && !FSequenceRecorder::Get().IsDelaying();
 }
 
 void SSequenceRecorder::HandleStopAll()
@@ -295,12 +306,12 @@ void SSequenceRecorder::HandleStopAll()
 
 bool SSequenceRecorder::CanStopAll() const
 {
-	return FAnimationRecorderManager::Get().IsRecording() || FSequenceRecorder::Get().IsDelaying();
+	return FSequenceRecorder::Get().IsRecording() || FAnimationRecorderManager::Get().IsRecording() || FSequenceRecorder::Get().IsDelaying();
 }
 
 bool SSequenceRecorder::IsStopAllVisible() const
 {
-	return FAnimationRecorderManager::Get().IsRecording() || FSequenceRecorder::Get().IsDelaying();
+	return FSequenceRecorder::Get().IsRecording() || FAnimationRecorderManager::Get().IsRecording() || FSequenceRecorder::Get().IsDelaying();
 }
 
 void SSequenceRecorder::HandleAddRecording()
@@ -359,4 +370,9 @@ EVisibility SSequenceRecorder::GetDelayProgressVisibilty() const
 	return FSequenceRecorder::Get().IsDelaying() ? EVisibility::Visible : EVisibility::Hidden;
 }
 
+FText SSequenceRecorder::GetTargetSequenceName() const
+{
+	return FText::Format(LOCTEXT("NextSequenceFormat", "Next Sequence: {0}"), FText::FromString(FSequenceRecorder::Get().GetNextSequenceName()));
+}
+
 #undef LOCTEXT_NAMESPACE
\ No newline at end of file
diff --git a/Engine/Source/Editor/SequenceRecorder/Private/SSequenceRecorder.h b/Engine/Source/Editor/SequenceRecorder/Private/SSequenceRecorder.h
--- a/Engine/Source/Editor/SequenceRecorder/Private/SSequenceRecorder.h
+++ b/Engine/Source/Editor/SequenceRecorder/Private/SSequenceRecorder.h
@@ -54,6 +54,8 @@ private:
 
 	EVisibility GetDelayProgressVisibilty() const;
 
+	FText GetTargetSequenceName() const;
+
 private:
 	TSharedPtr<IDetailsView> SequenceRecordingDetailsView;
 
diff --git a/Engine/Source/Editor/SequenceRecorder/Private/Sections/MovieSceneMultiPropertyRecorder.cpp b/Engine/Source/Editor/SequenceRecorder/Private/Sections/MovieSceneMultiPropertyRecorder.cpp
--- a/Engine/Source/Editor/SequenceRecorder/Private/Sections/MovieSceneMultiPropertyRecorder.cpp
+++ b/Engine/Source/Editor/SequenceRecorder/Private/Sections/MovieSceneMultiPropertyRecorder.cpp
@@ -14,7 +14,7 @@ bool FMovieSceneMultiPropertyRecorderFactory::CanRecordObject(UObject* InObjectT
 {
 	const USequenceRecorderSettings* Settings = GetDefault<USequenceRecorderSettings>();
 
-	for (const FPropertiesToRecordForClass& PropertiesToRecordForClass : Settings->PropertiesToRecord)
+	for (const FPropertiesToRecordForClass& PropertiesToRecordForClass : Settings->ClassesAndPropertiesToRecord)
 	{
 		if (*PropertiesToRecordForClass.Class == InObjectToRecord->GetClass() && PropertiesToRecordForClass.Properties.Num() > 0)
 		{
@@ -33,21 +33,21 @@ void FMovieSceneMultiPropertyRecorder::CreateSection(UObject* InObjectToRecord,
 {
 	ObjectToRecord = InObjectToRecord;
 
+	// collect all properties to record from classes we are recording
+	TArray<FName> PropertiesToRecord;
+
 	const FPropertiesToRecordForClass* PropertiesToRecordForClassPtr = nullptr;
 	const USequenceRecorderSettings* Settings = GetDefault<USequenceRecorderSettings>();
-	for (const FPropertiesToRecordForClass& PropertiesToRecordForClass : Settings->PropertiesToRecord)
+	for (const FPropertiesToRecordForClass& PropertiesToRecordForClass : Settings->ClassesAndPropertiesToRecord)
 	{
 		if (*PropertiesToRecordForClass.Class == InObjectToRecord->GetClass() && PropertiesToRecordForClass.Properties.Num() > 0)
 		{
-			PropertiesToRecordForClassPtr = &PropertiesToRecordForClass;
-			break;
+			PropertiesToRecord.Append(PropertiesToRecordForClass.Properties);
 		}
 	}
 
-	check(PropertiesToRecordForClassPtr);
-
 	// create a recorder for each property name
-	for (const FName& PropertyName : PropertiesToRecordForClassPtr->Properties)
+	for (const FName& PropertyName : PropertiesToRecord)
 	{
 		FTrackInstancePropertyBindings Binding(PropertyName, PropertyName.ToString());
 		Binding.UpdateBinding(InObjectToRecord);
diff --git a/Engine/Source/Editor/SequenceRecorder/Private/Sections/MovieScenePropertyRecorder.cpp b/Engine/Source/Editor/SequenceRecorder/Private/Sections/MovieScenePropertyRecorder.cpp
--- a/Engine/Source/Editor/SequenceRecorder/Private/Sections/MovieScenePropertyRecorder.cpp
+++ b/Engine/Source/Editor/SequenceRecorder/Private/Sections/MovieScenePropertyRecorder.cpp
@@ -22,12 +22,12 @@ bool FMovieScenePropertyRecorder<bool>::ShouldAddNewKey(const bool& InNewValue)
 }
 
 template <>
-UMovieSceneSection* FMovieScenePropertyRecorder<bool>::AddSection(UMovieScene* InMovieScene, const FGuid& InGuid, float InTime)
+UMovieSceneSection* FMovieScenePropertyRecorder<bool>::AddSection(UObject* InObjectToRecord, UMovieScene* InMovieScene, const FGuid& InGuid, float InTime)
 {
 	UMovieSceneBoolTrack* Track = InMovieScene->AddTrack<UMovieSceneBoolTrack>(InGuid);
 	if (Track)
 	{
-		Track->SetPropertyNameAndPath(Binding.GetPropertyName(), Binding.GetPropertyPath());
+		Track->SetPropertyNameAndPath(*Binding.GetProperty(InObjectToRecord)->GetDisplayNameText().ToString(), Binding.GetPropertyPath());
 
 		UMovieSceneBoolSection* Section = Cast<UMovieSceneBoolSection>(Track->CreateNewSection());
 		Section->SetDefault(PreviousValue);
@@ -61,12 +61,12 @@ bool FMovieScenePropertyRecorder<uint8>::ShouldAddNewKey(const uint8& InNewValue
 }
 
 template <>
-UMovieSceneSection* FMovieScenePropertyRecorder<uint8>::AddSection(UMovieScene* InMovieScene, const FGuid& InGuid, float InTime)
+UMovieSceneSection* FMovieScenePropertyRecorder<uint8>::AddSection(UObject* InObjectToRecord, UMovieScene* InMovieScene, const FGuid& InGuid, float InTime)
 {
 	UMovieSceneByteTrack* Track = InMovieScene->AddTrack<UMovieSceneByteTrack>(InGuid);
 	if (Track)
 	{
-		Track->SetPropertyNameAndPath(Binding.GetPropertyName(), Binding.GetPropertyPath());
+		Track->SetPropertyNameAndPath(*Binding.GetProperty(InObjectToRecord)->GetDisplayNameText().ToString(), Binding.GetPropertyPath());
 
 		UMovieSceneByteSection* Section = Cast<UMovieSceneByteSection>(Track->CreateNewSection());
 		Section->SetDefault(PreviousValue);
@@ -100,12 +100,12 @@ bool FMovieScenePropertyRecorder<float>::ShouldAddNewKey(const float& InNewValue
 }
 
 template <>
-UMovieSceneSection* FMovieScenePropertyRecorder<float>::AddSection(UMovieScene* InMovieScene, const FGuid& InGuid, float InTime)
+UMovieSceneSection* FMovieScenePropertyRecorder<float>::AddSection(UObject* InObjectToRecord, UMovieScene* InMovieScene, const FGuid& InGuid, float InTime)
 {
 	UMovieSceneFloatTrack* Track = InMovieScene->AddTrack<UMovieSceneFloatTrack>(InGuid);
 	if (Track)
 	{
-		Track->SetPropertyNameAndPath(Binding.GetPropertyName(), Binding.GetPropertyPath());
+		Track->SetPropertyNameAndPath(*Binding.GetProperty(InObjectToRecord)->GetDisplayNameText().ToString(), Binding.GetPropertyPath());
 
 		UMovieSceneFloatSection* Section = Cast<UMovieSceneFloatSection>(Track->CreateNewSection());
 		Section->SetDefault(PreviousValue);
@@ -140,12 +140,12 @@ bool FMovieScenePropertyRecorder<FColor>::ShouldAddNewKey(const FColor& InNewVal
 }
 
 template <>
-UMovieSceneSection* FMovieScenePropertyRecorder<FColor>::AddSection(UMovieScene* InMovieScene, const FGuid& InGuid, float InTime)
+UMovieSceneSection* FMovieScenePropertyRecorder<FColor>::AddSection(UObject* InObjectToRecord, UMovieScene* InMovieScene, const FGuid& InGuid, float InTime)
 {
 	UMovieSceneColorTrack* Track = InMovieScene->AddTrack<UMovieSceneColorTrack>(InGuid);
 	if (Track)
 	{
-		Track->SetPropertyNameAndPath(Binding.GetPropertyName(), Binding.GetPropertyPath());
+		Track->SetPropertyNameAndPath(*Binding.GetProperty(InObjectToRecord)->GetDisplayNameText().ToString(), Binding.GetPropertyPath());
 
 		UMovieSceneColorSection* Section = Cast<UMovieSceneColorSection>(Track->CreateNewSection());
 		Section->SetDefault(FColorKey(EKeyColorChannel::Red, PreviousValue.R, false));
@@ -196,13 +196,13 @@ bool FMovieScenePropertyRecorder<FVector>::ShouldAddNewKey(const FVector& InNewV
 }
 
 template <>
-UMovieSceneSection* FMovieScenePropertyRecorder<FVector>::AddSection(UMovieScene* InMovieScene, const FGuid& InGuid, float InTime)
+UMovieSceneSection* FMovieScenePropertyRecorder<FVector>::AddSection(UObject* InObjectToRecord, UMovieScene* InMovieScene, const FGuid& InGuid, float InTime)
 {
 	UMovieSceneVectorTrack* Track = InMovieScene->AddTrack<UMovieSceneVectorTrack>(InGuid);
 	Track->SetNumChannelsUsed(3);
 	if (Track)
 	{
-		Track->SetPropertyNameAndPath(Binding.GetPropertyName(), Binding.GetPropertyPath());
+		Track->SetPropertyNameAndPath(*Binding.GetProperty(InObjectToRecord)->GetDisplayNameText().ToString(), Binding.GetPropertyPath());
 
 		UMovieSceneVectorSection* Section = Cast<UMovieSceneVectorSection>(Track->CreateNewSection());
 
diff --git a/Engine/Source/Editor/SequenceRecorder/Private/Sections/MovieScenePropertyRecorder.h b/Engine/Source/Editor/SequenceRecorder/Private/Sections/MovieScenePropertyRecorder.h
--- a/Engine/Source/Editor/SequenceRecorder/Private/Sections/MovieScenePropertyRecorder.h
+++ b/Engine/Source/Editor/SequenceRecorder/Private/Sections/MovieScenePropertyRecorder.h
@@ -40,7 +40,7 @@ public:
 		Binding.UpdateBinding(InObjectToRecord);
 		PreviousValue = Binding.GetCurrentValue<PropertyType>(InObjectToRecord);
 
-		MovieSceneSection = AddSection(InMovieScene, InGuid, InTime);
+		MovieSceneSection = AddSection(InObjectToRecord, InMovieScene, InGuid, InTime);
 	}
 
 	virtual void Record(UObject* InObjectToRecord, float InCurrentTime) override
@@ -81,7 +81,7 @@ private:
 	bool ShouldAddNewKey(const PropertyType& InNewValue) const;
 
 	/** Helper function, specialized by type, used to add an appropriate section to the movie scene */
-	class UMovieSceneSection* AddSection(class UMovieScene* InMovieScene, const FGuid& InGuid, float InTime);
+	class UMovieSceneSection* AddSection(UObject* InObjectToRecord, class UMovieScene* InMovieScene, const FGuid& InGuid, float InTime);
 
 	/** Helper function, specialized by type, used to add keys to the movie scene section at Finalize() time */
 	void AddKeyToSection(UMovieSceneSection* InSection, const FPropertyKey<PropertyType>& InKey);
diff --git a/Engine/Source/Editor/SequenceRecorder/Private/SequenceRecorder.cpp b/Engine/Source/Editor/SequenceRecorder/Private/SequenceRecorder.cpp
--- a/Engine/Source/Editor/SequenceRecorder/Private/SequenceRecorder.cpp
+++ b/Engine/Source/Editor/SequenceRecorder/Private/SequenceRecorder.cpp
@@ -15,12 +15,16 @@
 #include "DrawDebugHelpers.h"
 #include "LevelSequenceActor.h"
 #include "Runtime/Core/Public/Features/IModularFeatures.h"
+#include "Tracks/MovieSceneAudioTrack.h"
+#include "Sections/MovieSceneAudioSection.h"
+#include "ILevelViewport.h"
 #include "Animation/AnimSequence.h"
 
 #define LOCTEXT_NAMESPACE "SequenceRecorder"
 
 FSequenceRecorder::FSequenceRecorder()
 	: bQueuedRecordingsDirty(false)
+	, bWasImmersive(false)
 	, CurrentDelay(0.0f)
 	, CurrentTime(0.0f)
 {
@@ -38,6 +42,8 @@ void FSequenceRecorder::Initialize()
 	IModularFeatures::Get().RegisterModularFeature("MovieSceneSectionRecorderFactory", &AnimationSectionRecorderFactory);
 	IModularFeatures::Get().RegisterModularFeature("MovieSceneSectionRecorderFactory", &TransformSectionRecorderFactory);
 	IModularFeatures::Get().RegisterModularFeature("MovieSceneSectionRecorderFactory", &MultiPropertySectionRecorder);
+
+	RefreshNextSequence();
 }
 
 void FSequenceRecorder::Shutdown()
@@ -336,28 +342,70 @@ void FSequenceRecorder::DrawDebug(UCanvas* InCanvas, APlayerController* InPlayer
 
 	if(bCountingDown || IsRecording())
 	{
-		FText Text;
+		FText LabelText;
 		if(IsRecording())
 		{
-			Text = FText::Format(LOCTEXT("RecordingIndicator", "Recording Sequence: {0}\nTime: {1}"), FText::FromName(CurrentSequence.Get()->GetFName()), FText::AsNumber(CurrentTime));
+			LabelText = FText::Format(LOCTEXT("RecordingIndicatorFormat", "{0}"), FText::FromName(CurrentSequence.Get()->GetFName()));
 		}
 		else
 		{
-			Text = LOCTEXT("RecordingIndicatorPending", "Sequence Recording Pending");
+			LabelText = FText::Format(LOCTEXT("RecordingIndicatorPending", "Pending recording: {0}"), FText::FromString(NextSequenceName));
 		}
 
+
+		float TimeAccumulator = CurrentTime;
+		float Hours = FMath::FloorToFloat(TimeAccumulator / (60.0f * 60.0f));
+		TimeAccumulator -= Hours * 60.0f * 60.0f;
+		float Minutes = FMath::FloorToFloat(TimeAccumulator / 60.0f);
+		TimeAccumulator -= Minutes * 60.0f;
+		float Seconds = FMath::FloorToFloat(TimeAccumulator);
+		TimeAccumulator -= Seconds;
+		float Frames = FMath::FloorToFloat(TimeAccumulator * GetDefault<USequenceRecorderSettings>()->DefaultAnimationSettings.SampleRate);
+
+		FNumberFormattingOptions Options;
+		Options.MinimumIntegralDigits = 2;
+		Options.MaximumIntegralDigits = 2;
+
+		FFormatNamedArguments NamedArgs;
+		NamedArgs.Add(TEXT("Hours"), FText::AsNumber((int32)Hours, &Options));
+		NamedArgs.Add(TEXT("Minutes"), FText::AsNumber((int32)Minutes, &Options));
+		NamedArgs.Add(TEXT("Seconds"), FText::AsNumber((int32)Seconds, &Options));
+		NamedArgs.Add(TEXT("Frames"), FText::AsNumber((int32)Frames, &Options));
+		FText TimeText = FText::Format(LOCTEXT("RecordingTimerFormat", "{Hours}:{Minutes}:{Seconds}:{Frames}"), NamedArgs);
+
 		const FVector2D IconSize(32.0f, 32.0f);
-		const FVector2D Offset(32.0f, 64.0f);
+		const FVector2D Offset(8.0f, 32.0f);
 
 		InCanvas->SetDrawColor(FColor::White);
 
-		FVector2D IconPosition(Offset.X, Offset.Y * 2.0f);
+		FVector2D IconPosition(Offset.X, InCanvas->SizeY - (Offset.Y + IconSize.Y));
 		FCanvasIcon Icon = UCanvas::MakeIcon(RecordingIndicatorTexture.Get(), FMath::FloorToFloat(NumFrames - CurrentDelay) * IconSize.X, 0.0f, IconSize.X, IconSize.Y);
 		InCanvas->DrawIcon(Icon, IconPosition.X, IconPosition.Y);
 
-		FFontRenderInfo Info;
-		Info.bEnableShadow = true;
-		InCanvas->DrawText(GEngine->GetMediumFont(), Text, IconPosition.X + IconSize.X + 4.0f, IconPosition.Y + 2.0f, 1.0f, 1.0f, Info);
+		const float TextScale = 1.2f;
+		float TextPositionY = 0.0f;
+		// draw label
+		{
+			float XSize, YSize;
+			InCanvas->TextSize(GEngine->GetLargeFont(), LabelText.ToString(), XSize, YSize, TextScale, TextScale);
+
+			TextPositionY = (IconPosition.Y + (IconSize.Y * 0.5f)) - (YSize * 0.5f);
+
+			FFontRenderInfo Info;
+			Info.bEnableShadow = true;
+			InCanvas->DrawText(GEngine->GetLargeFont(), LabelText, IconPosition.X + IconSize.X + 4.0f, TextPositionY, TextScale, TextScale, Info);
+		}
+		// draw time
+		{
+			float XSize, YSize;
+			InCanvas->TextSize(GEngine->GetLargeFont(), TimeText.ToString(), XSize, YSize, TextScale, TextScale);
+
+			FVector2D TimePosition(InCanvas->SizeX - (Offset.X + XSize), TextPositionY);
+
+			FFontRenderInfo Info;
+			Info.bEnableShadow = true;
+			InCanvas->DrawText(GEngine->GetLargeFont(), TimeText, TimePosition.X, TimePosition.Y, TextScale, TextScale, Info);
+		}
 	}
 }
 
@@ -388,6 +436,26 @@ bool FSequenceRecorder::StartRecording(const FOnRecordingStarted& OnRecordingSta
 
 	CurrentTime = 0.0f;
 
+	if (Settings->bImmersiveMode)
+	{
+		FLevelEditorModule& LevelEditorModule = FModuleManager::Get().LoadModuleChecked<FLevelEditorModule>("LevelEditor");
+		TSharedPtr< ILevelViewport > ActiveLevelViewport = LevelEditorModule.GetFirstActiveViewport();
+
+		if( ActiveLevelViewport.IsValid() )
+		{
+			bWasImmersive = ActiveLevelViewport->IsImmersive();
+
+			if (!ActiveLevelViewport->IsImmersive())
+			{
+				const bool bWantImmersive = true;
+				const bool bAllowAnimation = false;
+				ActiveLevelViewport->MakeImmersive( bWantImmersive, bAllowAnimation );
+			}
+		}
+	}
+
+	RefreshNextSequence();
+
 	if(Settings->RecordingDelay > 0.0f)
 	{
 		CurrentDelay = Settings->RecordingDelay;
@@ -470,6 +538,8 @@ bool FSequenceRecorder::StartRecordingInternal(UWorld* World)
 				CurrentSequence = LevelSequence;
 
 				FAssetRegistryModule::AssetCreated(LevelSequence);
+
+				RefreshNextSequence();
 			}
 		}
 
@@ -555,6 +625,23 @@ void FSequenceRecorder::HandleEndPIE(bool bSimulating)
 
 bool FSequenceRecorder::StopRecording()
 {
+	const USequenceRecorderSettings* Settings = GetDefault<USequenceRecorderSettings>();
+
+	if (Settings->bImmersiveMode)
+	{
+		FLevelEditorModule& LevelEditorModule = FModuleManager::Get().LoadModuleChecked<FLevelEditorModule>("LevelEditor");
+		TSharedPtr< ILevelViewport > ActiveLevelViewport = LevelEditorModule.GetFirstActiveViewport();
+
+		if( ActiveLevelViewport.IsValid() )
+		{
+			if (ActiveLevelViewport->IsImmersive() != bWasImmersive)
+			{
+				const bool bAllowAnimation = false;
+				ActiveLevelViewport->MakeImmersive(bWasImmersive, bAllowAnimation);
+			}
+		}
+	}
+
 	FScopedSlowTask SlowTask((float)(QueuedRecordings.Num() + DeadRecordings.Num()), LOCTEXT("ProcessingRecording", "Processing Recording"));
 	SlowTask.MakeDialog(false, true);
 
@@ -592,7 +679,6 @@ bool FSequenceRecorder::StopRecording()
 
 	DeadRecordings.Empty();
 
-	const USequenceRecorderSettings* Settings = GetDefault<USequenceRecorderSettings>();
 	if(Settings->bCreateLevelSequence)
 	{
 		if(CurrentSequence.IsValid())
@@ -758,4 +844,12 @@ void FSequenceRecorder::HandleActorDespawned(AActor* Actor)
 	}
 }
 
+void FSequenceRecorder::RefreshNextSequence()
+{
+	// Cache the name of the next sequence we will try to record to
+	const USequenceRecorderSettings* Settings = GetDefault<USequenceRecorderSettings>();
+	SequenceName = Settings->SequenceName.Len() > 0 ? Settings->SequenceName : TEXT("RecordedSequence");
+	NextSequenceName = SequenceRecorderUtils::MakeNewAssetName<ULevelSequence>(Settings->SequenceRecordingBasePath.Path, SequenceName);
+}
+
 #undef LOCTEXT_NAMESPACE
\ No newline at end of file
diff --git a/Engine/Source/Editor/SequenceRecorder/Private/SequenceRecorder.h b/Engine/Source/Editor/SequenceRecorder/Private/SequenceRecorder.h
--- a/Engine/Source/Editor/SequenceRecorder/Private/SequenceRecorder.h
+++ b/Engine/Source/Editor/SequenceRecorder/Private/SequenceRecorder.h
@@ -90,6 +90,12 @@ public:
 		return TransformSectionRecorderFactory;
 	}
 
+	/** Get the name of the next sequence we are targeting */
+	const FString& GetNextSequenceName() const { return NextSequenceName; }
+
+	/** Refresh the name of the next sequence we will be recording */
+	void RefreshNextSequence();
+
 private:
 	/** Starts recording a sequence, possibly delayed */
 	bool StartRecordingInternal(UWorld* World);
@@ -121,6 +127,8 @@ private:
 
 	bool bQueuedRecordingsDirty;
 
+	bool bWasImmersive;
+
 	/** The delay we are currently waiting for */
 	float CurrentDelay;
 
@@ -145,6 +153,9 @@ private:
 	/** Cached sequence name to record to */
 	FString SequenceName;
 
+	/** The next sequence we will be targeting. Name can change depending on assets being deleted, moved, renamed etc. */
+	FString NextSequenceName;
+
 	/** Cached sequence path to record to */
 	FString PathToRecordTo;
 
diff --git a/Engine/Source/Editor/SequenceRecorder/Private/SequenceRecorderSettings.cpp b/Engine/Source/Editor/SequenceRecorder/Private/SequenceRecorderSettings.cpp
--- a/Engine/Source/Editor/SequenceRecorder/Private/SequenceRecorderSettings.cpp
+++ b/Engine/Source/Editor/SequenceRecorder/Private/SequenceRecorderSettings.cpp
@@ -2,11 +2,13 @@
 
 #include "SequenceRecorderPrivatePCH.h"
 #include "SequenceRecorderSettings.h"
+#include "SequenceRecorder.h"
 
 USequenceRecorderSettings::USequenceRecorderSettings(const FObjectInitializer& ObjectInitializer)
 	: Super(ObjectInitializer)
 {
 	bCreateLevelSequence = true;
+	bImmersiveMode = false;
 	SequenceLength = FAnimationRecordingSettings::DefaultMaximumLength;
 	RecordingDelay = 4.0f;
 	SequenceName = TEXT("RecordedSequence");
@@ -16,11 +18,11 @@ USequenceRecorderSettings::USequenceRecorderSettings(const FObjectInitializer& O
 	NearbyActorRecordingProximity = 5000.0f;
 	bRecordWorldSettingsActor = true;
 
-	ComponentClassesToRecord.Add(USkeletalMeshComponent::StaticClass());
-	ComponentClassesToRecord.Add(UStaticMeshComponent::StaticClass());
-	ComponentClassesToRecord.Add(UParticleSystemComponent::StaticClass());
-	ComponentClassesToRecord.Add(ULightComponent::StaticClass());
-	ComponentClassesToRecord.Add(UCameraComponent::StaticClass());
+	ClassesAndPropertiesToRecord.Add(FPropertiesToRecordForClass(USkeletalMeshComponent::StaticClass()));
+	ClassesAndPropertiesToRecord.Add(FPropertiesToRecordForClass(UStaticMeshComponent::StaticClass()));
+	ClassesAndPropertiesToRecord.Add(FPropertiesToRecordForClass(UParticleSystemComponent::StaticClass()));
+	ClassesAndPropertiesToRecord.Add(FPropertiesToRecordForClass(ULightComponent::StaticClass()));
+	ClassesAndPropertiesToRecord.Add(FPropertiesToRecordForClass(UCameraComponent::StaticClass()));
 }
 
 void USequenceRecorderSettings::PostEditChangeChainProperty(struct FPropertyChangedChainEvent& PropertyChangedEvent)
@@ -28,4 +30,13 @@ void USequenceRecorderSettings::PostEditChangeChainProperty(struct FPropertyChan
 	Super::PostEditChangeChainProperty(PropertyChangedEvent);
 
 	SaveConfig();
+
+	if (PropertyChangedEvent.Property)
+	{
+		if (PropertyChangedEvent.Property->GetFName() == GET_MEMBER_NAME_CHECKED(USequenceRecorderSettings, SequenceName) ||
+			PropertyChangedEvent.Property->GetFName() == GET_MEMBER_NAME_CHECKED(USequenceRecorderSettings, SequenceRecordingBasePath))
+		{
+			FSequenceRecorder::Get().RefreshNextSequence();
+		}
+	}
 }
\ No newline at end of file
diff --git a/Engine/Source/Editor/SequenceRecorder/Public/SequenceRecorderSettings.h b/Engine/Source/Editor/SequenceRecorder/Public/SequenceRecorderSettings.h
--- a/Engine/Source/Editor/SequenceRecorder/Public/SequenceRecorderSettings.h
+++ b/Engine/Source/Editor/SequenceRecorder/Public/SequenceRecorderSettings.h
@@ -11,9 +11,16 @@ struct FPropertiesToRecordForClass
 {
 	GENERATED_BODY()
 
+	FPropertiesToRecordForClass()
+	{}
+
+	FPropertiesToRecordForClass(TSubclassOf<USceneComponent> InClass)
+		: Class(InClass)
+	{}
+
 	/** The class of the object we can record */
 	UPROPERTY(Config, EditAnywhere, AdvancedDisplay, Category = "Sequence Recording")
-	TSubclassOf<UObject> Class;
+	TSubclassOf<USceneComponent> Class;
 
 	/** List of properties we want to record for this class */
 	UPROPERTY(Config, EditAnywhere, AdvancedDisplay, Category = "Sequence Recording")
@@ -32,6 +39,10 @@ public:
 	UPROPERTY(Config, EditAnywhere, Category = "Sequence Recording")
 	bool bCreateLevelSequence;
 
+	/** Whether to maximize the viewport when recording */
+	UPROPERTY(Config, EditAnywhere, Category = "Sequence Recording")
+	bool bImmersiveMode;
+
 	/** The length of the recorded sequence */
 	UPROPERTY(Config, EditAnywhere, Category = "Sequence Recording")
 	float SequenceLength;
@@ -76,15 +87,11 @@ public:
 	UPROPERTY(Config, EditAnywhere, Category = "Sequence Recording")
 	FAnimationRecordingSettings DefaultAnimationSettings;
 
-	/** Component classes we record by default. If an actor does not contain one of these classes it will be ignored. */
-	UPROPERTY(Config, EditAnywhere, AdvancedDisplay, Category = "Sequence Recording")
-	TArray<TSubclassOf<USceneComponent>> ComponentClassesToRecord;
-
 	/** Whether to record actors that are spawned by sequencer itself (this is usually disabled as results can be unexpected) */
 	UPROPERTY(Config, EditAnywhere, AdvancedDisplay, Category = "Sequence Recording")
 	bool bRecordSequencerSpawnedActors;
 
-	/** The properties to record for specified classes */
-	UPROPERTY(Config, EditAnywhere, Category = "General Property Recording")
-	TArray<FPropertiesToRecordForClass> PropertiesToRecord;
+	/** The properties to record for specified classes. Component classes specified here will be recorded. If an actor does not contain one of these classes it will be ignored. */
+	UPROPERTY(Config, EditAnywhere, Category = "Sequence Recording")
+	TArray<FPropertiesToRecordForClass> ClassesAndPropertiesToRecord;
 };
\ No newline at end of file
diff --git a/Engine/Source/Editor/SequenceRecorder/Public/SequenceRecorderUtils.h b/Engine/Source/Editor/SequenceRecorder/Public/SequenceRecorderUtils.h
--- a/Engine/Source/Editor/SequenceRecorder/Public/SequenceRecorderUtils.h
+++ b/Engine/Source/Editor/SequenceRecorder/Public/SequenceRecorderUtils.h
@@ -51,6 +51,33 @@ static AssetType* MakeNewAsset(const FString& BaseAssetPath, const FString& Base
 	return nullptr;
 }
 
+template<typename AssetType>
+static FString MakeNewAssetName(const FString& BaseAssetPath, const FString& BaseAssetName)
+{
+	const FString Dot(TEXT("."));
+	FString AssetPath = BaseAssetPath;
+	FString AssetName = BaseAssetName;
+
+	AssetPath /= AssetName;
+	AssetPath += Dot + AssetName;
+
+	FAssetRegistryModule& AssetRegistryModule = FModuleManager::LoadModuleChecked<FAssetRegistryModule>(TEXT("AssetRegistry"));
+	FAssetData AssetData = AssetRegistryModule.Get().GetAssetByObjectPath(*AssetPath);
+
+	// if object with same name exists, try a different name until we don't find one
+	int32 ExtensionIndex = 0;
+	while (AssetData.IsValid())
+	{
+		AssetName = FString::Printf(TEXT("%s_%d"), *BaseAssetName, ExtensionIndex);
+		AssetPath = (BaseAssetPath / AssetName) + Dot + AssetName;
+		AssetData = AssetRegistryModule.Get().GetAssetByObjectPath(*AssetPath);
+
+		ExtensionIndex++;
+	}
+
+	return AssetName;
+}
+
 /** Helper function - check whether our component hierarchy has some attachment outside of its owned components */
 SEQUENCERECORDER_API AActor* GetAttachment(AActor* InActor, FName& SocketName, FName& ComponentName);
 
diff --git a/Engine/Source/Editor/Sequencer/Private/Sequencer.cpp b/Engine/Source/Editor/Sequencer/Private/Sequencer.cpp
--- a/Engine/Source/Editor/Sequencer/Private/Sequencer.cpp
+++ b/Engine/Source/Editor/Sequencer/Private/Sequencer.cpp
@@ -223,6 +223,7 @@ void FSequencer::InitSequencer(const FSequencerInitParams& InitParams, const TSh
 	// NOTE: Could fill in asset editor commands here!
 
 	BindCommands();
+	BindLevelEditorCommands();
 	ActivateSequencerEditorMode();
 
 	if( bIsEditingWithinLevelEditor )
@@ -287,6 +288,8 @@ FSequencer::~FSequencer()
 			TrackEditor->OnRelease();
 		}
 		TrackEditors.Empty();
+	
+		UnbindLevelEditorCommands();
 	}
 	SequencerWidget.Reset();
 }
@@ -328,6 +331,8 @@ void FSequencer::Close()
 			FName SettingsSectionName = *Settings->GetName();
 			SettingsModule->UnregisterSettings("Editor", "ContentEditors", SettingsSectionName);
 		}
+
+		UnbindLevelEditorCommands();
 	}
 
 	if (bIsEditingWithinLevelEditor)
@@ -1451,8 +1456,13 @@ void FSequencer::SetGlobalTimeDirectly( float NewTime, ESnapTimeMode SnapTimeMod
 
 	if (!IsInSilentMode())
 	{
-		// Redraw
-		FEditorSupportDelegates::RedrawAllViewports.Broadcast();
+		// Redraw if not in PIE/simulate
+		const bool bIsInPIEOrSimulate = GEditor->PlayWorld != NULL || GEditor->bIsSimulatingInEditor;
+		if (!bIsInPIEOrSimulate)
+		{
+			FEditorSupportDelegates::RedrawAllViewports.Broadcast();
+		}
+		
 		OnGlobalTimeChangedDelegate.Broadcast();
 	}
 }
@@ -1743,11 +1753,21 @@ void FSequencer::PossessPIEViewports(UObject* CameraObject, UObject* UnlockIfCam
 			PrePossessionViewTargets.Add(FCachedViewTarget{ PC, ViewTarget });
 		}
 
+		UCameraComponent* CameraComponent = MovieSceneHelpers::CameraComponentFromRuntimeObject(CameraObject);
+
 		if (CameraObject == ViewTarget)
 		{
-			if ( bJumpCut && PC->PlayerCameraManager )
+			if ( bJumpCut )
 			{
-				PC->PlayerCameraManager->bGameCameraCutThisFrame = true;
+				if (PC->PlayerCameraManager)
+				{
+					PC->PlayerCameraManager->bGameCameraCutThisFrame = true;
+				}
+
+				if (CameraComponent)
+				{
+					CameraComponent->NotifyCameraCut();
+				}
 			}
 			continue;
 		}
@@ -1775,6 +1795,11 @@ void FSequencer::PossessPIEViewports(UObject* CameraObject, UObject* UnlockIfCam
 		FViewTargetTransitionParams TransitionParams;
 		PC->SetViewTarget(CameraActor, TransitionParams);
 
+		if (CameraComponent)
+		{
+			CameraComponent->NotifyCameraCut();
+		}
+
 		if (PC->PlayerCameraManager)
 		{
 			PC->PlayerCameraManager->bClientSimulatingViewTarget = (CameraActor != nullptr);
@@ -1787,13 +1812,6 @@ void FSequencer::UpdateCameraCut(UObject* CameraObject, UObject* UnlockIfCameraO
 {
 	OnCameraCutEvent.Broadcast(CameraObject, bJumpCut);
 
-	// tell the camera we cut
-	UCameraComponent* const CamComp = MovieSceneHelpers::CameraComponentFromRuntimeObject(CameraObject);
-	if (CamComp)
-	{
-		CamComp->NotifyCameraCut();
-	}
-
 	if (!IsPerspectiveViewportCameraCutEnabled())
 	{
 		return;
@@ -2023,18 +2041,14 @@ TSharedRef<FExtender> FSequencer::GetLevelViewportExtender(const TSharedRef<FUIC
 		ActorName = FText::Format(LOCTEXT("ActorNamePlural", "{0} Actors"), FText::AsNumber(InActors.Num()));
 	}
 
-	Extender->AddMenuExtension("ActorControl", EExtensionHook::After, CommandList, FMenuExtensionDelegate::CreateLambda(
+	FLevelEditorModule& LevelEditor = FModuleManager::GetModuleChecked<FLevelEditorModule>(TEXT("LevelEditor"));
+	TSharedRef<FUICommandList> LevelEditorCommandBindings = LevelEditor.GetGlobalLevelEditorActions();
+
+	Extender->AddMenuExtension("ActorControl", EExtensionHook::After, LevelEditorCommandBindings, FMenuExtensionDelegate::CreateLambda(
 		[this, ActorName](FMenuBuilder& MenuBuilder){
 			MenuBuilder.BeginSection("SequenceRecording", LOCTEXT("SequenceRecordingHeading", "Sequence Recording"));
 
-			MenuBuilder.AddMenuEntry(
-				FText::Format(LOCTEXT("RecordSelectedActorsText", "Record {0} In Sequencer"), ActorName),
-				LOCTEXT("RecordSelectedActors_Tooltip", "Start recording the specified actors in sequencer."),
-				FSlateIcon(FEditorStyle::GetStyleSetName(), "SequenceRecorder.TabIcon"),
-				FUIAction(
-					FExecuteAction::CreateLambda([this]{ RecordSelectedActors(); })
-				)
-			);
+			MenuBuilder.AddMenuEntry(FSequencerCommands::Get().RecordSelectedActors, NAME_None, FText::Format(LOCTEXT("RecordSelectedActorsText", "Record {0} In Sequencer"), ActorName));
 
 			MenuBuilder.EndSection();
 		}
@@ -3707,7 +3721,11 @@ void FSequencer::OnPreBeginPIE(bool bIsSimulating)
 
 void FSequencer::OnEndPlayMap()
 {
+	// Update and clear any stale bindings 
 	UpdateRuntimeInstances();
+	
+	// Update for new bindings on next tick
+	bNeedInstanceRefresh = true;
 }
 
 
@@ -3765,46 +3783,53 @@ void FSequencer::UpdatePreviewLevelViewportClientFromCameraCut(FLevelEditorViewp
 {
 	AActor* CameraActor = Cast<AActor>(InCameraObject);
 
+	bool bCameraHasBeenCut = bJumpCut;
+
 	if (CameraActor)
 	{
+		bCameraHasBeenCut = bCameraHasBeenCut || !InViewportClient.IsLockedToActor(CameraActor);
 		InViewportClient.SetViewLocation(CameraActor->GetActorLocation());
 		InViewportClient.SetViewRotation(CameraActor->GetActorRotation());
-		InViewportClient.SetIsCameraCut(!InViewportClient.IsLockedToActor(CameraActor) || bJumpCut);
 	}
 	else
 	{
 		InViewportClient.ViewFOV = InViewportClient.FOVAngle;
-		InViewportClient.SetIsCameraCut(bJumpCut);
 	}
 
+	InViewportClient.SetIsCameraCut(bCameraHasBeenCut);
+
+
 	// Set the actor lock.
 	InViewportClient.SetMatineeActorLock(CameraActor);
 	InViewportClient.bLockedCameraView = CameraActor != nullptr;
 	InViewportClient.RemoveCameraRoll();
 
-	// If viewing through a camera - enforce aspect ratio.
-	if (CameraActor)
+	UCameraComponent* CameraComponent = MovieSceneHelpers::CameraComponentFromRuntimeObject(InCameraObject);
+	if (CameraComponent)
 	{
-		UCameraComponent* CameraComponent = Cast<UCameraComponent>(CameraActor->GetComponentByClass(UCameraComponent::StaticClass()));
-		if (CameraComponent)
+		if (bCameraHasBeenCut)
 		{
-			if (CameraComponent->AspectRatio == 0)
-			{
-				InViewportClient.AspectRatio = 1.7f;
-			}
-			else
-			{
-				InViewportClient.AspectRatio = CameraComponent->AspectRatio;
-			}
+			// tell the camera we cut
+			CameraComponent->NotifyCameraCut();
+		}
+
+		// enforce aspect ratio.
+		if (CameraComponent->AspectRatio == 0)
+		{
+			InViewportClient.AspectRatio = 1.7f;
+		}
+		else
+		{
+			InViewportClient.AspectRatio = CameraComponent->AspectRatio;
+		}
 
-			//don't stop the camera from zooming when not playing back
-			InViewportClient.ViewFOV = CameraComponent->FieldOfView;
+		//don't stop the camera from zooming when not playing back
+		InViewportClient.ViewFOV = CameraComponent->FieldOfView;
 
-			// If there are selected actors, invalidate the viewports hit proxies, otherwise they won't be selectable afterwards
-			if (InViewportClient.Viewport && GEditor->GetSelectedActorCount() > 0)
-			{
-				InViewportClient.Viewport->InvalidateHitProxy();
-			}
+		// If there are selected actors, invalidate the viewports hit proxies, otherwise they won't be selectable afterwards
+		if (InViewportClient.Viewport && GEditor->GetSelectedActorCount() > 0)
+		{
+			InViewportClient.Viewport->InvalidateHitProxy();
 		}
 	}
 
@@ -6105,14 +6130,33 @@ void FSequencer::BindCommands()
 		Commands.SelectKeysInSelectionRange,
 		FExecuteAction::CreateLambda([this]{ SelectKeysInSelectionRange(); }));
 
-	SequencerCommandBindings->MapAction(
-		Commands.RecordSelectedActors,
-		FExecuteAction::CreateLambda([this]{ RecordSelectedActors(); }));
-
 	// bind widget specific commands
 	SequencerWidget->BindCommands(SequencerCommandBindings);
 }
 
+void FSequencer::BindLevelEditorCommands()
+{
+	const FSequencerCommands& Commands = FSequencerCommands::Get();
+
+	FLevelEditorModule& LevelEditor = FModuleManager::GetModuleChecked<FLevelEditorModule>(TEXT("LevelEditor"));
+	TSharedRef<FUICommandList> LevelEditorCommandBindings = LevelEditor.GetGlobalLevelEditorActions();
+
+	LevelEditorCommandBindings->MapAction(
+		Commands.RecordSelectedActors,
+		FExecuteAction::CreateSP(this, &FSequencer::RecordSelectedActors));
+}
+
+void FSequencer::UnbindLevelEditorCommands()
+{
+	if (FModuleManager::Get().IsModuleLoaded(TEXT("LevelEditor")))
+	{
+		// unbind the commands we have bound to the level editor
+		FLevelEditorModule& LevelEditor = FModuleManager::GetModuleChecked<FLevelEditorModule>(TEXT("LevelEditor"));
+		TSharedRef<FUICommandList> LevelEditorCommandBindings = LevelEditor.GetGlobalLevelEditorActions();
+
+		LevelEditorCommandBindings->UnmapAction(FSequencerCommands::Get().RecordSelectedActors);
+	}
+}
 
 void FSequencer::BuildAddTrackMenu(class FMenuBuilder& MenuBuilder)
 {
diff --git a/Engine/Source/Editor/Sequencer/Private/Sequencer.h b/Engine/Source/Editor/Sequencer/Private/Sequencer.h
--- a/Engine/Source/Editor/Sequencer/Private/Sequencer.h
+++ b/Engine/Source/Editor/Sequencer/Private/Sequencer.h
@@ -325,6 +325,12 @@ public:
 
 	/** Get the extender for the level viewport menu */
 	TSharedRef<FExtender> GetLevelViewportExtender(const TSharedRef<FUICommandList> CommandList, const TArray<AActor*> InActors);
+	
+	/** Add level editor actions */
+	void BindLevelEditorCommands();
+	
+	/** Remove level editor actions */
+	void UnbindLevelEditorCommands();
 
 	void RecordSelectedActors();
 	
diff --git a/Engine/Source/Editor/Sequencer/Private/SequencerCommands.cpp b/Engine/Source/Editor/Sequencer/Private/SequencerCommands.cpp
--- a/Engine/Source/Editor/Sequencer/Private/SequencerCommands.cpp
+++ b/Engine/Source/Editor/Sequencer/Private/SequencerCommands.cpp
@@ -53,11 +53,11 @@ void FSequencerCommands::RegisterCommands()
 	UI_COMMAND( TrimSectionRight, "Trim Section Right", "Trim section at current time to the right (keeps the left)", EUserInterfaceActionType::Button, FInputChord(EModifierKey::Control, EKeys::Period) );
 	UI_COMMAND( SplitSection, "Split Section", "Split section at current time", EUserInterfaceActionType::Button, FInputChord(EModifierKey::Control, EKeys::Slash) );
 
-	UI_COMMAND( SetAutoKeyModeAll, "Auto-key All", "Enables auto keying for all properties.", EUserInterfaceActionType::ToggleButton, FInputChord() );
-	UI_COMMAND( SetAutoKeyModeAnimated, "Auto-key Animated", "Enables auto keying for properties with existing animations.", EUserInterfaceActionType::ToggleButton, FInputChord());
+	UI_COMMAND( SetAutoKeyModeAll, "Auto-key All", "Auto-key all channels/properties when they change.", EUserInterfaceActionType::ToggleButton, FInputChord() );
+	UI_COMMAND( SetAutoKeyModeAnimated, "Auto-key Animated", "Auto-key channels/properties when they change only if they have existing keys/animation.", EUserInterfaceActionType::ToggleButton, FInputChord());
 	UI_COMMAND( SetAutoKeyModeNone, "Disable Auto-key", "Disables auto keying.", EUserInterfaceActionType::ToggleButton, FInputChord());
 
-	UI_COMMAND( ToggleKeyAllEnabled, "Key All", "Enables and disables key all", EUserInterfaceActionType::ToggleButton, FInputChord());
+	UI_COMMAND( ToggleKeyAllEnabled, "Key All", "Key all channels/properties when only one of them changes. ie. Keys all translation, rotation, scale channels when only translation Y changes", EUserInterfaceActionType::ToggleButton, FInputChord());
 
 	UI_COMMAND( ToggleAutoScroll, "Auto Scroll", "Toggle auto-scroll: When enabled, automatically scrolls the sequencer view to keep the current time visible", EUserInterfaceActionType::ToggleButton, FInputChord(EModifierKey::Shift, EKeys::S) );
 
@@ -99,7 +99,7 @@ void FSequencerCommands::RegisterCommands()
 
 	UI_COMMAND( FixFrameTiming, "Fix Frame Timing", "Moves all time data for this sequence to a valid frame time.", EUserInterfaceActionType::Button, FInputChord() );
 
-	UI_COMMAND( RecordSelectedActors, "Record Selected Actors", "Records the selected actors into a new sub sequence of the currently active sequence in Sequencer.", EUserInterfaceActionType::Button, FInputChord() );
+	UI_COMMAND( RecordSelectedActors, "Record Selected Actors", "Records the selected actors into a new sub sequence of the currently active sequence in Sequencer.", EUserInterfaceActionType::Button, FInputChord(EModifierKey::Alt, EKeys::R) );
 	UI_COMMAND( ExportSceneAndSequence, "Export Scene and Sequence", "Exports the current scene and sequence to an FBX file. (Shots and sub-scenes not supported)", EUserInterfaceActionType::Button, FInputChord() );
 }
 
diff --git a/Engine/Source/Editor/Sequencer/Public/MovieSceneTrackEditor.h b/Engine/Source/Editor/Sequencer/Public/MovieSceneTrackEditor.h
--- a/Engine/Source/Editor/Sequencer/Public/MovieSceneTrackEditor.h
+++ b/Engine/Source/Editor/Sequencer/Public/MovieSceneTrackEditor.h
@@ -238,7 +238,7 @@ protected:
 	/**
 	 * Gets the currently focused movie scene, if any.
 	 *
-	 * @result Focused movie scene, or nullptr if no movie scene is focused.
+	 * @return Focused movie scene, or nullptr if no movie scene is focused.
 	 */
 	UMovieScene* GetFocusedMovieScene() const
 	{
diff --git a/Engine/Source/Editor/UnrealEd/Classes/Editor/EditorEngine.h b/Engine/Source/Editor/UnrealEd/Classes/Editor/EditorEngine.h
--- a/Engine/Source/Editor/UnrealEd/Classes/Editor/EditorEngine.h
+++ b/Engine/Source/Editor/UnrealEd/Classes/Editor/EditorEngine.h
@@ -1084,7 +1084,7 @@ public:
 	 * @param	Transform		The world-space transform to spawn the actor with.
 	 * @param	bSilent			If true, suppress logging (optional, defaults to false).
 	 * @param	ObjectFlags		The object flags to place on the spawned actor.
-	 * @result					A pointer to the newly added actor, or NULL if add failed.
+	 * @return					A pointer to the newly added actor, or NULL if add failed.
 	 */
 	virtual AActor* AddActor(ULevel* InLevel, UClass* Class, const FTransform& Transform, bool bSilent = false, EObjectFlags ObjectFlags = RF_Transactional);
 
@@ -1094,7 +1094,7 @@ public:
 	 * @param	ExportText		A T3D representation of the actor to create.
 	 * @param	bSilent			If true, suppress logging
 	 * @param	ObjectFlags		The object flags to place on the spawned actor.
-	 * @result					A pointer to the newly added actor, or NULL if add failed.
+	 * @return					A pointer to the newly added actor, or NULL if add failed.
 	 */
 	virtual TArray<AActor*> AddExportTextActors(const FString& ExportText, bool bSilent, EObjectFlags ObjectFlags = RF_Transactional);
 
diff --git a/Engine/Source/Editor/UnrealEd/Classes/Factories/Factory.h b/Engine/Source/Editor/UnrealEd/Classes/Factories/Factory.h
--- a/Engine/Source/Editor/UnrealEd/Classes/Factories/Factory.h
+++ b/Engine/Source/Editor/UnrealEd/Classes/Factories/Factory.h
@@ -2,6 +2,7 @@
 
 #pragma once
 
+#include "CoreUObject.h"
 #include "Factory.generated.h"
 
 
@@ -20,11 +21,11 @@ public:
 
 	/** The class manufactured by this factory. */
 	UPROPERTY()
-	TSubclassOf<class UObject>  SupportedClass;
+	TSubclassOf<UObject>  SupportedClass;
 
 	/** Class of the context object used to help create the object. */
 	UPROPERTY()
-	TSubclassOf<class UObject>  ContextClass;
+	TSubclassOf<UObject>  ContextClass;
 
 	/** List of formats supported by the factory. Each entry is of the form "ext;Description" where ext is the file extension. */
 	UPROPERTY()
diff --git a/Engine/Source/Editor/UnrealEd/Private/Factories/Factory.cpp b/Engine/Source/Editor/UnrealEd/Private/Factories/Factory.cpp
--- a/Engine/Source/Editor/UnrealEd/Private/Factories/Factory.cpp
+++ b/Engine/Source/Editor/UnrealEd/Private/Factories/Factory.cpp
@@ -179,8 +179,12 @@ FText UFactory::GetDisplayName() const
 		TWeakPtr<IAssetTypeActions> AssetTypeActions = AssetToolsModule.Get().GetAssetTypeActionsForClass(LocalSupportedClass);
 		if ( AssetTypeActions.IsValid() )
 		{
-			// @todo AssetTypeActions should be returning a FText so we dont have to do this conversion here.
-			return AssetTypeActions.Pin()->GetName();
+			FText Name = AssetTypeActions.Pin()->GetName();
+
+			if (!Name.IsEmpty())
+			{
+				return Name;
+			}
 		}
 
 		// Factories whose classes do not have asset type actions should just display the sanitized class name
diff --git a/Engine/Source/Editor/UnrealEd/Public/AssetThumbnail.h b/Engine/Source/Editor/UnrealEd/Public/AssetThumbnail.h
--- a/Engine/Source/Editor/UnrealEd/Public/AssetThumbnail.h
+++ b/Engine/Source/Editor/UnrealEd/Public/AssetThumbnail.h
@@ -1,13 +1,15 @@
 // Copyright 1998-2016 Epic Games, Inc. All Rights Reserved.
 
-
 #pragma once
 
+#include "SlateCore.h" // for ISlateViewport
+#include "Runtime/AssetRegistry/Public/AssetData.h"
+#include "TickableEditorObject.h"
+
+
 class FSlateTexture2DRHIRef;
 class FSlateTextureRenderTarget2DResource;
 
-#include "Runtime/AssetRegistry/Public/AssetData.h"
-#include "TickableEditorObject.h"
 
 namespace EThumbnailLabel
 {
@@ -19,6 +21,7 @@ namespace EThumbnailLabel
 	};
 };
 
+
 /** A struct containing details about how the asset thumbnail should behave */
 struct FAssetThumbnailConfig
 {
@@ -46,10 +49,13 @@ struct FAssetThumbnailConfig
 	TOptional< FLinearColor > AssetTypeColorOverride;
 };
 
+
 /**
  * Interface for rendering a thumbnail in a slate viewport                   
  */
-class FAssetThumbnail : public ISlateViewport, public TSharedFromThis<FAssetThumbnail>
+class FAssetThumbnail
+	: public ISlateViewport
+	, public TSharedFromThis<FAssetThumbnail>
 {
 public:
 	/**
diff --git a/Engine/Source/Programs/UnrealLightmass/Private/ImportExport/Importer.h b/Engine/Source/Programs/UnrealLightmass/Private/ImportExport/Importer.h
--- a/Engine/Source/Programs/UnrealLightmass/Private/ImportExport/Importer.h
+++ b/Engine/Source/Programs/UnrealLightmass/Private/ImportExport/Importer.h
@@ -24,18 +24,18 @@ public:
 	 * @param Scene Scene object to fill out
 	 * @param SceneGuid Guid of the scene to load from a swarm channel
 	 */
-	bool	ImportScene( class FScene& Scene, const FGuid& SceneGuid );
+	bool ImportScene( class FScene& Scene, const FGuid& SceneGuid );
 
 	/** Imports a buffer of raw data */
-	bool	Read( void* Data, int32 NumBytes );
+	bool Read( void* Data, int32 NumBytes );
 
 	/** Imports one object */
 	template <class DataType>
-	bool	ImportData( DataType* Data );
+	bool ImportData( DataType* Data );
 
 	/** Imports a TArray of simple elements in one bulk read. */
 	template <class ArrayType>
-	bool	ImportArray( ArrayType& Array, int32 Count );
+	bool ImportArray( ArrayType& Array, int32 Count );
 
 	/** Imports a TArray of objects, while also adding them to the specified LookupMap. */
 	template <class ObjType, class LookupMapType>
@@ -43,7 +43,7 @@ public:
 
 	/** Imports an array of GUIDs and stores the corresponding pointers into a TArray */
 	template <class ArrayType, class LookupMapType>
-	bool	ImportGuidArray( ArrayType& Array, int32 Count, const LookupMapType& LookupMap );
+	bool ImportGuidArray( ArrayType& Array, int32 Count, const LookupMapType& LookupMap );
 
 	/**
 	 * Finds existing or imports new object by Guid
@@ -51,8 +51,7 @@ public:
 	 * @param Guid Guid of object
 	 * @param Version Version of object to load
 	 * @param Extension Type of object to load (@todo UE4: This could be removed if Version could imply extension)
-	 *
-	 * @result The object that was loaded or found, or NULL if the Guid failed
+	 * @return The object that was loaded or found, or NULL if the Guid failed
 	 */
 	template <class ObjType, class LookupMapType>
 	ObjType* ConditionalImportObject(const FGuid& Guid, const FGuid& Version, const TCHAR* Extension, int32 ChannelFlags, LookupMapType& LookupMap);
@@ -92,12 +91,14 @@ private:
 	float LevelScale;
 };
 
+
 template <typename DataType>
 FORCEINLINE bool FLightmassImporter::ImportData( DataType* Data )
 {
 	return Read( Data, sizeof(DataType) );
 }
 
+
 /** Imports a TArray of simple elements in one bulk read. */
 template <class ArrayType>
 bool FLightmassImporter::ImportArray( ArrayType& Array, int32 Count )
@@ -107,6 +108,7 @@ bool FLightmassImporter::ImportArray( ArrayType& Array, int32 Count )
 	return Read( Array.GetData(), Count*sizeof(typename ArrayType::ElementType) );
 }
 
+
 /** Imports a TArray of objects, while also adding them to the specified LookupMap. */
 template <class ObjType, class LookupMapType>
 bool FLightmassImporter::ImportObjectArray( TArray<ObjType>& Array, int32 Count, LookupMapType& LookupMap )
@@ -121,6 +123,7 @@ bool FLightmassImporter::ImportObjectArray( TArray<ObjType>& Array, int32 Count,
 	return true;
 }
 
+
 /** Imports an array of GUIDs and stores the corresponding pointers into a TArray */
 template <class ArrayType, class LookupMapType>
 bool FLightmassImporter::ImportGuidArray( ArrayType& Array, int32 Count, const LookupMapType& LookupMap )
@@ -136,14 +139,14 @@ bool FLightmassImporter::ImportGuidArray( ArrayType& Array, int32 Count, const L
 	return bOk;
 }
 
+
 /**
  * Finds existing or imports new object by Guid
  *
  * @param Guid Guid of object
  * @param Version Version of object to load
  * @param Extension Type of object to load (@todo UE4: This could be removed if Version could imply extension)
- *
- * @result The object that was loaded or found, or NULL if the Guid failed
+ * @return The object that was loaded or found, or NULL if the Guid failed
  */
 template <class ObjType, class LookupMapType>
 ObjType* FLightmassImporter::ConditionalImportObject(const FGuid& Guid, const FGuid& Version, const TCHAR* Extension, int32 ChannelFlags, LookupMapType& LookupMap)
diff --git a/Engine/Source/Runtime/CinematicCamera/Private/CineCameraComponent.cpp b/Engine/Source/Runtime/CinematicCamera/Private/CineCameraComponent.cpp
--- a/Engine/Source/Runtime/CinematicCamera/Private/CineCameraComponent.cpp
+++ b/Engine/Source/Runtime/CinematicCamera/Private/CineCameraComponent.cpp
@@ -80,6 +80,7 @@ void UCineCameraComponent::PostInitProperties()
 void UCineCameraComponent::PostLoad()
 {
 	RecalcDerivedData();
+	bResetInterpolation = true;
 	Super::PostLoad();
 }
 
diff --git a/Engine/Source/Runtime/Core/Private/Android/AndroidJava.cpp b/Engine/Source/Runtime/Core/Private/Android/AndroidJava.cpp
--- a/Engine/Source/Runtime/Core/Private/Android/AndroidJava.cpp
+++ b/Engine/Source/Runtime/Core/Private/Android/AndroidJava.cpp
@@ -91,6 +91,20 @@ jobject FJavaClassObject::CallMethod<jobject>(FJavaClassMethod Method, ...)
 	return RetVal;
 }
 
+template<>
+jobjectArray FJavaClassObject::CallMethod<jobjectArray>(FJavaClassMethod Method, ...)
+{
+	JNIEnv*	JEnv = FAndroidApplication::GetJavaEnv();
+	va_list Params;
+	va_start(Params, Method);
+	jobject val = JEnv->CallObjectMethodV(Object, Method.Method, Params);
+	va_end(Params);
+	VerifyException();
+	jobjectArray RetVal = (jobjectArray)JEnv->NewGlobalRef(val);
+	JEnv->DeleteLocalRef(val);
+	return RetVal;
+}
+
 template<>
 int64 FJavaClassObject::CallMethod<int64>(FJavaClassMethod Method, ...)
 {
