Copying //UE4/Dev-Sequencer to //UE4/Dev-Main (Source: //UE4/Dev-Sequencer @ 3178529)

diff --git a/Engine/Source/Editor/PropertyEditor/Public/DetailCategoryBuilder.h b/Engine/Source/Editor/PropertyEditor/Public/DetailCategoryBuilder.h
--- a/Engine/Source/Editor/PropertyEditor/Public/DetailCategoryBuilder.h
+++ b/Engine/Source/Editor/PropertyEditor/Public/DetailCategoryBuilder.h
@@ -93,6 +93,15 @@ public:
 	 */
 	virtual IDetailPropertyRow* AddExternalProperty( TSharedPtr<FStructOnScope> StructData, FName PropertyName, EPropertyLocation::Type Location = EPropertyLocation::Default ) = 0;
 
+	/**
+	 * Adds all properties for the specified external structure to this category
+	 *
+	 * @param StructData		External struct data to add to the root
+	 * @param Location			The location within the category where the properties are to be shown
+	 * @return Array of all properties added to this category
+	 */
+	virtual TArray<TSharedPtr<IPropertyHandle>> AddExternalProperties( TSharedRef<FStructOnScope> StructData, EPropertyLocation::Type Location = EPropertyLocation::Default ) = 0;
+
 	/**
 	 * Adds a custom widget row to the category
 	 *
diff --git a/Engine/Source/Editor/PropertyEditor/Public/DetailLayoutBuilder.h b/Engine/Source/Editor/PropertyEditor/Public/DetailLayoutBuilder.h
--- a/Engine/Source/Editor/PropertyEditor/Public/DetailLayoutBuilder.h
+++ b/Engine/Source/Editor/PropertyEditor/Public/DetailLayoutBuilder.h
@@ -63,6 +63,13 @@ public:
 	 */
 	virtual void GetObjectsBeingCustomized( TArray< TWeakObjectPtr<UObject> >& OutObjects ) const = 0;
 
+	/**
+	 * Gets the current struct(s) being customized by this builder
+	 *
+	 * If this is a sub-struct customization it will return those sub struct.  Otherwise the root struct will be returned.
+	 */
+	virtual void GetStructsBeingCustomized( TArray< TSharedPtr<FStructOnScope> >& OutStructs ) const = 0;
+
 	/**
 	 *	@return the utilities various widgets need access to certain features of PropertyDetails
 	 */
diff --git a/Engine/Source/Editor/PropertyEditor/Public/IDetailChildrenBuilder.h b/Engine/Source/Editor/PropertyEditor/Public/IDetailChildrenBuilder.h
--- a/Engine/Source/Editor/PropertyEditor/Public/IDetailChildrenBuilder.h
+++ b/Engine/Source/Editor/PropertyEditor/Public/IDetailChildrenBuilder.h
@@ -45,6 +45,14 @@ public:
 	 */
 	virtual class IDetailPropertyRow& AddChildProperty( TSharedRef<IPropertyHandle> PropertyHandle ) = 0;
 
+	/**
+	 * Adds a child structure
+	 * 
+	 * @param ChildStructure	The structure to add
+	 * @return An array of interfaces to the properties that were added
+	 */
+	virtual TArray<TSharedPtr<IPropertyHandle>> AddChildStructure( TSharedRef<FStructOnScope> ChildStructure ) = 0;
+
 	/**
 	 * Generates a value widget from a customized struct
 	 * If the customized struct has no value widget an empty widget will be returned
diff --git a/Engine/Source/Editor/SequenceRecorder/Private/ActorRecording.cpp b/Engine/Source/Editor/SequenceRecorder/Private/ActorRecording.cpp
--- a/Engine/Source/Editor/SequenceRecorder/Private/ActorRecording.cpp
+++ b/Engine/Source/Editor/SequenceRecorder/Private/ActorRecording.cpp
@@ -292,6 +292,13 @@ void UActorRecording::StartRecordingActorProperties(ULevelSequence* CurrentSeque
 					SkeletalMeshComponent->ForcedLodModel = 1;
 				}
 
+				// Disable possession of pawns otherwise the recorded character will auto possess the player
+				if (ObjectTemplate->IsA(APawn::StaticClass()))
+				{
+					APawn* Pawn = CastChecked<APawn>(ObjectTemplate);
+					Pawn->AutoPossessPlayer = EAutoReceiveInput::Disabled;
+				}
+
 				Guid = MovieScene->AddSpawnable(TemplateName, *ObjectTemplate);
 			}
 		}
diff --git a/Engine/Source/Editor/SequenceRecorder/Private/SSequenceRecorder.cpp b/Engine/Source/Editor/SequenceRecorder/Private/SSequenceRecorder.cpp
--- a/Engine/Source/Editor/SequenceRecorder/Private/SSequenceRecorder.cpp
+++ b/Engine/Source/Editor/SequenceRecorder/Private/SSequenceRecorder.cpp
@@ -136,6 +136,7 @@ void SSequenceRecorder::Construct(const FArguments& Args)
 	{
 		ToolBarBuilder.AddToolBarButton(FSequenceRecorderCommands::Get().AddRecording);
 		ToolBarBuilder.AddToolBarButton(FSequenceRecorderCommands::Get().RemoveRecording);
+		ToolBarBuilder.AddToolBarButton(FSequenceRecorderCommands::Get().RemoveAllRecordings);
 	}
 	ToolBarBuilder.EndSection();
 
@@ -263,6 +264,11 @@ void SSequenceRecorder::BindCommands()
 		FExecuteAction::CreateSP(this, &SSequenceRecorder::HandleRemoveRecording),
 		FCanExecuteAction::CreateSP(this, &SSequenceRecorder::CanRemoveRecording)
 		);
+
+	CommandList->MapAction(FSequenceRecorderCommands::Get().RemoveAllRecordings,
+		FExecuteAction::CreateSP(this, &SSequenceRecorder::HandleRemoveAllRecordings),
+		FCanExecuteAction::CreateSP(this, &SSequenceRecorder::CanRemoveAllRecordings)
+		);
 }
 
 TSharedRef<ITableRow> SSequenceRecorder::MakeListViewWidget(UActorRecording* Recording, const TSharedRef<STableViewBase>& OwnerTable) const
@@ -346,6 +352,17 @@ bool SSequenceRecorder::CanRemoveRecording() const
 	return ListView->GetNumItemsSelected() > 0 && !FAnimationRecorderManager::Get().IsRecording();
 }
 
+void SSequenceRecorder::HandleRemoveAllRecordings()
+{
+	FSequenceRecorder::Get().ClearQueuedRecordings();
+	ActorRecordingDetailsView->SetObject(nullptr);
+}
+
+bool SSequenceRecorder::CanRemoveAllRecordings() const
+{
+	return FSequenceRecorder::Get().HasQueuedRecordings() && !FAnimationRecorderManager::Get().IsRecording();
+}
+
 EActiveTimerReturnType SSequenceRecorder::HandleRefreshItems(double InCurrentTime, float InDeltaTime)
 {
 	if(FSequenceRecorder::Get().AreQueuedRecordingsDirty())
diff --git a/Engine/Source/Editor/SequenceRecorder/Private/SSequenceRecorder.h b/Engine/Source/Editor/SequenceRecorder/Private/SSequenceRecorder.h
--- a/Engine/Source/Editor/SequenceRecorder/Private/SSequenceRecorder.h
+++ b/Engine/Source/Editor/SequenceRecorder/Private/SSequenceRecorder.h
@@ -46,6 +46,10 @@ private:
 
 	bool CanRemoveRecording() const;
 
+	void HandleRemoveAllRecordings();
+
+	bool CanRemoveAllRecordings() const;
+
 	EActiveTimerReturnType HandleRefreshItems(double InCurrentTime, float InDeltaTime);
 
 	TOptional<float> GetDelayPercent() const;
diff --git a/Engine/Source/Editor/SequenceRecorder/Private/Sections/MovieSceneMultiPropertyRecorder.cpp b/Engine/Source/Editor/SequenceRecorder/Private/Sections/MovieSceneMultiPropertyRecorder.cpp
--- a/Engine/Source/Editor/SequenceRecorder/Private/Sections/MovieSceneMultiPropertyRecorder.cpp
+++ b/Engine/Source/Editor/SequenceRecorder/Private/Sections/MovieSceneMultiPropertyRecorder.cpp
@@ -50,8 +50,7 @@ void FMovieSceneMultiPropertyRecorder::CreateSection(UObject* InObjectToRecord,
 	for (const FName& PropertyName : PropertiesToRecord)
 	{
 		FTrackInstancePropertyBindings Binding(PropertyName, PropertyName.ToString());
-		Binding.UpdateBinding(InObjectToRecord);
-		UProperty* Property = Binding.GetProperty(InObjectToRecord);
+		UProperty* Property = Binding.GetProperty(*InObjectToRecord);
 		if (Property != nullptr)
 		{
 			if (Property->IsA<UBoolProperty>())
diff --git a/Engine/Source/Editor/SequenceRecorder/Private/Sections/MovieScenePropertyRecorder.cpp b/Engine/Source/Editor/SequenceRecorder/Private/Sections/MovieScenePropertyRecorder.cpp
--- a/Engine/Source/Editor/SequenceRecorder/Private/Sections/MovieScenePropertyRecorder.cpp
+++ b/Engine/Source/Editor/SequenceRecorder/Private/Sections/MovieScenePropertyRecorder.cpp
@@ -27,7 +27,10 @@ UMovieSceneSection* FMovieScenePropertyRecorder<bool>::AddSection(UObject* InObj
 	UMovieSceneBoolTrack* Track = InMovieScene->AddTrack<UMovieSceneBoolTrack>(InGuid);
 	if (Track)
 	{
-		Track->SetPropertyNameAndPath(*Binding.GetProperty(InObjectToRecord)->GetDisplayNameText().ToString(), Binding.GetPropertyPath());
+		if (InObjectToRecord)
+		{
+			Track->SetPropertyNameAndPath(*Binding.GetProperty(*InObjectToRecord)->GetDisplayNameText().ToString(), Binding.GetPropertyPath());
+		}
 
 		UMovieSceneBoolSection* Section = Cast<UMovieSceneBoolSection>(Track->CreateNewSection());
 		Section->SetDefault(PreviousValue);
@@ -66,7 +69,10 @@ UMovieSceneSection* FMovieScenePropertyRecorder<uint8>::AddSection(UObject* InOb
 	UMovieSceneByteTrack* Track = InMovieScene->AddTrack<UMovieSceneByteTrack>(InGuid);
 	if (Track)
 	{
-		Track->SetPropertyNameAndPath(*Binding.GetProperty(InObjectToRecord)->GetDisplayNameText().ToString(), Binding.GetPropertyPath());
+		if (InObjectToRecord)
+		{
+			Track->SetPropertyNameAndPath(*Binding.GetProperty(*InObjectToRecord)->GetDisplayNameText().ToString(), Binding.GetPropertyPath());
+		}
 
 		UMovieSceneByteSection* Section = Cast<UMovieSceneByteSection>(Track->CreateNewSection());
 		Section->SetDefault(PreviousValue);
@@ -105,7 +111,10 @@ UMovieSceneSection* FMovieScenePropertyRecorder<float>::AddSection(UObject* InOb
 	UMovieSceneFloatTrack* Track = InMovieScene->AddTrack<UMovieSceneFloatTrack>(InGuid);
 	if (Track)
 	{
-		Track->SetPropertyNameAndPath(*Binding.GetProperty(InObjectToRecord)->GetDisplayNameText().ToString(), Binding.GetPropertyPath());
+		if (InObjectToRecord)
+		{
+			Track->SetPropertyNameAndPath(*Binding.GetProperty(*InObjectToRecord)->GetDisplayNameText().ToString(), Binding.GetPropertyPath());
+		}
 
 		UMovieSceneFloatSection* Section = Cast<UMovieSceneFloatSection>(Track->CreateNewSection());
 		Section->SetDefault(PreviousValue);
@@ -145,7 +154,10 @@ UMovieSceneSection* FMovieScenePropertyRecorder<FColor>::AddSection(UObject* InO
 	UMovieSceneColorTrack* Track = InMovieScene->AddTrack<UMovieSceneColorTrack>(InGuid);
 	if (Track)
 	{
-		Track->SetPropertyNameAndPath(*Binding.GetProperty(InObjectToRecord)->GetDisplayNameText().ToString(), Binding.GetPropertyPath());
+		if (InObjectToRecord)
+		{
+			Track->SetPropertyNameAndPath(*Binding.GetProperty(*InObjectToRecord)->GetDisplayNameText().ToString(), Binding.GetPropertyPath());
+		}
 
 		UMovieSceneColorSection* Section = Cast<UMovieSceneColorSection>(Track->CreateNewSection());
 		Section->SetDefault(FColorKey(EKeyColorChannel::Red, PreviousValue.R, false));
@@ -202,7 +214,10 @@ UMovieSceneSection* FMovieScenePropertyRecorder<FVector>::AddSection(UObject* In
 	Track->SetNumChannelsUsed(3);
 	if (Track)
 	{
-		Track->SetPropertyNameAndPath(*Binding.GetProperty(InObjectToRecord)->GetDisplayNameText().ToString(), Binding.GetPropertyPath());
+		if (InObjectToRecord)
+		{
+			Track->SetPropertyNameAndPath(*Binding.GetProperty(*InObjectToRecord)->GetDisplayNameText().ToString(), Binding.GetPropertyPath());
+		}
 
 		UMovieSceneVectorSection* Section = Cast<UMovieSceneVectorSection>(Track->CreateNewSection());
 
diff --git a/Engine/Source/Editor/SequenceRecorder/Private/Sections/MovieScenePropertyRecorder.h b/Engine/Source/Editor/SequenceRecorder/Private/Sections/MovieScenePropertyRecorder.h
--- a/Engine/Source/Editor/SequenceRecorder/Private/Sections/MovieScenePropertyRecorder.h
+++ b/Engine/Source/Editor/SequenceRecorder/Private/Sections/MovieScenePropertyRecorder.h
@@ -37,8 +37,10 @@ public:
 
 	virtual void Create(UObject* InObjectToRecord, class UMovieScene* InMovieScene, const FGuid& InGuid, float InTime) override
 	{
-		Binding.UpdateBinding(InObjectToRecord);
-		PreviousValue = Binding.GetCurrentValue<PropertyType>(InObjectToRecord);
+		if (InObjectToRecord)
+		{
+			PreviousValue = Binding.GetCurrentValue<PropertyType>(*InObjectToRecord);
+		}
 
 		MovieSceneSection = AddSection(InObjectToRecord, InMovieScene, InGuid, InTime);
 	}
@@ -49,7 +51,7 @@ public:
 		{
 			MovieSceneSection->SetEndTime(InCurrentTime);
 
-			PropertyType NewValue = Binding.GetCurrentValue<PropertyType>(InObjectToRecord);
+			PropertyType NewValue = Binding.GetCurrentValue<PropertyType>(*InObjectToRecord);
 			if (ShouldAddNewKey(NewValue))
 			{
 				FPropertyKey<PropertyType> Key;
diff --git a/Engine/Source/Editor/SequenceRecorder/Private/SequenceRecorder.cpp b/Engine/Source/Editor/SequenceRecorder/Private/SequenceRecorder.cpp
--- a/Engine/Source/Editor/SequenceRecorder/Private/SequenceRecorder.cpp
+++ b/Engine/Source/Editor/SequenceRecorder/Private/SequenceRecorder.cpp
@@ -766,12 +766,12 @@ bool FSequenceRecorder::StopRecording()
 	{
 		if(LevelSequence)
 		{
-			// Stop referencing the sequence so we are listed as 'not recording'
-			CurrentSequence = nullptr;
-
 			// set movie scene playback range to encompass all sections
 			UpdateSequencePlaybackRange();
 
+			// Stop referencing the sequence so we are listed as 'not recording'
+			CurrentSequence = nullptr;
+
 			if(GEditor == nullptr)
 			{
 				// auto-save asset outside of the editor
@@ -840,6 +840,13 @@ void FSequenceRecorder::UpdateSequencePlaybackRange()
 			}
 
 			MovieScene->SetPlaybackRange(MinRange, MaxRange);
+
+			// Initialize the working and view range with a little bit more space
+			const float OutputViewSize = MaxRange - MinRange;
+			const float OutputChange = OutputViewSize * 0.1f;
+
+			MovieScene->SetWorkingRange(MinRange - OutputChange, MaxRange + OutputChange);
+			MovieScene->SetViewRange(MinRange - OutputChange, MaxRange + OutputChange);
 		}
 	}
 }
@@ -929,9 +936,13 @@ void FSequenceRecorder::HandleActorDespawned(AActor* Actor)
 
 void FSequenceRecorder::RefreshNextSequence()
 {
-	// Cache the name of the next sequence we will try to record to
 	const USequenceRecorderSettings* Settings = GetDefault<USequenceRecorderSettings>();
-	SequenceName = Settings->SequenceName.Len() > 0 ? Settings->SequenceName : TEXT("RecordedSequence");
+	if (SequenceName.IsEmpty())
+	{
+		SequenceName = Settings->SequenceName.Len() > 0 ? Settings->SequenceName : TEXT("RecordedSequence");
+	}
+
+	// Cache the name of the next sequence we will try to record to
 	NextSequenceName = SequenceRecorderUtils::MakeNewAssetName<ULevelSequence>(Settings->SequenceRecordingBasePath.Path, SequenceName);
 }
 
diff --git a/Engine/Source/Editor/SequenceRecorder/Private/SequenceRecorderCommands.cpp b/Engine/Source/Editor/SequenceRecorder/Private/SequenceRecorderCommands.cpp
--- a/Engine/Source/Editor/SequenceRecorder/Private/SequenceRecorderCommands.cpp
+++ b/Engine/Source/Editor/SequenceRecorder/Private/SequenceRecorderCommands.cpp
@@ -16,6 +16,7 @@ void FSequenceRecorderCommands::RegisterCommands()
 	UI_COMMAND(StopAll, "StopAll", "Stop all recordings", EUserInterfaceActionType::Button, FInputChord(EModifierKey::Control | EModifierKey::Shift, EKeys::S));
 	UI_COMMAND(AddRecording, "Add", "Add a new recording", EUserInterfaceActionType::Button, FInputChord());
 	UI_COMMAND(RemoveRecording, "Remove", "Remove selected recording", EUserInterfaceActionType::Button, FInputChord());
+	UI_COMMAND(RemoveAllRecordings, "Remove All", "Remove all recordings", EUserInterfaceActionType::Button, FInputChord());
 }
 
 #undef LOCTEXT_NAMESPACE
\ No newline at end of file
diff --git a/Engine/Source/Editor/SequenceRecorder/Private/SequenceRecorderCommands.h b/Engine/Source/Editor/SequenceRecorder/Private/SequenceRecorderCommands.h
--- a/Engine/Source/Editor/SequenceRecorder/Private/SequenceRecorderCommands.h
+++ b/Engine/Source/Editor/SequenceRecorder/Private/SequenceRecorderCommands.h
@@ -13,6 +13,7 @@ public:
 	TSharedPtr<FUICommandInfo> StopAll;
 	TSharedPtr<FUICommandInfo> AddRecording;
 	TSharedPtr<FUICommandInfo> RemoveRecording;
+	TSharedPtr<FUICommandInfo> RemoveAllRecordings;
 
 	/** Initialize commands */
 	virtual void RegisterCommands() override;
diff --git a/Engine/Source/Editor/SequenceRecorder/Public/SequenceRecorderSettings.h b/Engine/Source/Editor/SequenceRecorder/Public/SequenceRecorderSettings.h
--- a/Engine/Source/Editor/SequenceRecorder/Public/SequenceRecorderSettings.h
+++ b/Engine/Source/Editor/SequenceRecorder/Public/SequenceRecorderSettings.h
@@ -66,11 +66,11 @@ public:
 	bool bImmersiveMode;
 
 	/** The length of the recorded sequence */
-	UPROPERTY(Config, EditAnywhere, Category = "Sequence Recording")
+	UPROPERTY(Config, EditAnywhere, Category = "Sequence Recording", meta = (ClampMin="0.0", UIMin = "0.0"))
 	float SequenceLength;
 
 	/** Delay that we will use before starting recording */
-	UPROPERTY(Config, EditAnywhere, Category = "Sequence Recording", meta = (ClampMax="9.0", UIMax = "9.0"))
+	UPROPERTY(Config, EditAnywhere, Category = "Sequence Recording", meta = (ClampMin="0.0", UIMin = "0.0", ClampMax="9.0", UIMax = "9.0"))
 	float RecordingDelay;
 
 	/** The base name of the sequence to record to. This name will also be used to auto-generate any assets created by this recording. */
@@ -94,11 +94,11 @@ public:
 	EAudioRecordingMode RecordAudio;
 
 	/** Gain in decibels to apply to recorded audio */
-	UPROPERTY(Config, EditAnywhere, Category = "Sequence Recording")
+	UPROPERTY(Config, EditAnywhere, Category = "Sequence Recording", meta = (ClampMin="0.0", UIMin = "0.0"))
 	float AudioGain;
 
 	/** The buffer size to use on mic input callbacks. Larger sizes increase latency but reduce chances of buffer overruns (pops and discontinuities). */
-	UPROPERTY(Config, EditAnywhere, Category = "Sequence Recording")
+	UPROPERTY(Config, EditAnywhere, Category = "Sequence Recording", meta = (ClampMin="0", UIMin = "0"))
 	int32 AudioInputBufferSize;
 
 	/** Whether to record nearby spawned actors. */
@@ -106,7 +106,7 @@ public:
 	bool bRecordNearbySpawnedActors;
 
 	/** Proximity to currently recorded actors to record newly spawned actors. */
-	UPROPERTY(Config, EditAnywhere, Category = "Sequence Recording")
+	UPROPERTY(Config, EditAnywhere, Category = "Sequence Recording", meta = (ClampMin="0.0", UIMin = "0.0"))
 	float NearbyActorRecordingProximity;
 
 	/** Whether to record the world settings actor in the sequence (some games use this to attach world SFX) */
diff --git a/Engine/Source/Editor/SequenceRecorderSections/Private/MovieSceneVisibilitySectionRecorder.cpp b/Engine/Source/Editor/SequenceRecorderSections/Private/MovieSceneVisibilitySectionRecorder.cpp
--- a/Engine/Source/Editor/SequenceRecorderSections/Private/MovieSceneVisibilitySectionRecorder.cpp
+++ b/Engine/Source/Editor/SequenceRecorderSections/Private/MovieSceneVisibilitySectionRecorder.cpp
@@ -2,12 +2,12 @@
 
 #include "SequenceRecorderSectionsPrivatePCH.h"
 #include "MovieSceneVisibilitySectionRecorder.h"
-#include "MovieSceneVisibilitySection.h"
+#include "MovieSceneBoolSection.h"
 #include "MovieSceneVisibilityTrack.h"
 #include "MovieScene.h"
 #include "ActorRecordingSettings.h"
 
-static const FName VisibilityTrackName = TEXT("Visibility");
+static const FName VisibilityTrackName = TEXT("bHidden");
 
 TSharedPtr<IMovieSceneSectionRecorder> FMovieSceneVisibilitySectionRecorderFactory::CreateSectionRecorder(const struct FActorRecordingSettings& InActorRecordingSettings) const
 {
@@ -35,7 +35,7 @@ void FMovieSceneVisibilitySectionRecorder::CreateSection(UObject* InObjectToReco
 	{
 		VisibilityTrack->SetPropertyNameAndPath(VisibilityTrackName, VisibilityTrackName.ToString());
 
-		MovieSceneSection = Cast<UMovieSceneVisibilitySection>(VisibilityTrack->CreateNewSection());
+		MovieSceneSection = Cast<UMovieSceneBoolSection>(VisibilityTrack->CreateNewSection());
 
 		VisibilityTrack->AddSection(*MovieSceneSection);
 
diff --git a/Engine/Source/Editor/SequenceRecorderSections/Private/MovieSceneVisibilitySectionRecorder.h b/Engine/Source/Editor/SequenceRecorderSections/Private/MovieSceneVisibilitySectionRecorder.h
--- a/Engine/Source/Editor/SequenceRecorderSections/Private/MovieSceneVisibilitySectionRecorder.h
+++ b/Engine/Source/Editor/SequenceRecorderSections/Private/MovieSceneVisibilitySectionRecorder.h
@@ -38,7 +38,7 @@ private:
 	TLazyObjectPtr<class UObject> ObjectToRecord;
 
 	/** Section to record to */
-	TWeakObjectPtr<class UMovieSceneVisibilitySection> MovieSceneSection;
+	TWeakObjectPtr<class UMovieSceneBoolSection> MovieSceneSection;
 
 	/** Flag used to track visibility state and add keys when this changes */
 	bool bWasVisible;
diff --git a/Engine/Source/Editor/Sequencer/Private/DisplayNodes/SequencerObjectBindingNode.cpp b/Engine/Source/Editor/Sequencer/Private/DisplayNodes/SequencerObjectBindingNode.cpp
--- a/Engine/Source/Editor/Sequencer/Private/DisplayNodes/SequencerObjectBindingNode.cpp
+++ b/Engine/Source/Editor/Sequencer/Private/DisplayNodes/SequencerObjectBindingNode.cpp
@@ -11,6 +11,10 @@
 #include "SequencerUtilities.h"
 #include "SlateIconFinder.h"
 #include "SequencerCommands.h"
+#include "ScopedTransaction.h"
+
+#include "MovieSceneSpawnTrack.h"
+#include "MovieSceneSpawnSection.h"
 
 #define LOCTEXT_NAMESPACE "FObjectBindingNode"
 
@@ -113,6 +117,8 @@ void FSequencerObjectBindingNode::BuildContextMenu(FMenuBuilder& MenuBuilder)
 				FNewMenuDelegate::CreateSP(this, &FSequencerObjectBindingNode::AddSpawnOwnershipMenu)
 			);
 
+			MenuBuilder.AddMenuEntry( FSequencerCommands::Get().SaveCurrentSpawnableState );
+
 			MenuBuilder.AddMenuEntry( FSequencerCommands::Get().ConvertToPossessable );
 		}
 		else
@@ -164,18 +170,39 @@ void FSequencerObjectBindingNode::BuildContextMenu(FMenuBuilder& MenuBuilder)
 
 void FSequencerObjectBindingNode::AddSpawnOwnershipMenu(FMenuBuilder& MenuBuilder)
 {
-	FMovieSceneSpawnable* Spawnable = GetSequencer().GetFocusedMovieSceneSequence()->GetMovieScene()->FindSpawnable(ObjectBinding);
+	UMovieScene* MovieScene = GetSequencer().GetFocusedMovieSceneSequence()->GetMovieScene();
+	FMovieSceneSpawnable* Spawnable = MovieScene->FindSpawnable(ObjectBinding);
 	if (!Spawnable)
 	{
 		return;
 	}
+	auto Callback = [=](ESpawnOwnership NewOwnership){
+
+		FScopedTransaction Transaction(LOCTEXT("SetSpawnOwnership", "Set Spawnable Ownership"));
+
+		Spawnable->SetSpawnOwnership(NewOwnership);
+
+		// Overwrite the completion state for all spawn sections to ensure the expected behaviour.
+		EMovieSceneCompletionMode NewCompletionMode = NewOwnership == ESpawnOwnership::InnerSequence ? EMovieSceneCompletionMode::RestoreState : EMovieSceneCompletionMode::KeepState;
+
+		// Make all spawn sections retain state
+		UMovieSceneSpawnTrack* SpawnTrack = MovieScene->FindTrack<UMovieSceneSpawnTrack>(ObjectBinding);
+		if (SpawnTrack)
+		{
+			for (UMovieSceneSection* Section : SpawnTrack->GetAllSections())
+			{
+				Section->Modify();
+				Section->EvalOptions.CompletionMode = NewCompletionMode;
+			}
+		}
+	};
 
 	MenuBuilder.AddMenuEntry(
 		LOCTEXT("ThisSequence_Label", "This Sequence"),
 		LOCTEXT("ThisSequence_Tooltip", "Indicates that this sequence will own the spawned object. The object will be destroyed at the end of the sequence."),
 		FSlateIcon(),
 		FUIAction(
-			FExecuteAction::CreateLambda([=]{ Spawnable->SetSpawnOwnership(ESpawnOwnership::InnerSequence); }),
+			FExecuteAction::CreateLambda(Callback, ESpawnOwnership::InnerSequence),
 			FCanExecuteAction(),
 			FIsActionChecked::CreateLambda([=]{ return Spawnable->GetSpawnOwnership() == ESpawnOwnership::InnerSequence; })
 		),
@@ -188,7 +215,7 @@ void FSequencerObjectBindingNode::AddSpawnOwnershipMenu(FMenuBuilder& MenuBuilde
 		LOCTEXT("MasterSequence_Tooltip", "Indicates that the outermost sequence will own the spawned object. The object will be destroyed when the outermost sequence stops playing."),
 		FSlateIcon(),
 		FUIAction(
-			FExecuteAction::CreateLambda([=]{ Spawnable->SetSpawnOwnership(ESpawnOwnership::MasterSequence); }),
+			FExecuteAction::CreateLambda(Callback, ESpawnOwnership::MasterSequence),
 			FCanExecuteAction(),
 			FIsActionChecked::CreateLambda([=]{ return Spawnable->GetSpawnOwnership() == ESpawnOwnership::MasterSequence; })
 		),
@@ -201,30 +228,13 @@ void FSequencerObjectBindingNode::AddSpawnOwnershipMenu(FMenuBuilder& MenuBuilde
 		LOCTEXT("External_Tooltip", "Indicates this object's lifetime is managed externally once spawned. It will not be destroyed by sequencer."),
 		FSlateIcon(),
 		FUIAction(
-			FExecuteAction::CreateLambda([=]{ Spawnable->SetSpawnOwnership(ESpawnOwnership::External); }),
+			FExecuteAction::CreateLambda(Callback, ESpawnOwnership::External),
 			FCanExecuteAction(),
 			FIsActionChecked::CreateLambda([=]{ return Spawnable->GetSpawnOwnership() == ESpawnOwnership::External; })
 		),
 		NAME_None,
 		EUserInterfaceActionType::ToggleButton
 	);
-
-	MenuBuilder.AddMenuSeparator();
-	MenuBuilder.AddMenuEntry(
-		LOCTEXT("IgnoreOwnership_Label", "Keep Alive Outside Playback Range (In Sequencer)"),
-		LOCTEXT("IgnoreOwnership_Tooltip", "Keeps the spawned object alive when viewing this specific sequence outside of its playback range. Does not apply to runtime evaluation."),
-		FSlateIcon(),
-		FUIAction(
-			FExecuteAction::CreateLambda([=]{
-				Spawnable->SetIgnoreOwnershipInEditor(!Spawnable->ShouldIgnoreOwnershipInEditor());
-				GetSequencer().SetGlobalTimeDirectly(GetSequencer().GetGlobalTime());
-			}),
-			FCanExecuteAction(),
-			FIsActionChecked::CreateLambda([=]{ return Spawnable->ShouldIgnoreOwnershipInEditor(); })
-		),
-		NAME_None,
-		EUserInterfaceActionType::ToggleButton
-	);
 }
 
 bool FSequencerObjectBindingNode::CanRenameNode() const
@@ -273,24 +283,15 @@ FText FSequencerObjectBindingNode::GetDisplayName() const
 FLinearColor FSequencerObjectBindingNode::GetDisplayNameColor() const
 {
 	FSequencer& Sequencer = ParentTree.GetSequencer();
-	TArray<TWeakObjectPtr<UObject>> RuntimeObjects;
-	Sequencer.GetRuntimeObjects( Sequencer.GetFocusedMovieSceneSequenceInstance(), ObjectBinding, RuntimeObjects );
-	if ( RuntimeObjects.Num() == 0 )
-	{
-		return FLinearColor::Red;
-	}
-	else
-	{
-		return FSequencerDisplayNode::GetDisplayNameColor();
-	}
+	
+	const bool bHasObjects = Sequencer.FindBoundObjects(ObjectBinding, Sequencer.GetFocusedTemplateID()).Num() != 0;
+	return bHasObjects ? FSequencerDisplayNode::GetDisplayNameColor() : FLinearColor::Red;
 }
 
 FText FSequencerObjectBindingNode::GetDisplayNameToolTipText() const
 {
 	FSequencer& Sequencer = ParentTree.GetSequencer();
-	TArray<TWeakObjectPtr<UObject>> RuntimeObjects;
-	Sequencer.GetRuntimeObjects( Sequencer.GetFocusedMovieSceneSequenceInstance(), ObjectBinding, RuntimeObjects );
-	if ( RuntimeObjects.Num() == 0 )
+	if ( Sequencer.FindObjectsInCurrentSequence(ObjectBinding).Num() == 0 )
 	{
 		return LOCTEXT("InvalidBoundObjectToolTip", "The object bound to this track is missing.");
 	}
diff --git a/Engine/Source/Editor/Sequencer/Private/DisplayNodes/SequencerTrackNode.cpp b/Engine/Source/Editor/Sequencer/Private/DisplayNodes/SequencerTrackNode.cpp
--- a/Engine/Source/Editor/Sequencer/Private/DisplayNodes/SequencerTrackNode.cpp
+++ b/Engine/Source/Editor/Sequencer/Private/DisplayNodes/SequencerTrackNode.cpp
@@ -102,10 +102,65 @@ void FSequencerTrackNode::FixRowIndices()
 /* FSequencerDisplayNode interface
  *****************************************************************************/
 
+namespace
+{
+	void AddBoolPropertyMenuItem(FMenuBuilder& MenuBuilder, UMovieSceneTrack* Track, const UBoolProperty* Property, void* PropertyContainer)
+	{
+		MenuBuilder.AddMenuEntry(
+			Property->GetDisplayNameText(),
+			Property->GetToolTipText(),
+			FSlateIcon(),
+			FUIAction(
+				FExecuteAction::CreateLambda([Track, Property, PropertyContainer]{
+					FScopedTransaction Transaction(FText::Format(NSLOCTEXT("Sequencer", "TrackNodeSetRoundEvaluation", "Set '{0}'"), Property->GetDisplayNameText()));
+					Track->Modify();
+
+					bool bIsSet = Property->GetPropertyValue(PropertyContainer);
+					Property->SetPropertyValue(PropertyContainer, !bIsSet);
+				}),
+				FCanExecuteAction::CreateLambda([]{ return true; }),
+				FIsActionChecked::CreateLambda([PropertyContainer, Property]{ return Property->GetPropertyValue(PropertyContainer); })
+			),
+			NAME_None,
+			EUserInterfaceActionType::Check
+		);
+	}
+}
+
 void FSequencerTrackNode::BuildContextMenu(FMenuBuilder& MenuBuilder)
 {
 	AssociatedEditor.BuildTrackContextMenu(MenuBuilder, AssociatedTrack.Get());
 	FSequencerDisplayNode::BuildContextMenu(MenuBuilder );
+
+	MenuBuilder.BeginSection("GeneralTrackOptions", NSLOCTEXT("Sequencer", "TrackNodeGeneralOptions", "Track Options"));
+	{
+		UMovieSceneTrack* Track = AssociatedTrack.Get();
+		if (!Track)
+		{
+			return;
+		}
+
+		UStruct* EvalOptionsStruct = FMovieSceneTrackEvalOptions::StaticStruct();
+
+		const UBoolProperty* NearestSectionProperty = Cast<UBoolProperty>(EvalOptionsStruct->FindPropertyByName(GET_MEMBER_NAME_CHECKED(FMovieSceneTrackEvalOptions, bEvaluateNearestSection)));
+		if (NearestSectionProperty && Track->EvalOptions.bCanEvaluateNearestSection)
+		{
+			AddBoolPropertyMenuItem(MenuBuilder, Track, NearestSectionProperty, NearestSectionProperty->ContainerPtrToValuePtr<void>(&Track->EvalOptions));
+		}
+
+		const UBoolProperty* PrerollProperty = Cast<UBoolProperty>(EvalOptionsStruct->FindPropertyByName(GET_MEMBER_NAME_CHECKED(FMovieSceneTrackEvalOptions, bEvaluateInPreroll)));
+		if (PrerollProperty)
+		{
+			AddBoolPropertyMenuItem(MenuBuilder, Track, PrerollProperty, PrerollProperty->ContainerPtrToValuePtr<void>(&Track->EvalOptions));
+		}
+
+		const UBoolProperty* PostrollProperty = Cast<UBoolProperty>(EvalOptionsStruct->FindPropertyByName(GET_MEMBER_NAME_CHECKED(FMovieSceneTrackEvalOptions, bEvaluateInPostroll)));
+		if (PostrollProperty)
+		{
+			AddBoolPropertyMenuItem(MenuBuilder, Track, PostrollProperty, PostrollProperty->ContainerPtrToValuePtr<void>(&Track->EvalOptions));
+		}
+	}
+	MenuBuilder.EndSection();
 }
 
 
diff --git a/Engine/Source/Editor/Sequencer/Private/GroupedKeyArea.cpp b/Engine/Source/Editor/Sequencer/Private/GroupedKeyArea.cpp
--- a/Engine/Source/Editor/Sequencer/Private/GroupedKeyArea.cpp
+++ b/Engine/Source/Editor/Sequencer/Private/GroupedKeyArea.cpp
@@ -347,7 +347,7 @@ TArray<FKeyHandle> FGroupedKeyArea::GetUnsortedKeyHandles() const
 	return Array;
 }
 
-void FGroupedKeyArea::SetKeyTime(FKeyHandle KeyHandle, float NewKeyTime) const
+void FGroupedKeyArea::SetKeyTime(FKeyHandle KeyHandle, float NewKeyTime)
 {
 	auto* Group = FindGroup(KeyHandle);
 
diff --git a/Engine/Source/Editor/Sequencer/Private/SKeyNavigationButtons.h b/Engine/Source/Editor/Sequencer/Private/SKeyNavigationButtons.h
--- a/Engine/Source/Editor/Sequencer/Private/SKeyNavigationButtons.h
+++ b/Engine/Source/Editor/Sequencer/Private/SKeyNavigationButtons.h
@@ -126,29 +126,47 @@ public:
 	{
 		FSequencer& Sequencer = DisplayNode->GetSequencer();
 		float ClosestPreviousKeyDistance = MAX_FLT;
-		float CurrentTime = Sequencer.GetCurrentLocalTime(*Sequencer.GetFocusedMovieSceneSequence());
+		float CurrentTime = Sequencer.GetLocalTime();
 		float PreviousTime = 0;
 		bool PreviousKeyFound = false;
 
+		TArray<float> AllTimes;
+
 		TSet<TSharedPtr<IKeyArea>> KeyAreas;
-		SequencerHelpers::GetAllKeyAreas(DisplayNode, KeyAreas);
-		for (TSharedPtr<IKeyArea> KeyArea : KeyAreas)
+		SequencerHelpers::GetAllKeyAreas( DisplayNode, KeyAreas );
+		for ( TSharedPtr<IKeyArea> KeyArea : KeyAreas )
 		{
-			for (FKeyHandle& KeyHandle : KeyArea->GetUnsortedKeyHandles())
+			for ( FKeyHandle& KeyHandle : KeyArea->GetUnsortedKeyHandles() )
 			{
-				float KeyTime = KeyArea->GetKeyTime(KeyHandle);
-				if (KeyTime < CurrentTime && CurrentTime - KeyTime < ClosestPreviousKeyDistance)
-				{
-					PreviousTime = KeyTime;
-					ClosestPreviousKeyDistance = CurrentTime - KeyTime;
-					PreviousKeyFound = true;
-				}
+				float KeyTime = KeyArea->GetKeyTime( KeyHandle );
+				AllTimes.Add(KeyTime);
+			}
+		}
+
+		TSet<TWeakObjectPtr<UMovieSceneSection> > Sections;
+		SequencerHelpers::GetAllSections( DisplayNode.ToSharedRef(), Sections );				
+		for ( TWeakObjectPtr<UMovieSceneSection> Section : Sections )
+		{
+			if (Section.IsValid() && !Section->IsInfinite())
+			{
+				AllTimes.Add(Section->GetStartTime());
+				AllTimes.Add(Section->GetEndTime());
+			}
+		}
+
+		for (float Time : AllTimes)
+		{
+			if (Time < CurrentTime && CurrentTime - Time < ClosestPreviousKeyDistance)
+			{
+				PreviousTime = Time;
+				ClosestPreviousKeyDistance = CurrentTime - Time;
+				PreviousKeyFound = true;
 			}
 		}
 
 		if (PreviousKeyFound)
 		{
-			Sequencer.SetGlobalTime(PreviousTime);
+			Sequencer.SetLocalTime(PreviousTime);
 		}
 		return FReply::Handled();
 	}
@@ -158,29 +176,47 @@ public:
 	{
 		FSequencer& Sequencer = DisplayNode->GetSequencer();
 		float ClosestNextKeyDistance = MAX_FLT;
-		float CurrentTime = Sequencer.GetCurrentLocalTime(*Sequencer.GetFocusedMovieSceneSequence());
+		float CurrentTime = Sequencer.GetLocalTime();
 		float NextTime = 0;
 		bool NextKeyFound = false;
 
+		TArray<float> AllTimes;
+
 		TSet<TSharedPtr<IKeyArea>> KeyAreas;
-		SequencerHelpers::GetAllKeyAreas(DisplayNode, KeyAreas);
-		for (TSharedPtr<IKeyArea> KeyArea : KeyAreas)
+		SequencerHelpers::GetAllKeyAreas( DisplayNode, KeyAreas );
+		for ( TSharedPtr<IKeyArea> KeyArea : KeyAreas )
 		{
-			for (FKeyHandle& KeyHandle : KeyArea->GetUnsortedKeyHandles())
+			for ( FKeyHandle& KeyHandle : KeyArea->GetUnsortedKeyHandles() )
 			{
-				float KeyTime = KeyArea->GetKeyTime(KeyHandle);
-				if (KeyTime > CurrentTime && KeyTime - CurrentTime < ClosestNextKeyDistance)
-				{
-					NextTime = KeyTime;
-					ClosestNextKeyDistance = KeyTime - CurrentTime;
-					NextKeyFound = true;
-				}
+				float KeyTime = KeyArea->GetKeyTime( KeyHandle );
+				AllTimes.Add(KeyTime);
+			}
+		}
+
+		TSet<TWeakObjectPtr<UMovieSceneSection> > Sections;
+		SequencerHelpers::GetAllSections( DisplayNode.ToSharedRef(), Sections );				
+		for ( TWeakObjectPtr<UMovieSceneSection> Section : Sections )
+		{
+			if (Section.IsValid() && !Section->IsInfinite())
+			{
+				AllTimes.Add(Section->GetStartTime());
+				AllTimes.Add(Section->GetEndTime());
+			}
+		}
+
+		for (float Time : AllTimes)
+		{
+			if (Time > CurrentTime && Time - CurrentTime < ClosestNextKeyDistance)
+			{
+				NextTime = Time;
+				ClosestNextKeyDistance = Time - CurrentTime;
+				NextKeyFound = true;
 			}
 		}
 
 		if (NextKeyFound)
 		{
-			Sequencer.SetGlobalTime(NextTime);
+			Sequencer.SetLocalTime(NextTime);
 		}
 
 		return FReply::Handled();
@@ -190,7 +226,7 @@ public:
 	FReply OnAddKeyClicked()
 	{
 		FSequencer& Sequencer = DisplayNode->GetSequencer();
-		float CurrentTime = Sequencer.GetCurrentLocalTime(*Sequencer.GetFocusedMovieSceneSequence());
+		float CurrentTime = Sequencer.GetLocalTime();
 
 		TSet<TSharedPtr<IKeyArea>> KeyAreas;
 		SequencerHelpers::GetAllKeyAreas(DisplayNode, KeyAreas);
diff --git a/Engine/Source/Editor/Sequencer/Private/SSequencer.cpp b/Engine/Source/Editor/Sequencer/Private/SSequencer.cpp
--- a/Engine/Source/Editor/Sequencer/Private/SSequencer.cpp
+++ b/Engine/Source/Editor/Sequencer/Private/SSequencer.cpp
@@ -13,6 +13,7 @@
 #include "SSequencerSectionOverlay.h"
 #include "SSequencerTrackArea.h"
 #include "SSequencerTrackOutliner.h"
+#include "SSequencerDebugVisualizer.h"
 #include "SequencerNodeTree.h"
 #include "SequencerTimeSliderController.h"
 #include "DragAndDrop/AssetDragDropOp.h"
@@ -37,7 +38,7 @@
 #include "NumericUnitTypeInterface.inl"
 #include "SNumericEntryBox.h"
 #include "SSequencerGotoBox.h"
-
+#include "SSequencerTransformBox.h"
 
 #define LOCTEXT_NAMESPACE "Sequencer"
 
@@ -142,6 +143,7 @@ void SSequencer::Construct(const FArguments& InArgs, TSharedRef<FSequencer> InSe
 		TimeSliderArgs.OnEndScrubberMovement = InArgs._OnEndScrubbing;
 		TimeSliderArgs.OnScrubPositionChanged = InArgs._OnScrubPositionChanged;
 		TimeSliderArgs.PlaybackStatus = InArgs._PlaybackStatus;
+		TimeSliderArgs.SubSequenceRange = InArgs._SubSequenceRange;
 
 		TimeSliderArgs.Settings = Settings;
 	}
@@ -213,7 +215,7 @@ void SSequencer::Construct(const FArguments& InArgs, TSharedRef<FSequencer> InSe
 	TrackArea->SetTreeView(TreeView);
 
 	const int32 Column0 = 0, Column1 = 1;
-	const int32 Row0 = 0, Row1 = 1, Row2 = 2, Row3 = 3;
+	const int32 Row0 = 0, Row1 = 1, Row2 = 2, Row3 = 3, Row4 = 4;
 
 	const float CommonPadding = 3.f;
 	const FMargin ResizeBarPadding(4.f, 0, 0, 0);
@@ -381,7 +383,7 @@ void SSequencer::Construct(const FArguments& InArgs, TSharedRef<FSequencer> InSe
 					]
 
 					// playback buttons
-					+ SGridPanel::Slot( Column0, Row3, SGridPanel::Layer(10) )
+					+ SGridPanel::Slot( Column0, Row4, SGridPanel::Layer(10) )
 					[
 						SNew(SBorder)
 						.BorderImage(FEditorStyle::GetBrush("ToolPanel.GroupBorder"))
@@ -454,8 +456,25 @@ void SSequencer::Construct(const FArguments& InArgs, TSharedRef<FSequencer> InSe
 							SAssignNew(GotoBox, SSequencerGotoBox, SequencerPtr.Pin().ToSharedRef(), *Settings, NumericTypeInterface.ToSharedRef())
 						]
 
-					// play range sliders
+					// Transform box
+					+ SGridPanel::Slot(Column1, Row2, SGridPanel::Layer(50))
+						.Padding(ResizeBarPadding)
+						.HAlign(HAlign_Left)
+						.VAlign(VAlign_Top)
+						[
+							SAssignNew(TransformBox, SSequencerTransformBox, SequencerPtr.Pin().ToSharedRef(), *Settings, NumericTypeInterface.ToSharedRef())
+						]
+
+					// debug vis
 					+ SGridPanel::Slot( Column1, Row3, SGridPanel::Layer(10) )
+					[
+						SNew(SSequencerDebugVisualizer, InSequencer)
+						.ViewRange(FAnimatedRange::WrapAttribute(InArgs._ViewRange))
+						.Visibility(this, &SSequencer::GetDebugVisualizerVisibility)
+					]
+
+					// play range sliders
+					+ SGridPanel::Slot( Column1, Row4, SGridPanel::Layer(10) )
 					.Padding(ResizeBarPadding)
 					[
 						SNew( SBorder )
@@ -538,6 +557,11 @@ void SSequencer::BindCommands(TSharedRef<FUICommandList> SequencerCommandBinding
 		FSequencerCommands::Get().ToggleShowGotoBox,
 		FExecuteAction::CreateLambda([this]{ GotoBox->ToggleVisibility(); })
 	);
+
+	SequencerCommandBindings->MapAction(
+		FSequencerCommands::Get().ToggleShowTransformBox,
+		FExecuteAction::CreateLambda([this]{ TransformBox->ToggleVisibility(); })
+	);
 }
 	
 const ISequencerEditTool* SSequencer::GetEditTool() const
@@ -929,6 +953,7 @@ TSharedRef<SWidget> SSequencer::MakeGeneralMenu()
 
 	MenuBuilder.BeginSection( "PlaybackAllSequences", LOCTEXT( "PlaybackRangeAllSequencesHeader", "Playback Range - All Sequences" ) );
 	{
+		MenuBuilder.AddMenuEntry( FSequencerCommands::Get().ToggleEvaluateSubSequencesInIsolation );
 		MenuBuilder.AddMenuEntry( FSequencerCommands::Get().ToggleKeepCursorInPlaybackRange );
 		MenuBuilder.AddMenuEntry( FSequencerCommands::Get().ToggleKeepPlaybackRangeInSectionBounds );
 		MenuBuilder.AddMenuEntry( FSequencerCommands::Get().ToggleLinkCurveEditorTimeRange );
@@ -960,6 +985,7 @@ TSharedRef<SWidget> SSequencer::MakeGeneralMenu()
 			LOCTEXT("ZeroPaddingText", "Zero Pad Frame Numbers"));
 	}
 	MenuBuilder.AddMenuEntry(FSequencerCommands::Get().ToggleShowGotoBox);
+	MenuBuilder.AddMenuEntry(FSequencerCommands::Get().ToggleShowTransformBox);
 	MenuBuilder.EndSection();
 
 	// selection range options
@@ -1164,18 +1190,23 @@ void SSequencer::UpdateLayoutTree()
 
 void SSequencer::UpdateBreadcrumbs()
 {
-	TSharedRef<FMovieSceneSequenceInstance> FocusedMovieSceneInstance = SequencerPtr.Pin()->GetFocusedMovieSceneSequenceInstance();
+	TSharedPtr<FSequencer> Sequencer = SequencerPtr.Pin();
+	if (!Sequencer.IsValid())
+	{
+		return;
+	}
 
+	FMovieSceneSequenceID FocusedID = Sequencer->GetFocusedTemplateID();
 	if (BreadcrumbTrail->PeekCrumb().BreadcrumbType == FSequencerBreadcrumb::ShotType)
 	{
 		BreadcrumbTrail->PopCrumb();
 	}
 
-	if( BreadcrumbTrail->PeekCrumb().BreadcrumbType == FSequencerBreadcrumb::MovieSceneType && BreadcrumbTrail->PeekCrumb().MovieSceneInstance.Pin() != FocusedMovieSceneInstance )
+	if( BreadcrumbTrail->PeekCrumb().BreadcrumbType == FSequencerBreadcrumb::MovieSceneType && BreadcrumbTrail->PeekCrumb().SequenceID != FocusedID )
 	{
-		FText CrumbName = FocusedMovieSceneInstance->GetSequence()->GetDisplayName();
+		FText CrumbName = Sequencer->GetFocusedMovieSceneSequence()->GetDisplayName();
 		// The current breadcrumb is not a moviescene so we need to make a new breadcrumb in order return to the parent moviescene later
-		BreadcrumbTrail->PushCrumb( CrumbName, FSequencerBreadcrumb( FocusedMovieSceneInstance ) );
+		BreadcrumbTrail->PushCrumb( CrumbName, FSequencerBreadcrumb( FocusedID ) );
 	}
 }
 
@@ -1183,7 +1214,7 @@ void SSequencer::UpdateBreadcrumbs()
 void SSequencer::ResetBreadcrumbs()
 {
 	BreadcrumbTrail->ClearCrumbs();
-	BreadcrumbTrail->PushCrumb(TAttribute<FText>::Create(TAttribute<FText>::FGetter::CreateSP(this, &SSequencer::GetRootAnimationName)), FSequencerBreadcrumb(SequencerPtr.Pin()->GetRootMovieSceneSequenceInstance()));
+	BreadcrumbTrail->PushCrumb(TAttribute<FText>::Create(TAttribute<FText>::FGetter::CreateSP(this, &SSequencer::GetRootAnimationName)), FSequencerBreadcrumb(MovieSceneSequenceID::Root));
 }
 
 
@@ -1454,7 +1485,7 @@ void SSequencer::OnCrumbClicked(const FSequencerBreadcrumb& Item)
 {
 	if (Item.BreadcrumbType != FSequencerBreadcrumb::ShotType)
 	{
-		if( SequencerPtr.Pin()->GetFocusedMovieSceneSequenceInstance() == Item.MovieSceneInstance.Pin() ) 
+		if( SequencerPtr.Pin()->GetFocusedTemplateID() == Item.SequenceID ) 
 		{
 			// then do zooming
 		}
@@ -1465,7 +1496,7 @@ void SSequencer::OnCrumbClicked(const FSequencerBreadcrumb& Item)
 				SequencerPtr.Pin()->SetShowCurveEditor(false);
 			}
 
-			SequencerPtr.Pin()->PopToSequenceInstance( Item.MovieSceneInstance.Pin().ToSharedRef() );
+			SequencerPtr.Pin()->PopToSequenceInstance( Item.SequenceID );
 		}
 	}
 }
@@ -1580,12 +1611,10 @@ void SSequencer::StepToKey(bool bStepToNextKey, bool bCameraOnly)
 			for ( auto RootNode : RootNodes )
 			{
 				TSharedRef<FSequencerObjectBindingNode> ObjectBindingNode = StaticCastSharedRef<FSequencerObjectBindingNode>( RootNode );
-				TArray<TWeakObjectPtr<UObject>> RuntimeObjects;
-				Sequencer->GetRuntimeObjects( Sequencer->GetFocusedMovieSceneSequenceInstance(), ObjectBindingNode->GetObjectBinding(), RuntimeObjects );
 
-				for ( int32 RuntimeIndex = 0; RuntimeIndex < RuntimeObjects.Num(); ++RuntimeIndex )
+				for (TWeakObjectPtr<>& Object : Sequencer->FindObjectsInCurrentSequence(ObjectBindingNode->GetObjectBinding()))
 				{
-					AActor* RuntimeActor = Cast<AActor>( RuntimeObjects[RuntimeIndex].Get() );
+					AActor* RuntimeActor = Cast<AActor>( Object.Get() );
 					if ( RuntimeActor != nullptr && LockedActors.Contains( RuntimeActor ) )
 					{
 						Nodes.Add( RootNode );
@@ -1613,7 +1642,7 @@ void SSequencer::StepToKey(bool bStepToNextKey, bool bCameraOnly)
 		if ( Nodes.Num() > 0 )
 		{
 			float ClosestKeyDistance = MAX_FLT;
-			float CurrentTime = Sequencer->GetCurrentLocalTime( *Sequencer->GetFocusedMovieSceneSequence() );
+			float CurrentTime = Sequencer->GetLocalTime();
 			float StepToTime = 0;
 			bool StepToKeyFound = false;
 
@@ -1622,31 +1651,48 @@ void SSequencer::StepToKey(bool bStepToNextKey, bool bCameraOnly)
 
 			for ( auto Node : Nodes )
 			{
+				TArray<float> AllTimes;
+
 				TSet<TSharedPtr<IKeyArea>> KeyAreas;
 				SequencerHelpers::GetAllKeyAreas( Node, KeyAreas );
-
 				for ( TSharedPtr<IKeyArea> KeyArea : KeyAreas )
 				{
 					for ( FKeyHandle& KeyHandle : KeyArea->GetUnsortedKeyHandles() )
 					{
 						float KeyTime = KeyArea->GetKeyTime( KeyHandle );
-						if ( bStepToNextKey )
+						AllTimes.Add(KeyTime);
+					}
+				}
+
+				TSet<TWeakObjectPtr<UMovieSceneSection> > Sections;
+				SequencerHelpers::GetAllSections( Node, Sections );				
+				for ( TWeakObjectPtr<UMovieSceneSection> Section : Sections )
+				{
+					if (Section.IsValid() && !Section->IsInfinite())
+					{
+						AllTimes.Add(Section->GetStartTime());
+						AllTimes.Add(Section->GetEndTime());
+					}
+				}
+
+				for (float Time : AllTimes)
+				{
+					if ( bStepToNextKey )
+					{
+						if ( Time > CurrentTime && Time - CurrentTime < ClosestKeyDistance )
 						{
-							if ( KeyTime > CurrentTime && KeyTime - CurrentTime < ClosestKeyDistance )
-							{
-								StepToTime = KeyTime;
-								ClosestKeyDistance = KeyTime - CurrentTime;
-								StepToKeyFound = true;
-							}
+							StepToTime = Time;
+							ClosestKeyDistance = Time - CurrentTime;
+							StepToKeyFound = true;
 						}
-						else
+					}
+					else
+					{
+						if ( Time < CurrentTime && CurrentTime - Time < ClosestKeyDistance )
 						{
-							if ( KeyTime < CurrentTime && CurrentTime - KeyTime < ClosestKeyDistance )
-							{
-								StepToTime = KeyTime;
-								ClosestKeyDistance = CurrentTime - KeyTime;
-								StepToKeyFound = true;
-							}
+							StepToTime = Time;
+							ClosestKeyDistance = CurrentTime - Time;
+							StepToKeyFound = true;
 						}
 					}
 				}
@@ -1654,7 +1700,7 @@ void SSequencer::StepToKey(bool bStepToNextKey, bool bCameraOnly)
 
 			if ( StepToKeyFound )
 			{
-				Sequencer->SetGlobalTime( StepToTime );
+				Sequencer->SetLocalTime( StepToTime );
 			}
 		}
 	}
@@ -1912,7 +1958,7 @@ void SSequencer::OpenPasteMenu()
 	TSharedPtr<FSequencer> Sequencer = SequencerPtr.Pin();
 	if (Sequencer->GetClipboardStack().Num() != 0)
 	{
-		FPasteContextMenuArgs Args = GeneratePasteArgs(Sequencer->GetGlobalTime(), Sequencer->GetClipboardStack().Last());
+		FPasteContextMenuArgs Args = GeneratePasteArgs(Sequencer->GetLocalTime(), Sequencer->GetClipboardStack().Last());
 		ContextMenu = FPasteContextMenu::CreateMenu(*Sequencer, Args);
 	}
 
@@ -1950,7 +1996,7 @@ void SSequencer::PasteFromHistory()
 		return;
 	}
 
-	FPasteContextMenuArgs Args = GeneratePasteArgs(Sequencer->GetGlobalTime());
+	FPasteContextMenuArgs Args = GeneratePasteArgs(Sequencer->GetLocalTime());
 	TSharedPtr<FPasteFromHistoryContextMenu> ContextMenu = FPasteFromHistoryContextMenu::CreateMenu(*Sequencer, Args);
 
 	if (ContextMenu.IsValid())
@@ -1973,7 +2019,7 @@ void SSequencer::PasteFromHistory()
 	}
 }
 
-void SSequencer::OnSequenceInstanceActivated( FMovieSceneSequenceInstance& ActiveInstance )
+void SSequencer::OnSequenceInstanceActivated( FMovieSceneSequenceIDRef ActiveInstanceID )
 {
 	TSharedPtr<FSequencer> Sequencer = SequencerPtr.Pin();
 	if ( Sequencer.IsValid() )
@@ -1986,5 +2032,10 @@ void SSequencer::OnSequenceInstanceActivated( FMovieSceneSequenceInstance& Activ
 	}
 }
 
+EVisibility SSequencer::GetDebugVisualizerVisibility() const
+{
+	return Settings->ShouldShowDebugVisualization() ? EVisibility::Visible : EVisibility::Collapsed;
+}
+
 #undef LOCTEXT_NAMESPACE
 
diff --git a/Engine/Source/Editor/Sequencer/Private/SSequencer.h b/Engine/Source/Editor/Sequencer/Private/SSequencer.h
--- a/Engine/Source/Editor/Sequencer/Private/SSequencer.h
+++ b/Engine/Source/Editor/Sequencer/Private/SSequencer.h
@@ -12,7 +12,6 @@ class FActorDragDropGraphEdOp;
 class FAssetDragDropOp;
 class FClassDragDropOp;
 class FEditPropertyChain;
-class FMovieSceneSequenceInstance;
 class FSequencer;
 class FSequencerNodeTree;
 class FUnloadedClassDragDropOp;
@@ -23,6 +22,7 @@ class SSequencerGotoBox;
 class SSequencerLabelBrowser;
 class SSequencerTrackArea;
 class SSequencerTrackOutliner;
+class SSequencerTransformBox;
 class SSequencerTreeView;
 class FSequencerTimeSliderController;
 
@@ -67,11 +67,11 @@ struct FSequencerBreadcrumb
 	FSequencerBreadcrumb::Type BreadcrumbType;
 
 	/** The movie scene this may point to */
-	TWeakPtr<FMovieSceneSequenceInstance> MovieSceneInstance;
+	FMovieSceneSequenceID SequenceID;
 
-	FSequencerBreadcrumb(TSharedRef<FMovieSceneSequenceInstance> InMovieSceneInstance)
+	FSequencerBreadcrumb(FMovieSceneSequenceIDRef InSequenceID)
 		: BreadcrumbType(FSequencerBreadcrumb::MovieSceneType)
-		, MovieSceneInstance(InMovieSceneInstance)
+		, SequenceID(InSequenceID)
 	{ }
 
 	FSequencerBreadcrumb()
@@ -103,9 +103,12 @@ public:
 		/** The playback range */
 		SLATE_ATTRIBUTE( TRange<float>, PlaybackRange )
 
-		/** The selection selection range */
+		/** The selection range */
 		SLATE_ATTRIBUTE(TRange<float>, SelectionRange)
 
+		/** The current sub sequence range */
+		SLATE_ATTRIBUTE(TOptional<TRange<float>>, SubSequenceRange)
+
 		/** The playback status */
 		SLATE_ATTRIBUTE( EMovieScenePlayerStatus::Type, PlaybackStatus )
 
@@ -386,8 +389,10 @@ private:
 	FPaintPlaybackRangeArgs GetSectionPlaybackRangeArgs() const;
 
 	/** Called whenever the active sequence instance changes on the FSequencer */
-	void OnSequenceInstanceActivated( FMovieSceneSequenceInstance& ActiveInstance );
+	void OnSequenceInstanceActivated( FMovieSceneSequenceIDRef ActiveInstanceID );
 
+	EVisibility GetDebugVisualizerVisibility() const;
+	
 public:
 	/** On Paste Command */
 	void OnPaste();
@@ -410,6 +415,9 @@ private:
 	/** Goto box widget. */
 	TSharedPtr<SSequencerGotoBox> GotoBox;
 
+	/** Transform box widget. */
+	TSharedPtr<SSequencerTransformBox> TransformBox;
+
 	/** Section area widget */
 	TSharedPtr<SSequencerTrackArea> TrackArea;
 
diff --git a/Engine/Source/Editor/Sequencer/Private/SSequencerDebugVisualizer.cpp b/Engine/Source/Editor/Sequencer/Private/SSequencerDebugVisualizer.cpp
new file mode 100644
--- /dev/null
+++ b/Engine/Source/Editor/Sequencer/Private/SSequencerDebugVisualizer.cpp
@@ -0,0 +1,224 @@
+// Copyright 1998-2016 Epic Games, Inc. All Rights Reserved.
+
+#include "SequencerPrivatePCH.h"
+#include "SSequencerDebugVisualizer.h"
+#include "CommonMovieSceneTools.h"
+
+#define LOCTEXT_NAMESPACE "SSequencerDebugVisualizer"
+
+void SSequencerDebugVisualizer::Construct(const FArguments& InArgs, TSharedRef<FSequencer> InSequencer)
+{
+	WeakSequencer = InSequencer;
+
+	InSequencer->GetSequenceInstance().OnUpdated().AddSP(this, &SSequencerDebugVisualizer::Refresh);
+	InSequencer->OnActivateSequence().AddSP(this, &SSequencerDebugVisualizer::OnSequenceActivated);
+
+	ViewRange = InArgs._ViewRange;
+
+	Refresh();
+}
+
+void SSequencerDebugVisualizer::OnSequenceActivated(FMovieSceneSequenceIDRef)
+{
+	Refresh();
+}
+
+const FMovieSceneEvaluationTemplateInstance* SSequencerDebugVisualizer::GetTemplate() const
+{
+	TSharedPtr<FSequencer> Sequencer = WeakSequencer.Pin();
+	return Sequencer.IsValid() ? Sequencer->GetSequenceInstance().GetInstance(FocusedSequenceID) : nullptr;
+}
+
+void SSequencerDebugVisualizer::Refresh()
+{
+	Children.Empty();
+
+	TSharedPtr<FSequencer> Sequencer = WeakSequencer.Pin();
+	if (!Sequencer.IsValid())
+	{
+		return;
+	}
+
+	FocusedSequenceID = Sequencer->GetFocusedTemplateID();
+
+	const FMovieSceneEvaluationTemplateInstance* ActiveInstance = GetTemplate();
+	if (!ActiveInstance)
+	{
+		return;
+	}
+
+	const FMovieSceneEvaluationField& EvaluationField = ActiveInstance->Template->EvaluationField;
+
+	TArray<int32> SegmentComplexity;
+	SegmentComplexity.Reserve(EvaluationField.Groups.Num());
+
+	// Heatmap complexity
+	float AverageComplexity = 0;
+	int32 MaxComplexity = 0;
+
+	for (const FMovieSceneEvaluationGroup& Group : EvaluationField.Groups)
+	{
+		int32 Complexity = 0;
+		for (const FMovieSceneEvaluationGroupLUTIndex& LUTIndex : Group.LUTIndices)
+		{
+			// more groups is more complex
+			Complexity += 1;
+			// Add total init and eval complexity
+			Complexity += LUTIndex.NumInitPtrs + LUTIndex.NumEvalPtrs;
+		}
+
+		SegmentComplexity.Add(Complexity);
+		MaxComplexity = FMath::Max(MaxComplexity, Complexity);
+		AverageComplexity += Complexity;
+	}
+
+	AverageComplexity /= SegmentComplexity.Num();
+
+	static const FSlateBrush* SectionBackgroundBrush = FEditorStyle::GetBrush("Sequencer.Section.Background");
+	static const FSlateBrush* SectionBackgroundTintBrush = FEditorStyle::GetBrush("Sequencer.Section.BackgroundTint");
+
+	for (int32 Index = 0; Index < EvaluationField.Ranges.Num(); ++Index)
+	{
+		TRange<float> Range = EvaluationField.Ranges[Index];
+
+		const float Complexity = SegmentComplexity[Index];
+		
+		float Lerp = FMath::Clamp((Complexity - AverageComplexity) / (MaxComplexity - AverageComplexity), 0.f, 1.f) * 0.5f +
+			FMath::Clamp(Complexity / AverageComplexity, 0.f, 1.f) * 0.5f;
+
+		// Blend from blue (240deg) to red (0deg)
+		FLinearColor ComplexityColor = FLinearColor(FMath::Lerp(240.f, 0.f, FMath::Clamp(Lerp, 0.f, 1.f)), 1.f, 1.f, 0.5f).HSVToLinearRGB();
+
+		Children.Add(
+			SNew(SSequencerDebugSlot, Index)
+			.Visibility(this, &SSequencerDebugVisualizer::GetSegmentVisibility, Range)
+			.ToolTip(
+				SNew(SToolTip)
+				[
+					GetTooltipForSegment(Index)
+				]
+			)
+			[
+				SNew(SBorder)
+				.BorderImage(SectionBackgroundBrush)
+				.Padding(FMargin(1.f))
+				[
+					SNew(SBorder)
+					.BorderImage(SectionBackgroundTintBrush)
+					.BorderBackgroundColor(ComplexityColor)
+					.ForegroundColor(FLinearColor::Black)
+					[
+						SNew(STextBlock)
+						.Text(FText::AsNumber(Index))
+					]
+				]
+			]
+		);
+	}
+}
+
+FVector2D SSequencerDebugVisualizer::ComputeDesiredSize(float) const
+{
+	// Note: X Size is not used
+	FVector2D Size(100, 0.f);
+	if (Children.Num())
+	{
+		for (int32 Index = 0; Index < Children.Num(); ++Index)
+		{
+			Size.Y = FMath::Max(Size.Y, Children[Index]->GetDesiredSize().Y);
+		}
+	}
+	return Size;
+}
+
+FGeometry SSequencerDebugVisualizer::GetSegmentGeometry(const FGeometry& AllottedGeometry, const SSequencerDebugSlot& Slot, const FTimeToPixel& TimeToPixelConverter) const
+{
+	const FMovieSceneEvaluationTemplateInstance* ActiveInstance = GetTemplate();
+	if (!ActiveInstance)
+	{
+		return AllottedGeometry.MakeChild(FVector2D(0,0), FVector2D(0,0));
+	}
+
+	TRange<float> SegmentRange = ActiveInstance->Template->EvaluationField.Ranges[Slot.GetSegmentIndex()];
+
+	float PixelStartX = SegmentRange.GetLowerBound().IsOpen() ? 0.f : TimeToPixelConverter.TimeToPixel(SegmentRange.GetLowerBoundValue());
+	float PixelEndX = SegmentRange.GetUpperBound().IsOpen() ? AllottedGeometry.GetDrawSize().X : TimeToPixelConverter.TimeToPixel(SegmentRange.GetUpperBoundValue());
+
+	const float MinSectionWidth = 1.f;
+	float SectionLength = FMath::Max(MinSectionWidth, PixelEndX - PixelStartX);
+
+	return AllottedGeometry.MakeChild(
+		FVector2D(PixelStartX, 0),
+		FVector2D(SectionLength, Slot.GetDesiredSize().Y)
+		);
+}
+
+EVisibility SSequencerDebugVisualizer::GetSegmentVisibility(TRange<float> Range) const
+{
+	return ViewRange.Get().Overlaps(Range) ? EVisibility::Visible : EVisibility::Collapsed;
+}
+
+TSharedRef<SWidget> SSequencerDebugVisualizer::GetTooltipForSegment(int32 SegmentIndex) const
+{
+	const FMovieSceneEvaluationTemplateInstance* ActiveInstance = GetTemplate();
+	if (!ActiveInstance)
+	{
+		return SNullWidget::NullWidget;
+	}
+
+	const FMovieSceneEvaluationGroup& Group = ActiveInstance->Template->EvaluationField.Groups[SegmentIndex];
+
+	TSharedRef<SVerticalBox> VerticalBox = SNew(SVerticalBox);
+
+	int32 NumInitEntities = 0;
+	for (int32 Index = 0; Index < Group.LUTIndices.Num(); ++Index)
+	{
+		VerticalBox->AddSlot()
+		[
+			SNew(SVerticalBox)
+
+			+ SVerticalBox::Slot()
+			.AutoHeight()
+			[
+				SNew(STextBlock)
+				.Text(FText::Format(LOCTEXT("EvalGroupFormat", "Evaluation Group {0}:"), Index))
+			]
+
+			+ SVerticalBox::Slot()
+			.AutoHeight()
+			[
+				SNew(STextBlock)
+				.Text(FText::Format(
+					LOCTEXT("EvalTrackFormat", "{0} initialization steps, {1} evaluation steps"),
+					FText::AsNumber(Group.LUTIndices[Index].NumInitPtrs),
+					FText::AsNumber(Group.LUTIndices[Index].NumEvalPtrs)
+				))
+			]
+		];
+	}
+
+	return VerticalBox;
+}
+
+void SSequencerDebugVisualizer::OnArrangeChildren( const FGeometry& AllottedGeometry, FArrangedChildren& ArrangedChildren ) const
+{
+	FTimeToPixel TimeToPixelConverter = FTimeToPixel(AllottedGeometry, ViewRange.Get());
+
+	for (int32 WidgetIndex = 0; WidgetIndex < Children.Num(); ++WidgetIndex)
+	{
+		const TSharedRef<SSequencerDebugSlot>& Child = Children[WidgetIndex];
+
+		EVisibility WidgetVisibility = Child->GetVisibility();
+		if( ArrangedChildren.Accepts( WidgetVisibility ) )
+		{
+			FGeometry SegmentGeometry = GetSegmentGeometry(AllottedGeometry, *Child, TimeToPixelConverter);
+			
+			ArrangedChildren.AddWidget( 
+				WidgetVisibility, 
+				AllottedGeometry.MakeChild(Child, SegmentGeometry.Position, SegmentGeometry.GetDrawSize())
+				);
+		}
+	}
+}
+
+#undef LOCTEXT_NAMESPACE
\ No newline at end of file
diff --git a/Engine/Source/Editor/Sequencer/Private/SSequencerDebugVisualizer.h b/Engine/Source/Editor/Sequencer/Private/SSequencerDebugVisualizer.h
new file mode 100644
--- /dev/null
+++ b/Engine/Source/Editor/Sequencer/Private/SSequencerDebugVisualizer.h
@@ -0,0 +1,69 @@
+// Copyright 1998-2016 Epic Games, Inc. All Rights Reserved.
+
+#pragma once
+
+struct FTimeToPixel;
+
+class SSequencerDebugSlot : public SCompoundWidget
+{
+	SLATE_BEGIN_ARGS(SSequencerDebugSlot){}
+		SLATE_DEFAULT_SLOT(FArguments, Content)
+	SLATE_END_ARGS()
+
+	void Construct(const FArguments& InArgs, int32 InSegmentIndex)
+	{
+		SegmentIndex = InSegmentIndex;
+
+		ChildSlot
+		[
+			InArgs._Content.Widget
+		];
+	}
+
+	int32 GetSegmentIndex() const { return SegmentIndex; }
+
+private:
+	int32 SegmentIndex;
+};
+
+class SSequencerDebugVisualizer : public SPanel
+{
+public:
+	SLATE_BEGIN_ARGS(SSequencerDebugVisualizer){}
+		SLATE_ATTRIBUTE( TRange<float>, ViewRange )
+	SLATE_END_ARGS()
+
+	void Construct(const FArguments& InArgs, TSharedRef<FSequencer> InSequencer);
+
+protected:
+	/** SPanel Interface */
+	virtual void OnArrangeChildren( const FGeometry& AllottedGeometry, FArrangedChildren& ArrangedChildren ) const override;
+	virtual FVector2D ComputeDesiredSize(float) const override;
+	virtual FChildren* GetChildren() override { return &Children; }
+
+protected:
+	
+	void Refresh();
+
+	FGeometry GetSegmentGeometry(const FGeometry& AllottedGeometry, const SSequencerDebugSlot& Slot, const FTimeToPixel& TimeToPixelConverter) const;
+
+	EVisibility GetSegmentVisibility(TRange<float> Range) const;
+
+	TSharedRef<SWidget> GetTooltipForSegment(int32 SegmentIndex) const;
+
+	void OnSequenceActivated(FMovieSceneSequenceIDRef);
+
+	const FMovieSceneEvaluationTemplateInstance* GetTemplate() const;
+
+private:
+
+	FMovieSceneSequenceID FocusedSequenceID;
+
+	/** The current view range */
+	TAttribute<TRange<float>> ViewRange;
+
+	/** All the widgets in the panel */
+	TSlotlessChildren<SSequencerDebugSlot> Children;
+
+	TWeakPtr<FSequencer> WeakSequencer;
+};
\ No newline at end of file
diff --git a/Engine/Source/Editor/Sequencer/Private/SSequencerGotoBox.cpp b/Engine/Source/Editor/Sequencer/Private/SSequencerGotoBox.cpp
--- a/Engine/Source/Editor/Sequencer/Private/SSequencerGotoBox.cpp
+++ b/Engine/Source/Editor/Sequencer/Private/SSequencerGotoBox.cpp
@@ -40,7 +40,7 @@ void SSequencerGotoBox::Construct(const FArguments& InArgs, const TSharedRef<FSe
 							.MinDesiredValueWidth(64.0f)
 							.OnValueCommitted(this, &SSequencerGotoBox::HandleEntryBoxValueCommitted)
 							.TypeInterface(NumericTypeInterface)
-							.Value_Lambda([this](){ return SequencerPtr.Pin()->GetGlobalTime(); })
+							.Value_Lambda([this](){ return SequencerPtr.Pin()->GetLocalTime(); })
 					]
 			]
 	];
@@ -86,7 +86,7 @@ void SSequencerGotoBox::HandleEntryBoxValueCommitted(float Value, ETextCommit::T
 		Sequencer->SetViewRange(NewRange);
 	}
 
-	Sequencer->SetGlobalTime(Value);
+	Sequencer->SetLocalTimeDirectly(Value);
 }
 
 
diff --git a/Engine/Source/Editor/Sequencer/Private/SSequencerTransformBox.cpp b/Engine/Source/Editor/Sequencer/Private/SSequencerTransformBox.cpp
new file mode 100644
--- /dev/null
+++ b/Engine/Source/Editor/Sequencer/Private/SSequencerTransformBox.cpp
@@ -0,0 +1,139 @@
+// Copyright 1998-2016 Epic Games, Inc. All Rights Reserved.
+
+#include "SequencerPrivatePCH.h"
+#include "Sequencer.h"
+#include "SequencerSettings.h"
+#include "SSequencerTransformBox.h"
+
+#define LOCTEXT_NAMESPACE "Sequencer"
+
+void SSequencerTransformBox::Construct(const FArguments& InArgs, const TSharedRef<FSequencer>& InSequencer, USequencerSettings& InSettings, const TSharedRef<INumericTypeInterface<float>>& InNumericTypeInterface)
+{
+	SequencerPtr = InSequencer;
+	Settings = &InSettings;
+	NumericTypeInterface = InNumericTypeInterface;
+	
+	DeltaTime = 1.f;
+
+	// Initialize to 10 frames if showing frame numbers
+	float TimeSnapInterval = SequencerPtr.Pin()->GetSettings()->GetTimeSnapInterval();
+	if (SequencerSnapValues::IsTimeSnapIntervalFrameRate(TimeSnapInterval))
+	{	
+		float FrameRate = 1.0f / TimeSnapInterval;
+		DeltaTime = SequencerHelpers::FrameToTime(10, FrameRate);
+	}
+
+	const FDockTabStyle* GenericTabStyle = &FCoreStyle::Get().GetWidgetStyle<FDockTabStyle>("Docking.Tab");
+	const FButtonStyle* const CloseButtonStyle = &GenericTabStyle->CloseButtonStyle;
+
+	ChildSlot
+	[
+		SAssignNew(Border, SBorder)
+			.BorderImage(FEditorStyle::GetBrush("ToolPanel.DarkGroupBorder"))
+			.Padding(6.0f)
+			.Visibility(EVisibility::Collapsed)
+			[
+				SNew(SHorizontalBox)
+
+				+ SHorizontalBox::Slot()
+					.AutoWidth()
+					.VAlign(VAlign_Center)
+					[
+						SNew(SButton)
+							.Text(LOCTEXT("PlusLabel", "+"))
+							.OnClicked(this, &SSequencerTransformBox::OnPlusButtonClicked)
+					]
+
+				+ SHorizontalBox::Slot()
+					.AutoWidth()
+					.VAlign(VAlign_Center)
+					[
+						SNew(SButton)
+							.Text(LOCTEXT("MinusLabel", "-"))
+							.OnClicked(this, &SSequencerTransformBox::OnMinusButtonClicked)
+					]
+
+				+ SHorizontalBox::Slot()
+					.Padding(6.0f, 0.0f, 0.0f, 0.0f)
+					.AutoWidth()
+					[
+						SAssignNew(EntryBox, SNumericEntryBox<float>)
+							.MinDesiredValueWidth(32.0f)
+							.TypeInterface(NumericTypeInterface)
+							.OnValueCommitted(this, &SSequencerTransformBox::OnDeltaChanged)
+							.Value_Lambda([this](){ return DeltaTime; })
+					]
+
+				+ SHorizontalBox::Slot()
+				.Padding(3.0f, 0.0f, 0.0f, 0.0f)
+				.AutoWidth()
+				.HAlign(HAlign_Center)
+				.VAlign(VAlign_Center)
+				[
+					SNew(SButton)
+					.ButtonStyle( CloseButtonStyle )
+					.OnClicked( this, &SSequencerTransformBox::OnCloseButtonClicked )
+					.ContentPadding( 0 )
+					[
+						SNew(SSpacer)
+						.Size( CloseButtonStyle->Normal.ImageSize )
+					]
+				]
+			]
+	];
+}
+
+
+void SSequencerTransformBox::ToggleVisibility()
+{
+	FSlateApplication& SlateApplication = FSlateApplication::Get();
+
+	if (Border->GetVisibility() == EVisibility::Visible)
+	{
+		SlateApplication.SetAllUserFocus(LastFocusedWidget.Pin(), EFocusCause::Navigation);
+		Border->SetVisibility(EVisibility::Collapsed);
+	}
+	else
+	{
+		Border->SetVisibility(EVisibility::Visible);
+		LastFocusedWidget = SlateApplication.GetUserFocusedWidget(0);
+		SlateApplication.SetAllUserFocus(EntryBox, EFocusCause::Navigation);
+	}
+}
+
+void SSequencerTransformBox::OnDeltaChanged(float Value, ETextCommit::Type CommitType)
+{
+	if (CommitType != ETextCommit::OnEnter)
+	{
+		return;
+	}
+
+	DeltaTime = Value;
+}
+
+FReply SSequencerTransformBox::OnPlusButtonClicked()
+{
+	TSharedPtr<FSequencer> Sequencer = SequencerPtr.Pin();
+
+	Sequencer->TransformSelectedKeysAndSections(DeltaTime);
+
+	return FReply::Handled();
+}
+
+FReply SSequencerTransformBox::OnMinusButtonClicked()
+{
+	TSharedPtr<FSequencer> Sequencer = SequencerPtr.Pin();
+
+	Sequencer->TransformSelectedKeysAndSections(-DeltaTime);
+
+	return FReply::Handled();
+}
+
+FReply SSequencerTransformBox::OnCloseButtonClicked()
+{
+	ToggleVisibility();
+
+	return FReply::Handled();
+}
+
+#undef LOCTEXT_NAMESPACE
diff --git a/Engine/Source/Editor/Sequencer/Private/SSequencerTransformBox.h b/Engine/Source/Editor/Sequencer/Private/SSequencerTransformBox.h
new file mode 100644
--- /dev/null
+++ b/Engine/Source/Editor/Sequencer/Private/SSequencerTransformBox.h
@@ -0,0 +1,59 @@
+// Copyright 1998-2016 Epic Games, Inc. All Rights Reserved.
+
+#pragma once
+
+#include "NumericTypeInterface.h"
+#include "SNumericEntryBox.h"
+
+
+class FSequencer;
+class USequencerSettings;
+
+
+class SSequencerTransformBox
+	: public SCompoundWidget
+{
+public:
+
+	SLATE_BEGIN_ARGS(SSequencerTransformBox) { }
+	SLATE_END_ARGS()
+
+	/** Construct the widget. */
+	void Construct(const FArguments& InArgs, const TSharedRef<FSequencer>& InSequencer, USequencerSettings& InSettings, const TSharedRef<INumericTypeInterface<float>>& InNumericTypeInterface);
+
+	/** Toggle the widget's visibility. */
+	void ToggleVisibility();
+
+private:
+
+	/** Callback for transform operations. */
+	FReply OnPlusButtonClicked();
+	FReply OnMinusButtonClicked();
+	void OnDeltaChanged(float Value, ETextCommit::Type CommitType);
+
+	/** Callback for when the close button is clicked. */
+	FReply OnCloseButtonClicked();
+
+private:
+
+	/** The border widget. */
+	TSharedPtr<SWidget> Border;
+
+	/** The entry box widget. */
+	TSharedPtr<SNumericEntryBox<float>> EntryBox;
+
+	/** The widget that focused prior to this transform box. */
+	TWeakPtr<SWidget> LastFocusedWidget;
+
+	/** Numeric type interface used for converting parsing and generating strings from numbers. */
+	TSharedPtr<INumericTypeInterface<float>> NumericTypeInterface;
+
+	/** The main sequencer interface. */
+	TWeakPtr<FSequencer> SequencerPtr;
+
+	/** Cached settings provided to the sequencer itself on creation. */
+	USequencerSettings* Settings;
+
+	/** Cached delta time. */
+	float DeltaTime;
+};
\ No newline at end of file
