Copying //UE4/Dev-Sequencer to //UE4/Dev-Main (Source: //UE4/Dev-Sequencer @ 3178529)

diff --git a/Engine/Source/Editor/Sequencer/Private/Sequencer.cpp b/Engine/Source/Editor/Sequencer/Private/Sequencer.cpp
--- a/Engine/Source/Editor/Sequencer/Private/Sequencer.cpp
+++ b/Engine/Source/Editor/Sequencer/Private/Sequencer.cpp
@@ -37,7 +37,6 @@
 #include "MovieSceneSubTrack.h"
 #include "MovieSceneCinematicShotSection.h"
 #include "ISequencerSection.h"
-#include "MovieSceneSequenceInstance.h"
 #include "IKeyArea.h"
 #include "ISettingsModule.h"
 #include "IDetailsView.h"
@@ -90,11 +89,59 @@ namespace
 	TWeakPtr<FSequencer> ActiveSequencerPtr;
 }
 
+struct FSequencerTemplateStore : FMovieSceneSequenceTemplateStore
+{
+	FSequencerTemplateStore()
+	{
+		TemplateParameters.bForEditorPreview = true;
+		bTemplatesAreVolatile = true;
+	}
+
+	void Reset()
+	{
+		Templates.Reset();
+	}
+
+	void PurgeStaleTracks()
+	{
+		for (auto& Pair : Templates)
+		{
+			Pair.Value->PurgeStaleTracks();
+		}
+	}
+
+	virtual FMovieSceneEvaluationTemplate& GetCompiledTemplate(UMovieSceneSequence& Sequence)
+	{
+		FObjectKey SequenceKey(&Sequence);
+
+		if (TUniquePtr<FCachedMovieSceneEvaluationTemplate>* ExistingTemplate = Templates.Find(SequenceKey))
+		{
+			(*ExistingTemplate)->Regenerate(TemplateParameters);
+			return **ExistingTemplate;
+		}
+		else
+		{
+			FCachedMovieSceneEvaluationTemplate* NewTemplate = new FCachedMovieSceneEvaluationTemplate;
+			NewTemplate->Initialize(Sequence, this);
+			NewTemplate->Regenerate(TemplateParameters);
+
+			Templates.Add(SequenceKey, TUniquePtr<FCachedMovieSceneEvaluationTemplate>(NewTemplate));
+			return *NewTemplate;
+		}
+	}
+
+	// Store templates as unique ptrs to ensure that external pointers don't become invalid when the array is reallocated
+	TMap<FObjectKey, TUniquePtr<FCachedMovieSceneEvaluationTemplate>> Templates;
+	FMovieSceneTrackCompilationParams TemplateParameters;
+};
+
 
 void FSequencer::InitSequencer(const FSequencerInitParams& InitParams, const TSharedRef<ISequencerObjectChangeListener>& InObjectChangeListener, const TSharedRef<IDetailKeyframeHandler>& InDetailKeyframeHandler, const TArray<FOnCreateTrackEditor>& TrackEditorDelegates)
 {
 	bIsEditingWithinLevelEditor = InitParams.bEditWithinLevelEditor;
 	SilentModeCount = 0;
+	
+	PreAnimatedState.EnableGlobalCapture();
 
 	if (InitParams.SpawnRegister.IsValid())
 	{
@@ -162,6 +209,8 @@ void FSequencer::InitSequencer(const FSequencerInitParams& InitParams, const TSh
 	Settings->GetOnLockPlaybackToAudioClockChanged().AddSP(this, &FSequencer::ResetTimingManager);
 	ResetTimingManager(Settings->ShouldLockPlaybackToAudioClock());
 
+	Settings->GetOnEvaluateSubSequencesInIsolationChanged().AddSP(this, &FSequencer::RestorePreAnimatedState);
+
 	ToolkitHost = InitParams.ToolkitHost;
 
 	ScrubPosition = InitParams.ViewParams.InitialScrubPosition;
@@ -171,12 +220,13 @@ void FSequencer::InitSequencer(const FSequencerInitParams& InitParams, const TSh
 	DetailKeyframeHandler = InDetailKeyframeHandler;
 
 	check( ObjectChangeListener.IsValid() );
-		
-	UMovieSceneSequence& RootSequence = *InitParams.RootSequence;
-		
-	// Focusing the initial movie scene needs to be done before the first time GetFocusedMovieSceneSequenceInstance or GetRootMovieSceneInstance is used
-	RootMovieSceneSequenceInstance = MakeShareable(new FMovieSceneSequenceInstance(RootSequence));
-	SequenceInstanceStack.Add(RootMovieSceneSequenceInstance.ToSharedRef());
+	
+	RootSequence = InitParams.RootSequence;
+
+	TemplateStore = MakeShared<FSequencerTemplateStore>();
+
+	ActiveTemplateIDs.Add(MovieSceneSequenceID::Root);
+	RootTemplateInstance.Initialize(*InitParams.RootSequence, *this, TemplateStore.ToSharedRef());
 
 	// Make internal widgets
 	SequencerWidget = SNew( SSequencer, SharedThis( this ) )
@@ -185,13 +235,14 @@ void FSequencer::InitSequencer(const FSequencerInitParams& InitParams, const TSh
 		.PlaybackRange( this, &FSequencer::GetPlaybackRange )
 		.PlaybackStatus( this, &FSequencer::GetPlaybackStatus )
 		.SelectionRange( this, &FSequencer::GetSelectionRange )
+		.SubSequenceRange( this, &FSequencer::GetSubSequenceRange )
 		.OnPlaybackRangeChanged( this, &FSequencer::SetPlaybackRange )
 		.OnPlaybackRangeBeginDrag( this, &FSequencer::OnPlaybackRangeBeginDrag )
 		.OnPlaybackRangeEndDrag( this, &FSequencer::OnPlaybackRangeEndDrag )
 		.OnSelectionRangeChanged( this, &FSequencer::SetSelectionRange )
 		.OnSelectionRangeBeginDrag( this, &FSequencer::OnSelectionRangeBeginDrag )
 		.OnSelectionRangeEndDrag( this, &FSequencer::OnSelectionRangeEndDrag )
-		.ScrubPosition( this, &FSequencer::OnGetScrubPosition )
+		.ScrubPosition( this, &FSequencer::GetLocalTime )
 		.OnBeginScrubbing( this, &FSequencer::OnBeginScrubbing )
 		.OnEndScrubbing( this, &FSequencer::OnEndScrubbing )
 		.OnScrubPositionChanged( this, &FSequencer::OnScrubPositionChanged )
@@ -249,7 +300,7 @@ void FSequencer::InitSequencer(const FSequencerInitParams& InitParams, const TSh
 		TrackEditor->OnInitialize();
 	}
 
-	OnActivateSequenceEvent.Broadcast(*RootMovieSceneSequenceInstance);
+	OnActivateSequenceEvent.Broadcast(ActiveTemplateIDs[0]);
 }
 
 
@@ -273,6 +324,7 @@ FSequencer::FSequencer()
 	, bUpdatingSequencerSelection( false )
 	, bUpdatingExternalSelection( false )
 	, OldMaxTickRate(GEngine->GetMaxFPS())
+	, bNeedsEvaluate(false)
 {
 	Selection.GetOnOutlinerNodeSelectionChanged().AddRaw(this, &FSequencer::OnSelectedOutlinerNodesChanged);
 	Selection.GetOnNodesWithSelectedKeysOrSectionsChanged().AddRaw(this, &FSequencer::OnSelectedOutlinerNodesChanged);
@@ -314,7 +366,7 @@ void FSequencer::Close()
 {
 	if (ActiveSequencerPtr.IsValid() && ActiveSequencerPtr.Pin().Get() == this)
 	{
-		RootMovieSceneSequenceInstance->RestoreState(*this);
+		RestorePreAnimatedState();
 
 		FEditorDelegates::PreSaveWorld.RemoveAll(this);
 		FEditorDelegates::PostSaveWorld.RemoveAll(this);
@@ -368,22 +420,24 @@ void FSequencer::Close()
 
 void FSequencer::Tick(float InDeltaTime)
 {
-	if (EventContextsAttribute.IsBound())
+	Selection.Tick();
+	
+	if (PlaybackContextAttribute.IsBound())
 	{
-		CachedEventContexts.Reset();
-		for (UObject* Object : EventContextsAttribute.Get())
+		TWeakObjectPtr<UObject> NewPlaybackContext = PlaybackContextAttribute.Get();
+
+		if (CachedPlaybackContext != NewPlaybackContext)
 		{
-			CachedEventContexts.Add(Object);
+			State.ClearObjectCaches();
+			CachedPlaybackContext = NewPlaybackContext;
 		}
 	}
 
-	if (PlaybackContextAttribute.IsBound())
+	if (RootTemplateInstance.IsDirty())
 	{
-		CachedPlaybackContext = PlaybackContextAttribute.Get();
+		bNeedsEvaluate = true;
 	}
-
-	Selection.Tick();
-
+	
 	if ( bNeedInstanceRefresh )
 	{
 		UpdateRuntimeInstances();
@@ -400,7 +454,10 @@ void FSequencer::Tick(float InDeltaTime)
 		SetPlaybackStatus(StoredPlaybackState);
 	}
 
-	FTimeAndDelta TimeAndDelta = TimingManager->AdjustTime(GetGlobalTime(), InDeltaTime, PlayRate);
+	UObject* PlaybackContext = GetPlaybackContext();
+	UWorld* World = PlaybackContext ? PlaybackContext->GetWorld() : nullptr;
+	float Dilation = World ? World->GetWorldSettings()->MatineeTimeDilation : 1.f;
+	FTimeAndDelta TimeAndDelta = TimingManager->AdjustTime(GetGlobalTime(), InDeltaTime, PlayRate, Dilation);
 
 	static const float AutoScrollFactor = 0.1f;
 
@@ -415,7 +472,7 @@ void FSequencer::Tick(float InDeltaTime)
 	if (AutoscrubOffset.IsSet())
 	{
 		float Offset = AutoscrubOffset.GetValue() * AutoScrollFactor;
-		SetGlobalTimeDirectly(GetGlobalTime() + Offset);
+		SetLocalTimeDirectly(GetLocalTime() + Offset);
 	}
 
 	// override max frame rate
@@ -442,8 +499,15 @@ void FSequencer::Tick(float InDeltaTime)
 	if (PlaybackState == EMovieScenePlayerStatus::Playing ||
 		PlaybackState == EMovieScenePlayerStatus::Recording)
 	{
-		SetGlobalTimeLooped(TimeAndDelta.Time);
+		// Put the time into local space
+		SetLocalTimeLooped(TimeAndDelta.Time * RootToLocalTransform);
 	}
+	else
+	{
+		PlayPosition.Reset(GetGlobalTime());
+	}
+
+	UpdateSubSequenceData();
 
 	// Tick all the tools we own as well
 	for (int32 EditorIndex = 0; EditorIndex < TrackEditors.Num(); ++EditorIndex)
@@ -453,7 +517,12 @@ void FSequencer::Tick(float InDeltaTime)
 
 	if (!IsInSilentMode())
 	{
-		PostTickRenderStateFixup();
+		if (bNeedsEvaluate)
+		{
+			FMovieSceneEvaluationRange Range = PlayPosition.JumpTo(ScrubPosition, GetFocusedMovieSceneSequence()->GetMovieScene()->GetOptionalFixedFrameInterval());
+			EvaluateInternal(Range);
+		}
+		//PostTickRenderStateFixup();
 	}
 
 	ISequenceRecorder& SequenceRecorder = FModuleManager::LoadModuleChecked<ISequenceRecorder>("SequenceRecorder");
@@ -537,7 +606,7 @@ void FSequencer::PostTickRenderStateFixup()
 
 			// Force previous frame in sequence to re-render
 			SetPlaybackStatus(EMovieScenePlayerStatus::Stopped);
-			SetGlobalTimeDirectly(PreviousTime);
+			SetGlobalTime(PreviousTime);
 			Tick(StepDelta);
 
 			UWorld* PlaybackContext = Cast<UWorld>(GetPlaybackContext());
@@ -547,7 +616,7 @@ void FSequencer::PostTickRenderStateFixup()
 			}
 
 			SetPlaybackStatus(EMovieScenePlayerStatus::Stopped);
-			SetGlobalTimeDirectly(Time);
+			SetGlobalTime(Time);
 			Tick(StepDelta);
 		}
 	}
@@ -562,21 +631,16 @@ TSharedRef<SWidget> FSequencer::GetSequencerWidget() const
 
 UMovieSceneSequence* FSequencer::GetRootMovieSceneSequence() const
 {
-	if (SequenceInstanceStack.Num())
-	{
-		return SequenceInstanceStack[0]->GetSequence();
-	}
-
-	return nullptr;
+	return RootSequence.Get();
 }
 
 
 UMovieSceneSequence* FSequencer::GetFocusedMovieSceneSequence() const
 {
 	// the last item is the focused movie scene
-	if (SequenceInstanceStack.Num())
+	if (ActiveTemplateIDs.Num())
 	{
-		return SequenceInstanceStack.Top()->GetSequence();
+		return RootTemplateInstance.GetSequence(ActiveTemplateIDs.Last());
 	}
 
 	return nullptr;
@@ -585,77 +649,64 @@ UMovieSceneSequence* FSequencer::GetFocusedMovieSceneSequence() const
 
 void FSequencer::ResetToNewRootSequence(UMovieSceneSequence& NewSequence)
 {
-	if (RootMovieSceneSequenceInstance.IsValid())
-	{
-		RootMovieSceneSequenceInstance->RestoreState(*this);
-	}
+	RootSequence = &NewSequence;
+	RestorePreAnimatedState();
 
-	//@todo Sequencer - Encapsulate this better
-	SequenceInstanceStack.Empty();
-	SequenceInstanceBySection.Empty();
+	TemplateStore->Reset();
 
-	// Focusing the initial movie scene needs to be done before the first time NewSequence or GetRootMovieSceneInstance is used
-	RootMovieSceneSequenceInstance = MakeShareable(new FMovieSceneSequenceInstance(NewSequence));
-	SequenceInstanceStack.Add(RootMovieSceneSequenceInstance.ToSharedRef());
+	ActiveTemplateIDs.Reset();
+	ActiveTemplateIDs.Add(MovieSceneSequenceID::Root);
 
-	ResetPerMovieSceneData();
-	SequencerWidget->ResetBreadcrumbs();
+	RootTemplateInstance.Initialize(NewSequence, *this);
 
-	OnActivateSequenceEvent.Broadcast(*SequenceInstanceStack.Top());
-}
+	RootToLocalTransform = FMovieSceneSequenceTransform();
 
+	ResetPerMovieSceneData();
+	SequencerWidget->ResetBreadcrumbs();
 
-TSharedRef<FMovieSceneSequenceInstance> FSequencer::GetRootMovieSceneSequenceInstance() const
-{
-	return SequenceInstanceStack[0];
+	OnActivateSequenceEvent.Broadcast(ActiveTemplateIDs.Top());
 }
 
 
-TSharedRef<FMovieSceneSequenceInstance> FSequencer::GetFocusedMovieSceneSequenceInstance() const
+void FSequencer::FocusSequenceInstance(UMovieSceneSubSection& InSubSection)
 {
-	// the last item is the focused movie scene
-	return SequenceInstanceStack.Top();
-}
+	// Root out the SequenceID for the sub section
+	FMovieSceneSequenceID SequenceID = MovieSceneSequenceID::Root;
 
+	const FMovieSceneSequenceHierarchyNode* Node = RootTemplateInstance.GetHierarchy().FindNode(ActiveTemplateIDs.Last());
 
-void FSequencer::FocusSequenceInstance(UMovieSceneSubSection& InSubSection)
-{
-	if (RootMovieSceneSequenceInstance.IsValid())
+	FName SearchForName(*InSubSection.GetPathNameInMovieScene());
+	for (FMovieSceneSequenceIDRef ChildID : Node->Children)
 	{
-		RootMovieSceneSequenceInstance->RestoreState(*this);
+		const FMovieSceneSubSequenceData* SubSequence = RootTemplateInstance.GetHierarchy().FindSubData(ChildID);
+		if (SearchForName == SubSequence->SectionPath)
+		{
+			SequenceID = ChildID;
+			break;
+		}
 	}
 
-	TSharedRef<FMovieSceneSequenceInstance> SequenceInstance = GetSequenceInstanceForSection(InSubSection);
-
-	// Check for infinite recursion
-	check(SequenceInstance != SequenceInstanceStack.Top());
-
-	// Focus the movie scene
-	SequenceInstanceStack.Push(SequenceInstance);
-
-	// Find the time in the focused sequence instance that is relative to the current time
-	TRange<float> PlaybackRange = TRange<float>(0, 0);
-	float PlaybackRangeStart = 0.f;
-	if (InSubSection.GetSequence() != nullptr)
+	if (!ensure(SequenceID != MovieSceneSequenceID::Root))
 	{
-		UMovieScene* ShotMovieScene = InSubSection.GetSequence()->GetMovieScene();
-		PlaybackRange = ShotMovieScene->GetPlaybackRange();
+		return;
 	}
 
-	const float ShotOffset = InSubSection.StartOffset + PlaybackRange.GetLowerBoundValue() - InSubSection.PrerollTime;
-	float AbsoluteShotPosition = ShotOffset + (GetGlobalTime() - (InSubSection.GetStartTime() - InSubSection.PrerollTime)) / InSubSection.TimeScale;
-	if (!PlaybackRange.Contains(AbsoluteShotPosition))
+	ActiveTemplateIDs.Push(SequenceID);
+
+	if (Settings->ShouldEvaluateSubSequencesInIsolation())
 	{
-		AbsoluteShotPosition = PlaybackRange.GetLowerBoundValue();
+		RestorePreAnimatedState();
 	}
 
+	UpdateSubSequenceData();
+
 	// Reset data that is only used for the previous movie scene
 	ResetPerMovieSceneData();
 	SequencerWidget->UpdateBreadcrumbs();
 
-	SetGlobalTime(AbsoluteShotPosition, ESnapTimeMode::STM_Interval);
+	OnActivateSequenceEvent.Broadcast(ActiveTemplateIDs.Top());
 
-	OnActivateSequenceEvent.Broadcast(*SequenceInstance);
+	bNeedsEvaluate = true;
 }
 
 
@@ -666,7 +717,8 @@ FGuid FSequencer::CreateBinding(UObject& InObject, const FString& InName)
 		
 	const FGuid PossessableGuid = OwnerMovieScene->AddPossessable(InName, InObject.GetClass());
 
-	UObject* BindingContext = OwnerSequence->GetParentObject(&InObject);
+	// Attempt to use the parent as a context if necessary
+	UObject* BindingContext = OwnerSequence->AreParentContextsSignificant() ? OwnerSequence->GetParentObject(&InObject) : nullptr;
 	if (BindingContext)
 	{
 		// Ensure we have possessed the outer object, if necessary
@@ -711,71 +763,6 @@ TArray<UObject*> FSequencer::GetEventContexts() const
 	return Temp;
 }
 
-void FSequencer::GetAllKeyedProperties(UObject& Object, TSet<UProperty*>& OutProperties)
-{
-	GetAllKeyedPropertiesForInstance(Object, *GetRootMovieSceneSequenceInstance(), OutProperties);
-}
-
-
-void FSequencer::GetAllKeyedPropertiesForInstance(UObject& Object, FMovieSceneSequenceInstance& Instance, TSet<UProperty*>& OutProperties)
-{
-	GetKeyedProperties(Object, Instance, OutProperties);
-
-	UMovieSceneSubTrack* SubTrack = Instance.GetSequence()->GetMovieScene()->FindMasterTrack<UMovieSceneSubTrack>();
-	if (SubTrack)
-	{
-		// Iterate sub movie scenes
-		for (UMovieSceneSection* SubSection : SubTrack->GetAllSections())
-		{
-			// Find the instance for the section
-			TSharedRef<FMovieSceneSequenceInstance>* SubInstance = SequenceInstanceBySection.Find(SubSection);
-			GetAllKeyedPropertiesForInstance(Object, SubInstance->Get(), OutProperties);
-		}
-	}
-}
-
-
-void FSequencer::GetKeyedProperties(UObject& Object, FMovieSceneSequenceInstance& Instance, TSet<UProperty*>& OutProperties)
-{
-	FGuid BindingId = Instance.FindObjectId(Object);
-
-	const FMovieSceneBinding* Binding = BindingId.IsValid() ? Instance.GetSequence()->GetMovieScene()->GetBindings().FindByPredicate([&](const FMovieSceneBinding& InBinding){
-		return BindingId == InBinding.GetObjectGuid();
-	}) : nullptr;
-
-	if (!Binding)
-	{
-		return;
-	}
-
-	// Add any properties for this object
-	for (UMovieSceneTrack* Track : Binding->GetTracks())
-	{
-		if (UMovieScene3DTransformTrack* TransformTrack = Cast<UMovieScene3DTransformTrack>(Track))
-		{
-			// Add transform properties manually. This is technically not correct, since it's the scene component that's being keyed, not the actor,
-			// But since the transform track is always stored on an actor binding, we have to add them this way
-			AActor* Actor = Cast<AActor>(&Object);
-			if (Actor)
-			{
-				UClass* Class = USceneComponent::StaticClass();
-				OutProperties.Add(Class->FindPropertyByName(GET_MEMBER_NAME_CHECKED(USceneComponent, RelativeLocation)));
-				OutProperties.Add(Class->FindPropertyByName(GET_MEMBER_NAME_CHECKED(USceneComponent, RelativeRotation)));
-				OutProperties.Add(Class->FindPropertyByName(GET_MEMBER_NAME_CHECKED(USceneComponent, RelativeScale3D)));
-			}
-		}
-		else if (UMovieScenePropertyTrack* PropertyTrack = Cast<UMovieScenePropertyTrack>(Track))
-		{
-			FTrackInstancePropertyBindings InstanceBinding(PropertyTrack->GetPropertyName(), PropertyTrack->GetPropertyPath());
-			UProperty* Property = InstanceBinding.GetProperty(&Object);
-			if (Property)
-			{
-				OutProperties.Add(Property);
-			}
-		}
-	}
-}
-
 void FSequencer::GetKeysFromSelection(TUniquePtr<ISequencerKeyCollection>& KeyCollection)
 {
 	if (!KeyCollection.IsValid())
@@ -796,29 +783,19 @@ void FSequencer::GetKeysFromSelection(TUniquePtr<ISequencerKeyCollection>& KeyCo
 	KeyCollection->InitializeRecursive(SelectedNodes, DuplicateThreshold);
 }
 
-TSharedRef<FMovieSceneSequenceInstance> FSequencer::GetSequenceInstanceForSection(UMovieSceneSection& Section) const
-{
-	return SequenceInstanceBySection.FindChecked(&Section);
-}
-
-
-bool FSequencer::HasSequenceInstanceForSection(UMovieSceneSection& Section) const
-{
-	return SequenceInstanceBySection.Contains(&Section);
-}
-
 
-void FSequencer::PopToSequenceInstance(TSharedRef<FMovieSceneSequenceInstance> SequenceInstance)
+void FSequencer::PopToSequenceInstance(FMovieSceneSequenceIDRef SequenceID)
 {
-	if( SequenceInstanceStack.Num() > 1 )
+	if( ActiveTemplateIDs.Num() > 1 )
 	{
 		// Pop until we find the movie scene to focus
-		while( SequenceInstance != SequenceInstanceStack.Last() )
+		while( SequenceID != ActiveTemplateIDs.Last() )
 		{
-			SequenceInstanceStack.Pop();
+			ActiveTemplateIDs.Pop();
 		}
-	
-		check( SequenceInstanceStack.Num() > 0 );
+
+		check( ActiveTemplateIDs.Num() > 0 );
+		UpdateSubSequenceData();
 
 		// Pop out of any potentially locked cameras from the shot and toggle on camera cuts
 		for (int32 i = 0; i < GEditor->LevelViewportClients.Num(); ++i)
@@ -837,10 +814,31 @@ void FSequencer::PopToSequenceInstance(TSharedRef<FMovieSceneSequenceInstance> S
 		ResetPerMovieSceneData();
 		SequencerWidget->UpdateBreadcrumbs();
 
-		OnActivateSequenceEvent.Broadcast(*SequenceInstanceStack.Top());
+		OnActivateSequenceEvent.Broadcast(ActiveTemplateIDs.Top());
+
+		bNeedsEvaluate = true;
 	}
 }
 
+void FSequencer::UpdateSubSequenceData()
+{
+	SubSequenceRange = TRange<float>(0.f);
+	RootToLocalTransform = FMovieSceneSequenceTransform();
+
+	// Find the parent sub section and set up the sub sequence range, if necessary
+	if (ActiveTemplateIDs.Num() <= 1)
+	{
+		return;
+	}
+
+	const FMovieSceneSubSequenceData* SubSequenceData = RootTemplateInstance.GetHierarchy().FindSubData(ActiveTemplateIDs.Top());
+
+	if (SubSequenceData)
+	{
+		SubSequenceRange = SubSequenceData->ValidPlayRange;
+		RootToLocalTransform = SubSequenceData->RootToSequenceTransform;
+	}
+}
 
 void FSequencer::DeleteSections(const TSet<TWeakObjectPtr<UMovieSceneSection>>& Sections)
 {
@@ -1000,6 +998,70 @@ bool FSequencer::CanSnapToFrame() const
 	return bKeysSelected && CanShowFrameNumbers();
 }
 
+void FSequencer::TransformSelectedKeysAndSections(float InDeltaTime)
+{
+	FScopedTransaction TransformKeysAndSectionsTransaction(NSLOCTEXT("Sequencer", "TransformKeysandSections_Transaction", "Transform Keys and Sections"));
+	bool bAnythingChanged = false;
+	TArray<FSequencerSelectedKey> SelectedKeysArray = Selection.GetSelectedKeys().Array();
+
+	TArray<UMovieSceneSection*> SectionsWithKeys;
+	
+	FSequencerDisplayNode::DisableKeyGoupingRegeneration();
+	
+	for (const FSequencerSelectedKey& Key : SelectedKeysArray)
+	{
+		if (Key.IsValid())
+		{
+			UMovieSceneSection* Section = Key.Section;
+
+			bool SectionModified = SectionsWithKeys.Contains(Section);
+
+			if (!SectionModified)
+			{
+				if (Key.Section->TryModify())
+				{
+					SectionModified = true;
+					SectionsWithKeys.Add(Section);
+				}
+			}
+			
+			if (SectionModified)
+			{
+				Key.KeyArea->MoveKey(Key.KeyHandle.GetValue(), InDeltaTime);
+				bAnythingChanged = true;
+
+				float NewKeyTime = Key.KeyArea->GetKeyTime(Key.KeyHandle.GetValue()) + InDeltaTime;
+				if (NewKeyTime > Section->GetEndTime())
+				{
+					Section->SetEndTime(NewKeyTime);
+				}
+				else if (NewKeyTime < Section->GetStartTime())
+				{
+					Section->SetStartTime(NewKeyTime);
+				}
+			}
+		}
+	}
+
+	FSequencerDisplayNode::EnableKeyGoupingRegeneration();
+
+	for (TWeakObjectPtr<UMovieSceneSection> SelectedSection : Selection.GetSelectedSections())
+	{
+		if (SelectedSection.IsValid())
+		{
+			TSet<FKeyHandle> EmptyKeyHandles;
+			SelectedSection.Get()->SetFlags( RF_Transactional );
+			SelectedSection.Get()->MoveSection(InDeltaTime, EmptyKeyHandles);
+			bAnythingChanged = true;
+		}
+	}
+
+	if (bAnythingChanged)
+	{
+		NotifyMovieSceneDataChanged( EMovieSceneDataChangeType::TrackValueChanged );
+	}
+}
+
 
 void FSequencer::NotifyMapChanged( class UWorld* NewWorld, EMapChangeType MapChangeType )
 {
@@ -1030,7 +1092,7 @@ void FSequencer::NotifyMovieSceneDataChanged( EMovieSceneDataChangeType DataChan
 	if ( DataChangeType == EMovieSceneDataChangeType::ActiveMovieSceneChanged ||
 		DataChangeType == EMovieSceneDataChangeType::Unknown )
 	{
-		LabelManager.SetMovieScene( SequenceInstanceStack.Top()->GetSequence()->GetMovieScene() );
+		LabelManager.SetMovieScene( GetFocusedMovieSceneSequence()->GetMovieScene() );
 	}
 
 	StoredPlaybackState = GetPlaybackStatus();
@@ -1052,6 +1114,9 @@ void FSequencer::NotifyMovieSceneDataChanged( EMovieSceneDataChangeType DataChan
 	else if (DataChangeType == EMovieSceneDataChangeType::TrackValueChangedRefreshImmediately)
 	{
 		UpdateRuntimeInstances();
+
+		// Evaluate now
+		EvaluateInternal(FMovieSceneEvaluationRange(ScrubPosition));
 	}
 	else
 	{
@@ -1064,6 +1129,9 @@ void FSequencer::NotifyMovieSceneDataChanged( EMovieSceneDataChangeType DataChan
 		bNeedInstanceRefresh = true;
 	}
 
+	bNeedsEvaluate = true;
+	State.ClearObjectCaches();
+
 	UpdatePlaybackRange();
 	OnMovieSceneDataChangedDelegate.Broadcast();
 }
@@ -1097,6 +1165,16 @@ void FSequencer::SetClampRange(TRange<float> InNewClampRange)
 }
 
 
+TOptional<TRange<float>> FSequencer::GetSubSequenceRange() const
+{
+	if (Settings->ShouldEvaluateSubSequencesInIsolation() || ActiveTemplateIDs.Num() == 1)
+	{
+		return TOptional<TRange<float>>();
+	}
+	return SubSequenceRange;
+}
+
+
 TRange<float> FSequencer::GetSelectionRange() const
 {
 	return GetFocusedMovieSceneSequence()->GetMovieScene()->GetSelectionRange();
@@ -1114,30 +1192,30 @@ void FSequencer::SetSelectionRange(TRange<float> Range)
 
 void FSequencer::SetSelectionRangeEnd()
 {
-	const float GlobalTime = GetGlobalTime();
+	const float LocalTime = GetLocalTime();
 
-	if (GetSelectionRange().GetLowerBoundValue() >= GlobalTime)
+	if (GetSelectionRange().GetLowerBoundValue() >= LocalTime)
 	{
-		SetSelectionRange(TRange<float>(GlobalTime));
+		SetSelectionRange(TRange<float>(LocalTime));
 	}
 	else
 	{
-		SetSelectionRange(TRange<float>(GetSelectionRange().GetLowerBoundValue(), GlobalTime));
+		SetSelectionRange(TRange<float>(GetSelectionRange().GetLowerBoundValue(), LocalTime));
 	}
 }
 
 
 void FSequencer::SetSelectionRangeStart()
 {
-	const float GlobalTime = GetGlobalTime();
+	const float LocalTime = GetLocalTime();
 
-	if (GetSelectionRange().GetUpperBoundValue() <= GlobalTime)
+	if (GetSelectionRange().GetUpperBoundValue() <= LocalTime)
 	{
-		SetSelectionRange(TRange<float>(GlobalTime));
+		SetSelectionRange(TRange<float>(LocalTime));
 	}
 	else
 	{
-		SetSelectionRange(TRange<float>(GlobalTime, GetSelectionRange().GetUpperBoundValue()));
+		SetSelectionRange(TRange<float>(LocalTime, GetSelectionRange().GetUpperBoundValue()));
 	}
 }
 
@@ -1213,6 +1291,8 @@ void FSequencer::SetPlaybackRange(TRange<float> Range)
 
 		UMovieScene* FocusedMovieScene = GetFocusedMovieSceneSequence()->GetMovieScene();
 		FocusedMovieScene->SetPlaybackRange(Range.GetLowerBoundValue(), Range.GetUpperBoundValue());
+
+		bNeedsEvaluate = true;
 	}
 }
 
@@ -1435,11 +1515,9 @@ bool FSequencer::IsRecordingLive() const
 }
 
 
-float FSequencer::GetCurrentLocalTime( UMovieSceneSequence& InMovieSceneSequence )
+float FSequencer::GetLocalTime() const
 {
-	//@todo Sequencer: Nested movie scenes:  Figure out the parent of the passed in movie scene and 
-	// calculate local time
-	return ScrubPosition;
+	return ScrubPosition * RootToLocalTransform;
 }
 
 
@@ -1448,33 +1526,12 @@ float FSequencer::GetGlobalTime() const
 	return ScrubPosition;
 }
 
-
-void FSequencer::SetGlobalTime( float NewTime, ESnapTimeMode SnapTimeMode, bool bRestarted )
-{
-	if (IsAutoScrollEnabled())
-	{
-		float RangeOffset = CalculateAutoscrollEncroachment(NewTime).Get(0.f);
-			
-		// When not scrubbing, we auto scroll the view range immediately
-		if (RangeOffset != 0.f)
-		{
-			TRange<float> WorkingRange = GetClampRange();
-			if (TargetViewRange.GetLowerBoundValue() + RangeOffset > WorkingRange.GetLowerBoundValue() &&
-				TargetViewRange.GetUpperBoundValue() + RangeOffset < WorkingRange.GetUpperBoundValue())
-			{
-				SetViewRange(TRange<float>(TargetViewRange.GetLowerBoundValue() + RangeOffset, TargetViewRange.GetUpperBoundValue() + RangeOffset), EViewRangeInterpolation::Immediate);
-			}
-		}
-	}
-
-	SetGlobalTimeDirectly(NewTime, SnapTimeMode, bRestarted);
-}
-
-
-void FSequencer::SetGlobalTimeDirectly( float NewTime, ESnapTimeMode SnapTimeMode, bool bRestarted )
+void FSequencer::SetLocalTime( float NewTime, ESnapTimeMode SnapTimeMode )
 {
-	float LastTime = ScrubPosition;
+	// Ensure the time is in the current view
+	ScrollIntoView(NewTime);
 
+	// Perform snapping
 	if ((SnapTimeMode & ESnapTimeMode::STM_Interval) && Settings->GetIsSnapEnabled())
 	{
 		NewTime = Settings->SnapTimeToInterval(NewTime);
@@ -1485,39 +1542,67 @@ void FSequencer::SetGlobalTimeDirectly( float NewTime, ESnapTimeMode SnapTimeMod
 		NewTime = FindNearestKey(NewTime);
 	}
 
+	SetLocalTimeDirectly(NewTime);
+}
+
+
+void FSequencer::SetLocalTimeDirectly(float NewTime)
+{
+	// Transform the time to the root time-space
+	SetGlobalTime(NewTime * RootToLocalTransform.Inverse());
+}
+
+
+void FSequencer::SetGlobalTime( float NewTime )
+{
 	// Update the position
 	ScrubPosition = NewTime;
 
-	UMovieScene* MovieScene = SequenceInstanceStack.Top()->GetSequence()->GetMovieScene();
-
-	if (LastTime == GetPlaybackRange().GetLowerBoundValue())
+	// Don't update the sequence if the time hasn't changed as this will cause duplicate events and the like to fire.
+	// If we need to reevaluate the sequence at the same time for whetever reason, we should call ForceEvaluate()
+	TOptional<float> LastPosition = PlayPosition.GetPreviousPosition();
+	if (!LastPosition.IsSet() || LastPosition.GetValue() != ScrubPosition)
 	{
-		bRestarted = true;
+		FMovieSceneEvaluationRange Range = PlayPosition.JumpTo(ScrubPosition, GetRootMovieSceneSequence()->GetMovieScene()->GetOptionalFixedFrameInterval());
+		EvaluateInternal(Range);
 	}
+}
 
-	if ( MovieScene->GetForceFixedFrameIntervalPlayback() && MovieScene->GetFixedFrameInterval() > 0 )
+void FSequencer::ForceEvaluate()
+{
+	FMovieSceneEvaluationRange Range = PlayPosition.JumpTo(ScrubPosition, GetFocusedMovieSceneSequence()->GetMovieScene()->GetOptionalFixedFrameInterval());
+	EvaluateInternal(Range);
+}
+
+void FSequencer::EvaluateInternal(FMovieSceneEvaluationRange InRange)
+{
+	bNeedsEvaluate = false;
+
+	if (PlaybackContextAttribute.IsBound())
 	{
-		NewTime = UMovieScene::CalculateFixedFrameTime( NewTime, MovieScene->GetFixedFrameInterval() );
-		LastTime = UMovieScene::CalculateFixedFrameTime( LastTime, MovieScene->GetFixedFrameInterval() );
+		CachedPlaybackContext = PlaybackContextAttribute.Get();
 	}
-
-	// Force last time to be the lower bound of the playback range when restarting or looping around 
-	if (bRestarted)
+	
+	if (EventContextsAttribute.IsBound())
 	{
-		LastTime = GetPlaybackRange().GetLowerBoundValue();
+		CachedEventContexts.Reset();
+		for (UObject* Object : EventContextsAttribute.Get())
+		{
+			CachedEventContexts.Add(Object);
+		}
 	}
 
-	EMovieSceneUpdateData UpdateData(NewTime, LastTime);
-	SequenceInstanceStack.Top()->Update(UpdateData, *this);
+	FMovieSceneContext Context = FMovieSceneContext(InRange, PlaybackState).SetIsSilent(SilentModeCount != 0);
 
-	// Ensure any relevant spawned objects are cleaned up if we're playing back the master sequence
-	if (SequenceInstanceStack.Num() == 1)
+	FMovieSceneSequenceID RootOverride = MovieSceneSequenceID::Root;
+	if (Settings->ShouldEvaluateSubSequencesInIsolation())
 	{
-		if (!GetPlaybackRange().Contains(NewTime))
-		{
-			SpawnRegister->OnMasterSequenceExpired(*this);
-		}
+		RootOverride = ActiveTemplateIDs.Top();
 	}
+	
+	RootTemplateInstance.Evaluate(Context, *this, RootOverride);
+
+	TemplateStore->PurgeStaleTracks();
 
 	// If realtime is off, this needs to be called to update the pivot location when scrubbing.
 	GUnrealEd->UpdatePivotLocationForSelection();
@@ -1535,6 +1620,24 @@ void FSequencer::SetGlobalTimeDirectly( float NewTime, ESnapTimeMode SnapTimeMod
 	}
 }
 
+void FSequencer::ScrollIntoView(float InLocalTime)
+{
+	if (IsAutoScrollEnabled())
+	{
+		float RangeOffset = CalculateAutoscrollEncroachment(InLocalTime).Get(0.f);
+		
+		// When not scrubbing, we auto scroll the view range immediately
+		if (RangeOffset != 0.f)
+		{
+			TRange<float> WorkingRange = GetClampRange();
+			if (TargetViewRange.GetLowerBoundValue() + RangeOffset > WorkingRange.GetLowerBoundValue() &&
+				TargetViewRange.GetUpperBoundValue() + RangeOffset < WorkingRange.GetUpperBoundValue())
+			{
+				SetViewRange(TRange<float>(TargetViewRange.GetLowerBoundValue() + RangeOffset, TargetViewRange.GetUpperBoundValue() + RangeOffset), EViewRangeInterpolation::Immediate);
+			}
+		}
+	}
+}
 
 void FSequencer::UpdateAutoScroll(float NewTime)
 {
@@ -1550,16 +1653,18 @@ void FSequencer::UpdateAutoScroll(float NewTime)
 	TRange<float> ViewRange = GetViewRange();
 	const float Threshold = (ViewRange.GetUpperBoundValue() - ViewRange.GetLowerBoundValue()) * ThresholdPercentage;
 
+	const float LocalPosition = GetLocalTime();
+
 	// If we have no autoscrub offset yet, we move the scrub position to the boundary of the autoscroll threasdhold, then autoscrub from there
 	if (!AutoscrubOffset.IsSet())
 	{
-		if (AutoscrollOffset.GetValue() < 0 && ScrubPosition > ViewRange.GetLowerBoundValue() + Threshold)
+		if (AutoscrollOffset.GetValue() < 0 && LocalPosition > ViewRange.GetLowerBoundValue() + Threshold)
 		{
-			SetGlobalTimeDirectly( ViewRange.GetLowerBoundValue() + Threshold );
+			SetLocalTimeDirectly( ViewRange.GetLowerBoundValue() + Threshold );
 		}
-		else if (AutoscrollOffset.GetValue() > 0 && ScrubPosition < ViewRange.GetUpperBoundValue() - Threshold)
+		else if (AutoscrollOffset.GetValue() > 0 && LocalPosition < ViewRange.GetUpperBoundValue() - Threshold)
 		{
-			SetGlobalTimeDirectly( ViewRange.GetUpperBoundValue() - Threshold );
+			SetLocalTimeDirectly( ViewRange.GetUpperBoundValue() - Threshold );
 		}
 	}
 
@@ -1581,7 +1686,7 @@ void FSequencer::UpdateAutoScroll(float NewTime)
 TOptional<float> FSequencer::CalculateAutoscrollEncroachment(float NewTime, float ThresholdPercentage) const
 {
 	enum class EDirection { Positive, Negative };
-	const EDirection Movement = NewTime - ScrubPosition >= 0 ? EDirection::Positive : EDirection::Negative;
+	const EDirection Movement = NewTime - GetLocalTime() >= 0 ? EDirection::Positive : EDirection::Negative;
 
 	const TRange<float> CurrentRange = GetViewRange();
 	const float RangeMin = CurrentRange.GetLowerBoundValue(), RangeMax = CurrentRange.GetUpperBoundValue();
@@ -1700,8 +1805,8 @@ FGuid FindUnspawnedObjectGuid(UObject& InObject, UMovieSceneSequence& Sequence)
 			// The only way to find the object now is to resolve all the child bindings, and see if they are the same
 			for (const FGuid& ChildGuid : ParentSpawnable->GetChildPossessables())
 			{
-				UObject* ChildObject = Sequence.FindPossessableObject(ChildGuid, ParentContext);
-				if (ChildObject == &InObject)
+				const bool bHasObject = Sequence.LocateBoundObjects(ChildGuid, ParentContext).Contains(&InObject);
+				if (bHasObject)
 				{
 					return ChildGuid;
 				}
@@ -1730,31 +1835,31 @@ FGuid FSequencer::GetHandleToObject( UObject* Object, bool bCreateHandleIfMissin
 		return FGuid();
 	}
 
-	TSharedRef<FMovieSceneSequenceInstance> FocusedMovieSceneSequenceInstance = GetFocusedMovieSceneSequenceInstance();
-	UMovieSceneSequence* FocusedMovieSceneSequence = FocusedMovieSceneSequenceInstance->GetSequence();
-
+	UMovieSceneSequence* FocusedMovieSceneSequence = GetFocusedMovieSceneSequence();
 	UMovieScene* FocusedMovieScene = FocusedMovieSceneSequence->GetMovieScene();
 	
 	// Attempt to resolve the object through the movie scene instance first, 
-	FGuid ObjectGuid = FocusedMovieSceneSequenceInstance->FindObjectId(*Object);
+	FGuid ObjectGuid = FindObjectId(*Object, ActiveTemplateIDs.Top());
 
-	// Check here for spawnable otherwise spawnables get recreated as possessables, which doesn't make sense
-	FMovieSceneSpawnable* Spawnable = FocusedMovieScene->FindSpawnable(ObjectGuid);
-	if (Spawnable)
+	if (ObjectGuid.IsValid())
 	{
-		return ObjectGuid;
-	}
+		// Check here for spawnable otherwise spawnables get recreated as possessables, which doesn't make sense
+		FMovieSceneSpawnable* Spawnable = FocusedMovieScene->FindSpawnable(ObjectGuid);
+		if (Spawnable)
+		{
+			return ObjectGuid;
+		}
 
-	// Make sure that the possessable is still valid, if it's not remove the binding so new one 
-	// can be created.  This can happen due to undo.
-	FMovieScenePossessable* Possessable = FocusedMovieScene->FindPossessable(ObjectGuid);
-	if(Possessable == nullptr)
-	{
-		FocusedMovieSceneSequence->UnbindPossessableObjects(ObjectGuid);
-		ObjectGuid.Invalidate();
+		// Make sure that the possessable is still valid, if it's not remove the binding so new one 
+		// can be created.  This can happen due to undo.
+		FMovieScenePossessable* Possessable = FocusedMovieScene->FindPossessable(ObjectGuid);
+		if(Possessable == nullptr)
+		{
+			FocusedMovieSceneSequence->UnbindPossessableObjects(ObjectGuid);
+			ObjectGuid.Invalidate();
+		}
 	}
-
-	if (!ObjectGuid.IsValid())
+	else
 	{
 		ObjectGuid = FindUnspawnedObjectGuid(*Object, *FocusedMovieSceneSequence);
 	}
@@ -1764,9 +1869,11 @@ FGuid FSequencer::GetHandleToObject( UObject* Object, bool bCreateHandleIfMissin
 		return ObjectGuid;
 	}
 
+	UObject* PlaybackContext = PlaybackContextAttribute.Get(nullptr);
+
 	// If the object guid was not found attempt to add it
 	// Note: Only possessed actors can be added like this
-	if (FocusedMovieSceneSequence->CanPossessObject(*Object) && bCreateHandleIfMissing)
+	if (FocusedMovieSceneSequence->CanPossessObject(*Object, PlaybackContext) && bCreateHandleIfMissing)
 	{
 		AActor* PossessedActor = Cast<AActor>(Object);
 
@@ -1784,16 +1891,6 @@ ISequencerObjectChangeListener& FSequencer::GetObjectChangeListener()
 	return *ObjectChangeListener;
 }
 
-
-void FSequencer::GetRuntimeObjects( TSharedRef<FMovieSceneSequenceInstance> MovieSceneInstance, const FGuid& ObjectHandle, TArray<TWeakObjectPtr<UObject>>& OutObjects ) const
-{
-	UObject* FoundObject = MovieSceneInstance->FindObject(ObjectHandle, *this);
-	if (FoundObject)
-	{
-		OutObjects.Add(FoundObject);
-	}
-}
-
 void FSequencer::PossessPIEViewports(UObject* CameraObject, UObject* UnlockIfCameraObject, bool bJumpCut)
 {
 	for (const FWorldContext& WorldContext : GEngine->GetWorldContexts())
@@ -1998,33 +2095,23 @@ void FSequencer::SetPlaybackStatus(EMovieScenePlayerStatus::Type InPlaybackStatu
 }
 
 
-void FSequencer::AddOrUpdateMovieSceneInstance( UMovieSceneSection& MovieSceneSection, TSharedRef<FMovieSceneSequenceInstance> InstanceToAdd )
-{
-	if( !SequenceInstanceBySection.Contains( &MovieSceneSection ) )
-	{
-		SequenceInstanceBySection.Add( &MovieSceneSection, InstanceToAdd );
-	}
-}
-
-
-void FSequencer::RemoveMovieSceneInstance( UMovieSceneSection& MovieSceneSection, TSharedRef<FMovieSceneSequenceInstance> InstanceToRemove )
-{
-	UMovieSceneSequence* Sequence = InstanceToRemove->GetSequence();
-
-	SpawnRegister->CleanUpSequence(*InstanceToRemove, *this);
-
-	SequenceInstanceBySection.Remove( &MovieSceneSection );
-}
-
-
 void FSequencer::AddReferencedObjects( FReferenceCollector& Collector )
 {
 	Collector.AddReferencedObject( Settings );
 
-	for( int32 MovieSceneIndex = 0; MovieSceneIndex < SequenceInstanceStack.Num(); ++MovieSceneIndex )
+	if (UMovieSceneSequence* RootSequencePtr = RootSequence.Get())
 	{
-		UMovieSceneSequence* Sequence = SequenceInstanceStack[MovieSceneIndex]->GetSequence();
-		Collector.AddReferencedObject( Sequence );
+		Collector.AddReferencedObject( RootSequencePtr );
+	}
+
+	if (RootTemplateInstance.IsValid())
+	{
+		// Sequencer references all sub movie scene sequences contained within the root
+		for (auto& Pair : RootTemplateInstance.GetSubInstances())
+		{
+			UMovieSceneSequence* Sequence = Pair.Value.Sequence;
+			Collector.AddReferencedObject(Sequence);
+		}
 	}
 }
 
@@ -2039,7 +2126,7 @@ void FSequencer::ResetPerMovieSceneData()
 	UpdateTimeBoundsToFocusedMovieScene();
 	UpdateRuntimeInstances();
 
-	LabelManager.SetMovieScene( SequenceInstanceStack.Top()->GetSequence()->GetMovieScene() );
+	LabelManager.SetMovieScene( GetFocusedMovieSceneSequence()->GetMovieScene() );
 
 	// @todo run through all tracks for new movie scene changes
 	//  needed for audio track decompression
@@ -2048,29 +2135,9 @@ void FSequencer::ResetPerMovieSceneData()
 
 void FSequencer::UpdateRuntimeInstances()
 {
-	if (PlaybackContextAttribute.IsBound())
-	{
-		CachedPlaybackContext = PlaybackContextAttribute.Get();
-	}
-
-	// Refresh the current root instance
-	SequenceInstanceStack.Top()->RefreshInstance( *this );
-
-	EMovieSceneUpdateData UpdateData(ScrubPosition, ScrubPosition);
-	SequenceInstanceStack.Top()->Update(UpdateData, *this);
-
-	// Ensure any relevant spawned objects are cleaned up if we're playing back the master sequence
-	if (SequenceInstanceStack.Num() == 1)
-	{
-		if (!GetPlaybackRange().Contains(ScrubPosition))
-		{
-			SpawnRegister->OnMasterSequenceExpired(*this);
-		}
-	}
-
 	// If realtime is off, this needs to be called to update the pivot location when scrubbing.
 	GUnrealEd->UpdatePivotLocationForSelection();
-		
+	
 	// Redraw
 	FEditorSupportDelegates::RedrawAllViewports.Broadcast();
 }
@@ -2265,6 +2332,7 @@ TSharedRef<SWidget> FSequencer::MakeTransportControls(bool bExtended)
 			TransportControlArgs.WidgetsToCreate.Add(FTransportControlWidget(FOnMakeTransportWidget::CreateSP(this, &FSequencer::OnCreateTransportSetPlaybackEnd)));
 		}
 		TransportControlArgs.WidgetsToCreate.Add(FTransportControlWidget(ETransportControlWidgetType::Loop));
+		TransportControlArgs.bAreButtonsFocusable = false;
 	}
 
 	return EditorWidgetsModule.CreateTransportControl( TransportControlArgs );
@@ -2451,6 +2519,16 @@ void FSequencer::AttachTransportControlsToViewports()
 		
 	FEditorWidgetsModule& EditorWidgetsModule = FModuleManager::Get().LoadModuleChecked<FEditorWidgetsModule>( "EditorWidgets" );
 
+	// Transport controls in the viewport need to have something that is focusable to prevent mouse input dropping through to the viewport.
+	// We don't want the buttons themselves to be focusable, so we just add them to a parent box that is.
+	class SFocusableBox : public SBox
+	{
+		virtual bool SupportsKeyboardFocus() const
+		{
+			return true;
+		}
+	};
+
 	for (int32 i = 0; i < LevelViewports.Num(); ++i)
 	{
 		const TSharedPtr<ILevelViewport>& LevelViewport = LevelViewports[i];
@@ -2481,7 +2559,10 @@ void FSequencer::AttachTransportControlsToViewports()
 							.HAlign(HAlign_Center)
 							.VAlign(VAlign_Center)
 							[
-								MakeTransportControls(false)
+								SNew(SFocusableBox)
+								[
+									MakeTransportControls(false)
+								]
 							]
 							+SHorizontalBox::Slot()
 							.AutoWidth()
@@ -2533,11 +2614,12 @@ void FSequencer::DetachTransportControlsFromViewports()
 	}
 }
 
-UObject* FSequencer::FindSpawnedObjectOrTemplate(const FGuid& BindingId) const
+UObject* FSequencer::FindSpawnedObjectOrTemplate(const FGuid& BindingId)
 {
-	if (UObject* Object = GetFocusedMovieSceneSequenceInstance()->FindObject(BindingId, *this))
+	TArrayView<TWeakObjectPtr<>> Objects = FindObjectsInCurrentSequence(BindingId);
+	if (Objects.Num())
 	{
-		return Object;
+		return Objects[0].Get();
 	}
 
 	UMovieSceneSequence* Sequence = GetFocusedMovieSceneSequence();
@@ -2559,7 +2641,10 @@ UObject* FSequencer::FindSpawnedObjectOrTemplate(const FGuid& BindingId) const
 			UObject* ParentObject = ParentSpawnable->GetObjectTemplate();
 			if (ParentObject)
 			{
-				return Sequence->FindPossessableObject(BindingId, ParentObject);
+				for (UObject* Obj : Sequence->LocateBoundObjects(BindingId, ParentObject))
+				{
+					return Obj;
+				}
 			}
 		}
 	}
@@ -2673,8 +2758,10 @@ void FSequencer::HandleRecordingStarted(UMovieSceneSequence* Sequence)
 	UMovieSceneSubSection* Section = UMovieSceneSubSection::GetRecordingSection();
 	if(Section != nullptr)
 	{
-		Section->SetStartTime(ScrubPosition);
-		Section->SetEndTime(ScrubPosition + Settings->GetTimeSnapInterval());
+		float LocalTime = GetLocalTime();
+
+		Section->SetStartTime(LocalTime);
+		Section->SetEndTime(LocalTime + Settings->GetTimeSnapInterval());
 	}
 }
 
@@ -2721,8 +2808,8 @@ void FSequencer::HandleRecordingFinished(UMovieSceneSequence* Sequence)
 FReply FSequencer::OnStepForward()
 {
 	SetPlaybackStatus(EMovieScenePlayerStatus::Stepping);
-	float NewPosition = ScrubPosition + Settings->GetTimeSnapInterval();
-	SetGlobalTime(NewPosition, ESnapTimeMode::STM_Interval);
+	float NewPosition = GetLocalTime() + Settings->GetTimeSnapInterval();
+	SetLocalTime(NewPosition, ESnapTimeMode::STM_Interval);
 	return FReply::Handled();
 }
 
@@ -2730,9 +2817,9 @@ FReply FSequencer::OnStepForward()
 FReply FSequencer::OnStepBackward()
 {
 	SetPlaybackStatus(EMovieScenePlayerStatus::Stepping);
-	float NewPosition = ScrubPosition - Settings->GetTimeSnapInterval();
+	float NewPosition = GetLocalTime() - Settings->GetTimeSnapInterval();
 	const bool bAllowSnappingToFrames = true;
-	SetGlobalTime(NewPosition, ESnapTimeMode::STM_Interval);
+	SetLocalTime(NewPosition, ESnapTimeMode::STM_Interval);
 	return FReply::Handled();
 }
 
@@ -2740,14 +2827,14 @@ FReply FSequencer::OnStepBackward()
 FReply FSequencer::OnStepToEnd()
 {
 	SetPlaybackStatus(EMovieScenePlayerStatus::Stepping);
-	SetGlobalTime(GetPlaybackRange().GetUpperBoundValue());
+	SetLocalTime(GetPlaybackRange().GetUpperBoundValue());
 	return FReply::Handled();
 }
 
 FReply FSequencer::OnStepToBeginning()
 {
 	SetPlaybackStatus(EMovieScenePlayerStatus::Stepping);
-	SetGlobalTime(GetPlaybackRange().GetLowerBoundValue());
+	SetLocalTime(GetPlaybackRange().GetLowerBoundValue());
 	return FReply::Handled();
 }
 
@@ -2765,7 +2852,7 @@ FReply FSequencer::SetPlaybackEnd()
 	if (FocusedSequence)
 	{
 		TRange<float> CurrentRange = FocusedSequence->GetMovieScene()->GetPlaybackRange();
-		float NewPos = FMath::Max(GetGlobalTime(), CurrentRange.GetLowerBoundValue());
+		float NewPos = FMath::Max(GetLocalTime(), CurrentRange.GetLowerBoundValue());
 		SetPlaybackRange(TRange<float>(CurrentRange.GetLowerBoundValue(), NewPos));
 	}
 
@@ -2778,7 +2865,7 @@ FReply FSequencer::SetPlaybackStart()
 	if (FocusedSequence)
 	{
 		TRange<float> CurrentRange = FocusedSequence->GetMovieScene()->GetPlaybackRange();
-		float NewPos = FMath::Min(GetGlobalTime(), CurrentRange.GetUpperBoundValue());
+		float NewPos = FMath::Min(GetLocalTime(), CurrentRange.GetUpperBoundValue());
 		SetPlaybackRange(TRange<float>(NewPos, CurrentRange.GetUpperBoundValue()));	
 	}
 
@@ -2791,13 +2878,13 @@ FReply FSequencer::JumpToPreviousKey()
 	GetKeysFromSelection(ActiveKeyCollection);
 	if (ActiveKeyCollection.IsValid())
 	{
-		TRange<float> FindRange(TRange<float>::BoundsType(), GetGlobalTime());
+		TRange<float> FindRange(TRange<float>::BoundsType(), GetLocalTime());
 
 		TOptional<float> NewTime = ActiveKeyCollection->FindFirstKeyInRange(FindRange, EFindKeyDirection::Backwards);
 		if (NewTime.IsSet())
 		{
 			SetPlaybackStatus(EMovieScenePlayerStatus::Stepping);
-			SetGlobalTime(NewTime.GetValue());
+			SetLocalTimeDirectly(NewTime.GetValue());
 		}
 	}
 
@@ -2810,13 +2897,13 @@ FReply FSequencer::JumpToNextKey()
 	GetKeysFromSelection(ActiveKeyCollection);
 	if (ActiveKeyCollection.IsValid())
 	{
-		TRange<float> FindRange(GetGlobalTime(), TRange<float>::BoundsType());
+		TRange<float> FindRange(GetLocalTime(), TRange<float>::BoundsType());
 
 		TOptional<float> NewTime = ActiveKeyCollection->FindFirstKeyInRange(FindRange, EFindKeyDirection::Forwards);
 		if (NewTime.IsSet())
 		{
 			SetPlaybackStatus(EMovieScenePlayerStatus::Stepping);
-			SetGlobalTime(NewTime.GetValue());
+			SetLocalTimeDirectly(NewTime.GetValue());
 		}
 	}
 
@@ -2829,64 +2916,96 @@ bool FSequencer::IsLooping() const
 }
 
 
-void FSequencer::SetGlobalTimeLooped(float InTime)
+void FSequencer::SetLocalTimeLooped(float NewLocalTime)
 {
+	TOptional<EMovieScenePlayerStatus::Type> NewPlaybackStatus;
+
+	float NewGlobalTime = NewLocalTime * RootToLocalTransform.Inverse();
+
+	TRange<float> TimeBounds = GetTimeBounds();
+
 	bool bRestarted = false;
 	if (Settings->IsLooping())
 	{
 		const UMovieSceneSequence* FocusedSequence = GetFocusedMovieSceneSequence();
 		if (FocusedSequence)
 		{
-			if (InTime <= FocusedSequence->GetMovieScene()->GetPlaybackRange().GetLowerBoundValue() ||
-				InTime >= FocusedSequence->GetMovieScene()->GetPlaybackRange().GetUpperBoundValue())
+			if (NewLocalTime <= TimeBounds.GetLowerBoundValue() || NewLocalTime >= TimeBounds.GetUpperBoundValue())
 			{
-				InTime = PlayRate > 0 ? FocusedSequence->GetMovieScene()->GetPlaybackRange().GetLowerBoundValue() : FocusedSequence->GetMovieScene()->GetPlaybackRange().GetUpperBoundValue();
-				TimingManager->OnStartPlaying(InTime);
-				bRestarted = true;
+				NewGlobalTime = (PlayRate > 0 ? TimeBounds.GetLowerBoundValue() : TimeBounds.GetUpperBoundValue()) * RootToLocalTransform.Inverse();
+				TimingManager->OnStartPlaying(NewGlobalTime);
+
+				// Always evaluate from the start/end when looping
+				PlayPosition.Reset(NewGlobalTime);
 			}
 		}
 	}
 	else
 	{
-		TRange<float> TimeBounds = GetTimeBounds();
 		TRange<float> WorkingRange = GetClampRange();
 
 		bool bReachedEnd = false;
 		if (PlayRate > 0)
 		{
-			bReachedEnd = GetGlobalTime() < TimeBounds.GetUpperBoundValue() && InTime >= TimeBounds.GetUpperBoundValue();
+			bReachedEnd = GetLocalTime() < TimeBounds.GetUpperBoundValue() && NewLocalTime >= TimeBounds.GetUpperBoundValue();
 		}
 		else
 		{
-			bReachedEnd = GetGlobalTime() > TimeBounds.GetLowerBoundValue() && InTime <= TimeBounds.GetLowerBoundValue();
+			bReachedEnd = GetLocalTime() > TimeBounds.GetLowerBoundValue() && NewLocalTime <= TimeBounds.GetLowerBoundValue();
 		}
 
 		// Stop if we hit the playback range end
 		if (bReachedEnd)
 		{
-			InTime = PlayRate > 0 ? TimeBounds.GetUpperBoundValue() : TimeBounds.GetLowerBoundValue();
-			SetPlaybackStatus(EMovieScenePlayerStatus::Stopped);
+			NewGlobalTime = (PlayRate > 0 ? TimeBounds.GetUpperBoundValue() : TimeBounds.GetLowerBoundValue()) * RootToLocalTransform.Inverse();
+			TimingManager->OnStartPlaying(NewGlobalTime);
+			NewPlaybackStatus = EMovieScenePlayerStatus::Stopped;
 		}
 		// Constrain to the play range if necessary
 		else if (Settings->ShouldKeepCursorInPlayRange())
 		{
-			// Clamp to bound or jumpp back if necessary
-			if (InTime <= TimeBounds.GetLowerBoundValue() || 
-				InTime >= TimeBounds.GetUpperBoundValue())
+			// Clamp to bound or jump back if necessary
+			if (NewLocalTime <= TimeBounds.GetLowerBoundValue() || 
+				NewLocalTime >= TimeBounds.GetUpperBoundValue())
 			{
-				InTime = PlayRate > 0 ? TimeBounds.GetLowerBoundValue() : TimeBounds.GetUpperBoundValue();					
-				bRestarted = true;
-			}			
+				NewGlobalTime = (PlayRate > 0 ? TimeBounds.GetLowerBoundValue() : TimeBounds.GetUpperBoundValue()) * RootToLocalTransform.Inverse();
+				TimingManager->OnStartPlaying(NewGlobalTime);
+
+				// Always evaluate from the start/end when looping
+				PlayPosition.Reset(NewGlobalTime);
+			}
 		}
 		// Ensure the time is within the working range
-		else if (!WorkingRange.Contains(InTime))
+		else if (!WorkingRange.Contains(NewLocalTime))
 		{
-			InTime = FMath::Clamp(InTime, WorkingRange.GetLowerBoundValue(), WorkingRange.GetUpperBoundValue());
-			SetPlaybackStatus(EMovieScenePlayerStatus::Stopped);
+			NewGlobalTime = FMath::Clamp(NewLocalTime, WorkingRange.GetLowerBoundValue(), WorkingRange.GetUpperBoundValue()) * RootToLocalTransform.Inverse();
+			TimingManager->OnStartPlaying(NewGlobalTime);
+
+			// Always evaluate from the start/end when looping
+			PlayPosition.Reset(NewGlobalTime);
+			NewPlaybackStatus = EMovieScenePlayerStatus::Stopped;
 		}
 	}
 
-	SetGlobalTime(InTime, ESnapTimeMode::STM_None, bRestarted);
+	// Ensure the time is in the current view
+	ScrollIntoView(NewGlobalTime * RootToLocalTransform);
+
+	// Update the position before fixing it to the time interval
+	ScrubPosition = NewGlobalTime;
+
+	// Evaluate the sequence
+	UMovieScene* MovieScene = GetFocusedMovieSceneSequence()->GetMovieScene();
+	FMovieSceneEvaluationRange EvalRange = PlayPosition.PlayTo(ScrubPosition, MovieScene->GetOptionalFixedFrameInterval());
+
+	EvaluateInternal(EvalRange);
+
+	// Set the playback status if we need to
+	if (NewPlaybackStatus.IsSet())
+	{
+		SetPlaybackStatus(NewPlaybackStatus.GetValue());
+		// Evaluate the sequence with the new status
+		EvaluateInternal(EvalRange);
+	}
 }
 
 
@@ -2941,9 +3060,9 @@ void FSequencer::UpdateTimeBoundsToFocusedMovieScene()
 	SetViewRange(NewRange, EViewRangeInterpolation::Immediate);
 
 	// Make sure the current time is within the bounds
-	if (!TargetViewRange.Contains(ScrubPosition))
+	if (!TargetViewRange.Contains(GetLocalTime()))
 	{
-		ScrubPosition = LastViewRange.GetLowerBoundValue();
+		SetLocalTimeDirectly(LastViewRange.GetLowerBoundValue());
 		OnGlobalTimeChangedDelegate.Broadcast();
 	}
 }
@@ -2960,7 +3079,12 @@ TRange<float> FSequencer::GetTimeBounds() const
 		return TRange<float>( -100000.0f, 100000.0f );
 	}
 	
-	return FocusedSequence->GetMovieScene()->GetPlaybackRange();
+	if (Settings->ShouldEvaluateSubSequencesInIsolation() || ActiveTemplateIDs.Num() == 1)
+	{
+		return FocusedSequence->GetMovieScene()->GetPlaybackRange();
+	}
+
+	return SubSequenceRange;
 }
 
 
@@ -3087,10 +3211,24 @@ void FSequencer::OnScrubPositionChanged( float NewScrubPosition, bool bScrubbing
 
 	if (bClampToViewRange)
 	{
-		NewScrubPosition = FMath::Clamp(NewScrubPosition, TargetViewRange.GetLowerBoundValue(), TargetViewRange.GetUpperBoundValue());
+		float LowerBound = TargetViewRange.GetLowerBoundValue();
+		float UpperBound = TargetViewRange.GetUpperBoundValue();
+
+		if (Settings->GetIsSnapEnabled() && Settings->GetSnapPlayTimeToInterval())
+		{
+			LowerBound = Settings->SnapTimeToInterval(LowerBound);
+			UpperBound = Settings->SnapTimeToInterval(UpperBound);
+		}
+
+		NewScrubPosition = FMath::Clamp(NewScrubPosition, LowerBound, UpperBound);		
+	}
+
+	if (Settings->GetSnapPlayTimeToKeys())
+	{
+		NewScrubPosition = FindNearestKey(NewScrubPosition);
 	}
 
-	SetGlobalTimeDirectly( NewScrubPosition, Settings->GetSnapKeyTimesToKeys() ? ESnapTimeMode::STM_Keys : ESnapTimeMode::STM_None );
+	SetLocalTimeDirectly( NewScrubPosition );
 }
 
 
@@ -3106,6 +3244,8 @@ void FSequencer::OnEndScrubbing()
 	SetPlaybackStatus(EMovieScenePlayerStatus::Stepping);
 	AutoscrubOffset.Reset();
 	StopAutoscroll();
+
+	bNeedsEvaluate = true;
 }
 
 
@@ -3231,6 +3371,9 @@ FGuid FSequencer::MakeNewSpawnable( UObject& Object )
 	// @todo sequencer: Undo doesn't seem to be working at all
 	const FScopedTransaction Transaction( LOCTEXT("UndoAddingObject", "Add Object to MovieScene") );
 
+	UMovieSceneSequence* Sequence = GetFocusedMovieSceneSequence();
+	UMovieScene* MovieScene = Sequence->GetMovieScene();
+
 	FGuid NewGuid = AddSpawnable(Object);
 	if (!NewGuid.IsValid())
 	{
@@ -3248,7 +3391,7 @@ FGuid FSequencer::MakeNewSpawnable( UObject& Object )
 	Spawnable->SetSpawnOwnership(ESpawnOwnership::External);
 
 	// Spawn the object so we can position it correctly, it's going to get spawned anyway since things default to spawned.
-	UObject* SpawnedObject = SpawnRegister->SpawnObject(NewGuid, *GetFocusedMovieSceneSequenceInstance(), *this);
+	UObject* SpawnedObject = SpawnRegister->SpawnObject(NewGuid, *MovieScene, ActiveTemplateIDs.Top(), *this);
 
 	FTransformData DefaultTransform;
 
@@ -3351,7 +3494,7 @@ void FSequencer::AddSubSequence(UMovieSceneSequence* Sequence)
 {
 	// @todo Sequencer - sub-moviescenes This should be moved to the sub-moviescene editor
 
-	// Grab the MovieScene that is currently focused.  THis is the movie scene that will contain the sub-moviescene
+	// Grab the MovieScene that is currently focused.  This is the movie scene that will contain the sub-moviescene
 	UMovieScene* OwnerMovieScene = GetFocusedMovieSceneSequence()->GetMovieScene();
 
 	// @todo sequencer: Undo doesn't seem to be working at all
@@ -3413,8 +3556,7 @@ bool FSequencer::OnRequestNodeDeleted( TSharedRef<const FSequencerDisplayNode> N
 {
 	bool bAnythingRemoved = false;
 	
-	TSharedRef<FMovieSceneSequenceInstance> MovieSceneInstance = GetFocusedMovieSceneSequenceInstance();
-	UMovieSceneSequence* Sequence = MovieSceneInstance->GetSequence();
+	UMovieSceneSequence* Sequence = GetFocusedMovieSceneSequence();
 	UMovieScene* OwnerMovieScene = Sequence->GetMovieScene();
 
 	if ( NodeToBeDeleted->GetType() == ESequencerNode::Folder )
@@ -3466,7 +3608,7 @@ bool FSequencer::OnRequestNodeDeleted( TSharedRef<const FSequencerDisplayNode> N
 		// Try to remove as a spawnable first
 		if (OwnerMovieScene->RemoveSpawnable(BindingToRemove))
 		{
-			SpawnRegister->DestroySpawnedObject(BindingToRemove, *MovieSceneInstance, *this);
+			SpawnRegister->DestroySpawnedObject(BindingToRemove, ActiveTemplateIDs.Top(), *this);
 		}
 		// The guid should be associated with a possessable if it wasnt a spawnable
 		else if (OwnerMovieScene->RemovePossessable(BindingToRemove))
@@ -3492,6 +3634,7 @@ bool FSequencer::OnRequestNodeDeleted( TSharedRef<const FSequencerDisplayNode> N
 
 		if (Track != nullptr)
 		{
+			OwnerMovieScene->Modify();
 			if (OwnerMovieScene->IsAMasterTrack(*Track))
 			{
 				OwnerMovieScene->RemoveMasterTrack(*Track);
@@ -3586,22 +3729,22 @@ void FSequencer::PostUndo(bool bSuccess)
 	bUpdatingSequencerSelection = true;
 	SynchronizeSequencerSelectionWithExternalSelection();
 	bUpdatingSequencerSelection = false;
-	OnActivateSequenceEvent.Broadcast( *SequenceInstanceStack.Top() );
-}
+	OnActivateSequenceEvent.Broadcast(ActiveTemplateIDs.Top());
+}	
 
 
 void FSequencer::OnPreSaveWorld(uint32 SaveFlags, class UWorld* World)
 {
 	// Restore the saved state so that the level save can save that instead of the animated state.
-	RootMovieSceneSequenceInstance->RestoreState(*this);
+	RestorePreAnimatedState();
 }
 
 
 void FSequencer::OnPostSaveWorld(uint32 SaveFlags, class UWorld* World, bool bSuccess)
 {
 	// Reset the time after saving so that an update will be triggered to put objects back to their animated state.
-	RootMovieSceneSequenceInstance->RefreshInstance(*this);
-	SetGlobalTime(GetGlobalTime());
+	FMovieSceneEvaluationRange Range = PlayPosition.JumpTo(ScrubPosition, GetFocusedMovieSceneSequence()->GetMovieScene()->GetOptionalFixedFrameInterval());
+	EvaluateInternal(Range);
 }
 
 
@@ -3663,16 +3806,17 @@ void FSequencer::OnNewActorsDropped(const TArray<UObject*>& DroppedObjects, cons
 
 				UpdateRuntimeInstances();
 
-				UObject* RuntimeObject = GetFocusedMovieSceneSequenceInstance()->FindObject(Spawnable->GetGuid(), *this);
-
-				AActor* SpawnedActor = Cast<AActor>(RuntimeObject);
-				if (SpawnedActor)
+				for (TWeakObjectPtr<> WeakObject : FindBoundObjects(Spawnable->GetGuid(), ActiveTemplateIDs.Top()))
 				{
-					SpawnedActors.Add(SpawnedActor);
+					AActor* SpawnedActor = Cast<AActor>(WeakObject.Get());
+					if (SpawnedActor)
+					{
+						SpawnedActors.Add(SpawnedActor);
+						NewActor = SpawnedActor;
+					}
 				}
 
 				NewGuid = Spawnable->GetGuid();
-				NewActor = SpawnedActor;
 			}
 
 			if (bCreateAndAttachCamera)
@@ -3701,9 +3845,15 @@ void FSequencer::OnNewActorsDropped(const TArray<UObject*>& DroppedObjects, cons
 
 					UpdateRuntimeInstances();
 
-					UObject* RuntimeObject = GetFocusedMovieSceneSequenceInstance()->FindObject(Spawnable->GetGuid(), *this);
+					for (TWeakObjectPtr<> WeakObject : FindBoundObjects(Spawnable->GetGuid(), ActiveTemplateIDs.Top()))
+					{
+						NewCamera = Cast<ACineCameraActor>(WeakObject.Get());
+						if (NewCamera)
+						{
+							break;
+						}
+					}
 
-					NewCamera = Cast<ACineCameraActor>(RuntimeObject);
 					NewCameraGuid = Spawnable->GetGuid();
 
 					// Create an attach track
@@ -3739,17 +3889,17 @@ void FSequencer::OnNewActorsDropped(const TArray<UObject*>& DroppedObjects, cons
 
 					FKeyPropertyParams KeyPropertyParams(TArrayBuilder<UObject*>().Add(RailActor), PropertyPath, ESequencerKeyMode::ManualKeyForced);
 
-					float OriginalTime = GetGlobalTime();
+					float OriginalTime = GetLocalTime();
 
-					SetGlobalTime(GetPlaybackRange().GetLowerBoundValue());
+					SetLocalTimeDirectly(GetPlaybackRange().GetLowerBoundValue());
 					RailActor->CurrentPositionOnRail = 0.f;
 					KeyProperty(KeyPropertyParams);
 
-					SetGlobalTime(GetPlaybackRange().GetUpperBoundValue());
+					SetLocalTimeDirectly(GetPlaybackRange().GetUpperBoundValue());
 					RailActor->CurrentPositionOnRail = 1.f;
 					KeyProperty(KeyPropertyParams);
 
-					SetGlobalTime(OriginalTime);
+					SetLocalTimeDirectly(OriginalTime);
 				}
 
 				// New camera added, don't lock the view to the camera because we want to see where the camera rig was placed
@@ -3800,7 +3950,7 @@ void FSequencer::ActivateSequencerEditorMode()
 
 void FSequencer::OnPreBeginPIE(bool bIsSimulating)
 {
-	RootMovieSceneSequenceInstance->RestoreState(*this);
+	RestorePreAnimatedState();
 	PrePossessionViewTargets.Reset();
 
 	bNeedInstanceRefresh = true;
@@ -4011,7 +4161,7 @@ void FSequencer::SaveCurrentMovieScene()
 
 	TArray<UPackage*> PackagesToSave;
 	TArray<UMovieScene*> MovieScenesToSave;
-	GetDescendantMovieScenes(RootMovieSceneSequenceInstance->GetSequence(), MovieScenesToSave);
+	GetDescendantMovieScenes(GetRootMovieSceneSequence(), MovieScenesToSave);
 	for (auto MovieSceneToSave : MovieScenesToSave)
 	{
 		UPackage* MovieScenePackageToSave = MovieSceneToSave->GetOuter()->GetOutermost();
@@ -4027,11 +4177,9 @@ void FSequencer::SaveCurrentMovieScene()
 
 	FEditorFileUtils::PromptForCheckoutAndSave( PackagesToSave, bCheckDirty, bPromptToSave );
 
-	// Ensure the sequence is up-to-date
-	EMovieSceneUpdateData UpdateData(ScrubPosition, ScrubPosition);
-	SequenceInstanceStack.Top()->Update(UpdateData, *this);
-
-	FEditorSupportDelegates::RedrawAllViewports.Broadcast();
+	UpdateRuntimeInstances();
+	FMovieSceneEvaluationRange Range = PlayPosition.JumpTo(ScrubPosition, GetFocusedMovieSceneSequence()->GetMovieScene()->GetOptionalFixedFrameInterval());
+	EvaluateInternal(Range);
 }
 
 
@@ -4070,7 +4218,7 @@ void FSequencer::AddActors(const TArray<TWeakObjectPtr<AActor> >& InActors)
 	{
 		if (AActor* Actor = WeakActor.Get())
 		{
-			FGuid ExistingGuid = SequenceInstanceStack.Top()->FindObjectId(*Actor);
+			FGuid ExistingGuid = FindObjectId(*Actor, ActiveTemplateIDs.Top());
 			if (!ExistingGuid.IsValid())
 			{
 				FGuid PossessableGuid = CreateBinding(*Actor, Actor->GetActorLabel());
@@ -4140,10 +4288,7 @@ void FSequencer::SynchronizeExternalSelectionWithSequencerSelection()
 		// If the root node is an object node, try to get the actor nodes from it.
 		if ( ObjectBindingNode.IsValid() )
 		{
-			TArray<TWeakObjectPtr<UObject>> RuntimeObjects;
-			GetRuntimeObjects( GetFocusedMovieSceneSequenceInstance(), ObjectBindingNode->GetObjectBinding(), RuntimeObjects );
-			
-			for (auto RuntimeObject : RuntimeObjects )
+			for (auto RuntimeObject : FindBoundObjects(ObjectBindingNode->GetObjectBinding(), ActiveTemplateIDs.Top()) )
 			{
 				AActor* Actor = Cast<AActor>(RuntimeObject.Get());
 				if ( Actor != nullptr )
@@ -4222,10 +4367,7 @@ void FSequencer::SynchronizeSequencerSelectionWithExternalSelection()
 	TSet<TSharedRef<FSequencerDisplayNode>> NodesToSelect;
 	for ( TSharedRef<FSequencerObjectBindingNode> ObjectBindingNode : RootObjectBindingNodes )
 	{
-		TArray<TWeakObjectPtr<UObject>> RuntimeObjects;
-		GetRuntimeObjects( GetFocusedMovieSceneSequenceInstance(), ObjectBindingNode->GetObjectBinding(), RuntimeObjects );
-
-		for ( TWeakObjectPtr<UObject> RuntimeObjectPtr : RuntimeObjects )
+		for ( TWeakObjectPtr<UObject> RuntimeObjectPtr : FindBoundObjects(ObjectBindingNode->GetObjectBinding(), ActiveTemplateIDs.Top()) )
 		{
 			UObject* RuntimeObject = RuntimeObjectPtr.Get();
 			if ( RuntimeObject != nullptr && GEditor->GetSelectedActors()->IsSelected( RuntimeObject ) )
@@ -4360,10 +4502,19 @@ void FSequencer::DeleteSelectedItems()
 
 void FSequencer::AssignActor(FMenuBuilder& MenuBuilder, FGuid InObjectBinding)
 {
-	UObject* RuntimeObject = SequenceInstanceStack.Top()->FindObject(InObjectBinding, *this);
+	TSet<const AActor*> BoundObjects;
+	{
+		for (TWeakObjectPtr<> Ptr : FindObjectsInCurrentSequence(InObjectBinding))
+		{
+			if (const AActor* Actor = Cast<AActor>(Ptr.Get()))
+			{
+				BoundObjects.Add(Actor);
+			}
+		}
+	}
 
-	auto IsActorValidForAssignment = [=](const AActor* InActor, UObject* CurrentObject){
-		return CurrentObject != InActor;
+	auto IsActorValidForAssignment = [BoundObjects](const AActor* InActor){
+		return !BoundObjects.Contains(InActor);
 	};
 
 	using namespace SceneOutliner;
@@ -4381,7 +4532,7 @@ void FSequencer::AssignActor(FMenuBuilder& MenuBuilder, FGuid InObjectBinding)
 		InitOptions.ColumnMap.Add(FBuiltInColumnTypes::Label(), FColumnInfo(EColumnVisibility::Visible, 0));
 
 		// Only display actors that are not possessed already
-		InitOptions.Filters->AddFilterPredicate( FActorFilterPredicate::CreateLambda( IsActorValidForAssignment, RuntimeObject ) );
+		InitOptions.Filters->AddFilterPredicate( FActorFilterPredicate::CreateLambda( IsActorValidForAssignment ) );
 	}
 
 	// actor selector to allow the user to choose an actor
@@ -4413,6 +4564,8 @@ void FSequencer::DoAssignActor(AActor*const* InActors, int32 NumActors, FGuid In
 		return;
 	}
 
+	//@todo: this code doesn't work with multiple actors, or when the existing binding is bound to multiple actors
+
 	AActor* Actor = InActors[0];
 
 	if (Actor == nullptr)
@@ -4422,8 +4575,6 @@ void FSequencer::DoAssignActor(AActor*const* InActors, int32 NumActors, FGuid In
 
 	FScopedTransaction AssignActor( NSLOCTEXT("Sequencer", "AssignActor", "Assign Actor") );
 
-	RootMovieSceneSequenceInstance->SaveState(*this);
-
 	UMovieSceneSequence* OwnerSequence = GetFocusedMovieSceneSequence();
 	UMovieScene* OwnerMovieScene = OwnerSequence->GetMovieScene();
 
@@ -4431,13 +4582,15 @@ void FSequencer::DoAssignActor(AActor*const* InActors, int32 NumActors, FGuid In
 	OwnerSequence->Modify();
 	OwnerMovieScene->Modify();
 
-	UObject* RuntimeObject = SequenceInstanceStack.Top()->FindObject(InObjectBinding, *this);
+	TArrayView<TWeakObjectPtr<>> RuntimeObjects = FindObjectsInCurrentSequence(InObjectBinding);
+
+	UObject* RuntimeObject = RuntimeObjects.Num() ? RuntimeObjects[0].Get() : nullptr;
 
 	// Replace the object itself
 	FGuid ParentGuid;
 	{
 		// Get the object guid to assign, remove the binding if it already exists
-		ParentGuid = SequenceInstanceStack.Top()->FindObjectId(*Actor);
+		ParentGuid = FindObjectId(*Actor, ActiveTemplateIDs.Top());
 		FString NewActorLabel = Actor->GetActorLabel();
 		if (ParentGuid.IsValid())
 		{
@@ -4447,18 +4600,18 @@ void FSequencer::DoAssignActor(AActor*const* InActors, int32 NumActors, FGuid In
 
 		// Add this object
 		FMovieScenePossessable NewPossessable( NewActorLabel, Actor->GetClass());
-		ParentGuid = NewPossessable.GetGuid();
-		OwnerSequence->BindPossessableObject(ParentGuid, *Actor, GetPlaybackContext());
+		OwnerSequence->BindPossessableObject(NewPossessable.GetGuid(), *Actor, GetPlaybackContext());
 
 		// Replace
-		OwnerMovieScene->ReplacePossessable(InObjectBinding, ParentGuid, NewActorLabel);
+		OwnerMovieScene->ReplacePossessable(InObjectBinding, NewPossessable);
+
+		State.Invalidate(InObjectBinding, ActiveTemplateIDs.Top());
 	}
 
 	auto UpdateComponent = [&]( FGuid OldComponentGuid, UActorComponent* NewComponent )
 	{
 		// Get the object guid to assign, remove the binding if it already exists
-		FGuid NewComponentGuid = SequenceInstanceStack.Top()->FindObjectId( *NewComponent );
-		FString NewComponentLabel = NewComponent->GetName();
+		FGuid NewComponentGuid = FindObjectId( *NewComponent, ActiveTemplateIDs.Top() );
 		if ( NewComponentGuid.IsValid() )
 		{
 			OwnerMovieScene->RemovePossessable( NewComponentGuid );
@@ -4466,14 +4619,14 @@ void FSequencer::DoAssignActor(AActor*const* InActors, int32 NumActors, FGuid In
 		}
 
 		// Add this object
-		FMovieScenePossessable NewPossessable( NewComponentLabel, NewComponent->GetClass() );
-		NewComponentGuid = NewPossessable.GetGuid();
-		OwnerSequence->BindPossessableObject( NewComponentGuid, *NewComponent, Actor );
+		FMovieScenePossessable NewPossessable( NewComponent->GetName(), NewComponent->GetClass() );
+		OwnerSequence->BindPossessableObject( NewPossessable.GetGuid(), *NewComponent, Actor );
 
 		// Replace
-		OwnerMovieScene->ReplacePossessable( OldComponentGuid, NewComponentGuid, NewComponentLabel );
+		OwnerMovieScene->ReplacePossessable( OldComponentGuid, NewPossessable );
+		State.Invalidate(OldComponentGuid, ActiveTemplateIDs.Top());
 
-		FMovieScenePossessable* ThisPossessable = OwnerMovieScene->FindPossessable( NewComponentGuid );
+		FMovieScenePossessable* ThisPossessable = OwnerMovieScene->FindPossessable( NewPossessable.GetGuid() );
 		if ( ensure( ThisPossessable ) )
 		{
 			ThisPossessable->SetParent( ParentGuid );
@@ -4488,7 +4641,7 @@ void FSequencer::DoAssignActor(AActor*const* InActors, int32 NumActors, FGuid In
 		{
 			if (ComponentToReplace != nullptr)
 			{
-				FGuid ComponentGuid = SequenceInstanceStack.Top()->FindObjectId(*ComponentToReplace);
+				FGuid ComponentGuid = FindObjectId(*ComponentToReplace, ActiveTemplateIDs.Top());
 				if (ComponentGuid.IsValid())
 				{
 					UActorComponent* NewComponent = Actor->GetComponentByClass(ComponentToReplace->GetClass());
@@ -4543,7 +4696,7 @@ void FSequencer::DoAssignActor(AActor*const* InActors, int32 NumActors, FGuid In
 		}
 	}
 
-	RootMovieSceneSequenceInstance->RestoreState(*this);
+	RestorePreAnimatedState();
 
 	NotifyMovieSceneDataChanged( EMovieSceneDataChangeType::MovieSceneStructureItemsChanged );
 }
@@ -4655,9 +4808,9 @@ void FSequencer::PasteCopiedTracks(TArray<TSharedPtr<FSequencerObjectBindingNode
 	{
 		FGuid ObjectGuid = ObjectNode->GetObjectBinding();
 		UObject* Context = nullptr;
-		UObject* PossessableObject = GetFocusedMovieSceneSequence()->FindPossessableObject(ObjectGuid, Context);
+		TArray<UObject*, TInlineAllocator<1>> PossessableObjects = GetFocusedMovieSceneSequence()->LocateBoundObjects(ObjectGuid, Context);
 
-		if (!PossessableObject)
+		if (PossessableObjects.Num() == 0)
 		{
 			FNotificationInfo Info(LOCTEXT("InvalidTargetObject", "Invalid Target Object"));
 			Info.FadeInDuration = 0.1f;
@@ -4691,7 +4844,7 @@ void FSequencer::PasteCopiedTracks(TArray<TSharedPtr<FSequencerObjectBindingNode
 			FString PropertyName = CastedPropertyTrack->GetPropertyName().ToString();
 			PropertyName.RemoveFromStart("b", ESearchCase::CaseSensitive);
 			FString FunctionName = "Set" + PropertyName;
-			if (PossessableObject->FindFunction(FName(*FunctionName)))
+			if (PossessableObjects[0]->FindFunction(FName(*FunctionName)))
 			{
 				if (!GetFocusedMovieSceneSequence()->GetMovieScene()->AddGivenTrack(ImportedTrack, ObjectGuid))
 				{
@@ -4784,6 +4937,7 @@ void FSequencer::ToggleNodeActive()
 
 		for (auto Section : Sections)
 		{
+			Section->Modify();
 			Section->SetIsActive(bIsActive);
 		}
 	}
@@ -4823,6 +4977,7 @@ void FSequencer::ToggleNodeLocked()
 
 		for (auto Section : Sections)
 		{
+			Section->Modify();
 			Section->SetIsLocked(bIsLocked);
 		}
 	}
@@ -4850,12 +5005,53 @@ bool FSequencer::IsNodeLocked() const
 	return NumSections > 0;
 }
 
+void FSequencer::SaveSelectedNodesSpawnableState()
+{
+	const FScopedTransaction Transaction( LOCTEXT("SaveSpawnableState", "Save spawnable state") );
+
+	UMovieScene* MovieScene = GetFocusedMovieSceneSequence()->GetMovieScene();
+
+	MovieScene->Modify();
+
+	TArray<FMovieSceneSpawnable*> Spawnables;
+
+	for (const TSharedRef<FSequencerDisplayNode>& Node : Selection.GetSelectedOutlinerNodes())
+	{
+		if (Node->GetType() == ESequencerNode::Object)
+		{
+			FMovieSceneSpawnable* Spawnable = MovieScene->FindSpawnable(StaticCastSharedRef<FSequencerObjectBindingNode>(Node)->GetObjectBinding());
+			if (Spawnable)
+			{
+				Spawnables.Add(Spawnable);
+			}
+		}
+	}
+
+	FScopedSlowTask SlowTask(Spawnables.Num(), LOCTEXT("SaveSpawnableStateProgress", "Saving selected spawnables"));
+	SlowTask.MakeDialog(true);
+
+	TArray<AActor*> PossessedActors;
+	for (FMovieSceneSpawnable* Spawnable : Spawnables)
+	{
+		SlowTask.EnterProgressFrame();
+		
+		SpawnRegister->SaveDefaultSpawnableState(*Spawnable, ActiveTemplateIDs.Top(), *this);
+
+		if (GWarn->ReceivedUserCancel())
+		{
+			break;
+		}
+	}
+
+	NotifyMovieSceneDataChanged(EMovieSceneDataChangeType::MovieSceneStructureItemsChanged);
+}
+
 void FSequencer::ConvertToSpawnable(TSharedRef<FSequencerObjectBindingNode> NodeToBeConverted)
 {
 	const FScopedTransaction Transaction( LOCTEXT("ConvertSelectedNodeSpawnable", "Convert Node to Spawnables") );
 
 	// Ensure we're in a non-possessed state
-	RootMovieSceneSequenceInstance->RestoreState(*this);
+	RestorePreAnimatedState();
 	GetFocusedMovieSceneSequence()->GetMovieScene()->Modify();
 	FMovieScenePossessable* Possessable = GetFocusedMovieSceneSequence()->GetMovieScene()->FindPossessable(NodeToBeConverted->GetObjectBinding());
 	if (Possessable)
@@ -4873,7 +5069,7 @@ void FSequencer::ConvertSelectedNodesToSpawnables()
 	UMovieScene* MovieScene = GetFocusedMovieSceneSequence()->GetMovieScene();
 
 	// Ensure we're in a non-possessed state
-	RootMovieSceneSequenceInstance->RestoreState(*this);
+	RestorePreAnimatedState();
 	MovieScene->Modify();
 
 	TArray<TSharedRef<FSequencerObjectBindingNode>> ObjectBindingNodes;
@@ -4908,10 +5104,13 @@ void FSequencer::ConvertSelectedNodesToSpawnables()
 
 			UpdateRuntimeInstances();
 
-			UObject* RuntimeObject = GetFocusedMovieSceneSequenceInstance()->FindObject(Spawnable->GetGuid(), *this);
-
-			AActor* SpawnedActor = Cast<AActor>(RuntimeObject);
-			SpawnedActors.Add(SpawnedActor);
+			for (TWeakObjectPtr<> WeakObject : FindBoundObjects(Spawnable->GetGuid(), ActiveTemplateIDs.Top()))
+			{
+				if (AActor* SpawnedActor = Cast<AActor>(WeakObject.Get()))
+				{
+					SpawnedActors.Add(SpawnedActor);
+				}
+			}
 		}
 
 		if (GWarn->ReceivedUserCancel())
@@ -4947,12 +5146,17 @@ FMovieSceneSpawnable* FSequencer::ConvertToSpawnableInternal(FGuid PossessableGu
 	UMovieScene* MovieScene = Sequence->GetMovieScene();
 
 	// Find the object in the environment
-	TSharedRef<FMovieSceneSequenceInstance> FocusedSequenceInstance = SequenceInstanceStack.Top();
-	UMovieSceneSequence* FocusedSequence = FocusedSequenceInstance->GetSequence();
 	FMovieScenePossessable* Possessable = MovieScene->FindPossessable(PossessableGuid);
 
-	UObject* FoundObject = FocusedSequenceInstance->FindObject(PossessableGuid, *this);
-	if (!FoundObject || !FocusedSequence)
+	//@todo: this code doesn't work where multiple objects are bound
+	TArrayView<TWeakObjectPtr<>> FoundObjects = FindBoundObjects(PossessableGuid, ActiveTemplateIDs.Top());
+	if (FoundObjects.Num() != 1)
+	{
+		return nullptr;
+	}
+
+	UObject* FoundObject = FoundObjects[0].Get();
+	if (!FoundObject)
 	{
 		return nullptr;
 	}
@@ -4981,20 +5185,8 @@ FMovieSceneSpawnable* FSequencer::ConvertToSpawnableInternal(FGuid PossessableGu
 		for (int32 Index = 0; Index < MovieScene->GetPossessableCount(); ++Index)
 		{
 			FMovieScenePossessable& MovieScenePossessable = MovieScene->GetPossessable(Index);
-
-			bool bBelongsToNewSpawnable = MovieScenePossessable.GetParent() == PersistentGuid;
-			if (!bBelongsToNewSpawnable)
-			{
-				// This is potentially slow, but fixes old data where we didn't assign parent object guids
-				UObject* PotentialChildObject = FocusedSequenceInstance->FindObject(MovieScenePossessable.GetGuid(), *this);
-				bBelongsToNewSpawnable = PotentialChildObject && FocusedSequence->GetParentObject(PotentialChildObject) == FoundObject;
-			}
-			
-			if (bBelongsToNewSpawnable)
-			{
-				MovieScenePossessable.SetParent(PersistentGuid);
-				Spawnable->AddChildPossessable(MovieScenePossessable.GetGuid());
-			}
+			MovieScenePossessable.SetParent(PersistentGuid);
+			Spawnable->AddChildPossessable(MovieScenePossessable.GetGuid());
 		}
 	}
 	
@@ -5015,7 +5207,7 @@ FMovieSceneSpawnable* FSequencer::ConvertToSpawnableInternal(FGuid PossessableGu
 	}
 
 	SetupDefaultsForSpawnable(Spawnable->GetGuid(), DefaultTransform);
-	SetGlobalTimeDirectly(ScrubPosition);
+	SetLocalTimeDirectly(ScrubPosition);
 
 	return Spawnable;
 }
@@ -5025,7 +5217,7 @@ void FSequencer::ConvertToPossessable(TSharedRef<FSequencerObjectBindingNode> No
 	const FScopedTransaction Transaction( LOCTEXT("ConvertSelectedNodePossessable", "Convert Node to Possessables") );
 
 	// Ensure we're in a non-possessed state
-	RootMovieSceneSequenceInstance->RestoreState(*this);
+	RestorePreAnimatedState();
 	GetFocusedMovieSceneSequence()->GetMovieScene()->Modify();
 	FMovieSceneSpawnable* Spawnable = GetFocusedMovieSceneSequence()->GetMovieScene()->FindSpawnable(NodeToBeConverted->GetObjectBinding());
 	if (Spawnable)
@@ -5074,10 +5266,13 @@ void FSequencer::ConvertSelectedNodesToPossessables()
 
 			UpdateRuntimeInstances();
 
-			UObject* RuntimeObject = GetFocusedMovieSceneSequenceInstance()->FindObject(Possessable->GetGuid(), *this);
-
-			AActor* PossessedActor = Cast<AActor>(RuntimeObject);
-			PossessedActors.Add(PossessedActor);
+			for (TWeakObjectPtr<> WeakObject : FindBoundObjects(Possessable->GetGuid(), ActiveTemplateIDs.Top()))
+			{
+				if (AActor* PossessedActor = Cast<AActor>(WeakObject.Get()))
+				{
+					PossessedActors.Add(PossessedActor);
+				}
+			}
 		}
 
 		if (GWarn->ReceivedUserCancel())
@@ -5113,8 +5308,7 @@ FMovieScenePossessable* FSequencer::ConvertToPossessableInternal(FGuid Spawnable
 	UMovieScene* MovieScene = Sequence->GetMovieScene();
 
 	// Find the object in the environment
-	TSharedRef<FMovieSceneSequenceInstance> FocusedSequenceInstance = SequenceInstanceStack.Top();
-	UMovieSceneSequence* FocusedSequence = FocusedSequenceInstance->GetSequence();
+	UMovieSceneSequence* FocusedSequence = GetFocusedMovieSceneSequence();
 	FMovieSceneSpawnable* Spawnable = MovieScene->FindSpawnable(SpawnableGuid);
 	if (!Spawnable || !Spawnable->GetObjectTemplate())
 	{
@@ -5160,7 +5354,7 @@ FMovieScenePossessable* FSequencer::ConvertToPossessableInternal(FGuid Spawnable
 
 		if (MovieScene->RemoveSpawnable(PersistentGuid))
 		{
-			SpawnRegister->DestroySpawnedObject(BenignSpawnableGuid, *FocusedSequenceInstance, *this);
+			SpawnRegister->DestroySpawnedObject(BenignSpawnableGuid, ActiveTemplateIDs.Top(), *this);
 		}
 
 		Sequence->BindPossessableObject(BenignSpawnableGuid, *PossessedActor, PlaybackContext);
@@ -5168,7 +5362,7 @@ FMovieScenePossessable* FSequencer::ConvertToPossessableInternal(FGuid Spawnable
 	
 	GEditor->SelectActor(PossessedActor, false, true);
 
-	SetGlobalTimeDirectly(ScrubPosition);
+	SetLocalTimeDirectly(ScrubPosition);
 
 	return Possessable;
 }
@@ -5288,11 +5482,11 @@ void FSequencer::Pause()
 	SetPlaybackStatus(EMovieScenePlayerStatus::Stopped);
 
 	// Snap time when stopping play
-	SetGlobalTime(ScrubPosition, ESnapTimeMode::STM_Interval);
+	SetLocalTime(GetLocalTime(), ESnapTimeMode::STM_Interval);
 
 	// Update on stop (cleans up things like sounds that are playing)
-	EMovieSceneUpdateData UpdateData(ScrubPosition, ScrubPosition);
-	SequenceInstanceStack.Top()->Update(UpdateData, *this );
+	FMovieSceneEvaluationRange Range = PlayPosition.JumpTo(ScrubPosition, GetFocusedMovieSceneSequence()->GetMovieScene()->GetOptionalFixedFrameInterval());
+	EvaluateInternal(Range);
 }
 
 void FSequencer::StepForward()
@@ -5388,7 +5582,7 @@ void FSequencer::SetKey()
 		{
 			if (KeyArea->GetOwningSection()->TryModify())
 			{
-				KeyArea->AddKeyUnique(GetGlobalTime(), GetKeyInterpolation());
+				KeyArea->AddKeyUnique(GetLocalTime(), GetKeyInterpolation());
 			}
 		}
 	}
@@ -5692,7 +5886,7 @@ void FSequencer::OnClipboardUsed(TSharedPtr<FMovieSceneClipboard> Clipboard)
 
 void FSequencer::DiscardChanges()
 {
-	if (SequenceInstanceStack.Num() == 0)
+	if (ActiveTemplateIDs.Num() == 0)
 	{
 		return;
 	}
@@ -5704,7 +5898,7 @@ void FSequencer::DiscardChanges()
 		return;
 	}
 
-	UMovieSceneSequence* EditedSequence = SequenceInstanceStack.Top()->GetSequence();
+	UMovieSceneSequence* EditedSequence = GetFocusedMovieSceneSequence();
 
 	if (EditedSequence == nullptr)
 	{
@@ -5856,7 +6050,14 @@ void FSequencer::CreateCamera()
 		// Destroy the old actor
 		World->EditorDestroyActor(NewCamera, false);
 
-		NewCamera = Cast<ACineCameraActor>(GetFocusedMovieSceneSequenceInstance()->FindObject(CameraGuid, *this));
+		for (TWeakObjectPtr<UObject>& Object : FindBoundObjects(CameraGuid, ActiveTemplateIDs.Top()))
+		{
+			NewCamera = Cast<ACineCameraActor>(Object.Get());
+			if (NewCamera)
+			{
+				break;
+			}
+		}
 		ensure(NewCamera);
 	}
 	else
@@ -5920,7 +6121,7 @@ void FSequencer::NewCameraAdded(ACineCameraActor* NewCamera, FGuid CameraGuid, b
 
 	if (CameraCutTrack)
 	{
-		UMovieSceneSection* Section = MovieSceneHelpers::FindSectionAtTime(CameraCutTrack->GetAllSections(), GetGlobalTime());
+		UMovieSceneSection* Section = MovieSceneHelpers::FindSectionAtTime(CameraCutTrack->GetAllSections(), GetLocalTime());
 		UMovieSceneCameraCutSection* CameraCutSection = Cast<UMovieSceneCameraCutSection>(Section);
 
 		if (CameraCutSection)
@@ -5963,9 +6164,7 @@ void FSequencer::FixActorReferences()
 		// Possesbles with parents are components so ignore them.
 		if ( Possessable.GetParent().IsValid() == false )
 		{
-			TArray<TWeakObjectPtr<UObject>> RuntimeObjects;
-			GetRuntimeObjects( GetFocusedMovieSceneSequenceInstance(), Possessable.GetGuid(), RuntimeObjects );
-			if ( RuntimeObjects.Num() == 0 )
+			if ( FindBoundObjects(Possessable.GetGuid(), ActiveTemplateIDs.Top()).Num() == 0 )
 			{
 				ActorsPossessablesToFix.Add( Possessable );
 			}
@@ -6151,7 +6350,7 @@ void FSequencer::ImportFBX()
 		}
 	}
 	
-	if (MovieSceneToolHelpers::ImportFBX(MovieScene, *GetFocusedMovieSceneSequenceInstance(), *this, ObjectBindingNameMap))
+	if (MovieSceneToolHelpers::ImportFBX(MovieScene, *this, ObjectBindingNameMap))
 	{
 		NotifyMovieSceneDataChanged( EMovieSceneDataChangeType::MovieSceneStructureItemAdded );
 	}
@@ -6210,7 +6409,7 @@ void FSequencer::ExportFBX()
 
 		const bool bSelectedOnly = Bindings.Num() != 0;
 
-		UnFbx::FFbxExporter::FLevelSequenceNodeNameAdapter NodeNameAdapter(GetFocusedMovieSceneSequence()->GetMovieScene(), this);
+		UnFbx::FFbxExporter::FLevelSequenceNodeNameAdapter NodeNameAdapter(GetFocusedMovieSceneSequence()->GetMovieScene(), this, GetFocusedTemplateID());
 
 		// Export the persistent level and all of it's actors
 		UWorld* World = Cast<UWorld>( GetPlaybackContext() );
@@ -6227,7 +6426,7 @@ void FSequencer::ExportFBX()
 		}
 
 		// Export the movie scene data.
-		Exporter->ExportLevelSequence( GetFocusedMovieSceneSequence()->GetMovieScene(), Bindings, this );
+		Exporter->ExportLevelSequence( GetFocusedMovieSceneSequence()->GetMovieScene(), Bindings, this, GetFocusedTemplateID() );
 
 		// Save to disk
 		Exporter->WriteToFile( *ExportFilename );
@@ -6268,7 +6467,7 @@ void FSequencer::CloseEntryPopupMenu()
 void FSequencer::TrimSection(bool bTrimLeft)
 {
 	FScopedTransaction TrimSectionTransaction( NSLOCTEXT("Sequencer", "TrimSection_Transaction", "Trim Section") );
-	MovieSceneToolHelpers::TrimSection(Selection.GetSelectedSections(), GetGlobalTime(), bTrimLeft);
+	MovieSceneToolHelpers::TrimSection(Selection.GetSelectedSections(), GetLocalTime(), bTrimLeft);
 	NotifyMovieSceneDataChanged( EMovieSceneDataChangeType::TrackValueChanged );
 }
 
@@ -6276,7 +6475,7 @@ void FSequencer::TrimSection(bool bTrimLeft)
 void FSequencer::SplitSection()
 {
 	FScopedTransaction SplitSectionTransaction( NSLOCTEXT("Sequencer", "SplitSection_Transaction", "Split Section") );
-	MovieSceneToolHelpers::SplitSection(Selection.GetSelectedSections(), GetGlobalTime());
+	MovieSceneToolHelpers::SplitSection(Selection.GetSelectedSections(), GetLocalTime());
 	NotifyMovieSceneDataChanged( EMovieSceneDataChangeType::MovieSceneStructureItemAdded );
 }
 
@@ -6317,7 +6516,8 @@ void FSequencer::BindCommands()
 
 	SequencerCommandBindings->MapAction(
 		Commands.SetStartPlaybackRange,
-		FExecuteAction::CreateSP( this, &FSequencer::SetPlaybackRangeStart ) );
+		FExecuteAction::CreateSP( this, &FSequencer::SetPlaybackRangeStart ),
+		FCanExecuteAction::CreateSP( this, &FSequencer::IsViewingMasterSequence ) );
 
 	SequencerCommandBindings->MapAction(
 		Commands.ResetViewRange,
@@ -6337,7 +6537,8 @@ void FSequencer::BindCommands()
 
 	SequencerCommandBindings->MapAction(
 		Commands.SetEndPlaybackRange,
-		FExecuteAction::CreateSP( this, &FSequencer::SetPlaybackRangeEnd ) );
+		FExecuteAction::CreateSP( this, &FSequencer::SetPlaybackRangeEnd ),
+		FCanExecuteAction::CreateSP( this, &FSequencer::IsViewingMasterSequence ) );
 
 	SequencerCommandBindings->MapAction(
 		Commands.ExpandAllNodesAndDescendants,
@@ -6382,11 +6583,7 @@ void FSequencer::BindCommands()
 				FMovieScenePossessable* Possessable = MovieScene->FindPossessable(static_cast<FSequencerObjectBindingNode&>(*Node).GetObjectBinding());
 				if (Possessable && !Possessable->GetParent().IsValid())
 				{
-					UObject* RuntimeObject = GetFocusedMovieSceneSequenceInstance()->FindObject(Possessable->GetGuid(), *this);
-					if (RuntimeObject)
-					{
-						return true;
-					}
+					return true;
 				}
 			}
 		}
@@ -6398,8 +6595,7 @@ void FSequencer::BindCommands()
 		FCanExecuteAction::CreateLambda(CanConvertToSpawnables)
 	);
 
-	// We can convert to possessable if anything selected is a spawnable
-	auto CanConvertToPossessables = [this]{
+	auto AreSpawnablesSelected = [this]{
 		UMovieScene* MovieScene = GetFocusedMovieSceneSequence()->GetMovieScene();
 
 		for (const TSharedRef<FSequencerDisplayNode>& Node : Selection.GetSelectedOutlinerNodes())
@@ -6418,7 +6614,14 @@ void FSequencer::BindCommands()
 	SequencerCommandBindings->MapAction(
 		FSequencerCommands::Get().ConvertToPossessable,
 		FExecuteAction::CreateSP(this, &FSequencer::ConvertSelectedNodesToPossessables),
-		FCanExecuteAction::CreateLambda(CanConvertToPossessables)
+		FCanExecuteAction::CreateLambda(AreSpawnablesSelected)
+	);
+
+
+	SequencerCommandBindings->MapAction(
+		FSequencerCommands::Get().SaveCurrentSpawnableState,
+		FExecuteAction::CreateSP(this, &FSequencer::SaveSelectedNodesSpawnableState),
+		FCanExecuteAction::CreateLambda(AreSpawnablesSelected)
 	);
 
 	SequencerCommandBindings->MapAction(
@@ -6657,6 +6860,16 @@ void FSequencer::BindCommands()
 		FCanExecuteAction::CreateLambda( []{ return true; } ),
 		FIsActionChecked::CreateLambda( [this]{ return Settings->ShouldKeepPlayRangeInSectionBounds(); } ) );
 
+	SequencerCommandBindings->MapAction(
+		Commands.ToggleEvaluateSubSequencesInIsolation,
+		FExecuteAction::CreateLambda( [this]{
+			Settings->SetEvaluateSubSequencesInIsolation( !Settings->ShouldEvaluateSubSequencesInIsolation() );
+			FMovieSceneEvaluationRange Range = PlayPosition.JumpTo(ScrubPosition, GetFocusedMovieSceneSequence()->GetMovieScene()->GetOptionalFixedFrameInterval());
+			EvaluateInternal(Range);
+		} ),
+		FCanExecuteAction::CreateLambda( [this]{ return ActiveTemplateIDs.Num() > 1; } ),
+		FIsActionChecked::CreateLambda( [this]{ return Settings->ShouldEvaluateSubSequencesInIsolation(); } ) );
+
 	SequencerCommandBindings->MapAction(
 		Commands.RenderMovie,
 		FExecuteAction::CreateLambda([this]{ RenderMovieInternal(GetPlaybackRange().GetLowerBoundValue(), GetPlaybackRange().GetUpperBoundValue()); })
@@ -6671,12 +6884,12 @@ void FSequencer::BindCommands()
 		Commands.DiscardChanges,
 		FExecuteAction::CreateSP(this, &FSequencer::DiscardChanges),
 		FCanExecuteAction::CreateLambda([this]{
-			if (SequenceInstanceStack.Num() == 0)
+			UMovieSceneSequence* EditedSequence = GetFocusedMovieSceneSequence();
+			if (!EditedSequence)
 			{
 				return false;
 			}
 
-			UMovieSceneSequence* EditedSequence = SequenceInstanceStack.Top()->GetSequence();
 			UPackage* EditedPackage = EditedSequence->GetOutermost();
 
 			return ((EditedSequence != nullptr) && (EditedPackage->FileSize != 0) && EditedPackage->IsDirty());
diff --git a/Engine/Source/Editor/Sequencer/Private/Sequencer.h b/Engine/Source/Editor/Sequencer/Private/Sequencer.h
--- a/Engine/Source/Editor/Sequencer/Private/Sequencer.h
+++ b/Engine/Source/Editor/Sequencer/Private/Sequencer.h
@@ -7,12 +7,12 @@
 #include "MovieSceneClipboard.h"
 #include "MovieScenePossessable.h"
 #include "SequencerLabelManager.h"
+#include "MovieSceneEvaluationTemplateInstance.h"
 #include "LevelEditor.h"
 #include "SequencerTimingManager.h"
 
 class ACineCameraActor;
 class FMenuBuilder;
-class FMovieSceneSequenceInstance;
 class IDetailKeyframeHandler;
 class IMenu;
 class ISequencerEditTool;
@@ -24,7 +24,7 @@ class UMovieScene;
 class UMovieSceneSection;
 class UMovieSceneSequence;
 class UWorld;
-class IMovieSceneSpawnRegister;
+class FMovieSceneSpawnRegister;
 class FSequencerNodeTree;
 
 struct ISequencerHotspot;
@@ -106,6 +106,13 @@ public:
 	/** Select all keys that fall into the current selection range. */
 	void SelectKeysInSelectionRange();
 
+	/**
+	 * Get the currently viewed sub sequence range
+	 *
+	 * @return The sub sequence range, or an empty optional if we're viewing the root.
+	 */
+	TOptional<TRange<float>> GetSubSequenceRange() const;
+
 public:
 
 	/**
@@ -131,7 +138,7 @@ public:
 	 */
 	void SetPlaybackRangeEnd()
 	{
-		SetPlaybackRange(TRange<float>(GetPlaybackRange().GetLowerBoundValue(), GetGlobalTime()));
+		SetPlaybackRange(TRange<float>(GetPlaybackRange().GetLowerBoundValue(), GetLocalTime()));
 	}
 
 	/**
@@ -141,7 +148,7 @@ public:
 	 */
 	void SetPlaybackRangeStart()
 	{
-		SetPlaybackRange(TRange<float>(GetGlobalTime(), GetPlaybackRange().GetUpperBoundValue()));
+		SetPlaybackRange(TRange<float>(GetLocalTime(), GetPlaybackRange().GetUpperBoundValue()));
 	}
 
 public:
@@ -178,7 +185,7 @@ public:
 	/**
 	 * Pops the current focused movie scene from the stack.  The parent of this movie scene will be come the focused one
 	 */
-	void PopToSequenceInstance( TSharedRef<FMovieSceneSequenceInstance> SequenceInstance );
+	void PopToSequenceInstance( FMovieSceneSequenceIDRef SequenceID );
 
 	/** Deletes the passed in sections. */
 	void DeleteSections(const TSet<TWeakObjectPtr<UMovieSceneSection> > & Sections);
@@ -198,6 +205,9 @@ public:
 	/** Are there keys to snap? */
 	bool CanSnapToFrame() const;
 
+	/** Transform the selected keys and sections */
+	void TransformSelectedKeysAndSections(float InDeltaTime);
+
 	/**
 	 * @return Movie scene tools used by the sequencer
 	 */
@@ -251,6 +261,11 @@ protected:
 	 */
 	void SetupDefaultsForSpawnable( const FGuid& Guid, const FTransformData& Transform );
 
+	/**
+	 * Save default spawnable state for the currently selected objects
+	 */
+	void SaveSelectedNodesSpawnableState();
+
 public:
 
 	/**
@@ -376,7 +391,9 @@ public:
 	bool IsLooping() const;
 
 	/** Set the new global time, accounting for looping options */
-	void SetGlobalTimeLooped(float InTime);
+	void SetLocalTimeLooped(float InTime);
+
+	float AutoScroll(float InTime, ESnapTimeMode SnapTimeMode);
 
 	EPlaybackMode::Type GetPlaybackMode() const;
 
@@ -532,14 +549,13 @@ public:
 
 	virtual void Close() override;
 	virtual TSharedRef<SWidget> GetSequencerWidget() const override;
+	virtual FMovieSceneSequenceIDRef GetRootTemplateID() const override { return ActiveTemplateIDs[0]; }
+	virtual FMovieSceneSequenceIDRef GetFocusedTemplateID() const override { return ActiveTemplateIDs.Top(); }
 	virtual UMovieSceneSequence* GetRootMovieSceneSequence() const override;
 	virtual UMovieSceneSequence* GetFocusedMovieSceneSequence() const override;
+	virtual FMovieSceneRootEvaluationTemplateInstance& GetEvaluationTemplate() override { return RootTemplateInstance; }
 	virtual void ResetToNewRootSequence(UMovieSceneSequence& NewSequence) override;
-	virtual TSharedRef<FMovieSceneSequenceInstance> GetRootMovieSceneSequenceInstance() const override;
-	virtual TSharedRef<FMovieSceneSequenceInstance> GetFocusedMovieSceneSequenceInstance() const override;
 	virtual void FocusSequenceInstance( UMovieSceneSubSection& InSubSection ) override;
-	virtual TSharedRef<FMovieSceneSequenceInstance> GetSequenceInstanceForSection(UMovieSceneSection& Section) const override;
-	virtual bool HasSequenceInstanceForSection(UMovieSceneSection& Section) const override;
 	virtual EAutoKeyMode GetAutoKeyMode() const override;
 	virtual void SetAutoKeyMode(EAutoKeyMode AutoKeyMode) override;
 	virtual bool GetKeyAllEnabled() const override;
@@ -552,10 +568,12 @@ public:
 	virtual void SetInfiniteKeyAreas(bool bInfiniteKeyAreas) override;
 	virtual bool GetAutoSetTrackDefaults() const override;
 	virtual bool IsRecordingLive() const override;
-	virtual float GetCurrentLocalTime(UMovieSceneSequence& InMovieSceneSequence) override;
+	virtual float GetLocalTime() const override;
 	virtual float GetGlobalTime() const override;
-	virtual void SetGlobalTime(float Time, ESnapTimeMode SnapTimeMode = ESnapTimeMode::STM_None, bool bRestarted = false) override;
-	virtual void SetGlobalTimeDirectly(float Time, ESnapTimeMode SnapTimeMode = ESnapTimeMode::STM_None, bool bRestarted = false) override;
+	virtual void SetLocalTime(float Time, ESnapTimeMode SnapTimeMode = ESnapTimeMode::STM_None) override;
+	virtual void SetLocalTimeDirectly(float NewTime) override;
+	virtual void SetGlobalTime(float Time) override;
+	virtual void ForceEvaluate() override;
 	virtual void SetPerspectiveViewportPossessionEnabled(bool bEnabled) override;
 	virtual void SetPerspectiveViewportCameraCutEnabled(bool bEnabled) override;
 	virtual void RenderMovie(UMovieSceneSection* InSection) const override;
@@ -583,7 +601,6 @@ public:
 	virtual FGuid CreateBinding(UObject& InObject, const FString& InName) override;
 	virtual UObject* GetPlaybackContext() const override;
 	virtual TArray<UObject*> GetEventContexts() const override;
-	virtual void GetAllKeyedProperties(UObject& Object, TSet<UProperty*>& OutProperties) override;
 	virtual FOnActorAddedToSequencer& OnActorAddedToSequencer() override;
 	virtual FOnPreSave& OnPreSave() override;
 	virtual FOnActivateSequence& OnActivateSequence() override;
@@ -593,25 +610,30 @@ public:
 	virtual TSharedRef<SWidget> MakeTransportControls(bool bExtended) override;
 	virtual TSharedRef<SWidget> MakeTimeRange(const TSharedRef<SWidget>& InnerContent, bool bShowWorkingRange, bool bShowViewRange, bool bShowPlaybackRange) override;
 	virtual void SetViewportTransportControlsVisibility(bool bVisible) override;
-	virtual UObject* FindSpawnedObjectOrTemplate(const FGuid& BindingId) const override;
+	virtual UObject* FindSpawnedObjectOrTemplate(const FGuid& BindingId) override;
 	virtual FGuid MakeNewSpawnable(UObject& SourceObject) override;
 	
 public:
 
 	// IMovieScenePlayer interface
 
-	virtual void GetRuntimeObjects(TSharedRef<FMovieSceneSequenceInstance> MovieSceneInstance, const FGuid& ObjectHandle, TArray<TWeakObjectPtr<UObject>>& OutObjects) const override;
 	virtual void UpdateCameraCut(UObject* CameraObject, UObject* UnlockIfCameraObject, bool bJumpCut) override;
 	virtual void SetViewportSettings(const TMap<FViewportClient*, EMovieSceneViewportParams>& ViewportParamsMap) override;
 	virtual void GetViewportSettings(TMap<FViewportClient*, EMovieSceneViewportParams>& ViewportParamsMap) const override;
 	virtual EMovieScenePlayerStatus::Type GetPlaybackStatus() const override;
 	virtual void SetPlaybackStatus(EMovieScenePlayerStatus::Type InPlaybackStatus) override;
-	virtual void AddOrUpdateMovieSceneInstance(UMovieSceneSection& MovieSceneSection, TSharedRef<FMovieSceneSequenceInstance> InstanceToAdd) override;
-	virtual void RemoveMovieSceneInstance(UMovieSceneSection& MovieSceneSection, TSharedRef<FMovieSceneSequenceInstance> InstanceToRemove) override;
-	virtual IMovieSceneSpawnRegister& GetSpawnRegister() override { return *SpawnRegister; }
+	virtual FMovieSceneSpawnRegister& GetSpawnRegister() override { return *SpawnRegister; }
 	virtual bool IsPreview() const override { return SilentModeCount != 0; }
+
+public:
+
+	FMovieSceneRootEvaluationTemplateInstance& GetSequenceInstance() { return RootTemplateInstance; }
+
 protected:
 
+	/** Reevaluate the sequence at the current time */
+	void EvaluateInternal(FMovieSceneEvaluationRange InRange);
+
 	/** Reset data about a movie scene when pushing or popping a movie scene. */
 	void ResetPerMovieSceneData();
 
@@ -634,9 +656,6 @@ protected:
 	 */
 	TRange<float> GetFilteringShotsTimeBounds() const;
 
-	/** @return The current scrub position */
-	virtual float OnGetScrubPosition() const { return ScrubPosition; }
-
 	/**
 	 * Called when the clamp range is changed by the user
 	 *
@@ -677,6 +696,11 @@ protected:
 	 */
 	void UpdateAutoScroll(float NewTime);
 
+	/**
+	 * Ensure that the specified local time is in the view
+	 */
+	void ScrollIntoView(float InLocalTime);
+
 	/**
 	 * Calculates the amount of encroachment the specified time has into the autoscroll region, if any
 	 */
@@ -700,24 +724,6 @@ protected:
 
 protected:
 
-	/**
-	 * Populate the specified set with all UProperties that are currently keyed for the specified object and sequence instance, including any sub sequences
-	 *
-	 * @param Object		The object to get keyed properties for
-	 * @param Instance		The sequence instance in which to look for the object bindings
-	 * @param OutProperties	set to populate with properties
-	 */
-	void GetAllKeyedPropertiesForInstance(UObject& Object, FMovieSceneSequenceInstance& Instance, TSet<UProperty*>& OutProperties);
-
-	/**
-	 * Populate the specified set with all UProperties that are currently keyed for the specified object and sequence instance, excluding any sub sequences
-	 *
-	 * @param Object		The object to get keyed properties for
-	 * @param Instance		The sequence instance in which to look for the object bindings
-	 * @param OutProperties	set to populate with properties
-	 */
-	void GetKeyedProperties(UObject& Object, FMovieSceneSequenceInstance& Instance, TSet<UProperty*>& OutProperties);
-
 	/** Get all the keys for the current sequencer selection */
 	virtual void GetKeysFromSelection(TUniquePtr<ISequencerKeyCollection>& KeyCollection) override;
 
@@ -856,6 +862,12 @@ protected:
 	/** Possess PIE viewports with the specified camera settings (a mirror of level viewport possession, but for game viewport clients) */
 	void PossessPIEViewports(UObject* CameraObject, UObject* UnlockIfCameraObject, bool bJumpCut);
 
+	/** Update the locked subsequence range (displayed as playback range for subsequences), and root to local transform */
+	void UpdateSubSequenceData();
+
+	/** Check whether we're viewing the master sequence or not */
+	bool IsViewingMasterSequence() const { return ActiveTemplateIDs.Num() == 1; }
+
 private:
 
 	/** Reset the timing manager to default, or audio clock locked */
@@ -867,8 +879,6 @@ private:
 	/** User-supplied settings object for this sequencer */
 	USequencerSettings* Settings;
 
-	TMap< TWeakObjectPtr<UMovieSceneSection>, TSharedRef<FMovieSceneSequenceInstance> > SequenceInstanceBySection;
-
 	/** Command list for sequencer commands (Sequencer widgets only). */
 	TSharedRef<FUICommandList> SequencerCommandBindings;
 
@@ -884,25 +894,21 @@ private:
 	/** Listener for object changes being made while this sequencer is open*/
 	TSharedPtr<IDetailKeyframeHandler> DetailKeyframeHandler;
 
-	/** The runtime instance for the root movie scene */
-	TSharedPtr<FMovieSceneSequenceInstance> RootMovieSceneSequenceInstance;
-
 	/** Main sequencer widget */
 	TSharedPtr<SSequencer> SequencerWidget;
 	
 	/** Spawn register for keeping track of what is spawned */
-	TSharedPtr<IMovieSceneSpawnRegister> SpawnRegister;
+	TSharedPtr<FMovieSceneSpawnRegister> SpawnRegister;
 
 	/** The asset editor that created this Sequencer if any */
 	TWeakPtr<IToolkitHost> ToolkitHost;
 
-	/**
-	 * Stack of sequence instances.
-	 *
-	 * The first element is always the root instance.
-	 * The last element is the focused instance.
-	 */
-	TArray<TSharedRef<FMovieSceneSequenceInstance>> SequenceInstanceStack;
+	TWeakObjectPtr<UMovieSceneSequence> RootSequence;
+	FMovieSceneRootEvaluationTemplateInstance RootTemplateInstance;
+
+	TArray<FMovieSceneSequenceID> ActiveTemplateIDs;
+
+	FMovieSceneSequenceTransform RootToLocalTransform;
 
 	/** The time range target to be viewed */
 	TRange<float> TargetViewRange;
@@ -934,6 +940,9 @@ private:
 	// @todo sequencer: Should use FTimespan or "double" for Time Cursor Position! (cascades)
 	float ScrubPosition;
 
+	/** Current play position */
+	FMovieScenePlaybackPosition PlayPosition;
+
 	/** The playback rate */
 	float PlayRate;
 
@@ -1014,7 +1023,7 @@ private:
 		/** The view target it was pointing at before we took over */
 		TWeakObjectPtr<AActor> ViewTarget;
 	};
-	
+
 	/** Cached array of view targets that were set before we possessed the player controller with a camera from sequencer */
 	TArray<FCachedViewTarget> PrePossessionViewTargets;
 
@@ -1031,4 +1040,11 @@ private:
 
 	/** Event contexts retrieved from the above attribute once per frame */
 	TArray<TWeakObjectPtr<UObject>> CachedEventContexts;
+
+	bool bNeedsEvaluate;
+
+	/** The range of the currently displayed sub sequence in relation to its parent section */
+	TRange<float> SubSequenceRange;
+
+	TSharedPtr<struct FSequencerTemplateStore> TemplateStore;
 };
