Copying //UE4/Dev-Sequencer to //UE4/Dev-Main (Source: //UE4/Dev-Sequencer @ 3178529)

diff --git a/Engine/Source/Editor/Sequencer/Private/SequencerCommands.cpp b/Engine/Source/Editor/Sequencer/Private/SequencerCommands.cpp
--- a/Engine/Source/Editor/Sequencer/Private/SequencerCommands.cpp
+++ b/Engine/Source/Editor/Sequencer/Private/SequencerCommands.cpp
@@ -61,6 +61,7 @@ void FSequencerCommands::RegisterCommands()
 
 	UI_COMMAND( ToggleShowFrameNumbers, "Show Frame Numbers", "Enables and disables showing frame numbers", EUserInterfaceActionType::ToggleButton, FInputChord(EModifierKey::Control, EKeys::T) );
 	UI_COMMAND( ToggleShowGotoBox, "Go to Time...", "Go to a particular point on the timeline", EUserInterfaceActionType::Button, FInputChord(EModifierKey::Control, EKeys::G) );
+	UI_COMMAND( ToggleShowTransformBox, "Transform...", "Transform the selected keys and sections by a given amount", EUserInterfaceActionType::Button, FInputChord(EModifierKey::Control, EKeys::M) );
 	UI_COMMAND( ToggleShowRangeSlider, "Show Range Slider", "Enables and disables showing the time range slider", EUserInterfaceActionType::ToggleButton, FInputChord() );
 	UI_COMMAND( ToggleIsSnapEnabled, "Enable Snapping", "Enables and disables snapping", EUserInterfaceActionType::ToggleButton, FInputChord() );
 
@@ -92,6 +93,8 @@ void FSequencerCommands::RegisterCommands()
 	UI_COMMAND( ConvertToSpawnable, "Convert to Spawnable", "Make the specified possessed objects spawnable from sequencer. This will allow sequencer to have control over the lifetime of the object.", EUserInterfaceActionType::Button, FInputChord() );
 	UI_COMMAND( ConvertToPossessable, "Convert to Possessable", "Make the specified spawned objects possessed by sequencer.", EUserInterfaceActionType::Button, FInputChord() );
 
+	UI_COMMAND( SaveCurrentSpawnableState, "Save Default State", "Save the current state of this spawnable object as its default properties.", EUserInterfaceActionType::Button, FInputChord() );
+
 	UI_COMMAND( DiscardChanges, "Discard All Changes", "Revert the currently edited movie scene to its last saved state.", EUserInterfaceActionType::Button, FInputChord());
 
 	UI_COMMAND( FixActorReferences, "Fix Actor References", "Try to automatically fix up broken actor bindings.", EUserInterfaceActionType::Button, FInputChord() );
@@ -102,6 +105,8 @@ void FSequencerCommands::RegisterCommands()
 
 	UI_COMMAND( ImportFBX, "Import...", "Imports the animation from an FBX file.", EUserInterfaceActionType::Button, FInputChord() );
 	UI_COMMAND( ExportFBX, "Export...", "Exports the animation to an FBX file. (Shots and sub-scenes not supported)", EUserInterfaceActionType::Button, FInputChord() );
+
+	UI_COMMAND( ToggleEvaluateSubSequencesInIsolation, "Evaluate Sub Sequences In Isolation", "When enabled, will only evaluate the currently focused sequence; otherwise evaluate from the master sequence.", EUserInterfaceActionType::ToggleButton, FInputChord() );
 }
 
 #undef LOCTEXT_NAMESPACE
\ No newline at end of file
diff --git a/Engine/Source/Editor/Sequencer/Private/SequencerCommonHelpers.cpp b/Engine/Source/Editor/Sequencer/Private/SequencerCommonHelpers.cpp
--- a/Engine/Source/Editor/Sequencer/Private/SequencerCommonHelpers.cpp
+++ b/Engine/Source/Editor/Sequencer/Private/SequencerCommonHelpers.cpp
@@ -348,11 +348,9 @@ TSharedPtr<SWidget> SequencerHelpers::SummonContextMenu(FSequencer& Sequencer, c
 {
 	// @todo sequencer replace with UI Commands instead of faking it
 
-	FVector2D MouseDownPos = MyGeometry.AbsoluteToLocal(MouseEvent.GetScreenSpacePosition());
-
 	// Attempt to paste into either the current node selection, or the clicked on track
 	TSharedRef<SSequencer> SequencerWidget = StaticCastSharedRef<SSequencer>(Sequencer.GetSequencerWidget());
-	const float PasteAtTime = SequencerWidget->GetVirtualTrackArea().PixelToTime(MouseDownPos.X);
+	const float PasteAtTime = Sequencer.GetLocalTime();
 
 	const bool bShouldCloseWindowAfterMenuSelection = true;
 	FMenuBuilder MenuBuilder(bShouldCloseWindowAfterMenuSelection, Sequencer.GetCommandBindings());
diff --git a/Engine/Source/Editor/Sequencer/Private/SequencerContextMenus.cpp b/Engine/Source/Editor/Sequencer/Private/SequencerContextMenus.cpp
--- a/Engine/Source/Editor/Sequencer/Private/SequencerContextMenus.cpp
+++ b/Engine/Source/Editor/Sequencer/Private/SequencerContextMenus.cpp
@@ -192,6 +192,7 @@ void FKeyContextMenu::AddPropertiesMenu(FMenuBuilder& MenuBuilder)
 		DetailsViewArgs.bUpdatesFromSelection = false;
 		DetailsViewArgs.bShowOptions = false;
 		DetailsViewArgs.bShowModifiedPropertiesOption = false;
+		DetailsViewArgs.bShowScrollBar = false;
 	}
 
 	FStructureDetailsViewArgs StructureViewArgs;
@@ -233,6 +234,7 @@ void FKeyContextMenu::AddPropertiesMenu(FMenuBuilder& MenuBuilder)
 			StructureDetailsView->GetOnFinishedChangingPropertiesDelegate().AddLambda(
 				[=](const FPropertyChangedEvent& ChangeEvent) {
 					
+					Key.Value.Section->Modify();
 					if (Key.Key->GetStruct()->IsChildOf(FMovieSceneKeyStruct::StaticStruct()))
 					{
 						((FMovieSceneKeyStruct*)Key.Key->GetStructMemory())->PropagateChanges(ChangeEvent);
@@ -686,7 +688,7 @@ void FSectionContextMenu::TrimSection(bool bTrimLeft)
 {
 	FScopedTransaction TrimSectionTransaction(LOCTEXT("TrimSection_Transaction", "Trim Section"));
 
-	MovieSceneToolHelpers::TrimSection(Sequencer->GetSelection().GetSelectedSections(), Sequencer->GetGlobalTime(), bTrimLeft);
+	MovieSceneToolHelpers::TrimSection(Sequencer->GetSelection().GetSelectedSections(), Sequencer->GetLocalTime(), bTrimLeft);
 	Sequencer->NotifyMovieSceneDataChanged( EMovieSceneDataChangeType::TrackValueChanged );
 }
 
@@ -695,7 +697,7 @@ void FSectionContextMenu::SplitSection()
 {
 	FScopedTransaction SplitSectionTransaction(LOCTEXT("SplitSection_Transaction", "Split Section"));
 
-	MovieSceneToolHelpers::SplitSection(Sequencer->GetSelection().GetSelectedSections(), Sequencer->GetGlobalTime());
+	MovieSceneToolHelpers::SplitSection(Sequencer->GetSelection().GetSelectedSections(), Sequencer->GetLocalTime());
 	Sequencer->NotifyMovieSceneDataChanged( EMovieSceneDataChangeType::MovieSceneStructureItemAdded );
 }
 
@@ -704,7 +706,7 @@ bool FSectionContextMenu::IsTrimmable() const
 {
 	for (auto Section : Sequencer->GetSelection().GetSelectedSections())
 	{
-		if (Section.IsValid() && Section->IsTimeWithinSection(Sequencer->GetGlobalTime()))
+		if (Section.IsValid() && Section->IsTimeWithinSection(Sequencer->GetLocalTime()))
 		{
 			return true;
 		}
@@ -1002,7 +1004,7 @@ void FSectionContextMenu::BringToFront()
 		}
 	}
 
-	Sequencer->SetGlobalTime(Sequencer->GetGlobalTime());
+	Sequencer->SetLocalTimeDirectly(Sequencer->GetLocalTime());
 }
 
 
@@ -1050,7 +1052,7 @@ void FSectionContextMenu::SendToBack()
 		}
 	}
 
-	Sequencer->SetGlobalTime(Sequencer->GetGlobalTime());
+	Sequencer->SetLocalTimeDirectly(Sequencer->GetLocalTime());
 }
 
 
@@ -1095,7 +1097,7 @@ void FSectionContextMenu::BringForward()
 		}
 	}
 
-	Sequencer->SetGlobalTime(Sequencer->GetGlobalTime());
+	Sequencer->SetLocalTimeDirectly(Sequencer->GetLocalTime());
 }
 
 
@@ -1140,7 +1142,7 @@ void FSectionContextMenu::SendBackward()
 		}
 	}
 
-	Sequencer->SetGlobalTime(Sequencer->GetGlobalTime());
+	Sequencer->SetLocalTimeDirectly(Sequencer->GetLocalTime());
 }
 
 
diff --git a/Engine/Source/Editor/Sequencer/Private/SequencerEdMode.cpp b/Engine/Source/Editor/Sequencer/Private/SequencerEdMode.cpp
--- a/Engine/Source/Editor/Sequencer/Private/SequencerEdMode.cpp
+++ b/Engine/Source/Editor/Sequencer/Private/SequencerEdMode.cpp
@@ -102,7 +102,7 @@ void FSequencerEdMode::OnKeySelected(FViewport* Viewport, HMovieSceneKeyProxy* K
 	bool bAltDown = Viewport->KeyState(EKeys::LeftAlt) || Viewport->KeyState(EKeys::RightAlt);
 	bool bShiftDown = Viewport->KeyState(EKeys::LeftShift) || Viewport->KeyState(EKeys::RightShift);
 
-	Sequencer->SetGlobalTimeDirectly(KeyProxy->Time);
+	Sequencer->SetLocalTimeDirectly(KeyProxy->Time);
 	Sequencer->SelectTrackKeys(KeyProxy->MovieSceneSection, KeyProxy->Time, bShiftDown, bCtrlDown);
 }
 
@@ -407,7 +407,10 @@ void FSequencerEdMode::DrawTracks3D(const FSceneView* View, FPrimitiveDrawInterf
 		FGuid ObjectBinding = StaticCastSharedRef<FSequencerObjectBindingNode>(ObjectBindingNode.Key)->GetObjectBinding();
 
 		TArray<TWeakObjectPtr<UObject>> BoundObjects;
-		Sequencer->GetRuntimeObjects(Sequencer->GetFocusedMovieSceneSequenceInstance(), ObjectBinding, BoundObjects);
+		for (TWeakObjectPtr<> Ptr : Sequencer->FindObjectsInCurrentSequence(ObjectBinding))
+		{
+			BoundObjects.Add(Ptr);
+		}
 
 		for (auto DisplayNode : AllNodes)
 		{
diff --git a/Engine/Source/Editor/Sequencer/Private/SequencerNodeTree.cpp b/Engine/Source/Editor/Sequencer/Private/SequencerNodeTree.cpp
--- a/Engine/Source/Editor/Sequencer/Private/SequencerNodeTree.cpp
+++ b/Engine/Source/Editor/Sequencer/Private/SequencerNodeTree.cpp
@@ -299,7 +299,6 @@ TSharedRef<FSequencerObjectBindingNode> FSequencerNodeTree::AddObjectBinding(con
 		TSharedPtr<FSequencerObjectBindingNode> ParentNode;
 
 		UMovieSceneSequence* Sequence = Sequencer.GetFocusedMovieSceneSequence();
-		TSharedRef<FMovieSceneSequenceInstance> SequenceInstance = Sequencer.GetFocusedMovieSceneSequenceInstance();
 
 		// Prefer to use the parent spawnable if possible, rather than relying on runtime object presence
 		FMovieScenePossessable* Possessable = Sequence->GetMovieScene()->FindPossessable(ObjectBinding);
@@ -311,38 +310,9 @@ TSharedRef<FSequencerObjectBindingNode> FSequencerNodeTree::AddObjectBinding(con
 				ParentNode = AddObjectBinding( ParentBinding->GetName(), Possessable->GetParent(), GuidToBindingMap, OutNodeList );
 			}
 		}
-		
-		UObject* RuntimeObject = SequenceInstance->FindObject(ObjectBinding, Sequencer);
-
-		// fallback to using the parent runtime object
-		if (!ParentNode.IsValid() && RuntimeObject)
-		{
-			UObject* ParentObject = Sequence->GetParentObject(RuntimeObject);
-
-			if (ParentObject != nullptr)
-			{
-				FGuid ParentBinding = SequenceInstance->FindObjectId(*ParentObject);
-				TSharedPtr<FSequencerObjectBindingNode>* FoundParentNode = ObjectBindingMap.Find( ParentBinding );
-				if ( FoundParentNode != nullptr )
-				{
-					ParentNode = *FoundParentNode;
-				}
-				else
-				{
-					const FMovieSceneBinding** FoundParentMovieSceneBinding = GuidToBindingMap.Find( ParentBinding );
-					if ( FoundParentMovieSceneBinding != nullptr )
-					{
-						ParentNode = AddObjectBinding( (*FoundParentMovieSceneBinding)->GetName(), ParentBinding, GuidToBindingMap, OutNodeList );
-					}
-				}
-			}
-		}
 
 		// get human readable name of the object
-		AActor* RuntimeActor = Cast<AActor>(RuntimeObject);
-		const FString& DisplayString = (RuntimeActor != nullptr)
-			? RuntimeActor->GetActorLabel()
-			: ObjectName;
+		const FString& DisplayString = ObjectName;
 
 		// Create the node.
 		ObjectNode = MakeShareable(new FSequencerObjectBindingNode(ObjectNodeName, FText::FromString(DisplayString), ObjectBinding, ParentNode, *this));
diff --git a/Engine/Source/Editor/Sequencer/Private/SequencerSelection.cpp b/Engine/Source/Editor/Sequencer/Private/SequencerSelection.cpp
--- a/Engine/Source/Editor/Sequencer/Private/SequencerSelection.cpp
+++ b/Engine/Source/Editor/Sequencer/Private/SequencerSelection.cpp
@@ -106,8 +106,6 @@ void FSequencerSelection::AddToSelection(FSequencerSelectedKey Key)
 	{
 		EmptySelectedOutlinerNodesWithoutSection(Key.KeyArea->GetOwningSection());
 	}
-
-	EmptySelectedSections();
 }
 
 void FSequencerSelection::AddToSelection(UMovieSceneSection* Section)
@@ -124,8 +122,6 @@ void FSequencerSelection::AddToSelection(UMovieSceneSection* Section)
 	{
 		EmptySelectedOutlinerNodesWithoutSection(Section);
 	}
-
-	EmptySelectedKeys();
 }
 
 void FSequencerSelection::AddToSelection(TSharedRef<FSequencerDisplayNode> OutlinerNode)
diff --git a/Engine/Source/Editor/Sequencer/Private/SequencerSettings.cpp b/Engine/Source/Editor/Sequencer/Private/SequencerSettings.cpp
--- a/Engine/Source/Editor/Sequencer/Private/SequencerSettings.cpp
+++ b/Engine/Source/Editor/Sequencer/Private/SequencerSettings.cpp
@@ -22,7 +22,7 @@ USequencerSettings::USequencerSettings( const FObjectInitializer& ObjectInitiali
 	bSnapSectionTimesToSections = true;
 	bSnapPlayTimeToKeys = false;
 	bSnapPlayTimeToInterval = true;
-	bSnapPlayTimeToDraggedKey = false;
+	bSnapPlayTimeToDraggedKey = true;
 	CurveValueSnapInterval = 10.0f;
 	bSnapCurveValueToInterval = true;
 	bLabelBrowserVisible = false;
@@ -41,6 +41,7 @@ USequencerSettings::USequencerSettings( const FObjectInitializer& ObjectInitiali
 	bShowViewportTransportControls = true;
 	bLockPlaybackToAudioClock = false;
 	bAllowPossessionOfPIEViewports = false;
+	bEvaluateSubSequencesInIsolation = false;
 }
 
 void USequencerSettings::PostEditChangeProperty(FPropertyChangedEvent& PropertyChangedEvent)
@@ -571,6 +572,36 @@ void USequencerSettings::SetAutoSetTrackDefaults(bool bInAutoSetTrackDefaults)
 	}
 }
 
+bool USequencerSettings::ShouldShowDebugVisualization() const
+{
+	return bShowDebugVisualization;
+}
+
+void USequencerSettings::SetShowDebugVisualization(bool bInShowDebugVisualization)
+{
+	if (bShowDebugVisualization != bInShowDebugVisualization)
+	{
+		bShowDebugVisualization = bInShowDebugVisualization;
+		SaveConfig();
+	}
+}
+
+bool USequencerSettings::ShouldEvaluateSubSequencesInIsolation() const
+{
+	return bEvaluateSubSequencesInIsolation;
+}
+
+void USequencerSettings::SetEvaluateSubSequencesInIsolation(bool bInEvaluateSubSequencesInIsolation)
+{
+	if (bEvaluateSubSequencesInIsolation != bInEvaluateSubSequencesInIsolation)
+	{
+		bEvaluateSubSequencesInIsolation = bInEvaluateSubSequencesInIsolation;
+		SaveConfig();
+
+		OnEvaluateSubSequencesInIsolationChangedEvent.Broadcast();
+	}
+}
+
 USequencerSettings::FOnTimeSnapIntervalChanged& USequencerSettings::GetOnTimeSnapIntervalChanged()
 {
 	return OnTimeSnapIntervalChanged;
diff --git a/Engine/Source/Editor/Sequencer/Private/SequencerTimeSliderController.cpp b/Engine/Source/Editor/Sequencer/Private/SequencerTimeSliderController.cpp
--- a/Engine/Source/Editor/Sequencer/Private/SequencerTimeSliderController.cpp
+++ b/Engine/Source/Editor/Sequencer/Private/SequencerTimeSliderController.cpp
@@ -275,6 +275,9 @@ int32 FSequencerTimeSliderController::OnPaintTimeSlider( bool bMirrorLabels, con
 		);
 
 		LayerId = DrawPlaybackRange(AllottedGeometry, MyClippingRect, OutDrawElements, LayerId, RangeToScreen, PlaybackRangeArgs);
+		LayerId = DrawSubSequenceRange(AllottedGeometry, MyClippingRect, OutDrawElements, LayerId, RangeToScreen, PlaybackRangeArgs);
+
+		PlaybackRangeArgs.SolidFillOpacity = 0.05f;
 		LayerId = DrawSelectionRange(AllottedGeometry, MyClippingRect, OutDrawElements, LayerId, RangeToScreen, PlaybackRangeArgs);
 
 		float HalfSize = FMath::CeilToFloat(ScrubHandleSize / 2.0f);
@@ -388,17 +391,20 @@ int32 FSequencerTimeSliderController::DrawSelectionRange(const FGeometry& Allott
 	{
 		const float SelectionRangeL = RangeToScreen.InputToLocalX(SelectionRange.GetLowerBoundValue()) - 1;
 		const float SelectionRangeR = RangeToScreen.InputToLocalX(SelectionRange.GetUpperBoundValue()) + 1;
-		const auto DrawColor = FLinearColor(FColor(72, 72, 255));
+		const auto DrawColor = FEditorStyle::GetSlateColor("SelectionColor").GetColor(FWidgetStyle());
 
-		FSlateDrawElement::MakeBox(
-			OutDrawElements,
-			LayerId + 1,
-			AllottedGeometry.ToPaintGeometry(FVector2D(SelectionRangeL, 0.f), FVector2D(SelectionRangeR - SelectionRangeL, AllottedGeometry.Size.Y)),
-			FEditorStyle::GetBrush("WhiteBrush"),
-			MyClippingRect,
-			ESlateDrawEffect::None,
-			DrawColor.CopyWithNewOpacity(0.1f)
-		);
+		if (Args.SolidFillOpacity > 0.f)
+		{
+			FSlateDrawElement::MakeBox(
+				OutDrawElements,
+				LayerId + 1,
+				AllottedGeometry.ToPaintGeometry(FVector2D(SelectionRangeL, 0.f), FVector2D(SelectionRangeR - SelectionRangeL, AllottedGeometry.Size.Y)),
+				FEditorStyle::GetBrush("WhiteBrush"),
+				MyClippingRect,
+				ESlateDrawEffect::None,
+				DrawColor.CopyWithNewOpacity(Args.SolidFillOpacity)
+			);
+		}
 
 		FSlateDrawElement::MakeBox(
 			OutDrawElements,
@@ -432,6 +438,11 @@ int32 FSequencerTimeSliderController::DrawPlaybackRange(const FGeometry& Allotte
 		return LayerId;
 	}
 
+	TOptional<TRange<float>> SubSequenceRangeValue;
+	SubSequenceRangeValue = TimeSliderArgs.SubSequenceRange.Get(SubSequenceRangeValue);
+
+	const uint8 OpacityBlend = SubSequenceRangeValue.IsSet() ? 128 : 255;
+
 	const TRange<float> PlaybackRange = TimeSliderArgs.PlaybackRange.Get();
 	const float PlaybackRangeL = RangeToScreen.InputToLocalX(PlaybackRange.GetLowerBoundValue()) - 1;
 	const float PlaybackRangeR = RangeToScreen.InputToLocalX(PlaybackRange.GetUpperBoundValue()) + 1;
@@ -443,7 +454,7 @@ int32 FSequencerTimeSliderController::DrawPlaybackRange(const FGeometry& Allotte
 		Args.StartBrush,
 		MyClippingRect,
 		ESlateDrawEffect::None,
-		FColor(32, 128, 32)	// 120, 75, 50 (HSV)
+		FColor(32, 128, 32, OpacityBlend)	// 120, 75, 50 (HSV)
 	);
 
 	FSlateDrawElement::MakeBox(
@@ -453,7 +464,7 @@ int32 FSequencerTimeSliderController::DrawPlaybackRange(const FGeometry& Allotte
 		Args.EndBrush,
 		MyClippingRect,
 		ESlateDrawEffect::None,
-		FColor(128, 32, 32)	// 0, 75, 50 (HSV)
+		FColor(128, 32, 32, OpacityBlend)	// 0, 75, 50 (HSV)
 	);
 
 	// Black tint for excluded regions
@@ -464,7 +475,7 @@ int32 FSequencerTimeSliderController::DrawPlaybackRange(const FGeometry& Allotte
 		FEditorStyle::GetBrush("WhiteBrush"),
 		MyClippingRect,
 		ESlateDrawEffect::None,
-		FLinearColor::Black.CopyWithNewOpacity(0.3f)
+		FLinearColor::Black.CopyWithNewOpacity(0.3f * OpacityBlend / 255.f)
 	);
 
 	FSlateDrawElement::MakeBox(
@@ -474,9 +485,92 @@ int32 FSequencerTimeSliderController::DrawPlaybackRange(const FGeometry& Allotte
 		FEditorStyle::GetBrush("WhiteBrush"),
 		MyClippingRect,
 		ESlateDrawEffect::None,
+		FLinearColor::Black.CopyWithNewOpacity(0.3f * OpacityBlend / 255.f)
+	);
+
+	return LayerId + 1;
+}
+
+int32 FSequencerTimeSliderController::DrawSubSequenceRange(const FGeometry& AllottedGeometry, const FSlateRect& MyClippingRect, FSlateWindowElementList& OutDrawElements, int32 LayerId, const FScrubRangeToScreen& RangeToScreen, const FPaintPlaybackRangeArgs& Args) const
+{
+	TOptional<TRange<float>> RangeValue;
+	RangeValue = TimeSliderArgs.SubSequenceRange.Get(RangeValue);
+
+	if (!RangeValue.IsSet() || RangeValue->IsEmpty())
+	{
+		return LayerId;
+	}
+
+	const float SubSequenceRangeL = RangeToScreen.InputToLocalX(RangeValue->GetLowerBoundValue()) - 1;
+	const float SubSequenceRangeR = RangeToScreen.InputToLocalX(RangeValue->GetUpperBoundValue()) + 1;
+
+	static const FSlateBrush* LineBrushL(FEditorStyle::GetBrush("Sequencer.Timeline.PlayRange_L"));
+	static const FSlateBrush* LineBrushR(FEditorStyle::GetBrush("Sequencer.Timeline.PlayRange_R"));
+
+	FColor GreenTint(32, 128, 32);	// 120, 75, 50 (HSV)
+	FSlateDrawElement::MakeBox(
+		OutDrawElements,
+		LayerId+1,
+		AllottedGeometry.ToPaintGeometry(FVector2D(SubSequenceRangeL, 0.f), FVector2D(Args.BrushWidth, AllottedGeometry.Size.Y)),
+		LineBrushL,
+		MyClippingRect,
+		ESlateDrawEffect::None,
+		GreenTint
+	);
+
+	FColor RedTint(128, 32, 32);	// 0, 75, 50 (HSV)
+	FSlateDrawElement::MakeBox(
+		OutDrawElements,
+		LayerId+1,
+		AllottedGeometry.ToPaintGeometry(FVector2D(SubSequenceRangeR - Args.BrushWidth, 0.f), FVector2D(Args.BrushWidth, AllottedGeometry.Size.Y)),
+		LineBrushR,
+		MyClippingRect,
+		ESlateDrawEffect::None,
+		RedTint
+	);
+
+	// Black tint for excluded regions
+	FSlateDrawElement::MakeBox(
+		OutDrawElements,
+		LayerId+1,
+		AllottedGeometry.ToPaintGeometry(FVector2D(0.f, 0.f), FVector2D(SubSequenceRangeL, AllottedGeometry.Size.Y)),
+		FEditorStyle::GetBrush("WhiteBrush"),
+		MyClippingRect,
+		ESlateDrawEffect::None,
 		FLinearColor::Black.CopyWithNewOpacity(0.3f)
 	);
 
+	FSlateDrawElement::MakeBox(
+		OutDrawElements,
+		LayerId+1,
+		AllottedGeometry.ToPaintGeometry(FVector2D(SubSequenceRangeR, 0.f), FVector2D(AllottedGeometry.Size.X - SubSequenceRangeR, AllottedGeometry.Size.Y)),
+		FEditorStyle::GetBrush("WhiteBrush"),
+		MyClippingRect,
+		ESlateDrawEffect::None,
+		FLinearColor::Black.CopyWithNewOpacity(0.3f)
+	);
+
+	// Hash applied to the left and right of the sequence bounds
+	FSlateDrawElement::MakeBox(
+		OutDrawElements,
+		LayerId+1,
+		AllottedGeometry.ToPaintGeometry(FVector2D(SubSequenceRangeL - 16.f, 0.f), FVector2D(16.f, AllottedGeometry.Size.Y)),
+		FEditorStyle::GetBrush("Sequencer.Timeline.SubSequenceRangeHashL"),
+		MyClippingRect,
+		ESlateDrawEffect::None,
+		GreenTint
+	);
+
+	FSlateDrawElement::MakeBox(
+		OutDrawElements,
+		LayerId+1,
+		AllottedGeometry.ToPaintGeometry(FVector2D(SubSequenceRangeR, 0.f), FVector2D(16.f, AllottedGeometry.Size.Y)),
+		FEditorStyle::GetBrush("Sequencer.Timeline.SubSequenceRangeHashR"),
+		MyClippingRect,
+		ESlateDrawEffect::None,
+		RedTint
+	);
+
 	return LayerId + 1;
 }
 
@@ -756,12 +850,7 @@ FReply FSequencerTimeSliderController::OnMouseMove( SWidget& WidgetOwner, const
 					NewValue = TimeSliderArgs.Settings->SnapTimeToInterval(NewValue);
 				}
 
-				TRange<float> SelectionRange = TimeSliderArgs.SelectionRange.Get();
-
-				if (NewValue <= SelectionRange.GetUpperBoundValue())
-				{
-					TimeSliderArgs.OnSelectionRangeChanged.ExecuteIfBound(TRange<float>(NewValue, SelectionRange.GetUpperBoundValue()));
-				}
+				SetSelectionRangeStart(NewValue);
 			}
 			// Set the end range time?
 			else if(MouseDragType == DRAG_SELECTION_END)
@@ -770,13 +859,8 @@ FReply FSequencerTimeSliderController::OnMouseMove( SWidget& WidgetOwner, const
 				{
 					NewValue = TimeSliderArgs.Settings->SnapTimeToInterval(NewValue);
 				}
-					
-				TRange<float> SelectionRange = TimeSliderArgs.SelectionRange.Get();
 
-				if (NewValue >= SelectionRange.GetLowerBoundValue())
-				{
-					TimeSliderArgs.OnSelectionRangeChanged.ExecuteIfBound(TRange<float>(SelectionRange.GetLowerBoundValue(), NewValue));
-				}
+				SetSelectionRangeEnd(NewValue);
 			}
 			else if (MouseDragType == DRAG_SCRUBBING_TIME)
 			{
@@ -884,8 +968,11 @@ int32 FSequencerTimeSliderController::OnPaintSectionView( const FGeometry& Allot
 
 	if (Args.PlaybackRangeArgs.IsSet())
 	{
-		LayerId = DrawPlaybackRange(AllottedGeometry, MyClippingRect, OutDrawElements, LayerId, RangeToScreen, Args.PlaybackRangeArgs.GetValue());
-		LayerId = DrawSelectionRange(AllottedGeometry, MyClippingRect, OutDrawElements, LayerId, RangeToScreen, Args.PlaybackRangeArgs.GetValue());
+		FPaintPlaybackRangeArgs PaintArgs = Args.PlaybackRangeArgs.GetValue();
+		LayerId = DrawPlaybackRange(AllottedGeometry, MyClippingRect, OutDrawElements, LayerId, RangeToScreen, PaintArgs);
+		LayerId = DrawSubSequenceRange(AllottedGeometry, MyClippingRect, OutDrawElements, LayerId, RangeToScreen, PaintArgs);
+		PaintArgs.SolidFillOpacity = 0.f;
+		LayerId = DrawSelectionRange(AllottedGeometry, MyClippingRect, OutDrawElements, LayerId, RangeToScreen, PaintArgs);
 	}
 
 	if( Args.bDisplayTickLines )
@@ -941,42 +1028,77 @@ TSharedRef<SWidget> FSequencerTimeSliderController::OpenSetPlaybackRangeMenu(flo
 	const bool bShouldCloseWindowAfterMenuSelection = true;
 	FMenuBuilder MenuBuilder(bShouldCloseWindowAfterMenuSelection, nullptr);
 
-	FText NumericText;
+	FText CurrentTimeText;
 	if (SequencerSnapValues::IsTimeSnapIntervalFrameRate(TimeSliderArgs.Settings->GetTimeSnapInterval()) && TimeSliderArgs.Settings->GetShowFrameNumbers())
 	{
-		NumericText = FText::Format(LOCTEXT("FrameTextFormat", "Playback Range (at frame {0}):"), FText::AsNumber(TimeToFrame(MouseTime)));
+		CurrentTimeText = FText::Format(LOCTEXT("FrameTextFormat", "at frame {0}"), FText::AsNumber(TimeToFrame(MouseTime)));
 	}
 	else
 	{
-		NumericText = FText::Format(LOCTEXT("TimeTextFormat", "Playback Range (at {0}s):"), FText::AsNumber(MouseTime));
+		CurrentTimeText = FText::Format(LOCTEXT("TimeTextFormat", "at {0}s"), FText::AsNumber(MouseTime));
 	}
 
 	TRange<float> PlaybackRange = TimeSliderArgs.PlaybackRange.Get();
 
-	MenuBuilder.BeginSection("SequencerPlaybackRangeMenu", NumericText);
+	MenuBuilder.BeginSection("SequencerPlaybackRangeMenu", FText::Format(LOCTEXT("PlaybackRangeTextFormat", "Playback Range ({0}):"), CurrentTimeText));
 	{
 		MenuBuilder.AddMenuEntry(
-			FText::Format(LOCTEXT("SetPlaybackStart", "Set Start Time"), NumericText),
+			LOCTEXT("SetPlaybackStart", "Set Start Time"),
 			FText(),
 			FSlateIcon(),
 			FUIAction(
-				FExecuteAction::CreateLambda([=]{ return SetPlaybackRangeStart(MouseTime); }),
+				FExecuteAction::CreateLambda([=]{ SetPlaybackRangeStart(MouseTime); }),
 				FCanExecuteAction::CreateLambda([=]{ return MouseTime <= PlaybackRange.GetUpperBoundValue(); })
 			)
 		);
 
 		MenuBuilder.AddMenuEntry(
-			FText::Format(LOCTEXT("SetPlaybackEnd", "Set End Time"), NumericText),
+			LOCTEXT("SetPlaybackEnd", "Set End Time"),
 			FText(),
 			FSlateIcon(),
 			FUIAction(
-				FExecuteAction::CreateLambda([=]{ return SetPlaybackRangeEnd(MouseTime); }),
+				FExecuteAction::CreateLambda([=]{ SetPlaybackRangeEnd(MouseTime); }),
 				FCanExecuteAction::CreateLambda([=]{ return MouseTime >= PlaybackRange.GetLowerBoundValue(); })
 			)
 		);
 	}
 	MenuBuilder.EndSection(); // SequencerPlaybackRangeMenu
 
+	TRange<float> SelectionRange = TimeSliderArgs.SelectionRange.Get();
+	MenuBuilder.BeginSection("SequencerSelectionRangeMenu", FText::Format(LOCTEXT("SelectionRangeTextFormat", "Selection Range ({0}):"), CurrentTimeText));
+	{
+		MenuBuilder.AddMenuEntry(
+			LOCTEXT("SetSelectionStart", "Set Selection Start"),
+			FText(),
+			FSlateIcon(),
+			FUIAction(
+				FExecuteAction::CreateLambda([=]{ SetSelectionRangeStart(MouseTime); }),
+				FCanExecuteAction::CreateLambda([=]{ return SelectionRange.IsEmpty() || MouseTime <= SelectionRange.GetUpperBoundValue(); })
+			)
+		);
+
+		MenuBuilder.AddMenuEntry(
+			LOCTEXT("SetSelectionEnd", "Set Selection End"),
+			FText(),
+			FSlateIcon(),
+			FUIAction(
+				FExecuteAction::CreateLambda([=]{ SetSelectionRangeEnd(MouseTime); }),
+				FCanExecuteAction::CreateLambda([=]{ return SelectionRange.IsEmpty() || MouseTime >= SelectionRange.GetLowerBoundValue(); })
+			)
+		);
+
+		MenuBuilder.AddMenuEntry(
+			LOCTEXT("ClearSelectionRange", "Clear Selection Range"),
+			FText(),
+			FSlateIcon(),
+			FUIAction(
+				FExecuteAction::CreateLambda([=]{ TimeSliderArgs.OnSelectionRangeChanged.ExecuteIfBound(TRange<float>::Empty()); }),
+				FCanExecuteAction::CreateLambda([=]{ return !SelectionRange.IsEmpty(); })
+			)
+		);
+	}
+	MenuBuilder.EndSection(); // SequencerPlaybackRangeMenu
+
 	return MenuBuilder.MakeWidget();
 }
 
@@ -1145,4 +1267,32 @@ void FSequencerTimeSliderController::SetPlaybackRangeEnd(float NewEnd)
 	}
 }
 
+void FSequencerTimeSliderController::SetSelectionRangeStart(float NewStart)
+{
+	TRange<float> SelectionRange = TimeSliderArgs.SelectionRange.Get();
+
+	if (SelectionRange.IsEmpty())
+	{
+		TimeSliderArgs.OnSelectionRangeChanged.ExecuteIfBound(TRange<float>(NewStart, NewStart + 1));
+	}
+	else if (NewStart <= SelectionRange.GetUpperBoundValue())
+	{
+		TimeSliderArgs.OnSelectionRangeChanged.ExecuteIfBound(TRange<float>(NewStart, SelectionRange.GetUpperBoundValue()));
+	}
+}
+
+void FSequencerTimeSliderController::SetSelectionRangeEnd(float NewEnd)
+{
+	TRange<float> SelectionRange = TimeSliderArgs.SelectionRange.Get();
+
+	if (SelectionRange.IsEmpty())
+	{
+		TimeSliderArgs.OnSelectionRangeChanged.ExecuteIfBound(TRange<float>(NewEnd - 1, NewEnd));
+	}
+	else if (NewEnd >= SelectionRange.GetLowerBoundValue())
+	{
+		TimeSliderArgs.OnSelectionRangeChanged.ExecuteIfBound(TRange<float>(SelectionRange.GetLowerBoundValue(), NewEnd));
+	}
+}
+
 #undef LOCTEXT_NAMESPACE
diff --git a/Engine/Source/Editor/Sequencer/Private/SequencerTimeSliderController.h b/Engine/Source/Editor/Sequencer/Private/SequencerTimeSliderController.h
--- a/Engine/Source/Editor/Sequencer/Private/SequencerTimeSliderController.h
+++ b/Engine/Source/Editor/Sequencer/Private/SequencerTimeSliderController.h
@@ -9,7 +9,7 @@ struct FContextMenuSuppressor;
 struct FPaintPlaybackRangeArgs
 {
 	FPaintPlaybackRangeArgs()
-		: StartBrush(nullptr), EndBrush(nullptr), BrushWidth(0.f)
+		: StartBrush(nullptr), EndBrush(nullptr), BrushWidth(0.f), SolidFillOpacity(0.f)
 	{}
 
 	FPaintPlaybackRangeArgs(const FSlateBrush* InStartBrush, const FSlateBrush* InEndBrush, float InBrushWidth)
@@ -21,6 +21,8 @@ struct FPaintPlaybackRangeArgs
 	const FSlateBrush* EndBrush;
 	/** The width of the above brushes, in slate units */
 	float BrushWidth;
+	/** level of opacity for the fill color between the range markers */
+	float SolidFillOpacity;
 };
 
 struct FPaintSectionAreaViewArgs
@@ -166,6 +168,13 @@ private:
 	 */
 	int32 DrawPlaybackRange(const FGeometry& AllottedGeometry, const FSlateRect& MyClippingRect, FSlateWindowElementList& OutDrawElements, int32 LayerId, const FScrubRangeToScreen& RangeToScreen, const FPaintPlaybackRangeArgs& Args) const;
 
+	/**
+	 * Draw the playback range.
+	 *
+	 * @return the new layer ID
+	 */
+	int32 DrawSubSequenceRange(const FGeometry& AllottedGeometry, const FSlateRect& MyClippingRect, FSlateWindowElementList& OutDrawElements, int32 LayerId, const FScrubRangeToScreen& RangeToScreen, const FPaintPlaybackRangeArgs& Args) const;
+	
 private:
 
 	/**
@@ -181,6 +190,9 @@ private:
 	void SetPlaybackRangeStart(float NewStart);
 	void SetPlaybackRangeEnd(float NewEnd);
 
+	void SetSelectionRangeStart(float NewStart);
+	void SetSelectionRangeEnd(float NewEnd);
+
 	TSharedRef<SWidget> OpenSetPlaybackRangeMenu(float MouseTime);
 
 private:
@@ -220,7 +232,7 @@ private:
 	
 	/** Range stack */
 	TArray<FVector2D> RangeStack;
-	
+
 	/** When > 0, we should not show context menus */
 	int32 ContextMenuSuppression;
 	
diff --git a/Engine/Source/Editor/Sequencer/Private/SequencerTimingManager.cpp b/Engine/Source/Editor/Sequencer/Private/SequencerTimingManager.cpp
--- a/Engine/Source/Editor/Sequencer/Private/SequencerTimingManager.cpp
+++ b/Engine/Source/Editor/Sequencer/Private/SequencerTimingManager.cpp
@@ -4,11 +4,11 @@
 #include "SequencerTimingManager.h"
 #include "AudioDevice.h"
 
-FTimeAndDelta FSequencerDefaultTimingManager::AdjustTime(float InCurrentTime, float InDelta, float InPlayRate)
+FTimeAndDelta FSequencerDefaultTimingManager::AdjustTime(float InCurrentTime, float InDelta, float InPlayRate, float InDilation)
 {
 	FTimeAndDelta RetVal;
 	RetVal.Delta = InDelta;
-	RetVal.Time = InCurrentTime + InDelta * GWorld->GetWorldSettings()->MatineeTimeDilation * InPlayRate;
+	RetVal.Time = InCurrentTime + InDelta * InDilation * InPlayRate;
 
 	return RetVal;
 }
@@ -25,12 +25,12 @@ void FSequencerAudioClockTimer::OnStopPlaying(float InStopTime)
 	bIsPlaying = false;
 }
 
-FTimeAndDelta FSequencerAudioClockTimer::AdjustTime(float InCurrentTime, float InDelta, float InPlayRate)
+FTimeAndDelta FSequencerAudioClockTimer::AdjustTime(float InCurrentTime, float InDelta, float InPlayRate, float InDilation)
 {
 	if (!bIsPlaying)
 	{
 		// If we're not playing, we just use the default impl
-		return FSequencerDefaultTimingManager::AdjustTime(InCurrentTime, InDelta, InPlayRate);
+		return FSequencerDefaultTimingManager::AdjustTime(InCurrentTime, InDelta, InPlayRate, InDilation);
 	}
 
 	const double Now = GEngine->GetMainAudioDevice()->GetAudioClock();
diff --git a/Engine/Source/Editor/Sequencer/Private/SequencerTimingManager.h b/Engine/Source/Editor/Sequencer/Private/SequencerTimingManager.h
--- a/Engine/Source/Editor/Sequencer/Private/SequencerTimingManager.h
+++ b/Engine/Source/Editor/Sequencer/Private/SequencerTimingManager.h
@@ -32,7 +32,7 @@ public:
 
 	virtual void OnStartPlaying(float InStartTime) {}
 	virtual void OnStopPlaying(float InStopTime) {}
-	virtual FTimeAndDelta AdjustTime(float InCurrentTime, float InDelta, float InPlayRate) = 0;
+	virtual FTimeAndDelta AdjustTime(float InCurrentTime, float InDelta, float InPlayRate, float InDilation) = 0;
 
 private:
 	TOptional<EMovieScenePlayerStatus::Type> Status;
@@ -41,7 +41,7 @@ private:
 /** Default timing manager that is forced to the audio device clock */
 struct FSequencerDefaultTimingManager : FSequencerTimingManager
 {
-	virtual FTimeAndDelta AdjustTime(float InCurrentTime, float InDelta, float InPlayRate) override;
+	virtual FTimeAndDelta AdjustTime(float InCurrentTime, float InDelta, float InPlayRate, float InDilation) override;
 };
 
 /**
@@ -64,5 +64,5 @@ struct FSequencerAudioClockTimer : FSequencerDefaultTimingManager
 	virtual void OnStartPlaying(float InStartTime) override;
 	virtual void OnStopPlaying(float InStopTime) override;
 
-	virtual FTimeAndDelta AdjustTime(float InCurrentTime, float InDelta, float InPlayRate) override;
+	virtual FTimeAndDelta AdjustTime(float InCurrentTime, float InDelta, float InPlayRate, float InDilation) override;
 };
\ No newline at end of file
diff --git a/Engine/Source/Editor/Sequencer/Private/SequencerUtilities.cpp b/Engine/Source/Editor/Sequencer/Private/SequencerUtilities.cpp
--- a/Engine/Source/Editor/Sequencer/Private/SequencerUtilities.cpp
+++ b/Engine/Source/Editor/Sequencer/Private/SequencerUtilities.cpp
@@ -76,7 +76,7 @@ FName FSequencerUtilities::GetUniqueName( FName CandidateName, const TArray<FNam
 	int32 NameIndex = 1;
 	while ( ExistingNames.Contains( UniqueName ) )
 	{
-		UniqueName = FName( *FString::Printf(TEXT("%s%03i"), *BaseNameString, NameIndex ) );
+		UniqueName = FName( *FString::Printf(TEXT("%s%i"), *BaseNameString, NameIndex ) );
 		NameIndex++;
 	}
 
diff --git a/Engine/Source/Editor/Sequencer/Private/Tools/EditToolDragOperations.cpp b/Engine/Source/Editor/Sequencer/Private/Tools/EditToolDragOperations.cpp
--- a/Engine/Source/Editor/Sequencer/Private/Tools/EditToolDragOperations.cpp
+++ b/Engine/Source/Editor/Sequencer/Private/Tools/EditToolDragOperations.cpp
@@ -68,17 +68,25 @@ TOptional<FSequencerSnapField::FSnapResult> SnapToInterval(const TArray<float>&
 /** How many pixels near the mouse has to be before snapping occurs */
 const float PixelSnapWidth = 10.f;
 
-TRange<float> GetSectionBoundaries(UMovieSceneSection* Section, TSharedPtr<FSequencerTrackNode> SequencerNode)
+TRange<float> GetSectionBoundaries(const UMovieSceneSection* Section, TArray<FSectionHandle>& SectionHandles, TSharedPtr<FSequencerTrackNode> SequencerNode)
 {
+	// Only get boundaries for the sections that aren't being moved
+	TArray<const UMovieSceneSection*> SectionsBeingMoved;
+	for (auto SectionHandle : SectionHandles)
+	{
+		SectionsBeingMoved.Add(SectionHandle.GetSectionObject());
+	}
+
 	// Find the borders of where you can drag to
 	float LowerBound = -FLT_MAX, UpperBound = FLT_MAX;
 
 	// Also get the closest borders on either side
-	const TArray< TSharedRef<ISequencerSection> >& Sections = SequencerNode->GetSections();
-	for (int32 SectionIndex = 0; SectionIndex < Sections.Num(); ++SectionIndex)
+	const TArray< TSharedRef<ISequencerSection> >& AllSections = SequencerNode->GetSections();
+	for (int32 SectionIndex = 0; SectionIndex < AllSections.Num(); ++SectionIndex)
 	{
-		const UMovieSceneSection* TestSection = Sections[SectionIndex]->GetSectionObject();
-		if (Section != TestSection && Section->GetRowIndex() == TestSection->GetRowIndex())
+		const UMovieSceneSection* TestSection = AllSections[SectionIndex]->GetSectionObject();
+
+		if (!SectionsBeingMoved.Contains(TestSection) && Section->GetRowIndex() == TestSection->GetRowIndex())
 		{
 			if (TestSection->GetEndTime() <= Section->GetStartTime() && TestSection->GetEndTime() > LowerBound)
 			{
@@ -402,6 +410,20 @@ void FMoveSection::OnDrag(const FPointerEvent& MouseEvent, FVector2D LocalMouseP
 		}
 	}
 
+	// If sections are all on different rows, don't set row indices for anything because it leads to odd behavior.
+	bool bSectionsAreOnDifferentRows = false;
+	int32 FirstRowIndex = Sections[0].GetSectionObject()->GetRowIndex();
+	TArray<const UMovieSceneSection*> SectionsBeingMoved;
+	for (auto SectionHandle : Sections)
+	{
+		if (FirstRowIndex != SectionHandle.GetSectionObject()->GetRowIndex())
+		{
+			bSectionsAreOnDifferentRows = true;
+		}
+		SectionsBeingMoved.Add(SectionHandle.GetSectionObject());
+
+	}
+
 	for (int32 Index = 0; Index < Sections.Num(); ++Index)
 	{
 		auto& Handle = Sections[Index];
@@ -411,11 +433,15 @@ void FMoveSection::OnDrag(const FPointerEvent& MouseEvent, FVector2D LocalMouseP
 		const float DeltaTime = VirtualMousePos.X + RelativeOffsets[Index].StartTime - Section->GetStartTime();
 
 		auto SequencerNodeSections = Handle.TrackNode->GetSections();
-		
+
 		TArray<UMovieSceneSection*> MovieSceneSections;
 		for (int32 i = 0; i < SequencerNodeSections.Num(); ++i)
 		{
-			MovieSceneSections.Add(SequencerNodeSections[i]->GetSectionObject());
+			UMovieSceneSection* TrackSection = SequencerNodeSections[i]->GetSectionObject();
+			if (!SectionsBeingMoved.Contains(TrackSection))
+			{
+				MovieSceneSections.Add(SequencerNodeSections[i]->GetSectionObject());
+			}
 		}
 
 		int32 TargetRowIndex = Section->GetRowIndex();
@@ -462,12 +488,16 @@ void FMoveSection::OnDrag(const FPointerEvent& MouseEvent, FVector2D LocalMouseP
 			!Section->OverlapsWithSections(MovieSceneSections, TargetRowIndex - Section->GetRowIndex(), DeltaTime))
 		{
 			Section->MoveSection(DeltaTime, DraggedKeyHandles);
-			Section->SetRowIndex(TargetRowIndex);
+			if (!bSectionsAreOnDifferentRows)
+			{
+				Section->SetRowIndex(TargetRowIndex);
+			}
 		}
 		else
 		{
 			if (bDeltaY &&
-				!Section->OverlapsWithSections(MovieSceneSections, TargetRowIndex - Section->GetRowIndex(), 0.f))
+				!Section->OverlapsWithSections(MovieSceneSections, TargetRowIndex - Section->GetRowIndex(), 0.f) &&
+				!bSectionsAreOnDifferentRows)
 			{
 				Section->SetRowIndex(TargetRowIndex);
 			}
@@ -481,7 +511,7 @@ void FMoveSection::OnDrag(const FPointerEvent& MouseEvent, FVector2D LocalMouseP
 				else
 				{
 					// Find the borders of where you can move to
-					TRange<float> SectionBoundaries = GetSectionBoundaries(Section, Handle.TrackNode);
+					TRange<float> SectionBoundaries = GetSectionBoundaries(Section, Sections, Handle.TrackNode);
 
 					float LeftMovementMaximum = SectionBoundaries.GetLowerBoundValue() - Section->GetStartTime();
 					float RightMovementMaximum = SectionBoundaries.GetUpperBoundValue() - Section->GetEndTime();
@@ -614,9 +644,13 @@ void FMoveKeys::OnDrag(const FPointerEvent& MouseEvent, FVector2D LocalMousePos,
 	// Snap the play time to the new dragged key time if all the keyframes were dragged to the same time
 	if (Settings->GetSnapPlayTimeToDraggedKey() && PrevNewKeyTime != FLT_MAX && PrevNewKeyTime != -FLT_MAX)
 	{
-		Sequencer.SetGlobalTime(PrevNewKeyTime);
+		Sequencer.SetLocalTime(PrevNewKeyTime);
 	}
 
+	for (UMovieSceneSection* Section : ModifiedSections)
+	{
+		Section->MarkAsChanged();
+	}
 	Sequencer.NotifyMovieSceneDataChanged( EMovieSceneDataChangeType::TrackValueChanged );
 }
 
diff --git a/Engine/Source/Editor/Sequencer/Private/Tools/SequencerEditTool_Movement.cpp b/Engine/Source/Editor/Sequencer/Private/Tools/SequencerEditTool_Movement.cpp
--- a/Engine/Source/Editor/Sequencer/Private/Tools/SequencerEditTool_Movement.cpp
+++ b/Engine/Source/Editor/Sequencer/Private/Tools/SequencerEditTool_Movement.cpp
@@ -36,7 +36,7 @@ FReply FSequencerEditTool_Movement::OnMouseButtonDown(SWidget& OwnerWidget, cons
 				if (DelayedDrag->Hotspot->GetType() == ESequencerHotspot::Key)
 				{
 					FSequencerSelectedKey& ThisKey = StaticCastSharedPtr<FKeyHotspot>(DelayedDrag->Hotspot)->Key;
-					Sequencer.SetGlobalTime(ThisKey.KeyArea->GetKeyTime(ThisKey.KeyHandle.GetValue()));
+					Sequencer.SetLocalTime(ThisKey.KeyArea->GetKeyTime(ThisKey.KeyHandle.GetValue()));
 				}
 			}
 		}
diff --git a/Engine/Source/Editor/Sequencer/Private/Tools/SequencerSnapField.cpp b/Engine/Source/Editor/Sequencer/Private/Tools/SequencerSnapField.cpp
--- a/Engine/Source/Editor/Sequencer/Private/Tools/SequencerSnapField.cpp
+++ b/Engine/Source/Editor/Sequencer/Private/Tools/SequencerSnapField.cpp
@@ -70,7 +70,7 @@ FSequencerSnapField::FSequencerSnapField(const ISequencer& InSequencer, ISequenc
 	Visitor.Snaps.Add(FSequencerSnapPoint{ FSequencerSnapPoint::PlaybackRange, PlaybackRange.GetUpperBoundValue() });
 
 	// Add the current time as a potential snap candidate
-	Visitor.Snaps.Add(FSequencerSnapPoint{ FSequencerSnapPoint::CurrentTime, InSequencer.GetGlobalTime() });
+	Visitor.Snaps.Add(FSequencerSnapPoint{ FSequencerSnapPoint::CurrentTime, InSequencer.GetLocalTime() });
 
 	// Add the selection range bounds as a potential snap candidate
 	TRange<float> SelectionRange = InSequencer.GetFocusedMovieSceneSequence()->GetMovieScene()->GetSelectionRange();
diff --git a/Engine/Source/Editor/Sequencer/Public/IKeyArea.h b/Engine/Source/Editor/Sequencer/Public/IKeyArea.h
--- a/Engine/Source/Editor/Sequencer/Public/IKeyArea.h
+++ b/Engine/Source/Editor/Sequencer/Public/IKeyArea.h
@@ -5,7 +5,7 @@
 #include "Curves/CurveBase.h"
 #include "KeyParams.h"
 
-
+class ISequencer;
 class UMovieSceneSection;
 class FMovieSceneClipboardBuilder;
 class FMovieSceneClipboardKeyTrack;
@@ -49,7 +49,7 @@ public:
 	 * @param KeyHandle Handle of the key
 	 * @param The new time of the key
 	 */
-	virtual void SetKeyTime(FKeyHandle KeyHandle, float NewKeyTime) const = 0;
+	virtual void SetKeyTime(FKeyHandle KeyHandle, float NewKeyTime) = 0;
 
 	/**
 	 * Gets the time of a key given a handle
diff --git a/Engine/Source/Editor/Sequencer/Public/ISequencer.h b/Engine/Source/Editor/Sequencer/Public/ISequencer.h
--- a/Engine/Source/Editor/Sequencer/Public/ISequencer.h
+++ b/Engine/Source/Editor/Sequencer/Public/ISequencer.h
@@ -118,10 +118,14 @@ public:
 	/** @return Returns the MovieScene that is currently focused for editing by the sequencer.  This can change at any time. */
 	virtual UMovieSceneSequence* GetFocusedMovieSceneSequence() const = 0;
 
-	/**
-	 * @return the currently focused movie scene instance
-	 */
-	virtual TSharedRef<FMovieSceneSequenceInstance> GetFocusedMovieSceneSequenceInstance() const = 0;
+	/** @return The root movie scene being used */
+	virtual FMovieSceneSequenceIDRef GetRootTemplateID() const = 0;
+	virtual FMovieSceneSequenceIDRef GetFocusedTemplateID() const = 0;
+
+	TArrayView<TWeakObjectPtr<>> FindObjectsInCurrentSequence(const FGuid& InObjectBinding)
+	{
+		return FindBoundObjects(InObjectBinding, GetFocusedTemplateID());
+	}
 
 	/** Resets sequencer with a new animation */
 	virtual void ResetToNewRootSequence(UMovieSceneSequence& NewAnimation) = 0;
@@ -146,19 +150,6 @@ public:
 	 */
 	virtual FGuid MakeNewSpawnable(UObject& SourceObject) = 0;
 
-	/**
-	 * Given a sub-movie scene section, returns the instance of the movie scene for that section.
-	 *
-	 * @param Section The sub-movie scene section containing the sequence instance to get.
-	 * @return The instance for the sub-movie scene
-	 */
-	virtual TSharedRef<FMovieSceneSequenceInstance> GetSequenceInstanceForSection(UMovieSceneSection& Section) const = 0;
-
-	/**
-	 * @return Whether the section has a sequence instance
-	 */
-	virtual bool HasSequenceInstanceForSection(UMovieSceneSection& Section) const = 0;
-
 	/**
 	 * Add actors as possessable objects to sequencer.
 	 * 
@@ -210,12 +201,10 @@ public:
 	virtual bool IsRecordingLive() const = 0;
 
 	/**
-	 * Gets the current time of the time slider relative to the passed in movie scene
-	 *
-	 * @param MovieScene The movie scene to get the local time for (must exist in this sequencer).
+	 * Gets the current time of the time slider relative to the currently focused movie scene
 	 */
-	virtual float GetCurrentLocalTime(UMovieSceneSequence& InMovieSceneSequence) = 0;
-	
+	virtual float GetLocalTime() const = 0;
+
 	/**
 	 * Gets the global time.
 	 *
@@ -225,17 +214,21 @@ public:
 	virtual float GetGlobalTime() const = 0;
 
 	/**
-	 * Sets the global position to the time.
+	 * Sets the cursor position relative to the currently focused sequence
 	 *
-	 * @param Time The global time to set.
+	 * @param Time The local time to set.
 	 * @param SnapTimeMode The type of time snapping allowed.
-	 * @param bRestarted Whether the time has been restarted from the beginning or looped.
-	 * @see GetGlobalTime
 	 */
-	virtual void SetGlobalTime(float Time, ESnapTimeMode SnapTimeMode = ESnapTimeMode::STM_None, bool bRestarted = false) = 0;
+	virtual void SetLocalTime(float Time, ESnapTimeMode SnapTimeMode = ESnapTimeMode::STM_None) = 0;
 
-	/** Set the global time directly, without performing any auto-scroll */
-	virtual void SetGlobalTimeDirectly(float Time, ESnapTimeMode SnapTimeMode = ESnapTimeMode::STM_None, bool bRestarted = false) = 0;
+	/** Set the current local time directly, with no other snapping, scrolling or manipulation */
+	virtual void SetLocalTimeDirectly(float NewTime) = 0;
+
+	/** Set the global time directly, without performing any auto-scroll, snapping or other adjustments to the supplied time  */
+	virtual void SetGlobalTime(float Time) = 0;
+
+	/** Forcefully reevaluate the sequence */
+	virtual void ForceEvaluate() = 0;
 
 	/** @return The current view range */
 	virtual FAnimatedRange GetViewRange() const
@@ -305,7 +298,7 @@ public:
 	DECLARE_EVENT_OneParam(ISequencer, FOnPreSave, ISequencer&)
 	virtual FOnPreSave& OnPreSave() = 0;
 
-	DECLARE_EVENT_OneParam(ISequencer, FOnActivateSequence, FMovieSceneSequenceInstance&)
+	DECLARE_EVENT_OneParam(ISequencer, FOnActivateSequence, FMovieSceneSequenceIDRef)
 	virtual FOnActivateSequence& OnActivateSequence() = 0;
 
 	/**
@@ -322,11 +315,6 @@ public:
 	 */
 	virtual class ISequencerObjectChangeListener& GetObjectChangeListener() = 0;
 
-	/**
-	 * Get a set of all properties currently keyed on the specified object
-	 */
-	virtual void GetAllKeyedProperties(UObject& Object, TSet<UProperty*>& OutProperties) = 0;
-
 	virtual bool CanKeyProperty(FCanKeyPropertyParams CanKeyPropertyParams) const = 0;
 
 	virtual void KeyProperty(FKeyPropertyParams KeyPropertyParams) = 0;
@@ -379,7 +367,7 @@ public:
 	virtual void SetViewportTransportControlsVisibility(bool bVisible) = 0;
 
 	/** Attempt to find a spawned object in the currently focused movie scene, or the template object for the specified binding ID, if possible */
-	virtual UObject* FindSpawnedObjectOrTemplate(const FGuid& BindingId) const = 0;
+	virtual UObject* FindSpawnedObjectOrTemplate(const FGuid& BindingId) = 0;
 
 	/**
 	 * Create a widget containing the spinboxes for setting the working and playback range
diff --git a/Engine/Source/Editor/Sequencer/Public/ISequencerModule.h b/Engine/Source/Editor/Sequencer/Public/ISequencerModule.h
--- a/Engine/Source/Editor/Sequencer/Public/ISequencerModule.h
+++ b/Engine/Source/Editor/Sequencer/Public/ISequencerModule.h
@@ -68,7 +68,7 @@ struct FSequencerInitParams
 	bool bEditWithinLevelEditor;
 
 	/** Domain-specific spawn register for the movie scene */
-	TSharedPtr<IMovieSceneSpawnRegister> SpawnRegister;
+	TSharedPtr<FMovieSceneSpawnRegister> SpawnRegister;
 
 	/** Accessor for event contexts */
 	TAttribute<TArray<UObject*>> EventContexts;
diff --git a/Engine/Source/Editor/Sequencer/Public/MovieSceneTrackEditor.h b/Engine/Source/Editor/Sequencer/Public/MovieSceneTrackEditor.h
--- a/Engine/Source/Editor/Sequencer/Public/MovieSceneTrackEditor.h
+++ b/Engine/Source/Editor/Sequencer/Public/MovieSceneTrackEditor.h
@@ -54,15 +54,12 @@ public:
 	}
 
 	/**
-	 * Gets the movie scene that should be used for adding new keys/sections to and the time where new keys should be added
-	 *
-	 * @param InMovieScene	The movie scene to be used for keying
-	 * @param OutTime		The current time of the sequencer which should be used for adding keys during auto-key
-	 * @return true	if we can auto-key
+	 * @return The current local time at which we should add a key
 	 */
-	float GetTimeForKey( UMovieSceneSequence* InMovieSceneSequence )
+	float GetTimeForKey()
 	{ 
-		return Sequencer.Pin()->GetCurrentLocalTime( *InMovieSceneSequence );
+		TSharedPtr<ISequencer> SequencerPin = Sequencer.Pin();
+		return SequencerPin.IsValid() ? SequencerPin->GetLocalTime() : 0.f;
 	}
 
 	void UpdatePlaybackRange()
@@ -82,7 +79,7 @@ public:
 		UMovieSceneSequence* MovieSceneSequence = GetMovieSceneSequence();
 		if (MovieSceneSequence)
 		{
-			float KeyTime = GetTimeForKey( MovieSceneSequence );
+			float KeyTime = GetTimeForKey();
 
 			if( !Sequencer.Pin()->IsRecordingLive() )
 			{
diff --git a/Engine/Source/Editor/Sequencer/Public/SequencerCommands.h b/Engine/Source/Editor/Sequencer/Public/SequencerCommands.h
--- a/Engine/Source/Editor/Sequencer/Public/SequencerCommands.h
+++ b/Engine/Source/Editor/Sequencer/Public/SequencerCommands.h
@@ -144,6 +144,9 @@ public:
 	/** Toggle the visibility of the goto box. */
 	TSharedPtr< FUICommandInfo > ToggleShowGotoBox;
 
+	/** Toggle the visibility of the transform box. */
+	TSharedPtr< FUICommandInfo > ToggleShowTransformBox;
+
 	/** Turns the range slider on and off. */
 	TSharedPtr< FUICommandInfo > ToggleShowRangeSlider;
 
@@ -216,6 +219,9 @@ public:
 	/** Convert the selected spawnable objects to possessables. The newly created possessables will be created in the current level. */
 	TSharedPtr< FUICommandInfo > ConvertToPossessable;
 
+	/** Saves the current state of this object as the default spawnable state. */
+	TSharedPtr< FUICommandInfo > SaveCurrentSpawnableState;
+
 	/** Discard all changes to the current movie scene. */
 	TSharedPtr< FUICommandInfo > DiscardChanges;
 
@@ -234,6 +240,9 @@ public:
 	/** Exports animation to fbx. */
 	TSharedPtr< FUICommandInfo > ExportFBX;
 
+	/** Toggle whether we should evaluate sub sequences in isolation */
+	TSharedPtr< FUICommandInfo > ToggleEvaluateSubSequencesInIsolation;
+
 	/**
 	 * Initialize commands
 	 */
diff --git a/Engine/Source/Editor/Sequencer/Public/SequencerSettings.h b/Engine/Source/Editor/Sequencer/Public/SequencerSettings.h
--- a/Engine/Source/Editor/Sequencer/Public/SequencerSettings.h
+++ b/Engine/Source/Editor/Sequencer/Public/SequencerSettings.h
@@ -66,6 +66,7 @@ public:
 	GENERATED_UCLASS_BODY()
 
 	DECLARE_MULTICAST_DELEGATE( FOnTimeSnapIntervalChanged );
+	DECLARE_MULTICAST_DELEGATE( FOnEvaluateSubSequencesInIsolationChanged );
 	DECLARE_MULTICAST_DELEGATE_OneParam( FOnLockPlaybackToAudioClockChanged, bool );
 
 	virtual void PostEditChangeProperty(FPropertyChangedEvent& PropertyChangedEvent) override;
@@ -245,6 +246,18 @@ public:
 	/** Sets whether or not track defaults will be automatically set when modifying tracks. */
 	void SetAutoSetTrackDefaults(bool bInAutoSetTrackDefaults);
 
+	/** @return Whether to show debug vis */
+	bool ShouldShowDebugVisualization() const;
+	/** Toggle whether to show debug vis */
+	void SetShowDebugVisualization(bool bInShowDebugVisualization);
+
+	/** @return Whether to evaluate sub sequences in isolation */
+	bool ShouldEvaluateSubSequencesInIsolation() const;
+	/** Set whether to evaluate sub sequences in isolation */
+	void SetEvaluateSubSequencesInIsolation(bool bInEvaluateSubSequencesInIsolation);
+	/** Gets the multicast delegate which is run whenever the time snap interval is changed. */
+	FOnEvaluateSubSequencesInIsolationChanged& GetOnEvaluateSubSequencesInIsolationChanged() { return OnEvaluateSubSequencesInIsolationChangedEvent; }
+
 	/** Snaps a time value in seconds to the currently selected interval. */
 	float SnapTimeToInterval(float InTimeValue) const;
 
@@ -404,7 +417,16 @@ protected:
 	UPROPERTY(config, EditAnywhere, Category=General)
 	bool bAllowPossessionOfPIEViewports;
 
+	/** When enabled, entering a sub sequence will evaluate that sub sequence in isolation, rather than from the master sequence */
+	UPROPERTY(config, EditAnywhere, Category=Playback)
+	bool bEvaluateSubSequencesInIsolation;
+
 	FOnLockPlaybackToAudioClockChanged OnLockPlaybackToAudioClockChanged;
 
+	/** Enable or disable showing of debug visualization. */
+	UPROPERTY( config, EditAnywhere, Category=General )
+	bool bShowDebugVisualization;
+
 	FOnTimeSnapIntervalChanged OnTimeSnapIntervalChanged;
+	FOnEvaluateSubSequencesInIsolationChanged OnEvaluateSubSequencesInIsolationChangedEvent;
 };
diff --git a/Engine/Source/Editor/SequencerWidgets/Public/ITimeSlider.h b/Engine/Source/Editor/SequencerWidgets/Public/ITimeSlider.h
--- a/Engine/Source/Editor/SequencerWidgets/Public/ITimeSlider.h
+++ b/Engine/Source/Editor/SequencerWidgets/Public/ITimeSlider.h
@@ -93,6 +93,9 @@ struct FTimeSliderArgs
 	/** Called when the clamp range changes */
 	FOnRangeChanged OnClampRangeChanged;
 
+	/** Attribute defining the active sub-sequence range for this controller */
+	TAttribute<TOptional<TRange<float>>> SubSequenceRange;
+
 	/** Attribute defining the playback range for this controller */
 	TAttribute<TRange<float>> PlaybackRange;
 
diff --git a/Engine/Source/Editor/UMGEditor/Private/Animation/Sequencer2DTransformTrackEditor.cpp b/Engine/Source/Editor/UMGEditor/Private/Animation/Sequencer2DTransformTrackEditor.cpp
--- a/Engine/Source/Editor/UMGEditor/Private/Animation/Sequencer2DTransformTrackEditor.cpp
+++ b/Engine/Source/Editor/UMGEditor/Private/Animation/Sequencer2DTransformTrackEditor.cpp
@@ -58,11 +58,11 @@ public:
 		// Shear
 		TAttribute<TOptional<float>> ShearXExternalValue = TAttribute<TOptional<float>>::Create(
 			TAttribute<TOptional<float>>::FGetter::CreateRaw(this, &F2DTransformSection::GetShearXValue));
-		TSharedRef<FFloatCurveKeyArea> ShearXKeyArea = MakeShareable(new FFloatCurveKeyArea(&TransformSection->GetSheerCurve(EAxis::X), ShearXExternalValue, TransformSection, RedKeyAreaColor));
+		TSharedRef<FFloatCurveKeyArea> ShearXKeyArea = MakeShareable(new FFloatCurveKeyArea(&TransformSection->GetShearCurve(EAxis::X), ShearXExternalValue, TransformSection, RedKeyAreaColor));
 
 		TAttribute<TOptional<float>> ShearYExternalValue = TAttribute<TOptional<float>>::Create(
 			TAttribute<TOptional<float>>::FGetter::CreateRaw(this, &F2DTransformSection::GetShearYValue));
-		TSharedRef<FFloatCurveKeyArea> ShearYKeyArea = MakeShareable(new FFloatCurveKeyArea(&TransformSection->GetSheerCurve(EAxis::Y), ShearYExternalValue, TransformSection, GreenKeyAreaColor));
+		TSharedRef<FFloatCurveKeyArea> ShearYKeyArea = MakeShareable(new FFloatCurveKeyArea(&TransformSection->GetShearCurve(EAxis::Y), ShearYExternalValue, TransformSection, GreenKeyAreaColor));
 
 		// This generates the tree structure for the transform section
 		LayoutBuilder.PushCategory("Translation", NSLOCTEXT("F2DTransformSection", "TranslationArea", "Translation"));
diff --git a/Engine/Source/Editor/UMGEditor/Private/Animation/WidgetMaterialTrackEditor.cpp b/Engine/Source/Editor/UMGEditor/Private/Animation/WidgetMaterialTrackEditor.cpp
--- a/Engine/Source/Editor/UMGEditor/Private/Animation/WidgetMaterialTrackEditor.cpp
+++ b/Engine/Source/Editor/UMGEditor/Private/Animation/WidgetMaterialTrackEditor.cpp
@@ -4,6 +4,7 @@
 #include "WidgetMaterialTrackEditor.h"
 #include "MovieSceneWidgetMaterialTrack.h"
 #include "WidgetMaterialTrackUtilities.h"
+#include "MovieSceneSequenceInstance.h"
 
 
 FWidgetMaterialTrackEditor::FWidgetMaterialTrackEditor( TSharedRef<ISequencer> InSequencer )
@@ -26,15 +27,17 @@ bool FWidgetMaterialTrackEditor::SupportsType( TSubclassOf<UMovieSceneTrack> Typ
 
 UMaterialInterface* FWidgetMaterialTrackEditor::GetMaterialInterfaceForTrack( FGuid ObjectBinding, UMovieSceneMaterialTrack* MaterialTrack )
 {
-	UObject* WidgetObject = GetSequencer()->GetFocusedMovieSceneSequenceInstance()->FindObject( ObjectBinding, *GetSequencer() );
-	UWidget* Widget = Cast<UWidget>( WidgetObject );
-	UMovieSceneWidgetMaterialTrack* WidgetMaterialTrack = Cast<UMovieSceneWidgetMaterialTrack>( MaterialTrack );
-	if ( Widget != nullptr && WidgetMaterialTrack != nullptr )
+	for (TWeakObjectPtr<> WeakObjectPtr : GetSequencer()->FindObjectsInCurrentSequence(ObjectBinding))
 	{
-		FWidgetMaterialHandle Handle = WidgetMaterialTrackUtilities::GetMaterialHandle( Widget, WidgetMaterialTrack->GetBrushPropertyNamePath() );
-		if (Handle.IsValid() )
+		UWidget* Widget = Cast<UWidget>( WeakObjectPtr.Get() );
+		UMovieSceneWidgetMaterialTrack* WidgetMaterialTrack = Cast<UMovieSceneWidgetMaterialTrack>( MaterialTrack );
+		if ( Widget != nullptr && WidgetMaterialTrack != nullptr )
 		{
-			return Handle.GetMaterial();
+			FWidgetMaterialHandle Handle = WidgetMaterialTrackUtilities::GetMaterialHandle( Widget, WidgetMaterialTrack->GetBrushPropertyNamePath() );
+			if (Handle.IsValid() )
+			{
+				return Handle.GetMaterial();
+			}
 		}
 	}
 	return nullptr;
diff --git a/Engine/Source/Editor/UMGEditor/Private/TabFactory/AnimationTabSummoner.cpp b/Engine/Source/Editor/UMGEditor/Private/TabFactory/AnimationTabSummoner.cpp
--- a/Engine/Source/Editor/UMGEditor/Private/TabFactory/AnimationTabSummoner.cpp
+++ b/Engine/Source/Editor/UMGEditor/Private/TabFactory/AnimationTabSummoner.cpp
@@ -35,23 +35,25 @@ FAnimationTabSummoner::FAnimationTabSummoner(TSharedPtr<class FWidgetBlueprintEd
 }
 
 
-bool VerifyAnimationRename( UWidgetBlueprint* Blueprint, UWidgetAnimation* Animation, FString NewAnimationName, FText& OutErrorMessage )
+bool VerifyAnimationRename( FWidgetBlueprintEditor& BlueprintEditor, UWidgetAnimation* Animation, FString NewAnimationName, FText& OutErrorMessage )
 {
-	if ( FindObject<UWidgetAnimation>( Blueprint, *NewAnimationName, true ) )
+	UWidgetBlueprint* Blueprint = BlueprintEditor.GetWidgetBlueprintObj();
+	if (Blueprint && FindObject<UWidgetAnimation>( Blueprint, *NewAnimationName, true ) )
 	{
 		OutErrorMessage = LOCTEXT( "NameInUseByAnimation", "An animation with this name already exists" );
 		return false;
 	}
 
 	FName NewAnimationNameAsName( *NewAnimationName );
-	if ( Blueprint->WidgetTree->FindWidget<UWidget>( NewAnimationNameAsName ) != nullptr )
+	if ( Blueprint && Blueprint->WidgetTree->FindWidget<UWidget>( NewAnimationNameAsName ) != nullptr )
 	{
 		OutErrorMessage = LOCTEXT( "NameInUseByWidget", "A widget with this name already exists" );
 		return false;
 	}
 
+	UUserWidget* PreviewWidget = BlueprintEditor.GetPreview();
 	FName FunctionName(*NewAnimationName);
-	if (Animation->GetPreviewWidget().IsValid() && Animation->GetPreviewWidget().Get()->FindFunction(FunctionName))
+	if (PreviewWidget && PreviewWidget->FindFunction(FunctionName))
 	{
 		OutErrorMessage = LOCTEXT("NameInUseByFunction", "A function with this name already exists");
 		return false;
@@ -138,8 +140,8 @@ private:
 
 		if ( Animation->GetName() != NewName )
 		{
-			UWidgetBlueprint* Blueprint = BlueprintEditor.Pin()->GetWidgetBlueprintObj();
-			return VerifyAnimationRename( Blueprint, Animation, NewName, OutErrorMessage );
+			TSharedPtr<FWidgetBlueprintEditor> Editor = BlueprintEditor.Pin();
+			return Editor.IsValid() && VerifyAnimationRename( *Editor, Animation, NewName, OutErrorMessage );
 		}
 
 		return true;
@@ -367,7 +369,13 @@ private:
 		const float InTime = 0.f;
 		const float OutTime = 5.0f;
 
-		UWidgetBlueprint* WidgetBlueprint = BlueprintEditor.Pin()->GetWidgetBlueprintObj();
+		TSharedPtr<FWidgetBlueprintEditor> Editor = BlueprintEditor.Pin();
+		if (!Editor.IsValid())
+		{
+			return FReply::Handled();
+		}
+
+		UWidgetBlueprint* WidgetBlueprint = Editor->GetWidgetBlueprintObj();
 
 		FString BaseName = "NewAnimation";
 		UWidgetAnimation* NewAnimation = NewObject<UWidgetAnimation>(WidgetBlueprint, *BaseName, RF_Transactional);
@@ -375,7 +383,7 @@ private:
 		FString UniqueName = BaseName;
 		int32 NameIndex = 1;
 		FText Unused;
-		while ( VerifyAnimationRename( WidgetBlueprint, NewAnimation, UniqueName, Unused ) == false )
+		while ( VerifyAnimationRename( *Editor, NewAnimation, UniqueName, Unused ) == false )
 		{
 			UniqueName = FString::Printf( TEXT( "%s_%i" ), *BaseName, NameIndex );
 			NameIndex++;
diff --git a/Engine/Source/Editor/UMGEditor/Private/WidgetBlueprintEditor.cpp b/Engine/Source/Editor/UMGEditor/Private/WidgetBlueprintEditor.cpp
--- a/Engine/Source/Editor/UMGEditor/Private/WidgetBlueprintEditor.cpp
+++ b/Engine/Source/Editor/UMGEditor/Private/WidgetBlueprintEditor.cpp
@@ -117,12 +117,6 @@ void FWidgetBlueprintEditor::InitWidgetBlueprintEditor(const EToolkitMode::Type
 
 	UpdatePreview(GetWidgetBlueprintObj(), true);
 
-	// If the user has close the sequencer tab, this will not be initialized.
-	if (CurrentAnimation.IsValid())
-	{
-		CurrentAnimation->Initialize(PreviewWidgetPtr.Get());
-	}
-
 	DesignerCommandList = MakeShareable(new FUICommandList);
 
 	DesignerCommandList->MapAction(FGenericCommands::Get().Delete,
@@ -673,6 +667,8 @@ TSharedPtr<ISequencer>& FWidgetBlueprintEditor::GetSequencer()
 			SequencerInitParams.RootSequence = NullAnimation;
 			SequencerInitParams.bEditWithinLevelEditor = false;
 			SequencerInitParams.ToolkitHost = GetToolkitHost();
+			SequencerInitParams.PlaybackContext = TAttribute<UObject*>(this, &FWidgetBlueprintEditor::GetAnimationPlaybackContext);
+			SequencerInitParams.EventContexts = TAttribute<TArray<UObject*>>(this, &FWidgetBlueprintEditor::GetAnimationEventContexts);
 		};
 
 		Sequencer = FModuleManager::LoadModuleChecked<ISequencerModule>("Sequencer").CreateSequencer(SequencerInitParams);
@@ -685,16 +681,7 @@ TSharedPtr<ISequencer>& FWidgetBlueprintEditor::GetSequencer()
 
 void FWidgetBlueprintEditor::ChangeViewedAnimation( UWidgetAnimation& InAnimationToView )
 {
-	if (CurrentAnimation.IsValid())
-	{
-		CurrentAnimation->Initialize(nullptr);
-	}
-
 	CurrentAnimation = &InAnimationToView;
-	if (CurrentAnimation.IsValid())
-	{
-		CurrentAnimation->Initialize(PreviewWidgetPtr.Get());
-	}
 
 	if (Sequencer.IsValid())
 	{
@@ -863,13 +850,15 @@ void FWidgetBlueprintEditor::UpdatePreview(UBlueprint* InBlueprint, bool bInForc
 		PreviewWidgetPtr = PreviewActor;
 	}
 
-	if (CurrentAnimation.IsValid())
+	OnWidgetPreviewUpdated.Broadcast();
+
+	// We've changed the binding context so drastically that we should just clear all knowledge of our previous cached bindings
+
+	if (Sequencer.IsValid())
 	{
-		CurrentAnimation->Initialize(PreviewActor);
+		Sequencer->State.ClearObjectCaches();
+		Sequencer->ForceEvaluate();
 	}
-
-	OnWidgetPreviewUpdated.Broadcast();
-	GetSequencer()->UpdateRuntimeInstances();
 }
 
 FGraphAppearanceInfo FWidgetBlueprintEditor::GetGraphAppearance(UEdGraph* InGraph) const
@@ -1013,7 +1002,7 @@ void FWidgetBlueprintEditor::OnGetAnimationAddMenuContent(FMenuBuilder& MenuBuil
 	{
 		for (FObjectAndDisplayName& BindableObject : BindableObjects)
 		{
-			FGuid BoundObjectGuid = Sequencer->GetFocusedMovieSceneSequenceInstance()->FindObjectId(*BindableObject.Object);
+			FGuid BoundObjectGuid = Sequencer->FindObjectId(*BindableObject.Object, Sequencer->GetFocusedTemplateID());
 			if (BoundObjectGuid.IsValid() == false)
 			{
 				FUIAction AddMenuAction(FExecuteAction::CreateSP(this, &FWidgetBlueprintEditor::AddObjectToAnimation, BindableObject.Object));
@@ -1105,14 +1094,14 @@ void FWidgetBlueprintEditor::ReplaceTrackWithSelectedWidget(FWidgetReference Sel
 {
 	const FScopedTransaction Transaction( LOCTEXT( "ReplaceTrackWithSelectedWidget", "Replace Track with Selected Widget" ) );
 
-	UWidgetAnimation* WidgetAnimation = Cast<UWidgetAnimation> (GetSequencer().Get()->GetFocusedMovieSceneSequence());
+	UWidgetAnimation* WidgetAnimation = Cast<UWidgetAnimation>(Sequencer->GetFocusedMovieSceneSequence());
 	UMovieScene* MovieScene = WidgetAnimation->GetMovieScene();
 	UWidget* PreviewWidget = SelectedWidget.GetPreview();
 	UWidget* TemplateWidget = SelectedWidget.GetTemplate();
-	FGuid ObjectId = WidgetAnimation->FindPossessableObjectId(*BoundWidget);
+	FGuid ObjectId = BoundWidget ? Sequencer->FindObjectId(*BoundWidget, MovieSceneSequenceID::Root) : FGuid();
 
 	// Try find if the SelectedWidget is already bound, if so return
-	FGuid SelectedWidgetId = WidgetAnimation->FindPossessableObjectId(*PreviewWidget);
+	FGuid SelectedWidgetId = Sequencer->FindObjectId(*PreviewWidget, MovieSceneSequenceID::Root);
 	if (SelectedWidgetId.IsValid())
 	{
 		FNotificationInfo Info(LOCTEXT("SelectedWidgetalreadybound", "Selected Widget already bound"));
@@ -1126,7 +1115,7 @@ void FWidgetBlueprintEditor::ReplaceTrackWithSelectedWidget(FWidgetReference Sel
 		return;
 	}
 
-	if (TemplateWidget->GetClass() != BoundWidget->GetClass())
+	if (!BoundWidget || TemplateWidget->GetClass() != BoundWidget->GetClass())
 	{
 		TArray<FMovieSceneBinding> MovieSceneBindings = MovieScene->GetBindings();
 		for (FMovieSceneBinding Binding : MovieSceneBindings)
@@ -1162,53 +1151,41 @@ void FWidgetBlueprintEditor::ReplaceTrackWithSelectedWidget(FWidgetReference Sel
 		}
 	}
 
+	if (BoundWidget)
+	{
+		Sequencer->PreAnimatedState.RestorePreAnimatedState(*Sequencer, *BoundWidget);
+	}
+
 	// else it's safe to modify
 	MovieScene->Modify();
 	MovieScene->SetObjectDisplayName(ObjectId, FText::FromString(PreviewWidget->GetName()));
-	
-	FMovieScenePossessable NewPossessable(PreviewWidget->GetName(), PreviewWidget->GetClass());
-	FGuid PossessableGuid = NewPossessable.GetGuid();
-
-	MovieScene->ReplacePossessable(ObjectId, PossessableGuid, PreviewWidget->GetName());
-
-	UObject* BindingContext = GetSequencer().Get()->GetPlaybackContext();
 
 	// Replace bindings in WidgetAnimation
 	WidgetAnimation->Modify();
 	{
-		TArray<FWidgetAnimationBinding> NewBindings;
-		for (FWidgetAnimationBinding Binding : WidgetAnimation->AnimationBindings)
-		{
-			if (Binding.WidgetName == BoundWidget->GetFName())
+		FWidgetAnimationBinding* SourceBinding = WidgetAnimation->AnimationBindings.FindByPredicate(
+			[&](FWidgetAnimationBinding& In)
 			{
-				FWidgetAnimationBinding NewBinding;
-				{
-					if (Binding.SlotWidgetName == NAME_None)
-					{
-						NewBinding.AnimationGuid = PossessableGuid;
-					}
-					else
-					{
-						NewBinding.AnimationGuid = Binding.AnimationGuid;
-					}
-					NewBinding.WidgetName = PreviewWidget->GetFName();
-					NewBinding.SlotWidgetName = Binding.SlotWidgetName;
-					NewBinding.bIsRootWidget = Binding.bIsRootWidget;
-				}
-				NewBindings.Add(NewBinding);
+				return In.AnimationGuid == ObjectId;
 			}
-		}
-		WidgetAnimation->AnimationBindings.RemoveAll([&](const FWidgetAnimationBinding& Binding) {
-			return Binding.WidgetName == BoundWidget->GetFName();
-		});
-		for (FWidgetAnimationBinding Binding : NewBindings)
+		);
+
+		check(SourceBinding);
+
+		// Set binding names used for lookup
+		FName PredicateName = SourceBinding->WidgetName;
+		for (FWidgetAnimationBinding& Binding : WidgetAnimation->AnimationBindings)
 		{
-			WidgetAnimation->AnimationBindings.Add(Binding);
+			if (Binding.WidgetName != PredicateName)
+			{
+				continue;
+			}
+		
+			Binding.WidgetName = PreviewWidget->GetFName();
 		}
 	}
-	WidgetAnimation->ReplacePossessableObject(ObjectId, PossessableGuid, *BoundWidget, *PreviewWidget);
 
-	GetSequencer().Get()->NotifyMovieSceneDataChanged(EMovieSceneDataChangeType::MovieSceneStructureItemsChanged);
+	Sequencer->NotifyMovieSceneDataChanged(EMovieSceneDataChangeType::MovieSceneStructureItemsChanged);
 }
 
 void FWidgetBlueprintEditor::ExtendSequencerObjectBindingMenu(FMenuBuilder& ObjectBindingMenuBuilder, const TArray<UObject*> ContextObjects)
@@ -1276,23 +1253,23 @@ void FWidgetBlueprintEditor::OnMovieSceneDataChanged()
 
 void FWidgetBlueprintEditor::SyncSelectedWidgetsWithSequencerSelection(TArray<FGuid> ObjectGuids)
 {
-	UWidgetAnimation* WidgetAnimation = Cast<UWidgetAnimation>(GetSequencer().Get()->GetFocusedMovieSceneSequence());
-	UObject* BindingContext = GetSequencer().Get()->GetPlaybackContext();
+	UMovieSceneSequence* AnimationSequence = GetSequencer().Get()->GetFocusedMovieSceneSequence();
+	UObject* BindingContext = GetAnimationPlaybackContext();
 	TSet<FWidgetReference> SequencerSelectedWidgets;
 	for (FGuid Guid : ObjectGuids)
 	{
-		UObject* BoundObject = WidgetAnimation->FindPossessableObject(Guid, BindingContext);
-		if (!BoundObject)
+		TArray<UObject*, TInlineAllocator<1>> BoundObjects = AnimationSequence->LocateBoundObjects(Guid, BindingContext);
+		if (BoundObjects.Num() == 0)
 		{
 			continue;
 		}
-		else if (Cast<UPanelSlot>(BoundObject))
+		else if (Cast<UPanelSlot>(BoundObjects[0]))
 		{
-			SequencerSelectedWidgets.Add(GetReferenceFromPreview(Cast<UPanelSlot>(BoundObject)->Content));
+			SequencerSelectedWidgets.Add(GetReferenceFromPreview(Cast<UPanelSlot>(BoundObjects[0])->Content));
 		}
 		else
 		{
-			UWidget* BoundWidget = Cast<UWidget>(BoundObject);
+			UWidget* BoundWidget = Cast<UWidget>(BoundObjects[0]);
 			SequencerSelectedWidgets.Add(GetReferenceFromPreview(BoundWidget));
 		}
 	}
diff --git a/Engine/Source/Editor/UnrealEd/Private/Fbx/FbxMainExport.cpp b/Engine/Source/Editor/UnrealEd/Private/Fbx/FbxMainExport.cpp
--- a/Engine/Source/Editor/UnrealEd/Private/Fbx/FbxMainExport.cpp
+++ b/Engine/Source/Editor/UnrealEd/Private/Fbx/FbxMainExport.cpp
@@ -48,6 +48,7 @@
 #include "UnitConversion.h"
 
 #include "MovieScene.h"
+#include "MovieSceneFwd.h"
 #include "MovieSceneBinding.h"
 #include "IMovieScenePlayer.h"
 #include "Tracks/MovieScene3DTransformTrack.h"
@@ -55,6 +56,8 @@
 #include "Tracks/MovieSceneSkeletalAnimationTrack.h"
 #include "Sections/MovieScene3DTransformSection.h"
 #include "Sections/MovieSceneFloatSection.h"
+#include "MovieSceneEvaluationTemplateInstance.h"
+#include "MovieSceneSequence.h"
 
 #if WITH_PHYSX
 #include "PhysicsEngine/AggregateGeom.h"
@@ -1299,10 +1302,11 @@ bool FFbxExporter::ExportMatinee(AMatineeActor* InMatineeActor)
 }
 
 
-FFbxExporter::FLevelSequenceNodeNameAdapter::FLevelSequenceNodeNameAdapter( UMovieScene* InMovieScene, IMovieScenePlayer* InMovieScenePlayer )
+FFbxExporter::FLevelSequenceNodeNameAdapter::FLevelSequenceNodeNameAdapter( UMovieScene* InMovieScene, IMovieScenePlayer* InMovieScenePlayer, FMovieSceneSequenceIDRef InSequenceID )
 {
 	MovieScene = InMovieScene;
 	MovieScenePlayer = InMovieScenePlayer;
+	SequenceID = InSequenceID;
 }
 
 FString FFbxExporter::FLevelSequenceNodeNameAdapter::GetActorNodeName(const AActor* Actor)
@@ -1311,10 +1315,7 @@ FString FFbxExporter::FLevelSequenceNodeNameAdapter::GetActorNodeName(const AAct
 
 	for ( const FMovieSceneBinding& MovieSceneBinding : MovieScene->GetBindings() )
 	{
-		TArray<TWeakObjectPtr<UObject>> RuntimeObjects;
-		MovieScenePlayer->GetRuntimeObjects( MovieScenePlayer->GetRootMovieSceneSequenceInstance(), MovieSceneBinding.GetObjectGuid(), RuntimeObjects );
-
-		for ( TWeakObjectPtr<UObject> RuntimeObject : RuntimeObjects )
+		for ( TWeakObjectPtr<UObject> RuntimeObject : MovieScenePlayer->FindBoundObjects(MovieSceneBinding.GetObjectGuid(), SequenceID) )
 		{
 			if (RuntimeObject.Get() == Actor)
 			{
@@ -1352,12 +1353,12 @@ float FFbxExporter::FLevelSequenceAnimTrackAdapter::GetAnimationLength() const
 
 void FFbxExporter::FLevelSequenceAnimTrackAdapter::UpdateAnimation( float Time )
 {
-	EMovieSceneUpdateData UpdateData( Time, Time );
-	return MovieScenePlayer->GetRootMovieSceneSequenceInstance()->Update( UpdateData, *MovieScenePlayer );
+	FMovieSceneContext Context = FMovieSceneContext(FMovieSceneEvaluationRange(Time), MovieScenePlayer->GetPlaybackStatus()).SetHasJumped(true);
+	return MovieScenePlayer->GetEvaluationTemplate().Evaluate( Context, *MovieScenePlayer );
 }
 
 
-bool FFbxExporter::ExportLevelSequence( UMovieScene* MovieScene, const TArray<FGuid>& Bindings, IMovieScenePlayer* MovieScenePlayer )
+bool FFbxExporter::ExportLevelSequence( UMovieScene* MovieScene, const TArray<FGuid>& Bindings, IMovieScenePlayer* MovieScenePlayer, FMovieSceneSequenceIDRef SequenceID )
 {
 	if ( MovieScene == nullptr || MovieScenePlayer == nullptr )
 	{
@@ -1372,10 +1373,7 @@ bool FFbxExporter::ExportLevelSequence( UMovieScene* MovieScene, const TArray<FG
 			continue;
 		}
 
-		TArray<TWeakObjectPtr<UObject>> RuntimeObjects;
-		MovieScenePlayer->GetRuntimeObjects( MovieScenePlayer->GetRootMovieSceneSequenceInstance(), MovieSceneBinding.GetObjectGuid(), RuntimeObjects );
-
-		for ( TWeakObjectPtr<UObject> RuntimeObject : RuntimeObjects )
+		for ( TWeakObjectPtr<UObject> RuntimeObject : MovieScenePlayer->FindBoundObjects(MovieSceneBinding.GetObjectGuid(), SequenceID) )
 		{
 			if ( RuntimeObject.IsValid() )
 			{
diff --git a/Engine/Source/Editor/UnrealEd/Private/FbxExporter.h b/Engine/Source/Editor/UnrealEd/Private/FbxExporter.h
--- a/Engine/Source/Editor/UnrealEd/Private/FbxExporter.h
+++ b/Engine/Source/Editor/UnrealEd/Private/FbxExporter.h
@@ -4,6 +4,7 @@
 
 #include "Factories.h"
 #include "MatineeExporter.h"
+#include "MovieSceneFwd.h"
 #include "FbxImporter.h"
 
 class ALandscapeProxy;
@@ -82,7 +83,7 @@ public:
 	*
 	* @return	true, if successful
 	*/
-	bool ExportLevelSequence( UMovieScene* MovieScene, const TArray<FGuid>& InBindings, IMovieScenePlayer* MovieScenePlayer );
+	bool ExportLevelSequence( UMovieScene* MovieScene, const TArray<FGuid>& InBindings, IMovieScenePlayer* MovieScenePlayer, FMovieSceneSequenceIDRef SequenceID );
 
 	/**
 	 * Exports all the animation sequences part of a single Group in a Matinee sequence
@@ -158,11 +159,12 @@ public:
 	class UNREALED_API FLevelSequenceNodeNameAdapter : public INodeNameAdapter
 	{
 	public:
-		FLevelSequenceNodeNameAdapter( UMovieScene* InMovieScene, IMovieScenePlayer* InMovieScenePlayer );
+		FLevelSequenceNodeNameAdapter( UMovieScene* InMovieScene, IMovieScenePlayer* InMovieScenePlayer, FMovieSceneSequenceIDRef InSequenceID);
 		virtual FString GetActorNodeName(const AActor* InActor) override;
 	private:
 		UMovieScene* MovieScene;
 		IMovieScenePlayer* MovieScenePlayer;
+		FMovieSceneSequenceID SequenceID;
 	};
 
 	/* Get a valid unique name from a name */
diff --git a/Engine/Source/Editor/UnrealEd/Private/SCurveEditor.cpp b/Engine/Source/Editor/UnrealEd/Private/SCurveEditor.cpp
--- a/Engine/Source/Editor/UnrealEd/Private/SCurveEditor.cpp
+++ b/Engine/Source/Editor/UnrealEd/Private/SCurveEditor.cpp
@@ -3010,7 +3010,11 @@ void SCurveEditor::OnFlattenOrStraightenTangents(bool bFlattenTangents)
 
 			Key.Curve->GetKey(Key.KeyHandle).LeaveTangent = LeaveTangent;
 			Key.Curve->GetKey(Key.KeyHandle).ArriveTangent = ArriveTangent;
-			Key.Curve->GetKey(Key.KeyHandle).TangentMode = RCTM_User;
+			if (Key.Curve->GetKey(Key.KeyHandle).InterpMode == RCIM_Cubic &&
+				Key.Curve->GetKey(Key.KeyHandle).TangentMode == RCTM_Auto)
+			{
+				Key.Curve->GetKey(Key.KeyHandle).TangentMode = RCTM_User;
+			}
 		}
 				
 		for(auto It = SelectedTangents.CreateIterator();It;++It)
@@ -3034,7 +3038,11 @@ void SCurveEditor::OnFlattenOrStraightenTangents(bool bFlattenTangents)
 
 			Tangent.Key.Curve->GetKey(Tangent.Key.KeyHandle).LeaveTangent = LeaveTangent;
 			Tangent.Key.Curve->GetKey(Tangent.Key.KeyHandle).ArriveTangent = ArriveTangent;
-			Tangent.Key.Curve->GetKey(Tangent.Key.KeyHandle).TangentMode = RCTM_User;
+			if (Tangent.Key.Curve->GetKey(Tangent.Key.KeyHandle).InterpMode == RCIM_Cubic &&
+				Tangent.Key.Curve->GetKey(Tangent.Key.KeyHandle).TangentMode == RCTM_Auto)
+			{
+				Tangent.Key.Curve->GetKey(Tangent.Key.KeyHandle).TangentMode = RCTM_User;
+			}
 		}
 
 		TArray<FRichCurveEditInfo> ChangedCurveEditInfos;
diff --git a/Engine/Source/Editor/UnrealEd/Public/TickableEditorObject.h b/Engine/Source/Editor/UnrealEd/Public/TickableEditorObject.h
--- a/Engine/Source/Editor/UnrealEd/Public/TickableEditorObject.h
+++ b/Engine/Source/Editor/UnrealEd/Public/TickableEditorObject.h
@@ -1,40 +1,39 @@
 // Copyright 1998-2016 Epic Games, Inc. All Rights Reserved.
 
 #pragma once
+
 #include "Tickable.h"
 
+
 /**
  * This class provides common registration for gamethread editor only tickable objects. It is an
- * abstract base class requiring you to implement the Tick() method.
+ * abstract base class requiring you to implement the GetStatId, IsTickable, and Tick methods.
  */
 class FTickableEditorObject : public FTickableObjectBase
 {
 public:
-	static void TickObjects( float DeltaSeconds )
+
+	static void TickObjects(float DeltaSeconds)
 	{
 		const TArray<FTickableEditorObject*>& TickableObjects = GetTickableObjects();
 
-		for( int32 ObjectIndex=0; ObjectIndex < TickableObjects.Num(); ++ObjectIndex)
+		for (int32 ObjectIndex=0; ObjectIndex < TickableObjects.Num(); ++ObjectIndex)
 		{
 			FTickableEditorObject* TickableObject = TickableObjects[ObjectIndex];
-			if ( TickableObject->IsTickable() )
+			if (TickableObject->IsTickable())
 			{
 				TickableObject->Tick(DeltaSeconds);
 			}
 		}
 	}
 
-	/**
-	 * Registers this instance with the static array of tickable objects.	
-	 */
+	/** Registers this instance with the static array of tickable objects. */
 	FTickableEditorObject()
 	{
-		GetTickableObjects().Add( this );
+		GetTickableObjects().Add(this);
 	}
 
-	/**
-	 * Removes this instance from the static array of tickable objects.
-	 */
+	/** Removes this instance from the static array of tickable objects. */
 	virtual ~FTickableEditorObject()
 	{
 		UnregisterTickableObject(this);
@@ -44,7 +43,9 @@ private:
 
 	/**
 	 * Class that avoids crashes when unregistering a tickable editor object too late.
-	 * Some tickable objects can outlive the collection (global/static destructor order is unpredictable).
+	 *
+	 * Some tickable objects can outlive the collection
+	 * (global/static destructor order is unpredictable).
 	 */
 	class TTickableObjectsCollection : public TArray<FTickableEditorObject*>
 	{
diff --git a/Engine/Source/Programs/UnrealBuildTool/System/UEPlatformProjectGenerator.cs b/Engine/Source/Programs/UnrealBuildTool/System/UEPlatformProjectGenerator.cs
--- a/Engine/Source/Programs/UnrealBuildTool/System/UEPlatformProjectGenerator.cs
+++ b/Engine/Source/Programs/UnrealBuildTool/System/UEPlatformProjectGenerator.cs
@@ -229,12 +229,22 @@ namespace UnrealBuildTool
 		}
 
 		/// <summary>
-		/// Return any custom property import lines
+		/// Return any custom property settings. These will be included right after Global properties to make values available to all other imports.
+		/// </summary>
+		/// <param name="InPlatform">  The UnrealTargetPlatform being built</param>
+		/// <returns>string    The custom property import lines for the project file; Empty string if it doesn't require one</returns>
+		public virtual string GetVisualStudioGlobalProperties(UnrealTargetPlatform InPlatform)
+		{
+			return "";
+		}
+
+		/// <summary>
+		/// Return any custom target overrides. These will be included last in the project file so they have the opportunity to override any existing settings.
 		/// </summary>
 		/// <param name="InPlatform">  The UnrealTargetPlatform being built</param>
 		/// <param name="InProjectFileFormat">The visual studio project file format being generated</param>
 		/// <returns>string    The custom property import lines for the project file; Empty string if it doesn't require one</returns>
-		public virtual string GetAdditionalVisualStudioImportSettings(UnrealTargetPlatform InPlatform, VCProjectFileFormat InProjectFileFormat)
+		public virtual string GetVisualStudioTargetOverrides(UnrealTargetPlatform InPlatform, VCProjectFileFormat InProjectFileFormat)
 		{
 			return "";
 		}
diff --git a/Engine/Source/Programs/UnrealBuildTool/System/VCProject.cs b/Engine/Source/Programs/UnrealBuildTool/System/VCProject.cs
--- a/Engine/Source/Programs/UnrealBuildTool/System/VCProject.cs
+++ b/Engine/Source/Programs/UnrealBuildTool/System/VCProject.cs
@@ -494,6 +494,19 @@ namespace UnrealBuildTool
                     "	</PropertyGroup>" + ProjectFileGenerator.NewLine);
 			}
 
+			// look for additional import lines for all platforms for non stub projects
+			if (!IsStubProject)
+			{
+				foreach (UnrealTargetPlatform Platform in ProjectPlatforms)
+				{
+					UEPlatformProjectGenerator ProjGenerator = UEPlatformProjectGenerator.GetPlatformProjectGenerator(Platform, true);
+					if (ProjGenerator != null && ProjGenerator.HasVisualStudioSupport(Platform, UnrealTargetConfiguration.Development, ProjectFileFormat))
+					{
+						VCProjectFileContent.Append(ProjGenerator.GetVisualStudioGlobalProperties(Platform));
+					}
+				}
+			}
+
 			// Write each project configuration PreDefaultProps section
 			foreach (Tuple<string, UnrealTargetConfiguration> ConfigurationTuple in ProjectConfigurationNameAndConfigurations)
 			{
@@ -688,7 +701,7 @@ namespace UnrealBuildTool
 			}
 
 			// look for additional import lines for all platforms for non stub projects
-			StringBuilder AdditionalImportSettings = new StringBuilder();
+			StringBuilder AdditionalTargetSettings = new StringBuilder();
 			if (!IsStubProject)
 			{
 				foreach (UnrealTargetPlatform Platform in ProjectPlatforms)
@@ -696,7 +709,7 @@ namespace UnrealBuildTool
 					UEPlatformProjectGenerator ProjGenerator = UEPlatformProjectGenerator.GetPlatformProjectGenerator(Platform, true);
 					if (ProjGenerator != null && ProjGenerator.HasVisualStudioSupport(Platform, UnrealTargetConfiguration.Development, ProjectFileFormat))
 					{
-						AdditionalImportSettings.Append(ProjGenerator.GetAdditionalVisualStudioImportSettings(Platform, ProjectFileFormat));
+						AdditionalTargetSettings.Append(ProjGenerator.GetVisualStudioTargetOverrides(Platform, ProjectFileFormat));
 					}
 				}
 			}
@@ -738,7 +751,7 @@ namespace UnrealBuildTool
 					"	<ItemDefinitionGroup>" + ProjectFileGenerator.NewLine +
 					"	</ItemDefinitionGroup>" + ProjectFileGenerator.NewLine +
 					"	<Import Project=\"$(VCTargetsPath)\\Microsoft.Cpp.targets\" />" + ProjectFileGenerator.NewLine +
-					AdditionalImportSettings.ToString() +
+					AdditionalTargetSettings.ToString() +
 					"	<ImportGroup Label=\"ExtensionTargets\">" + ProjectFileGenerator.NewLine +
 					"	</ImportGroup>" + ProjectFileGenerator.NewLine +
 					"</Project>" + ProjectFileGenerator.NewLine);
diff --git a/Engine/Source/Programs/UnrealFrontend/Private/Commands/StatsDumpMemoryCommand.cpp b/Engine/Source/Programs/UnrealFrontend/Private/Commands/StatsDumpMemoryCommand.cpp
--- a/Engine/Source/Programs/UnrealFrontend/Private/Commands/StatsDumpMemoryCommand.cpp
+++ b/Engine/Source/Programs/UnrealFrontend/Private/Commands/StatsDumpMemoryCommand.cpp
@@ -2,19 +2,15 @@
 
 #include "UnrealFrontendPrivatePCH.h"
 #include "StatsDumpMemoryCommand.h"
-#include "Profiler.h"
+#include "IProfilerModule.h"
 
-/*-----------------------------------------------------------------------------
-	FStatsMemoryDumpCommand
------------------------------------------------------------------------------*/
 
 void FStatsMemoryDumpCommand::Run()
 {
 	FString SourceFilepath;
-	FParse::Value( FCommandLine::Get(), TEXT( "-INFILE=" ), SourceFilepath );
+	FParse::Value(FCommandLine::Get(), TEXT("-INFILE="), SourceFilepath);
 
-	const FName NAME_ProfilerModule = TEXT( "Profiler" );
-	IProfilerModule& ProfilerModule = FModuleManager::LoadModuleChecked<IProfilerModule>( NAME_ProfilerModule );
-	ProfilerModule.StatsMemoryDumpCommand( *SourceFilepath );
+	const FName NAME_ProfilerModule = TEXT("Profiler");
+	IProfilerModule& ProfilerModule = FModuleManager::LoadModuleChecked<IProfilerModule>(NAME_ProfilerModule);
+	ProfilerModule.StatsMemoryDumpCommand(*SourceFilepath);
 }
-
diff --git a/Engine/Source/Runtime/CinematicCamera/Private/CineCameraActor.cpp b/Engine/Source/Runtime/CinematicCamera/Private/CineCameraActor.cpp
--- a/Engine/Source/Runtime/CinematicCamera/Private/CineCameraActor.cpp
+++ b/Engine/Source/Runtime/CinematicCamera/Private/CineCameraActor.cpp
@@ -17,23 +17,19 @@ ACineCameraActor::ACineCameraActor(const FObjectInitializer& ObjectInitializer)
 	CineCameraComponent = Cast<UCineCameraComponent>(GetCameraComponent());
 
 	PrimaryActorTick.bCanEverTick = true;
-	SetActorTickEnabled(ShouldTickForTracking());
+	SetActorTickEnabled(true);
 }
 
 #if WITH_EDITOR
 void ACineCameraActor::PostEditChangeProperty(FPropertyChangedEvent& PropertyChangedEvent)
 {
 	Super::PostEditChangeProperty(PropertyChangedEvent);
-
-	SetActorTickEnabled(ShouldTickForTracking());
 }
 #endif
 
 void ACineCameraActor::PostInitializeComponents()
 {
 	Super::PostInitializeComponents();
-	SetActorTickEnabled(ShouldTickForTracking());
-
 	LookatTrackingSettings.LastLookatTrackingRotation = GetActorRotation();
 }
 
@@ -123,11 +119,6 @@ void ACineCameraActor::Tick(float DeltaTime)
 		}
 #endif // WITH_EDITORONLY_DATA
 	}
-	else
-	{
-		// no tracking, no ticking
-		SetActorTickEnabled(false);
-	}
 
 	bResetInterplation = false;
 }
diff --git a/Engine/Source/Runtime/Core/Private/Tests/Math/RangeTest.cpp b/Engine/Source/Runtime/Core/Private/Tests/Math/RangeTest.cpp
--- a/Engine/Source/Runtime/Core/Private/Tests/Math/RangeTest.cpp
+++ b/Engine/Source/Runtime/Core/Private/Tests/Math/RangeTest.cpp
@@ -322,7 +322,7 @@ bool FRangeTest::RunTest( const FString& Parameters )
 	TestEqual(TEXT("The intersection of [7, 14) and [2, 8] must be [7, 8]"), FFloatRange::Intersection(r17_1, r17_3), FFloatRange(FFloatRangeBound::Inclusive(7.0f), FFloatRangeBound::Inclusive(8.0f)));
 	TestEqual(TEXT("The intersection of [7, 14) and (8, 16) must be (8, 14)"), FFloatRange::Intersection(r17_1, r17_4), FFloatRange(FFloatRangeBound::Exclusive(8.0f), FFloatRangeBound::Exclusive(14.0f)));
 	TestEqual(TEXT("The intersection of [7, 14) and [8, 16) must be [8, 14)"), FFloatRange::Intersection(r17_1, r17_5), FFloatRange(FFloatRangeBound::Inclusive(8.0f), FFloatRangeBound::Exclusive(14.0f)));
-	TestEqual(TEXT("The intersection of a non-empty range and an empty range must be the non-empty range"), FFloatRange::Intersection(r17_1, r17_6), r17_1);
+	TestTrue(TEXT("The intersection of a non-empty range and an empty range must be empty"), FFloatRange::Intersection(r17_1, r17_6).IsEmpty());
 	TestTrue(TEXT("The intersection of two empty ranges must be empty"), FFloatRange::Intersection(r17_6, r17_7).IsEmpty());
 
 	// union
diff --git a/Engine/Source/Runtime/Core/Public/Math/Range.h b/Engine/Source/Runtime/Core/Public/Math/Range.h
--- a/Engine/Source/Runtime/Core/Public/Math/Range.h
+++ b/Engine/Source/Runtime/Core/Public/Math/Range.h
@@ -138,13 +138,13 @@ public:
 			return false;
 		}
 
-		if (UpperBound.GetValue() == Other.LowerBound.GetValue())
+		if (!UpperBound.IsOpen() && !Other.LowerBound.IsOpen() && UpperBound.GetValue() == Other.LowerBound.GetValue())
 		{
 			return ((UpperBound.IsInclusive() && Other.LowerBound.IsExclusive()) ||
 					(UpperBound.IsExclusive() && Other.LowerBound.IsInclusive()));
 		}
 		
-		if (Other.UpperBound.GetValue() == LowerBound.GetValue())
+		if (!Other.UpperBound.IsOpen() && !LowerBound.IsOpen() && Other.UpperBound.GetValue() == LowerBound.GetValue())
 		{
 			return ((Other.UpperBound.IsInclusive() && LowerBound.IsExclusive()) ||
 					(Other.UpperBound.IsExclusive() && LowerBound.IsInclusive()));
@@ -506,12 +506,12 @@ public:
 	{
 		if (X.IsEmpty())
 		{
-			return Y;
+			return TRange::Empty();
 		}
 
 		if (Y.IsEmpty())
 		{
-			return X;
+			return TRange::Empty();
 		}
 
 		return TRange(BoundsType::MaxLower(X.LowerBound, Y.LowerBound), BoundsType::MinUpper(X.UpperBound, Y.UpperBound));
diff --git a/Engine/Source/Runtime/CoreUObject/Public/UObject/ObjectKey.h b/Engine/Source/Runtime/CoreUObject/Public/UObject/ObjectKey.h
--- a/Engine/Source/Runtime/CoreUObject/Public/UObject/ObjectKey.h
+++ b/Engine/Source/Runtime/CoreUObject/Public/UObject/ObjectKey.h
@@ -18,15 +18,15 @@ public:
 
 	/** Construct from an object pointer */
 	FORCEINLINE FObjectKey(const UObject *Object)
-	    : ObjectIndex(INDEX_NONE)
-	    , ObjectSerialNumber(0)
+		: ObjectIndex(INDEX_NONE)
+		, ObjectSerialNumber(0)
 	{
-	    if (Object)
-	    {
-	        FWeakObjectPtr Weak(Object);
-	        ObjectIndex = Weak.ObjectIndex;
-	        ObjectSerialNumber = Weak.ObjectSerialNumber;
-	    }
+		if (Object)
+		{
+			FWeakObjectPtr Weak(Object);
+			ObjectIndex = Weak.ObjectIndex;
+			ObjectSerialNumber = Weak.ObjectSerialNumber;
+		}
 	}
 
 	/** Compare this key with another */
@@ -41,6 +41,19 @@ public:
 		return ObjectIndex != Other.ObjectIndex || ObjectSerialNumber != Other.ObjectSerialNumber;
 	}
 
+	/**
+	 * Attempt to access the object from which this key was constructed.
+	 * @return The object used to construct this key, or nullptr if it is no longer valid
+	 */
+	UObject* ResolveObjectPtr() const
+	{
+		FWeakObjectPtr WeakPtr;
+		WeakPtr.ObjectIndex = ObjectIndex;
+		WeakPtr.ObjectSerialNumber = ObjectSerialNumber;
+
+		return WeakPtr.Get();
+	}
+
 	/** Hash function */
 	friend uint32 GetTypeHash(const FObjectKey& Key)
 	{
diff --git a/Engine/Source/Runtime/Engine/Classes/Curves/CurveLinearColor.h b/Engine/Source/Runtime/Engine/Classes/Curves/CurveLinearColor.h
--- a/Engine/Source/Runtime/Engine/Classes/Curves/CurveLinearColor.h
+++ b/Engine/Source/Runtime/Engine/Classes/Curves/CurveLinearColor.h
@@ -1,9 +1,11 @@
 // Copyright 1998-2016 Epic Games, Inc. All Rights Reserved.
 
-
 #pragma once
+
+#include "CurveBase.h"
 #include "CurveLinearColor.generated.h"
 
+
 USTRUCT()
 struct ENGINE_API FRuntimeCurveLinearColor
 {
diff --git a/Engine/Source/Runtime/Engine/Classes/Curves/KeyHandle.h b/Engine/Source/Runtime/Engine/Classes/Curves/KeyHandle.h
--- a/Engine/Source/Runtime/Engine/Classes/Curves/KeyHandle.h
+++ b/Engine/Source/Runtime/Engine/Classes/Curves/KeyHandle.h
@@ -89,3 +89,63 @@ struct TStructOpsTypeTraits<FKeyHandleMap>
 		WithIdenticalViaEquality = true,
 	};
 };
+
+/**
+ * Lookup table that maps key handles to indices in an external data structure
+ * Maintains a map of key handle to last known index,
+ * and an array of optional key handles that's used to validate map entries.
+ */
+struct FKeyHandleLookupTable
+{
+	/**
+	 * Get the index that corresponds to the specified key handle
+	 *
+	 * @param KeyHandle		The key handle to retrieve an up to date index for
+	 * @retrun The index that corresponds to this key handle, or INDEX_NONE if the key handle is invalid
+	 */
+	ENGINE_API int32 GetIndex(FKeyHandle KeyHandle);
+
+	/**
+	 * Attempt to find the handle for the specified index, or allocate a new one if it doesn't have one
+	 *
+	 * @param Index			The index to find or allocate a handle for
+	 * @retrun A valid key handle for the specified index
+	 */
+	ENGINE_API FKeyHandle FindOrAddKeyHandle(int32 Index);
+
+	/**
+	 * Move the specified handle from its previous index, to its new index
+	 *
+	 * @param OldIndex		The previous index of the handle
+	 * @param NewIndex		The new index of the handle
+	 */
+	ENGINE_API void MoveHandle(int32 OldIndex, int32 NewIndex);
+
+	/**
+	 * Allocate a new handle for the specified index
+	 *
+	 * @param Index			The index to associate a handle with
+	 * @return A new key handle that corresponds to the specified index
+	 */
+	ENGINE_API FKeyHandle AllocateHandle(int32 Index);
+
+	/**
+	 * Deallocate the specified index
+	 *
+	 * @param Index			The index to invalidate a handle for
+	 */
+	ENGINE_API void DeallocateHandle(int32 Index);
+
+	/**
+	 * Reset this lookup table, forgetting all key handles and indices
+	 */
+	ENGINE_API void Reset();
+
+private:
+
+	/** Array of optional key handles that reside in corresponding indices for an externally owned serial data structure */
+	TArray<TOptional<FKeyHandle>> KeyHandles;
+
+	/** Map of which key handles go to which indices. Indices may point to incorrect indices. Check against entry in KeyHandles array to verify. */
+	TMap<FKeyHandle, int32> KeyHandlesToIndices;
+};
\ No newline at end of file
diff --git a/Engine/Source/Runtime/Engine/Classes/Curves/StringCurve.h b/Engine/Source/Runtime/Engine/Classes/Curves/StringCurve.h
--- a/Engine/Source/Runtime/Engine/Classes/Curves/StringCurve.h
+++ b/Engine/Source/Runtime/Engine/Classes/Curves/StringCurve.h
@@ -49,13 +49,6 @@ public:
 };
 
 
-template<>
-struct TIsPODType<FStringCurveKey>
-{
-	enum { Value = true };
-};
-
-
 template<>
 struct TStructOpsTypeTraits<FStringCurveKey>
 	: public TStructOpsTypeTraitsBase
diff --git a/Engine/Source/Runtime/Engine/Private/Curves/KeyHandle.cpp b/Engine/Source/Runtime/Engine/Private/Curves/KeyHandle.cpp
--- a/Engine/Source/Runtime/Engine/Private/Curves/KeyHandle.cpp
+++ b/Engine/Source/Runtime/Engine/Private/Curves/KeyHandle.cpp
@@ -90,3 +90,95 @@ bool FKeyHandleMap::operator!=(const FKeyHandleMap& Other) const
 {
 	return !(*this==Other);
 }
+
+
+int32 FKeyHandleLookupTable::GetIndex(FKeyHandle KeyHandle)
+{
+	const int32* Index = KeyHandlesToIndices.Find(KeyHandle);
+	if (!Index)
+	{
+		// If it's not even in the map, there's no way this could be a valid handle for this container
+		return INDEX_NONE;
+	}
+	else if (KeyHandles.IsValidIndex(*Index) && KeyHandles[*Index] == KeyHandle)
+	{
+		return *Index;
+	}
+
+	// slow lookup and cache
+	const int32 NewCacheIndex = KeyHandles.IndexOfByPredicate(
+		[KeyHandle](const TOptional<FKeyHandle>& PredKeyHandle)
+		{
+			return PredKeyHandle.IsSet() && PredKeyHandle.GetValue() == KeyHandle;
+		}
+	);
+
+	if (NewCacheIndex == INDEX_NONE)
+	{
+		return INDEX_NONE;
+	}
+
+	KeyHandlesToIndices.Add(KeyHandle, NewCacheIndex);
+	return NewCacheIndex;
+}
+
+FKeyHandle FKeyHandleLookupTable::FindOrAddKeyHandle(int32 Index)
+{
+	if (KeyHandles.IsValidIndex(Index) && KeyHandles[Index].IsSet())
+	{
+		return KeyHandles[Index].GetValue();
+	}
+	
+	int32 NumToAdd = Index + 1 - KeyHandles.Num();
+	if (NumToAdd > 0)
+	{
+		KeyHandles.AddDefaulted(NumToAdd);
+	}
+
+	// Allocate a new key handle
+	FKeyHandle NewKeyHandle;
+
+	KeyHandles[Index] = NewKeyHandle;
+	KeyHandlesToIndices.Add(NewKeyHandle, Index);
+
+	return NewKeyHandle;
+}
+
+void FKeyHandleLookupTable::MoveHandle(int32 OldIndex, int32 NewIndex)
+{
+	if (KeyHandles.IsValidIndex(OldIndex))
+	{
+		TOptional<FKeyHandle> Handle = KeyHandles[OldIndex];
+
+		KeyHandles.RemoveAtSwap(OldIndex, 1, false);
+		KeyHandles.Insert(Handle, NewIndex);
+		if (Handle.IsSet())
+		{
+			KeyHandlesToIndices.Add(Handle.GetValue(), NewIndex);
+		}
+	}
+}
+
+FKeyHandle FKeyHandleLookupTable::AllocateHandle(int32 Index)
+{
+	FKeyHandle NewKeyHandle;
+	KeyHandles.Insert(NewKeyHandle, Index);
+	KeyHandlesToIndices.Add(NewKeyHandle, Index);
+	return NewKeyHandle;
+}
+
+void FKeyHandleLookupTable::DeallocateHandle(int32 Index)
+{
+	TOptional<FKeyHandle> KeyHandle = KeyHandles[Index];
+	KeyHandles.RemoveAtSwap(Index, 1, false);
+	if (KeyHandle.IsSet())
+	{
+		KeyHandlesToIndices.Remove(KeyHandle.GetValue());
+	}
+}
+
+void FKeyHandleLookupTable::Reset()
+{
+	KeyHandles.Reset();
+	KeyHandlesToIndices.Reset();
+}
\ No newline at end of file
diff --git a/Engine/Source/Runtime/Engine/Public/Curves/CurveInterface.h b/Engine/Source/Runtime/Engine/Public/Curves/CurveInterface.h
new file mode 100644
--- /dev/null
+++ b/Engine/Source/Runtime/Engine/Public/Curves/CurveInterface.h
@@ -0,0 +1,162 @@
+#pragma once
+
+#include "Curves/KeyHandle.h"
+#include "KeyFrameManipulator.h"
+
+/**
+ * Proxy type used to reference a key's time and value
+ */
+template<typename KeyValueType, typename TimeType>
+struct TKeyFrameProxy
+{
+	TKeyFrameProxy(TimeType InTime, KeyValueType& InValue)
+		: Time(InTime)
+		, Value(InValue)
+	{}
+
+	/** The key time */
+	TimeType Time;
+
+	/** (potentially const) Reference to the key's value */
+	KeyValueType& Value;
+};
+
+/**
+ * Templated interface to externally owned curve data data
+ */
+template<typename KeyValueType, typename TimeType>
+class TCurveInterface : public TKeyFrameManipulator<TimeType>
+{
+public:
+	// Pass by value/ref parameter types
+	typedef typename TCallTraits<TimeType>::ParamType TimeTypeRef;
+	typedef typename TCallTraits<KeyValueType>::ParamType KeyValueTypeRef;
+
+	/**
+	 * Construction from externally owned curve data
+	 *
+	 * @param KeyTimesParam			Ptr to array of key times
+	 * @param KeyTimesParam			Ptr to array of key values
+	 * @param ExternalKeyHandleLUT	Optional external look-up-table for key handles. If not supplied, an internal temporary LUT will be used.
+	 */
+	TCurveInterface(TArray<TimeType>* KeyTimesParam, TArray<KeyValueType>* KeyValuesParam, FKeyHandleLookupTable* ExternalKeyHandleLUT = nullptr)
+		: TKeyFrameManipulator<TimeType>(KeyTimesParam, ExternalKeyHandleLUT)
+		, KeyValues(KeyValuesParam)
+	{
+	}
+
+	TCurveInterface(const TCurveInterface&) = default;
+	TCurveInterface& operator=(const TCurveInterface&) = default;
+
+#if PLATFORM_COMPILER_HAS_DEFAULTED_CONSTRUCTORS
+	TCurveInterface(TCurveInterface&&) = default;
+	TCurveInterface& operator=(TCurveInterface&&) = default;
+#endif
+
+	/**
+	  * Add a new key to the curve with the supplied Time and Value. Returns the handle of the new key.
+	  * 
+	  * @param InTime				The time at which to add the key.
+	  * @param InValue				The value of the key.
+	  * @return A unique identifier for the newly added key
+	  */
+	FKeyHandle AddKeyValue(TimeTypeRef InTime, KeyValueTypeRef InValue)
+	{
+		int32 InsertIndex = this->ComputeInsertIndex(InTime);
+		FKeyHandle NewHandle = this->InsertKeyImpl(InTime, InsertIndex);
+
+		KeyValues->Insert(InValue, InsertIndex);
+		return NewHandle;
+	}
+
+	/**
+	 * Attempt to retrieve a key from its handle
+	 *
+	 * @param 						KeyHandle Handle to the key to get.
+	 * @return A proxy to the key, or empty container if it was not found
+	 */
+	TOptional<TKeyFrameProxy<const KeyValueType, TimeType>> GetKey(FKeyHandle KeyHandle) const
+	{
+		int32 Index = this->GetIndex(KeyHandle);
+		if (Index != INDEX_NONE)
+		{
+			return TKeyFrameProxy<const KeyValueType, const TimeType>(this->GetKeyTimeChecked(Index), (*KeyValues)[Index]);
+		}
+		return TOptional<TKeyFrameProxy<const KeyValueType, TimeType>>();
+	}
+
+	/**
+	 * Attempt to retrieve a key from its handle
+	 *
+	 * @param 						KeyHandle Handle to the key to get.
+	 * @return A proxy to the key, or empty container if it was not found
+	 */
+	TOptional<TKeyFrameProxy<KeyValueType, TimeType>> GetKey(FKeyHandle KeyHandle)
+	{
+		int32 Index = this->GetIndex(KeyHandle);
+		if (Index != INDEX_NONE)
+		{
+			return TKeyFrameProxy<KeyValueType, TimeType>(this->GetKeyTimeChecked(Index), (*KeyValues)[Index]);
+		}
+		return TOptional<TKeyFrameProxy<KeyValueType, TimeType>>();
+	}
+
+	/**
+	 * Update the key at the specified time with a new value, or add one if none is present
+	 *
+	 * @param InTime				The time at which to update/add the value
+	 * @param Value					The new key value
+	 * @param KeyTimeTolerance		Tolerance to use when looking for an existing key
+	 * @return The key's handle
+	 */
+	FKeyHandle UpdateOrAddKey(TimeTypeRef InTime, KeyValueTypeRef Value, TimeTypeRef KeyTimeTolerance)
+	{
+		TOptional<FKeyHandle> Handle = this->FindKey(
+			[=](TimeType ExistingTime)
+			{
+				return FMath::IsNearlyEqual(InTime, ExistingTime, KeyTimeTolerance);
+			});
+		
+		if (!Handle.IsSet())
+		{
+			return AddKeyValue(InTime, Value);
+		}
+		else
+		{
+			TOptional<TKeyFrameProxy<KeyValueType, TimeType>> Key = GetKey(Handle.GetValue());
+
+			if (ensure(Key.IsSet()))
+			{
+				Key->Value = Value;
+			}
+
+			return Handle.GetValue();
+		}
+	}
+
+protected:
+
+	virtual void OnKeyAdded(int32 Index) override
+	{
+		KeyValueType NewValue;
+		KeyValues->Insert(MoveTemp(NewValue), Index);
+	}
+
+	virtual void OnKeyRelocated(int32 OldIndex, int32 NewIndex) override
+	{
+		KeyValueType Value = (*KeyValues)[OldIndex];
+
+		KeyValues->RemoveAtSwap(OldIndex, 1, false);
+		KeyValues->Insert(MoveTemp(Value), NewIndex);
+	}
+
+	virtual void OnKeyRemoved(int32 Index) override
+	{
+		KeyValues->RemoveAtSwap(Index, 1, false);
+	}
+
+private:
+
+	/** Array of associated key data */
+	TArray<KeyValueType>* KeyValues;
+};
\ No newline at end of file
