UMG: Fixed parent bindings not being adhered to correctly. Fixed slot widgets that get recreated not having their object bindings updated. This fixes a regression in all types of keyframing in UMG found after copying Dev-Sequencer to main.

Copied from Dev-Sequencer

diff --git a/Engine/Source/Editor/Sequencer/Private/Sequencer.cpp b/Engine/Source/Editor/Sequencer/Private/Sequencer.cpp
--- a/Engine/Source/Editor/Sequencer/Private/Sequencer.cpp
+++ b/Engine/Source/Editor/Sequencer/Private/Sequencer.cpp
@@ -718,12 +718,19 @@ FGuid FSequencer::CreateBinding(UObject& InObject, const FString& InName)
 	const FGuid PossessableGuid = OwnerMovieScene->AddPossessable(InName, InObject.GetClass());
 
 	// Attempt to use the parent as a context if necessary
-	UObject* BindingContext = OwnerSequence->AreParentContextsSignificant() ? OwnerSequence->GetParentObject(&InObject) : nullptr;
-	if (BindingContext)
+	UObject* ParentObject = OwnerSequence->GetParentObject(&InObject);
+	UObject* BindingContext = GetPlaybackContext();
+
+	if (ParentObject)
 	{
 		// Ensure we have possessed the outer object, if necessary
-		FGuid ParentGuid = GetHandleToObject(BindingContext);
+		FGuid ParentGuid = GetHandleToObject(ParentObject);
 		
+		if (OwnerSequence->AreParentContextsSignificant())
+		{
+			BindingContext = ParentObject;
+		}
+
 		// Set up parent/child guids for possessables within spawnables
 		if (ParentGuid.IsValid())
 		{
@@ -740,10 +747,6 @@ FGuid FSequencer::CreateBinding(UObject& InObject, const FString& InName)
 			}
 		}
 	}
-	else
-	{
-		BindingContext = GetPlaybackContext();
-	}
 
 	OwnerSequence->BindPossessableObject(PossessableGuid, InObject, BindingContext);
 
@@ -4807,7 +4810,7 @@ void FSequencer::PasteCopiedTracks(TArray<TSharedPtr<FSequencerObjectBindingNode
 	for (TSharedPtr<FSequencerObjectBindingNode> ObjectNode : ObjectNodes)
 	{
 		FGuid ObjectGuid = ObjectNode->GetObjectBinding();
-		UObject* Context = nullptr;
+		UObject* Context = GetPlaybackContext();
 		TArray<UObject*, TInlineAllocator<1>> PossessableObjects = GetFocusedMovieSceneSequence()->LocateBoundObjects(ObjectGuid, Context);
 
 		if (PossessableObjects.Num() == 0)
diff --git a/Engine/Source/Runtime/MovieScene/Private/Evaluation/MovieSceneEvaluationState.cpp b/Engine/Source/Runtime/MovieScene/Private/Evaluation/MovieSceneEvaluationState.cpp
--- a/Engine/Source/Runtime/MovieScene/Private/Evaluation/MovieSceneEvaluationState.cpp
+++ b/Engine/Source/Runtime/MovieScene/Private/Evaluation/MovieSceneEvaluationState.cpp
@@ -52,6 +52,10 @@ FGuid FMovieSceneObjectCache::FindObjectId(UObject& InObject, IMovieScenePlayer&
 		return FGuid();
 	}
 
+	// @todo: Currently we nuke the entire object cache when attempting to find an object's ID to ensure that we do a 
+	// complete lookup from scratch. This is required for UMG as it interchanges content slots without notifying sequencer.
+	Clear();
+
 	TWeakObjectPtr<> ObjectToFind(&InObject);
 
 	// Search all possessables
