VR Editor: New Create Sequence button on the Quick Menu allows you to create a new level sequence in VR. Sequencer UI accessible from this new asset as well as opening sequence assets from the Content Browser. In VR Editing mode, all autokeying is enabled. Settings are restored when leaving VR Editing mode. Meshes are placed in the world to represent keys and frames (VR and non-VR editing mode in this change). Transforming key meshes pushes the transform change back into the correct Sequencer track.

diff --git a/Engine/Plugins/MovieScene/LevelSequenceEditor/Source/LevelSequenceEditor/LevelSequenceEditor.Build.cs b/Engine/Plugins/MovieScene/LevelSequenceEditor/Source/LevelSequenceEditor/LevelSequenceEditor.Build.cs
--- a/Engine/Plugins/MovieScene/LevelSequenceEditor/Source/LevelSequenceEditor/LevelSequenceEditor.Build.cs
+++ b/Engine/Plugins/MovieScene/LevelSequenceEditor/Source/LevelSequenceEditor/LevelSequenceEditor.Build.cs
@@ -34,6 +34,7 @@ public class LevelSequenceEditor : ModuleRules
                 "Slate",
                 "SlateCore",
                 "UnrealEd",
+                "VREditor"
 			}
 		);
 
diff --git a/Engine/Plugins/MovieScene/LevelSequenceEditor/Source/LevelSequenceEditor/Private/LevelSequenceEditorToolkit.cpp b/Engine/Plugins/MovieScene/LevelSequenceEditor/Source/LevelSequenceEditor/Private/LevelSequenceEditorToolkit.cpp
--- a/Engine/Plugins/MovieScene/LevelSequenceEditor/Source/LevelSequenceEditor/Private/LevelSequenceEditorToolkit.cpp
+++ b/Engine/Plugins/MovieScene/LevelSequenceEditor/Source/LevelSequenceEditor/Private/LevelSequenceEditorToolkit.cpp
@@ -46,9 +46,16 @@
 #include "MovieScene.h"
 
 // @todo sequencer: hack: setting defaults for transform tracks
+
 #include "Sections/MovieScene3DTransformSection.h"
 #include "Tracks/MovieScene3DTransformTrack.h"
 
+// To override Sequencer editor behavior for VR Editor 
+#include "IVREditorModule.h"
+#include "VREditorMode.h"
+#include "EditorWorldManager.h"
+
+
 #define LOCTEXT_NAMESPACE "LevelSequenceEditor"
 
 
@@ -155,6 +162,7 @@ FLevelSequenceEditorToolkit::FLevelSequenceEditorToolkit(const TSharedRef<ISlate
 		FAssetEditorExtender::CreateRaw(this, &FLevelSequenceEditorToolkit::HandleMenuExtensibilityGetExtender));
 	SequencerExtenderHandle = SequencerModule.GetAddTrackMenuExtensibilityManager()->GetExtenderDelegates()[NewIndex].GetHandle();
 
+
 	OpenToolkits.Add(this);
 }
 
@@ -261,6 +269,21 @@ void FLevelSequenceEditorToolkit::Initialize(const EToolkitMode::Type Mode, cons
 	}
 
 	FLevelSequenceEditorToolkit::OnOpened().Broadcast(*this);
+	
+	// If we are currently in the VR editor
+	bool bIsInVREditor = IVREditorModule::Get().IsVREditorModeActive();
+	if (bIsInVREditor)
+	{
+		const UWorld* World = Cast<UWorld>(GetLevelSequenceEditorPlaybackContext());
+		UVREditorMode* VRMode = GEditor->GetEditorWorldManager()->GetEditorWorldWrapper(World)->GetVREditorMode();
+		VRMode->OnVREditingModeExit_Handler.BindSP(this, &FLevelSequenceEditorToolkit::HandleVREditorModeExit);
+		VRMode->SaveSequencerSettings(Sequencer->GetKeyAllEnabled(), Sequencer->GetAutoKeyMode());
+		// Override currently set autokey behavior to always autokey all
+		Sequencer->SetAutoKeyMode(EAutoKeyMode::KeyAll);
+		Sequencer->SetKeyAllEnabled(true);
+		// Tell the VR Editor mode that Sequencer has refreshed
+		VRMode->RefreshVREditorSequencer();
+	}
 }
 
 
@@ -684,6 +707,18 @@ void FLevelSequenceEditorToolkit::HandleActorAddedToSequencer(AActor* Actor, con
 }
 
 
+void FLevelSequenceEditorToolkit::HandleVREditorModeExit()
+{
+	const UWorld* World = Cast<UWorld>(GetLevelSequenceEditorPlaybackContext());
+	UVREditorMode* VRMode = GEditor->GetEditorWorldManager()->GetEditorWorldWrapper(World)->GetVREditorMode();
+
+	// Reset auto key settings
+	Sequencer->SetAutoKeyMode(VRMode->GetSavedEditorState().AutoKeyMode);
+	Sequencer->SetKeyAllEnabled(VRMode->GetSavedEditorState().bKeyAllEnabled);
+
+	VRMode->OnVREditingModeExit_Handler.Unbind();
+}
+
 void FLevelSequenceEditorToolkit::HandleMapChanged(class UWorld* NewWorld, EMapChangeType MapChangeType)
 {
 	Sequencer->NotifyMapChanged(NewWorld, MapChangeType);
@@ -751,7 +786,6 @@ void FLevelSequenceEditorToolkit::AddShot(UMovieSceneCinematicShotTrack* ShotTra
 		{
 			CameraGuid = GetSequencer()->CreateBinding(*NewCamera, NewCamera->GetActorLabel());
 		}
-
 		NewCamera->SetActorLocation( GCurrentLevelEditingViewportClient->GetViewLocation(), false );
 		NewCamera->SetActorRotation( GCurrentLevelEditingViewportClient->GetViewRotation() );
 		//pNewCamera->CameraComponent->FieldOfView = ViewportClient->ViewFOV; //@todo set the focal length from this field of view
diff --git a/Engine/Plugins/MovieScene/LevelSequenceEditor/Source/LevelSequenceEditor/Private/LevelSequenceEditorToolkit.h b/Engine/Plugins/MovieScene/LevelSequenceEditor/Source/LevelSequenceEditor/Private/LevelSequenceEditorToolkit.h
--- a/Engine/Plugins/MovieScene/LevelSequenceEditor/Source/LevelSequenceEditor/Private/LevelSequenceEditorToolkit.h
+++ b/Engine/Plugins/MovieScene/LevelSequenceEditor/Source/LevelSequenceEditor/Private/LevelSequenceEditorToolkit.h
@@ -138,6 +138,9 @@ private:
 	/** Callback for actor added to sequencer. */
 	void HandleActorAddedToSequencer(AActor* Actor, const FGuid Binding);
 
+	/** Callback for VR Editor mode exiting */
+	void HandleVREditorModeExit();
+
 private:
 
 	/** Level sequence for our edit operation. */
diff --git a/Engine/Source/Developer/AssetTools/Private/AssetTools.cpp b/Engine/Source/Developer/AssetTools/Private/AssetTools.cpp
--- a/Engine/Source/Developer/AssetTools/Private/AssetTools.cpp
+++ b/Engine/Source/Developer/AssetTools/Private/AssetTools.cpp
@@ -716,7 +716,7 @@ void FAssetTools::CreateUniqueAssetName(const FString& InBasePackageName, const
 	const FString BaseAssetNameWithSuffix = FPackageName::GetLongPackageAssetName(SanitizedBasePackageName) + InSuffix;
 	const FString SanitizedBaseAssetName = ObjectTools::SanitizeObjectName(BaseAssetNameWithSuffix);
 
-	int32 IntSuffix = 1;
+	int32 IntSuffix = 0;
 	bool bObjectExists = false;
 
 	int32 CharIndex = SanitizedBaseAssetName.Len() - 1;
@@ -744,7 +744,7 @@ void FAssetTools::CreateUniqueAssetName(const FString& InBasePackageName, const
 	do
 	{
 		bObjectExists = false;
-		if ( IntSuffix <= 1 )
+		if ( IntSuffix < 1 )
 		{
 			OutAssetName = SanitizedBaseAssetName;
 		}
diff --git a/Engine/Source/Editor/LevelEditor/Private/SLevelEditor.cpp b/Engine/Source/Editor/LevelEditor/Private/SLevelEditor.cpp
--- a/Engine/Source/Editor/LevelEditor/Private/SLevelEditor.cpp
+++ b/Engine/Source/Editor/LevelEditor/Private/SLevelEditor.cpp
@@ -521,13 +521,16 @@ void SLevelEditor::AttachSequencer( TSharedPtr<SWidget> SequencerWidget, TShared
 			SequencerAssetEditor.Pin()->CloseWindow();
 		}
 
-		TSharedRef<SDockTab> Tab = InvokeTab("Sequencer");
-
-		if(SequencerWidget.IsValid() && NewSequencerAssetEditor.IsValid())
+		TSharedRef<SDockTab> Tab = SNew(SDockTab);
+		if(!FGlobalTabmanager::Get()->OnOverrideDockableAreaRestore_Handler.IsBound())
 		{
+			Tab = InvokeTab("Sequencer");
 			Tab->SetOnTabClosed(SDockTab::FOnTabClosedCallback::CreateStatic(&Local::OnSequencerClosed, TWeakPtr<IAssetEditorInstance>(NewSequencerAssetEditor)));
+		}
+		if(SequencerWidget.IsValid() && NewSequencerAssetEditor.IsValid())
+		{
 			Tab->SetContent(SequencerWidget.ToSharedRef());
-
+			SequencerWidgetPtr = SequencerWidget;
 			SequencerAssetEditor = NewSequencerAssetEditor;
 		}
 		else
@@ -739,8 +742,8 @@ TSharedRef<SDockTab> SLevelEditor::SpawnLevelEditorTab( const FSpawnTabArgs& Arg
 		{
 			// @todo sequencer: remove when world-centric mode is added
 			return SNew(SDockTab)
-				.Icon( FSlateStyleRegistry::FindSlateStyle("LevelSequenceEditorStyle")->GetBrush("LevelSequenceEditor.Tabs.Sequencer") )
-				.Label( NSLOCTEXT("Sequencer", "SequencerMainTitle", "Sequencer") )
+				.Icon(FSlateStyleRegistry::FindSlateStyle("LevelSequenceEditorStyle")->GetBrush("LevelSequenceEditor.Tabs.Sequencer"))
+				.Label(NSLOCTEXT("Sequencer", "SequencerMainTitle", "Sequencer"))
 				[
 					SNullWidget::NullWidget
 				];
@@ -1524,3 +1527,8 @@ TSharedRef<SWidget> SLevelEditor::CreateToolBox()
 
 	return NewToolBox;
 }
+
+TSharedPtr<SWidget> SLevelEditor::GetSequencerWidgetForLevelEditor()
+{
+	return SequencerWidgetPtr.Pin();
+}
\ No newline at end of file
diff --git a/Engine/Source/Editor/LevelEditor/Private/SLevelEditor.h b/Engine/Source/Editor/LevelEditor/Private/SLevelEditor.h
--- a/Engine/Source/Editor/LevelEditor/Private/SLevelEditor.h
+++ b/Engine/Source/Editor/LevelEditor/Private/SLevelEditor.h
@@ -6,7 +6,6 @@
 
 #pragma once
 
-
 #include "CoreMinimal.h"
 #include "Input/Reply.h"
 #include "Widgets/SWidget.h"
@@ -77,6 +76,7 @@ public:
 	virtual TSharedPtr< class FAssetThumbnailPool > GetThumbnailPool() const override;
 	virtual void AppendCommands( const TSharedRef<FUICommandList>& InCommandsToAppend ) override;
 	virtual void AddStandaloneLevelViewport( const TSharedRef<SLevelViewport>& LevelViewport ) override;
+	virtual TSharedPtr<SWidget> GetSequencerWidgetForLevelEditor() override;
 
 	/**
 	 * Given a tab ID, summons a new tab in the position saved in the current layout, or in a default position.
@@ -246,6 +246,9 @@ private:
 
 	/** Attached sequencer asset editor */
 	TWeakPtr<IAssetEditorInstance> SequencerAssetEditor;
+
+	/** Weak pointer to the level editor's Sequencer widget */
+	TWeakPtr<SWidget> SequencerWidgetPtr;
 };
 
 
diff --git a/Engine/Source/Editor/LevelEditor/Public/ILevelEditor.h b/Engine/Source/Editor/LevelEditor/Public/ILevelEditor.h
--- a/Engine/Source/Editor/LevelEditor/Public/ILevelEditor.h
+++ b/Engine/Source/Editor/LevelEditor/Public/ILevelEditor.h
@@ -63,6 +63,9 @@ public:
 	/** Spawns a level editor ToolBox widget (aka. "Modes") */
 	virtual TSharedRef<SWidget> CreateToolBox() = 0;
 
+	/** Return the level editor's Sequencer widget */
+	virtual TSharedPtr<SWidget> GetSequencerWidgetForLevelEditor() = 0;
+
 };
 
 
diff --git a/Engine/Source/Editor/MovieSceneTools/Private/TrackEditors/TransformTrackEditor.cpp b/Engine/Source/Editor/MovieSceneTools/Private/TrackEditors/TransformTrackEditor.cpp
--- a/Engine/Source/Editor/MovieSceneTools/Private/TrackEditors/TransformTrackEditor.cpp
+++ b/Engine/Source/Editor/MovieSceneTools/Private/TrackEditors/TransformTrackEditor.cpp
@@ -408,11 +408,12 @@ TSharedRef<ISequencerSection> F3DTransformTrackEditor::MakeSectionInterface(UMov
 void F3DTransformTrackEditor::OnPreTransformChanged( UObject& InObject )
 {
 	float AutoKeyTime = GetTimeForKey();
-
-	if( IsAllowedToAutoKey() )
+	AActor* Actor = Cast<AActor>(&InObject);
+	// If Sequencer is allowed to autokey and we are clicking on an Actor that can be autokeyed
+	if( IsAllowedToAutoKey() && (Actor && !Actor->IsEditorOnly() ) )
 	{
 		USceneComponent* SceneComponent = nullptr;
-		AActor* Actor = Cast<AActor>( &InObject );
+		
 		if( Actor && Actor->GetRootComponent() )
 		{
 			SceneComponent = Actor->GetRootComponent();
@@ -428,8 +429,8 @@ void F3DTransformTrackEditor::OnPreTransformChanged( UObject& InObject )
 			// Cache off the existing transform so we can detect which components have changed
 			// and keys only when something has changed
 			FTransformData Transform( SceneComponent );
-
-			ObjectToExistingTransform.Add( &InObject, Transform );
+			
+			ObjectToExistingTransform.Add(&InObject, Transform);
 		}
 	}
 }
@@ -456,7 +457,8 @@ void F3DTransformTrackEditor::OnTransformChanged( UObject& InObject )
 	USceneComponent* SceneComponentThatChanged = nullptr;
 	GetActorAndSceneComponentFromObject(&InObject, Actor, SceneComponentThatChanged);
 
-	if( SceneComponentThatChanged != nullptr && Actor != nullptr )
+	// If the Actor that just finished transforming doesn't have autokey disabled
+	if( SceneComponentThatChanged != nullptr && (Actor && !Actor->IsEditorOnly()))
 	{
 		// Find an existing transform if possible.  If one exists we will compare against the new one to decide what components of the transform need keys
 		FTransformData ExistingTransform = ObjectToExistingTransform.FindRef( &InObject );
diff --git a/Engine/Source/Editor/Sequencer/Private/MovieSceneTrackEditor.cpp b/Engine/Source/Editor/Sequencer/Private/MovieSceneTrackEditor.cpp
new file mode 100644
--- /dev/null
+++ b/Engine/Source/Editor/Sequencer/Private/MovieSceneTrackEditor.cpp
@@ -0,0 +1,44 @@
+// Copyright 1998-2017 Epic Games, Inc. All Rights Reserved.
+
+#include "MovieSceneTrackEditor.h"
+#include "Sequencer.h"
+
+void FMovieSceneTrackEditor::AnimatablePropertyChanged(FOnKeyProperty OnKeyProperty)
+{
+	check(OnKeyProperty.IsBound());
+
+	// Get the movie scene we want to autokey
+	UMovieSceneSequence* MovieSceneSequence = GetMovieSceneSequence();
+	if (MovieSceneSequence)
+	{
+		float KeyTime = GetTimeForKey();
+
+		if (!Sequencer.Pin()->IsRecordingLive())
+		{
+			// @todo Sequencer - The sequencer probably should have taken care of this
+			MovieSceneSequence->SetFlags(RF_Transactional);
+
+			// Create a transaction record because we are about to add keys
+			const bool bShouldActuallyTransact = !Sequencer.Pin()->IsRecordingLive();		// Don't transact if we're recording in a PIE world.  That type of keyframe capture cannot be undone.
+			FScopedTransaction AutoKeyTransaction(NSLOCTEXT("AnimatablePropertyTool", "PropertyChanged", "Animatable Property Changed"), bShouldActuallyTransact);
+
+			if (OnKeyProperty.Execute(KeyTime))
+			{
+				// TODO: This should pass an appropriate change type here instead of always passing structure changed since most
+				// changes will be value changes.
+				Sequencer.Pin()->NotifyMovieSceneDataChanged(EMovieSceneDataChangeType::MovieSceneStructureItemAdded);
+			}
+
+			UpdatePlaybackRange();
+	
+			TSharedPtr<FSequencer> SequencerToUpdate = StaticCastSharedPtr<FSequencer>(Sequencer.Pin());
+			if (SequencerToUpdate.IsValid())
+			{
+				SequencerToUpdate->NotifyMovieSceneDataChanged(EMovieSceneDataChangeType::Unknown);
+				SequencerToUpdate->SynchronizeSequencerSelectionWithExternalSelection();
+				SequencerToUpdate->SynchronizeExternalSelectionWithSequencerSelection();
+				
+			}
+		}
+	}
+}
\ No newline at end of file
diff --git a/Engine/Source/Editor/Sequencer/Private/Sequencer.cpp b/Engine/Source/Editor/Sequencer/Private/Sequencer.cpp
--- a/Engine/Source/Editor/Sequencer/Private/Sequencer.cpp
+++ b/Engine/Source/Editor/Sequencer/Private/Sequencer.cpp
@@ -4286,6 +4286,13 @@ void FSequencer::OnSelectedOutlinerNodesChanged()
 		SynchronizeExternalSelectionWithSequencerSelection();
 		bUpdatingExternalSelection = false;
 	}
+	FSequencerEdMode* SequencerEdMode = (FSequencerEdMode*)(GLevelEditorModeTools().GetActiveMode(FSequencerEdMode::EM_SequencerMode));
+	AActor* NewlySelectedActor = GEditor->GetSelectedActors()->GetTop<AActor>();
+	// If we selected an Actor or a node for an Actor that is a potential autokey candidate, clean up any existing mesh trails
+	if (NewlySelectedActor && !NewlySelectedActor->IsEditorOnly())
+	{
+		SequencerEdMode->CleanUpMeshTrails();
+	}
 	OnSelectionChangedObjectGuidsDelegate.Broadcast(Selection.GetBoundObjectsGuids());
 }
 
diff --git a/Engine/Source/Editor/Sequencer/Private/SequencerEdMode.cpp b/Engine/Source/Editor/Sequencer/Private/SequencerEdMode.cpp
--- a/Engine/Source/Editor/Sequencer/Private/SequencerEdMode.cpp
+++ b/Engine/Source/Editor/Sequencer/Private/SequencerEdMode.cpp
@@ -14,6 +14,10 @@
 #include "Tracks/MovieScene3DTransformTrack.h"
 #include "Sections/MovieScene3DTransformSection.h"
 #include "SubtitleManager.h"
+#include "SequencerMeshTrail.h"
+#include "SequencerKeyActor.h"
+#include "ViewportWorldInteraction.h"
+#include "EditorWorldManager.h"
 
 const FEditorModeID FSequencerEdMode::EM_SequencerMode(TEXT("EM_SequencerMode"));
 
@@ -23,6 +27,9 @@ FSequencerEdMode::FSequencerEdMode()
 
 	Tools.Add( SequencerEdModeTool );
 	SetCurrentTool( SequencerEdModeTool );
+
+	// todo vreditor: make this a setting
+	bDrawMeshTrails = true;
 }
 
 FSequencerEdMode::~FSequencerEdMode()
@@ -81,9 +88,16 @@ void FSequencerEdMode::Render(const FSceneView* View, FViewport* Viewport, FPrim
 	FEdMode::Render(View, Viewport, PDI);
 
 #if WITH_EDITORONLY_DATA
-	if (View->Family->EngineShowFlags.Splines)
+	// Draw spline trails using the PDI
+	if (View->Family->EngineShowFlags.Splines && !bDrawMeshTrails)
+	{
+		DrawTracks3D(PDI);
+	}
+	// Draw mesh trails (doesn't use the PDI)
+	else 	
 	{
-		DrawTracks3D(View, PDI);
+		PDI = nullptr;
+		DrawTracks3D(PDI);
 	}
 #endif
 }
@@ -106,6 +120,15 @@ void FSequencerEdMode::DrawHUD(FEditorViewportClient* ViewportClient,FViewport*
 	}
 }
 
+void FSequencerEdMode::AddReferencedObjects(FReferenceCollector& Collector)
+{
+	for (auto MeshTrail : MeshTrails)
+	{
+		Collector.AddReferencedObject(MeshTrail.Get<0>());
+		Collector.AddReferencedObject(MeshTrail.Get<1>());
+	}
+}
+
 void FSequencerEdMode::OnKeySelected(FViewport* Viewport, HMovieSceneKeyProxy* KeyProxy)
 {
 	TSharedPtr<FSequencer> Sequencer = SequencerPtr.Pin();
@@ -123,6 +146,36 @@ void FSequencerEdMode::OnKeySelected(FViewport* Viewport, HMovieSceneKeyProxy* K
 	Sequencer->SelectTrackKeys(KeyProxy->MovieSceneSection, KeyProxy->Time, bShiftDown, bCtrlDown);
 }
 
+void FSequencerEdMode::DrawMeshTransformTrailFromKey(const class ASequencerKeyActor* KeyActor)
+{
+	ASequencerMeshTrail* Trail = Cast<ASequencerMeshTrail>(KeyActor->GetOwner());
+	TTuple<UMovieScene3DTransformTrack*, ASequencerMeshTrail*>* TrailPtr = MeshTrails.FindByPredicate([Trail](const TTuple<UMovieScene3DTransformTrack*, ASequencerMeshTrail*> InTrail)
+	{
+		return Trail == InTrail.Get<1>();
+	});
+	if(TrailPtr)
+	{
+		// From the key, get the mesh trail, and then the track associated with that mesh trail
+		UMovieScene3DTransformTrack* Track = TrailPtr->Get<0>();
+		// Draw a mesh trail for the key's associated actor
+		TArray<TWeakObjectPtr<UObject>> KeyObjects;
+		AActor* TrailActor = KeyActor->GetAssociatedActor();
+		KeyObjects.Add(TrailActor);
+		FPrimitiveDrawInterface* PDI = nullptr;
+		DrawTransformTrack(PDI, Track, KeyObjects, true);
+	}
+}
+
+void FSequencerEdMode::CleanUpMeshTrails()
+{
+	// Clean up any existing trails
+	for (auto& MeshTrail : MeshTrails)
+	{
+		MeshTrail.Get<1>()->Cleanup();
+	}
+	MeshTrails.Empty();
+}
+
 namespace SequencerEdMode_Draw3D
 {
 static const FColor	KeySelectedColor(255,128,0);
@@ -149,10 +202,28 @@ void GetLocationAtTime(UMovieScene3DTransformSection* TransformSection, float Ke
 	TransformSection->EvalRotation(KeyTime, KeyRot);
 }
 
-void DrawTransformTrack(const FSceneView* View, FPrimitiveDrawInterface* PDI, UMovieScene3DTransformTrack* TransformTrack, const TArray<TWeakObjectPtr<UObject>>& BoundObjects, const bool& bIsSelected)
+void FSequencerEdMode::DrawTransformTrack(FPrimitiveDrawInterface* PDI, UMovieScene3DTransformTrack* TransformTrack, const TArray<TWeakObjectPtr<UObject>>& BoundObjects, const bool& bIsSelected)
 {
-	const bool bHitTesting = PDI->IsHitTesting();
 	FLinearColor TrackColor = FLinearColor::Yellow; //@todo - customizable per track
+	bool bHitTesting = true;
+	if( PDI != nullptr )
+	{
+		bHitTesting = PDI->IsHitTesting();
+	}
+	
+	ASequencerMeshTrail* TrailActor = nullptr;
+	// Get the Trail Actor associated with this track if we are drawing mesh trails
+	if (bDrawMeshTrails)
+	{
+		TTuple<UMovieScene3DTransformTrack*, ASequencerMeshTrail*>* TrailPtr = MeshTrails.FindByPredicate([TransformTrack](const TTuple<UMovieScene3DTransformTrack*, ASequencerMeshTrail*> InTrail)
+		{
+			return InTrail.Get<0>() == TransformTrack;
+		});
+		if (TrailPtr != nullptr)
+		{
+			TrailActor = TrailPtr->Get<1>();
+		}
+	}
 
 	for (int32 SectionIndex = 0; SectionIndex < TransformTrack->GetAllSections().Num(); ++SectionIndex)
 	{
@@ -241,11 +312,14 @@ void DrawTransformTrack(const FSceneView* View, FPrimitiveDrawInterface* PDI, UM
 
 					FVector OldPos_G = RefTM.TransformPosition(OldPos);
 					FVector NewKeyPos_G = RefTM.TransformPosition(NewKeyPos);
-
+					
 					// For constant interpolation - don't draw ticks - just draw dotted line.
 					if (bIsConstantKey)
 					{
-						DrawDashedLine(PDI, OldPos_G, NewKeyPos_G, TrackColor, 20, SDPG_Foreground);
+						if(PDI != nullptr)
+						{
+							DrawDashedLine(PDI, OldPos_G, NewKeyPos_G, TrackColor, 20, SDPG_Foreground);
+						}
 					}
 					else
 					{
@@ -260,12 +334,22 @@ void DrawTransformTrack(const FSceneView* View, FPrimitiveDrawInterface* PDI, UM
 
 							FVector NewPos_G = RefTM.TransformPosition(NewPos);
 
-							PDI->DrawLine(OldPos_G, NewPos_G, TrackColor, SDPG_Foreground);
-
+							if (PDI != nullptr)
+							{
+								PDI->DrawLine(OldPos_G, NewPos_G, TrackColor, SDPG_Foreground);
+							}
+							// Drawing frames
 							// Don't draw point for last one - its the keypoint drawn above.
 							if (j != NumSteps)
 							{
-								PDI->DrawPoint(NewPos_G, TrackColor, 3.f, SDPG_Foreground);
+								if (PDI != nullptr)
+								{
+									PDI->DrawPoint(NewPos_G, TrackColor, 3.f, SDPG_Foreground);
+								}
+								else if (TrailActor != nullptr)
+								{
+									TrailActor->AddFrameMeshComponent(NewTime, FTransform::FTransform(NewRot, NewPos_G, FVector::FVector(0.25f)));
+								}
 							}
 							OldTime = NewTime;
 							OldPos_G = NewPos_G;
@@ -307,12 +391,22 @@ void DrawTransformTrack(const FSceneView* View, FPrimitiveDrawInterface* PDI, UM
 
 				if (bHitTesting) 
 				{
-					PDI->SetHitProxy( new HMovieSceneKeyProxy(TransformTrack, TransformSection, NewKeyTime) );
+					if (PDI != nullptr)
+					{
+						PDI->SetHitProxy(new HMovieSceneKeyProxy(TransformTrack, TransformSection, NewKeyTime));
+					}
 				}
 
 				FVector NewKeyPos_G = RefTM.TransformPosition(NewKeyPos);
-
-				PDI->DrawPoint(NewKeyPos_G, KeyColor, 6.f, SDPG_Foreground);
+				// Drawing keys
+				if (PDI != nullptr)
+				{
+					PDI->DrawPoint(NewKeyPos_G, KeyColor, 6.f, SDPG_Foreground);
+				}
+				else if (TrailActor != nullptr)
+				{
+					TrailActor->AddKeyMeshActor(NewKeyTime, FTransform::FTransform(NewKeyRot, NewKeyPos_G, FVector::FVector(0.5f)), TransformSection);
+				}
 
 				//@todo
 				// If desired, draw directional arrow at each keyframe.
@@ -321,7 +415,7 @@ void DrawTransformTrack(const FSceneView* View, FPrimitiveDrawInterface* PDI, UM
 				//	FRotationTranslationMatrix ArrowToWorld(NewKeyRot,NewKeyPos);
 				//	DrawDirectionalArrow(PDI, FScaleMatrix(FVector(16.f,16.f,16.f)) * ArrowToWorld, KeyColor, 3.f, 1.f, SDPG_Foreground );
 				//}
-				if (bHitTesting) 
+				if (bHitTesting && PDI != nullptr) 
 				{
 					PDI->SetHitProxy( NULL );
 				}
@@ -381,7 +475,7 @@ void DrawTransformTrack(const FSceneView* View, FPrimitiveDrawInterface* PDI, UM
 }
 
 
-void FSequencerEdMode::DrawTracks3D(const FSceneView* View, FPrimitiveDrawInterface* PDI)
+void FSequencerEdMode::DrawTracks3D(FPrimitiveDrawInterface* PDI)
 {
 	TSet<TSharedRef<FSequencerDisplayNode> > ObjectBindingNodes;
 
@@ -420,9 +514,7 @@ void FSequencerEdMode::DrawTracks3D(const FSceneView* View, FPrimitiveDrawInterf
 		}
 	}
 
-
 	// Gather up the transform track nodes from the object binding nodes
-
 	for (auto ObjectBindingNode : ObjectBindingNodesSelectionMap)
 	{
 		TSet<TSharedRef<FSequencerDisplayNode> > AllNodes;
@@ -445,7 +537,22 @@ void FSequencerEdMode::DrawTracks3D(const FSceneView* View, FPrimitiveDrawInterf
 				UMovieScene3DTransformTrack* TransformTrack = Cast<UMovieScene3DTransformTrack>(TrackNodeTrack);
 				if (TransformTrack != nullptr)
 				{
-					DrawTransformTrack(View, PDI, TransformTrack, BoundObjects, ObjectBindingNode.Value);
+					// If we are drawing mesh trails but we haven't made one for this track yet
+					if (bDrawMeshTrails)
+					{
+						TTuple<UMovieScene3DTransformTrack*, ASequencerMeshTrail*>* TrailPtr = MeshTrails.FindByPredicate([TransformTrack](const TTuple<UMovieScene3DTransformTrack*, ASequencerMeshTrail*> InTrail)
+						{
+							return InTrail.Get<0>() == TransformTrack;
+						});
+						if (TrailPtr == nullptr)
+						{
+							UViewportWorldInteraction* ViewportWorldInteraction = GEditor->GetEditorWorldManager()->GetEditorWorldWrapper(GEditor->GetActiveViewport()->GetClient()->GetWorld())->GetViewportWorldInteraction();
+							ASequencerMeshTrail* TrailActor = ViewportWorldInteraction->SpawnTransientSceneActor<ASequencerMeshTrail>(TEXT("SequencerMeshTrail"), true);
+							TTuple<UMovieScene3DTransformTrack*, ASequencerMeshTrail*> MeshTrail = MakeTuple(TransformTrack, TrailActor);
+							MeshTrails.Add(MeshTrail);
+						}
+					}
+					DrawTransformTrack(PDI, TransformTrack, BoundObjects, ObjectBindingNode.Value);
 				}
 			}
 		}
diff --git a/Engine/Source/Editor/Sequencer/Private/SequencerEdMode.h b/Engine/Source/Editor/Sequencer/Private/SequencerEdMode.h
--- a/Engine/Source/Editor/Sequencer/Private/SequencerEdMode.h
+++ b/Engine/Source/Editor/Sequencer/Private/SequencerEdMode.h
@@ -34,16 +34,31 @@ public:
 	virtual bool InputKey( FEditorViewportClient* ViewportClient, FViewport* Viewport, FKey Key, EInputEvent Event ) override;
 	virtual void Render(const FSceneView* View,FViewport* Viewport,FPrimitiveDrawInterface* PDI) override;
 	virtual void DrawHUD(FEditorViewportClient* ViewportClient,FViewport* Viewport,const FSceneView* View,FCanvas* Canvas) override;
+	virtual void AddReferencedObjects(FReferenceCollector& Collector) override;
 
 	void SetSequencer(const TWeakPtr<FSequencer>& InSequencer) { SequencerPtr = InSequencer; }
 
 	void OnKeySelected(FViewport* Viewport, HMovieSceneKeyProxy* KeyProxy);
 
+	/** Draw a single mesh transform track, given a key that is on that track */
+	void DrawMeshTransformTrailFromKey(const class ASequencerKeyActor* KeyActor);
+
+	/** Clean up any mesh trails and their associated key actors */
+	void CleanUpMeshTrails();
+
 protected:
-	void DrawTracks3D(const FSceneView* View, FPrimitiveDrawInterface* PDI);
+	void DrawTracks3D(FPrimitiveDrawInterface* PDI);
+
+	void DrawTransformTrack(FPrimitiveDrawInterface* PDI, class UMovieScene3DTransformTrack* TransformTrack, const TArray<TWeakObjectPtr<UObject>>& BoundObjects, const bool& bIsSelected);
 
 private:
 	TWeakPtr<FSequencer> SequencerPtr;
+
+	/**The map of any transform tracks with mesh trails, and their mesh trail representations */
+	TArray<TTuple<class UMovieScene3DTransformTrack*,class ASequencerMeshTrail*>> MeshTrails;
+
+	/** True when drawing mesh trails instead of debug trails */
+	bool bDrawMeshTrails;
 };
 
 /**
diff --git a/Engine/Source/Editor/Sequencer/Private/SequencerKeyActor.cpp b/Engine/Source/Editor/Sequencer/Private/SequencerKeyActor.cpp
new file mode 100644
--- /dev/null
+++ b/Engine/Source/Editor/Sequencer/Private/SequencerKeyActor.cpp
@@ -0,0 +1,72 @@
+// Copyright 1998-2016 Epic Games, Inc. All Rights Reserved.
+
+#include "SequencerKeyActor.h"
+#include "Sections/MovieScene3DTransformSection.h"
+#include "SequencerEdMode.h"
+#include "Materials/MaterialInstance.h" 
+#include "Engine/Selection.h"
+#include "EditorModeTools.h"
+#include "EditorModeManager.h"
+#include "ModuleManager.h"
+
+ASequencerKeyActor::ASequencerKeyActor()
+	: Super()
+{
+	UStaticMesh* KeyEditorMesh = LoadObject<UStaticMesh>(nullptr, TEXT("/Engine/BasicShapes/Sphere"));
+	check(KeyEditorMesh != nullptr);
+	UMaterialInstance* KeyEditorMaterial = LoadObject<UMaterialInstance>(nullptr, TEXT("/Engine/VREditor/LaserPointer/LaserPointerMaterial_Inst"));
+	check(KeyEditorMaterial != nullptr);
+	GetStaticMeshComponent()->SetStaticMesh(KeyEditorMesh);
+	GetStaticMeshComponent()->SetMaterial(0, KeyEditorMaterial);
+}
+
+
+void ASequencerKeyActor::PostEditMove(bool bFinished)
+{
+	// Push the key's transform to the Sequencer track
+	PropagateKeyChange();
+	Super::PostEditMove(bFinished);
+}
+
+void ASequencerKeyActor::SetKeyData(class UMovieScene3DTransformSection* NewTrackSection, float NewKeyTime)
+{
+	TrackSection = NewTrackSection;
+	KeyTime = NewKeyTime;
+	// Associate the currently selected actor with this key
+	AssociatedActor = GEditor->GetSelectedActors()->GetTop<AActor>();
+}
+
+
+void ASequencerKeyActor::PropagateKeyChange()
+{
+	if (TrackSection != nullptr)
+	{
+		// Mark the track section as dirty
+		TrackSection->Modify();
+
+		// Update the translation keys
+		FRichCurve& TransXCurve = TrackSection->GetTranslationCurve(EAxis::X);
+		FRichCurve& TransYCurve = TrackSection->GetTranslationCurve(EAxis::Y);
+		FRichCurve& TransZCurve = TrackSection->GetTranslationCurve(EAxis::Z);
+		TransXCurve.UpdateOrAddKey(KeyTime, GetActorTransform().GetLocation().X);
+		TransYCurve.UpdateOrAddKey(KeyTime, GetActorTransform().GetLocation().Y);
+		TransZCurve.UpdateOrAddKey(KeyTime, GetActorTransform().GetLocation().Z);
+
+		// Update the rotation keys
+		FRichCurve& RotXCurve = TrackSection->GetRotationCurve(EAxis::X);
+		FRichCurve& RotYCurve = TrackSection->GetRotationCurve(EAxis::Y);
+		FRichCurve& RotZCurve = TrackSection->GetRotationCurve(EAxis::Z);
+		RotXCurve.UpdateOrAddKey(KeyTime, GetActorTransform().GetRotation().X);
+		RotYCurve.UpdateOrAddKey(KeyTime, GetActorTransform().GetRotation().Y);
+		RotZCurve.UpdateOrAddKey(KeyTime, GetActorTransform().GetRotation().Z);
+
+		// Draw a single transform track based on the data from this key
+		FEditorViewportClient* ViewportClient = StaticCast<FEditorViewportClient*>(GEditor->GetActiveViewport()->GetClient());
+		if (ViewportClient != nullptr)
+		{
+			FSequencerEdMode* SequencerEdMode = (FSequencerEdMode*)(ViewportClient->GetModeTools()->GetActiveMode(FSequencerEdMode::EM_SequencerMode));
+			SequencerEdMode->DrawMeshTransformTrailFromKey(this);
+		}
+	}
+}
+
diff --git a/Engine/Source/Editor/Sequencer/Private/SequencerMeshTrail.cpp b/Engine/Source/Editor/Sequencer/Private/SequencerMeshTrail.cpp
new file mode 100644
--- /dev/null
+++ b/Engine/Source/Editor/Sequencer/Private/SequencerMeshTrail.cpp
@@ -0,0 +1,90 @@
+// Copyright 1998-2017 Epic Games, Inc. All Rights Reserved.
+
+#include "SequencerMeshTrail.h"
+#include "SequencerKeyActor.h"
+#include "ViewportWorldInteraction.h"
+#include "EditorWorldManager.h"
+#include "Materials/MaterialInstance.h" 
+
+ASequencerMeshTrail::ASequencerMeshTrail()
+	: Super()
+{
+	RootComponent = CreateDefaultSubobject<USceneComponent>(TEXT("RootComponent"));
+}
+
+void ASequencerMeshTrail::Cleanup()
+{
+	// Destroy all the key actors this trail created
+	for (auto& KeyMesh : KeyMeshActors)
+	{
+		KeyMesh.Get<1>()->Destroy();
+	}
+	this->Destroy();
+}
+
+void ASequencerMeshTrail::AddKeyMeshActor(float KeyTime, const FTransform KeyTransform, class UMovieScene3DTransformSection* TrackSection)
+{
+	ASequencerKeyActor* KeyMeshActor = nullptr;
+	UStaticMesh* KeyEditorMesh = nullptr;
+	UMaterialInstance* KeyEditorMaterial = nullptr;
+
+	TTuple<float, ASequencerKeyActor*>* KeyPtr = KeyMeshActors.FindByPredicate([KeyTime](const TTuple<float, ASequencerKeyActor*> InKey)
+	{
+		return FMath::IsNearlyEqual(KeyTime, InKey.Get<0>());
+	});
+
+	// If we don't currently have an actor for this time, create one
+	if (KeyPtr == nullptr)
+	{
+		UViewportWorldInteraction* ViewportWorldInteraction = GEditor->GetEditorWorldManager()->GetEditorWorldWrapper(GEditor->GetActiveViewport()->GetClient()->GetWorld())->GetViewportWorldInteraction();
+		KeyMeshActor = ViewportWorldInteraction->SpawnTransientSceneActor<ASequencerKeyActor>(TEXT("KeyMesh"), false);
+		KeyMeshActor->SetMobility(EComponentMobility::Movable);
+		KeyMeshActor->SetActorTransform(KeyTransform);
+		KeyMeshActor->SetKeyData(TrackSection, KeyTime);
+		KeyMeshActor->SetOwner(this);
+		TTuple<float, ASequencerKeyActor*> NewKey = MakeTuple(KeyTime, KeyMeshActor);
+		KeyMeshActors.Add(NewKey);
+	}
+	else
+	{
+		// Just update the transform
+		KeyPtr->Get<1>()->SetActorTransform(KeyTransform);
+	}
+}
+
+void ASequencerMeshTrail::AddFrameMeshComponent(const float FrameTime, const FTransform FrameTransform)
+{
+	UStaticMeshComponent* FrameMeshComponent = nullptr;
+	UStaticMesh* KeyEditorMesh = nullptr;
+	UMaterialInterface* KeyEditorMaterial = nullptr;
+
+	TTuple<float, UStaticMeshComponent*>* FramePtr = FrameMeshComponents.FindByPredicate([FrameTime](const TTuple<float, UStaticMeshComponent*> InFrame)
+	{
+		return FMath::IsNearlyEqual(FrameTime, InFrame.Get<0>());
+	});
+
+	// If we don't currently have a component for this time, create one
+	if (FramePtr == nullptr)
+	{
+		FrameMeshComponent = NewObject<UStaticMeshComponent>(this);
+		this->AddOwnedComponent(FrameMeshComponent);
+		FrameMeshComponent->RegisterComponent();
+
+		KeyEditorMesh = LoadObject<UStaticMesh>(nullptr, TEXT("/Engine/BasicShapes/Sphere"));
+		check(KeyEditorMesh != nullptr);
+		KeyEditorMaterial = LoadObject<UMaterialInstance>(nullptr, TEXT("/Engine/VREditor/LaserPointer/LaserPointerMaterial_Inst"));
+		check(KeyEditorMaterial != nullptr);
+		FrameMeshComponent->SetStaticMesh(KeyEditorMesh);
+		FrameMeshComponent->SetMaterial(0, KeyEditorMaterial);
+		FrameMeshComponent->SetWorldTransform(FrameTransform);
+		FrameMeshComponent->SetMobility(EComponentMobility::Movable);
+		FrameMeshComponent->bSelectable = false;
+		TTuple<float, UStaticMeshComponent*> NewFrame = MakeTuple(FrameTime, FrameMeshComponent);
+		FrameMeshComponents.Add(NewFrame);
+	}
+	else
+	{
+		// Just update the transform
+		FramePtr->Get<1>()->SetWorldTransform(FrameTransform);
+	}
+}
diff --git a/Engine/Source/Editor/Sequencer/Private/SequencerMeshTrail.h b/Engine/Source/Editor/Sequencer/Private/SequencerMeshTrail.h
new file mode 100644
--- /dev/null
+++ b/Engine/Source/Editor/Sequencer/Private/SequencerMeshTrail.h
@@ -0,0 +1,37 @@
+// Copyright 1998-2017 Epic Games, Inc. All Rights Reserved.
+
+#pragma once
+#include "SequencerMeshTrail.generated.h"
+
+UCLASS()
+class ASequencerMeshTrail : public AActor
+{
+	GENERATED_BODY()
+
+public:
+	ASequencerMeshTrail();
+
+	/** Clean up the key mesh Actors and then destroy the trail itself */
+	void Cleanup();
+
+	/** Add a SequencerKeyMesh Actor associated with the given track section at the KeyTransform. KeyTime is used as a key to the KeyMeshes TMap */
+	void AddKeyMeshActor(float KeyTime, const FTransform KeyTransform, class UMovieScene3DTransformSection* TrackSection);
+
+	/** Add a static mesh component at the KeyTransform. KeyTime is used as a key to the FrameMeshes TMap */
+	void AddFrameMeshComponent(const float FrameTime, const FTransform FrameTransform);
+
+	virtual bool IsEditorOnly() const final
+	{
+		return true;
+	}
+
+private:
+
+	/** Map of all Key Mesh Actors for a given trail, mapped to the key time they represent */
+	TArray<TTuple<float, class ASequencerKeyActor*>> KeyMeshActors;
+
+	/** Map of all Frame Mesh Components for a given trail, mapped to the frame time they represent */
+	TArray<TTuple<float, UStaticMeshComponent*>> FrameMeshComponents;
+
+	TSharedPtr<class FSequencer> Sequencer;
+};
\ No newline at end of file
diff --git a/Engine/Source/Editor/Sequencer/Public/MovieSceneTrackEditor.h b/Engine/Source/Editor/Sequencer/Public/MovieSceneTrackEditor.h
--- a/Engine/Source/Editor/Sequencer/Public/MovieSceneTrackEditor.h
+++ b/Engine/Source/Editor/Sequencer/Public/MovieSceneTrackEditor.h
@@ -75,36 +75,8 @@ public:
 		}
 	}
 
-	void AnimatablePropertyChanged( FOnKeyProperty OnKeyProperty )
-	{
-		check(OnKeyProperty.IsBound());
-
-		// Get the movie scene we want to autokey
-		UMovieSceneSequence* MovieSceneSequence = GetMovieSceneSequence();
-		if (MovieSceneSequence)
-		{
-			float KeyTime = GetTimeForKey();
-
-			if( !Sequencer.Pin()->IsRecordingLive() )
-			{
-				// @todo Sequencer - The sequencer probably should have taken care of this
-				MovieSceneSequence->SetFlags(RF_Transactional);
-		
-				// Create a transaction record because we are about to add keys
-				const bool bShouldActuallyTransact = !Sequencer.Pin()->IsRecordingLive();		// Don't transact if we're recording in a PIE world.  That type of keyframe capture cannot be undone.
-				FScopedTransaction AutoKeyTransaction( NSLOCTEXT("AnimatablePropertyTool", "PropertyChanged", "Animatable Property Changed"), bShouldActuallyTransact );
-
-				if( OnKeyProperty.Execute( KeyTime ) )
-				{
-					// TODO: This should pass an appropriate change type here instead of always passing structure changed since most
-					// changes will be value changes.
-					Sequencer.Pin()->NotifyMovieSceneDataChanged( EMovieSceneDataChangeType::MovieSceneStructureItemAdded );
-				}
-
-				UpdatePlaybackRange();
-			}
-		}
-	}
+/** Creates new keys after an animatable property is changed */
+	SEQUENCER_API void AnimatablePropertyChanged(FOnKeyProperty OnKeyProperty);
 
 	struct FFindOrCreateHandleResult
 	{
diff --git a/Engine/Source/Editor/Sequencer/Public/SequencerKeyActor.h b/Engine/Source/Editor/Sequencer/Public/SequencerKeyActor.h
new file mode 100644
--- /dev/null
+++ b/Engine/Source/Editor/Sequencer/Public/SequencerKeyActor.h
@@ -0,0 +1,52 @@
+// Copyright 1998-2016 Epic Games, Inc. All Rights Reserved.
+
+#pragma once
+
+#include "Engine/StaticMeshActor.h"
+#include "SequencerKeyActor.generated.h"
+
+UCLASS()
+class SEQUENCER_API ASequencerKeyActor : public AStaticMeshActor
+{
+	GENERATED_BODY()
+
+public:
+
+	ASequencerKeyActor();
+
+	/** AActor overrides */
+	virtual void PostEditMove(bool bFinished) override;
+	virtual bool IsEditorOnly() const final
+	{
+		return true;
+	}
+	/** AActor */
+
+	/** Set the track section and time associated with this key */
+	void SetKeyData(class UMovieScene3DTransformSection* NewTrackSection, float NewKeyTime);
+
+	/** Return the actor associated with this key */
+	AActor* GetAssociatedActor() const
+	{ 
+		return AssociatedActor; 
+	}
+
+protected:
+
+	/** Push the data from a key actor change back to Sequencer */
+	void PropagateKeyChange();
+
+private:
+
+	/** The actor whose transform was used to build this key */
+	UPROPERTY()
+	AActor*	AssociatedActor;
+
+	/** The track section this key resides on */
+	UPROPERTY()
+	class UMovieScene3DTransformSection* TrackSection;
+
+	/** The time this key is associated with in Sequencer */
+	UPROPERTY()
+	float KeyTime;
+};
\ No newline at end of file
diff --git a/Engine/Source/Editor/Sequencer/Sequencer.Build.cs b/Engine/Source/Editor/Sequencer/Sequencer.Build.cs
--- a/Engine/Source/Editor/Sequencer/Sequencer.Build.cs
+++ b/Engine/Source/Editor/Sequencer/Sequencer.Build.cs
@@ -34,7 +34,8 @@ public class Sequencer : ModuleRules
 				"EditorWidgets", 
 				"SequencerWidgets",
 				"BlueprintGraph",
-				"LevelSequence"
+				"LevelSequence",
+                "ViewportInteraction"
 			}
 		);
 
diff --git a/Engine/Source/Editor/VREditor/Public/VREditorMode.h b/Engine/Source/Editor/VREditor/Public/VREditorMode.h
--- a/Engine/Source/Editor/VREditor/Public/VREditorMode.h
+++ b/Engine/Source/Editor/VREditor/Public/VREditorMode.h
@@ -18,7 +18,7 @@
 class AActor;
 class FEditorViewportClient;
 class SLevelViewport;
-enum class EGizmoHandleTypes : uint8;
+enum class EAutoKeyMode : uint8;
 
 // Forward declare the GizmoHandleTypes that is defined in VIBaseTransformGizmo.h
 enum class EGizmoHandleTypes : uint8;
@@ -267,21 +267,25 @@ public:
 		bool bOnScreenMessages;
 		EHMDTrackingOrigin::Type TrackingOrigin;
 		float WorldToMetersScale;
+		bool bKeyAllEnabled;
+		EAutoKeyMode AutoKeyMode;
 
 		FSavedEditorState()
-			: ViewportType( LVT_Perspective ),
-			  ViewLocation( FVector::ZeroVector ),
-			  ViewRotation( FRotator::ZeroRotator ),
-			  ShowFlags( ESFIM_Editor ),
-			  bLockedPitch( false ),
-			  bGameView( false ),
-			  bAlwaysShowModeWidgetAfterSelectionChanges( false ),
-			  NearClipPlane( 0.0f ),
-			  bRealTime( false ),
-			  DragTriggerDistance( 0.0f ),
-			  bOnScreenMessages( false ),
-			  TrackingOrigin( EHMDTrackingOrigin::Eye ),
-			  WorldToMetersScale( 100.0f )
+			: ViewportType(LVT_Perspective),
+			ViewLocation(FVector::ZeroVector),
+			ViewRotation(FRotator::ZeroRotator),
+			ShowFlags(ESFIM_Editor),
+			bLockedPitch(false),
+			bGameView(false),
+			bAlwaysShowModeWidgetAfterSelectionChanges(false),
+			NearClipPlane(0.0f),
+			bRealTime(false),
+			DragTriggerDistance(0.0f),
+			bOnScreenMessages(false),
+			TrackingOrigin(EHMDTrackingOrigin::Eye),
+			WorldToMetersScale(100.0f),
+			bKeyAllEnabled(false),
+			AutoKeyMode()
 		{
 		}
 	};
@@ -289,6 +293,13 @@ public:
 	/** Gets the saved editor state from entering the mode */
 	const FSavedEditorState& GetSavedEditorState() const;
 
+	DECLARE_DELEGATE(FOnVREditingModeExit);
+	/** Used to override dockable area restoration behavior */
+	FOnVREditingModeExit OnVREditingModeExit_Handler;
+
+	void SaveSequencerSettings(bool bInKeyAllEnabled, EAutoKeyMode InAutoKeyMode);
+
+
 private:
 
 	/** Called when the editor is closed */
@@ -445,6 +456,9 @@ public:
 	// Get the default near clipping plane for VR editing
 	float GetDefaultVRNearClipPlane() const;
 
+	/** Runtime and plugin modules can force VR Editor to refresh using this function */
+	void RefreshVREditorSequencer();
+
 private:
 
 	// All the colors for this mode
diff --git a/Engine/Source/Editor/VREditor/UI/VREditorFloatingUI.h b/Engine/Source/Editor/VREditor/UI/VREditorFloatingUI.h
--- a/Engine/Source/Editor/VREditor/UI/VREditorFloatingUI.h
+++ b/Engine/Source/Editor/VREditor/UI/VREditorFloatingUI.h
@@ -59,6 +59,11 @@ public:
 
 	/** AActor overrides */
 	virtual void Destroyed() override;
+	virtual bool IsEditorOnly() const final
+	{
+		return true;
+	}
+	//~ End AActor interface
 
 	/** Returns the owner of this object */
 	UVREditorUISystem& GetOwner()
@@ -87,6 +92,7 @@ public:
 	/** Gets the initial size of this UI */
 	float GetInitialScale() const;
 
+
 protected:
 
 	/** Returns a scale to use for this widget that takes into account animation */
diff --git a/Engine/Source/Editor/VREditor/UI/VREditorUISystem.cpp b/Engine/Source/Editor/VREditor/UI/VREditorUISystem.cpp
--- a/Engine/Source/Editor/VREditor/UI/VREditorUISystem.cpp
+++ b/Engine/Source/Editor/VREditor/UI/VREditorUISystem.cpp
@@ -78,7 +78,6 @@
 #include "GenericCommands.h"
 #include "SVRRadialPanel.h"
 
-
 #define LOCTEXT_NAMESPACE "VREditor"
 
 namespace VREd
@@ -192,6 +191,12 @@ void UVREditorUISystem::Init()
 	// Bind the color picker creation & destruction overrides
 	SColorPicker::OnColorPickerNonModalCreateOverride.BindUObject(this, &UVREditorUISystem::CreateVRColorPicker);
 	SColorPicker::OnColorPickerDestroyOverride.BindUObject(this, &UVREditorUISystem::DestroyVRColorPicker);
+
+	// Bind the global tab manager's dockable area restore override
+	FGlobalTabmanager::Get()->OnOverrideDockableAreaRestore_Handler.BindUObject(this, &UVREditorUISystem::DockableAreaRestored);
+
+	FVREditorActionCallbacks::GizmoCoordinateSystemText = FVREditorActionCallbacks::GetGizmoCoordinateSystemText();
+	FVREditorActionCallbacks::GizmoModeText = FVREditorActionCallbacks::GetGizmoModeText();
 }
 
 
@@ -206,6 +211,7 @@ void UVREditorUISystem::Shutdown()
 	// Unbind the color picker creation & destruction overrides
 	SColorPicker::OnColorPickerNonModalCreateOverride.Unbind();
 	SColorPicker::OnColorPickerDestroyOverride.Unbind();
+	FGlobalTabmanager::Get()->OnOverrideDockableAreaRestore_Handler.Unbind();
 
 	CleanUpActorsBeforeMapChangeOrSimulate();
 
@@ -1106,6 +1112,29 @@ void UVREditorUISystem::CreateUIs()
 
 			EditorUIPanels[ (int32)EEditorUIPanel::WorldSettings ] = WorldSettingsUI;
 		}
+
+		{
+			const TSharedRef< ILevelEditor >& LevelEditor = FModuleManager::GetModuleChecked<FLevelEditorModule>("LevelEditor").GetFirstLevelEditor().ToSharedRef();
+
+			const FName TabIdentifier = NAME_None;	// No tab for us!
+			TSharedRef<SWidget> SequencerWidget = SNullWidget::NullWidget;
+			
+			TSharedRef<SWidget> WidgetToDraw =
+				SNew(SDPIScaler)
+				.DPIScale(VREd::EditorUIScale->GetFloat())
+				[
+					SequencerWidget
+				]
+			;
+
+			const bool bWithSceneComponent = false;
+			AVREditorFloatingUI* SequencerUI = GetOwner().SpawnTransientSceneActor< AVREditorDockableWindow >(TEXT("SequencerUI"), bWithSceneComponent);
+			SequencerUI->SetSlateWidget(*this, WidgetToDraw, FIntPoint(VREd::ContentBrowserUIResolutionX->GetFloat(), VREd::ContentBrowserUIResolutionY->GetFloat()), VREd::ContentBrowserUISize->GetFloat(), AVREditorFloatingUI::EDockedTo::Nothing);
+			SequencerUI->ShowUI(false);
+			FloatingUIs.Add(SequencerUI);
+
+			EditorUIPanels[(int32)EEditorUIPanel::SequencerUI] = SequencerUI;
+		}
 	}
 }
 
@@ -1138,7 +1167,15 @@ void UVREditorUISystem::OnAssetEditorOpened(UObject* Asset)
 	TArray<IAssetEditorInstance*> Editors = FAssetEditorManager::Get().FindEditorsForAsset(Asset);
 	for ( IAssetEditorInstance* Editor : Editors )
 	{
-		Editor->GetAssociatedTabManager()->SetCanDoDragOperation(false);
+		if (Editor->GetAssociatedTabManager().IsValid())
+		{
+			Editor->GetAssociatedTabManager()->SetCanDoDragOperation(false);
+		}
+		else
+		{
+			FLevelEditorModule& LevelEditorModule = FModuleManager::GetModuleChecked<FLevelEditorModule>(TEXT("LevelEditor"));
+			LevelEditorModule.GetLevelEditorTabManager().ToSharedRef()->SetCanDoDragOperation(false);
+		}
 	}
 }
 
@@ -1643,6 +1680,10 @@ void UVREditorUISystem::SetDefaultWindowLayout()
 		}
 	}
 }
+void UVREditorUISystem::DockableAreaRestored()
+{
+}
+
 void UVREditorUISystem::CreateVRColorPicker(const TSharedRef<SColorPicker>& ColorPicker)
 {
 	// Check that the Color Picker Panel isn't currently showing. Also handles the EEditorUIPanel::ColorPicker being null.
@@ -2067,6 +2108,16 @@ TSharedRef<SWidget> UVREditorUISystem::BuildQuickMenuWidget()
 		NAME_None,
 		EUserInterfaceActionType::Check
 		);
+	InterfaceMenuBuilder.AddMenuEntry(
+		LOCTEXT("CreateNewSequence", "Create Sequence"),
+		LOCTEXT("CreateNewSequenceTooltip", "Create a new Level Sequence and open it"),
+		FSlateIcon(),
+		FUIAction
+		(
+			FExecuteAction::CreateStatic(&FVREditorActionCallbacks::CreateNewSequence, VRMode),
+			FCanExecuteAction::CreateStatic(&FLevelEditorActionCallbacks::DefaultCanExecuteAction)
+			)
+		);
 	InterfaceMenuBuilder.AddMenuEntry(
 		LOCTEXT("ToggleAllUI", "Toggle All UI"),
 		LOCTEXT("ToggleAllUITooltip", "ToggleAllUI"),
@@ -2523,6 +2574,8 @@ const TSharedRef<SWidget>& UVREditorUISystem::FindWidgetOfType(const TSharedRef<
 	return SNullWidget::NullWidget;
 }
 
+
+
 void UVREditorUISystem::OnHoverBeginEffect(TSharedRef<SButton> Button)
 {
 	FVRButton* ButtonToAnimate = VRButtons.FindByPredicate([&Button](const FVRButton& ButtonData)
@@ -2584,4 +2637,36 @@ void UVREditorUISystem::SwapRadialMenu()
 	QuickRadialMenu->GetWidgetComponent()->SetBlendMode( EWidgetBlendMode::Masked );
 }
 
+void UVREditorUISystem::UpdateSequencerUI()
+{
+	if (EditorUIPanels[(int32)EEditorUIPanel::SequencerUI] != nullptr)
+	{
+		AVREditorFloatingUI* SequencerPanel = EditorUIPanels[(int32)EEditorUIPanel::SequencerUI];
+		const FIntPoint DefaultResolution(VREd::DefaultEditorUIResolutionX->GetInt(), VREd::DefaultEditorUIResolutionY->GetInt());
+		const TSharedRef< ILevelEditor >& LevelEditor = FModuleManager::GetModuleChecked<FLevelEditorModule>("LevelEditor").GetFirstLevelEditor().ToSharedRef();
+
+		const FName TabIdentifier = NAME_None;	// No tab for us!
+		TSharedRef<SWidget> SequencerWidget = LevelEditor->GetSequencerWidgetForLevelEditor().ToSharedRef();
+
+		TSharedRef<SWidget> WidgetToDraw =
+			SNew(SDPIScaler)
+			.DPIScale(VREd::EditorUIScale->GetFloat())
+			[
+				SequencerWidget
+			]
+		;
+
+		const bool bWithSceneComponent = false;
+		SequencerPanel->SetSlateWidget(*this, WidgetToDraw, FIntPoint(VREd::ContentBrowserUIResolutionX->GetFloat(), VREd::ContentBrowserUIResolutionY->GetFloat()), VREd::ContentBrowserUISize->GetFloat(), AVREditorFloatingUI::EDockedTo::Nothing);
+		UVREditorInteractor* VREditorInteractor = VRMode->GetHandInteractor(EControllerHand::Left);
+		if (!VREditorInteractor->IsHoveringOverUI())
+		{
+			VREditorInteractor = VRMode->GetHandInteractor(EControllerHand::Right);
+
+		}
+		ShowEditorUIPanel(UVREditorUISystem::EEditorUIPanel::SequencerUI, VREditorInteractor, true);
+
+	}
+}
+
 #undef LOCTEXT_NAMESPACE
\ No newline at end of file
diff --git a/Engine/Source/Editor/VREditor/UI/VREditorUISystem.h b/Engine/Source/Editor/VREditor/UI/VREditorUISystem.h
--- a/Engine/Source/Editor/VREditor/UI/VREditorUISystem.h
+++ b/Engine/Source/Editor/VREditor/UI/VREditorUISystem.h
@@ -98,6 +98,7 @@ public:
 		AssetEditor,
 		WorldSettings,
 		ColorPicker,
+		SequencerUI,
 		// ...
 
 		TotalCount,
@@ -199,6 +200,8 @@ public:
 	 */
 	static const TSharedRef<SWidget>& FindWidgetOfType(const TSharedRef<SWidget>& Content, FName WidgetType);
 
+	/** Function to force an update of the Sequencer UI based on a change */
+	void UpdateSequencerUI();
 
 protected:
 
@@ -230,6 +233,8 @@ protected:
 	/** Sets the main windows to their default transform */
 	void SetDefaultWindowLayout();
 
+	/** Can be used when the tab manager is restoring a dockable tab area */
+	void DockableAreaRestored();
 
 	/** Creates a VR-specific color picker. Gets bound to SColorPicker's creation override delegate */
 	void CreateVRColorPicker(const TSharedRef<SColorPicker>& ColorPicker);
diff --git a/Engine/Source/Editor/VREditor/VREditor.Build.cs b/Engine/Source/Editor/VREditor/VREditor.Build.cs
--- a/Engine/Source/Editor/VREditor/VREditor.Build.cs
+++ b/Engine/Source/Editor/VREditor/VREditor.Build.cs
@@ -22,7 +22,9 @@ namespace UnrealBuildTool.Rules
 					"LevelEditor",
 					"HeadMountedDisplay",
                     "ViewportInteraction",
-					"Analytics"
+					"Analytics",
+                    "LevelSequence",
+                    "Sequencer"
 				}
 			);
 
diff --git a/Engine/Source/Editor/VREditor/VREditorActions.cpp b/Engine/Source/Editor/VREditor/VREditorActions.cpp
--- a/Engine/Source/Editor/VREditor/VREditorActions.cpp
+++ b/Engine/Source/Editor/VREditor/VREditorActions.cpp
@@ -10,12 +10,22 @@
 #include "VREditorFloatingUI.h"
 #include "VREditorTransformGizmo.h"
 #include "SLevelViewport.h"
+#include "SlateApplication.h"
 #include "ImageUtils.h"
 #include "FileHelper.h"
 
+#include "AssetEditorManager.h"
+#include "Developer/AssetTools/Public/IAssetTools.h"
+#include "Developer/AssetTools/Public/AssetToolsModule.h"
+#include "ModuleManager.h"
+#include "LevelSequence.h"
+#include "LevelSequenceActor.h"
 
 #define LOCTEXT_NAMESPACE "VREditorActions"
 
+FText FVREditorActionCallbacks::GizmoCoordinateSystemText;
+FText FVREditorActionCallbacks::GizmoModeText;
+
 ECheckBoxState FVREditorActionCallbacks::GetTranslationSnapState()
 {
 	return GetDefault<ULevelEditorViewportSettings>()->GridEnabled ? ECheckBoxState::Checked : ECheckBoxState::Unchecked;
@@ -108,13 +118,13 @@ void FVREditorActionCallbacks::OnGizmoCoordinateSystemButtonClicked(class UVREdi
 
 FText FVREditorActionCallbacks::GetGizmoCoordinateSystemText()
 {
-	return EVREditorActions::GizmoCoordinateSystemText;
+	return FVREditorActionCallbacks::GizmoCoordinateSystemText;
 }
 
 void FVREditorActionCallbacks::UpdateGizmoCoordinateSystemText(class UVREditorMode* InVRMode) 
 {
 	const ECoordSystem CurrentCoordSystem = InVRMode->GetWorldInteraction().GetTransformGizmoCoordinateSpace(); //@todo VREditor
-	EVREditorActions::GizmoCoordinateSystemText = (CurrentCoordSystem == COORD_World ? LOCTEXT("WorldCoordinateSystem", "World Space") : LOCTEXT("LocalCoordinateSystem", "Local Space"));
+	FVREditorActionCallbacks::GizmoCoordinateSystemText = (CurrentCoordSystem == COORD_World ? LOCTEXT("WorldCoordinateSystem", "World Space") : LOCTEXT("LocalCoordinateSystem", "Local Space"));
 }
 
 void FVREditorActionCallbacks::OnGizmoModeButtonClicked(class UVREditorMode* InVRMode)
@@ -128,7 +138,7 @@ void FVREditorActionCallbacks::OnGizmoModeButtonClicked(class UVREditorMode* InV
 
 FText FVREditorActionCallbacks::GetGizmoModeText()
 {
-	return EVREditorActions::GizmoModeText;
+	return FVREditorActionCallbacks::GizmoModeText;
 }
 
 void FVREditorActionCallbacks::UpdateGizmoModeText(class UVREditorMode* InVRMode) 
@@ -158,7 +168,7 @@ void FVREditorActionCallbacks::UpdateGizmoModeText(class UVREditorMode* InVRMode
 		break;
 	}
 
-	EVREditorActions::GizmoModeText = GizmoTypeText;
+	FVREditorActionCallbacks::GizmoModeText = GizmoTypeText;
 }
 
 void FVREditorActionCallbacks::OnGizmoCycle(class UVREditorMode* InVRMode)
@@ -270,7 +280,7 @@ void FVREditorActionCallbacks::SimulateCharacterEntry(const FString InChar)
 
 }
 
-void FVREditorActionCallbacks::SimulateKeyDown(FKey Key, bool bRepeat)
+void FVREditorActionCallbacks::SimulateKeyDown(const FKey Key, const bool bRepeat)
 {
 	const uint32* KeyCodePtr;
 	const uint32* CharCodePtr;
@@ -289,7 +299,7 @@ void FVREditorActionCallbacks::SimulateKeyDown(FKey Key, bool bRepeat)
 	}
 }
 
-void FVREditorActionCallbacks::SimulateKeyUp(FKey Key)
+void FVREditorActionCallbacks::SimulateKeyUp(const FKey Key)
 {
 	const uint32* KeyCodePtr;
 	const uint32* CharCodePtr;
@@ -302,4 +312,48 @@ void FVREditorActionCallbacks::SimulateKeyUp(FKey Key)
 	FSlateApplication::Get().ProcessKeyUpEvent( KeyEvent );
 }
 
+void FVREditorActionCallbacks::CreateNewSequence(class UVREditorMode* InVRMode)
+{
+	// Create a new level sequence
+	IAssetTools& AssetTools = FModuleManager::GetModuleChecked<FAssetToolsModule>("AssetTools").Get();
+
+	UObject* NewAsset = nullptr;
+
+	// Attempt to create a new asset
+	for (TObjectIterator<UClass> It; It; ++It)
+	{
+		UClass* CurrentClass = *It;
+		if (CurrentClass->IsChildOf(UFactory::StaticClass()) && !(CurrentClass->HasAnyClassFlags(CLASS_Abstract)))
+		{
+			UFactory* Factory = Cast<UFactory>(CurrentClass->GetDefaultObject());
+			if (Factory->CanCreateNew() && Factory->ImportPriority >= 0 && Factory->SupportedClass == ULevelSequence::StaticClass())
+			{
+				FString NewPackageName;
+				FString NewAssetName;
+				// Sequences created in VR editor will have a sequential VRSequencer00X naming scheme and be stored in Game/Sequences
+				AssetTools.CreateUniqueAssetName(TEXT("/Game/Cinematics/Sequences/VRSequence"), TEXT("001"), NewPackageName, NewAssetName);
+				NewAsset = AssetTools.CreateAsset(NewAssetName, TEXT("/Game/Cinematics/Sequences"), ULevelSequence::StaticClass(), Factory);
+				break;
+			}
+		}
+	}
+
+	if (!NewAsset)
+	{
+		return;
+	}
+
+	// Spawn an actor at the origin, and move in front of the camera and open for edit
+	UActorFactory* ActorFactory = GEditor->FindActorFactoryForActorClass(ALevelSequenceActor::StaticClass());
+	if (!ensure(ActorFactory))
+	{
+		return;
+	}
+
+	ALevelSequenceActor* NewActor = CastChecked<ALevelSequenceActor>(GEditor->UseActorFactory(ActorFactory, FAssetData(NewAsset), &FTransform::Identity));
+
+	// Open the Sequencer window
+	FAssetEditorManager::Get().OpenEditorForAsset(NewAsset);
+}
+
 #undef LOCTEXT_NAMESPACE
diff --git a/Engine/Source/Editor/VREditor/VREditorActions.h b/Engine/Source/Editor/VREditor/VREditorActions.h
--- a/Engine/Source/Editor/VREditor/VREditorActions.h
+++ b/Engine/Source/Editor/VREditor/VREditorActions.h
@@ -3,15 +3,6 @@
 
 #include "VREditorUISystem.h"
 
-namespace EVREditorActions
-{
-	//Gizmo coordinate button text starts blank before it is initialized
-	static FText GizmoCoordinateSystemText = NSLOCTEXT("GizmoCoordinateDefault", "", "");
-
-	//Gizmo mode button text starts blank before it is initialized
-	static FText GizmoModeText = NSLOCTEXT("GizmoModeDefault", "", "");;
-}
-
 /**
 * Implementation of various VR editor action callback functions
 */
@@ -137,16 +128,22 @@ public:
 	 * @param Key Key to press
 	 * @param bRepeat Whether or not to repeat
 	 */
-	static void SimulateKeyDown( FKey Key, bool bRepeat );
+	static void SimulateKeyDown( const FKey Key, const bool bRepeat );
 
 	/**
 	 * Simulates the user releasing a key
 	 * @param Key Key to release
 	 */
-	static void SimulateKeyUp( FKey Key );
+	static void SimulateKeyUp( const FKey Key );
 
-}
-;
+	/** Create a new level sequence with an auto-generated name */
+	static void CreateNewSequence(class UVREditorMode* InVRMode);
 
+public:
+	
+	static FText GizmoCoordinateSystemText;
 
+	static FText GizmoModeText;
+}
+;
 
diff --git a/Engine/Source/Editor/VREditor/VREditorMode.cpp b/Engine/Source/Editor/VREditor/VREditorMode.cpp
--- a/Engine/Source/Editor/VREditor/VREditorMode.cpp
+++ b/Engine/Source/Editor/VREditor/VREditorMode.cpp
@@ -34,6 +34,8 @@
 
 #include "ViewportWorldInteractionManager.h"
 
+#include "ISequencer.h"
+
 #define LOCTEXT_NAMESPACE "VREditorMode"
 
 namespace VREd
@@ -452,6 +454,7 @@ void UVREditorMode::Exit(const bool bHMDShouldExitStereo)
 						VRViewportClient.GetWorld()->GetWorldSettings()->WorldToMeters = SavedEditorState.WorldToMetersScale;
 						ENGINE_API extern float GNewWorldToMetersScale;
 						GNewWorldToMetersScale = 0.0f;
+						OnVREditingModeExit_Handler.ExecuteIfBound();
 					}
 				}
 
@@ -909,6 +912,12 @@ float UVREditorMode::GetDefaultVRNearClipPlane() const
 	return VREd::VRNearClipPlane->GetFloat();
 }
 
+void UVREditorMode::RefreshVREditorSequencer()
+{
+	// Tell the VR Editor UI system to refresh the Sequencer UI
+	GetUISystem().UpdateSequencerUI();
+}
+
 bool UVREditorMode::IsHandAimingTowardsCapsule( UViewportInteractor* Interactor, const FTransform& CapsuleTransform, FVector CapsuleStart, const FVector CapsuleEnd, const float CapsuleRadius, const float MinDistanceToCapsule, const FVector CapsuleFrontDirection, const float MinDotForAimingAtCapsule ) const
 {
 	bool bIsAimingTowards = false;
@@ -993,4 +1002,10 @@ const UVREditorMode::FSavedEditorState& UVREditorMode::GetSavedEditorState() con
 	return SavedEditorState;
 }
 
+void UVREditorMode::SaveSequencerSettings(bool bInKeyAllEnabled, EAutoKeyMode InAutoKeyMode)
+{
+	SavedEditorState.bKeyAllEnabled = bInKeyAllEnabled;
+	SavedEditorState.AutoKeyMode = InAutoKeyMode;
+}
+
 #undef LOCTEXT_NAMESPACE
diff --git a/Engine/Source/Editor/ViewportInteraction/Public/VIBaseTransformGizmo.h b/Engine/Source/Editor/ViewportInteraction/Public/VIBaseTransformGizmo.h
--- a/Engine/Source/Editor/ViewportInteraction/Public/VIBaseTransformGizmo.h
+++ b/Engine/Source/Editor/ViewportInteraction/Public/VIBaseTransformGizmo.h
@@ -50,6 +50,13 @@ public:
 	/** Deconstructor */
 	virtual ~ABaseTransformGizmo();
 
+	//~ Begin AActor interface
+	virtual bool IsEditorOnly() const final
+	{
+		return true;
+	}
+	//~ End AActor interface
+
 	/** Call this when new objects become selected.  This triggers an animation transition. */
 	void OnNewObjectsSelected();
 
diff --git a/Engine/Source/Editor/ViewportInteraction/ViewportWorldInteraction.cpp b/Engine/Source/Editor/ViewportInteraction/ViewportWorldInteraction.cpp
--- a/Engine/Source/Editor/ViewportInteraction/ViewportWorldInteraction.cpp
+++ b/Engine/Source/Editor/ViewportInteraction/ViewportWorldInteraction.cpp
@@ -2267,6 +2267,7 @@ void UViewportWorldInteraction::StartDraggingActors( UViewportInteractor* Intera
 			// using the laser impact point as the target transform
 			if ( bIsPlacingActors )
 			{
+				GEditor->BroadcastBeginObjectMovement(*Actor);
 				bIsInterpolatingTransformablesFromSnapshotTransform = true;
 				const FTimespan CurrentTime = FTimespan::FromSeconds( FPlatformTime::Seconds() );
 				TransformablesInterpolationStartTime = CurrentTime;
@@ -2325,6 +2326,12 @@ void UViewportWorldInteraction::StopDragging( UViewportInteractor* Interactor )
 				InteractorData.DraggingMode == EViewportInteractionDraggingMode::ActorsFreely ||
 				InteractorData.DraggingMode == EViewportInteractionDraggingMode::ActorsAtLaserImpact )
 			{
+				// vreditor todo: Move this to where we detect a move from inertia finishes
+				FHitResult HitResult = Interactor->GetHitResultFromLaserPointer();
+				if (HitResult.Actor.IsValid())
+				{
+					GEditor->BroadcastEndObjectMovement(*HitResult.Actor);
+				}
 				// No hand is dragging the objects anymore.
 				if ( InteractorData.DraggingMode == EViewportInteractionDraggingMode::ActorsWithGizmo &&
 					InteractorData.TransformGizmoInteractionType == ETransformGizmoInteractionType::RotateOnAngle )
@@ -2340,7 +2347,7 @@ void UViewportWorldInteraction::StopDragging( UViewportInteractor* Interactor )
 					--TrackingTransaction.TransCount;
 					TrackingTransaction.End();
 					GEditor->DisableDeltaModification( false );
-				}
+				}	
 			}
 		}
 		
diff --git a/Engine/Source/Runtime/Slate/Public/Framework/Docking/TabManager.h b/Engine/Source/Runtime/Slate/Public/Framework/Docking/TabManager.h
--- a/Engine/Source/Runtime/Slate/Public/Framework/Docking/TabManager.h
+++ b/Engine/Source/Runtime/Slate/Public/Framework/Docking/TabManager.h
@@ -904,6 +904,10 @@ public:
 
 	void SetProxyTabManager(TSharedPtr<FProxyTabmanager> InProxyTabManager);
 
+	DECLARE_DELEGATE(FOnOverrideDockableAreaRestore);
+	/** Used to override dockable area restoration behavior */
+	FOnOverrideDockableAreaRestore OnOverrideDockableAreaRestore_Handler;
+
 protected:
 	virtual void OnTabForegrounded( const TSharedPtr<SDockTab>& NewForegroundTab, const TSharedPtr<SDockTab>& BackgroundedTab ) override;
 	virtual void OnTabRelocated( const TSharedRef<SDockTab>& RelocatedTab, const TSharedPtr<SWindow>& NewOwnerWindow ) override;
