Changing TArray<TTuple>s to TArray<Struct> to fix clang crash, improve code readability

diff --git a/Engine/Source/Editor/Sequencer/Private/SequencerEdMode.cpp b/Engine/Source/Editor/Sequencer/Private/SequencerEdMode.cpp
--- a/Engine/Source/Editor/Sequencer/Private/SequencerEdMode.cpp
+++ b/Engine/Source/Editor/Sequencer/Private/SequencerEdMode.cpp
@@ -122,10 +122,10 @@ void FSequencerEdMode::DrawHUD(FEditorViewportClient* ViewportClient,FViewport*
 
 void FSequencerEdMode::AddReferencedObjects(FReferenceCollector& Collector)
 {
-	for (auto MeshTrail : MeshTrails)
+	for (FMeshTrailData MeshTrail : MeshTrails)
 	{
-		Collector.AddReferencedObject(MeshTrail.Get<0>());
-		Collector.AddReferencedObject(MeshTrail.Get<1>());
+		Collector.AddReferencedObject(MeshTrail.Track);
+		Collector.AddReferencedObject(MeshTrail.Trail);
 	}
 }
 
@@ -149,14 +149,14 @@ void FSequencerEdMode::OnKeySelected(FViewport* Viewport, HMovieSceneKeyProxy* K
 void FSequencerEdMode::DrawMeshTransformTrailFromKey(const class ASequencerKeyActor* KeyActor)
 {
 	ASequencerMeshTrail* Trail = Cast<ASequencerMeshTrail>(KeyActor->GetOwner());
-	TTuple<UMovieScene3DTransformTrack*, ASequencerMeshTrail*>* TrailPtr = MeshTrails.FindByPredicate([Trail](const TTuple<UMovieScene3DTransformTrack*, ASequencerMeshTrail*> InTrail)
+	FMeshTrailData* TrailPtr = MeshTrails.FindByPredicate([Trail](const FMeshTrailData InTrail)
 	{
-		return Trail == InTrail.Get<1>();
+		return Trail == InTrail.Trail;
 	});
 	if(TrailPtr)
 	{
 		// From the key, get the mesh trail, and then the track associated with that mesh trail
-		UMovieScene3DTransformTrack* Track = TrailPtr->Get<0>();
+		UMovieScene3DTransformTrack* Track = TrailPtr->Track;
 		// Draw a mesh trail for the key's associated actor
 		TArray<TWeakObjectPtr<UObject>> KeyObjects;
 		AActor* TrailActor = KeyActor->GetAssociatedActor();
@@ -169,9 +169,9 @@ void FSequencerEdMode::DrawMeshTransformTrailFromKey(const class ASequencerKeyAc
 void FSequencerEdMode::CleanUpMeshTrails()
 {
 	// Clean up any existing trails
-	for (auto& MeshTrail : MeshTrails)
+	for (FMeshTrailData& MeshTrail : MeshTrails)
 	{
-		MeshTrail.Get<1>()->Cleanup();
+		MeshTrail.Trail->Cleanup();
 	}
 	MeshTrails.Empty();
 }
@@ -215,13 +215,13 @@ void FSequencerEdMode::DrawTransformTrack(FPrimitiveDrawInterface* PDI, UMovieSc
 	// Get the Trail Actor associated with this track if we are drawing mesh trails
 	if (bDrawMeshTrails)
 	{
-		TTuple<UMovieScene3DTransformTrack*, ASequencerMeshTrail*>* TrailPtr = MeshTrails.FindByPredicate([TransformTrack](const TTuple<UMovieScene3DTransformTrack*, ASequencerMeshTrail*> InTrail)
+		FMeshTrailData* TrailPtr = MeshTrails.FindByPredicate([TransformTrack](const FMeshTrailData InTrail)
 		{
-			return InTrail.Get<0>() == TransformTrack;
+			return InTrail.Track == TransformTrack;
 		});
 		if (TrailPtr != nullptr)
 		{
-			TrailActor = TrailPtr->Get<1>();
+			TrailActor = TrailPtr->Trail;
 		}
 	}
 
@@ -540,15 +540,15 @@ void FSequencerEdMode::DrawTracks3D(FPrimitiveDrawInterface* PDI)
 					// If we are drawing mesh trails but we haven't made one for this track yet
 					if (bDrawMeshTrails)
 					{
-						TTuple<UMovieScene3DTransformTrack*, ASequencerMeshTrail*>* TrailPtr = MeshTrails.FindByPredicate([TransformTrack](const TTuple<UMovieScene3DTransformTrack*, ASequencerMeshTrail*> InTrail)
+						FMeshTrailData* TrailPtr = MeshTrails.FindByPredicate([TransformTrack](const FMeshTrailData InTrail)
 						{
-							return InTrail.Get<0>() == TransformTrack;
+							return InTrail.Track == TransformTrack;
 						});
 						if (TrailPtr == nullptr)
 						{
 							UViewportWorldInteraction* ViewportWorldInteraction = GEditor->GetEditorWorldManager()->GetEditorWorldWrapper(GEditor->GetActiveViewport()->GetClient()->GetWorld())->GetViewportWorldInteraction();
 							ASequencerMeshTrail* TrailActor = ViewportWorldInteraction->SpawnTransientSceneActor<ASequencerMeshTrail>(TEXT("SequencerMeshTrail"), true);
-							TTuple<UMovieScene3DTransformTrack*, ASequencerMeshTrail*> MeshTrail = MakeTuple(TransformTrack, TrailActor);
+							FMeshTrailData MeshTrail = FMeshTrailData(TransformTrack, TrailActor);
 							MeshTrails.Add(MeshTrail);
 						}
 					}
diff --git a/Engine/Source/Editor/Sequencer/Private/SequencerEdMode.h b/Engine/Source/Editor/Sequencer/Private/SequencerEdMode.h
--- a/Engine/Source/Editor/Sequencer/Private/SequencerEdMode.h
+++ b/Engine/Source/Editor/Sequencer/Private/SequencerEdMode.h
@@ -15,6 +15,19 @@ class FSequencer;
 class FViewport;
 struct HMovieSceneKeyProxy;
 
+/** Stores the transform track and associated mesh trail for each drawn track */
+struct FMeshTrailData
+{
+	class UMovieScene3DTransformTrack* Track;
+	class ASequencerMeshTrail* Trail;
+
+	FMeshTrailData(class UMovieScene3DTransformTrack* InTrack, class ASequencerMeshTrail* InTrail) :
+		Track(InTrack),
+		Trail(InTrail)
+	{
+	}
+};
+
 /**
  * FSequencerEdMode is the editor mode for additional drawing and handling sequencer hotkeys in the editor
  */
@@ -54,8 +67,8 @@ protected:
 private:
 	TWeakPtr<FSequencer> SequencerPtr;
 
-	/**The map of any transform tracks with mesh trails, and their mesh trail representations */
-	TArray<TTuple<class UMovieScene3DTransformTrack*,class ASequencerMeshTrail*>> MeshTrails;
+	/**The array of any transform tracks with mesh trails, and their mesh trail representations */
+	TArray<FMeshTrailData> MeshTrails;
 
 	/** True when drawing mesh trails instead of debug trails */
 	bool bDrawMeshTrails;
diff --git a/Engine/Source/Editor/Sequencer/Private/SequencerMeshTrail.cpp b/Engine/Source/Editor/Sequencer/Private/SequencerMeshTrail.cpp
--- a/Engine/Source/Editor/Sequencer/Private/SequencerMeshTrail.cpp
+++ b/Engine/Source/Editor/Sequencer/Private/SequencerMeshTrail.cpp
@@ -15,9 +15,9 @@ ASequencerMeshTrail::ASequencerMeshTrail()
 void ASequencerMeshTrail::Cleanup()
 {
 	// Destroy all the key actors this trail created
-	for (auto& KeyMesh : KeyMeshActors)
+	for (FKeyActorData KeyMesh : KeyMeshActors)
 	{
-		KeyMesh.Get<1>()->Destroy();
+		KeyMesh.KeyActor->Destroy();
 	}
 	this->Destroy();
 }
@@ -28,9 +28,9 @@ void ASequencerMeshTrail::AddKeyMeshActor(float KeyTime, const FTransform KeyTra
 	UStaticMesh* KeyEditorMesh = nullptr;
 	UMaterialInstance* KeyEditorMaterial = nullptr;
 
-	TTuple<float, ASequencerKeyActor*>* KeyPtr = KeyMeshActors.FindByPredicate([KeyTime](const TTuple<float, ASequencerKeyActor*> InKey)
+	FKeyActorData* KeyPtr = KeyMeshActors.FindByPredicate([KeyTime](const FKeyActorData InKey)
 	{
-		return FMath::IsNearlyEqual(KeyTime, InKey.Get<0>());
+		return FMath::IsNearlyEqual(KeyTime, InKey.Time);
 	});
 
 	// If we don't currently have an actor for this time, create one
@@ -42,13 +42,13 @@ void ASequencerMeshTrail::AddKeyMeshActor(float KeyTime, const FTransform KeyTra
 		KeyMeshActor->SetActorTransform(KeyTransform);
 		KeyMeshActor->SetKeyData(TrackSection, KeyTime);
 		KeyMeshActor->SetOwner(this);
-		TTuple<float, ASequencerKeyActor*> NewKey = MakeTuple(KeyTime, KeyMeshActor);
+		FKeyActorData NewKey = FKeyActorData(KeyTime, KeyMeshActor);
 		KeyMeshActors.Add(NewKey);
 	}
 	else
 	{
 		// Just update the transform
-		KeyPtr->Get<1>()->SetActorTransform(KeyTransform);
+		KeyPtr->KeyActor->SetActorTransform(KeyTransform);
 	}
 }
 
@@ -58,9 +58,9 @@ void ASequencerMeshTrail::AddFrameMeshComponent(const float FrameTime, const FTr
 	UStaticMesh* KeyEditorMesh = nullptr;
 	UMaterialInterface* KeyEditorMaterial = nullptr;
 
-	TTuple<float, UStaticMeshComponent*>* FramePtr = FrameMeshComponents.FindByPredicate([FrameTime](const TTuple<float, UStaticMeshComponent*> InFrame)
+	FFrameComponentData* FramePtr = FrameMeshComponents.FindByPredicate([FrameTime](const FFrameComponentData InFrame)
 	{
-		return FMath::IsNearlyEqual(FrameTime, InFrame.Get<0>());
+		return FMath::IsNearlyEqual(FrameTime, InFrame.Time);
 	});
 
 	// If we don't currently have a component for this time, create one
@@ -79,12 +79,12 @@ void ASequencerMeshTrail::AddFrameMeshComponent(const float FrameTime, const FTr
 		FrameMeshComponent->SetWorldTransform(FrameTransform);
 		FrameMeshComponent->SetMobility(EComponentMobility::Movable);
 		FrameMeshComponent->bSelectable = false;
-		TTuple<float, UStaticMeshComponent*> NewFrame = MakeTuple(FrameTime, FrameMeshComponent);
+		FFrameComponentData NewFrame = FFrameComponentData(FrameTime, FrameMeshComponent);
 		FrameMeshComponents.Add(NewFrame);
 	}
 	else
 	{
 		// Just update the transform
-		FramePtr->Get<1>()->SetWorldTransform(FrameTransform);
+		FramePtr->FrameComponent->SetWorldTransform(FrameTransform);
 	}
 }
diff --git a/Engine/Source/Editor/Sequencer/Private/SequencerMeshTrail.h b/Engine/Source/Editor/Sequencer/Private/SequencerMeshTrail.h
--- a/Engine/Source/Editor/Sequencer/Private/SequencerMeshTrail.h
+++ b/Engine/Source/Editor/Sequencer/Private/SequencerMeshTrail.h
@@ -3,6 +3,32 @@
 #pragma once
 #include "SequencerMeshTrail.generated.h"
 
+/** Stores time and actor reference for each key Actor on the trail */
+struct FKeyActorData
+{
+	float Time;
+	class ASequencerKeyActor* KeyActor;
+
+	FKeyActorData(float InTime, class ASequencerKeyActor* InKeyActor) :
+		Time(InTime),
+		KeyActor(InKeyActor)
+	{
+	}
+};
+
+/** Stores time and component reference for each frame Component on the trail */
+struct FFrameComponentData
+{
+	float Time;
+	class UStaticMeshComponent* FrameComponent;
+
+	FFrameComponentData(float InTime, class UStaticMeshComponent* InFrameComponent) :
+		Time(InTime),
+		FrameComponent(InFrameComponent)
+	{
+	}
+};
+
 UCLASS()
 class ASequencerMeshTrail : public AActor
 {
@@ -27,11 +53,11 @@ public:
 
 private:
 
-	/** Map of all Key Mesh Actors for a given trail, mapped to the key time they represent */
-	TArray<TTuple<float, class ASequencerKeyActor*>> KeyMeshActors;
+	/** Array of all Key Mesh Actors for a given trail and the key time they represent */
+	TArray<FKeyActorData> KeyMeshActors;
 
-	/** Map of all Frame Mesh Components for a given trail, mapped to the frame time they represent */
-	TArray<TTuple<float, UStaticMeshComponent*>> FrameMeshComponents;
+	/** Array of all Frame Mesh Components for a given trail and the frame time they represent */
+	TArray<FFrameComponentData> FrameMeshComponents;
 
 	TSharedPtr<class FSequencer> Sequencer;
 };
\ No newline at end of file
