Copying //UE4/Dev-Sequencer to //UE4/Dev-Main (Source: //UE4/Dev-Sequencer @ 3237992)

diff --git a/Engine/Source/Editor/Sequencer/Private/DisplayNodes/SequencerTrackNode.cpp b/Engine/Source/Editor/Sequencer/Private/DisplayNodes/SequencerTrackNode.cpp
--- a/Engine/Source/Editor/Sequencer/Private/DisplayNodes/SequencerTrackNode.cpp
+++ b/Engine/Source/Editor/Sequencer/Private/DisplayNodes/SequencerTrackNode.cpp
@@ -110,7 +110,7 @@ void FSequencerTrackNode::FixRowIndices()
 
 namespace
 {
-	void AddBoolPropertyMenuItem(FMenuBuilder& MenuBuilder, UMovieSceneTrack* Track, const UBoolProperty* Property, void* PropertyContainer)
+	void AddBoolPropertyMenuItem(FMenuBuilder& MenuBuilder, FCanExecuteAction InCanExecute, UMovieSceneTrack* Track, const UBoolProperty* Property, void* PropertyContainer)
 	{
 		MenuBuilder.AddMenuEntry(
 			Property->GetDisplayNameText(),
@@ -124,7 +124,7 @@ namespace
 					bool bIsSet = Property->GetPropertyValue(PropertyContainer);
 					Property->SetPropertyValue(PropertyContainer, !bIsSet);
 				}),
-				FCanExecuteAction::CreateLambda([]{ return true; }),
+				InCanExecute,
 				FIsActionChecked::CreateLambda([PropertyContainer, Property]{ return Property->GetPropertyValue(PropertyContainer); })
 			),
 			NAME_None,
@@ -146,24 +146,27 @@ void FSequencerTrackNode::BuildContextMenu(FMenuBuilder& MenuBuilder)
 			return;
 		}
 
+		bool bIsReadOnly = !GetSequencer().IsReadOnly();
+		FCanExecuteAction CanExecute = FCanExecuteAction::CreateLambda([bIsReadOnly]{ return bIsReadOnly; });
+
 		UStruct* EvalOptionsStruct = FMovieSceneTrackEvalOptions::StaticStruct();
 
 		const UBoolProperty* NearestSectionProperty = Cast<UBoolProperty>(EvalOptionsStruct->FindPropertyByName(GET_MEMBER_NAME_CHECKED(FMovieSceneTrackEvalOptions, bEvaluateNearestSection)));
 		if (NearestSectionProperty && Track->EvalOptions.bCanEvaluateNearestSection)
 		{
-			AddBoolPropertyMenuItem(MenuBuilder, Track, NearestSectionProperty, NearestSectionProperty->ContainerPtrToValuePtr<void>(&Track->EvalOptions));
+			AddBoolPropertyMenuItem(MenuBuilder, CanExecute, Track, NearestSectionProperty, NearestSectionProperty->ContainerPtrToValuePtr<void>(&Track->EvalOptions));
 		}
 
 		const UBoolProperty* PrerollProperty = Cast<UBoolProperty>(EvalOptionsStruct->FindPropertyByName(GET_MEMBER_NAME_CHECKED(FMovieSceneTrackEvalOptions, bEvaluateInPreroll)));
 		if (PrerollProperty)
 		{
-			AddBoolPropertyMenuItem(MenuBuilder, Track, PrerollProperty, PrerollProperty->ContainerPtrToValuePtr<void>(&Track->EvalOptions));
+			AddBoolPropertyMenuItem(MenuBuilder, CanExecute, Track, PrerollProperty, PrerollProperty->ContainerPtrToValuePtr<void>(&Track->EvalOptions));
 		}
 
 		const UBoolProperty* PostrollProperty = Cast<UBoolProperty>(EvalOptionsStruct->FindPropertyByName(GET_MEMBER_NAME_CHECKED(FMovieSceneTrackEvalOptions, bEvaluateInPostroll)));
 		if (PostrollProperty)
 		{
-			AddBoolPropertyMenuItem(MenuBuilder, Track, PostrollProperty, PostrollProperty->ContainerPtrToValuePtr<void>(&Track->EvalOptions));
+			AddBoolPropertyMenuItem(MenuBuilder, CanExecute, Track, PostrollProperty, PostrollProperty->ContainerPtrToValuePtr<void>(&Track->EvalOptions));
 		}
 	}
 	MenuBuilder.EndSection();
@@ -203,6 +206,7 @@ TSharedRef<SWidget> FSequencerTrackNode::GetCustomOutlinerContent()
 					SNew(SBox)
 					.WidthOverride(100)
 					.HAlign(HAlign_Left)
+					.IsEnabled(!GetSequencer().IsReadOnly())
 					[
 						KeyAreas[0]->CreateKeyEditor(&GetSequencer())
 					]
@@ -243,7 +247,7 @@ TSharedRef<SWidget> FSequencerTrackNode::GetCustomOutlinerContent()
 		FBuildEditWidgetParams Params;
 		Params.NodeIsHovered = TAttribute<bool>::Create(TAttribute<bool>::FGetter::CreateSP(this, &FSequencerDisplayNode::IsHovered));
 
-		TSharedPtr<SWidget> Widget = AssociatedEditor.BuildOutlinerEditWidget(ObjectBinding, AssociatedTrack.Get(), Params);
+		TSharedPtr<SWidget> Widget = GetSequencer().IsReadOnly() ? SNullWidget::NullWidget : AssociatedEditor.BuildOutlinerEditWidget(ObjectBinding, AssociatedTrack.Get(), Params);
 
 		TSharedRef<SHorizontalBox> BoxPanel = SNew(SHorizontalBox);
 
diff --git a/Engine/Source/Editor/Sequencer/Private/LevelEditorSequencerIntegration.cpp b/Engine/Source/Editor/Sequencer/Private/LevelEditorSequencerIntegration.cpp
new file mode 100644
--- /dev/null
+++ b/Engine/Source/Editor/Sequencer/Private/LevelEditorSequencerIntegration.cpp
@@ -0,0 +1,957 @@
+// Copyright 1998-2016 Epic Games, Inc. All Rights Reserved.
+
+#include "LevelEditorSequencerIntegration.h"
+#include "SequencerEdMode.h"
+#include "SlateIconFinder.h"
+#include "PropertyHandle.h"
+#include "IDetailKeyframeHandler.h"
+#include "GameDelegates.h"
+#include "Editor/PropertyEditor/Public/PropertyEditorModule.h"
+#include "Editor/LevelEditor/Public/ILevelEditor.h"
+#include "Editor/LevelEditor/Public/ILevelViewport.h"
+#include "Editor/LevelEditor/Public/LevelEditor.h"
+#include "Framework/Application/SlateApplication.h"
+#include "IDetailsView.h"
+#include "ISequencer.h"
+#include "KeyPropertyParams.h"
+#include "Classes/Engine/Selection.h"
+#include "Sequencer.h"
+#include "EditorModeManager.h"
+#include "SequencerCommands.h"
+#include "MovieSceneSequence.h"
+#include "SequencerSettings.h"
+#include "LevelEditorViewport.h"
+#include "Modules/ModuleManager.h"
+#include "MultiBox/MultiBoxBuilder.h"
+#include "PropertyEditorModule.h"
+#include "Widgets/Input/SButton.h"
+#include "Widgets/Input/SComboButton.h"
+#include "Widgets/Layout/SBorder.h"
+#include "Widgets/Layout/SBox.h"
+#include "Widgets/Text/STextBlock.h"
+
+
+#define LOCTEXT_NAMESPACE "LevelEditorSequencerIntegration"
+
+class FDetailKeyframeHandlerWrapper : public IDetailKeyframeHandler
+{
+public:
+
+	void Add(TWeakPtr<ISequencer> InSequencer)
+	{
+		Sequencers.Add(InSequencer);
+	}
+
+	void Remove(TWeakPtr<ISequencer> InSequencer)
+	{
+		Sequencers.Remove(InSequencer);
+	}
+
+	virtual bool IsPropertyKeyable(UClass* InObjectClass, const IPropertyHandle& InPropertyHandle) const
+	{
+		FCanKeyPropertyParams CanKeyPropertyParams(InObjectClass, InPropertyHandle);
+
+		for (const TWeakPtr<ISequencer>& WeakSequencer : Sequencers)
+		{
+			TSharedPtr<ISequencer> Sequencer = WeakSequencer.Pin();
+			if (Sequencer.IsValid() && Sequencer->CanKeyProperty(CanKeyPropertyParams))
+			{
+				return true;
+			}
+		}
+		return false;
+	}
+
+	virtual bool IsPropertyKeyingEnabled() const
+	{
+		for (const TWeakPtr<ISequencer>& WeakSequencer : Sequencers)
+		{
+			TSharedPtr<ISequencer> Sequencer = WeakSequencer.Pin();
+			if (Sequencer.IsValid() && Sequencer->GetFocusedMovieSceneSequence())
+			{
+				return true;
+			}
+		}
+		return false;
+	}
+
+	virtual void OnKeyPropertyClicked(const IPropertyHandle& KeyedPropertyHandle)
+	{
+		TArray<UObject*> Objects;
+		KeyedPropertyHandle.GetOuterObjects( Objects );
+		FKeyPropertyParams KeyPropertyParams(Objects, KeyedPropertyHandle, ESequencerKeyMode::ManualKeyForced);
+
+		for (const TWeakPtr<ISequencer>& WeakSequencer : Sequencers)
+		{
+			TSharedPtr<ISequencer> Sequencer = WeakSequencer.Pin();
+			if (Sequencer.IsValid())
+			{
+				Sequencer->KeyProperty(KeyPropertyParams);
+			}
+		}
+	}
+
+private:
+	TArray<TWeakPtr<ISequencer>> Sequencers;
+};
+
+static const FName DetailsTabIdentifiers[] = { "LevelEditorSelectionDetails", "LevelEditorSelectionDetails2", "LevelEditorSelectionDetails3", "LevelEditorSelectionDetails4" };
+
+FLevelEditorSequencerIntegration::FLevelEditorSequencerIntegration()
+{
+	KeyFrameHandler = MakeShared<FDetailKeyframeHandlerWrapper>();
+}
+
+FLevelEditorSequencerIntegration& FLevelEditorSequencerIntegration::Get()
+{
+	static FLevelEditorSequencerIntegration Singleton;
+	return Singleton;
+}
+
+void FLevelEditorSequencerIntegration::IterateAllSequencers(TFunctionRef<void(FSequencer&, const FLevelEditorSequencerIntegrationOptions& Options)> It) const
+{
+	for (const FSequencerAndOptions& SequencerAndOptions : BoundSequencers)
+	{
+		TSharedPtr<FSequencer> Pinned = SequencerAndOptions.Sequencer.Pin();
+		if (Pinned.IsValid())
+		{
+			It(*Pinned, SequencerAndOptions.Options);
+		}
+	}
+}
+
+void FLevelEditorSequencerIntegration::Initialize()
+{
+	AcquiredResources.Release();
+
+	// Register for saving the level so that the state of the scene can be restored before saving and updated after saving.
+	{
+		FDelegateHandle Handle = FEditorDelegates::PreSaveWorld.AddRaw(this, &FLevelEditorSequencerIntegration::OnPreSaveWorld);
+		AcquiredResources.Add([=]{ FEditorDelegates::PreSaveWorld.Remove(Handle); });
+	}
+	{
+		FDelegateHandle Handle = FEditorDelegates::PostSaveWorld.AddRaw(this, &FLevelEditorSequencerIntegration::OnPostSaveWorld);
+		AcquiredResources.Add([=]{ FEditorDelegates::PostSaveWorld.Remove(Handle); });
+	}
+	{
+		FDelegateHandle Handle = FEditorDelegates::PreBeginPIE.AddRaw(this, &FLevelEditorSequencerIntegration::OnPreBeginPIE);
+		AcquiredResources.Add([=]{ FEditorDelegates::PreBeginPIE.Remove(Handle); });
+	}
+	{
+		FDelegateHandle Handle = FEditorDelegates::EndPIE.AddRaw(this, &FLevelEditorSequencerIntegration::OnEndPIE);
+		AcquiredResources.Add([=]{ FEditorDelegates::EndPIE.Remove(Handle); });
+	}
+	{
+		FDelegateHandle Handle = FGameDelegates::Get().GetEndPlayMapDelegate().AddRaw(this, &FLevelEditorSequencerIntegration::OnEndPlayMap);
+		AcquiredResources.Add([=]{ FGameDelegates::Get().GetEndPlayMapDelegate().Remove(Handle); });
+	}
+	{
+		FDelegateHandle Handle = FWorldDelegates::LevelAddedToWorld.AddRaw(this, &FLevelEditorSequencerIntegration::OnLevelAdded);
+		AcquiredResources.Add([=]{ FWorldDelegates::LevelAddedToWorld.Remove(Handle); });
+	}
+	{
+		FDelegateHandle Handle = FWorldDelegates::LevelRemovedFromWorld.AddRaw(this, &FLevelEditorSequencerIntegration::OnLevelRemoved);
+		AcquiredResources.Add([=]{ FWorldDelegates::LevelRemovedFromWorld.Remove(Handle); });
+	}
+	{
+		FDelegateHandle Handle = FEditorDelegates::NewCurrentLevel.AddRaw(this, &FLevelEditorSequencerIntegration::OnNewCurrentLevel);
+		AcquiredResources.Add([=]{ FEditorDelegates::NewCurrentLevel.Remove(Handle); });
+	}
+	{
+		FDelegateHandle Handle = FEditorDelegates::OnMapOpened.AddRaw(this, &FLevelEditorSequencerIntegration::OnMapOpened);
+		AcquiredResources.Add([=]{ FEditorDelegates::OnMapOpened.Remove(Handle); });
+	}
+	{
+		FDelegateHandle Handle = FEditorDelegates::OnNewActorsDropped.AddRaw(this, &FLevelEditorSequencerIntegration::OnNewActorsDropped);
+		AcquiredResources.Add([=]{ FEditorDelegates::OnNewActorsDropped.Remove(Handle); });
+	}
+
+	{
+		FDelegateHandle Handle = USelection::SelectionChangedEvent.AddRaw( this, &FLevelEditorSequencerIntegration::OnActorSelectionChanged );
+		AcquiredResources.Add([=]{ USelection::SelectionChangedEvent.Remove(Handle); });
+	}
+
+	{
+		// Hook into the editor's mechanism for checking whether we need live capture of PIE/SIE actor state
+		FDelegateHandle Handle = GEditor->GetActorRecordingState().AddRaw(this, &FLevelEditorSequencerIntegration::GetActorRecordingState);
+		AcquiredResources.Add([=]{ GEditor->GetActorRecordingState().Remove(Handle); });
+	}
+
+	AddLevelViewportMenuExtender();
+	ActivateDetailKeyframeHandler();
+	AttachTransportControlsToViewports();
+	ActivateSequencerEditorMode();
+
+	{
+		FPropertyEditorModule& EditModule = FModuleManager::Get().GetModuleChecked<FPropertyEditorModule>("PropertyEditor");
+
+		FDelegateHandle Handle = EditModule.OnPropertyEditorOpened().AddRaw(this, &FLevelEditorSequencerIntegration::OnPropertyEditorOpened);
+		AcquiredResources.Add(
+			[=]{
+				FPropertyEditorModule* EditModulePtr = FModuleManager::Get().GetModulePtr<FPropertyEditorModule>("PropertyEditor");
+				if (EditModulePtr)
+				{
+					EditModulePtr->OnPropertyEditorOpened().Remove(Handle);
+				}
+			}
+		);
+	}
+}
+
+void FLevelEditorSequencerIntegration::GetActorRecordingState( bool& bIsRecording ) const
+{
+	IterateAllSequencers(
+		[&](FSequencer& In, const FLevelEditorSequencerIntegrationOptions& Options)
+		{
+			if (In.IsRecordingLive())
+			{
+				bIsRecording = true;
+			}
+		}
+	);
+}
+
+void FLevelEditorSequencerIntegration::OnPreSaveWorld(uint32 SaveFlags, class UWorld* World)
+{
+	// Restore the saved state so that the level save can save that instead of the animated state.
+	IterateAllSequencers(
+		[](FSequencer& In, const FLevelEditorSequencerIntegrationOptions& Options)
+		{
+			if (Options.bRequiresLevelEvents)
+			{
+				In.RestorePreAnimatedState();
+			}
+		}
+	);
+}
+
+void FLevelEditorSequencerIntegration::OnPostSaveWorld(uint32 SaveFlags, class UWorld* World, bool bSuccess)
+{
+	// Reset the time after saving so that an update will be triggered to put objects back to their animated state.
+	IterateAllSequencers(
+		[](FSequencer& In, const FLevelEditorSequencerIntegrationOptions& Options)
+		{
+			if (Options.bRequiresLevelEvents)
+			{
+				In.ForceEvaluate();
+			}
+		}
+	);
+}
+
+void FLevelEditorSequencerIntegration::OnNewCurrentLevel()
+{
+	ActivateSequencerEditorMode();
+}
+
+void FLevelEditorSequencerIntegration::OnMapOpened(const FString& Filename, bool bLoadAsTemplate)
+{
+	ActivateSequencerEditorMode();
+}
+
+void FLevelEditorSequencerIntegration::OnLevelAdded(ULevel* InLevel, UWorld* InWorld)
+{
+	IterateAllSequencers(
+		[](FSequencer& In, const FLevelEditorSequencerIntegrationOptions& Options)
+		{
+			if (Options.bRequiresLevelEvents)
+			{
+				In.State.ClearObjectCaches();
+			}
+		}
+	);
+}
+
+void FLevelEditorSequencerIntegration::OnLevelRemoved(ULevel* InLevel, UWorld* InWorld)
+{
+	IterateAllSequencers(
+		[](FSequencer& In, const FLevelEditorSequencerIntegrationOptions& Options)
+		{
+			if (Options.bRequiresLevelEvents)
+			{
+				In.State.ClearObjectCaches();
+			}
+		}
+	);
+}
+
+void FLevelEditorSequencerIntegration::OnActorSelectionChanged( UObject* )
+{
+	IterateAllSequencers(
+		[](FSequencer& In, const FLevelEditorSequencerIntegrationOptions& Options)
+		{
+			if (Options.bRequiresActorEvents)
+			{
+				In.ExternalSelectionHasChanged();
+			}
+		}
+	);
+}
+
+void FLevelEditorSequencerIntegration::OnNewActorsDropped(const TArray<UObject*>& DroppedObjects, const TArray<AActor*>& DroppedActors)
+{
+	IterateAllSequencers(
+		[&](FSequencer& In, const FLevelEditorSequencerIntegrationOptions& Options)
+		{
+			if (Options.bRequiresActorEvents)
+			{
+				In.OnNewActorsDropped(DroppedObjects, DroppedActors);
+			}
+		}
+	);
+}
+
+void FLevelEditorSequencerIntegration::ActivateSequencerEditorMode()
+{
+	// Release the sequencer mode if we already enabled it
+	FName ResourceName("SequencerMode");
+	AcquiredResources.Release(ResourceName);
+
+	GLevelEditorModeTools().ActivateMode( FSequencerEdMode::EM_SequencerMode );
+	FSequencerEdMode* SequencerEdMode = (FSequencerEdMode*)GLevelEditorModeTools().GetActiveMode(FSequencerEdMode::EM_SequencerMode);
+
+	for (const FSequencerAndOptions& SequencerAndOptions : BoundSequencers)
+	{
+		TSharedPtr<FSequencer> Pinned = SequencerAndOptions.Sequencer.Pin();
+		if (Pinned.IsValid())
+		{
+			SequencerEdMode->AddSequencer(Pinned);
+		}
+	}
+
+	// Acquire the resource, which allows us to deactivate the mode later
+	AcquiredResources.Add(
+		ResourceName,
+		[] {
+			if (GLevelEditorModeTools().IsModeActive(FSequencerEdMode::EM_SequencerMode))
+			{
+				GLevelEditorModeTools().DeactivateMode(FSequencerEdMode::EM_SequencerMode);
+			}
+		}
+	);
+}
+
+
+void FLevelEditorSequencerIntegration::OnPreBeginPIE(bool bIsSimulating)
+{
+	IterateAllSequencers(
+		[](FSequencer& In, const FLevelEditorSequencerIntegrationOptions& Options)
+		{
+			if (Options.bRequiresLevelEvents)
+			{
+				In.RestorePreAnimatedState();
+				In.State.ClearObjectCaches();
+			}
+		}
+	);
+}
+
+void FLevelEditorSequencerIntegration::OnEndPlayMap()
+{
+	IterateAllSequencers(
+		[](FSequencer& In, const FLevelEditorSequencerIntegrationOptions& Options)
+		{
+			if (Options.bRequiresLevelEvents)
+			{
+				// Update and clear any stale bindings 
+				In.State.ClearObjectCaches();
+			}
+		}
+	);
+}
+
+void FLevelEditorSequencerIntegration::OnEndPIE(bool bIsSimulating)
+{
+	OnEndPlayMap();
+}
+
+void FLevelEditorSequencerIntegration::AddLevelViewportMenuExtender()
+{
+	typedef FLevelEditorModule::FLevelViewportMenuExtender_SelectedActors DelegateType;
+
+	FLevelEditorModule& LevelEditorModule = FModuleManager::Get().LoadModuleChecked<FLevelEditorModule>("LevelEditor");
+	auto& MenuExtenders = LevelEditorModule.GetAllLevelViewportContextMenuExtenders();
+	MenuExtenders.Add(DelegateType::CreateRaw(this, &FLevelEditorSequencerIntegration::GetLevelViewportExtender));
+
+	FDelegateHandle LevelViewportExtenderHandle = MenuExtenders.Last().GetHandle();
+	AcquiredResources.Add(
+		[LevelViewportExtenderHandle]{
+			FLevelEditorModule* LevelEditorModulePtr = FModuleManager::Get().GetModulePtr<FLevelEditorModule>("LevelEditor");
+			if (LevelEditorModulePtr)
+			{
+				LevelEditorModulePtr->GetAllLevelViewportContextMenuExtenders().RemoveAll([=](const DelegateType& In){ return In.GetHandle() == LevelViewportExtenderHandle; });
+			}
+		}
+	);
+}
+
+TSharedRef<FExtender> FLevelEditorSequencerIntegration::GetLevelViewportExtender(const TSharedRef<FUICommandList> CommandList, const TArray<AActor*> InActors)
+{
+	TSharedRef<FExtender> Extender = MakeShareable(new FExtender);
+
+	FText ActorName;
+	if (InActors.Num() == 1)
+	{
+		ActorName = FText::Format(LOCTEXT("ActorNameSingular", "\"{0}\""), FText::FromString(InActors[0]->GetActorLabel()));
+	}
+	else if (InActors.Num() > 1)
+	{
+		ActorName = FText::Format(LOCTEXT("ActorNamePlural", "{0} Actors"), FText::AsNumber(InActors.Num()));
+	}
+
+	TSharedRef<FUICommandList> LevelEditorCommandBindings  = FModuleManager::GetModuleChecked<FLevelEditorModule>(TEXT("LevelEditor")).GetGlobalLevelEditorActions();
+	Extender->AddMenuExtension("ActorControl", EExtensionHook::After, LevelEditorCommandBindings, FMenuExtensionDelegate::CreateLambda(
+		[ActorName](FMenuBuilder& MenuBuilder){
+			MenuBuilder.BeginSection("SequenceRecording", LOCTEXT("SequenceRecordingHeading", "Sequence Recording"));
+			MenuBuilder.AddMenuEntry(FSequencerCommands::Get().RecordSelectedActors, NAME_None, FText::Format(LOCTEXT("RecordSelectedActorsText", "Record {0} In Sequencer"), ActorName));
+			MenuBuilder.EndSection();
+		}
+	));
+
+	return Extender;
+}
+
+void FLevelEditorSequencerIntegration::ActivateDetailKeyframeHandler()
+{
+	FName DetailHandlerName("DetailHandler");
+
+	AcquiredResources.Release(DetailHandlerName);
+
+	// Add sequencer detail keyframe handler
+	FPropertyEditorModule& EditModule = FModuleManager::Get().GetModuleChecked<FPropertyEditorModule>("PropertyEditor");
+
+	for (const FName& DetailsTabIdentifier : DetailsTabIdentifiers)
+	{
+		TSharedPtr<IDetailsView> DetailsView = EditModule.FindDetailView(DetailsTabIdentifier);
+		if(DetailsView.IsValid())
+		{
+			DetailsView->SetKeyframeHandler(KeyFrameHandler);
+		}
+	}
+
+	FDelegateHandle OnPropertyEditorOpenedHandle = EditModule.OnPropertyEditorOpened().AddRaw(this, &FLevelEditorSequencerIntegration::OnPropertyEditorOpened);
+
+	auto DeactivateDetailKeyframeHandler =
+		[this, OnPropertyEditorOpenedHandle]
+		{
+			FPropertyEditorModule* EditModulePtr = FModuleManager::Get().GetModulePtr<FPropertyEditorModule>("PropertyEditor");
+			if (!EditModulePtr)
+			{
+				return;
+			}
+
+			EditModulePtr->OnPropertyEditorOpened().Remove(OnPropertyEditorOpenedHandle);
+
+			for (const FName& DetailsTabIdentifier : DetailsTabIdentifiers)
+			{
+				TSharedPtr<IDetailsView> DetailsView = EditModulePtr->FindDetailView(DetailsTabIdentifier);
+				if (DetailsView.IsValid())
+				{
+					if (DetailsView->GetKeyframeHandler() == KeyFrameHandler)
+					{
+						DetailsView->SetKeyframeHandler(nullptr);
+					}
+				}
+			}
+		};
+
+	AcquiredResources.Add(DetailHandlerName, DeactivateDetailKeyframeHandler);
+}
+
+void FLevelEditorSequencerIntegration::OnPropertyEditorOpened()
+{
+	ActivateDetailKeyframeHandler();
+}
+
+void FLevelEditorSequencerIntegration::BindLevelEditorCommands()
+{
+	FLevelEditorModule& LevelEditor = FModuleManager::GetModuleChecked<FLevelEditorModule>(TEXT("LevelEditor"));
+	LevelEditor.GetGlobalLevelEditorActions()->MapAction(
+		FSequencerCommands::Get().RecordSelectedActors,
+		FExecuteAction::CreateRaw(this, &FLevelEditorSequencerIntegration::RecordSelectedActors));
+
+	AcquiredResources.Add(
+		[]
+		{
+			FLevelEditorModule* LevelEditorPtr = FModuleManager::GetModulePtr<FLevelEditorModule>(TEXT("LevelEditor"));
+			if (LevelEditorPtr)
+			{
+				LevelEditorPtr->GetGlobalLevelEditorActions()->UnmapAction(FSequencerCommands::Get().RecordSelectedActors);
+			}
+		}
+	);
+}
+
+void FLevelEditorSequencerIntegration::RecordSelectedActors()
+{
+	IterateAllSequencers(
+		[&](FSequencer& In, const FLevelEditorSequencerIntegrationOptions& Options)
+		{
+			if (Options.bCanRecord)
+			{
+				In.RecordSelectedActors();
+			}
+		}
+	);
+}
+
+namespace FaderConstants
+{	
+	/** The opacity when we are hovered */
+	const float HoveredOpacity = 1.0f;
+	/** The opacity when we are not hovered */
+	const float NonHoveredOpacity = 0.75f;
+	/** The amount of time spent actually fading in or out */
+	const float FadeTime = 0.15f;
+}
+
+/** Wrapper widget allowing us to fade widgets in and out on hover state */
+class SFader : public SBorder
+{
+public:
+	SLATE_BEGIN_ARGS(SFader)
+		: _Content()
+	{}
+
+	SLATE_DEFAULT_SLOT(FArguments, Content)
+
+	SLATE_END_ARGS()
+
+	void Construct(const FArguments& InArgs)
+	{
+		FadeInSequence = FCurveSequence(0.0f, FaderConstants::FadeTime);
+		FadeOutSequence = FCurveSequence(0.0f, FaderConstants::FadeTime);
+		FadeOutSequence.JumpToEnd();
+
+		SBorder::Construct(SBorder::FArguments()
+			.BorderImage(FCoreStyle::Get().GetBrush("NoBorder"))
+			.Padding(0.0f)
+			.VAlign(VAlign_Center)
+			.ColorAndOpacity(this, &SFader::GetColorAndOpacity)
+			.Content()
+			[
+				InArgs._Content.Widget
+			]);
+	}
+
+	FLinearColor GetColorAndOpacity() const
+	{
+		FLinearColor Color = FLinearColor::White;
+	
+		if(FadeOutSequence.IsPlaying() || !bIsHovered)
+		{
+			Color.A = FMath::Lerp(FaderConstants::HoveredOpacity, FaderConstants::NonHoveredOpacity, FadeOutSequence.GetLerp());
+		}
+		else
+		{
+			Color.A = FMath::Lerp(FaderConstants::NonHoveredOpacity, FaderConstants::HoveredOpacity, FadeInSequence.GetLerp());
+		}
+
+		return Color;
+	}
+
+	virtual void OnMouseEnter(const FGeometry& MyGeometry, const FPointerEvent& MouseEvent) override
+	{
+		if(!FSlateApplication::Get().IsUsingHighPrecisionMouseMovment())
+		{
+			bIsHovered = true;
+			if(FadeOutSequence.IsPlaying())
+			{
+				// Fade out is already playing so just force the fade in curve to the end so we don't have a "pop" 
+				// effect from quickly resetting the alpha
+				FadeInSequence.JumpToEnd();
+			}
+			else
+			{
+				FadeInSequence.Play(AsShared());
+			}
+		}
+	}
+
+	virtual void OnMouseLeave(const FPointerEvent& MouseEvent) override
+	{
+		if(!FSlateApplication::Get().IsUsingHighPrecisionMouseMovment())
+		{
+			bIsHovered = false;
+			FadeOutSequence.Play(AsShared());
+		}
+	}
+
+private:
+	/** Curve sequence for fading out the widget */
+	FCurveSequence FadeOutSequence;
+	/** Curve sequence for fading in the widget */
+	FCurveSequence FadeInSequence;
+};
+
+class SViewportTransportControls : public SCompoundWidget
+{
+public:
+	SLATE_BEGIN_ARGS(SViewportTransportControls){}
+	SLATE_END_ARGS()
+
+	void Construct(const FArguments& InArgs, TSharedPtr<ILevelViewport> Viewport)
+	{
+		check(Viewport.IsValid());
+
+		WeakViewport = Viewport;
+
+		SetVisibility(EVisibility::SelfHitTestInvisible);
+
+		ChildSlot
+		[
+			SNew(SHorizontalBox)
+
+			+SHorizontalBox::Slot()
+			.FillWidth(1)
+			.HAlign(HAlign_Center)
+			.VAlign(VAlign_Bottom)
+			.Padding(4.f)
+			[
+				SNew(SFader)
+				.Content()
+				[
+					SNew(SBorder)
+					.Padding(4.f)
+					.Cursor( EMouseCursor::Default )
+					.BorderImage( FEditorStyle::GetBrush( "ToolPanel.GroupBorder" ) )
+					.Visibility(this, &SViewportTransportControls::GetVisibility)
+					.Content()
+					[
+						SNew(SHorizontalBox)
+
+						+SHorizontalBox::Slot()
+						.AutoWidth()
+						.HAlign(HAlign_Center)
+						.VAlign(VAlign_Center)
+						.Padding(FMargin(0,0,4.f,0))
+						[
+							SNew(SComboButton)
+							.Visibility(this, &SViewportTransportControls::GetComboVisibility)
+							.OnGetMenuContent(this, &SViewportTransportControls::GetBoundSequencerMenu)
+							.ButtonContent()
+							[
+								SNew(STextBlock)
+								.Text(this, &SViewportTransportControls::GetBoundSequencerName)
+							]
+						]
+
+						+SHorizontalBox::Slot()
+						.AutoWidth()
+						.HAlign(HAlign_Center)
+						.VAlign(VAlign_Center)
+						[
+							SAssignNew(ControlContent, SBox)
+						]
+						+SHorizontalBox::Slot()
+						.AutoWidth()
+						.HAlign(HAlign_Center)
+						.VAlign(VAlign_Top)
+						[
+							SNew(SButton)
+							.ButtonStyle(&FEditorStyle::Get().GetWidgetStyle< FButtonStyle >("Sequencer.Transport.CloseButton"))
+							.ToolTipText(LOCTEXT("CloseTransportControlsToolTip", "Hide the transport controls. You can re-enable transport controls from the viewport menu."))
+							.OnClicked(this, &SViewportTransportControls::Close)
+						]
+					]
+				]
+			]
+		];
+	}
+
+	TSharedPtr<FSequencer> GetSequencer() const
+	{
+		return WeakSequencer.Pin();
+	}
+
+	void AssignSequencer(TSharedRef<FSequencer> InSequencer)
+	{
+		WeakSequencer = InSequencer;
+		bool bExtendedControls = false;
+		ControlContent->SetContent(InSequencer->MakeTransportControls(bExtendedControls));
+	}
+
+	EVisibility GetComboVisibility() const
+	{
+		int32 NumSequencers = 0;
+		FLevelEditorSequencerIntegration::Get().IterateAllSequencers(
+			[&NumSequencers](FSequencer& In, const FLevelEditorSequencerIntegrationOptions& Options)
+			{
+				if (UMovieSceneSequence* RootSequence = In.GetRootMovieSceneSequence())
+				{
+					++NumSequencers;
+				}
+			}
+		);
+
+		return NumSequencers == 1 ? EVisibility::Collapsed : EVisibility::Visible;
+	}
+
+	FText GetBoundSequencerName() const
+	{
+		TSharedPtr<FSequencer> Sequencer = WeakSequencer.Pin();
+		if (Sequencer.IsValid() && Sequencer->GetRootMovieSceneSequence())
+		{
+			return Sequencer->GetRootMovieSceneSequence()->GetDisplayName();
+		}
+
+		return LOCTEXT("SelectSequencer", "Choose Sequence...");
+	}
+
+	TSharedRef<SWidget> GetBoundSequencerMenu()
+	{
+		FMenuBuilder MenuBuilder(true, nullptr);
+
+		for (const FLevelEditorSequencerIntegration::FSequencerAndOptions& SequencerAndOptions : FLevelEditorSequencerIntegration::Get().BoundSequencers)
+		{
+			TSharedPtr<FSequencer> PinnedSequencer = SequencerAndOptions.Sequencer.Pin();
+			if (PinnedSequencer.IsValid())
+			{
+				UMovieSceneSequence* RootSequence = PinnedSequencer->GetRootMovieSceneSequence();
+
+				// Be careful not to hold a strong reference in the lamdba below
+				TWeakPtr<FSequencer> ThisWeakSequencer = SequencerAndOptions.Sequencer;
+
+				MenuBuilder.AddMenuEntry(
+					RootSequence->GetDisplayName(),
+					FText(),
+					FSlateIconFinder::FindIconForClass(RootSequence->GetClass(), "MovieSceneSequence"),
+					FUIAction(
+						FExecuteAction::CreateLambda(
+							[this, ThisWeakSequencer]
+							{
+								TSharedPtr<FSequencer> LocalPinnedSequence = ThisWeakSequencer.Pin();
+								if (LocalPinnedSequence.IsValid())
+								{
+									AssignSequencer(LocalPinnedSequence.ToSharedRef());
+								}
+							}
+						)
+					));
+			}
+		}
+
+		return MenuBuilder.MakeWidget();
+	}
+
+	FReply Close()
+	{
+		TSharedPtr<FSequencer> Sequencer = WeakSequencer.Pin();
+		if (Sequencer.IsValid())
+		{
+			Sequencer->GetSettings()->SetShowViewportTransportControls(!Sequencer->GetSettings()->GetShowViewportTransportControls());
+		}
+		return FReply::Handled();
+	}
+
+	EVisibility GetVisibility() const
+	{
+		TSharedPtr<ILevelViewport> Viewport = WeakViewport.Pin();
+		TSharedPtr<FSequencer> Sequencer = WeakSequencer.Pin();
+		FLevelEditorViewportClient* ViewportClient = Viewport.IsValid() ? &Viewport->GetLevelViewportClient() : nullptr;
+
+		return (
+				Sequencer.IsValid() &&
+				ViewportClient &&
+				Sequencer->GetSettings()->GetShowViewportTransportControls() &&
+				ViewportClient->ViewportType == LVT_Perspective &&
+				ViewportClient->AllowsCinematicPreview()
+			) ? EVisibility::Visible : EVisibility::Collapsed;
+	}
+
+	virtual bool SupportsKeyboardFocus() const override
+	{
+		// Transport controls in the viewport need to have something that is focusable to prevent mouse input dropping through to the viewport.
+		// We don't want the buttons themselves to be focusable, so we just add them to a parent box that is.
+		return true;
+	}
+
+private:
+
+	TSharedPtr<SBox> ControlContent;
+	TWeakPtr<ILevelViewport> WeakViewport;
+	TWeakPtr<FSequencer> WeakSequencer;
+};
+
+void FLevelEditorSequencerIntegration::AttachTransportControlsToViewports()
+{
+	FLevelEditorModule* Module = FModuleManager::Get().LoadModulePtr<FLevelEditorModule>("LevelEditor");
+	
+	if (Module == nullptr)
+	{
+		return;
+	}
+
+	// register level editor viewport menu extenders
+	{
+		FLevelEditorModule::FLevelEditorMenuExtender ViewMenuExtender = FLevelEditorModule::FLevelEditorMenuExtender::CreateRaw(this, &FLevelEditorSequencerIntegration::OnExtendLevelEditorViewMenu);
+		Module->GetAllLevelViewportOptionsMenuExtenders().Add(ViewMenuExtender);
+
+		FDelegateHandle Handle = Module->GetAllLevelViewportOptionsMenuExtenders().Last().GetHandle();
+		AcquiredResources.Add(
+			[Handle]{
+				FLevelEditorModule* LevelEditorModule = FModuleManager::Get().GetModulePtr<FLevelEditorModule>("LevelEditor");
+				if (LevelEditorModule)
+				{
+					LevelEditorModule->GetAllLevelViewportOptionsMenuExtenders().RemoveAll([Handle](const FLevelEditorModule::FLevelEditorMenuExtender& In){ return In.GetHandle() == Handle; });
+				}
+			}
+		);
+	}
+
+	TSharedPtr<ILevelEditor> LevelEditor = Module->GetFirstLevelEditor();
+
+	for (TSharedPtr<ILevelViewport> LevelViewport : LevelEditor->GetViewports())
+	{
+		TSharedRef<SViewportTransportControls> TransportControl = SNew(SViewportTransportControls, LevelViewport);
+		LevelViewport->AddOverlayWidget(TransportControl);
+		TransportControls.Add(FTransportControl{ LevelViewport, TransportControl });
+	}
+
+	AcquiredResources.Add([=]{ this->DetachTransportControlsFromViewports(); });
+}
+
+void FLevelEditorSequencerIntegration::SetViewportTransportControlsVisibility(bool bVisible)
+{
+	IterateAllSequencers(
+		[=](FSequencer& In, const FLevelEditorSequencerIntegrationOptions& Options)
+		{
+			In.GetSettings()->SetShowViewportTransportControls(bVisible);
+		}
+	);
+}
+
+bool FLevelEditorSequencerIntegration::GetViewportTransportControlsVisibility() const
+{
+	bool bVisible = false;
+
+	IterateAllSequencers(
+		[&](FSequencer& In, const FLevelEditorSequencerIntegrationOptions& Options)
+		{
+			if (In.GetSettings()->GetShowViewportTransportControls())
+			{
+				bVisible = true;
+			}
+		}
+	);
+
+	return bVisible;
+}
+
+void FLevelEditorSequencerIntegration::DetachTransportControlsFromViewports()
+{
+	for (const FTransportControl& Control : TransportControls)
+	{
+		TSharedPtr<ILevelViewport> Viewport = Control.Viewport.Pin();
+		if (Viewport.IsValid())
+		{
+			Viewport->RemoveOverlayWidget(Control.Widget.ToSharedRef());
+		}
+	}
+	TransportControls.Reset();
+}
+
+TSharedRef<FExtender> FLevelEditorSequencerIntegration::OnExtendLevelEditorViewMenu(const TSharedRef<FUICommandList> CommandList)
+{
+	TSharedRef<FExtender> Extender(new FExtender());
+
+	Extender->AddMenuExtension(
+		"LevelViewportViewportOptions2",
+		EExtensionHook::First,
+		NULL,
+		FMenuExtensionDelegate::CreateRaw(this, &FLevelEditorSequencerIntegration::CreateTransportToggleMenuEntry));
+
+	return Extender;
+}
+
+void FLevelEditorSequencerIntegration::CreateTransportToggleMenuEntry(FMenuBuilder& MenuBuilder)
+{
+	MenuBuilder.AddMenuEntry(
+		LOCTEXT("ShowTransportControls", "Show Transport Controls"),
+		LOCTEXT("ShowTransportControlsToolTip", "Show or hide the Sequencer transport controls when a sequence is active."),
+		FSlateIcon(),
+		FUIAction(
+			FExecuteAction::CreateLambda([this](){ SetViewportTransportControlsVisibility(!GetViewportTransportControlsVisibility()); }),
+			FCanExecuteAction(), 
+			FGetActionCheckState::CreateLambda([this](){ return GetViewportTransportControlsVisibility() ? ECheckBoxState::Checked : ECheckBoxState::Unchecked; })),
+		NAME_None,
+		EUserInterfaceActionType::ToggleButton);
+}
+
+void FLevelEditorSequencerIntegration::AddSequencer(TSharedRef<ISequencer> InSequencer, const FLevelEditorSequencerIntegrationOptions& Options)
+{
+	if (!BoundSequencers.Num())
+	{
+		Initialize();
+	}
+
+	KeyFrameHandler->Add(InSequencer);
+	
+	auto DerivedSequencerPtr = StaticCastSharedRef<FSequencer>(InSequencer);
+	BoundSequencers.Add(FSequencerAndOptions{ DerivedSequencerPtr, Options });
+
+	FSequencerEdMode* SequencerEdMode = (FSequencerEdMode*)(GLevelEditorModeTools().GetActiveMode(FSequencerEdMode::EM_SequencerMode));
+	if (SequencerEdMode)
+	{
+		SequencerEdMode->AddSequencer(DerivedSequencerPtr);
+	}
+
+	// Set up any transport controls
+	for (const FTransportControl& Control : TransportControls)
+	{
+		if (!Control.Widget->GetSequencer().IsValid())
+		{
+			Control.Widget->AssignSequencer(DerivedSequencerPtr);
+		}
+	}
+}
+
+void FLevelEditorSequencerIntegration::OnSequencerReceivedFocus(TSharedRef<ISequencer> InSequencer)
+{
+	FSequencerEdMode* SequencerEdMode = (FSequencerEdMode*)(GLevelEditorModeTools().GetActiveMode(FSequencerEdMode::EM_SequencerMode));
+	if (SequencerEdMode)
+	{
+		SequencerEdMode->OnSequencerReceivedFocus(StaticCastSharedRef<FSequencer>(InSequencer));
+	}
+
+}
+
+void FLevelEditorSequencerIntegration::RemoveSequencer(TSharedRef<ISequencer> InSequencer)
+{
+	BoundSequencers.RemoveAll(
+		[=](const FSequencerAndOptions& In)
+		{
+			return In.Sequencer == InSequencer;
+		}
+	);
+
+	FSequencerEdMode* SequencerEdMode = (FSequencerEdMode*)(GLevelEditorModeTools().GetActiveMode(FSequencerEdMode::EM_SequencerMode));
+	if (SequencerEdMode)
+	{
+		SequencerEdMode->RemoveSequencer(StaticCastSharedRef<FSequencer>(InSequencer));
+	}
+
+	KeyFrameHandler->Remove(InSequencer);
+
+	auto IsValidSequencer = [](const FSequencerAndOptions& In){ return In.Sequencer.IsValid(); };
+	if (FSequencerAndOptions* FirstValidSequencer = BoundSequencers.FindByPredicate(IsValidSequencer))
+	{
+		TSharedPtr<FSequencer> SequencerPtr = FirstValidSequencer->Sequencer.Pin();
+		check(SequencerPtr.IsValid());
+
+		// Assign any transport controls
+		for (const FTransportControl& Control : TransportControls)
+		{
+			if (!Control.Widget->GetSequencer().IsValid())
+			{
+				Control.Widget->AssignSequencer(SequencerPtr.ToSharedRef());
+			}
+		}
+	}
+	else
+	{
+		AcquiredResources.Release();
+	}
+}
+
+#undef LOCTEXT_NAMESPACE
\ No newline at end of file
diff --git a/Engine/Source/Editor/Sequencer/Private/SAnimationOutlinerTreeNode.cpp b/Engine/Source/Editor/Sequencer/Private/SAnimationOutlinerTreeNode.cpp
--- a/Engine/Source/Editor/Sequencer/Private/SAnimationOutlinerTreeNode.cpp
+++ b/Engine/Source/Editor/Sequencer/Private/SAnimationOutlinerTreeNode.cpp
@@ -232,6 +232,7 @@ void SAnimationOutlinerTreeNode::Construct( const FArguments& InArgs, TSharedRef
 				.VAlign(VAlign_Fill)
 				.HasDownArrow(false)
 				.IsFocusable(false)
+				.IsEnabled(!DisplayNode->GetSequencer().IsReadOnly())
 				.ButtonStyle(FEditorStyle::Get(), "Sequencer.AnimationOutliner.ColorStrip")
 				.OnGetMenuContent(this, &SAnimationOutlinerTreeNode::OnGetColorPicker)
 				.OnMenuOpenChanged_Lambda([](bool bIsOpen){
@@ -438,7 +439,7 @@ FText SAnimationOutlinerTreeNode::GetDisplayName() const
 
 bool SAnimationOutlinerTreeNode::HandleNodeLabelCanEdit() const
 {
-	return DisplayNode->CanRenameNode();
+	return !DisplayNode->GetSequencer().IsReadOnly() && DisplayNode->CanRenameNode();
 }
 
 
diff --git a/Engine/Source/Editor/Sequencer/Private/SKeyNavigationButtons.h b/Engine/Source/Editor/Sequencer/Private/SKeyNavigationButtons.h
--- a/Engine/Source/Editor/Sequencer/Private/SKeyNavigationButtons.h
+++ b/Engine/Source/Editor/Sequencer/Private/SKeyNavigationButtons.h
@@ -79,6 +79,7 @@ public:
 				.Padding(0)
 				.BorderImage(NoBorder)
 				.ColorAndOpacity(HoverTint)
+				.IsEnabled(!InDisplayNode->GetSequencer().IsReadOnly())
 				[
 					SNew(SButton)
 					.ButtonStyle(FEditorStyle::Get(), "FlatButton")
diff --git a/Engine/Source/Editor/Sequencer/Private/SSequencer.cpp b/Engine/Source/Editor/Sequencer/Private/SSequencer.cpp
--- a/Engine/Source/Editor/Sequencer/Private/SSequencer.cpp
+++ b/Engine/Source/Editor/Sequencer/Private/SSequencer.cpp
@@ -144,6 +144,8 @@ void SSequencer::Construct(const FArguments& InArgs, TSharedRef<FSequencer> InSe
 	OnSelectionRangeBeginDrag = InArgs._OnSelectionRangeBeginDrag;
 	OnSelectionRangeEndDrag = InArgs._OnSelectionRangeEndDrag;
 
+	OnReceivedFocus = InArgs._OnReceivedFocus;
+
 	FTimeSliderArgs TimeSliderArgs;
 	{
 		TimeSliderArgs.ViewRange = InArgs._ViewRange;
@@ -158,6 +160,8 @@ void SSequencer::Construct(const FArguments& InArgs, TSharedRef<FSequencer> InSe
 		TimeSliderArgs.OnSelectionRangeEndDrag = OnSelectionRangeEndDrag;
 		TimeSliderArgs.OnViewRangeChanged = InArgs._OnViewRangeChanged;
 		TimeSliderArgs.OnClampRangeChanged = InArgs._OnClampRangeChanged;
+		TimeSliderArgs.IsPlaybackRangeLocked = InArgs._IsPlaybackRangeLocked;
+		TimeSliderArgs.OnTogglePlaybackRangeLocked = InArgs._OnTogglePlaybackRangeLocked;
 		TimeSliderArgs.ScrubPosition = InArgs._ScrubPosition;
 		TimeSliderArgs.OnBeginScrubberMovement = InArgs._OnBeginScrubbing;
 		TimeSliderArgs.OnEndScrubberMovement = InArgs._OnEndScrubbing;
@@ -669,6 +673,11 @@ void SSequencer::HandleOutlinerNodeSelectionChanged()
 
 TSharedRef<SWidget> SSequencer::MakeAddButton()
 {
+	if (SequencerPtr.Pin()->IsReadOnly())
+	{
+		return SNullWidget::NullWidget;
+	}
+
 	return SNew(SComboButton)
 	.OnGetMenuContent(this, &SSequencer::MakeAddMenu)
 	.ButtonStyle(FEditorStyle::Get(), "FlatButton.Success")
@@ -713,7 +722,9 @@ TSharedRef<SWidget> SSequencer::MakeAddButton()
 TSharedRef<SWidget> SSequencer::MakeToolBar()
 {
 	FToolBarBuilder ToolBarBuilder( SequencerPtr.Pin()->GetCommandBindings(), FMultiBoxCustomization::None, TSharedPtr<FExtender>(), Orient_Horizontal, true);
-	
+
+	const bool bIsReadOnly = SequencerPtr.Pin()->IsReadOnly();
+
 	ToolBarBuilder.BeginSection("Base Commands");
 	{
 		// General 
@@ -735,6 +746,7 @@ TSharedRef<SWidget> SSequencer::MakeToolBar()
 				FSlateIcon(FEditorStyle::GetStyleSetName(), "Sequencer.SaveAs")
 			);
 
+			ToolBarBuilder.AddToolBarButton( FSequencerCommands::Get().RestoreAnimatedState );
 			//ToolBarBuilder.AddToolBarButton( FSequencerCommands::Get().DiscardChanges );
 			ToolBarBuilder.AddToolBarButton( FSequencerCommands::Get().FindInContentBrowser );
 			ToolBarBuilder.AddToolBarButton( FSequencerCommands::Get().CreateCamera );
@@ -747,56 +759,75 @@ TSharedRef<SWidget> SSequencer::MakeToolBar()
 			FOnGetContent::CreateSP(this, &SSequencer::MakeGeneralMenu),
 			LOCTEXT("GeneralOptions", "General Options"),
 			LOCTEXT("GeneralOptionsToolTip", "General Options"),
-			FSlateIcon(FEditorStyle::GetStyleSetName(), "Sequencer.General")
+			FSlateIcon(FEditorStyle::GetStyleSetName(), "Sequencer.GeneralOptions")
 		);
 
-		ToolBarBuilder.AddSeparator();
-
-		if( SequencerPtr.Pin()->IsLevelEditorSequencer() )
+		if (!bIsReadOnly)
 		{
-			TAttribute<FSlateIcon> KeyAllIcon;
-			KeyAllIcon.Bind(TAttribute<FSlateIcon>::FGetter::CreateLambda([&]{
-				static FSlateIcon KeyAllEnabledIcon(FEditorStyle::GetStyleSetName(), "Sequencer.KeyAllEnabled");
-				static FSlateIcon KeyAllDisabledIcon(FEditorStyle::GetStyleSetName(), "Sequencer.KeyAllDisabled");
+			ToolBarBuilder.AddComboButton(
+				FUIAction(),
+				FOnGetContent::CreateSP(this, &SSequencer::MakePlaybackMenu),
+				LOCTEXT("PlaybackOptions", "Playback Options"),
+				LOCTEXT("PlaybackOptionsToolTip", "Playback Options"),
+				FSlateIcon(FEditorStyle::GetStyleSetName(), "Sequencer.PlaybackOptions")
+			);
 
-				return SequencerPtr.Pin()->GetKeyAllEnabled() ? KeyAllEnabledIcon : KeyAllDisabledIcon;
-			}));
+			ToolBarBuilder.AddComboButton(
+				FUIAction(),
+				FOnGetContent::CreateSP(this, &SSequencer::MakeSelectEditMenu),
+				LOCTEXT("SelectEditOptions", "Select/Edit Options"),
+				LOCTEXT("SelectEditOptionsToolTip", "Select/Edit Options"),
+				FSlateIcon(FEditorStyle::GetStyleSetName(), "Sequencer.SelectEditOptions")
+			);
 
-			ToolBarBuilder.AddToolBarButton( FSequencerCommands::Get().ToggleKeyAllEnabled, NAME_None, TAttribute<FText>(), TAttribute<FText>(), KeyAllIcon );
-		}
+			ToolBarBuilder.AddSeparator();
 
-		TAttribute<FSlateIcon> AutoKeyModeIcon;
-		AutoKeyModeIcon.Bind(TAttribute<FSlateIcon>::FGetter::CreateLambda( [&] {
-			switch ( SequencerPtr.Pin()->GetAutoKeyMode() )
+			if( SequencerPtr.Pin()->IsLevelEditorSequencer() )
 			{
-			case EAutoKeyMode::KeyAll:
-				return FSequencerCommands::Get().SetAutoKeyModeAll->GetIcon();
-			case EAutoKeyMode::KeyAnimated:
-				return FSequencerCommands::Get().SetAutoKeyModeAnimated->GetIcon();
-			default: // EAutoKeyMode::KeyNone
-				return FSequencerCommands::Get().SetAutoKeyModeNone->GetIcon();
-			}
-		} ) );
+				TAttribute<FSlateIcon> KeyAllIcon;
+				KeyAllIcon.Bind(TAttribute<FSlateIcon>::FGetter::CreateLambda([&]{
+					static FSlateIcon KeyAllEnabledIcon(FEditorStyle::GetStyleSetName(), "Sequencer.KeyAllEnabled");
+					static FSlateIcon KeyAllDisabledIcon(FEditorStyle::GetStyleSetName(), "Sequencer.KeyAllDisabled");
 
-		TAttribute<FText> AutoKeyModeToolTip;
-		AutoKeyModeToolTip.Bind( TAttribute<FText>::FGetter::CreateLambda( [&] {
-			switch ( SequencerPtr.Pin()->GetAutoKeyMode() )
-			{
-			case EAutoKeyMode::KeyAll:
-				return FSequencerCommands::Get().SetAutoKeyModeAll->GetDescription();
-			case EAutoKeyMode::KeyAnimated:
-				return FSequencerCommands::Get().SetAutoKeyModeAnimated->GetDescription();
-			default: // EAutoKeyMode::KeyNone
-				return FSequencerCommands::Get().SetAutoKeyModeNone->GetDescription();
+					return SequencerPtr.Pin()->GetKeyAllEnabled() ? KeyAllEnabledIcon : KeyAllDisabledIcon;
+				}));
+
+				ToolBarBuilder.AddToolBarButton( FSequencerCommands::Get().ToggleKeyAllEnabled, NAME_None, TAttribute<FText>(), TAttribute<FText>(), KeyAllIcon );
 			}
-		} ) );
 
-		ToolBarBuilder.AddComboButton(
-			FUIAction(),
-			FOnGetContent::CreateSP(this, &SSequencer::MakeAutoKeyMenu),
-			LOCTEXT("AutoKeyMode", "Auto-Key Mode"),
-			AutoKeyModeToolTip,
-			AutoKeyModeIcon);
+			TAttribute<FSlateIcon> AutoKeyModeIcon;
+			AutoKeyModeIcon.Bind(TAttribute<FSlateIcon>::FGetter::CreateLambda( [&] {
+				switch ( SequencerPtr.Pin()->GetAutoKeyMode() )
+				{
+				case EAutoKeyMode::KeyAll:
+					return FSequencerCommands::Get().SetAutoKeyModeAll->GetIcon();
+				case EAutoKeyMode::KeyAnimated:
+					return FSequencerCommands::Get().SetAutoKeyModeAnimated->GetIcon();
+				default: // EAutoKeyMode::KeyNone
+					return FSequencerCommands::Get().SetAutoKeyModeNone->GetIcon();
+				}
+			} ) );
+
+			TAttribute<FText> AutoKeyModeToolTip;
+			AutoKeyModeToolTip.Bind( TAttribute<FText>::FGetter::CreateLambda( [&] {
+				switch ( SequencerPtr.Pin()->GetAutoKeyMode() )
+				{
+				case EAutoKeyMode::KeyAll:
+					return FSequencerCommands::Get().SetAutoKeyModeAll->GetDescription();
+				case EAutoKeyMode::KeyAnimated:
+					return FSequencerCommands::Get().SetAutoKeyModeAnimated->GetDescription();
+				default: // EAutoKeyMode::KeyNone
+					return FSequencerCommands::Get().SetAutoKeyModeNone->GetDescription();
+				}
+			} ) );
+
+			ToolBarBuilder.AddComboButton(
+				FUIAction(),
+				FOnGetContent::CreateSP(this, &SSequencer::MakeAutoKeyMenu),
+				LOCTEXT("AutoKeyMode", "Auto-Key Mode"),
+				AutoKeyModeToolTip,
+				AutoKeyModeIcon);
+		}
 	}
 	ToolBarBuilder.EndSection();
 
@@ -832,11 +863,15 @@ TSharedRef<SWidget> SSequencer::MakeToolBar()
 	}
 	ToolBarBuilder.EndSection();
 
-	ToolBarBuilder.BeginSection("Curve Editor");
+	if (!bIsReadOnly)
 	{
-		ToolBarBuilder.AddToolBarButton( FSequencerCommands::Get().ToggleShowCurveEditor );
+		// Curve editor doesn't have any notion of read-only at the moment
+		ToolBarBuilder.BeginSection("Curve Editor");
+		{
+			ToolBarBuilder.AddToolBarButton( FSequencerCommands::Get().ToggleShowCurveEditor );
+		}
+		ToolBarBuilder.EndSection();
 	}
-	ToolBarBuilder.EndSection();
 
 	return ToolBarBuilder.MakeWidget();
 }
@@ -876,11 +911,7 @@ TSharedRef<SWidget> SSequencer::MakeGeneralMenu()
 	// view options
 	MenuBuilder.BeginSection( "ViewOptions", LOCTEXT( "ViewMenuHeader", "View" ) );
 	{
-		if (Sequencer->IsLevelEditorSequencer())
-		{
-			MenuBuilder.AddMenuEntry( FSequencerCommands::Get().ToggleLabelBrowser );
-		}
-
+		MenuBuilder.AddMenuEntry( FSequencerCommands::Get().ToggleLabelBrowser );
 		MenuBuilder.AddMenuEntry( FSequencerCommands::Get().ToggleCombinedKeyframes );
 		MenuBuilder.AddMenuEntry( FSequencerCommands::Get().ToggleChannelColors );
 
@@ -896,6 +927,32 @@ TSharedRef<SWidget> SSequencer::MakeGeneralMenu()
 	}
 	MenuBuilder.EndSection();
 
+	MenuBuilder.AddMenuEntry(FSequencerCommands::Get().ToggleShowGotoBox);
+
+	MenuBuilder.AddMenuSeparator();
+
+	if (SequencerPtr.Pin()->IsLevelEditorSequencer())
+	{
+		MenuBuilder.AddMenuEntry(FSequencerCommands::Get().FixActorReferences);
+	}
+	MenuBuilder.AddMenuEntry(FSequencerCommands::Get().FixFrameTiming);
+
+	if ( SequencerPtr.Pin()->IsLevelEditorSequencer() )
+	{
+		MenuBuilder.AddMenuSeparator();
+		
+		MenuBuilder.AddMenuEntry( FSequencerCommands::Get().ImportFBX );
+		MenuBuilder.AddMenuEntry( FSequencerCommands::Get().ExportFBX );
+	}
+
+	return MenuBuilder.MakeWidget();
+}
+
+TSharedRef<SWidget> SSequencer::MakePlaybackMenu()
+{
+	FMenuBuilder MenuBuilder( true, SequencerPtr.Pin()->GetCommandBindings() );
+	TSharedPtr<FSequencer> Sequencer = SequencerPtr.Pin();
+
 	// playback range options
 	MenuBuilder.BeginSection("PlaybackThisSequence", LOCTEXT("PlaybackThisSequenceHeader", "Playback - This Sequence"));
 	{
@@ -916,6 +973,9 @@ TSharedRef<SWidget> SSequencer::MakeGeneralMenu()
 				[
 					SNew(SSpinBox<float>)
 						.TypeInterface(NumericTypeInterface)
+						.IsEnabled_Lambda([=]() {
+							return !Sequencer->IsPlaybackRangeLocked();
+						})
 						.Style(&FEditorStyle::GetWidgetStyle<FSpinBoxStyle>("Sequencer.HyperlinkSpinBox"))
 						.OnValueCommitted_Lambda([=](float Value, ETextCommit::Type){ OnStartChanged(Value); })
 						.OnValueChanged_Lambda(OnStartChanged)
@@ -950,6 +1010,9 @@ TSharedRef<SWidget> SSequencer::MakeGeneralMenu()
 				[
 					SNew(SSpinBox<float>)
 						.TypeInterface(NumericTypeInterface)
+						.IsEnabled_Lambda([=]() {
+							return !Sequencer->IsPlaybackRangeLocked();
+						})
 						.Style(&FEditorStyle::GetWidgetStyle<FSpinBoxStyle>("Sequencer.HyperlinkSpinBox"))
 						.OnValueCommitted_Lambda([=](float Value, ETextCommit::Type){ OnEndChanged(Value); })
 						.OnValueChanged_Lambda(OnEndChanged)
@@ -967,13 +1030,18 @@ TSharedRef<SWidget> SSequencer::MakeGeneralMenu()
 				],
 			LOCTEXT("PlaybackStartEnd", "End"));
 
+		MenuBuilder.AddMenuEntry( FSequencerCommands::Get().TogglePlaybackRangeLocked );
 		MenuBuilder.AddMenuEntry( FSequencerCommands::Get().ToggleForceFixedFrameIntervalPlayback );
 	}
 	MenuBuilder.EndSection();
 
 	MenuBuilder.BeginSection( "PlaybackAllSequences", LOCTEXT( "PlaybackRangeAllSequencesHeader", "Playback Range - All Sequences" ) );
 	{
-		MenuBuilder.AddMenuEntry( FSequencerCommands::Get().ToggleEvaluateSubSequencesInIsolation );
+		if (SequencerPtr.Pin()->IsLevelEditorSequencer())
+		{
+			MenuBuilder.AddMenuEntry( FSequencerCommands::Get().ToggleEvaluateSubSequencesInIsolation );
+		}
+
 		MenuBuilder.AddMenuEntry( FSequencerCommands::Get().ToggleKeepCursorInPlaybackRange );
 		MenuBuilder.AddMenuEntry( FSequencerCommands::Get().ToggleKeepPlaybackRangeInSectionBounds );
 		MenuBuilder.AddMenuEntry( FSequencerCommands::Get().ToggleLinkCurveEditorTimeRange );
@@ -1004,10 +1072,18 @@ TSharedRef<SWidget> SSequencer::MakeGeneralMenu()
 				],
 			LOCTEXT("ZeroPaddingText", "Zero Pad Frame Numbers"));
 	}
-	MenuBuilder.AddMenuEntry(FSequencerCommands::Get().ToggleShowGotoBox);
-	MenuBuilder.AddMenuEntry(FSequencerCommands::Get().ToggleShowTransformBox);
 	MenuBuilder.EndSection();
 
+	return MenuBuilder.MakeWidget();
+}
+
+TSharedRef<SWidget> SSequencer::MakeSelectEditMenu()
+{
+	FMenuBuilder MenuBuilder( true, SequencerPtr.Pin()->GetCommandBindings() );
+	TSharedPtr<FSequencer> Sequencer = SequencerPtr.Pin();
+
+	MenuBuilder.AddMenuEntry(FSequencerCommands::Get().ToggleShowTransformBox);
+
 	// selection range options
 	MenuBuilder.BeginSection("SelectionRange", LOCTEXT("SelectionRangeHeader", "Selection Range"));
 	{
@@ -1015,26 +1091,11 @@ TSharedRef<SWidget> SSequencer::MakeGeneralMenu()
 		MenuBuilder.AddMenuEntry(FSequencerCommands::Get().SetSelectionRangeEnd);
 		MenuBuilder.AddMenuEntry(FSequencerCommands::Get().ResetSelectionRange);
 		MenuBuilder.AddMenuEntry(FSequencerCommands::Get().SelectKeysInSelectionRange);
+		MenuBuilder.AddMenuEntry(FSequencerCommands::Get().SelectSectionsInSelectionRange);
+		MenuBuilder.AddMenuEntry(FSequencerCommands::Get().SelectAllInSelectionRange);
 	}
 	MenuBuilder.EndSection();
 
-	// other options
-	MenuBuilder.AddMenuSeparator();
-
-	if (SequencerPtr.Pin()->IsLevelEditorSequencer())
-	{
-		MenuBuilder.AddMenuEntry(FSequencerCommands::Get().FixActorReferences);
-	}
-	MenuBuilder.AddMenuEntry(FSequencerCommands::Get().FixFrameTiming);
-
-	MenuBuilder.AddMenuSeparator();
-
-	if ( SequencerPtr.Pin()->IsLevelEditorSequencer() )
-	{
-		MenuBuilder.AddMenuEntry( FSequencerCommands::Get().ImportFBX );
-		MenuBuilder.AddMenuEntry( FSequencerCommands::Get().ExportFBX );
-	}
-
 	return MenuBuilder.MakeWidget();
 }
 
@@ -1366,6 +1427,13 @@ FReply SSequencer::OnKeyDown( const FGeometry& MyGeometry, const FKeyEvent& InKe
 	return FReply::Unhandled();
 }
 
+void SSequencer::OnFocusChanging( const FWeakWidgetPath& PreviousFocusPath, const FWidgetPath& NewWidgetPath, const FFocusEvent& InFocusEvent )
+{
+	if (NewWidgetPath.ContainsWidget(AsShared()))
+	{
+		OnReceivedFocus.ExecuteIfBound();
+	}
+}
 
 void SSequencer::OnAssetsDropped( const FAssetDragDropOp& DragDropOp )
 {
@@ -1918,30 +1986,27 @@ bool SSequencer::CanPaste()
 		FString TexttoImport;
 		FPlatformMisc::ClipboardPaste(TexttoImport);
 
-		TArray<UMovieSceneTrack*> ImportedTrack;
-		Sequencer->ImportTracksFromText(TexttoImport, ImportedTrack);
-		if (ImportedTrack.Num() == 0)
+		if (Sequencer->CanPaste(TexttoImport))
 		{
-			return false;
-		}
+			TArray<UMovieSceneTrack*> ImportedTrack;
+			Sequencer->ImportTracksFromText(TexttoImport, ImportedTrack);
+			if (ImportedTrack.Num() == 0)
+			{
+				return false;
+			}
 
-		for (TSharedRef<FSequencerDisplayNode> Node : SelectedNodes)
-		{
-			if (Node->GetType() == ESequencerNode::Object)
+			for (TSharedRef<FSequencerDisplayNode> Node : SelectedNodes)
 			{
-				return true;
+				if (Node->GetType() == ESequencerNode::Object)
+				{
+					return true;
+				}
 			}
-		}
-		return false;
-	}
-	else
-	{
-		if (!HasFocusedDescendants() && !HasKeyboardFocus())
-		{
 			return false;
 		}
-		return SequencerPtr.Pin()->GetClipboardStack().Num() != 0;
 	}
+
+	return SequencerPtr.Pin()->GetClipboardStack().Num() != 0;
 }
 
 void SSequencer::PasteTracks()
diff --git a/Engine/Source/Editor/Sequencer/Private/SSequencer.h b/Engine/Source/Editor/Sequencer/Private/SSequencer.h
--- a/Engine/Source/Editor/Sequencer/Private/SSequencer.h
+++ b/Engine/Source/Editor/Sequencer/Private/SSequencer.h
@@ -141,6 +141,12 @@ public:
 		/** Called when the user has finished dragging the selection range */
 		SLATE_EVENT( FSimpleDelegate, OnSelectionRangeEndDrag )
 
+		/** Whether the playback range is locked */
+		SLATE_ATTRIBUTE( bool, IsPlaybackRangeLocked )
+
+		/** Called when the user toggles the play back range lock */
+		SLATE_EVENT( FSimpleDelegate, OnTogglePlaybackRangeLocked )
+
 		/** The current scrub position in (seconds) */
 		SLATE_ATTRIBUTE( float, ScrubPosition )
 
@@ -162,6 +168,9 @@ public:
 		/** Called to populate the add combo button in the toolbar. */
 		SLATE_EVENT( FOnGetAddMenuContent, OnGetAddMenuContent )
 
+		/** Called when any widget contained within sequencer has received focus */
+		SLATE_EVENT( FSimpleDelegate, OnReceivedFocus )
+
 		/** Extender to use for the add menu. */
 		SLATE_ARGUMENT( TSharedPtr<FExtender>, AddMenuExtender )
 	SLATE_END_ARGS()
@@ -243,6 +252,7 @@ protected:
 	virtual FReply OnDragOver( const FGeometry& MyGeometry, const FDragDropEvent& DragDropEvent ) override;
 	virtual FReply OnDrop( const FGeometry& MyGeometry, const FDragDropEvent& DragDropEvent ) override;
 	virtual FReply OnKeyDown( const FGeometry& MyGeometry, const FKeyEvent& InKeyEvent ) override;
+	virtual void OnFocusChanging( const FWeakWidgetPath& PreviousFocusPath, const FWidgetPath& NewWidgetPath, const FFocusEvent& InFocusEvent ) override;
 
 private:
 	
@@ -276,6 +286,12 @@ private:
 	/** Makes the general menu for the toolbar. */
 	TSharedRef<SWidget> MakeGeneralMenu();
 
+	/** Makes the plabacky menu for the toolbar. */
+	TSharedRef<SWidget> MakePlaybackMenu();
+
+	/** Makes the select/edit menu for the toolbar. */
+	TSharedRef<SWidget> MakeSelectEditMenu();
+
 	/** Makes the snapping menu for the toolbar. */
 	TSharedRef<SWidget> MakeSnapMenu();
 
@@ -483,6 +499,9 @@ private:
 	/** Called when the user has finished dragging the playback range */
 	FSimpleDelegate OnPlaybackRangeEndDrag;
 
+	/** Called when any widget contained within sequencer has received focus */
+	FSimpleDelegate OnReceivedFocus;
+
 	/** Cached clamp and view range for unlinking the curve editor time range */
 	TRange<float> CachedClampRange;
 	TRange<float> CachedViewRange;
diff --git a/Engine/Source/Editor/Sequencer/Private/SSequencerCurveEditor.cpp b/Engine/Source/Editor/Sequencer/Private/SSequencerCurveEditor.cpp
--- a/Engine/Source/Editor/Sequencer/Private/SSequencerCurveEditor.cpp
+++ b/Engine/Source/Editor/Sequencer/Private/SSequencerCurveEditor.cpp
@@ -1,11 +1,14 @@
 // Copyright 1998-2017 Epic Games, Inc. All Rights Reserved.
 
 #include "SSequencerCurveEditor.h"
-#include "SequencerSettings.h"
+#include "SequencerCommonHelpers.h"
 #include "SequencerCurveOwner.h"
+#include "SequencerSettings.h"
+
 
 #define LOCTEXT_NAMESPACE "SequencerCurveEditor"
 
+
 void SSequencerCurveEditor::Construct( const FArguments& InArgs, TSharedRef<FSequencer> InSequencer, TSharedRef<ITimeSliderController> InTimeSliderController )
 {
 	Sequencer = InSequencer;
@@ -28,6 +31,7 @@ void SSequencerCurveEditor::Construct( const FArguments& InArgs, TSharedRef<FSeq
 		.ShowCurveSelector( false )
 		.ShowZoomButtons( false )
 		.ShowInputGridNumbers( false )
+		.ShowTimeInFrames(this, &SSequencerCurveEditor::GetShowTimeInFrames)
 		.InputSnappingEnabled(this, &SSequencerCurveEditor::GetInputCurveSnapEnabled)
 		.InputSnap( this, &SSequencerCurveEditor::GetCurveTimeSnapInterval )
 		.OutputSnap( this, &SSequencerCurveEditor::GetCurveValueSnapInterval )
@@ -163,6 +167,11 @@ float SSequencerCurveEditor::GetCurveValueSnapInterval() const
 		: 0;
 }
 
+bool SSequencerCurveEditor::GetShowTimeInFrames() const
+{
+	return SequencerSnapValues::IsTimeSnapIntervalFrameRate(SequencerSettings->GetTimeSnapInterval());
+}
+
 void SSequencerCurveEditor::NodeTreeSelectionChanged()
 {
 	if (SequencerNodeTree.IsValid())
diff --git a/Engine/Source/Editor/Sequencer/Private/SSequencerCurveEditor.h b/Engine/Source/Editor/Sequencer/Private/SSequencerCurveEditor.h
--- a/Engine/Source/Editor/Sequencer/Private/SSequencerCurveEditor.h
+++ b/Engine/Source/Editor/Sequencer/Private/SSequencerCurveEditor.h
@@ -54,6 +54,9 @@ private:
 	/** Gets the value snapping interval. */
 	float GetCurveValueSnapInterval() const;
 
+	/** Gets whether to show time in frames.*/
+	bool GetShowTimeInFrames() const;
+
 	/** Run whenever the selection on the FSequencerNodeTree changes. */
 	void NodeTreeSelectionChanged();
 
diff --git a/Engine/Source/Editor/Sequencer/Private/SSequencerTransformBox.cpp b/Engine/Source/Editor/Sequencer/Private/SSequencerTransformBox.cpp
--- a/Engine/Source/Editor/Sequencer/Private/SSequencerTransformBox.cpp
+++ b/Engine/Source/Editor/Sequencer/Private/SSequencerTransformBox.cpp
@@ -17,6 +17,7 @@ void SSequencerTransformBox::Construct(const FArguments& InArgs, const TSharedRe
 	NumericTypeInterface = InNumericTypeInterface;
 	
 	DeltaTime = 1.f;
+	ScaleFactor = 1.f;
 
 	// Initialize to 10 frames if showing frame numbers
 	float TimeSnapInterval = SequencerPtr.Pin()->GetSettings()->GetTimeSnapInterval();
@@ -63,10 +64,41 @@ void SSequencerTransformBox::Construct(const FArguments& InArgs, const TSharedRe
 						SAssignNew(EntryBox, SNumericEntryBox<float>)
 							.MinDesiredValueWidth(32.0f)
 							.TypeInterface(NumericTypeInterface)
+							.ToolTipText(LOCTEXT("Delta_Tooltip", "The amount to offset the keys/sections by"))
 							.OnValueCommitted(this, &SSequencerTransformBox::OnDeltaChanged)
 							.Value_Lambda([this](){ return DeltaTime; })
 					]
 
+				+ SHorizontalBox::Slot()
+					.Padding(6.0f, 0.0f, 0.0f, 0.0f)
+					.AutoWidth()
+					.VAlign(VAlign_Center)
+					[
+						SNew(SButton)
+							.Text(LOCTEXT("MultiplyLabel", "*"))
+							.OnClicked(this, &SSequencerTransformBox::OnMultiplyButtonClicked)
+					]
+
+				+ SHorizontalBox::Slot()
+					.AutoWidth()
+					.VAlign(VAlign_Center)
+					[
+						SNew(SButton)
+							.Text(LOCTEXT("DivideLabel", "/"))
+							.OnClicked(this, &SSequencerTransformBox::OnDivideButtonClicked)
+					]
+
+				+ SHorizontalBox::Slot()
+					.Padding(6.0f, 0.0f, 0.0f, 0.0f)
+					.AutoWidth()
+					[
+						SAssignNew(EntryBox, SNumericEntryBox<float>)
+							.MinDesiredValueWidth(32.0f)
+							.ToolTipText(LOCTEXT("Scale_Tooltip", "The amount to scale the keys/section by (about the local time)"))
+							.OnValueCommitted(this, &SSequencerTransformBox::OnScaleChanged)
+							.Value_Lambda([this](){ return ScaleFactor; })
+					]
+
 				+ SHorizontalBox::Slot()
 				.Padding(3.0f, 0.0f, 0.0f, 0.0f)
 				.AutoWidth()
@@ -114,21 +146,60 @@ void SSequencerTransformBox::OnDeltaChanged(float Value, ETextCommit::Type Commi
 	DeltaTime = Value;
 }
 
+void SSequencerTransformBox::OnScaleChanged(float Value, ETextCommit::Type CommitType)
+{
+	if (CommitType != ETextCommit::OnEnter)
+	{
+		return;
+	}
+
+	ScaleFactor = Value;
+}
+
 FReply SSequencerTransformBox::OnPlusButtonClicked()
 {
-	TSharedPtr<FSequencer> Sequencer = SequencerPtr.Pin();
+	if (DeltaTime != 0.f)
+	{
+		TSharedPtr<FSequencer> Sequencer = SequencerPtr.Pin();
 
-	Sequencer->TransformSelectedKeysAndSections(DeltaTime);
+		Sequencer->TransformSelectedKeysAndSections(DeltaTime, 1.f);
+	}
 
 	return FReply::Handled();
 }
 
 FReply SSequencerTransformBox::OnMinusButtonClicked()
 {
-	TSharedPtr<FSequencer> Sequencer = SequencerPtr.Pin();
+	if (DeltaTime != 0.f)
+	{
+		TSharedPtr<FSequencer> Sequencer = SequencerPtr.Pin();
+	
+		Sequencer->TransformSelectedKeysAndSections(-DeltaTime, 1.f);
+	}
+
+	return FReply::Handled();
+}
+
+FReply SSequencerTransformBox::OnMultiplyButtonClicked()
+{
+	if (ScaleFactor != 0.f && ScaleFactor != 1.f)
+	{
+		TSharedPtr<FSequencer> Sequencer = SequencerPtr.Pin();
+		
+		Sequencer->TransformSelectedKeysAndSections(0, ScaleFactor);
+	}
 
-	Sequencer->TransformSelectedKeysAndSections(-DeltaTime);
+	return FReply::Handled();
+}
+
+FReply SSequencerTransformBox::OnDivideButtonClicked()
+{
+	if (ScaleFactor != 0.f && ScaleFactor != 1.f)
+	{
+		TSharedPtr<FSequencer> Sequencer = SequencerPtr.Pin();
 
+		Sequencer->TransformSelectedKeysAndSections(0, 1.f/ScaleFactor);
+	}
 	return FReply::Handled();
 }
 
diff --git a/Engine/Source/Editor/Sequencer/Private/SSequencerTransformBox.h b/Engine/Source/Editor/Sequencer/Private/SSequencerTransformBox.h
--- a/Engine/Source/Editor/Sequencer/Private/SSequencerTransformBox.h
+++ b/Engine/Source/Editor/Sequencer/Private/SSequencerTransformBox.h
@@ -31,7 +31,10 @@ private:
 	/** Callback for transform operations. */
 	FReply OnPlusButtonClicked();
 	FReply OnMinusButtonClicked();
+	FReply OnMultiplyButtonClicked();
+	FReply OnDivideButtonClicked();
 	void OnDeltaChanged(float Value, ETextCommit::Type CommitType);
+	void OnScaleChanged(float Value, ETextCommit::Type CommitType);
 
 	/** Callback for when the close button is clicked. */
 	FReply OnCloseButtonClicked();
@@ -58,4 +61,7 @@ private:
 
 	/** Cached delta time. */
 	float DeltaTime;
+
+	/** Cached scale factor. */
+	float ScaleFactor;
 };
diff --git a/Engine/Source/Editor/Sequencer/Private/SSequencerTreeView.cpp b/Engine/Source/Editor/Sequencer/Private/SSequencerTreeView.cpp
--- a/Engine/Source/Editor/Sequencer/Private/SSequencerTreeView.cpp
+++ b/Engine/Source/Editor/Sequencer/Private/SSequencerTreeView.cpp
@@ -711,9 +711,10 @@ TSharedRef<ITableRow> SSequencerTreeView::OnGenerateRow(FDisplayNodeRef InDispla
 		if (!TrackLane.IsValid())
 		{
 			// Add a track slot for the row
-			TAttribute<TRange<float>> ViewRange = FAnimatedRange::WrapAttribute( TAttribute<FAnimatedRange>::Create(TAttribute<FAnimatedRange>::FGetter::CreateSP(&SequencerNodeTree->GetSequencer(), &FSequencer::GetViewRange)) );	
+			TAttribute<TRange<float>> ViewRange = FAnimatedRange::WrapAttribute( TAttribute<FAnimatedRange>::Create(TAttribute<FAnimatedRange>::FGetter::CreateSP(&SequencerNodeTree->GetSequencer(), &FSequencer::GetViewRange)) );
 
 			TrackLane = SNew(SSequencerTrackLane, SectionAuthority.ToSharedRef(), SharedThis(this))
+			.IsEnabled(!InDisplayNode->GetSequencer().IsReadOnly())
 			[
 				SectionAuthority->GenerateWidgetForSectionArea(ViewRange)
 			];
