Copying //UE4/Dev-Sequencer to //UE4/Dev-Main (Source: //UE4/Dev-Sequencer @ 3237992)

diff --git a/Engine/Source/Editor/Sequencer/Private/Sequencer.cpp b/Engine/Source/Editor/Sequencer/Private/Sequencer.cpp
--- a/Engine/Source/Editor/Sequencer/Private/Sequencer.cpp
+++ b/Engine/Source/Editor/Sequencer/Private/Sequencer.cpp
@@ -72,7 +72,6 @@
 #include "Tracks/MovieSceneSubTrack.h"
 #include "Sections/MovieSceneCinematicShotSection.h"
 #include "ISettingsModule.h"
-#include "IDetailsView.h"
 #include "SnappingUtils.h"
 #include "Framework/Commands/GenericCommands.h"
 #include "Tracks/MovieSceneSpawnTrack.h"
@@ -95,7 +94,6 @@
 #include "SequencerUtilities.h"
 #include "Tracks/MovieSceneCinematicShotTrack.h"
 #include "ISequenceRecorder.h"
-#include "GameDelegates.h"
 #include "CineCameraActor.h"
 #include "CameraRig_Rail.h"
 #include "CameraRig_Crane.h"
@@ -109,12 +107,6 @@
 
 DEFINE_LOG_CATEGORY(LogSequencer);
 
-
-namespace
-{
-	TWeakPtr<FSequencer> ActiveSequencerPtr;
-}
-
 struct FSequencerTemplateStore : FMovieSceneSequenceTemplateStore
 {
 	FSequencerTemplateStore()
@@ -162,11 +154,12 @@ struct FSequencerTemplateStore : FMovieSceneSequenceTemplateStore
 };
 
 
-void FSequencer::InitSequencer(const FSequencerInitParams& InitParams, const TSharedRef<ISequencerObjectChangeListener>& InObjectChangeListener, const TSharedRef<IDetailKeyframeHandler>& InDetailKeyframeHandler, const TArray<FOnCreateTrackEditor>& TrackEditorDelegates)
+void FSequencer::InitSequencer(const FSequencerInitParams& InitParams, const TSharedRef<ISequencerObjectChangeListener>& InObjectChangeListener, const TArray<FOnCreateTrackEditor>& TrackEditorDelegates)
 {
 	bIsEditingWithinLevelEditor = InitParams.bEditWithinLevelEditor;
 	SilentModeCount = 0;
-	
+	bReadOnly = InitParams.ViewParams.bReadOnly;
+
 	PreAnimatedState.EnableGlobalCapture();
 
 	if (InitParams.SpawnRegister.IsValid())
@@ -192,35 +185,6 @@ void FSequencer::InitSequencer(const FSequencerInitParams& InitParams, const TSh
 	PlaybackContextAttribute = InitParams.PlaybackContext;
 	CachedPlaybackContext = PlaybackContextAttribute.Get(nullptr);
 
-	// If this sequencer edits the level, close out the existing active sequencer and mark this as the active sequencer.
-	if (bIsEditingWithinLevelEditor)
-	{
-		if (ActiveSequencerPtr.IsValid())
-		{
-			ActiveSequencerPtr.Pin()->Close();
-		}
-
-		ActiveSequencerPtr = SharedThis(this);
-
-		// Register for saving the level so that the state of the scene can be restored before saving and updated after saving.
-		FEditorDelegates::PreSaveWorld.AddSP(this, &FSequencer::OnPreSaveWorld);
-		FEditorDelegates::PostSaveWorld.AddSP(this, &FSequencer::OnPostSaveWorld);
-		FEditorDelegates::PreBeginPIE.AddSP(this, &FSequencer::OnPreBeginPIE);
-		FEditorDelegates::EndPIE.AddSP(this, &FSequencer::OnEndPIE);
-		FGameDelegates::Get().GetEndPlayMapDelegate().AddSP(this, &FSequencer::OnEndPlayMap);
-
-		FWorldDelegates::LevelAddedToWorld.AddSP(this, &FSequencer::OnLevelAdded);
-		FWorldDelegates::LevelRemovedFromWorld.AddSP(this, &FSequencer::OnLevelRemoved);
-
-		FEditorDelegates::NewCurrentLevel.AddSP(this, &FSequencer::OnNewCurrentLevel);
-		FEditorDelegates::OnMapOpened.AddSP(this, &FSequencer::OnMapOpened);
-		FEditorDelegates::OnNewActorsDropped.AddSP(this, &FSequencer::OnNewActorsDropped);
-
-		USelection::SelectionChangedEvent.AddSP( this, &FSequencer::OnActorSelectionChanged );
-
-		AddLevelViewportMenuExtender();
-	}
-
 	Settings = USequencerSettingsContainer::GetOrCreate<USequencerSettings>(*InitParams.ViewParams.UniqueName);
 
 	if (ISettingsModule* SettingsModule = FModuleManager::GetModulePtr<ISettingsModule>("Settings"))
@@ -243,7 +207,6 @@ void FSequencer::InitSequencer(const FSequencerInitParams& InitParams, const TSh
 	PlayRate = 1.f;
 	ShuttleMultiplier = 0;
 	ObjectChangeListener = InObjectChangeListener;
-	DetailKeyframeHandler = InDetailKeyframeHandler;
 
 	check( ObjectChangeListener.IsValid() );
 	
@@ -268,6 +231,8 @@ void FSequencer::InitSequencer(const FSequencerInitParams& InitParams, const TSh
 		.OnSelectionRangeChanged( this, &FSequencer::SetSelectionRange )
 		.OnSelectionRangeBeginDrag( this, &FSequencer::OnSelectionRangeBeginDrag )
 		.OnSelectionRangeEndDrag( this, &FSequencer::OnSelectionRangeEndDrag )
+		.IsPlaybackRangeLocked( this, &FSequencer::IsPlaybackRangeLocked )
+		.OnTogglePlaybackRangeLocked( this, &FSequencer::TogglePlaybackRangeLocked )
 		.ScrubPosition( this, &FSequencer::GetLocalTime )
 		.OnBeginScrubbing( this, &FSequencer::OnBeginScrubbing )
 		.OnEndScrubbing( this, &FSequencer::OnEndScrubbing )
@@ -275,23 +240,12 @@ void FSequencer::InitSequencer(const FSequencerInitParams& InitParams, const TSh
 		.OnViewRangeChanged( this, &FSequencer::SetViewRange )
 		.OnClampRangeChanged( this, &FSequencer::OnClampRangeChanged )
 		.OnGetAddMenuContent(InitParams.ViewParams.OnGetAddMenuContent)
+		.OnReceivedFocus(InitParams.ViewParams.OnReceivedFocus)
 		.AddMenuExtender(InitParams.ViewParams.AddMenuExtender);
 
 	// When undo occurs, get a notification so we can make sure our view is up to date
 	GEditor->RegisterForUndo(this);
 
-	if( bIsEditingWithinLevelEditor )
-	{
-		// @todo remove when world-centric mode is added
-		// Hook into the editor's mechanism for checking whether we need live capture of PIE/SIE actor state
-		GEditor->GetActorRecordingState().AddSP(this, &FSequencer::GetActorRecordingState);
-
-		ActivateDetailKeyframeHandler();
-				
-		FPropertyEditorModule& EditModule = FModuleManager::Get().GetModuleChecked<FPropertyEditorModule>("PropertyEditor");
-		EditModule.OnPropertyEditorOpened().AddSP(this, &FSequencer::OnPropertyEditorOpened);
-	}
-
 	// Create tools and bind them to this sequencer
 	for( int32 DelegateIndex = 0; DelegateIndex < TrackEditorDelegates.Num(); ++DelegateIndex )
 	{
@@ -313,13 +267,6 @@ void FSequencer::InitSequencer(const FSequencerInitParams& InitParams, const TSh
 	// NOTE: Could fill in asset editor commands here!
 
 	BindCommands();
-	BindLevelEditorCommands();
-	ActivateSequencerEditorMode();
-
-	if( bIsEditingWithinLevelEditor )
-	{
-		AttachTransportControlsToViewports();
-	}
 
 	for (auto TrackEditor : TrackEditors)
 	{
@@ -345,7 +292,6 @@ FSequencer::FSequencer()
 	, bNeedTreeRefresh( false )
 	, bNeedInstanceRefresh( false )
 	, StoredPlaybackState( EMovieScenePlayerStatus::Stopped )
-	, bWasClosed( false )
 	, NodeTree( MakeShareable( new FSequencerNodeTree( *this ) ) )
 	, bUpdatingSequencerSelection( false )
 	, bUpdatingExternalSelection( false )
@@ -360,87 +306,35 @@ FSequencer::FSequencer()
 
 FSequencer::~FSequencer()
 {
-	GEditor->GetActorRecordingState().RemoveAll( this );
-	GEditor->UnregisterForUndo( this );
+	GEditor->UnregisterForUndo(this);
 
-	if( bIsEditingWithinLevelEditor )
+	if (ISettingsModule* SettingsModulePtr = FModuleManager::GetModulePtr<ISettingsModule>("Settings"))
 	{
-		DetachTransportControlsFromViewports();
-		RemoveLevelViewportMenuExtender();
+		FName SettingsSectionName = *Settings->GetName();
+		SettingsModulePtr->UnregisterSettings("Editor", "ContentEditors", SettingsSectionName);
 	}
 
-	// When a new level sequence asset is opened while an existing one is open, the original level sequence asset is closed, 
-	// the new level sequence asset is initialized, and then the original level sequence asset is destroyed. In that case, 
-	// the track editors should not be released a second time. However, a level sequence asset is opened and closed without 
-	// opening another, the track editors still need to be released because Close() wil not be called.
-	//
-	if (!bWasClosed)
+	for (auto TrackEditor : TrackEditors)
 	{
-		for (auto TrackEditor : TrackEditors)
-		{
-			TrackEditor->OnRelease();
-		}
-		TrackEditors.Empty();
-	
-		UnbindLevelEditorCommands();
+		TrackEditor->OnRelease();
 	}
+
 	SequencerWidget.Reset();
+	TrackEditors.Empty();
 }
 
 
 void FSequencer::Close()
 {
-	if (ActiveSequencerPtr.IsValid() && ActiveSequencerPtr.Pin().Get() == this)
-	{
-		RestorePreAnimatedState();
-
-		FEditorDelegates::PreSaveWorld.RemoveAll(this);
-		FEditorDelegates::PostSaveWorld.RemoveAll(this);
-		FEditorDelegates::PreBeginPIE.RemoveAll(this);
-		FEditorDelegates::EndPIE.RemoveAll(this);
-		FGameDelegates::Get().GetEndPlayMapDelegate().RemoveAll(this);
-
-		FWorldDelegates::LevelAddedToWorld.RemoveAll(this);
-		FWorldDelegates::LevelRemovedFromWorld.RemoveAll(this);
-
-		FEditorDelegates::NewCurrentLevel.RemoveAll(this);
-		FEditorDelegates::OnMapOpened.RemoveAll(this);
-		FEditorDelegates::OnNewActorsDropped.RemoveAll(this);
-		
-		for (auto TrackEditor : TrackEditors)
-		{
-			TrackEditor->OnRelease();
-		}
-
-		ActiveSequencerPtr.Reset();
-
-		if (GLevelEditorModeTools().IsModeActive(FSequencerEdMode::EM_SequencerMode))
-		{
-			GLevelEditorModeTools().DeactivateMode(FSequencerEdMode::EM_SequencerMode);
-		}
-
-		if (ISettingsModule* SettingsModule = FModuleManager::GetModulePtr<ISettingsModule>("Settings"))
-		{
-			FName SettingsSectionName = *Settings->GetName();
-			SettingsModule->UnregisterSettings("Editor", "ContentEditors", SettingsSectionName);
-		}
-
-		UnbindLevelEditorCommands();
-	}
+	RestorePreAnimatedState();
 
-	if (bIsEditingWithinLevelEditor)
+	for (auto TrackEditor : TrackEditors)
 	{
-		RemoveLevelViewportMenuExtender();
-		DeactivateDetailKeyframeHandler();
-			
-		FPropertyEditorModule& EditModule = FModuleManager::Get().GetModuleChecked<FPropertyEditorModule>("PropertyEditor");			
-		EditModule.OnPropertyEditorOpened().RemoveAll(this);
-
-		DetachTransportControlsFromViewports();
+		TrackEditor->OnRelease();
 	}
-	
+
 	SequencerWidget.Reset();
-	bWasClosed = true;
+	TrackEditors.Empty();
 }
 
 
@@ -454,6 +348,7 @@ void FSequencer::Tick(float InDeltaTime)
 
 		if (CachedPlaybackContext != NewPlaybackContext)
 		{
+			PrePossessionViewTargets.Reset();
 			State.ClearObjectCaches();
 			CachedPlaybackContext = NewPlaybackContext;
 		}
@@ -548,7 +443,6 @@ void FSequencer::Tick(float InDeltaTime)
 			FMovieSceneEvaluationRange Range = PlayPosition.JumpTo(ScrubPosition, GetFocusedMovieSceneSequence()->GetMovieScene()->GetOptionalFixedFrameInterval());
 			EvaluateInternal(Range);
 		}
-		//PostTickRenderStateFixup();
 	}
 
 	ISequenceRecorder& SequenceRecorder = FModuleManager::LoadModuleChecked<ISequenceRecorder>("SequenceRecorder");
@@ -578,76 +472,6 @@ void FSequencer::Tick(float InDeltaTime)
 	}
 }
 
-void FSequencer::PostTickRenderStateFixup()
-{
-	if ( !IsLevelEditorSequencer() )
-	{
-		return;
-	}
-
-	if (PlaybackState == EMovieScenePlayerStatus::Jumping || PlaybackState == EMovieScenePlayerStatus::Stepping)
-	{
-		const bool bApplyFixup = true;
-
-		if (!bApplyFixup)
-		{
-			// Clear any delayed state changes regardless
-			SetPlaybackStatus(EMovieScenePlayerStatus::Stopped);
-		}
-		else
-		{
-			const float Time = GetGlobalTime();
-			const float FixedTickDelta = 1.f / 60.f;
-			const float StepDelta = (PlaybackState == EMovieScenePlayerStatus::Jumping) ? 0.f : FixedTickDelta;
-			
-			// Find a time for our previous frame, if there is one
-			float TimeLowerBound = GetPlaybackRange().GetLowerBoundValue();
-
-			if (UMovieSceneSequence* MovieSequence = GetFocusedMovieSceneSequence())
-			{
-				UMovieSceneSubTrack* SubTrack = Cast<UMovieSceneSubTrack>(MovieSequence->GetMovieScene()->FindMasterTrack(UMovieSceneCinematicShotTrack::StaticClass()));
-				if (!SubTrack)
-				{
-					SubTrack = Cast<UMovieSceneSubTrack>(MovieSequence->GetMovieScene()->FindMasterTrack(UMovieSceneSubTrack::StaticClass()));
-				}
-
-				if (SubTrack)
-				{
-					for (UMovieSceneSection* Section : SubTrack->GetAllSections())
-					{
-						// Subsections are inclusive start and exclusive end
-						if (Section->IsInfinite() || (Time >= Section->GetStartTime() && Time < Section->GetEndTime()))
-						{
-							UMovieSceneSubSection* SubSection = CastChecked<UMovieSceneSubSection>(Section);
-
-							TimeLowerBound = FMath::Max(TimeLowerBound, SubSection->GetStartTime());
-							break;
-						}
-					}
-				}
-			}
-
-			float PreviousTime = FMath::Max(TimeLowerBound, Time - StepDelta);
-			PreviousTime = FMath::Min(PreviousTime, Time);
-
-			// Force previous frame in sequence to re-render
-			SetPlaybackStatus(EMovieScenePlayerStatus::Stopped);
-			SetGlobalTime(PreviousTime);
-			Tick(StepDelta);
-
-			UWorld* PlaybackContext = Cast<UWorld>(GetPlaybackContext());
-			if (PlaybackContext)
-			{
-				PlaybackContext->SendAllEndOfFrameUpdates();
-			}
-
-			SetPlaybackStatus(EMovieScenePlayerStatus::Stopped);
-			SetGlobalTime(Time);
-			Tick(StepDelta);
-		}
-	}
-}
-
 
 TSharedRef<SWidget> FSequencer::GetSequencerWidget() const
 {
@@ -1027,7 +851,7 @@ bool FSequencer::CanSnapToFrame() const
 	return bKeysSelected && CanShowFrameNumbers();
 }
 
-void FSequencer::TransformSelectedKeysAndSections(float InDeltaTime)
+void FSequencer::TransformSelectedKeysAndSections(float InDeltaTime, float InScale)
 {
 	FScopedTransaction TransformKeysAndSectionsTransaction(NSLOCTEXT("Sequencer", "TransformKeysandSections_Transaction", "Transform Keys and Sections"));
 	bool bAnythingChanged = false;
@@ -1056,8 +880,16 @@ void FSequencer::TransformSelectedKeysAndSections(float InDeltaTime)
 			
 			if (SectionModified)
 			{
-				Key.KeyArea->MoveKey(Key.KeyHandle.GetValue(), InDeltaTime);
-				bAnythingChanged = true;
+				if (InScale != 0.f)
+				{
+					Key.KeyArea->DilateKey(Key.KeyHandle.GetValue(), InScale, GetLocalTime());
+					bAnythingChanged = true;
+				}
+				if (InDeltaTime != 0.f)
+				{
+					Key.KeyArea->MoveKey(Key.KeyHandle.GetValue(), InDeltaTime);
+					bAnythingChanged = true;
+				}
 
 				float NewKeyTime = Key.KeyArea->GetKeyTime(Key.KeyHandle.GetValue()) + InDeltaTime;
 				if (NewKeyTime > Section->GetEndTime())
@@ -1080,8 +912,16 @@ void FSequencer::TransformSelectedKeysAndSections(float InDeltaTime)
 		{
 			TSet<FKeyHandle> EmptyKeyHandles;
 			SelectedSection.Get()->SetFlags( RF_Transactional );
-			SelectedSection.Get()->MoveSection(InDeltaTime, EmptyKeyHandles);
-			bAnythingChanged = true;
+			if (InScale != 1.f)
+			{
+				SelectedSection.Get()->DilateSection(InScale, GetLocalTime(), EmptyKeyHandles);
+				bAnythingChanged = true;
+			}
+			if (InDeltaTime != 0.f)
+			{
+				SelectedSection.Get()->MoveSection(InDeltaTime, EmptyKeyHandles);
+				bAnythingChanged = true;
+			}
 		}
 	}
 
@@ -1092,18 +932,6 @@ void FSequencer::TransformSelectedKeysAndSections(float InDeltaTime)
 }
 
 
-void FSequencer::NotifyMapChanged( class UWorld* NewWorld, EMapChangeType MapChangeType )
-{
-	// @todo sequencer: We should only wipe/respawn puppets that are affected by the world that is being changed! (multi-UWorld support)
-	if( ( MapChangeType == EMapChangeType::LoadMap || MapChangeType == EMapChangeType::NewMap || MapChangeType == EMapChangeType::TearDownWorld) )
-	{
-		SpawnRegister->CleanUp(*this);
-
-		UpdateRuntimeInstances();
-	}
-}
-
-
 void FSequencer::OnActorsDropped( const TArray<TWeakObjectPtr<AActor> >& Actors )
 {
 	AddActors(Actors);
@@ -1249,40 +1077,60 @@ void FSequencer::SetSelectionRangeStart()
 }
 
 
-void FSequencer::SelectKeysInSelectionRange(const TSharedRef<FSequencerDisplayNode>& DisplayNode, const TRange<float>& SelectionRange)
+void FSequencer::SelectInSelectionRange(const TSharedRef<FSequencerDisplayNode>& DisplayNode, const TRange<float>& SelectionRange, bool bSelectKeys, bool bSelectSections)
 {
 	if (DisplayNode->GetType() == ESequencerNode::Track)
 	{
-		TSet<TSharedPtr<IKeyArea>> OutKeyAreas;
-		SequencerHelpers::GetAllKeyAreas(DisplayNode, OutKeyAreas);
-
-		for (auto KeyArea : OutKeyAreas)
+		if (bSelectKeys)
 		{
-			if (!KeyArea.IsValid())
+			TSet<TSharedPtr<IKeyArea>> OutKeyAreas;
+			SequencerHelpers::GetAllKeyAreas(DisplayNode, OutKeyAreas);
+
+			for (auto KeyArea : OutKeyAreas)
 			{
-				continue;
-			}
+				if (!KeyArea.IsValid())
+				{
+					continue;
+				}
 
-			UMovieSceneSection* Section = KeyArea->GetOwningSection();
+				UMovieSceneSection* Section = KeyArea->GetOwningSection();
 
-			if (Section == nullptr)
-			{
-				continue;
+				if (Section == nullptr)
+				{
+					continue;
+				}
+
+				if (bSelectKeys)
+				{
+					TSet<FKeyHandle> OutKeyHandles;
+					Section->GetKeyHandles(OutKeyHandles, SelectionRange);
+
+					for (auto KeyHandle : OutKeyHandles)
+					{
+						Selection.AddToSelection(FSequencerSelectedKey(*Section, KeyArea, KeyHandle));
+					}
+				}
 			}
+		}
 
-			TSet<FKeyHandle> OutKeyHandles;
-			Section->GetKeyHandles(OutKeyHandles, SelectionRange);
+		if (bSelectSections)
+		{
+			TSet<TWeakObjectPtr<UMovieSceneSection>> OutSections;
+			SequencerHelpers::GetAllSections(DisplayNode, OutSections);
 
-			for (auto KeyHandle : OutKeyHandles)
+			for (auto Section : OutSections)
 			{
-				Selection.AddToSelection(FSequencerSelectedKey(*Section, KeyArea, KeyHandle));
+				if (Section.IsValid())
+				{
+					Selection.AddToSelection(Section.Get());
+				}
 			}
 		}
 	}
 
 	for (const auto& ChildNode : DisplayNode->GetChildNodes())
 	{
-		SelectKeysInSelectionRange(ChildNode, SelectionRange);
+		SelectInSelectionRange(ChildNode, SelectionRange, bSelectKeys, bSelectSections);
 	}
 }
 
@@ -1291,7 +1139,7 @@ void FSequencer::ResetSelectionRange()
 	SetSelectionRange(TRange<float>::Empty());
 }
 
-void FSequencer::SelectKeysInSelectionRange()
+void FSequencer::SelectInSelectionRange(bool bSelectKeys, bool bSelectSections)
 {
 	UMovieSceneSequence* Sequence = GetFocusedMovieSceneSequence();
 	UMovieScene* MovieScene = Sequence->GetMovieScene();
@@ -1301,7 +1149,7 @@ void FSequencer::SelectKeysInSelectionRange()
 
 	for (const TSharedRef<FSequencerDisplayNode>& DisplayNode : NodeTree->GetRootNodes())
 	{
-		SelectKeysInSelectionRange(DisplayNode, SelectionRange);
+		SelectInSelectionRange(DisplayNode, SelectionRange, bSelectKeys, bSelectSections);
 	}
 }
 
@@ -1316,15 +1164,41 @@ void FSequencer::SetPlaybackRange(TRange<float> Range)
 {
 	if (ensure(Range.HasLowerBound() && Range.HasUpperBound() && !Range.IsDegenerate()))
 	{
-		const FScopedTransaction Transaction(LOCTEXT("SetPlaybackRange_Transaction", "Set Playback Range"));
-
 		UMovieScene* FocusedMovieScene = GetFocusedMovieSceneSequence()->GetMovieScene();
-		FocusedMovieScene->SetPlaybackRange(Range.GetLowerBoundValue(), Range.GetUpperBoundValue());
+		
+		if (!FocusedMovieScene->IsPlaybackRangeLocked())
+		{
+			const FScopedTransaction Transaction(LOCTEXT("SetPlaybackRange_Transaction", "Set Playback Range"));
 
-		bNeedsEvaluate = true;
+			FocusedMovieScene->SetPlaybackRange(Range.GetLowerBoundValue(), Range.GetUpperBoundValue());
+
+			bNeedsEvaluate = true;
+		}
 	}
 }
 
+bool FSequencer::IsPlaybackRangeLocked() const
+{
+	UMovieSceneSequence* FocusedMovieSceneSequence = GetFocusedMovieSceneSequence();
+	if (FocusedMovieSceneSequence != nullptr)
+	{
+		return FocusedMovieSceneSequence->GetMovieScene()->IsPlaybackRangeLocked();
+	}
+
+	return false;
+}
+
+void FSequencer::TogglePlaybackRangeLocked()
+{
+	UMovieSceneSequence* FocusedMovieSceneSequence = GetFocusedMovieSceneSequence();
+	if ( FocusedMovieSceneSequence != nullptr )
+	{
+		FScopedTransaction TogglePlaybackRangeLockTransaction( NSLOCTEXT( "Sequencer", "TogglePlaybackRangeLocked", "Toggle playback range lock" ) );
+		UMovieScene* MovieScene = FocusedMovieSceneSequence->GetMovieScene();
+		MovieScene->Modify();
+		MovieScene->SetPlaybackRangeLocked( !MovieScene->IsPlaybackRangeLocked() );
+	}
+}
 
 void FSequencer::ResetViewRange()
 {
@@ -1443,7 +1317,7 @@ void FSequencer::UpdatePlaybackRange()
 			}
 
 			// When the playback range is determined by the section bounds, don't mark the change in the playback range otherwise the scene will be marked dirty
-			if (!NewBounds.IsDegenerate())
+			if (!NewBounds.IsDegenerate() && !FocusedMovieScene->IsPlaybackRangeLocked())
 			{
 				const bool bAlwaysMarkDirty = false;
 				FocusedMovieScene->SetPlaybackRange(NewBounds.GetLowerBoundValue(), NewBounds.GetUpperBoundValue(), bAlwaysMarkDirty);
@@ -1464,7 +1338,7 @@ void FSequencer::UpdatePlaybackRange()
 			}
 		}
 			
-		if (!NewBounds.IsEmpty())
+		if (!NewBounds.IsEmpty() && !FocusedMovieScene->IsPlaybackRangeLocked())
 		{
 			FocusedMovieScene->SetPlaybackRange(NewBounds.GetLowerBoundValue(), NewBounds.GetUpperBoundValue());
 		}
@@ -1893,7 +1767,7 @@ FGuid FSequencer::GetHandleToObject( UObject* Object, bool bCreateHandleIfMissin
 		ObjectGuid = FindUnspawnedObjectGuid(*Object, *FocusedMovieSceneSequence);
 	}
 
-	if (ObjectGuid.IsValid())
+	if (ObjectGuid.IsValid() || IsReadOnly())
 	{
 		return ObjectGuid;
 	}
@@ -2173,64 +2047,6 @@ void FSequencer::UpdateRuntimeInstances()
 	FEditorSupportDelegates::RedrawAllViewports.Broadcast();
 }
 
-void FSequencer::AddLevelViewportMenuExtender()
-{
-	FLevelEditorModule& LevelEditorModule = FModuleManager::Get().LoadModuleChecked<FLevelEditorModule>("LevelEditor");
-	auto& MenuExtenders = LevelEditorModule.GetAllLevelViewportContextMenuExtenders();
-
-	if (LevelViewportExtenderHandle.IsValid())
-	{
-		typedef FLevelEditorModule::FLevelViewportMenuExtender_SelectedActors DelegateType;
-		MenuExtenders.RemoveAll([=](const DelegateType& In){ return In.GetHandle() == LevelViewportExtenderHandle; });
-	}
-
-	MenuExtenders.Add(FLevelEditorModule::FLevelViewportMenuExtender_SelectedActors::CreateSP(this, &FSequencer::GetLevelViewportExtender));
-	LevelViewportExtenderHandle = MenuExtenders.Last().GetHandle();
-}
-
-void FSequencer::RemoveLevelViewportMenuExtender()
-{
-	if (LevelViewportExtenderHandle.IsValid())
-	{
-		FLevelEditorModule* LevelEditorModule = FModuleManager::Get().GetModulePtr<FLevelEditorModule>("LevelEditor");
-		if (LevelEditorModule)
-		{
-			typedef FLevelEditorModule::FLevelViewportMenuExtender_SelectedActors DelegateType;
-			LevelEditorModule->GetAllLevelViewportContextMenuExtenders().RemoveAll([=](const DelegateType& In){ return In.GetHandle() == LevelViewportExtenderHandle; });
-		}
-	}
-}
-
-TSharedRef<FExtender> FSequencer::GetLevelViewportExtender(const TSharedRef<FUICommandList> CommandList, const TArray<AActor*> InActors)
-{
-	TSharedRef<FExtender> Extender = MakeShareable(new FExtender);
-
-	FText ActorName;
-	if (InActors.Num() == 1)
-	{
-		ActorName = FText::Format(LOCTEXT("ActorNameSingular", "\"{0}\""), FText::FromString(InActors[0]->GetActorLabel()));
-	}
-	else if (InActors.Num() > 1)
-	{
-		ActorName = FText::Format(LOCTEXT("ActorNamePlural", "{0} Actors"), FText::AsNumber(InActors.Num()));
-	}
-
-	FLevelEditorModule& LevelEditor = FModuleManager::GetModuleChecked<FLevelEditorModule>(TEXT("LevelEditor"));
-	TSharedRef<FUICommandList> LevelEditorCommandBindings = LevelEditor.GetGlobalLevelEditorActions();
-
-	Extender->AddMenuExtension("ActorControl", EExtensionHook::After, LevelEditorCommandBindings, FMenuExtensionDelegate::CreateLambda(
-		[this, ActorName](FMenuBuilder& MenuBuilder){
-			MenuBuilder.BeginSection("SequenceRecording", LOCTEXT("SequenceRecordingHeading", "Sequence Recording"));
-
-			MenuBuilder.AddMenuEntry(FSequencerCommands::Get().RecordSelectedActors, NAME_None, FText::Format(LOCTEXT("RecordSelectedActorsText", "Record {0} In Sequencer"), ActorName));
-
-			MenuBuilder.EndSection();
-		}
-	));
-
-	return Extender;
-}
-
 void FSequencer::RecordSelectedActors()
 {
 	ISequenceRecorder& SequenceRecorder = FModuleManager::LoadModuleChecked<ISequenceRecorder>("SequenceRecorder");
@@ -2437,213 +2253,6 @@ TSharedRef<SWidget> FSequencer::OnCreateTransportRecord()
 	return RecordButton;
 }
 
-namespace FaderConstants
-{	
-	/** The opacity when we are hovered */
-	const float HoveredOpacity = 1.0f;
-	/** The opacity when we are not hovered */
-	const float NonHoveredOpacity = 0.75f;
-	/** The amount of time spent actually fading in or out */
-	const float FadeTime = 0.15f;
-}
-
-/** Wrapper widget allowing us to fade widgets in and out on hover state */
-class SFader : public SBorder
-{
-public:
-	SLATE_BEGIN_ARGS(SFader)
-		: _Content()
-	{}
-
-	SLATE_DEFAULT_SLOT(FArguments, Content)
-
-	SLATE_END_ARGS()
-
-	void Construct(const FArguments& InArgs)
-	{
-		FadeInSequence = FCurveSequence(0.0f, FaderConstants::FadeTime);
-		FadeOutSequence = FCurveSequence(0.0f, FaderConstants::FadeTime);
-		FadeOutSequence.JumpToEnd();
-
-		SBorder::Construct(SBorder::FArguments()
-			.BorderImage(FCoreStyle::Get().GetBrush("NoBorder"))
-			.Padding(0.0f)
-			.VAlign(VAlign_Center)
-			.ColorAndOpacity(this, &SFader::GetColorAndOpacity)
-			.Content()
-			[
-				InArgs._Content.Widget
-			]);
-	}
-
-	FLinearColor GetColorAndOpacity() const
-	{
-		FLinearColor Color = FLinearColor::White;
-	
-		if(FadeOutSequence.IsPlaying() || !bIsHovered)
-		{
-			Color.A = FMath::Lerp(FaderConstants::HoveredOpacity, FaderConstants::NonHoveredOpacity, FadeOutSequence.GetLerp());
-		}
-		else
-		{
-			Color.A = FMath::Lerp(FaderConstants::NonHoveredOpacity, FaderConstants::HoveredOpacity, FadeInSequence.GetLerp());
-		}
-
-		return Color;
-	}
-
-	virtual void OnMouseEnter(const FGeometry& MyGeometry, const FPointerEvent& MouseEvent) override
-	{
-		if(!FSlateApplication::Get().IsUsingHighPrecisionMouseMovment())
-		{
-			bIsHovered = true;
-			if(FadeOutSequence.IsPlaying())
-			{
-				// Fade out is already playing so just force the fade in curve to the end so we don't have a "pop" 
-				// effect from quickly resetting the alpha
-				FadeInSequence.JumpToEnd();
-			}
-			else
-			{
-				FadeInSequence.Play(AsShared());
-			}
-		}
-	}
-
-	virtual void OnMouseLeave(const FPointerEvent& MouseEvent) override
-	{
-		if(!FSlateApplication::Get().IsUsingHighPrecisionMouseMovment())
-		{
-			bIsHovered = false;
-			FadeOutSequence.Play(AsShared());
-		}
-	}
-
-private:
-	/** Curve sequence for fading out the widget */
-	FCurveSequence FadeOutSequence;
-	/** Curve sequence for fading in the widget */
-	FCurveSequence FadeInSequence;
-};
-
-void FSequencer::SetViewportTransportControlsVisibility(bool bVisible)
-{
-	Settings->SetShowViewportTransportControls(bVisible);
-}
-
-void FSequencer::AttachTransportControlsToViewports()
-{
-	FLevelEditorModule* Module = FModuleManager::Get().LoadModulePtr<FLevelEditorModule>("LevelEditor");
-	
-	if (Module == nullptr)
-	{
-		return;
-	}
-
-	// register level editor viewport menu extenders
-	ViewMenuExtender = FLevelEditorModule::FLevelEditorMenuExtender::CreateSP(this, &FSequencer::OnExtendLevelEditorViewMenu);
-	Module->GetAllLevelViewportOptionsMenuExtenders().Add(ViewMenuExtender);		
-	LevelEditorExtenderDelegateHandle = Module->GetAllLevelViewportOptionsMenuExtenders().Last().GetHandle();
-
-	TSharedPtr<ILevelEditor> LevelEditor = Module->GetFirstLevelEditor();
-	const TArray< TSharedPtr<ILevelViewport> >& LevelViewports = LevelEditor->GetViewports();
-		
-	FEditorWidgetsModule& EditorWidgetsModule = FModuleManager::Get().LoadModuleChecked<FEditorWidgetsModule>( "EditorWidgets" );
-
-	// Transport controls in the viewport need to have something that is focusable to prevent mouse input dropping through to the viewport.
-	// We don't want the buttons themselves to be focusable, so we just add them to a parent box that is.
-	class SFocusableBox : public SBox
-	{
-		virtual bool SupportsKeyboardFocus() const
-		{
-			return true;
-		}
-	};
-
-	for (int32 i = 0; i < LevelViewports.Num(); ++i)
-	{
-		const TSharedPtr<ILevelViewport>& LevelViewport = LevelViewports[i];
-			
-		TSharedRef<SHorizontalBox> TransportControl =
-			SNew(SHorizontalBox)
-				.Visibility(EVisibility::SelfHitTestInvisible)
-
-			+SHorizontalBox::Slot()
-				.FillWidth(1)
-				.HAlign(HAlign_Center)
-				.VAlign(VAlign_Bottom)
-				.Padding(4.f)
-				[
-					SNew(SFader)
-					.Content()
-					[
-						SNew(SBorder)
-						.Padding(4.f)
-						.Cursor( EMouseCursor::Default )
-						.BorderImage( FEditorStyle::GetBrush( "ToolPanel.GroupBorder" ) )
-						.Visibility(this, &FSequencer::GetTransportControlVisibility, LevelViewport)
-						.Content()
-						[
-							SNew(SHorizontalBox)
-							+SHorizontalBox::Slot()
-							.AutoWidth()
-							.HAlign(HAlign_Center)
-							.VAlign(VAlign_Center)
-							[
-								SNew(SFocusableBox)
-								[
-									MakeTransportControls(false)
-								]
-							]
-							+SHorizontalBox::Slot()
-							.AutoWidth()
-							.HAlign(HAlign_Center)
-							.VAlign(VAlign_Top)
-							[
-								SNew(SButton)
-								.ButtonStyle(&FEditorStyle::Get().GetWidgetStyle< FButtonStyle >("Sequencer.Transport.CloseButton"))
-								.ToolTipText(LOCTEXT("CloseTransportControlsToolTip", "Hide the transport controls. You can re-enable transport controls from the viewport menu."))
-								.OnClicked_Lambda([this]() { Settings->SetShowViewportTransportControls(!Settings->GetShowViewportTransportControls()); return FReply::Handled(); })
-								
-							]
-						]
-					]
-				];
-
-		LevelViewport->AddOverlayWidget(TransportControl);
-
-		TransportControls.Add(LevelViewport, TransportControl);
-	}
-}
-
-
-void FSequencer::DetachTransportControlsFromViewports()
-{
-	FLevelEditorModule* Module = FModuleManager::Get().LoadModulePtr<FLevelEditorModule>("LevelEditor");
-	if (Module)
-	{
-		TSharedPtr<ILevelEditor> LevelEditor = Module->GetFirstLevelEditor();
-		if (LevelEditor.IsValid())
-		{
-			const TArray< TSharedPtr<ILevelViewport> >& LevelViewports = LevelEditor->GetViewports();
-		
-			for (int32 i = 0; i < LevelViewports.Num(); ++i)
-			{
-				const TSharedPtr<ILevelViewport>& LevelViewport = LevelViewports[i];
-
-				TSharedPtr<SWidget>* TransportControl = TransportControls.Find(LevelViewport);
-				if (TransportControl && TransportControl->IsValid())
-				{
-					LevelViewport->RemoveOverlayWidget(TransportControl->ToSharedRef());
-				}
-			}
-		}
-
-		Module->GetAllLevelViewportOptionsMenuExtenders().RemoveAll([&](const FLevelEditorModule::FLevelEditorMenuExtender& Delegate) {
-			return Delegate.GetHandle() == LevelEditorExtenderDelegateHandle;
-		});
-	}
-}
 
 UObject* FSequencer::FindSpawnedObjectOrTemplate(const FGuid& BindingId)
 {
@@ -2688,41 +2297,6 @@ UObject* FSequencer::FindSpawnedObjectOrTemplate(const FGuid& BindingId)
 	return nullptr;
 }
 
-TSharedRef<FExtender> FSequencer::OnExtendLevelEditorViewMenu(const TSharedRef<FUICommandList> CommandList)
-{
-	TSharedRef<FExtender> Extender(new FExtender());
-
-	Extender->AddMenuExtension(
-		"LevelViewportViewportOptions2",
-		EExtensionHook::First,
-		NULL,
-		FMenuExtensionDelegate::CreateSP(this, &FSequencer::CreateTransportToggleMenuEntry));
-
-	return Extender;
-}
-
-void FSequencer::CreateTransportToggleMenuEntry(FMenuBuilder& MenuBuilder)
-{
-	MenuBuilder.AddMenuEntry(
-		LOCTEXT("ShowTransportControls", "Show Transport Controls"),
-		LOCTEXT("ShowTransportControlsToolTip", "Show or hide the Sequencer transport controls when a sequence is active."),
-		FSlateIcon(),
-		FUIAction(
-			FExecuteAction::CreateLambda([this](){ Settings->SetShowViewportTransportControls(!Settings->GetShowViewportTransportControls()); }),
-			FCanExecuteAction(), 
-			FGetActionCheckState::CreateLambda([this](){ return Settings->GetShowViewportTransportControls() ? ECheckBoxState::Checked : ECheckBoxState::Unchecked; })),
-		NAME_None,
-		EUserInterfaceActionType::ToggleButton);
-}
-
-EVisibility FSequencer::GetTransportControlVisibility(TSharedPtr<ILevelViewport> LevelViewport) const
-{
-	FLevelEditorViewportClient& ViewportClient = LevelViewport->GetLevelViewportClient();
-	return Settings->GetShowViewportTransportControls() &&
-		ViewportClient.ViewportType == LVT_Perspective &&
-		ViewportClient.AllowsCinematicPreview() ? EVisibility::Visible : EVisibility::Collapsed;
-}
-
 FReply FSequencer::OnPlay(bool bTogglePlay, float InPlayRate)
 {
 	if( (PlaybackState == EMovieScenePlayerStatus::Playing ||
@@ -3347,6 +2921,10 @@ UObject* FSequencer::GetCurrentAsset() const
 	return GetFocusedMovieSceneSequence()->GetMovieScene()->GetOuter();
 }
 
+bool FSequencer::IsReadOnly() const
+{
+	return bReadOnly;
+}
 
 void FSequencer::VerticalScroll(float ScrollAmountUnits)
 {
@@ -3745,59 +3323,11 @@ void FSequencer::PlaceActorInFrontOfCamera( AActor* ActorCDO )
 }
 
 
-void FSequencer::GetActorRecordingState( bool& bIsRecording /* In+Out */ ) const
-{
-	if (IsRecordingLive())
-	{
-		bIsRecording = true;
-	}
-}
-
-
 void FSequencer::PostUndo(bool bSuccess)
 {
 	NotifyMovieSceneDataChanged( EMovieSceneDataChangeType::Unknown );
-	bUpdatingSequencerSelection = true;
 	SynchronizeSequencerSelectionWithExternalSelection();
-	bUpdatingSequencerSelection = false;
 	OnActivateSequenceEvent.Broadcast(ActiveTemplateIDs.Top());
-}	
-
-
-void FSequencer::OnPreSaveWorld(uint32 SaveFlags, class UWorld* World)
-{
-	// Restore the saved state so that the level save can save that instead of the animated state.
-	RestorePreAnimatedState();
-}
-
-
-void FSequencer::OnPostSaveWorld(uint32 SaveFlags, class UWorld* World, bool bSuccess)
-{
-	// Reset the time after saving so that an update will be triggered to put objects back to their animated state.
-	FMovieSceneEvaluationRange Range = PlayPosition.JumpTo(ScrubPosition, GetFocusedMovieSceneSequence()->GetMovieScene()->GetOptionalFixedFrameInterval());
-	EvaluateInternal(Range);
-}
-
-
-void FSequencer::OnNewCurrentLevel()
-{
-	ActivateSequencerEditorMode();
-}
-
-
-void FSequencer::OnMapOpened(const FString& Filename, bool bLoadAsTemplate)
-{
-	ActivateSequencerEditorMode();
-}
-
-void FSequencer::OnLevelAdded(ULevel* InLevel, UWorld* InWorld)
-{
-	bNeedInstanceRefresh = true;
-}
-
-void FSequencer::OnLevelRemoved(ULevel* InLevel, UWorld* InWorld)
-{
-	bNeedInstanceRefresh = true;
 }
 
 void FSequencer::OnNewActorsDropped(const TArray<UObject*>& DroppedObjects, const TArray<AActor*>& DroppedActors)
@@ -3959,95 +3489,11 @@ void FSequencer::OnNewActorsDropped(const TArray<UObject*>& DroppedObjects, cons
 
 		NotifyMovieSceneDataChanged( EMovieSceneDataChangeType::MovieSceneStructureItemsChanged );
 
-		bUpdatingSequencerSelection = true;
 		SynchronizeSequencerSelectionWithExternalSelection();
-		bUpdatingSequencerSelection = false;
-	}
-}
-
-void FSequencer::ActivateSequencerEditorMode()
-{
-	if( GLevelEditorModeTools().IsModeActive( FSequencerEdMode::EM_SequencerMode ) )
-	{
-		GLevelEditorModeTools().DeactivateMode( FSequencerEdMode::EM_SequencerMode );
-	}
-
-	GLevelEditorModeTools().ActivateMode( FSequencerEdMode::EM_SequencerMode );
-
-	FSequencerEdMode* SequencerEdMode = (FSequencerEdMode*)(GLevelEditorModeTools().GetActiveMode(FSequencerEdMode::EM_SequencerMode));
-	SequencerEdMode->SetSequencer(ActiveSequencerPtr);
-}
-
-
-void FSequencer::OnPreBeginPIE(bool bIsSimulating)
-{
-	RestorePreAnimatedState();
-	PrePossessionViewTargets.Reset();
-
-	bNeedInstanceRefresh = true;
-}
-
-
-void FSequencer::OnEndPlayMap()
-{
-	// Update and clear any stale bindings 
-	UpdateRuntimeInstances();
-	
-	// Update for new bindings on next tick
-	bNeedInstanceRefresh = true;
-}
-
-
-void FSequencer::OnEndPIE(bool bIsSimulating)
-{
-	OnEndPlayMap();
-}
-
-
-void FSequencer::ActivateDetailKeyframeHandler()
-{			
-	// Add sequencer detail keyframe handler
-	static const FName DetailsTabIdentifiers[] = { "LevelEditorSelectionDetails", "LevelEditorSelectionDetails2", "LevelEditorSelectionDetails3", "LevelEditorSelectionDetails4" };
-	FPropertyEditorModule& EditModule = FModuleManager::Get().GetModuleChecked<FPropertyEditorModule>("PropertyEditor");
-
-	for (const FName& DetailsTabIdentifier : DetailsTabIdentifiers)
-	{
-		TSharedPtr<IDetailsView> DetailsView = EditModule.FindDetailView(DetailsTabIdentifier);
-
-		if(DetailsView.IsValid())
-		{
-			DetailsView->SetKeyframeHandler(DetailKeyframeHandler);
-		}
 	}
 }
 
 
-void FSequencer::DeactivateDetailKeyframeHandler()
-{
-	static const FName DetailsTabIdentifiers[] = { "LevelEditorSelectionDetails", "LevelEditorSelectionDetails2", "LevelEditorSelectionDetails3", "LevelEditorSelectionDetails4" };
-	FPropertyEditorModule& EditModule = FModuleManager::Get().GetModuleChecked<FPropertyEditorModule>("PropertyEditor");
-
-	for (const FName& DetailsTabIdentifier : DetailsTabIdentifiers)
-	{
-		TSharedPtr<IDetailsView> DetailsView = EditModule.FindDetailView(DetailsTabIdentifier);
-
-		if (DetailsView.IsValid())
-		{
-			if (DetailsView->GetKeyframeHandler() == DetailKeyframeHandler)
-			{
-				DetailsView->SetKeyframeHandler(0);
-			}
-		}
-	}
-}
-
-
-void FSequencer::OnPropertyEditorOpened()
-{
-	ActivateDetailKeyframeHandler();
-}
-
-
 void FSequencer::UpdatePreviewLevelViewportClientFromCameraCut(FLevelEditorViewportClient& InViewportClient, UObject* InCameraObject, bool bJumpCut) const
 {
 	AActor* CameraActor = Cast<AActor>(InCameraObject);
@@ -4229,6 +3675,11 @@ void FSequencer::SaveCurrentMovieSceneAs()
 	TArray<UObject*> SavedAssets;
 	FEditorFileUtils::SaveAssetsAs(AssetsToSave, SavedAssets);
 
+	if (SavedAssets.Num() == 0)
+	{
+		return;
+	}
+
 	if ((SavedAssets[0] != AssetsToSave[0]) && (SavedAssets[0] != nullptr))
 	{
 		FAssetEditorManager& AssetEditorManager = FAssetEditorManager::Get();
@@ -4266,45 +3717,36 @@ void FSequencer::AddActors(const TArray<TWeakObjectPtr<AActor> >& InActors)
 	{
 		SequencerWidget->UpdateLayoutTree();
 
-		bUpdatingSequencerSelection = true;
 		SynchronizeSequencerSelectionWithExternalSelection();
-		bUpdatingSequencerSelection = false;
 	}
 }
 
 
-void FSequencer::OnSectionSelectionChanged()
-{
-}
-
-
 void FSequencer::OnSelectedOutlinerNodesChanged()
 {
-	if ( bUpdatingSequencerSelection == false )
-	{
-		bUpdatingExternalSelection = true;
-		SynchronizeExternalSelectionWithSequencerSelection();
-		bUpdatingExternalSelection = false;
-	}
+	SynchronizeExternalSelectionWithSequencerSelection();
 	OnSelectionChangedObjectGuidsDelegate.Broadcast(Selection.GetBoundObjectsGuids());
 }
 
 
 void FSequencer::SynchronizeExternalSelectionWithSequencerSelection()
 {
-	if ( !IsLevelEditorSequencer() )
+	if ( bUpdatingSequencerSelection || !IsLevelEditorSequencer() )
 	{
 		return;
 	}
 
+	TGuardValue<bool> Guard(bUpdatingExternalSelection, true);
+
 	TSet<AActor*> SelectedSequencerActors;
+	TSet<USceneComponent*> SelectedSequencerComponents;
 
 	TSet<TSharedRef<FSequencerDisplayNode> > DisplayNodes = Selection.GetNodesWithSelectedKeysOrSections();
 	DisplayNodes.Append(Selection.GetSelectedOutlinerNodes());
 
 	for ( TSharedRef<FSequencerDisplayNode> DisplayNode : DisplayNodes)
 	{
-		// Get the root object binding node.
+		// Get the closest object binding node.
 		TSharedPtr<FSequencerDisplayNode> CurrentNode = DisplayNode;
 		TSharedPtr<FSequencerObjectBindingNode> ObjectBindingNode;
 		while ( CurrentNode.IsValid() )
@@ -4312,11 +3754,12 @@ void FSequencer::SynchronizeExternalSelectionWithSequencerSelection()
 			if ( CurrentNode->GetType() == ESequencerNode::Object )
 			{
 				ObjectBindingNode = StaticCastSharedPtr<FSequencerObjectBindingNode>(CurrentNode);
+				break;
 			}
 			CurrentNode = CurrentNode->GetParent();
 		}
 
-		// If the root node is an object node, try to get the actor nodes from it.
+		// If the closest node is an object node, try to get the actor/component nodes from it.
 		if ( ObjectBindingNode.IsValid() )
 		{
 			for (auto RuntimeObject : FindBoundObjects(ObjectBindingNode->GetObjectBinding(), ActiveTemplateIDs.Top()) )
@@ -4326,17 +3769,41 @@ void FSequencer::SynchronizeExternalSelectionWithSequencerSelection()
 				{
 					SelectedSequencerActors.Add( Actor );
 				}
+
+				USceneComponent* SceneComponent = Cast<USceneComponent>(RuntimeObject.Get());
+				if ( SceneComponent != nullptr )
+				{
+					SelectedSequencerComponents.Add( SceneComponent );
+
+					Actor = SceneComponent->GetOwner();
+					if ( Actor != nullptr )
+					{
+						SelectedSequencerActors.Add( Actor );
+					}
+				}
 			}
 		}
 	}
 
-	const FScopedTransaction Transaction( NSLOCTEXT( "Sequencer", "UpdatingActorSelection", "Select Actors" ) );
-
 	const bool bNotifySelectionChanged = false;
 	const bool bDeselectBSP = true;
 	const bool bWarnAboutTooManyActors = false;
 	const bool bSelectEvenIfHidden = true;
 
+	if (SelectedSequencerComponents.Num() + SelectedSequencerActors.Num() == 0)
+	{
+		if (GEditor->GetSelectedActorCount())
+		{
+			const FScopedTransaction Transaction( NSLOCTEXT( "Sequencer", "UpdatingActorComponentSelectionNone", "Select None" ) );
+			GEditor->SelectNone( bNotifySelectionChanged, bDeselectBSP, bWarnAboutTooManyActors );
+			GEditor->NoteSelectionChange();
+		}
+		return;
+	}
+
+	const FScopedTransaction Transaction( NSLOCTEXT( "Sequencer", "UpdatingActorComponentSelection", "Select Actors/Components" ) );
+
+
 	GEditor->GetSelectedActors()->Modify();
 	GEditor->GetSelectedActors()->BeginBatchSelectOperation();
 
@@ -4348,18 +3815,21 @@ void FSequencer::SynchronizeExternalSelectionWithSequencerSelection()
 	}
 
 	GEditor->GetSelectedActors()->EndBatchSelectOperation();
-	GEditor->NoteSelectionChange();
-}
 
-
-void FSequencer::OnActorSelectionChanged( UObject* )
-{
-	if ( bUpdatingExternalSelection == false )
+	if (SelectedSequencerComponents.Num())
 	{
-		bUpdatingSequencerSelection = true;
-		SynchronizeSequencerSelectionWithExternalSelection();
-		bUpdatingSequencerSelection = false;
+		GEditor->GetSelectedComponents()->Modify();
+		GEditor->GetSelectedComponents()->BeginBatchSelectOperation();
+
+		for ( USceneComponent* SelectedSequencerComponent : SelectedSequencerComponents )
+		{
+			GEditor->SelectComponent( SelectedSequencerComponent, true, bNotifySelectionChanged, bSelectEvenIfHidden );
+		}
+
+		GEditor->GetSelectedComponents()->EndBatchSelectOperation();
 	}
+		
+	GEditor->NoteSelectionChange();
 }
 
 
@@ -4382,13 +3852,12 @@ void GetRootObjectBindingNodes(const TArray<TSharedRef<FSequencerDisplayNode>>&
 
 void FSequencer::SynchronizeSequencerSelectionWithExternalSelection()
 {
-	if ( !IsLevelEditorSequencer() )
+	if ( bUpdatingExternalSelection || !IsLevelEditorSequencer() )
 	{
 		return;
 	}
 
-	TArray<TSharedRef<FSequencerObjectBindingNode>> RootObjectBindingNodes;
-	GetRootObjectBindingNodes( NodeTree->GetRootNodes(), RootObjectBindingNodes );
+	TGuardValue<bool> Guard(bUpdatingSequencerSelection, true);
 
 	// If all nodes are already selected, do nothing. This ensures that when an undo event happens, 
 	// nodes are not cleared and reselected, which can cause issues with the curve editor auto-fitting 
@@ -4396,12 +3865,19 @@ void FSequencer::SynchronizeSequencerSelectionWithExternalSelection()
 	bool bAllAlreadySelected = true;
 
 	TSet<TSharedRef<FSequencerDisplayNode>> NodesToSelect;
-	for ( TSharedRef<FSequencerObjectBindingNode> ObjectBindingNode : RootObjectBindingNodes )
+	for (auto ObjectBinding : NodeTree->GetObjectBindingMap() )
 	{
+		if (!ObjectBinding.Value.IsValid())
+		{
+			continue;
+		}
+
+		TSharedRef<FSequencerObjectBindingNode> ObjectBindingNode = ObjectBinding.Value.ToSharedRef();
 		for ( TWeakObjectPtr<UObject> RuntimeObjectPtr : FindBoundObjects(ObjectBindingNode->GetObjectBinding(), ActiveTemplateIDs.Top()) )
 		{
 			UObject* RuntimeObject = RuntimeObjectPtr.Get();
-			if ( RuntimeObject != nullptr && GEditor->GetSelectedActors()->IsSelected( RuntimeObject ) )
+			if ( RuntimeObject != nullptr &&
+				(GEditor->GetSelectedActors()->IsSelected( RuntimeObject ) || GEditor->GetSelectedComponents()->IsSelected( RuntimeObject) ))
 			{
 				NodesToSelect.Add( ObjectBindingNode );
 
@@ -4838,29 +4314,15 @@ void FSequencer::PasteCopiedTracks(TArray<TSharedPtr<FSequencerObjectBindingNode
 	for (TSharedPtr<FSequencerObjectBindingNode> ObjectNode : ObjectNodes)
 	{
 		FGuid ObjectGuid = ObjectNode->GetObjectBinding();
-		UObject* Context = GetPlaybackContext();
-		TArray<UObject*, TInlineAllocator<1>> PossessableObjects = GetFocusedMovieSceneSequence()->LocateBoundObjects(ObjectGuid, Context);
 
-		if (PossessableObjects.Num() == 0)
-		{
-			FNotificationInfo Info(LOCTEXT("InvalidTargetObject", "Invalid Target Object"));
-			Info.FadeInDuration = 0.1f;
-			Info.FadeOutDuration = 0.5f;
-			Info.ExpireDuration = 2.5f;
-			auto NotificationItem = FSlateNotificationManager::Get().AddNotification(Info);
-
-			NotificationItem->SetCompletionState(SNotificationItem::CS_Success);
-			NotificationItem->ExpireAndFadeout();
-			continue;
-		}
+		TArray<UMovieSceneTrack*> NewTracks;
+		FSequencer::ImportTracksFromText(TextToImport, NewTracks);
 
-		for (UMovieSceneTrack* ImportedTrack : ImportedTracks)
+		for (UMovieSceneTrack* NewTrack : NewTracks)
 		{
-			UMovieScenePropertyTrack* CastedPropertyTrack = Cast<UMovieScenePropertyTrack>(ImportedTrack);
-
-			if (!CastedPropertyTrack)
+			if (!GetFocusedMovieSceneSequence()->GetMovieScene()->AddGivenTrack(NewTrack, ObjectGuid))
 			{
-				FNotificationInfo Info(LOCTEXT("CopiedTrackIsNotPropertyTrack", "Can't Paste: Copied Track isn't a Property Track"));
+				FNotificationInfo Info(LOCTEXT("TrackAlreadyBound", "Can't Paste: Binding doesn't exist"));
 				Info.FadeInDuration = 0.1f;
 				Info.FadeOutDuration = 0.5f;
 				Info.ExpireDuration = 2.5f;
@@ -4871,46 +4333,12 @@ void FSequencer::PasteCopiedTracks(TArray<TSharedPtr<FSequencerObjectBindingNode
 
 				continue;
 			}
-
-			FString PropertyName = CastedPropertyTrack->GetPropertyName().ToString();
-			PropertyName.RemoveFromStart("b", ESearchCase::CaseSensitive);
-			FString FunctionName = "Set" + PropertyName;
-			if (PossessableObjects[0]->FindFunction(FName(*FunctionName)))
-			{
-				if (!GetFocusedMovieSceneSequence()->GetMovieScene()->AddGivenTrack(ImportedTrack, ObjectGuid))
-				{
-					FNotificationInfo Info(LOCTEXT("PropertyAlreadyBound", "Can't Paste: Property already bound to a track"));
-					Info.FadeInDuration = 0.1f;
-					Info.FadeOutDuration = 0.5f;
-					Info.ExpireDuration = 2.5f;
-					auto NotificationItem = FSlateNotificationManager::Get().AddNotification(Info);
-
-					NotificationItem->SetCompletionState(SNotificationItem::CS_Success);
-					NotificationItem->ExpireAndFadeout();
-
-					continue;
-				}
-				else
-				{
-					NotifyMovieSceneDataChanged(EMovieSceneDataChangeType::MovieSceneStructureItemsChanged);
-				}
-			}
 			else
 			{
-				FNotificationInfo Info(LOCTEXT("PropertyNotFound", "Can't Paste: Target Object doesn't have this property"));
-				Info.FadeInDuration = 0.1f;
-				Info.FadeOutDuration = 0.5f;
-				Info.ExpireDuration = 2.5f;
-				auto NotificationItem = FSlateNotificationManager::Get().AddNotification(Info);
-
-				NotificationItem->SetCompletionState(SNotificationItem::CS_Success);
-				NotificationItem->ExpireAndFadeout();
-
-				continue;
+				NotifyMovieSceneDataChanged(EMovieSceneDataChangeType::MovieSceneStructureItemsChanged);
 			}
 		}
 	}
-	
 }
 
 
@@ -4923,9 +4351,13 @@ public:
 	}
 
 	// FCustomizableTextObjectFactory implementation
-	virtual bool CanCreateClass(UClass* ObjectClass, bool& bOmitSubObjs) const override
+	virtual bool CanCreateClass(UClass* InObjectClass, bool& bOmitSubObjs) const override
 	{
-		return true;
+		if (InObjectClass->IsChildOf(UMovieSceneTrack::StaticClass()))
+		{
+			return true;
+		}
+		return false;
 	}
 	
 
@@ -4940,6 +4372,17 @@ public:
 	TArray<UMovieSceneTrack*> NewTracks;
 };
 
+bool FSequencer::CanPaste(const FString& TextToImport) const
+{
+	FTrackObjectTextFactory Factory;
+	if (!Factory.CanCreateObjectsFromText(TextToImport))
+	{
+		return false;
+	}
+
+	return true;
+}
+
 void FSequencer::ImportTracksFromText(const FString& TextToImport, /*out*/ TArray<UMovieSceneTrack*>& ImportedTracks)
 {
 	UPackage* TempPackage = NewObject<UPackage>(nullptr, TEXT("/Engine/Sequencer/Editor/Transient"), RF_Transient);
@@ -5133,13 +4576,16 @@ void FSequencer::ConvertSelectedNodesToSpawnables()
 		{
 			FMovieSceneSpawnable* Spawnable = ConvertToSpawnableInternal(Possessable->GetGuid());
 
-			UpdateRuntimeInstances();
-
-			for (TWeakObjectPtr<> WeakObject : FindBoundObjects(Spawnable->GetGuid(), ActiveTemplateIDs.Top()))
+			if (Spawnable)
 			{
-				if (AActor* SpawnedActor = Cast<AActor>(WeakObject.Get()))
+				UpdateRuntimeInstances();
+
+				for (TWeakObjectPtr<> WeakObject : FindBoundObjects(Spawnable->GetGuid(), ActiveTemplateIDs.Top()))
 				{
-					SpawnedActors.Add(SpawnedActor);
+					if (AActor* SpawnedActor = Cast<AActor>(WeakObject.Get()))
+					{
+						SpawnedActors.Add(SpawnedActor);
+					}
 				}
 			}
 		}
@@ -5356,6 +4802,14 @@ FMovieScenePossessable* FSequencer::ConvertToPossessableInternal(FGuid Spawnable
 		return nullptr;
 	}
 
+	// Delete the spawn track
+	UMovieSceneSpawnTrack* SpawnTrack = Cast<UMovieSceneSpawnTrack>(MovieScene->FindTrack(UMovieSceneSpawnTrack::StaticClass(), SpawnableGuid, NAME_None));
+	if (SpawnTrack)
+	{
+		MovieScene->Modify();
+		MovieScene->RemoveTrack(*SpawnTrack);
+	}
+
 	FTransform SpawnTransform = SpawnableActorTemplate->GetActorTransform();
 
 	UWorld* PlaybackContext = Cast<UWorld>(GetPlaybackContext());
@@ -6659,6 +6113,11 @@ void FSequencer::BindCommands()
 		FCanExecuteAction::CreateLambda(AreSpawnablesSelected)
 	);
 
+	SequencerCommandBindings->MapAction(
+		FSequencerCommands::Get().RestoreAnimatedState,
+		FExecuteAction::CreateSP(this, &FSequencer::RestorePreAnimatedState)
+	);
+
 	SequencerCommandBindings->MapAction(
 		Commands.SetAutoKeyModeAll,
 		FExecuteAction::CreateLambda( [this]{ Settings->SetAutoKeyMode( EAutoKeyMode::KeyAll ); } ),
@@ -6857,6 +6316,12 @@ void FSequencer::BindCommands()
 		FExecuteAction::CreateSP( this, &FSequencer::DeleteSelectedItems ),
 		FCanExecuteAction::CreateLambda(CanDelete));
 
+	SequencerCommandBindings->MapAction(
+		Commands.TogglePlaybackRangeLocked,
+		FExecuteAction::CreateSP( this, &FSequencer::TogglePlaybackRangeLocked ),
+		FCanExecuteAction::CreateLambda( [this] { return GetFocusedMovieSceneSequence() != nullptr;	} ),
+		FIsActionChecked::CreateSP( this, &FSequencer::IsPlaybackRangeLocked ));
+
 	SequencerCommandBindings->MapAction(
 		Commands.ToggleForceFixedFrameIntervalPlayback,
 		FExecuteAction::CreateLambda( [this] 
@@ -7028,40 +6493,18 @@ void FSequencer::BindCommands()
 
 	SequencerCommandBindings->MapAction(
 		Commands.SelectKeysInSelectionRange,
-		FExecuteAction::CreateLambda([this]{ SelectKeysInSelectionRange(); }));
+		FExecuteAction::CreateSP(this, &FSequencer::SelectInSelectionRange, true, false));
 
-	// bind widget specific commands
-	SequencerWidget->BindCommands(SequencerCommandBindings);
-}
-
-void FSequencer::BindLevelEditorCommands()
-{
-	const FSequencerCommands& Commands = FSequencerCommands::Get();
-
-	FLevelEditorModule& LevelEditor = FModuleManager::GetModuleChecked<FLevelEditorModule>(TEXT("LevelEditor"));
-	TSharedRef<FUICommandList> LevelEditorCommandBindings = LevelEditor.GetGlobalLevelEditorActions();
-
-	if (IsLevelEditorSequencer())
-	{
-		LevelEditorCommandBindings->MapAction(
-			Commands.RecordSelectedActors,
-			FExecuteAction::CreateSP(this, &FSequencer::RecordSelectedActors));
-	}
-}
+	SequencerCommandBindings->MapAction(
+		Commands.SelectSectionsInSelectionRange,
+		FExecuteAction::CreateSP(this, &FSequencer::SelectInSelectionRange, false, true));
 
-void FSequencer::UnbindLevelEditorCommands()
-{
-	if (IsLevelEditorSequencer())
-	{
-		if (FModuleManager::Get().IsModuleLoaded(TEXT("LevelEditor")))
-		{
-			// unbind the commands we have bound to the level editor
-			FLevelEditorModule& LevelEditor = FModuleManager::GetModuleChecked<FLevelEditorModule>(TEXT("LevelEditor"));
-			TSharedRef<FUICommandList> LevelEditorCommandBindings = LevelEditor.GetGlobalLevelEditorActions();
+	SequencerCommandBindings->MapAction(
+		Commands.SelectAllInSelectionRange,
+		FExecuteAction::CreateSP(this, &FSequencer::SelectInSelectionRange, true, true));
 
-			LevelEditorCommandBindings->UnmapAction(FSequencerCommands::Get().RecordSelectedActors);
-		}
-	}
+	// bind widget specific commands
+	SequencerWidget->BindCommands(SequencerCommandBindings);
 }
 
 void FSequencer::BuildAddTrackMenu(class FMenuBuilder& MenuBuilder)
diff --git a/Engine/Source/Editor/Sequencer/Private/Sequencer.h b/Engine/Source/Editor/Sequencer/Private/Sequencer.h
--- a/Engine/Source/Editor/Sequencer/Private/Sequencer.h
+++ b/Engine/Source/Editor/Sequencer/Private/Sequencer.h
@@ -35,6 +35,7 @@
 #include "Evaluation/MovieSceneEvaluationTemplateInstance.h"
 #include "LevelEditor.h"
 #include "SequencerTimingManager.h"
+#include "AcquiredResources.h"
 
 class AActor;
 class ACineCameraActor;
@@ -86,10 +87,9 @@ public:
 	 *
 	 * @param InitParams Initialization parameters.
 	 * @param InObjectChangeListener The object change listener to use.
-	 * @param InDetailKeyframeHandler The detail keyframe handler to use.
 	 * @param TrackEditorDelegates Delegates to call to create auto-key handlers for this sequencer.
 	 */
-	void InitSequencer(const FSequencerInitParams& InitParams, const TSharedRef<ISequencerObjectChangeListener>& InObjectChangeListener, const TSharedRef<IDetailKeyframeHandler>& InDetailKeyframeHandler, const TArray<FOnCreateTrackEditor>& TrackEditorDelegates);
+	void InitSequencer(const FSequencerInitParams& InitParams, const TSharedRef<ISequencerObjectChangeListener>& InObjectChangeListener, const TArray<FOnCreateTrackEditor>& TrackEditorDelegates);
 
 	/** @return The current view range */
 	virtual FAnimatedRange GetViewRange() const override;
@@ -135,7 +135,7 @@ public:
 	void ResetSelectionRange();
 
 	/** Select all keys that fall into the current selection range. */
-	void SelectKeysInSelectionRange();
+	void SelectInSelectionRange(bool bSelectKeys, bool bSelectSections);
 
 	/**
 	 * Get the currently viewed sub sequence range
@@ -184,6 +184,8 @@ public:
 
 public:
 
+	bool IsPlaybackRangeLocked() const;
+	void TogglePlaybackRangeLocked();
 	void ResetViewRange();
 	void ZoomViewRange(float InZoomDelta);
 	void ZoomInViewRange();
@@ -237,7 +239,7 @@ public:
 	bool CanSnapToFrame() const;
 
 	/** Transform the selected keys and sections */
-	void TransformSelectedKeysAndSections(float InDeltaTime);
+	void TransformSelectedKeysAndSections(float InDeltaTime, float InScale);
 
 	/**
 	 * @return Movie scene tools used by the sequencer
@@ -299,6 +301,9 @@ protected:
 
 public:
 
+	/** Called when new actors are dropped in the viewport. */
+	void OnNewActorsDropped(const TArray<UObject*>& DroppedObjects, const TArray<AActor*>& DroppedActors);
+
 	/**
 	 * Call when an asset is dropped into the sequencer. Will proprogate this
 	 * to all tracks, and potentially consume this asset
@@ -363,39 +368,7 @@ public:
 	/** Called when an actor is dropped into Sequencer */
 	void OnActorsDropped( const TArray<TWeakObjectPtr<AActor> >& Actors );
 
-	/** Generates and attaches transport control widgets to the main level editor viewports */
-	void AttachTransportControlsToViewports();
-	/** Purges all transport control widgets from the main level editor viewports */
-	void DetachTransportControlsFromViewports();
-
-	/** Gets the color and opacity of a transport control */
-	FLinearColor GetTransportControlsColorAndOpacity(TSharedRef<SWidget> TransportWidget) const;
-
-	/** Gets the visibility of a transport control based on the viewport it's attached to */
-	EVisibility GetTransportControlVisibility(TSharedPtr<ILevelViewport> LevelViewport) const;
-
-	/** Extend the level editor viewport menu */
-	TSharedRef<FExtender> OnExtendLevelEditorViewMenu(const TSharedRef<FUICommandList> CommandList);
-
-	/** Extend the level viewport context menu */
-	void AddLevelViewportMenuExtender();
-	
-	/** Remove a previously registered extender for the level viewport context menu */
-	void RemoveLevelViewportMenuExtender();
-
-	/** Get the extender for the level viewport menu */
-	TSharedRef<FExtender> GetLevelViewportExtender(const TSharedRef<FUICommandList> CommandList, const TArray<AActor*> InActors);
-	
-	/** Add level editor actions */
-	void BindLevelEditorCommands();
-	
-	/** Remove level editor actions */
-	void UnbindLevelEditorCommands();
-
 	void RecordSelectedActors();
-	
-	/** Create a menu entry we can use to toggle the transport controls */
-	void CreateTransportToggleMenuEntry(FMenuBuilder& MenuBuilder);
 
 	/** Functions to push on to the transport controls we use */
 	FReply OnPlay(bool bTogglePlay=true, float InPlayRate=1.f);
@@ -462,6 +435,7 @@ public:
 	void ExportTracksToText(TArray<UMovieSceneTrack*> TrackToExport, /*out*/ FString& ExportedText);
 
 	/** Called when a user executes the paste track menu item */
+	bool CanPaste(const FString& TextToImport) const;
 	void PasteCopiedTracks(TArray<TSharedPtr<FSequencerObjectBindingNode>>& ObjectNodes);
 	void ImportTracksFromText(const FString& TextToImport, /*out*/ TArray<UMovieSceneTrack*>& ImportedTrack);
 
@@ -625,7 +599,6 @@ public:
 	virtual void KeyProperty(FKeyPropertyParams KeyPropertyParams) override;
 	virtual FSequencerSelection& GetSelection() override;
 	virtual FSequencerSelectionPreview& GetSelectionPreview() override;
-	virtual void NotifyMapChanged(UWorld* NewWorld, EMapChangeType MapChangeType) override;
 	virtual FOnGlobalTimeChanged& OnGlobalTimeChanged() override { return OnGlobalTimeChangedDelegate; }
 	virtual FOnMovieSceneDataChanged& OnMovieSceneDataChanged() override { return OnMovieSceneDataChangedDelegate; }
 	virtual FOnSelectionChangedObjectGuids& GetSelectionChangedObjectGuids() override { return OnSelectionChangedObjectGuidsDelegate; }
@@ -640,10 +613,11 @@ public:
 	virtual TSharedRef<INumericTypeInterface<float>> GetZeroPadNumericTypeInterface() override;
 	virtual TSharedRef<SWidget> MakeTransportControls(bool bExtended) override;
 	virtual TSharedRef<SWidget> MakeTimeRange(const TSharedRef<SWidget>& InnerContent, bool bShowWorkingRange, bool bShowViewRange, bool bShowPlaybackRange) override;
-	virtual void SetViewportTransportControlsVisibility(bool bVisible) override;
 	virtual UObject* FindSpawnedObjectOrTemplate(const FGuid& BindingId) override;
 	virtual FGuid MakeNewSpawnable(UObject& SourceObject) override;
-	
+	virtual bool IsReadOnly() const override;
+	virtual void ExternalSelectionHasChanged() override { SynchronizeSequencerSelectionWithExternalSelection(); }
+
 public:
 
 	// IMovieScenePlayer interface
@@ -762,9 +736,6 @@ protected:
 	float FindNearestKey(float NewScrubPosition) override;
 
 protected:
-
-	/** Called via UEditorEngine::GetActorRecordingStateEvent to check to see whether we need to record actor state */
-	void GetActorRecordingState( bool& bIsRecording /* In+Out */ ) const;
 	
 	/** Called when a user executes the delete command to delete sections or keys */
 	void DeleteSelectedItems();
@@ -810,51 +781,13 @@ protected:
 	/** Generates command bindings for UI commands */
 	void BindCommands();
 
-	void ActivateSequencerEditorMode();
-
-	void ActivateDetailKeyframeHandler();
-	void DeactivateDetailKeyframeHandler();
-	void OnPropertyEditorOpened();
-
 	//~ Begin FEditorUndoClient Interface
 	virtual void PostUndo(bool bSuccess) override;
 	virtual void PostRedo(bool bSuccess) override { PostUndo(bSuccess); }
 	// End of FEditorUndoClient
 
-	void OnSectionSelectionChanged();
-
 	void OnSelectedOutlinerNodesChanged();
 
-	/** Called before the world is going to be saved. The sequencer puts everything back to its initial state. */
-	void OnPreSaveWorld(uint32 SaveFlags, UWorld* World);
-
-	/** Called after the world has been saved. The sequencer updates to the animated state. */
-	void OnPostSaveWorld(uint32 SaveFlags, UWorld* World, bool bSuccess);
-
-	/** Called after a level has been added */
-	void OnLevelAdded(ULevel* InLevel, UWorld* InWorld);
-
-	/** Called after a level has been removed */
-	void OnLevelRemoved(ULevel* InLevel, UWorld* InWorld);
-
-	/** Called after a new level has been created. The sequencer editor mode needs to be enabled. */
-	void OnNewCurrentLevel();
-
-	/** Called after a map has been opened. The sequencer editor mode needs to be enabled. */
-	void OnMapOpened(const FString& Filename, bool bLoadAsTemplate);
-
-	/** Called when new actors are dropped in the viewport. */
-	void OnNewActorsDropped(const TArray<UObject*>& DroppedObjects, const TArray<AActor*>& DroppedActors);
-
-	/** Called before a PIE session begins. */
-	void OnPreBeginPIE(bool bIsSimulating);
-
-	/** Called after a PIE session ends. */
-	void OnEndPIE(bool bIsSimulating);
-
-	/** Called after PIE session ends and maps have been cleaned up */
-	void OnEndPlayMap();
-
 	/** Updates a viewport client from camera cut data */
 	void UpdatePreviewLevelViewportClientFromCameraCut(FLevelEditorViewportClient& InViewportClient, UObject* InCameraObject, bool bJumpCut) const;
 
@@ -870,9 +803,6 @@ protected:
 	/** Handles adding a new folder to the outliner tree. */
 	void OnAddFolder();
 
-	/** Handles the actor selection changing externally .*/
-	void OnActorSelectionChanged( UObject* );
-
 	/** Create set playback start transport control */
 	TSharedRef<SWidget> OnCreateTransportSetPlaybackStart();
 
@@ -885,8 +815,9 @@ protected:
 	/** Create set playback end transport control */
 	TSharedRef<SWidget> OnCreateTransportSetPlaybackEnd();
 
-	/** Select all keys in a display node that fall into the current selection range. */
-	void SelectKeysInSelectionRange(const TSharedRef<FSequencerDisplayNode>& DisplayNode, const TRange<float>& SelectionRange);
+	/** Select keys and/or sections in a display node that fall into the current selection range. */
+	void SelectInSelectionRange(const TSharedRef<FSequencerDisplayNode>& DisplayNode, const TRange<float>& SelectionRange, bool bSelectKeys, bool bSelectSections);
+	
 	/** Create record transport control */
 	TSharedRef<SWidget> OnCreateTransportRecord();
 
@@ -904,9 +835,6 @@ private:
 	/** Reset the timing manager to default, or audio clock locked */
 	void ResetTimingManager(bool bUseAudioClock);
 
-	/** Performs any post-tick rendering work needed when moving through scenes */
-	void PostTickRenderStateFixup();
-
 	/** User-supplied settings object for this sequencer */
 	USequencerSettings* Settings;
 
@@ -922,9 +850,6 @@ private:
 	/** Listener for object changes being made while this sequencer is open*/
 	TSharedPtr<ISequencerObjectChangeListener> ObjectChangeListener;
 
-	/** Listener for object changes being made while this sequencer is open*/
-	TSharedPtr<IDetailKeyframeHandler> DetailKeyframeHandler;
-
 	/** Main sequencer widget */
 	TSharedPtr<SSequencer> SequencerWidget;
 	
@@ -988,6 +913,9 @@ private:
 
 	bool bShowCurveEditor;
 
+	/** Whether the sequence should be editable or read only */
+	bool bReadOnly;
+
 	/** Generic Popup Entry */
 	TWeakPtr<IMenu> EntryPopupMenu;
 
@@ -1002,9 +930,6 @@ private:
 	/** Stores the playback status to be restored on refresh. */
 	EMovieScenePlayerStatus::Type StoredPlaybackState;
 
-	/** A flag which indicates Close() was called on this Sequencer. */ 
-	bool bWasClosed;
-
 	FSequencerLabelManager LabelManager;
 	FSequencerSelection Selection;
 	FSequencerSelectionPreview SelectionPreview;
@@ -1021,19 +946,12 @@ private:
 	/** A delegate which is called any time the sequencer selection changes. */
 	FOnSelectionChangedObjectGuids OnSelectionChangedObjectGuidsDelegate;
 
-	/** A map of all the transport controls to viewports that this sequencer has made */
-	TMap< TSharedPtr<class ILevelViewport>, TSharedPtr<class SWidget> > TransportControls;
-
 	FOnActorAddedToSequencer OnActorAddedToSequencerEvent;
 	FOnCameraCut OnCameraCutEvent;
 	FOnPreSave OnPreSaveEvent;
 	FOnActivateSequence OnActivateSequenceEvent;
 
 	int32 SilentModeCount;
-	
-	/** Menu extender for level viewports */
-	FLevelEditorModule::FLevelEditorMenuExtender ViewMenuExtender;
-	FDelegateHandle LevelEditorExtenderDelegateHandle;
 
 	/** When true the sequencer selection is being updated from changes to the external selection. */
 	bool bUpdatingSequencerSelection;
@@ -1058,8 +976,6 @@ private:
 	/** Cached array of view targets that were set before we possessed the player controller with a camera from sequencer */
 	TArray<FCachedViewTarget> PrePossessionViewTargets;
 
-	FDelegateHandle LevelViewportExtenderHandle;
-
 	/** Attribute used to retrieve the playback context for this frame */
 	TAttribute<UObject*> PlaybackContextAttribute;
 
@@ -1074,8 +990,12 @@ private:
 
 	bool bNeedsEvaluate;
 
+	FAcquiredResources AcquiredResources;
+
 	/** The range of the currently displayed sub sequence in relation to its parent section */
 	TRange<float> SubSequenceRange;
 
 	TSharedPtr<struct FSequencerTemplateStore> TemplateStore;
+
+	TMap<FName, TFunction<void()>> CleanupFunctions;
 };
diff --git a/Engine/Source/Editor/Sequencer/Private/SequencerCommands.cpp b/Engine/Source/Editor/Sequencer/Private/SequencerCommands.cpp
--- a/Engine/Source/Editor/Sequencer/Private/SequencerCommands.cpp
+++ b/Engine/Source/Editor/Sequencer/Private/SequencerCommands.cpp
@@ -24,6 +24,7 @@ void FSequencerCommands::RegisterCommands()
 	UI_COMMAND( ZoomInViewRange, "Zoom into the View Range", "Zoom into the view range", EUserInterfaceActionType::Button, FInputChord(EKeys::Equals) );
 	UI_COMMAND( ZoomOutViewRange, "Zoom out of the View Range", "Zoom out of the view range", EUserInterfaceActionType::Button, FInputChord(EKeys::Hyphen) );
 
+	UI_COMMAND( TogglePlaybackRangeLocked, "Locked", "Lock/unlock the playback range.", EUserInterfaceActionType::ToggleButton, FInputChord() );
 	UI_COMMAND( ToggleForceFixedFrameIntervalPlayback, "Force Fixed Frame Interval Playback", "Forces scene evaluation to a fixed frame interval in editor and at runtime, even if this would result in duplicated or dropped frames.", EUserInterfaceActionType::ToggleButton, FInputChord() );
 	
 	UI_COMMAND( ToggleKeepCursorInPlaybackRange, "Keep Cursor in Playback Range", "When checked, the cursor will be constrained to the current playback range during playback", EUserInterfaceActionType::ToggleButton, FInputChord() );
@@ -37,7 +38,9 @@ void FSequencerCommands::RegisterCommands()
 	UI_COMMAND( SetSelectionRangeEnd, "Set Selection End", "Sets the end of the selection range", EUserInterfaceActionType::Button, FInputChord(EKeys::O) );
 	UI_COMMAND( SetSelectionRangeStart, "Set Selection Start", "Sets the start of the selection range", EUserInterfaceActionType::Button, FInputChord(EKeys::I) );
 	UI_COMMAND( ResetSelectionRange, "Reset Selection Range", "Reset the selection range", EUserInterfaceActionType::Button, FInputChord() );
-	UI_COMMAND( SelectKeysInSelectionRange, "Select in Selection Range", "Select all keys that fall into the selection range", EUserInterfaceActionType::Button, FInputChord() );
+	UI_COMMAND( SelectKeysInSelectionRange, "Select Keys in Selection Range", "Select all keys that fall into the selection range", EUserInterfaceActionType::Button, FInputChord() );
+	UI_COMMAND( SelectSectionsInSelectionRange, "Select Sections in Selection Range", "Select all sections that fall into the selection range", EUserInterfaceActionType::Button, FInputChord() );
+	UI_COMMAND( SelectAllInSelectionRange, "Select All in Selection Range", "Select all keys and section that fall into the selection range", EUserInterfaceActionType::Button, FInputChord() );
 
 	UI_COMMAND( SetKey, "Set Key", "Sets a key on the selected tracks", EUserInterfaceActionType::Button, FInputChord(EKeys::Enter) );
 	UI_COMMAND( SetInterpolationCubicAuto, "Set Key Auto", "Cubic interpolation - Automatic tangents", EUserInterfaceActionType::Button, FInputChord(EKeys::One));
@@ -93,6 +96,7 @@ void FSequencerCommands::RegisterCommands()
 	UI_COMMAND( ConvertToPossessable, "Convert to Possessable", "Make the specified spawned objects possessed by sequencer.", EUserInterfaceActionType::Button, FInputChord() );
 
 	UI_COMMAND( SaveCurrentSpawnableState, "Save Default State", "Save the current state of this spawnable object as its default properties.", EUserInterfaceActionType::Button, FInputChord() );
+	UI_COMMAND( RestoreAnimatedState, "Restore Animated State", "Restore any objects that have been animated by sequencer back to their original state.", EUserInterfaceActionType::Button, FInputChord(EModifierKey::Control, EKeys::R) );
 
 	UI_COMMAND( DiscardChanges, "Discard All Changes", "Revert the currently edited movie scene to its last saved state.", EUserInterfaceActionType::Button, FInputChord());
 
diff --git a/Engine/Source/Editor/Sequencer/Private/SequencerContextMenus.cpp b/Engine/Source/Editor/Sequencer/Private/SequencerContextMenus.cpp
--- a/Engine/Source/Editor/Sequencer/Private/SequencerContextMenus.cpp
+++ b/Engine/Source/Editor/Sequencer/Private/SequencerContextMenus.cpp
@@ -228,19 +228,23 @@ void FKeyContextMenu::AddPropertiesMenu(FMenuBuilder& MenuBuilder)
 		// @todo sequencer: only one struct per structure view supported right now :/
 		if (Keys.Num() == 1)
 		{
-			TPair<TSharedPtr<FStructOnScope>, const FSequencerSelectedKey&>& Key = Keys[0];
+			TSharedPtr<FStructOnScope> SharedKeyStruct = Keys[0].Key;
+			TWeakObjectPtr<UMovieSceneSection> WeakSection = Keys[0].Value.Section;
 
 			// register details customizations for this instance
-			StructureDetailsView->GetDetailsView().RegisterInstancedCustomPropertyLayout(FIntegralKey::StaticStruct(), FOnGetDetailCustomizationInstance::CreateStatic(&FIntegralKeyDetailsCustomization::MakeInstance, TWeakObjectPtr<const UMovieSceneSection>(Key.Value.Section)));
+			StructureDetailsView->GetDetailsView().RegisterInstancedCustomPropertyLayout(FIntegralKey::StaticStruct(), FOnGetDetailCustomizationInstance::CreateStatic(&FIntegralKeyDetailsCustomization::MakeInstance, TWeakObjectPtr<const UMovieSceneSection>(WeakSection)));
 
-			StructureDetailsView->SetStructureData(Key.Key);
+			StructureDetailsView->SetStructureData(Keys[0].Key);
 			StructureDetailsView->GetOnFinishedChangingPropertiesDelegate().AddLambda(
 				[=](const FPropertyChangedEvent& ChangeEvent) {
-					
-					Key.Value.Section->Modify();
-					if (Key.Key->GetStruct()->IsChildOf(FMovieSceneKeyStruct::StaticStruct()))
+
+					if (UMovieSceneSection* Section = WeakSection.Get())
+					{
+						Section->Modify();
+					}
+					if (SharedKeyStruct->GetStruct()->IsChildOf(FMovieSceneKeyStruct::StaticStruct()))
 					{
-						((FMovieSceneKeyStruct*)Key.Key->GetStructMemory())->PropagateChanges(ChangeEvent);
+						((FMovieSceneKeyStruct*)SharedKeyStruct->GetStructMemory())->PropagateChanges(ChangeEvent);
 					}
 					Sequencer->NotifyMovieSceneDataChanged( EMovieSceneDataChangeType::TrackValueChanged );
 				}
@@ -433,6 +437,15 @@ void FSectionContextMenu::AddEditMenu(FMenuBuilder& MenuBuilder)
 			FExecuteAction::CreateLambda([=]{ Shared->SplitSection(); }),
 			FCanExecuteAction::CreateLambda([=]{ return Shared->IsTrimmable(); }))
 	);
+		
+	MenuBuilder.AddMenuEntry(
+		LOCTEXT("ReduceKeysSection", "Reduce Keys"),
+		LOCTEXT("ReduceKeysTooltip", "Reduce keys in this section"),
+		FSlateIcon(),
+		FUIAction(
+			FExecuteAction::CreateLambda([=]{ Shared->ReduceKeys(); }),
+			FCanExecuteAction::CreateLambda([=]{ return Shared->CanReduceKeys(); }))
+	);
 }
 
 
@@ -705,6 +718,40 @@ void FSectionContextMenu::SplitSection()
 }
 
 
+void FSectionContextMenu::ReduceKeys()
+{
+	FScopedTransaction ReduceKeysTransaction(LOCTEXT("ReduceKeys_Transaction", "Reduce Keys"));
+
+	TSet<TSharedPtr<IKeyArea> > KeyAreas;
+	for (auto DisplayNode : Sequencer->GetSelection().GetSelectedOutlinerNodes())
+	{
+		SequencerHelpers::GetAllKeyAreas(DisplayNode, KeyAreas);
+	}
+
+	if (KeyAreas.Num() == 0)
+	{
+		TSet<TSharedRef<FSequencerDisplayNode>> SelectedNodes = Sequencer->GetSelection().GetNodesWithSelectedKeysOrSections();
+		for (auto DisplayNode : SelectedNodes)
+		{
+			SequencerHelpers::GetAllKeyAreas(DisplayNode, KeyAreas);
+		}
+	}
+
+	for (auto KeyArea : KeyAreas)
+	{
+		if (KeyArea.IsValid())
+		{
+			if (KeyArea->GetRichCurve() && KeyArea->GetOwningSection())
+			{
+				KeyArea->GetOwningSection()->Modify();
+				KeyArea->GetRichCurve()->RemoveRedundantKeys(KINDA_SMALL_NUMBER);
+			}
+		}
+	}
+
+	Sequencer->NotifyMovieSceneDataChanged( EMovieSceneDataChangeType::TrackValueChanged );
+}
+
 bool FSectionContextMenu::IsTrimmable() const
 {
 	for (auto Section : Sequencer->GetSelection().GetSelectedSections())
@@ -717,6 +764,26 @@ bool FSectionContextMenu::IsTrimmable() const
 	return false;
 }
 
+bool FSectionContextMenu::CanReduceKeys() const
+{
+	TSet<TSharedPtr<IKeyArea> > KeyAreas;
+	for (auto DisplayNode : Sequencer->GetSelection().GetSelectedOutlinerNodes())
+	{
+		SequencerHelpers::GetAllKeyAreas(DisplayNode, KeyAreas);
+	}
+
+	if (KeyAreas.Num() == 0)
+	{
+		TSet<TSharedRef<FSequencerDisplayNode>> SelectedNodes = Sequencer->GetSelection().GetNodesWithSelectedKeysOrSections();
+		for (auto DisplayNode : SelectedNodes)
+		{
+			SequencerHelpers::GetAllKeyAreas(DisplayNode, KeyAreas);
+		}
+	}
+
+	return KeyAreas.Num() != 0;
+}
+
 
 void FSectionContextMenu::SetExtrapolationMode(ERichCurveExtrapolation ExtrapMode, bool bPreInfinity)
 {
diff --git a/Engine/Source/Editor/Sequencer/Private/SequencerContextMenus.h b/Engine/Source/Editor/Sequencer/Private/SequencerContextMenus.h
--- a/Engine/Source/Editor/Sequencer/Private/SequencerContextMenus.h
+++ b/Engine/Source/Editor/Sequencer/Private/SequencerContextMenus.h
@@ -61,8 +61,12 @@ private:
 
 	void SplitSection();
 
+	void ReduceKeys();
+
 	bool IsTrimmable() const;
 
+	bool CanReduceKeys() const;
+
 	void SetExtrapolationMode(ERichCurveExtrapolation ExtrapMode, bool bPreInfinity);
 
 	bool IsExtrapolationModeSelected(ERichCurveExtrapolation ExtrapMode, bool bPreInfinity) const;
diff --git a/Engine/Source/Editor/Sequencer/Private/SequencerEdMode.cpp b/Engine/Source/Editor/Sequencer/Private/SequencerEdMode.cpp
--- a/Engine/Source/Editor/Sequencer/Private/SequencerEdMode.cpp
+++ b/Engine/Source/Editor/Sequencer/Private/SequencerEdMode.cpp
@@ -36,7 +36,7 @@ void FSequencerEdMode::Enter()
 
 void FSequencerEdMode::Exit()
 {
-	SequencerPtr.Reset();
+	Sequencers.Reset();
 
 	FEdMode::Exit();
 }
@@ -49,18 +49,22 @@ bool FSequencerEdMode::IsCompatibleWith(FEditorModeID OtherModeID) const
 
 bool FSequencerEdMode::InputKey( FEditorViewportClient* ViewportClient, FViewport* Viewport, FKey Key, EInputEvent Event )
 {
-	TSharedPtr<FSequencer> Sequencer = SequencerPtr.Pin();
+	TSharedPtr<FSequencer> ActiveSequencer;
 
-	if (!Sequencer.IsValid())
+	for (TWeakPtr<FSequencer> WeakSequencer : Sequencers)
 	{
-		return false;
+		ActiveSequencer = WeakSequencer.Pin();
+		if (ActiveSequencer.IsValid())
+		{
+			break;
+		}
 	}
 
-	if (Event != IE_Released)
+	if (ActiveSequencer.IsValid() && Event != IE_Released)
 	{
 		FModifierKeysState KeyState = FSlateApplication::Get().GetModifierKeys();
 
-		if (Sequencer->GetCommandBindings(ESequencerCommandBindings::Shared).Get()->ProcessCommandBindings(Key, KeyState, (Event == IE_Repeat) ))
+		if (ActiveSequencer->GetCommandBindings(ESequencerCommandBindings::Shared).Get()->ProcessCommandBindings(Key, KeyState, (Event == IE_Repeat) ))
 		{
 			return true;
 		}
@@ -71,13 +75,6 @@ bool FSequencerEdMode::InputKey( FEditorViewportClient* ViewportClient, FViewpor
 
 void FSequencerEdMode::Render(const FSceneView* View, FViewport* Viewport, FPrimitiveDrawInterface* PDI)
 {
-	TSharedPtr<FSequencer> Sequencer = SequencerPtr.Pin();
-
-	if (!Sequencer.IsValid())
-	{
-		return;
-	}
-
 	FEdMode::Render(View, Viewport, PDI);
 
 #if WITH_EDITORONLY_DATA
@@ -108,19 +105,19 @@ void FSequencerEdMode::DrawHUD(FEditorViewportClient* ViewportClient,FViewport*
 
 void FSequencerEdMode::OnKeySelected(FViewport* Viewport, HMovieSceneKeyProxy* KeyProxy)
 {
-	TSharedPtr<FSequencer> Sequencer = SequencerPtr.Pin();
-
-	if (!Sequencer.IsValid())
-	{
-		return;
-	}
-
 	bool bCtrlDown = Viewport->KeyState(EKeys::LeftControl) || Viewport->KeyState(EKeys::RightControl);
 	bool bAltDown = Viewport->KeyState(EKeys::LeftAlt) || Viewport->KeyState(EKeys::RightAlt);
 	bool bShiftDown = Viewport->KeyState(EKeys::LeftShift) || Viewport->KeyState(EKeys::RightShift);
 
-	Sequencer->SetLocalTimeDirectly(KeyProxy->Time);
-	Sequencer->SelectTrackKeys(KeyProxy->MovieSceneSection, KeyProxy->Time, bShiftDown, bCtrlDown);
+	for (TWeakPtr<FSequencer> WeakSequencer : Sequencers)
+	{
+		TSharedPtr<FSequencer> Sequencer = WeakSequencer.Pin();
+		if (Sequencer.IsValid())
+		{
+			Sequencer->SetLocalTimeDirectly(KeyProxy->Time);
+			Sequencer->SelectTrackKeys(KeyProxy->MovieSceneSection, KeyProxy->Time, bShiftDown, bCtrlDown);
+		}
+	}
 }
 
 namespace SequencerEdMode_Draw3D
@@ -383,22 +380,26 @@ void DrawTransformTrack(const FSceneView* View, FPrimitiveDrawInterface* PDI, UM
 
 void FSequencerEdMode::DrawTracks3D(const FSceneView* View, FPrimitiveDrawInterface* PDI)
 {
-	TSet<TSharedRef<FSequencerDisplayNode> > ObjectBindingNodes;
-
-	// Map between object binding nodes and selection
-	TMap<TSharedRef<FSequencerDisplayNode>, bool > ObjectBindingNodesSelectionMap;
+	for (TWeakPtr<FSequencer> WeakSequencer : Sequencers)
+	{
+		TSharedPtr<FSequencer> Sequencer = WeakSequencer.Pin();
+		if (!Sequencer.IsValid())
+		{
+			continue;
+		}
 
-	TSharedPtr<FSequencer> Sequencer = SequencerPtr.Pin();
+		TSet<TSharedRef<FSequencerDisplayNode> > ObjectBindingNodes;
 
-	if (!Sequencer.IsValid())
-	{
-		return;
-	}
+		// Map between object binding nodes and selection
+		TMap<TSharedRef<FSequencerDisplayNode>, bool > ObjectBindingNodesSelectionMap;
 
-	for (auto ObjectBinding : Sequencer->GetNodeTree()->GetObjectBindingMap() )
-	{
-		if (ObjectBinding.Value.IsValid())
+		for (auto ObjectBinding : Sequencer->GetNodeTree()->GetObjectBindingMap() )
 		{
+			if (!ObjectBinding.Value.IsValid())
+			{
+				continue;
+			}
+
 			TSharedRef<FSequencerObjectBindingNode> ObjectBindingNode = ObjectBinding.Value.ToSharedRef();
 
 			TSet<TSharedRef<FSequencerDisplayNode> > DescendantNodes;
@@ -418,34 +419,33 @@ void FSequencerEdMode::DrawTracks3D(const FSceneView* View, FPrimitiveDrawInterf
 
 			ObjectBindingNodesSelectionMap.Add(ObjectBindingNode, bSelected);
 		}
-	}
-
 
-	// Gather up the transform track nodes from the object binding nodes
+		// Gather up the transform track nodes from the object binding nodes
 
-	for (auto ObjectBindingNode : ObjectBindingNodesSelectionMap)
-	{
-		TSet<TSharedRef<FSequencerDisplayNode> > AllNodes;
-		SequencerHelpers::GetDescendantNodes(ObjectBindingNode.Key, AllNodes);
+		for (auto ObjectBindingNode : ObjectBindingNodesSelectionMap)
+		{
+			TSet<TSharedRef<FSequencerDisplayNode> > AllNodes;
+			SequencerHelpers::GetDescendantNodes(ObjectBindingNode.Key, AllNodes);
 
-		FGuid ObjectBinding = StaticCastSharedRef<FSequencerObjectBindingNode>(ObjectBindingNode.Key)->GetObjectBinding();
+			FGuid ObjectBinding = StaticCastSharedRef<FSequencerObjectBindingNode>(ObjectBindingNode.Key)->GetObjectBinding();
 
-		TArray<TWeakObjectPtr<UObject>> BoundObjects;
-		for (TWeakObjectPtr<> Ptr : Sequencer->FindObjectsInCurrentSequence(ObjectBinding))
-		{
-			BoundObjects.Add(Ptr);
-		}
+			TArray<TWeakObjectPtr<UObject>> BoundObjects;
+			for (TWeakObjectPtr<> Ptr : Sequencer->FindObjectsInCurrentSequence(ObjectBinding))
+			{
+				BoundObjects.Add(Ptr);
+			}
 
-		for (auto DisplayNode : AllNodes)
-		{
-			if (DisplayNode->GetType() == ESequencerNode::Track)
+			for (auto DisplayNode : AllNodes)
 			{
-				TSharedRef<FSequencerTrackNode> TrackNode = StaticCastSharedRef<FSequencerTrackNode>(DisplayNode);
-				UMovieSceneTrack* TrackNodeTrack = TrackNode->GetTrack();
-				UMovieScene3DTransformTrack* TransformTrack = Cast<UMovieScene3DTransformTrack>(TrackNodeTrack);
-				if (TransformTrack != nullptr)
+				if (DisplayNode->GetType() == ESequencerNode::Track)
 				{
-					DrawTransformTrack(View, PDI, TransformTrack, BoundObjects, ObjectBindingNode.Value);
+					TSharedRef<FSequencerTrackNode> TrackNode = StaticCastSharedRef<FSequencerTrackNode>(DisplayNode);
+					UMovieSceneTrack* TrackNodeTrack = TrackNode->GetTrack();
+					UMovieScene3DTransformTrack* TransformTrack = Cast<UMovieScene3DTransformTrack>(TrackNodeTrack);
+					if (TransformTrack != nullptr)
+					{
+						DrawTransformTrack(View, PDI, TransformTrack, BoundObjects, ObjectBindingNode.Value);
+					}
 				}
 			}
 		}
diff --git a/Engine/Source/Editor/Sequencer/Private/SequencerEdMode.h b/Engine/Source/Editor/Sequencer/Private/SequencerEdMode.h
--- a/Engine/Source/Editor/Sequencer/Private/SequencerEdMode.h
+++ b/Engine/Source/Editor/Sequencer/Private/SequencerEdMode.h
@@ -35,7 +35,10 @@ public:
 	virtual void Render(const FSceneView* View,FViewport* Viewport,FPrimitiveDrawInterface* PDI) override;
 	virtual void DrawHUD(FEditorViewportClient* ViewportClient,FViewport* Viewport,const FSceneView* View,FCanvas* Canvas) override;
 
-	void SetSequencer(const TWeakPtr<FSequencer>& InSequencer) { SequencerPtr = InSequencer; }
+	void AddSequencer(TWeakPtr<FSequencer> InSequencer) { Sequencers.AddUnique(InSequencer); }
+	void RemoveSequencer(TWeakPtr<FSequencer> InSequencer) { Sequencers.Remove(InSequencer); }
+
+	void OnSequencerReceivedFocus(TWeakPtr<FSequencer> InSequencer) { Sequencers.Sort([=](TWeakPtr<FSequencer> A, TWeakPtr<FSequencer> B){ return A == InSequencer; }); }
 
 	void OnKeySelected(FViewport* Viewport, HMovieSceneKeyProxy* KeyProxy);
 
@@ -43,7 +46,7 @@ protected:
 	void DrawTracks3D(const FSceneView* View, FPrimitiveDrawInterface* PDI);
 
 private:
-	TWeakPtr<FSequencer> SequencerPtr;
+	TArray<TWeakPtr<FSequencer>> Sequencers;
 };
 
 /**
diff --git a/Engine/Source/Editor/Sequencer/Private/SequencerModule.cpp b/Engine/Source/Editor/Sequencer/Private/SequencerModule.cpp
--- a/Engine/Source/Editor/Sequencer/Private/SequencerModule.cpp
+++ b/Engine/Source/Editor/Sequencer/Private/SequencerModule.cpp
@@ -13,7 +13,6 @@
 #include "SequencerEdMode.h"
 #include "SequencerObjectChangeListener.h"
 #include "IDetailKeyframeHandler.h"
-#include "SequencerDetailKeyframeHandler.h"
 
 
 // We disable the deprecation warnings here because otherwise it'll complain about us
@@ -40,9 +39,8 @@ public:
 	{
 		TSharedRef<FSequencer> Sequencer = MakeShareable(new FSequencer);
 		TSharedRef<ISequencerObjectChangeListener> ObjectChangeListener = MakeShareable(new FSequencerObjectChangeListener(Sequencer, InitParams.bEditWithinLevelEditor));
-		TSharedRef<IDetailKeyframeHandler> KeyframeHandler = MakeShareable( new FSequencerDetailKeyframeHandler( Sequencer ));
 
-		Sequencer->InitSequencer(InitParams, ObjectChangeListener, KeyframeHandler, TrackEditorDelegates);
+		Sequencer->InitSequencer(InitParams, ObjectChangeListener, TrackEditorDelegates);
 
 		return Sequencer;
 	}
diff --git a/Engine/Source/Editor/Sequencer/Private/SequencerTimeSliderController.cpp b/Engine/Source/Editor/Sequencer/Private/SequencerTimeSliderController.cpp
--- a/Engine/Source/Editor/Sequencer/Private/SequencerTimeSliderController.cpp
+++ b/Engine/Source/Editor/Sequencer/Private/SequencerTimeSliderController.cpp
@@ -791,6 +791,7 @@ FReply FSequencerTimeSliderController::OnMouseMove( SWidget& WidgetOwner, const
 				TRange<float> SelectionRange = TimeSliderArgs.SelectionRange.Get();
 				TRange<float> PlaybackRange = TimeSliderArgs.PlaybackRange.Get();
 				float LocalMouseDownPos = RangeToScreen.InputToLocalX(MouseDownRange[0]);
+				const bool bIsPlaybackRangeLocked = TimeSliderArgs.IsPlaybackRangeLocked.Get();
 
 				// Disable selection range test if it's empty so that the playback range scrubbing gets priority
 				if (!SelectionRange.IsEmpty() && HitTestScrubberEnd(RangeToScreen, SelectionRange, LocalMouseDownPos, ScrubPosition))
@@ -805,13 +806,13 @@ FReply FSequencerTimeSliderController::OnMouseMove( SWidget& WidgetOwner, const
 					MouseDragType = DRAG_SELECTION_START;
 					TimeSliderArgs.OnSelectionRangeBeginDrag.ExecuteIfBound();
 				}
-				else if (HitTestScrubberEnd(RangeToScreen, PlaybackRange, LocalMouseDownPos, ScrubPosition))
+				else if (!bIsPlaybackRangeLocked && HitTestScrubberEnd(RangeToScreen, PlaybackRange, LocalMouseDownPos, ScrubPosition))
 				{
 					// playback range end scrubber
 					MouseDragType = DRAG_PLAYBACK_END;
 					TimeSliderArgs.OnPlaybackRangeBeginDrag.ExecuteIfBound();
 				}
-				else if (HitTestScrubberStart(RangeToScreen, PlaybackRange, LocalMouseDownPos, ScrubPosition))
+				else if (!bIsPlaybackRangeLocked && HitTestScrubberStart(RangeToScreen, PlaybackRange, LocalMouseDownPos, ScrubPosition))
 				{
 					// playback range start scrubber
 					MouseDragType = DRAG_PLAYBACK_START;
@@ -947,6 +948,7 @@ FCursorReply FSequencerTimeSliderController::OnCursorQuery( TSharedRef<const SWi
 	TRange<float> PlaybackRange = TimeSliderArgs.PlaybackRange.Get();
 	TRange<float> SelectionRange = TimeSliderArgs.SelectionRange.Get();
 	const float ScrubPosition = TimeSliderArgs.ScrubPosition.Get();
+	const bool bIsPlaybackRangeLocked = TimeSliderArgs.IsPlaybackRangeLocked.Get();
 
 	// Use L/R resize cursor if we're dragging or hovering a playback range bound
 	float HitTestPosition = MyGeometry.AbsoluteToLocal(CursorEvent.GetScreenSpacePosition()).X;
@@ -954,8 +956,8 @@ FCursorReply FSequencerTimeSliderController::OnCursorQuery( TSharedRef<const SWi
 		(MouseDragType == DRAG_PLAYBACK_START) ||
 		(MouseDragType == DRAG_SELECTION_START) ||
 		(MouseDragType == DRAG_SELECTION_END) ||
-		HitTestScrubberStart(RangeToScreen, PlaybackRange, HitTestPosition, ScrubPosition) ||
-		HitTestScrubberEnd(RangeToScreen, PlaybackRange, HitTestPosition, ScrubPosition) ||
+		(!bIsPlaybackRangeLocked && HitTestScrubberStart(RangeToScreen, PlaybackRange, HitTestPosition, ScrubPosition)) ||
+		(!bIsPlaybackRangeLocked && HitTestScrubberEnd(RangeToScreen, PlaybackRange, HitTestPosition, ScrubPosition)) ||
 		(!SelectionRange.IsEmpty() && HitTestScrubberStart(RangeToScreen, SelectionRange, HitTestPosition, ScrubPosition)) ||
 		(!SelectionRange.IsEmpty() && HitTestScrubberEnd(RangeToScreen, SelectionRange, HitTestPosition, ScrubPosition)))
 	{
@@ -1060,7 +1062,7 @@ TSharedRef<SWidget> FSequencerTimeSliderController::OpenSetPlaybackRangeMenu(flo
 			FSlateIcon(),
 			FUIAction(
 				FExecuteAction::CreateLambda([=]{ SetPlaybackRangeStart(MouseTime); }),
-				FCanExecuteAction::CreateLambda([=]{ return MouseTime <= PlaybackRange.GetUpperBoundValue(); })
+				FCanExecuteAction::CreateLambda([=]{ return !TimeSliderArgs.IsPlaybackRangeLocked.Get() && MouseTime <= PlaybackRange.GetUpperBoundValue(); })
 			)
 		);
 
@@ -1070,9 +1072,22 @@ TSharedRef<SWidget> FSequencerTimeSliderController::OpenSetPlaybackRangeMenu(flo
 			FSlateIcon(),
 			FUIAction(
 				FExecuteAction::CreateLambda([=]{ SetPlaybackRangeEnd(MouseTime); }),
-				FCanExecuteAction::CreateLambda([=]{ return MouseTime >= PlaybackRange.GetLowerBoundValue(); })
+				FCanExecuteAction::CreateLambda([=]{ return !TimeSliderArgs.IsPlaybackRangeLocked.Get() && MouseTime >= PlaybackRange.GetLowerBoundValue(); })
 			)
 		);
+
+		MenuBuilder.AddMenuEntry(
+			LOCTEXT("ToggleLocked", "Locked"),
+			LOCTEXT("ToggleLockedTooltip", "Lock/Unlock the playback range"),
+			FSlateIcon(),
+			FUIAction(
+				FExecuteAction::CreateLambda([=] { TimeSliderArgs.OnTogglePlaybackRangeLocked.ExecuteIfBound(); }),
+				FCanExecuteAction(),
+				FIsActionChecked::CreateLambda([=] { return TimeSliderArgs.IsPlaybackRangeLocked.Get(); })
+			),
+			NAME_None,
+			EUserInterfaceActionType::ToggleButton
+		);
 	}
 	MenuBuilder.EndSection(); // SequencerPlaybackRangeMenu
 
diff --git a/Engine/Source/Editor/Sequencer/Public/IKeyArea.h b/Engine/Source/Editor/Sequencer/Public/IKeyArea.h
--- a/Engine/Source/Editor/Sequencer/Public/IKeyArea.h
+++ b/Engine/Source/Editor/Sequencer/Public/IKeyArea.h
@@ -64,15 +64,23 @@ public:
 	 */
 	virtual float GetKeyTime(FKeyHandle KeyHandle) const = 0;
 
+	/**
+	 * Dilates the key by a specific factor
+	 *
+	 * @param Scale The multiplier which scales this key
+	 * @param Origin Time to scale from
+	 * @return The handle of the key
+	 */
+	virtual FKeyHandle DilateKey(FKeyHandle KeyHandle, float Scale, float Origin) = 0;
+
 	/**
 	 * Moves a key
 	 * 
 	 * @param KeyHandle		Handle of the key to move
 	 * @param DeltaPosition	The delta position of the key
-	 * @return The handle of the key
 	 */
 	virtual FKeyHandle MoveKey(FKeyHandle KeyHandle, float DeltaPosition) = 0;
-	
+		
 	/**
 	 * Deletes a key
 	 * 
diff --git a/Engine/Source/Editor/Sequencer/Public/ISequencer.h b/Engine/Source/Editor/Sequencer/Public/ISequencer.h
--- a/Engine/Source/Editor/Sequencer/Public/ISequencer.h
+++ b/Engine/Source/Editor/Sequencer/Public/ISequencer.h
@@ -5,6 +5,7 @@
 
 #include "CoreMinimal.h"
 #include "UObject/ObjectMacros.h"
+#include "IMovieScenePlayer.h"
 #include "Misc/Guid.h"
 #include "Widgets/SWidget.h"
 #include "Containers/ArrayView.h"
@@ -343,8 +344,6 @@ public:
 	virtual FSequencerSelection& GetSelection() = 0;
 	virtual FSequencerSelectionPreview& GetSelectionPreview() = 0;
 
-	virtual void NotifyMapChanged(class UWorld* NewWorld, EMapChangeType MapChangeType) = 0;
-
 	/** Gets a multicast delegate which is executed whenever the global time changes. */
 	virtual FOnGlobalTimeChanged& OnGlobalTimeChanged() = 0;
 
@@ -366,12 +365,15 @@ public:
 	/** @return Returns a widget containing the sequencer's playback controls */
 	virtual TSharedRef<SWidget> MakeTransportControls(bool bExtended) = 0;
 
-	/** Turn viewport transport controls on or off */
-	virtual void SetViewportTransportControlsVisibility(bool bVisible) = 0;
-
 	/** Attempt to find a spawned object in the currently focused movie scene, or the template object for the specified binding ID, if possible */
 	virtual UObject* FindSpawnedObjectOrTemplate(const FGuid& BindingId) = 0;
 
+	/** Called when the external selection has changed in such a way that sequencer should re-synchronize its selection states */
+	virtual void ExternalSelectionHasChanged() = 0;
+
+	/** Whether the sequence is read-only */
+	virtual bool IsReadOnly() const = 0;
+
 	/**
 	 * Create a widget containing the spinboxes for setting the working and playback range
 	 * 
diff --git a/Engine/Source/Editor/Sequencer/Public/ISequencerModule.h b/Engine/Source/Editor/Sequencer/Public/ISequencerModule.h
--- a/Engine/Source/Editor/Sequencer/Public/ISequencerModule.h
+++ b/Engine/Source/Editor/Sequencer/Public/ISequencerModule.h
@@ -37,15 +37,22 @@ struct FSequencerViewParams
 
 	FOnGetAddMenuContent OnGetAddMenuContent;
 
+	/** Called when this sequencer has received user focus */
+	FSimpleDelegate OnReceivedFocus;
+
 	/** A menu extender for the add menu */
 	TSharedPtr<FExtender> AddMenuExtender;
 
 	/** Unique name for the sequencer. */
 	FString UniqueName;
 
+	/** Whether the sequencer is read-only */
+	bool bReadOnly;
+
 	FSequencerViewParams(FString InName = FString())
 		: InitialScrubPosition(0.0f)
 		, UniqueName(MoveTemp(InName))
+		, bReadOnly(false)
 	{ }
 };
 
diff --git a/Engine/Source/Editor/Sequencer/Public/LevelEditorSequencerIntegration.h b/Engine/Source/Editor/Sequencer/Public/LevelEditorSequencerIntegration.h
new file mode 100644
--- /dev/null
+++ b/Engine/Source/Editor/Sequencer/Public/LevelEditorSequencerIntegration.h
@@ -0,0 +1,141 @@
+// Copyright 1998-2016 Epic Games, Inc. All Rights Reserved.
+
+#pragma once
+
+#include "AcquiredResources.h"
+#include "CoreMinimal.h"
+#include "Layout/Visibility.h"
+
+
+class AActor;
+class FExtender;
+class FMenuBuilder;
+class FSequencer;
+class FUICommandList;
+class ILevelViewport;
+class ISequencer;
+class SViewportTransportControls;
+class ULevel;
+
+
+struct FLevelEditorSequencerIntegrationOptions
+{
+	FLevelEditorSequencerIntegrationOptions()
+		: bRequiresLevelEvents(true)
+		, bRequiresActorEvents(false)
+		, bCanRecord(false)
+	{}
+
+	bool bRequiresLevelEvents : 1;
+	bool bRequiresActorEvents : 1;
+	bool bCanRecord : 1;
+};
+
+
+class SEQUENCER_API FLevelEditorSequencerIntegration
+{
+public:
+
+	static FLevelEditorSequencerIntegration& Get();
+
+	void Initialize();
+
+	void AddSequencer(TSharedRef<ISequencer> InSequencer, const FLevelEditorSequencerIntegrationOptions& Options);
+
+	void OnSequencerReceivedFocus(TSharedRef<ISequencer> InSequencer);
+
+	void RemoveSequencer(TSharedRef<ISequencer> InSequencer);
+
+	void SetViewportTransportControlsVisibility(bool bVisible);
+
+	bool GetViewportTransportControlsVisibility() const;
+
+private:
+
+	/** Called before the world is going to be saved. The sequencer puts everything back to its initial state. */
+	void OnPreSaveWorld(uint32 SaveFlags, UWorld* World);
+
+	/** Called after the world has been saved. The sequencer updates to the animated state. */
+	void OnPostSaveWorld(uint32 SaveFlags, UWorld* World, bool bSuccess);
+
+	/** Called after a level has been added */
+	void OnLevelAdded(ULevel* InLevel, UWorld* InWorld);
+
+	/** Called after a level has been removed */
+	void OnLevelRemoved(ULevel* InLevel, UWorld* InWorld);
+
+	/** Called after a new level has been created. The sequencer editor mode needs to be enabled. */
+	void OnNewCurrentLevel();
+
+	/** Called after a map has been opened. The sequencer editor mode needs to be enabled. */
+	void OnMapOpened(const FString& Filename, bool bLoadAsTemplate);
+
+	/** Called when new actors are dropped in the viewport. */
+	void OnNewActorsDropped(const TArray<UObject*>& DroppedObjects, const TArray<AActor*>& DroppedActors);
+
+	/** Called before a PIE session begins. */
+	void OnPreBeginPIE(bool bIsSimulating);
+
+	/** Called after a PIE session ends. */
+	void OnEndPIE(bool bIsSimulating);
+
+	/** Called after PIE session ends and maps have been cleaned up */
+	void OnEndPlayMap();
+
+	/** Handles the actor selection changing externally .*/
+	void OnActorSelectionChanged( UObject* );
+
+	/** Called via UEditorEngine::GetActorRecordingStateEvent to check to see whether we need to record actor state */
+	void GetActorRecordingState( bool& bIsRecording ) const;
+
+	void OnPropertyEditorOpened();
+
+	TSharedRef<FExtender> GetLevelViewportExtender(const TSharedRef<FUICommandList> CommandList, const TArray<AActor*> InActors);
+
+	TSharedRef<FExtender> OnExtendLevelEditorViewMenu(const TSharedRef<FUICommandList> CommandList);
+
+	void RecordSelectedActors();
+
+	EVisibility GetTransportControlVisibility(TSharedPtr<ILevelViewport> LevelViewport) const;
+	
+	/** Create a menu entry we can use to toggle the transport controls */
+	void CreateTransportToggleMenuEntry(FMenuBuilder& MenuBuilder);
+
+private:
+
+	void ActivateSequencerEditorMode();
+	void AddLevelViewportMenuExtender();
+	void ActivateDetailKeyframeHandler();
+	void AttachTransportControlsToViewports();
+	void DetachTransportControlsFromViewports();
+	void BindLevelEditorCommands();
+
+private:
+
+	void IterateAllSequencers(TFunctionRef<void(FSequencer&, const FLevelEditorSequencerIntegrationOptions& Options)>) const;
+
+	FLevelEditorSequencerIntegration();
+
+private:
+
+	friend SViewportTransportControls;
+	
+	struct FSequencerAndOptions
+	{
+		TWeakPtr<FSequencer> Sequencer;
+		FLevelEditorSequencerIntegrationOptions Options;
+	};
+	TArray<FSequencerAndOptions> BoundSequencers;
+
+	/** A map of all the transport controls to viewports that this sequencer has made */
+	struct FTransportControl
+	{
+		TWeakPtr<ILevelViewport> Viewport;
+		TSharedPtr<SViewportTransportControls> Widget;
+	};
+	TArray<FTransportControl> TransportControls;
+
+	FAcquiredResources AcquiredResources;
+
+	TSharedPtr<class FDetailKeyframeHandlerWrapper> KeyFrameHandler;
+};
diff --git a/Engine/Source/Editor/Sequencer/Public/SequencerCommands.h b/Engine/Source/Editor/Sequencer/Public/SequencerCommands.h
--- a/Engine/Source/Editor/Sequencer/Public/SequencerCommands.h
+++ b/Engine/Source/Editor/Sequencer/Public/SequencerCommands.h
@@ -70,6 +70,9 @@ public:
 	/** Zoom out of the view range */
 	TSharedPtr< FUICommandInfo > ZoomOutViewRange;
 
+	/** Toggle locking the playback range. */
+	TSharedPtr< FUICommandInfo > TogglePlaybackRangeLocked;
+
 	/** Forces playback both in editor and at runtime to be evaluated at fixed frame intervals. */
 	TSharedPtr< FUICommandInfo > ToggleForceFixedFrameIntervalPlayback;
 
@@ -103,6 +106,12 @@ public:
 	/** Select all keys that fall into the selection range*/
 	TSharedPtr< FUICommandInfo > SelectKeysInSelectionRange;
 
+	/** Select all sections that fall into the selection range*/
+	TSharedPtr< FUICommandInfo > SelectSectionsInSelectionRange;
+
+	/** Select all keys and sections that fall into the selection range*/
+	TSharedPtr< FUICommandInfo > SelectAllInSelectionRange;
+
 	/** Sets a key at the current time for the selected actor */
 	TSharedPtr< FUICommandInfo > SetKey;
 
@@ -226,6 +235,9 @@ public:
 	/** Saves the current state of this object as the default spawnable state. */
 	TSharedPtr< FUICommandInfo > SaveCurrentSpawnableState;
 
+	/** Restores all animated state for the current sequence. */
+	TSharedPtr< FUICommandInfo > RestoreAnimatedState;
+
 	/** Discard all changes to the current movie scene. */
 	TSharedPtr< FUICommandInfo > DiscardChanges;
 
diff --git a/Engine/Source/Editor/SequencerWidgets/Private/STimeRangeSlider.cpp b/Engine/Source/Editor/SequencerWidgets/Private/STimeRangeSlider.cpp
--- a/Engine/Source/Editor/SequencerWidgets/Private/STimeRangeSlider.cpp
+++ b/Engine/Source/Editor/SequencerWidgets/Private/STimeRangeSlider.cpp
@@ -10,6 +10,7 @@
 namespace TimeRangeSliderConstants
 {
 	const int32 HandleSize = 14;
+	const int32 MinimumScrubberWidth = HandleSize * 2;
 }
 
 void STimeRangeSlider::Construct( const FArguments& InArgs, TSharedRef<ITimeSliderController> InTimeSliderController)
@@ -58,6 +59,14 @@ void STimeRangeSlider::ComputeHandleOffsets(float& LeftHandleOffset, float& Hand
 	LeftHandleOffset = (InTime - StartTime) * UnitsToPixel;
 	HandleOffset = LeftHandleOffset + TimeRangeSliderConstants::HandleSize;
 	RightHandleOffset = HandleOffset + (OutTime - InTime) * UnitsToPixel;
+	
+	float ScrubberWidth = RightHandleOffset-LeftHandleOffset-TimeRangeSliderConstants::HandleSize;
+	if (ScrubberWidth < (float)TimeRangeSliderConstants::MinimumScrubberWidth)
+	{
+		HandleOffset = HandleOffset - ((float)TimeRangeSliderConstants::MinimumScrubberWidth - ScrubberWidth) / 2.f;
+		LeftHandleOffset = HandleOffset - TimeRangeSliderConstants::HandleSize;
+		RightHandleOffset = HandleOffset + (float)TimeRangeSliderConstants::MinimumScrubberWidth;
+	}
 }
 
 
diff --git a/Engine/Source/Editor/SequencerWidgets/Public/ITimeSlider.h b/Engine/Source/Editor/SequencerWidgets/Public/ITimeSlider.h
--- a/Engine/Source/Editor/SequencerWidgets/Public/ITimeSlider.h
+++ b/Engine/Source/Editor/SequencerWidgets/Public/ITimeSlider.h
@@ -129,6 +129,12 @@ struct FTimeSliderArgs
 	/** Round the scrub position to an integer during playback */
 	TAttribute<EMovieScenePlayerStatus::Type> PlaybackStatus;
 
+	/** Attribute defining whether the playback range is locked */
+	TAttribute<bool> IsPlaybackRangeLocked;
+
+	/** Called when toggling the playback range lock */
+	FSimpleDelegate OnTogglePlaybackRangeLocked;
+
 	/** If we are allowed to zoom */
 	bool AllowZoom;
 
diff --git a/Engine/Source/Editor/UnrealEd/Private/Fbx/FbxMainExport.cpp b/Engine/Source/Editor/UnrealEd/Private/Fbx/FbxMainExport.cpp
--- a/Engine/Source/Editor/UnrealEd/Private/Fbx/FbxMainExport.cpp
+++ b/Engine/Source/Editor/UnrealEd/Private/Fbx/FbxMainExport.cpp
@@ -70,6 +70,7 @@
 #include "CineCameraComponent.h"
 #include "Math/UnitConversion.h"
 
+#include "IMovieScenePlayer.h"
 #include "MovieScene.h"
 #include "Tracks/MovieScene3DTransformTrack.h"
 #include "Tracks/MovieSceneFloatTrack.h"
diff --git a/Engine/Source/Editor/UnrealEd/Private/FileHelpers.cpp b/Engine/Source/Editor/UnrealEd/Private/FileHelpers.cpp
--- a/Engine/Source/Editor/UnrealEd/Private/FileHelpers.cpp
+++ b/Engine/Source/Editor/UnrealEd/Private/FileHelpers.cpp
@@ -967,7 +967,7 @@ void FEditorFileUtils::SaveAssetsAs(const TArray<UObject*>& Assets, TArray<UObje
 		{
 			if (!OpenSaveAsDialog(Asset->GetClass(), OldPackagePath, OldAssetName, NewPackageName))
 			{
-				break;
+				return;
 			}
 
 			FText OutError;
