Making sequencer mesh trails non-selectable, so that only keys are transformable.

diff --git a/Engine/Source/Editor/Sequencer/Private/SequencerMeshTrail.cpp b/Engine/Source/Editor/Sequencer/Private/SequencerMeshTrail.cpp
--- a/Engine/Source/Editor/Sequencer/Private/SequencerMeshTrail.cpp
+++ b/Engine/Source/Editor/Sequencer/Private/SequencerMeshTrail.cpp
@@ -5,11 +5,15 @@
 #include "ViewportWorldInteraction.h"
 #include "EditorWorldExtension.h"
 #include "Materials/MaterialInstance.h" 
+#include "Components/BillboardComponent.h"
 
 ASequencerMeshTrail::ASequencerMeshTrail()
 	: Super()
 {
-	RootComponent = CreateDefaultSubobject<USceneComponent>(TEXT("RootComponent"));
+	UBillboardComponent* PrimitiveRootComponent = CreateDefaultSubobject<UBillboardComponent>(TEXT("RootComponent"));
+	PrimitiveRootComponent->bSelectable = false;
+	PrimitiveRootComponent->bVisible = false;
+	SetRootComponent(PrimitiveRootComponent);
 }
 
 void ASequencerMeshTrail::Cleanup()
