Mesh trail root component visibility was incorrectly propogating to children

diff --git a/Engine/Source/Editor/Sequencer/Private/SequencerMeshTrail.cpp b/Engine/Source/Editor/Sequencer/Private/SequencerMeshTrail.cpp
--- a/Engine/Source/Editor/Sequencer/Private/SequencerMeshTrail.cpp
+++ b/Engine/Source/Editor/Sequencer/Private/SequencerMeshTrail.cpp
@@ -12,7 +12,7 @@ ASequencerMeshTrail::ASequencerMeshTrail()
 {
 	UBillboardComponent* PrimitiveRootComponent = CreateDefaultSubobject<UBillboardComponent>(TEXT("RootComponent"));
 	PrimitiveRootComponent->bSelectable = false;
-	PrimitiveRootComponent->bVisible = false;
+	PrimitiveRootComponent->SetVisibility(false, false);
 	SetRootComponent(PrimitiveRootComponent);
 }
 
