Moving LevelVC settings inside check for LevelVC being valid.

diff --git a/Engine/Source/Editor/Sequencer/Private/Sequencer.cpp b/Engine/Source/Editor/Sequencer/Private/Sequencer.cpp
--- a/Engine/Source/Editor/Sequencer/Private/Sequencer.cpp
+++ b/Engine/Source/Editor/Sequencer/Private/Sequencer.cpp
@@ -1925,32 +1925,35 @@ void FSequencer::SetViewportSettings(const TMap<FViewportClient*, EMovieSceneVie
 
 	for (FLevelEditorViewportClient* LevelVC : GEditor->LevelViewportClients)
 	{
-		if (LevelVC && LevelVC->IsPerspective() && LevelVC->AllowsCinematicPreview())
+		if (LevelVC && LevelVC->IsPerspective())
 		{
-			if (ViewportParamsMap.Contains(LevelVC))
+			if (LevelVC->AllowsCinematicPreview())
 			{
-				const EMovieSceneViewportParams* ViewportParams = ViewportParamsMap.Find(LevelVC);
-				if (ViewportParams->SetWhichViewportParam & EMovieSceneViewportParams::SVP_FadeAmount)
-				{
-					LevelVC->FadeAmount = ViewportParams->FadeAmount;
-					LevelVC->bEnableFading = true;
-				}
-				if (ViewportParams->SetWhichViewportParam & EMovieSceneViewportParams::SVP_FadeColor)
-				{
-					LevelVC->FadeColor = ViewportParams->FadeColor.ToRGBE();
-					LevelVC->bEnableFading = true;
-				}
-				if (ViewportParams->SetWhichViewportParam & EMovieSceneViewportParams::SVP_ColorScaling)
+				if (ViewportParamsMap.Contains(LevelVC))
 				{
-					LevelVC->bEnableColorScaling = ViewportParams->bEnableColorScaling;
-					LevelVC->ColorScale = ViewportParams->ColorScale;
+					const EMovieSceneViewportParams* ViewportParams = ViewportParamsMap.Find(LevelVC);
+					if (ViewportParams->SetWhichViewportParam & EMovieSceneViewportParams::SVP_FadeAmount)
+					{
+						LevelVC->FadeAmount = ViewportParams->FadeAmount;
+						LevelVC->bEnableFading = true;
+					}
+					if (ViewportParams->SetWhichViewportParam & EMovieSceneViewportParams::SVP_FadeColor)
+					{
+						LevelVC->FadeColor = ViewportParams->FadeColor.ToRGBE();
+						LevelVC->bEnableFading = true;
+					}
+					if (ViewportParams->SetWhichViewportParam & EMovieSceneViewportParams::SVP_ColorScaling)
+					{
+						LevelVC->bEnableColorScaling = ViewportParams->bEnableColorScaling;
+						LevelVC->ColorScale = ViewportParams->ColorScale;
+					}
 				}
 			}
-		}
-		else
-		{
-			LevelVC->bEnableFading = false;
-			LevelVC->bEnableColorScaling = false;
+			else
+			{
+				LevelVC->bEnableFading = false;
+				LevelVC->bEnableColorScaling = false;
+			}
 		}
 	}
 }
