Sequencer: Fix find property container so that it returns when it finds the first matching property.

diff --git a/Engine/Source/Editor/Sequencer/Private/KeyPropertyParams.cpp b/Engine/Source/Editor/Sequencer/Private/KeyPropertyParams.cpp
--- a/Engine/Source/Editor/Sequencer/Private/KeyPropertyParams.cpp
+++ b/Engine/Source/Editor/Sequencer/Private/KeyPropertyParams.cpp
@@ -46,6 +46,10 @@ const UStruct* FCanKeyPropertyParams::FindPropertyContainer(const UProperty* For
 		if (!bFoundProperty)
 		{
 			bFoundProperty = Property == ForProperty;
+			if (bFoundProperty)
+			{
+				return ObjectClass;
+			}
 		}
 		else if (const UStructProperty* StructProperty = Cast<const UStructProperty>(Property))
 		{
