Fixed localization CIS issues

diff --git a/Engine/Source/Editor/PIEPreviewDeviceProfileSelector/Private/PIEPreviewDeviceEnumeration.cpp b/Engine/Source/Editor/PIEPreviewDeviceProfileSelector/Private/PIEPreviewDeviceEnumeration.cpp
--- a/Engine/Source/Editor/PIEPreviewDeviceProfileSelector/Private/PIEPreviewDeviceEnumeration.cpp
+++ b/Engine/Source/Editor/PIEPreviewDeviceProfileSelector/Private/PIEPreviewDeviceEnumeration.cpp
@@ -77,4 +77,4 @@ FText FPIEPreviewDeviceContainerCategory::GetCategoryToolTip()
 	return FText::Format(LOCTEXT("PIEPreviewDeviceCategoryToolTip", "Preview {Device} devices"), Args);
 }
 
-
+#undef LOCTEXT_NAMESPACE
diff --git a/Engine/Source/Editor/Sequencer/Private/SequencerUtilities.cpp b/Engine/Source/Editor/Sequencer/Private/SequencerUtilities.cpp
--- a/Engine/Source/Editor/Sequencer/Private/SequencerUtilities.cpp
+++ b/Engine/Source/Editor/Sequencer/Private/SequencerUtilities.cpp
@@ -94,7 +94,7 @@ void FSequencerUtilities::PopulateMenu_CreateNewSection(FMenuBuilder& MenuBuilde
 		const float StartAtTime = InSequencer->GetLocalTime();
 		TRange<float> VisibleRange = InSequencer->GetViewRange();
 
-		FScopedTransaction Transaction(LOCTEXT("AddSectionText", "Add Section"));
+		FScopedTransaction Transaction(LOCTEXT("AddSectionTransactionText", "Add Section"));
 		if (UMovieSceneSection* NewSection = Track->CreateNewSection())
 		{
 			int32 OverlapPriority = 0;
