Fix blending curves offset with a high dpi monitor

diff --git a/Engine/Source/Editor/Sequencer/Private/SSequencerSection.cpp b/Engine/Source/Editor/Sequencer/Private/SSequencerSection.cpp
--- a/Engine/Source/Editor/Sequencer/Private/SSequencerSection.cpp
+++ b/Engine/Source/Editor/Sequencer/Private/SSequencerSection.cpp
@@ -403,8 +403,8 @@ struct FSequencerSectionPainterImpl : FSequencerSectionPainter
 
 		FSlateRenderTransform RenderTransform;
 
-		const FVector2D Pos = RangeGeometry.AbsolutePosition;
-		const FVector2D Size = RangeGeometry.Size;
+		const FVector2D Pos = RangeGeometry.GetAbsolutePosition();
+		const FVector2D Size = RangeGeometry.GetLocalSize();
 
 		FLinearColor EaseSelectionColor = FEditorStyle::GetSlateColor(SequencerSectionConstants::SelectionColorName).GetColor(FWidgetStyle());
 
@@ -446,12 +446,12 @@ struct FSequencerSectionPainterImpl : FSequencerSectionPainter
 
 				// Add verts top->bottom
 				FVector2D UV(U, 0.f);
-				Verts.Add(FSlateVertex::Make<ESlateVertexRounding::Disabled>(RenderTransform, Pos + UV*Size, AtlasOffset + UV*AtlasUVSize, FillColor));
+				Verts.Add(FSlateVertex::Make<ESlateVertexRounding::Disabled>(RenderTransform, (Pos + UV*Size*RangeGeometry.Scale), AtlasOffset + UV*AtlasUVSize, FillColor));
 
 				UV.Y = 1.f - Point.Location.Y;
 				BorderPoints.Add(UV*Size);
 				BorderPointColors.Add(Point.Color);
-				Verts.Add(FSlateVertex::Make<ESlateVertexRounding::Disabled>(RenderTransform, Pos + UV*Size, AtlasOffset + FVector2D(UV.X, 0.5f)*AtlasUVSize, FillColor));
+				Verts.Add(FSlateVertex::Make<ESlateVertexRounding::Disabled>(RenderTransform, (Pos + UV*Size*RangeGeometry.Scale), AtlasOffset + FVector2D(UV.X, 0.5f)*AtlasUVSize, FillColor));
 
 				if (Verts.Num() >= 4)
 				{
