Copying //UE4/Dev-Sequencer to //UE4/Main (Source: //UE4/Dev-Sequencer @ 3631098)

diff --git a/Engine/Source/Editor/Sequencer/Private/GroupedKeyArea.cpp b/Engine/Source/Editor/Sequencer/Private/GroupedKeyArea.cpp
--- a/Engine/Source/Editor/Sequencer/Private/GroupedKeyArea.cpp
+++ b/Engine/Source/Editor/Sequencer/Private/GroupedKeyArea.cpp
@@ -6,43 +6,16 @@
 #include "DisplayNodes/SequencerSectionKeyAreaNode.h"
 #include "DisplayNodes/SequencerTrackNode.h"
 
-FIndexKey::FIndexKey(FString InNodePath, UMovieSceneSection* InSection)
-	: NodePath(MoveTemp(InNodePath))
-	, Section(InSection)
-	, CachedHash(HashCombine(GetTypeHash(NodePath), GetTypeHash(InSection)))
-{}
-
-namespace 
+FSequencerKeyCollectionSignature FSequencerKeyCollectionSignature::FromNodes(const TArray<FSequencerDisplayNode*>& InNodes, float InDuplicateThresholdTime)
 {
-	/** Structure that defines the index for a particular section */
-	struct FIndexEntry
-	{
-		TMap<FKeyHandle, int32> HandleToGroup;
-		TArray<FKeyHandle> GroupHandles;
-		TArray<float> RepresentativeTimes;
-	};
-
-	/** A persistent index is required to ensure that generated key handles are maintained for the lifetime of specific display nodes */
-	TMap<FIndexKey, FIndexEntry> GlobalIndex;
-}
+	FSequencerKeyCollectionSignature Result;
+	Result.DuplicateThresholdTime = InDuplicateThresholdTime;
 
-FGroupedKeyCollection::FGroupedKeyCollection()
-	: IndexKey(FString(), nullptr)
-{
-	GroupingThreshold = SMALL_NUMBER;
-}
-
-void FGroupedKeyCollection::InitializeExplicit(const TArray<FSequencerDisplayNode*>& InNodes, float DuplicateThreshold)
-{
-	KeyAreas.Reset();
-	Groups.Reset();
-
-	TArray<TSharedRef<FSequencerSectionKeyAreaNode>> AllKeyAreaNodes;
-	AllKeyAreaNodes.Reserve(36);
 	for (const FSequencerDisplayNode* Node : InNodes)
 	{
 		const FSequencerSectionKeyAreaNode* KeyAreaNode = nullptr;
 
+		check(Node);
 		if (Node->GetType() == ESequencerNode::KeyArea)
 		{
 			KeyAreaNode = static_cast<const FSequencerSectionKeyAreaNode*>(Node);
@@ -50,32 +23,25 @@ void FGroupedKeyCollection::InitializeExplicit(const TArray<FSequencerDisplayNod
 		else if (Node->GetType() == ESequencerNode::Track)
 		{
 			KeyAreaNode = static_cast<const FSequencerTrackNode*>(Node)->GetTopLevelKeyNode().Get();
-			if (!KeyAreaNode)
-			{
-				continue;
-			}
-		}
-		else
-		{
-			continue;
 		}
 
-		const TArray<TSharedRef<IKeyArea>>& AllKeyAreas = KeyAreaNode->GetAllKeyAreas();
-		KeyAreas.Reserve(KeyAreas.Num() + AllKeyAreas.Num());
-
-		for (const TSharedRef<IKeyArea>& KeyArea : AllKeyAreas)
+		if (KeyAreaNode)
 		{
-			AddKeyArea(KeyArea);
+			for (const TSharedRef<IKeyArea>& KeyArea : KeyAreaNode->GetAllKeyAreas())
+			{
+				const UMovieSceneSection* Section = KeyArea->GetOwningSection();
+				Result.KeyAreaToSignature.Add(KeyArea, Section ? Section->GetSignature() : FGuid());
+			}
 		}
 	}
 
-	RemoveDuplicateKeys(DuplicateThreshold);
+	return Result;
 }
 
-void FGroupedKeyCollection::InitializeRecursive(const TArray<FSequencerDisplayNode*>& InNodes, float DuplicateThreshold)
+FSequencerKeyCollectionSignature FSequencerKeyCollectionSignature::FromNodesRecursive(const TArray<FSequencerDisplayNode*>& InNodes, float InDuplicateThresholdTime)
 {
-	KeyAreas.Reset();
-	Groups.Reset();
+	FSequencerKeyCollectionSignature Result;
+	Result.DuplicateThresholdTime = InDuplicateThresholdTime;
 
 	TArray<TSharedRef<FSequencerSectionKeyAreaNode>> AllKeyAreaNodes;
 	AllKeyAreaNodes.Reserve(36);
@@ -89,24 +55,22 @@ void FGroupedKeyCollection::InitializeRecursive(const TArray<FSequencerDisplayNo
 		Node->GetChildKeyAreaNodesRecursively(AllKeyAreaNodes);
 	}
 
-	for (const auto& Node : AllKeyAreaNodes)
+	for (const TSharedRef<FSequencerSectionKeyAreaNode>& Node : AllKeyAreaNodes)
 	{
-		const TArray<TSharedRef<IKeyArea>>& AllKeyAreas = Node->GetAllKeyAreas();
-		KeyAreas.Reserve(KeyAreas.Num() + AllKeyAreas.Num());
-
-		for (const TSharedRef<IKeyArea>& KeyArea : AllKeyAreas)
+		for (const TSharedRef<IKeyArea>& KeyArea : Node->GetAllKeyAreas())
 		{
-			AddKeyArea(KeyArea);
+			const UMovieSceneSection* Section = KeyArea->GetOwningSection();
+			Result.KeyAreaToSignature.Add(KeyArea, Section ? Section->GetSignature() : FGuid());
 		}
 	}
 
-	RemoveDuplicateKeys(DuplicateThreshold);
+	return Result;
 }
 
-void FGroupedKeyCollection::InitializeRecursive(FSequencerDisplayNode& InNode, UMovieSceneSection* InSection, float DuplicateThreshold)
+FSequencerKeyCollectionSignature FSequencerKeyCollectionSignature::FromNodeRecursive(FSequencerDisplayNode& InNode, UMovieSceneSection* InSection, float InDuplicateThresholdTime)
 {
-	KeyAreas.Reset();
-	Groups.Reset();
+	FSequencerKeyCollectionSignature Result;
+	Result.DuplicateThresholdTime = InDuplicateThresholdTime;
 
 	TArray<TSharedRef<FSequencerSectionKeyAreaNode>> AllKeyAreaNodes;
 	AllKeyAreaNodes.Reserve(36);
@@ -117,11 +81,162 @@ void FGroupedKeyCollection::InitializeRecursive(FSequencerDisplayNode& InNode, U
 		TSharedPtr<IKeyArea> KeyArea = Node->GetKeyArea(InSection);
 		if (KeyArea.IsValid())
 		{
-			AddKeyArea(KeyArea.ToSharedRef());
+			Result.KeyAreaToSignature.Add(KeyArea.ToSharedRef(), InSection ? InSection->GetSignature() : FGuid());
+		}
+	}
+
+	return Result;
+}
+
+bool FSequencerKeyCollectionSignature::HasUncachableContent() const
+{
+	for (auto& Pair : KeyAreaToSignature)
+	{
+		if (!Pair.Value.IsValid())
+		{
+			return true;
 		}
 	}
+	return false;
+}
+
+bool operator!=(const FSequencerKeyCollectionSignature& A, const FSequencerKeyCollectionSignature& B)
+{
+	if (A.HasUncachableContent() || B.HasUncachableContent())
+	{
+		return true;
+	}
+
+	if (A.KeyAreaToSignature.Num() != B.KeyAreaToSignature.Num() || A.DuplicateThresholdTime != B.DuplicateThresholdTime)
+	{
+		return true;
+	}
+
+	for (auto& Pair : A.KeyAreaToSignature)
+	{
+		const FGuid* BSig = B.KeyAreaToSignature.Find(Pair.Key);
+		if (!BSig || *BSig != Pair.Value)
+		{
+			return true;
+		}
+	}
+
+	return false;
+}
+
+bool operator==(const FSequencerKeyCollectionSignature& A, const FSequencerKeyCollectionSignature& B)
+{
+	if (A.HasUncachableContent() || B.HasUncachableContent())
+	{
+		return false;
+	}
+
+	if (A.KeyAreaToSignature.Num() != B.KeyAreaToSignature.Num() || A.DuplicateThresholdTime != B.DuplicateThresholdTime)
+	{
+		return false;
+	}
+
+	for (auto& Pair : A.KeyAreaToSignature)
+	{
+		const FGuid* BSig = B.KeyAreaToSignature.Find(Pair.Key);
+		if (!BSig || *BSig != Pair.Value)
+		{
+			return false;
+		}
+	}
+
+	return true;
+}
+
+FIndexKey::FIndexKey(FName InNodePath, UMovieSceneSection* InSection)
+	: NodePath(MoveTemp(InNodePath))
+	, Section(InSection)
+{}
+
+namespace 
+{
+	/** Structure that defines the index for a particular section */
+	struct FIndexEntry
+	{
+		TMap<FKeyHandle, int32> HandleToGroup;
+		TArray<FKeyHandle> GroupHandles;
+		TArray<float> RepresentativeTimes;
+	};
+
+	/** A persistent index is required to ensure that generated key handles are maintained for the lifetime of specific display nodes */
+	TMap<FIndexKey, FIndexEntry> GlobalIndex;
+}
+
+FGroupedKeyCollection::FGroupedKeyCollection()
+	: IndexKey(FName(), nullptr)
+{
+	GroupingThreshold = SMALL_NUMBER;
+}
+
+bool FGroupedKeyCollection::InitializeExplicit(const TArray<FSequencerDisplayNode*>& InNodes, float DuplicateThreshold)
+{
+	FSequencerKeyCollectionSignature UpToDateSignature = FSequencerKeyCollectionSignature::FromNodes(InNodes, DuplicateThreshold);
+	if (CacheSignature == UpToDateSignature)
+	{
+		return false;
+	}
+
+	Swap(UpToDateSignature, CacheSignature);
+
+	KeyAreas.Reset();
+	Groups.Reset();
+
+	for (auto& Pair : CacheSignature.GetKeyAreas())
+	{
+		AddKeyArea(Pair.Key);
+	}
+
+	RemoveDuplicateKeys(DuplicateThreshold);
+	return true;
+}
+
+bool FGroupedKeyCollection::InitializeRecursive(const TArray<FSequencerDisplayNode*>& InNodes, float DuplicateThreshold)
+{
+	FSequencerKeyCollectionSignature UpToDateSignature = FSequencerKeyCollectionSignature::FromNodesRecursive(InNodes, DuplicateThreshold);
+	if (CacheSignature == UpToDateSignature)
+	{
+		return false;
+	}
+
+	Swap(UpToDateSignature, CacheSignature);
+
+	KeyAreas.Reset();
+	Groups.Reset();
+
+	for (auto& Pair : CacheSignature.GetKeyAreas())
+	{
+		AddKeyArea(Pair.Key);
+	}
+
+	RemoveDuplicateKeys(DuplicateThreshold);
+	return true;
+}
+
+bool FGroupedKeyCollection::InitializeRecursive(FSequencerDisplayNode& InNode, UMovieSceneSection* InSection, float DuplicateThreshold)
+{
+	FSequencerKeyCollectionSignature UpToDateSignature = FSequencerKeyCollectionSignature::FromNodeRecursive(InNode, InSection, DuplicateThreshold);
+	if (CacheSignature == UpToDateSignature)
+	{
+		return false;
+	}
+
+	Swap(UpToDateSignature, CacheSignature);
+
+	KeyAreas.Reset();
+	Groups.Reset();
+
+	for (auto& Pair : CacheSignature.GetKeyAreas())
+	{
+		AddKeyArea(Pair.Key);
+	}
 
 	RemoveDuplicateKeys(DuplicateThreshold);
+	return true;
 }
 
 void FGroupedKeyCollection::AddKeyArea(const TSharedRef<IKeyArea>& InKeyArea)
@@ -142,34 +257,41 @@ void FGroupedKeyCollection::RemoveDuplicateKeys(float DuplicateThreshold)
 {
 	GroupingThreshold = DuplicateThreshold;
 
-	Groups.Sort([](const FKeyGrouping& A, const FKeyGrouping& B){
+	TArray<FKeyGrouping> OldGroups;
+	Swap(OldGroups, Groups);
+
+	Groups.Reserve(OldGroups.Num());
+
+	OldGroups.Sort([](const FKeyGrouping& A, const FKeyGrouping& B){
 		return A.RepresentativeTime < B.RepresentativeTime;
 	});
 
 	// Remove duplicates
-	for (int32 Index = 0; Index < Groups.Num(); ++Index)
+	for (int32 ReadIndex = 0; ReadIndex < OldGroups.Num(); ++ReadIndex)
 	{
-		const float CurrentTime = Groups[Index].RepresentativeTime;
+		const float CurrentTime = OldGroups[ReadIndex].RepresentativeTime;
 
 		int32 NumToMerge = 0;
-		for (int32 DuplIndex = Index + 1; DuplIndex < Groups.Num(); ++DuplIndex)
+		for (int32 DuplIndex = ReadIndex + 1; DuplIndex < OldGroups.Num(); ++DuplIndex)
 		{
-			if (!FMath::IsNearlyEqual(CurrentTime, Groups[DuplIndex].RepresentativeTime, DuplicateThreshold))
+			if (!FMath::IsNearlyEqual(CurrentTime, OldGroups[DuplIndex].RepresentativeTime, DuplicateThreshold))
 			{
 				break;
 			}
 			++NumToMerge;
 		}
 
+		Groups.Add(MoveTemp(OldGroups[ReadIndex]));
+
 		if (NumToMerge)
 		{
-			const int32 Start = Index + 1, End = Start + NumToMerge;
+			const int32 Start = ReadIndex + 1, End = Start + NumToMerge;
 			for (int32 MergeIndex = Start; MergeIndex < End; ++MergeIndex)
 			{
-				Groups[Index].Keys.Append(MoveTemp(Groups[MergeIndex].Keys));
+				Groups.Last().Keys.Append(MoveTemp(OldGroups[MergeIndex].Keys));
 			}
 
-			Groups.RemoveAt(Start, NumToMerge, false);
+			ReadIndex += NumToMerge;
 		}
 	}
 }
@@ -319,15 +441,22 @@ const FSlateBrush* FGroupedKeyCollection::GetBrush(FKeyHandle InHandle) const
 }
 
 FGroupedKeyArea::FGroupedKeyArea(FSequencerDisplayNode& InNode, UMovieSceneSection* InSection)
-	: Section(InSection)
+	: DisplayNode(InNode.AsShared())
+	, Section(InSection)
 {
 	check(Section);
+	IndexKey = FIndexKey(*InNode.GetPathName(), Section);
 
-	// Calling this virtual function is safe here, as it's just called through standard name-lookup, not runtime dispatch
-	FGroupedKeyCollection::InitializeRecursive(InNode, Section);
+	Update();
+}
 
-	IndexKey = FIndexKey(InNode.GetPathName(), Section);
-	UpdateIndex();
+void FGroupedKeyArea::Update()
+{
+	TSharedPtr<FSequencerDisplayNode> PinnedNode = DisplayNode.Pin();
+	if (PinnedNode.IsValid() && InitializeRecursive(*PinnedNode, Section))
+	{
+		UpdateIndex();
+	}
 }
 
 TArray<FKeyHandle> FGroupedKeyArea::GetUnsortedKeyHandles() const
diff --git a/Engine/Source/Editor/Sequencer/Private/GroupedKeyArea.h b/Engine/Source/Editor/Sequencer/Private/GroupedKeyArea.h
--- a/Engine/Source/Editor/Sequencer/Private/GroupedKeyArea.h
+++ b/Engine/Source/Editor/Sequencer/Private/GroupedKeyArea.h
@@ -23,28 +23,25 @@ enum class EMovieSceneKeyInterpolation : uint8;
 /** Keyable structure used to identify a particular FGroupedKeyArea */
 struct FIndexKey
 {
-	FIndexKey(FString InNodePath, UMovieSceneSection* InSection);
+	FIndexKey(FName InNodePath, UMovieSceneSection* InSection);
 
 	friend bool operator==(const FIndexKey& A, const FIndexKey& B)
 	{
-		return A.CachedHash == B.CachedHash && A.NodePath == B.NodePath && A.Section == B.Section;
+		return A.NodePath == B.NodePath && A.Section == B.Section;
 	}
 
 	friend uint32 GetTypeHash( const FIndexKey& IndexKey )
 	{
-		return IndexKey.CachedHash;
+		return GetTypeHash(IndexKey.NodePath) ^ GetTypeHash(IndexKey.Section);
 	}
 
 private:
 
 	/** Cached path of the node that we relate to */
-	FString NodePath;
+	FName NodePath;
 
 	/** The movie scene section that we relate to */
 	TWeakObjectPtr<UMovieSceneSection> Section;
-
-	/** Cached hash for efficiency (to avoid having to re-hash strings all the time) */
-	uint32 CachedHash;
 };
 
 
@@ -82,7 +79,7 @@ struct FKeyGrouping
 	/** The representative time at which all keys within this group are located */
 	float RepresentativeTime;
 	/** Array of keys contained within this group */
-	TArray<FKeyIndex> Keys;
+	TArray<FKeyIndex, TInlineAllocator<1>> Keys;
 };
 
 
@@ -97,14 +94,23 @@ public:
 
 public:
 
-	/** Initialize this key collection from the specified nodes. Only gathers keys from those explicitly specified. */
-	virtual void InitializeExplicit(const TArray<FSequencerDisplayNode*>& InNodes, float DuplicateThreshold = SMALL_NUMBER) override;
-
-	/** Initialize this key collection from the specified nodes. Gathers keys from all child nodes. */
-	virtual void InitializeRecursive(const TArray<FSequencerDisplayNode*>& InNodes, float DuplicateThreshold = SMALL_NUMBER) override;
+	/**
+	 * Initialize this key collection from the specified nodes. Only gathers keys from those explicitly specified.
+	 * @return true if this collection was (re)initialized, false if it did not need updating
+	 */
+	virtual bool InitializeExplicit(const TArray<FSequencerDisplayNode*>& InNodes, float DuplicateThreshold = SMALL_NUMBER) override;
+
+	/**
+	 * Initialize this key collection from the specified nodes. Gathers keys from all child nodes.
+	 * @return true if this collection was (re)initialized, false if it did not need updating
+	 */
+	virtual bool InitializeRecursive(const TArray<FSequencerDisplayNode*>& InNodes, float DuplicateThreshold = SMALL_NUMBER) override;
 	
-	/** Initialize this key collection from the specified node and section index. */
-	virtual void InitializeRecursive(FSequencerDisplayNode& InNode, UMovieSceneSection* InSection, float DuplicateThreshold = SMALL_NUMBER) override;
+	/**
+	 * Initialize this key collection from the specified node and section index.
+	 * @return true if this collection was (re)initialized, false if it did not need updating
+	 */
+	virtual bool InitializeRecursive(FSequencerDisplayNode& InNode, UMovieSceneSection* InSection, float DuplicateThreshold = SMALL_NUMBER) override;
 
 	/** Iterate the keys contained within this collection */
 	virtual void IterateKeys(const TFunctionRef<bool(float)>& Iter) override;
@@ -141,6 +147,10 @@ protected:
 	FKeyGrouping* FindGroup(FKeyHandle InHandle);
 
 protected:
+
+	/** A signature that is used to avoid unnecessary updates */
+	FSequencerKeyCollectionSignature CacheSignature;
+
 	/** Array of (child) key areas that we are reflecting */
 	TArray<TSharedRef<IKeyArea>> KeyAreas;
 
@@ -194,8 +204,16 @@ public:
 	virtual void PasteKeys(const FMovieSceneClipboardKeyTrack& KeyTrack, const FMovieSceneClipboardEnvironment& SrcEnvironment, const FSequencerPasteEnvironment& DstEnvironment) override;
 	virtual TOptional<FLinearColor> GetColor() override;
 
+public:
+
+	/** Ensure this key area is up to date based on the section's signature */
+	void Update();
+
 protected:
-	
+
+	/** Weak pointer back to the display node so we can update when necessary */
+	TWeakPtr<FSequencerDisplayNode> DisplayNode;
+
 	/** Pointer to the section to which this key area relates */
 	UMovieSceneSection* Section;
 };
diff --git a/Engine/Source/Editor/Sequencer/Private/MovieSceneTrackEditor.cpp b/Engine/Source/Editor/Sequencer/Private/MovieSceneTrackEditor.cpp
--- a/Engine/Source/Editor/Sequencer/Private/MovieSceneTrackEditor.cpp
+++ b/Engine/Source/Editor/Sequencer/Private/MovieSceneTrackEditor.cpp
@@ -151,13 +151,19 @@ TSharedPtr<SWidget> FMovieSceneTrackEditor::BuildOutlinerEditWidget(const FGuid&
 {
 	if (Track->GetSupportedBlendTypes().Num() > 0)
 	{
-		TSharedPtr<ISequencer> SequencerPtr = GetSequencer();
+		TWeakPtr<ISequencer> WeakSequencer = GetSequencer();
 
 		const int32 RowIndex = Params.TrackInsertRowIndex;
 		auto SubMenuCallback = [=]() -> TSharedRef<SWidget>
 		{
 			FMenuBuilder MenuBuilder(true, nullptr);
-			FSequencerUtilities::PopulateMenu_CreateNewSection(MenuBuilder, RowIndex, Track, SequencerPtr);
+
+			TSharedPtr<ISequencer> SequencerPtr = WeakSequencer.Pin();
+			if (SequencerPtr.IsValid())
+			{
+				FSequencerUtilities::PopulateMenu_CreateNewSection(MenuBuilder, RowIndex, Track, SequencerPtr);
+			}
+
 			return MenuBuilder.MakeWidget();
 		};
 
diff --git a/Engine/Source/Editor/Sequencer/Private/SSequencerSection.cpp b/Engine/Source/Editor/Sequencer/Private/SSequencerSection.cpp
--- a/Engine/Source/Editor/Sequencer/Private/SSequencerSection.cpp
+++ b/Engine/Source/Editor/Sequencer/Private/SSequencerSection.cpp
@@ -1178,9 +1178,18 @@ void SSequencerSection::PaintKeys( FSequencerSectionPainter& InPainter, const FW
 			continue;
 		}
 
-		TArray<FKeyHandle> KeyHandles = KeyArea->GetUnsortedKeyHandles();
+		// Gather keys for a region larger thean the view range to ensure we draw keys that are only just offscreen.
+		TRange<float> PaddedViewRange;
+		{
+			const float KeyWidthAsTime = TimeToPixelConverter.PixelToTime(SequencerSectionConstants::KeySize.X) - TimeToPixelConverter.PixelToTime(0);
+			const TRange<float> ViewRange = GetSequencer().GetViewRange();
+
+			PaddedViewRange = TRange<float>(ViewRange.GetLowerBoundValue() - KeyWidthAsTime, ViewRange.GetUpperBoundValue() + KeyWidthAsTime);
+		}
 
-		if (!KeyHandles.Num())
+		const FSequencerCachedKeys& CachedKeys = CachedKeyAreaPositions.FindChecked(KeyArea);
+		TArrayView<const FSequencerCachedKey> KeysInRange = CachedKeys.GetKeysInRange(PaddedViewRange);
+		if (!KeysInRange.Num())
 		{
 			continue;
 		}
@@ -1190,9 +1199,21 @@ void SSequencerSection::PaintKeys( FSequencerSectionPainter& InPainter, const FW
 		TOptional<FSlateClippingState> PreviousClipState = InPainter.DrawElements.GetClippingState();
 		InPainter.DrawElements.PopClip();
 
-		for (const FKeyHandle& KeyHandle : KeyHandles)
+		static float PixelOverlapThreshold = 3.f;
+
+		for (int32 KeyIndex = 0; KeyIndex < KeysInRange.Num(); ++KeyIndex)
 		{
-			float KeyTime = KeyArea->GetKeyTime(KeyHandle);
+			FKeyHandle KeyHandle = KeysInRange[KeyIndex].Handle;
+			const float KeyTime = KeysInRange[KeyIndex].Time;
+			const float KeyPosition = TimeToPixelConverter.TimeToPixel(KeyTime);
+
+			// Count the number of overlapping keys
+			int32 NumOverlaps = 0;
+			while (KeyIndex + 1 < KeysInRange.Num() && FMath::IsNearlyEqual(TimeToPixelConverter.TimeToPixel(KeysInRange[KeyIndex+1].Time), KeyPosition, PixelOverlapThreshold))
+			{
+				++KeyIndex;
+				++NumOverlaps;
+			}
 
 			// omit keys which would not be visible
 			if( !SectionObject.IsTimeWithinSection(KeyTime))
@@ -1302,6 +1323,12 @@ void SSequencerSection::PaintKeys( FSequencerSectionPainter& InPainter, const FW
 				FillColor = KeyColor;
 			}
 
+			// Color keys with overlaps with a red border
+			if (NumOverlaps > 0)
+			{
+				BorderColor = FLinearColor(0.83f, 0.12f, 0.12f, 1.0f); // Red
+			}
+
 			// allow group to tint the color
 			if (LayoutElement.GetType() == FSectionLayoutElement::Group)
 			{
@@ -1310,8 +1337,6 @@ void SSequencerSection::PaintKeys( FSequencerSectionPainter& InPainter, const FW
 			}
 
 			// draw border
-			float KeyPosition = TimeToPixelConverter.TimeToPixel(KeyTime);
-
 			static FVector2D ThrobAmount(12.f, 12.f);
 			const FVector2D KeySize = bSelected ? SequencerSectionConstants::KeySize + ThrobAmount * ThrobScaleValue : SequencerSectionConstants::KeySize;
 
@@ -1602,6 +1627,16 @@ void SSequencerSection::Tick( const FGeometry& AllottedGeometry, const double In
 	{
 		Layout = FSectionLayout(*ParentSectionArea, SectionIndex);
 
+		// Update cached key area key positions
+		for (const FSectionLayoutElement& LayoutElement : Layout->GetElements())
+		{
+			TSharedPtr<IKeyArea> KeyArea = LayoutElement.GetKeyArea();
+			if (KeyArea.IsValid())
+			{
+				CachedKeyAreaPositions.FindOrAdd(KeyArea).Update(KeyArea.ToSharedRef());
+			}
+		}
+
 		UMovieSceneSection* Section = SectionInterface->GetSectionObject();
 		if (Section && !Section->IsInfinite())
 		{
diff --git a/Engine/Source/Editor/Sequencer/Private/SSequencerSection.h b/Engine/Source/Editor/Sequencer/Private/SSequencerSection.h
--- a/Engine/Source/Editor/Sequencer/Private/SSequencerSection.h
+++ b/Engine/Source/Editor/Sequencer/Private/SSequencerSection.h
@@ -11,6 +11,7 @@
 #include "Rendering/RenderingCommon.h"
 #include "DisplayNodes/SequencerTrackNode.h"
 #include "SectionLayout.h"
+#include "SequencerKeyTimeCache.h"
 
 class FPaintArgs;
 class FSequencer;
@@ -149,6 +150,8 @@ private:
 	TArray<FSequencerOverlapRange> UnderlappingEasingSegments;
 	/** The signature of the track last time the overlapping segments were updated */
 	FGuid CachedTrackSignature;
+	/** Cache of key area positions */
+	TMap<TSharedPtr<IKeyArea>, FSequencerCachedKeys> CachedKeyAreaPositions;
 
 	friend struct FSequencerSectionPainterImpl;
 };
diff --git a/Engine/Source/Editor/Sequencer/Private/SectionLayout.cpp b/Engine/Source/Editor/Sequencer/Private/SectionLayout.cpp
--- a/Engine/Source/Editor/Sequencer/Private/SectionLayout.cpp
+++ b/Engine/Source/Editor/Sequencer/Private/SectionLayout.cpp
@@ -97,7 +97,7 @@ FSectionLayout::FSectionLayout(FSequencerTrackNode& TrackNode, int32 InSectionIn
 		{
 			Elements.Add(FSectionLayoutElement::FromGroup(
 				Node.AsShared(),
-				Node.UpdateKeyGrouping(Section),
+				Node.GetKeyGrouping(Section),
 				Offset
 			));
 		}
diff --git a/Engine/Source/Editor/Sequencer/Private/Sequencer.cpp b/Engine/Source/Editor/Sequencer/Private/Sequencer.cpp
--- a/Engine/Source/Editor/Sequencer/Private/Sequencer.cpp
+++ b/Engine/Source/Editor/Sequencer/Private/Sequencer.cpp
@@ -348,6 +348,13 @@ void FSequencer::Close()
 
 void FSequencer::Tick(float InDeltaTime)
 {
+	static bool bEnableRefCountCheck = true;
+	if (bEnableRefCountCheck && !FSlateApplication::Get().AnyMenusVisible())
+	{
+		const int32 SequencerRefCount = AsShared().GetSharedReferenceCount() - 1;
+		ensureAlwaysMsgf(SequencerRefCount == 1, TEXT("Multiple persistent shared references detected for Sequencer. There should only be one persistent authoritative reference. Found %d additional references which will result in FSequencer not being released correctly."), SequencerRefCount - 1);
+	}
+
 	Selection.Tick();
 	
 	if (PlaybackContextAttribute.IsBound())
@@ -631,7 +638,7 @@ TArray<UObject*> FSequencer::GetEventContexts() const
 	return Temp;
 }
 
-void FSequencer::GetKeysFromSelection(TUniquePtr<ISequencerKeyCollection>& KeyCollection)
+void FSequencer::GetKeysFromSelection(TUniquePtr<ISequencerKeyCollection>& KeyCollection, float DuplicateThreshold)
 {
 	if (!KeyCollection.IsValid())
 	{
@@ -644,10 +651,6 @@ void FSequencer::GetKeysFromSelection(TUniquePtr<ISequencerKeyCollection>& KeyCo
 		SelectedNodes.Add(&Node.Get());
 	}
 
-	// Anything within .5 pixel's worth of time is a duplicate as far as we're concerned
-	FVirtualTrackArea TrackArea = SequencerWidget->GetVirtualTrackArea();;
-	const float DuplicateThreshold = (TrackArea.PixelToTime(0.f) - TrackArea.PixelToTime(1.f)) * .5f;
-
 	KeyCollection->InitializeRecursive(SelectedNodes, DuplicateThreshold);
 }
 
@@ -809,7 +812,6 @@ void FSequencer::DeleteSections(const TSet<TWeakObjectPtr<UMovieSceneSection>>&
 		}
 
 		bAnythingRemoved = true;
-		Selection.RemoveFromSelection(Section.Get());
 	}
 
 	if (bAnythingRemoved)
@@ -818,6 +820,7 @@ void FSequencer::DeleteSections(const TSet<TWeakObjectPtr<UMovieSceneSection>>&
 		NotifyMovieSceneDataChanged( EMovieSceneDataChangeType::MovieSceneStructureItemRemoved );
 	}
 
+	Selection.EmptySelectedSections();
 	SequencerHelpers::ValidateNodesWithSelectedKeysOrSections(*this);
 }
 
@@ -2860,13 +2863,12 @@ FReply FSequencer::SetPlaybackStart()
 
 FReply FSequencer::JumpToPreviousKey()
 {
-	TUniquePtr<ISequencerKeyCollection> ActiveKeyCollection;
-	GetKeysFromSelection(ActiveKeyCollection);
-	if (ActiveKeyCollection.IsValid())
+	GetKeysFromSelection(SelectedKeyCollection, SMALL_NUMBER);
+	if (SelectedKeyCollection.IsValid())
 	{
 		TRange<float> FindRange(TRange<float>::BoundsType(), GetLocalTime());
 
-		TOptional<float> NewTime = ActiveKeyCollection->FindFirstKeyInRange(FindRange, EFindKeyDirection::Backwards);
+		TOptional<float> NewTime = SelectedKeyCollection->FindFirstKeyInRange(FindRange, EFindKeyDirection::Backwards);
 		if (NewTime.IsSet())
 		{
 			SetPlaybackStatus(EMovieScenePlayerStatus::Stepping);
@@ -2879,13 +2881,12 @@ FReply FSequencer::JumpToPreviousKey()
 
 FReply FSequencer::JumpToNextKey()
 {
-	TUniquePtr<ISequencerKeyCollection> ActiveKeyCollection;
-	GetKeysFromSelection(ActiveKeyCollection);
-	if (ActiveKeyCollection.IsValid())
+	GetKeysFromSelection(SelectedKeyCollection, SMALL_NUMBER);
+	if (SelectedKeyCollection.IsValid())
 	{
 		TRange<float> FindRange(GetLocalTime(), TRange<float>::BoundsType());
 
-		TOptional<float> NewTime = ActiveKeyCollection->FindFirstKeyInRange(FindRange, EFindKeyDirection::Forwards);
+		TOptional<float> NewTime = SelectedKeyCollection->FindFirstKeyInRange(FindRange, EFindKeyDirection::Forwards);
 		if (NewTime.IsSet())
 		{
 			SetPlaybackStatus(EMovieScenePlayerStatus::Stepping);
@@ -3146,16 +3147,15 @@ void FSequencer::OnClampRangeChanged( TRange<float> NewClampRange )
 float FSequencer::OnGetNearestKey(float InTime)
 {
 	float NearestKeyTime = InTime;
-	TUniquePtr<ISequencerKeyCollection> ActiveKeyCollection;
-	GetKeysFromSelection(ActiveKeyCollection);
+	GetKeysFromSelection(SelectedKeyCollection, SMALL_NUMBER);
 
-	if (ActiveKeyCollection.IsValid())
+	if (SelectedKeyCollection.IsValid())
 	{
 		TRange<float> FindRangeBackwards(TRange<float>::BoundsType(), NearestKeyTime);
-		TOptional<float> NewTimeBackwards = ActiveKeyCollection->FindFirstKeyInRange(FindRangeBackwards, EFindKeyDirection::Backwards);
+		TOptional<float> NewTimeBackwards = SelectedKeyCollection->FindFirstKeyInRange(FindRangeBackwards, EFindKeyDirection::Backwards);
 
 		TRange<float> FindRangeForwards(NearestKeyTime, TRange<float>::BoundsType());
-		TOptional<float> NewTimeForwards = ActiveKeyCollection->FindFirstKeyInRange(FindRangeForwards, EFindKeyDirection::Forwards);
+		TOptional<float> NewTimeForwards = SelectedKeyCollection->FindFirstKeyInRange(FindRangeForwards, EFindKeyDirection::Forwards);
 		if (NewTimeForwards.IsSet())
 		{
 			if (NewTimeBackwards.IsSet())
diff --git a/Engine/Source/Editor/Sequencer/Private/Sequencer.h b/Engine/Source/Editor/Sequencer/Private/Sequencer.h
--- a/Engine/Source/Editor/Sequencer/Private/Sequencer.h
+++ b/Engine/Source/Editor/Sequencer/Private/Sequencer.h
@@ -777,7 +777,7 @@ protected:
 protected:
 
 	/** Get all the keys for the current sequencer selection */
-	virtual void GetKeysFromSelection(TUniquePtr<ISequencerKeyCollection>& KeyCollection) override;
+	virtual void GetKeysFromSelection(TUniquePtr<ISequencerKeyCollection>& KeyCollection, float DuplicateThresoldTime) override;
 
 	UMovieSceneSection* FindNextOrPreviousShot(UMovieSceneSequence* Sequence, float Time, const bool bNext) const;
 
@@ -1074,4 +1074,7 @@ private:
 	TSharedPtr<struct FSequencerTemplateStore> TemplateStore;
 
 	TMap<FName, TFunction<void()>> CleanupFunctions;
+
+	/** Transient collection of keys that is used for jumping between keys contained within the current selection */
+	TUniquePtr<ISequencerKeyCollection> SelectedKeyCollection;
 };
diff --git a/Engine/Source/Editor/Sequencer/Private/SequencerCommonHelpers.cpp b/Engine/Source/Editor/Sequencer/Private/SequencerCommonHelpers.cpp
--- a/Engine/Source/Editor/Sequencer/Private/SequencerCommonHelpers.cpp
+++ b/Engine/Source/Editor/Sequencer/Private/SequencerCommonHelpers.cpp
@@ -245,13 +245,20 @@ bool AreKeysSelectedInNode(FSequencer& Sequencer, TSharedRef<FSequencerDisplayNo
 
 void SequencerHelpers::ValidateNodesWithSelectedKeysOrSections(FSequencer& Sequencer)
 {
+	TArray<TSharedRef<FSequencerDisplayNode>> NodesToRemove;
+
 	for (auto Node : Sequencer.GetSelection().GetNodesWithSelectedKeysOrSections())
 	{
 		if (!IsSectionSelectedInNode(Sequencer, Node) && !AreKeysSelectedInNode(Sequencer, Node))
 		{
-			Sequencer.GetSelection().RemoveFromNodesWithSelectedKeysOrSections(Node);
+			NodesToRemove.Add(Node);
 		}
 	}
+
+	for (auto Node : NodesToRemove)
+	{
+		Sequencer.GetSelection().RemoveFromNodesWithSelectedKeysOrSections(Node);
+	}
 }
 
 void SequencerHelpers::UpdateHoveredNodeFromSelectedSections(FSequencer& Sequencer)
diff --git a/Engine/Source/Editor/Sequencer/Private/SequencerContextMenus.cpp b/Engine/Source/Editor/Sequencer/Private/SequencerContextMenus.cpp
--- a/Engine/Source/Editor/Sequencer/Private/SequencerContextMenus.cpp
+++ b/Engine/Source/Editor/Sequencer/Private/SequencerContextMenus.cpp
@@ -2017,6 +2017,8 @@ void FEasingContextMenu::EasingTypeMenu(FMenuBuilder& MenuBuilder)
 
 void FEasingContextMenu::OnEasingTypeChanged(UClass* NewClass)
 {
+	FScopedTransaction Transaction(LOCTEXT("SetEasingType", "Set Easing Method"));
+
 	for (const FEasingAreaHandle& Handle : Easings)
 	{
 		UMovieSceneSection* Section = Handle.Section.GetSectionObject();
diff --git a/Engine/Source/Editor/Sequencer/Private/SequencerEdMode.cpp b/Engine/Source/Editor/Sequencer/Private/SequencerEdMode.cpp
--- a/Engine/Source/Editor/Sequencer/Private/SequencerEdMode.cpp
+++ b/Engine/Source/Editor/Sequencer/Private/SequencerEdMode.cpp
@@ -374,7 +374,7 @@ void FSequencerEdMode::DrawTransformTrack(const TSharedPtr<FSequencer>& Sequence
 	);
 	
 	const FMovieSceneEvaluationTemplateInstance* TemplateInstance = Sequencer->GetEvaluationTemplate().GetInstance(Sequencer->GetFocusedTemplateID());
-	if (!bShowTrajectory || !TemplateInstance)
+	if (!bShowTrajectory || !TemplateInstance || !TransformTrack->GetAllSections().ContainsByPredicate([](UMovieSceneSection* In){ return In->IsActive(); }))
 	{
 		return;
 	}
@@ -390,7 +390,7 @@ void FSequencerEdMode::DrawTransformTrack(const TSharedPtr<FSequencer>& Sequence
 			continue;
 		}
 
-		TArray<FTrajectoryKey> TrajectoryKeys = TransformTrack->GetTrajectoryData();
+		TArray<FTrajectoryKey> TrajectoryKeys = TransformTrack->GetTrajectoryData(Sequencer->GetLocalTime(), Sequencer->GetSettings()->GetTrajectoryPathCap());
 		for (TWeakObjectPtr<> WeakBinding : BoundObjects)
 		{
 			UObject* BoundObject = WeakBinding.Get();
diff --git a/Engine/Source/Editor/Sequencer/Private/SequencerHotspots.cpp b/Engine/Source/Editor/Sequencer/Private/SequencerHotspots.cpp
--- a/Engine/Source/Editor/Sequencer/Private/SequencerHotspots.cpp
+++ b/Engine/Source/Editor/Sequencer/Private/SequencerHotspots.cpp
@@ -47,6 +47,12 @@ TOptional<float> FSectionHotspot::GetTime() const
 	return ThisSection ? ThisSection->GetStartTime() : TOptional<float>();
 }
 
+TOptional<float> FSectionHotspot::GetOffsetTime() const
+{
+	UMovieSceneSection* ThisSection = Section.GetSectionObject();
+	return ThisSection ? ThisSection->GetOffsetTime() : TOptional<float>();
+}
+
 void FSectionHotspot::UpdateOnHover(SSequencerTrackArea& InTrackArea, ISequencer& InSequencer) const
 {
 	UMovieSceneSection* ThisSection = Section.GetSectionObject();
@@ -125,7 +131,8 @@ TSharedPtr<ISequencerEditToolDragOperation> FSectionResizeHotspot::InitiateDrag(
 		SectionHandles.Empty();
 		SectionHandles.Add(Section);
 	}
-	return MakeShareable( new FResizeSection(static_cast<FSequencer&>(Sequencer), SectionHandles, HandleType == Right) );
+	const bool bIsSlipping = false;
+	return MakeShareable( new FResizeSection(static_cast<FSequencer&>(Sequencer), SectionHandles, HandleType == Right, bIsSlipping) );
 }
 
 TOptional<float> FSectionEasingHandleHotspot::GetTime() const
diff --git a/Engine/Source/Editor/Sequencer/Private/SequencerHotspots.h b/Engine/Source/Editor/Sequencer/Private/SequencerHotspots.h
--- a/Engine/Source/Editor/Sequencer/Private/SequencerHotspots.h
+++ b/Engine/Source/Editor/Sequencer/Private/SequencerHotspots.h
@@ -68,6 +68,7 @@ struct FSectionHotspot
 	virtual ESequencerHotspot GetType() const override { return ESequencerHotspot::Section; }
 	virtual void UpdateOnHover(SSequencerTrackArea& InTrackArea, ISequencer& InSequencer) const override;
 	virtual TOptional<float> GetTime() const override;
+	virtual TOptional<float> GetOffsetTime() const override;
 	virtual TSharedPtr<ISequencerEditToolDragOperation> InitiateDrag(ISequencer&) override { return nullptr; }
 	virtual bool PopulateContextMenu(FMenuBuilder& MenuBuilder, ISequencer& Sequencer, float MouseDownTime) override;
 
diff --git a/Engine/Source/Editor/Sequencer/Private/SequencerSettings.cpp b/Engine/Source/Editor/Sequencer/Private/SequencerSettings.cpp
--- a/Engine/Source/Editor/Sequencer/Private/SequencerSettings.cpp
+++ b/Engine/Source/Editor/Sequencer/Private/SequencerSettings.cpp
@@ -51,6 +51,7 @@ USequencerSettings::USequencerSettings( const FObjectInitializer& ObjectInitiali
 	bEvaluateSubSequencesInIsolation = false;
 	bRerunConstructionScripts = false;
 	bVisualizePreAndPostRoll = true;
+	TrajectoryPathCap = 250;
 }
 
 void USequencerSettings::PostEditChangeProperty(FPropertyChangedEvent& PropertyChangedEvent)
diff --git a/Engine/Source/Editor/Sequencer/Private/Tools/EditToolDragOperations.cpp b/Engine/Source/Editor/Sequencer/Private/Tools/EditToolDragOperations.cpp
--- a/Engine/Source/Editor/Sequencer/Private/Tools/EditToolDragOperations.cpp
+++ b/Engine/Source/Editor/Sequencer/Private/Tools/EditToolDragOperations.cpp
@@ -149,10 +149,11 @@ void FEditToolDragOperation::EndTransaction()
 	Sequencer.NotifyMovieSceneDataChanged( EMovieSceneDataChangeType::TrackValueChanged );
 }
 
-FResizeSection::FResizeSection( FSequencer& InSequencer, TArray<FSectionHandle> InSections, bool bInDraggingByEnd )
+FResizeSection::FResizeSection( FSequencer& InSequencer, TArray<FSectionHandle> InSections, bool bInDraggingByEnd, bool bInIsSlipping )
 	: FEditToolDragOperation( InSequencer )
 	, Sections( MoveTemp(InSections) )
 	, bDraggingByEnd(bInDraggingByEnd)
+	, bIsSlipping(bInIsSlipping)
 	, MouseDownTime(0.f)
 {
 }
@@ -265,6 +266,10 @@ void FResizeSection::OnDrag(const FPointerEvent& MouseEvent, FVector2D LocalMous
 						float DilationFactor = NewSize / Section->GetTimeSize();
 						SequencerSection->DilateSection(DilationFactor, Section->GetStartTime(), DraggedKeyHandles);
 					}
+					else if (bIsSlipping)
+					{
+						SequencerSection->SlipSection( NewTime );
+					}
 					else
 					{
 						SequencerSection->ResizeSection( SSRM_TrailingEdge, NewTime );
@@ -282,6 +287,10 @@ void FResizeSection::OnDrag(const FPointerEvent& MouseEvent, FVector2D LocalMous
 						float DilationFactor = NewSize / Section->GetTimeSize();
 						SequencerSection->DilateSection(DilationFactor, Section->GetEndTime(), DraggedKeyHandles);
 					}
+					else if (bIsSlipping)
+					{
+						SequencerSection->SlipSection( NewTime );
+					}
 					else
 					{
 						SequencerSection->ResizeSection( SSRM_LeadingEdge, NewTime );
diff --git a/Engine/Source/Editor/Sequencer/Private/Tools/EditToolDragOperations.h b/Engine/Source/Editor/Sequencer/Private/Tools/EditToolDragOperations.h
--- a/Engine/Source/Editor/Sequencer/Private/Tools/EditToolDragOperations.h
+++ b/Engine/Source/Editor/Sequencer/Private/Tools/EditToolDragOperations.h
@@ -64,7 +64,7 @@ class FResizeSection
 public:
 
 	/** Create and initialize a new instance. */
-	FResizeSection( FSequencer& InSequencer, TArray<FSectionHandle> Sections, bool bInDraggingByEnd );
+	FResizeSection( FSequencer& InSequencer, TArray<FSectionHandle> Sections, bool bInDraggingByEnd, bool bIsSlipping );
 
 public:
 
@@ -83,6 +83,9 @@ private:
 	/** true if dragging  the end of the section, false if dragging the start */
 	bool bDraggingByEnd;
 
+	/** true if slipping, adjust only the start offset */
+	bool bIsSlipping;
+
 	/** Time where the mouse is pressed */
 	float MouseDownTime;
 
diff --git a/Engine/Source/Editor/Sequencer/Private/Tools/SequencerEditTool_Movement.cpp b/Engine/Source/Editor/Sequencer/Private/Tools/SequencerEditTool_Movement.cpp
--- a/Engine/Source/Editor/Sequencer/Private/Tools/SequencerEditTool_Movement.cpp
+++ b/Engine/Source/Editor/Sequencer/Private/Tools/SequencerEditTool_Movement.cpp
@@ -13,6 +13,7 @@
 #include "Tools/EditToolDragOperations.h"
 #include "IKeyArea.h"
 #include "SBox.h"
+#include "SlateApplication.h"
 
 const FName FSequencerEditTool_Movement::Identifier = "Movement";
 
@@ -105,6 +106,20 @@ bool FSequencerEditTool_Movement::GetHotspotTime(float& HotspotTime) const
 	return false;
 }
 
+float FSequencerEditTool_Movement::GetHotspotOffsetTime(float CurrentTime) const
+{
+	//@todo abstract dragging offset from shift
+	if (DelayedDrag->Hotspot.IsValid() && FSlateApplication::Get().GetModifierKeys().IsShiftDown())
+	{
+		TOptional<float> OptionalOffsetTime = DelayedDrag->Hotspot->GetOffsetTime();
+		if (OptionalOffsetTime.IsSet())
+		{
+			return OptionalOffsetTime.GetValue();
+		}
+	}
+	return CurrentTime - OriginalHotspotTime;
+}
+
 TSharedPtr<ISequencerEditToolDragOperation> FSequencerEditTool_Movement::CreateDrag(const FPointerEvent& MouseEvent)
 {
 	FSequencerSelection& Selection = Sequencer.GetSelection();
@@ -149,7 +164,16 @@ TSharedPtr<ISequencerEditToolDragOperation> FSequencerEditTool_Movement::CreateD
 				SequencerHelpers::UpdateHoveredNodeFromSelectedSections(Sequencer);
 				SectionHandles.Add(SectionToDrag.GetValue());
 			}
-			return MakeShareable( new FMoveSection( Sequencer, SectionHandles ) );
+			if (MouseEvent.IsShiftDown())
+			{
+				const bool bDraggingByEnd = false;
+				const bool bIsSlipping = true;
+				return MakeShareable( new FResizeSection( Sequencer, SectionHandles, bDraggingByEnd, bIsSlipping ) );
+			}
+			else
+			{
+				return MakeShareable( new FMoveSection( Sequencer, SectionHandles ) );
+			}
 		}
 		// Moving key(s)?
 		else if (HotspotType == ESequencerHotspot::Key)
@@ -345,7 +369,8 @@ int32 FSequencerEditTool_Movement::OnPaint(const FGeometry& AllottedGeometry, co
 				);
 
 				// draw offset string
-				const FString OffsetString = TimeToString(CurrentTime - OriginalHotspotTime, true);
+				float OffsetTime = GetHotspotOffsetTime(CurrentTime);
+				const FString OffsetString = TimeToString(OffsetTime, true);
 				const FVector2D OffsetStringSize = FontMeasureService->Measure(OffsetString, SmallLayoutFont);
 				const FVector2D OffsetPos = FVector2D(NewPos.X + MousePadding, NewPos.Y - 0.5f * OffsetStringSize.Y);
 
diff --git a/Engine/Source/Editor/Sequencer/Private/Tools/SequencerEditTool_Movement.h b/Engine/Source/Editor/Sequencer/Private/Tools/SequencerEditTool_Movement.h
--- a/Engine/Source/Editor/Sequencer/Private/Tools/SequencerEditTool_Movement.h
+++ b/Engine/Source/Editor/Sequencer/Private/Tools/SequencerEditTool_Movement.h
@@ -46,6 +46,7 @@ private:
 	TSharedPtr<ISequencerEditToolDragOperation> CreateDrag(const FPointerEvent& MouseEvent);
 
 	bool GetHotspotTime(float& HotspotTime) const;
+	float GetHotspotOffsetTime(float CurrentTime) const;
 
 	struct FDelayedDrag_Hotspot : FDelayedDrag
 	{
diff --git a/Engine/Source/Editor/Sequencer/Private/Tools/SequencerEntityVisitor.cpp b/Engine/Source/Editor/Sequencer/Private/Tools/SequencerEntityVisitor.cpp
--- a/Engine/Source/Editor/Sequencer/Private/Tools/SequencerEntityVisitor.cpp
+++ b/Engine/Source/Editor/Sequencer/Private/Tools/SequencerEntityVisitor.cpp
@@ -204,7 +204,7 @@ void FSequencerEntityWalker::HandleSingleNode(const ISequencerEntityVisitor& Vis
 				// Only handle grouped keys if we actually have children
 				if (InNode->GetChildNodes().Num() != 0 && Range.IntersectKeyArea(InNode, VirtualKeySize.X))
 				{
-					TSharedRef<IKeyArea> KeyArea = InNode->UpdateKeyGrouping(Section);
+					TSharedRef<IKeyArea> KeyArea = InNode->GetKeyGrouping(Section);
 					HandleKeyArea(Visitor, KeyArea, Section, InNode);
 				}
 			}
diff --git a/Engine/Source/Editor/Sequencer/Public/ISequencer.h b/Engine/Source/Editor/Sequencer/Public/ISequencer.h
--- a/Engine/Source/Editor/Sequencer/Public/ISequencer.h
+++ b/Engine/Source/Editor/Sequencer/Public/ISequencer.h
@@ -381,7 +381,7 @@ public:
 	virtual void UpdatePlaybackRange() = 0;
 
 	/** Get all the keys for the current sequencer selection */
-	virtual void GetKeysFromSelection(TUniquePtr<ISequencerKeyCollection>& KeyCollection) = 0;
+	virtual void GetKeysFromSelection(TUniquePtr<ISequencerKeyCollection>& KeyCollection, float DuplicateThresoldTime) = 0;
 
 	virtual FSequencerSelection& GetSelection() = 0;
 	virtual FSequencerSelectionPreview& GetSelectionPreview() = 0;
diff --git a/Engine/Source/Editor/Sequencer/Public/ISequencerHotspot.h b/Engine/Source/Editor/Sequencer/Public/ISequencerHotspot.h
--- a/Engine/Source/Editor/Sequencer/Public/ISequencerHotspot.h
+++ b/Engine/Source/Editor/Sequencer/Public/ISequencerHotspot.h
@@ -31,6 +31,7 @@ struct ISequencerHotspot
 	virtual ESequencerHotspot GetType() const = 0;
 	virtual void UpdateOnHover(SSequencerTrackArea& InTrackArea, ISequencer& InSequencer) const = 0;
 	virtual TOptional<float> GetTime() const { return TOptional<float>(); }
+	virtual TOptional<float> GetOffsetTime() const { return TOptional<float>(); }
 	virtual TSharedPtr<ISequencerEditToolDragOperation> InitiateDrag(ISequencer&) { return nullptr; }
 	virtual bool PopulateContextMenu(FMenuBuilder& MenuBuilder, ISequencer& Sequencer, float MouseDownTime){ return false; }
 	virtual FCursorReply GetCursor() const { return FCursorReply::Unhandled(); }
diff --git a/Engine/Source/Editor/Sequencer/Public/ISequencerKeyCollection.h b/Engine/Source/Editor/Sequencer/Public/ISequencerKeyCollection.h
--- a/Engine/Source/Editor/Sequencer/Public/ISequencerKeyCollection.h
+++ b/Engine/Source/Editor/Sequencer/Public/ISequencerKeyCollection.h
@@ -6,12 +6,57 @@
 
 class UMovieSceneSection;
 class FSequencerDisplayNode;
+class IKeyArea;
 
 enum class EFindKeyDirection
 {
 	Forwards, Backwards
 };
 
+struct FSequencerKeyCollectionSignature
+{
+	FSequencerKeyCollectionSignature()
+	{}
+
+	/** Initialize this key collection from the specified nodes. Only gathers keys from those explicitly specified. */
+	SEQUENCER_API static FSequencerKeyCollectionSignature FromNodes(const TArray<FSequencerDisplayNode*>& InNodes, float InDuplicateThresholdTime);
+
+	/** Initialize this key collection from the specified nodes. Gathers keys from all child nodes. */
+	SEQUENCER_API static FSequencerKeyCollectionSignature FromNodesRecursive(const TArray<FSequencerDisplayNode*>& InNodes, float InDuplicateThresholdTime);
+
+	/** Initialize this key collection from the specified node and section index. */
+	SEQUENCER_API static FSequencerKeyCollectionSignature FromNodeRecursive(FSequencerDisplayNode& InNode, UMovieSceneSection* InSection, float InDuplicateThresholdTime);
+
+	/** Compare this signature for inequality with another */
+	SEQUENCER_API friend bool operator!=(const FSequencerKeyCollectionSignature& A, const FSequencerKeyCollectionSignature& B);
+
+	/** Compare this signature for equality with another */
+	SEQUENCER_API friend bool operator==(const FSequencerKeyCollectionSignature& A, const FSequencerKeyCollectionSignature& B);
+
+	/** Access the map of keyareas and signatures this signature was generated for */
+	const TMap<TSharedRef<IKeyArea>, FGuid>& GetKeyAreas() const
+	{
+		return KeyAreaToSignature;
+	}
+
+	/** Access duplicate threshold that this signature was generated for */
+	float GetDuplicateThreshold() const
+	{
+		return DuplicateThresholdTime;
+	}
+
+private:
+
+	/** Check whether this signature contains content that cannot be cached (such content causes this signature to never compare equal with another) */
+	bool HasUncachableContent() const;
+
+	/** The time at which proximal keys are considered duplicates */
+	float DuplicateThresholdTime;
+
+	/** Map of key areas to the section signature with with this signature was generated */
+	TMap<TSharedRef<IKeyArea>, FGuid> KeyAreaToSignature;
+};
+
 /**
  * A collection of keys gathered recursively from a particular node or nodes
  */
@@ -32,13 +77,21 @@ public:
 	virtual TOptional<float> FindFirstKeyInRange(const TRange<float>& InRange, EFindKeyDirection Direction) const = 0;
 
 public:
+	/**
+	 * Initialize this key collection from the specified nodes. Only gathers keys from those explicitly specified.
+	 * @return true if this collection was (re)initialized, false if it did not need updating
+	 */
+	virtual bool InitializeExplicit(const TArray<FSequencerDisplayNode*>& InNodes, float DuplicateThreshold = SMALL_NUMBER) = 0;
 
-	/** Initialize this key collection from the specified nodes. Only gathers keys from those explicitly specified. */
-	virtual void InitializeExplicit(const TArray<FSequencerDisplayNode*>& InNodes, float DuplicateThreshold = SMALL_NUMBER) = 0;
-
-	/** Initialize this key collection from the specified nodes. Gathers keys from all child nodes. */
-	virtual void InitializeRecursive(const TArray<FSequencerDisplayNode*>& InNodes, float DuplicateThreshold = SMALL_NUMBER) = 0;
-
-	/** Initialize this key collection from the specified node and section index. */
-	virtual void InitializeRecursive(FSequencerDisplayNode& InNode, UMovieSceneSection* InSection, float DuplicateThreshold = SMALL_NUMBER) = 0;
+	/**
+	 * Initialize this key collection from the specified nodes. Gathers keys from all child nodes.
+	 * @return true if this collection was (re)initialized, false if it did not need updating
+	 */
+	virtual bool InitializeRecursive(const TArray<FSequencerDisplayNode*>& InNodes, float DuplicateThreshold = SMALL_NUMBER) = 0;
+	
+	/**
+	 * Initialize this key collection from the specified node and section index.
+	 * @return true if this collection was (re)initialized, false if it did not need updating
+	 */
+	virtual bool InitializeRecursive(FSequencerDisplayNode& InNode, UMovieSceneSection* InSection, float DuplicateThreshold = SMALL_NUMBER) = 0;
 };
diff --git a/Engine/Source/Editor/Sequencer/Public/ISequencerSection.h b/Engine/Source/Editor/Sequencer/Public/ISequencerSection.h
--- a/Engine/Source/Editor/Sequencer/Public/ISequencerSection.h
+++ b/Engine/Source/Editor/Sequencer/Public/ISequencerSection.h
@@ -189,4 +189,12 @@ public:
 	 */
 	virtual void BeginDilateSection() {}
 	virtual void DilateSection(float DilationFactor, float Origin, TSet<FKeyHandle>& KeyHandles) { GetSectionObject()->DilateSection(DilationFactor, Origin, KeyHandles); }
+
+	/**
+	 * Slips the section by a specific factor
+	 *
+	 * @param SlipTime The amount to slip this section by
+	 */
+	virtual void BeginSlipSection() {}
+	virtual void SlipSection(float SlipTime) {}
 };
diff --git a/Engine/Source/Editor/Sequencer/Public/SequencerSettings.h b/Engine/Source/Editor/Sequencer/Public/SequencerSettings.h
--- a/Engine/Source/Editor/Sequencer/Public/SequencerSettings.h
+++ b/Engine/Source/Editor/Sequencer/Public/SequencerSettings.h
@@ -345,6 +345,8 @@ public:
 	/** Toggle whether to show pre and post roll in sequencer */
 	void SetShouldShowPrePostRoll(bool bInVisualizePreAndPostRoll);
 
+	uint32 GetTrajectoryPathCap() const { return TrajectoryPathCap; }
+
 	/** Gets the multicast delegate which is invoked whenevcer the bLockPlaybackToAudioClock setting is changed. */
 	FOnLockPlaybackToAudioClockChanged& GetOnLockPlaybackToAudioClockChanged() { return OnLockPlaybackToAudioClockChanged; }
 
@@ -529,6 +531,10 @@ protected:
 	UPROPERTY( config, EditAnywhere, Category=General )
 	bool bVisualizePreAndPostRoll;
 
+	/** Specifies the maximum number of keys to draw when rendering trajectories in viewports */
+	UPROPERTY(config, EditAnywhere, Category=General)
+	uint32 TrajectoryPathCap;
+
 	FOnLockPlaybackToAudioClockChanged OnLockPlaybackToAudioClockChanged;
 	FOnEvaluateSubSequencesInIsolationChanged OnEvaluateSubSequencesInIsolationChangedEvent;
 	FOnAllowEditsModeChanged OnAllowEditsModeChangedEvent;
diff --git a/Engine/Source/Programs/UnrealBuildTool/Platform/Android/UEDeployAndroid.cs b/Engine/Source/Programs/UnrealBuildTool/Platform/Android/UEDeployAndroid.cs
--- a/Engine/Source/Programs/UnrealBuildTool/Platform/Android/UEDeployAndroid.cs
+++ b/Engine/Source/Programs/UnrealBuildTool/Platform/Android/UEDeployAndroid.cs
@@ -1004,6 +1004,7 @@ namespace UnrealBuildTool
 						File.Copy("F:/NVPACK/android-kk-egl-t124-a32/Stripped_libNvPmApi.Core.so", UE4BuildPath + "/libs/" + NDKArch + "/libNvPmApi.Core.so", true);
 						File.Copy("F:/NVPACK/android-kk-egl-t124-a32/Stripped_libNvidia_gfx_debugger.so", UE4BuildPath + "/libs/" + NDKArch + "/libNvidia_gfx_debugger.so", true);
 					}
+					break;
 				*/
 				default:
 					break;
@@ -2459,8 +2460,10 @@ namespace UnrealBuildTool
 					"subst " + UnusedDriveLetter + " \"%CD%\"\n" +
 					"pushd " + UnusedDriveLetter + "\n" +
 					"call \"%ANTPATH%\\ant.orig.bat\" %ANT_CMD_LINE_ARGS%\n" +
+					"set ANTERROR=%ERRORLEVEL%\n" +
 					"popd\n" +
-					"subst " + UnusedDriveLetter + " /d\n";
+					"subst " + UnusedDriveLetter + " /d\n" +
+					"exit /b %ANTERROR%\n";
 
 			File.WriteAllText(AntBatFilename, AntBatText);
 		}
@@ -2524,8 +2527,10 @@ namespace UnrealBuildTool
 					"subst " + UnusedDriveLetter + " \"%CD%\"\n" +
 					"pushd " + UnusedDriveLetter + "\n" +
 					"call \"%GRADLEPATH%\\gradlew.bat\" %GRADLE_CMD_LINE_ARGS%\n" +
+					"set GRADLEERROR=%ERRORLEVEL%\n" +
 					"popd\n" +
-					"subst " + UnusedDriveLetter + " /d\n";
+					"subst " + UnusedDriveLetter + " /d\n" +
+					"exit /b %GRADLEERROR%\n";
 
 			File.WriteAllText(RunGradleBatFilename, RunGradleBatText);
 
diff --git a/Engine/Source/Runtime/ApplicationCore/Private/IOS/IOSPlatformApplicationMisc.cpp b/Engine/Source/Runtime/ApplicationCore/Private/IOS/IOSPlatformApplicationMisc.cpp
--- a/Engine/Source/Runtime/ApplicationCore/Private/IOS/IOSPlatformApplicationMisc.cpp
+++ b/Engine/Source/Runtime/ApplicationCore/Private/IOS/IOSPlatformApplicationMisc.cpp
@@ -1,6 +1,10 @@
 // Copyright 1998-2017 Epic Games, Inc. All Rights Reserved.
 
 #include "IOSPlatformApplicationMisc.h"
+
+#include "HAL/PlatformProcess.h"
+#include "Modules/ModuleManager.h"
+
 #include "IOSApplication.h"
 #include "IOSAppDelegate.h"
 #include "IOSView.h"
diff --git a/Engine/Source/Runtime/AudioMixer/Private/Components/SynthComponent.cpp b/Engine/Source/Runtime/AudioMixer/Private/Components/SynthComponent.cpp
--- a/Engine/Source/Runtime/AudioMixer/Private/Components/SynthComponent.cpp
+++ b/Engine/Source/Runtime/AudioMixer/Private/Components/SynthComponent.cpp
@@ -37,12 +37,12 @@ void USynthSound::Init(USynthComponent* InSynthComponent, int32 InNumChannels)
 
 bool USynthSound::OnGeneratePCMAudio(TArray<uint8>& OutAudio, int32 NumSamples)
 {
+	OutAudio.Reset();
+
 	if (bAudioMixer)
 	{
 		// If running with audio mixer, the output audio buffer will be in floats already
-		OutAudio.Reset();
 		OutAudio.AddZeroed(NumSamples * sizeof(float));
-
 		OwningSynthComponent->OnGeneratePCMAudio((float*)OutAudio.GetData(), NumSamples);
 	}
 	else
@@ -55,6 +55,7 @@ bool USynthSound::OnGeneratePCMAudio(TArray<uint8>& OutAudio, int32 NumSamples)
 		OwningSynthComponent->OnGeneratePCMAudio(FloatBufferDataPtr, NumSamples);
 
 		// Convert the float buffer to int16 data
+		OutAudio.AddZeroed(NumSamples * sizeof(int16));
 		int16* OutAudioBuffer = (int16*)OutAudio.GetData();
 		for (int32 i = 0; i < NumSamples; ++i)
 		{
diff --git a/Engine/Source/Runtime/CinematicCamera/Private/CineCameraComponent.cpp b/Engine/Source/Runtime/CinematicCamera/Private/CineCameraComponent.cpp
--- a/Engine/Source/Runtime/CinematicCamera/Private/CineCameraComponent.cpp
+++ b/Engine/Source/Runtime/CinematicCamera/Private/CineCameraComponent.cpp
@@ -165,6 +165,66 @@ float UCineCameraComponent::GetVerticalFieldOfView() const
 		: 0.f;
 }
 
+FString UCineCameraComponent::GetFilmbackPresetName() const
+{
+	TArray<FNamedFilmbackPreset> const& Presets = UCineCameraComponent::GetFilmbackPresets();
+	int32 const NumPresets = Presets.Num();
+	for (int32 PresetIdx = 0; PresetIdx < NumPresets; ++PresetIdx)
+	{
+		FNamedFilmbackPreset const& P = Presets[PresetIdx];
+		if (P.FilmbackSettings == FilmbackSettings)
+		{
+			return P.Name;
+		}
+	}
+
+	return FString();
+}
+
+void UCineCameraComponent::SetFilmbackPresetByName(const FString& InPresetName)
+{
+	TArray<FNamedFilmbackPreset> const& Presets = UCineCameraComponent::GetFilmbackPresets();
+	int32 const NumPresets = Presets.Num();
+	for (int32 PresetIdx = 0; PresetIdx < NumPresets; ++PresetIdx)
+	{
+		FNamedFilmbackPreset const& P = Presets[PresetIdx];
+		if (P.Name == InPresetName)
+		{
+			FilmbackSettings = P.FilmbackSettings;
+		}
+	}
+}
+
+FString UCineCameraComponent::GetLensPresetName() const
+{
+	TArray<FNamedLensPreset> const& Presets = UCineCameraComponent::GetLensPresets();
+	int32 const NumPresets = Presets.Num();
+	for (int32 PresetIdx = 0; PresetIdx < NumPresets; ++PresetIdx)
+	{
+		FNamedLensPreset const& P = Presets[PresetIdx];
+		if (P.LensSettings == LensSettings)
+		{
+			return P.Name;
+		}
+	}
+
+	return FString();
+}
+
+void UCineCameraComponent::SetLensPresetByName(const FString& InPresetName)
+{
+	TArray<FNamedLensPreset> const& Presets = UCineCameraComponent::GetLensPresets();
+	int32 const NumPresets = Presets.Num();
+	for (int32 PresetIdx = 0; PresetIdx < NumPresets; ++PresetIdx)
+	{
+		FNamedLensPreset const& P = Presets[PresetIdx];
+		if (P.Name == InPresetName)
+		{
+			LensSettings = P.LensSettings;
+		}
+	}
+}
+
 float UCineCameraComponent::GetWorldToMetersScale() const
 {
 	UWorld const* const World = GetWorld();
diff --git a/Engine/Source/Runtime/CinematicCamera/Public/CineCameraComponent.h b/Engine/Source/Runtime/CinematicCamera/Public/CineCameraComponent.h
--- a/Engine/Source/Runtime/CinematicCamera/Public/CineCameraComponent.h
+++ b/Engine/Source/Runtime/CinematicCamera/Public/CineCameraComponent.h
@@ -30,6 +30,12 @@ struct FCameraFilmbackSettings
 	/** Read-only. Computed from Sensor dimensions. */
 	UPROPERTY(Interp, VisibleAnywhere, BlueprintReadOnly, Category = "Filmback")
 	float SensorAspectRatio;
+
+	bool operator==(const FCameraFilmbackSettings& Other) const
+	{
+		return (SensorWidth == Other.SensorWidth)
+			&& (SensorHeight == Other.SensorHeight);
+	}
 };
 
 /** A named bundle of filmback settings used to implement filmback presets */
@@ -242,6 +248,22 @@ public:
 	UFUNCTION(BlueprintCallable, Category = "Cine Camera")
 	float GetVerticalFieldOfView() const;
 
+	/** Returns the filmback name of the camera with the current settings. */
+	UFUNCTION(BlueprintCallable, Category = "Cine Camera")
+	FString GetFilmbackPresetName() const;
+
+	/** Set the current preset settings by preset name. */
+	UFUNCTION(BlueprintCallable, Category = "Cine Camera")
+	void SetFilmbackPresetByName(const FString& InPresetName);
+
+	/** Returns the lens name of the camera with the current settings. */
+	UFUNCTION(BlueprintCallable, Category = "Cine Camera")
+	FString GetLensPresetName() const;
+
+	/** Set the current lens settings by preset name. */
+	UFUNCTION(BlueprintCallable, Category = "Cine Camera")
+	void SetLensPresetByName(const FString& InPresetName);
+
 	/** Returns a list of available filmback presets. */
 	static TArray<FNamedFilmbackPreset> const& GetFilmbackPresets();
 	
diff --git a/Engine/Source/Runtime/Core/Private/IOS/IOSPlatformOutputDevices.cpp b/Engine/Source/Runtime/Core/Private/IOS/IOSPlatformOutputDevices.cpp
--- a/Engine/Source/Runtime/Core/Private/IOS/IOSPlatformOutputDevices.cpp
+++ b/Engine/Source/Runtime/Core/Private/IOS/IOSPlatformOutputDevices.cpp
@@ -1,15 +1,8 @@
 // Copyright 1998-2017 Epic Games, Inc. All Rights Reserved.
 
-/*=============================================================================
-	IOSPlatformOutputDevices.mm: iOS implementations of OutputDevices functions
-=============================================================================*/
-
 #include "IOSPlatformOutputDevices.h"
 #include "Misc/OutputDeviceFile.h"
 
-//////////////////////////////////
-// FIOSPlatformOutputDevices
-//////////////////////////////////
 
 class FOutputDevice*	FIOSPlatformOutputDevices::GetLog()
 {
diff --git a/Engine/Source/Runtime/Core/Public/Containers/Queue.h b/Engine/Source/Runtime/Core/Public/Containers/Queue.h
--- a/Engine/Source/Runtime/Core/Public/Containers/Queue.h
+++ b/Engine/Source/Runtime/Core/Public/Containers/Queue.h
@@ -66,7 +66,7 @@ public:
 	 * @param OutValue Will hold the returned value.
 	 * @return true if a value was returned, false if the queue was empty.
 	 * @note To be called only from consumer thread.
-	 * @see Empty, Enqueue, IsEmpty, Peek
+	 * @see Empty, Enqueue, IsEmpty, Peek, Pop
 	 */
 	bool Dequeue(ItemType& OutItem)
 	{
@@ -92,12 +92,11 @@ public:
 	 * Empty the queue, discarding all items.
 	 *
 	 * @note To be called only from consumer thread.
-	 * @see Dequeue, IsEmpty, Peek
+	 * @see Dequeue, IsEmpty, Peek, Pop
 	 */
 	void Empty()
 	{
-		ItemType DummyItem;
-		while (Dequeue(DummyItem));
+		while (Pop());
 	}
 
 	/**
@@ -106,7 +105,7 @@ public:
 	 * @param Item The item to add.
 	 * @return true if the item was added, false otherwise.
 	 * @note To be called only from producer thread(s).
-	 * @see Dequeue
+	 * @see Dequeue, Pop
 	 */
 	bool Enqueue(const ItemType& Item)
 	{
@@ -143,7 +142,7 @@ public:
 	 * @param Item The item to add.
 	 * @return true if the item was added, false otherwise.
 	 * @note To be called only from producer thread(s).
-	 * @see Dequeue
+	 * @see Dequeue, Pop
 	 */
 	bool Enqueue(ItemType&& Item)
 	{
@@ -179,7 +178,7 @@ public:
 	 *
 	 * @return true if the queue is empty, false otherwise.
 	 * @note To be called only from consumer thread.
-	 * @see Dequeue, Empty, Peek
+	 * @see Dequeue, Empty, Peek, Pop
 	 */
 	bool IsEmpty() const
 	{
@@ -192,7 +191,7 @@ public:
 	 * @param OutItem Will hold the peeked at item.
 	 * @return true if an item was returned, false if the queue was empty.
 	 * @note To be called only from consumer thread.
-	 * @see Dequeue, Empty, IsEmpty
+	 * @see Dequeue, Empty, IsEmpty, Pop
 	 */
 	bool Peek(ItemType& OutItem) const
 	{
@@ -206,6 +205,32 @@ public:
 		return true;
 	}
 
+	/**
+	 * Removes the item from the tail of the queue.
+	 *
+	 * @return true if a value was removed, false if the queue was empty.
+	 * @note To be called only from consumer thread.
+	 * @see Dequeue, Empty, Enqueue, IsEmpty, Peek
+	 */
+	bool Pop()
+	{
+		TNode* Popped = Tail->NextNode;
+
+		if (Popped == nullptr)
+		{
+			return false;
+		}
+		
+		TSAN_AFTER(&Tail->NextNode);
+
+		TNode* OldTail = Tail;
+		Tail = Popped;
+		Tail->Item = ItemType();
+		delete OldTail;
+
+		return true;
+	}
+
 private:
 
 	/** Structure for the internal linked list. */
diff --git a/Engine/Source/Runtime/CoreUObject/Public/UObject/StructOnScope.h b/Engine/Source/Runtime/CoreUObject/Public/UObject/StructOnScope.h
--- a/Engine/Source/Runtime/CoreUObject/Public/UObject/StructOnScope.h
+++ b/Engine/Source/Runtime/CoreUObject/Public/UObject/StructOnScope.h
@@ -6,6 +6,7 @@
 #include "UObject/Class.h"
 #include "UObject/WeakObjectPtr.h"
 #include "UObject/Package.h"
+#include "Casts.h"
 
 class FStructOnScope
 {
diff --git a/Engine/Source/Runtime/Engine/Engine.Build.cs b/Engine/Source/Runtime/Engine/Engine.Build.cs
--- a/Engine/Source/Runtime/Engine/Engine.Build.cs
+++ b/Engine/Source/Runtime/Engine/Engine.Build.cs
@@ -168,7 +168,6 @@ public class Engine : ModuleRules
 		DynamicallyLoadedModuleNames.AddRange(
 			new string[]
 			{
-				"ActorSequence",
 				"MovieScene",
 				"MovieSceneCapture",
 				"MovieSceneTracks",
diff --git a/Engine/Source/Runtime/Engine/Private/Materials/MaterialUniformExpressions.cpp b/Engine/Source/Runtime/Engine/Private/Materials/MaterialUniformExpressions.cpp
--- a/Engine/Source/Runtime/Engine/Private/Materials/MaterialUniformExpressions.cpp
+++ b/Engine/Source/Runtime/Engine/Private/Materials/MaterialUniformExpressions.cpp
@@ -10,6 +10,8 @@
 #include "Materials/MaterialInstanceSupport.h"
 #include "Materials/MaterialParameterCollection.h"
 #include "ExternalTexture.h"
+#include "MessageLog.h"
+#include "UObjectToken.h"
 
 TLinkedList<FMaterialUniformExpressionType*>*& FMaterialUniformExpressionType::GetTypeList()
 {
@@ -401,8 +403,15 @@ FUniformBufferRHIRef FUniformExpressionSet::CreateUniformBuffer(const FMaterialR
 				// Do not allow external textures to be applied to normal texture samplers
 				if (Value->GetMaterialType() == MCT_TextureExternal)
 				{
-					ensureMsgf(false, TEXT("External texture applied to a Texture2D sampler. This may work by chance on some platforms but is not portable. Please change sampler type to 'External'. Texture '%s' on parameter '%s' (slot %i) in material '%s'"),
-						*Value->GetPathName(), TextureParameter ? *TextureParameter->GetParameterName().ToString() : TEXT("non-parameter"), ExpressionIndex, *MaterialRenderContext.Material.GetFriendlyName());
+					FText MessageText = FText::Format(
+						NSLOCTEXT("MaterialExpressions", "IncompatibleExternalTexture", " applied to a non-external Texture2D sampler. This may work by chance on some platforms but is not portable. Please change sampler type to 'External'. Parameter '{0}' (slot {1}) in material '{2}'"),
+						FText::FromName(TextureParameter ? TextureParameter->GetParameterName() : FName()),
+						ExpressionIndex,
+						FText::FromString(*MaterialRenderContext.Material.GetFriendlyName()));
+
+					FMessageLog("PIE").Warning()
+						->AddToken(FUObjectToken::Create(Value))
+						->AddToken(FTextToken::Create(MessageText));
 				}
 			}
 
@@ -746,7 +755,7 @@ void FMaterialUniformExpressionScalarParameter::GetGameThreadNumberValue(const U
 
 namespace
 {
-	void SerializeOptional(FArchive Ar, TOptional<FName> OptionalName)
+	void SerializeOptional(FArchive& Ar, TOptional<FName>& OptionalName)
 	{
 		bool bIsSet = OptionalName.IsSet();
 		Ar << bIsSet;
diff --git a/Engine/Source/Runtime/Engine/Private/ShaderDerivedDataVersion.h b/Engine/Source/Runtime/Engine/Private/ShaderDerivedDataVersion.h
--- a/Engine/Source/Runtime/Engine/Private/ShaderDerivedDataVersion.h
+++ b/Engine/Source/Runtime/Engine/Private/ShaderDerivedDataVersion.h
@@ -12,4 +12,4 @@
 // guid as version
 
 #define GLOBALSHADERMAP_DERIVEDDATA_VER			TEXT("605300842A8542F8A2AD72835F18F142")
-#define MATERIALSHADERMAP_DERIVEDDATA_VER		TEXT("A6F6D46EF9844AD39968049DB1F003F5")
+#define MATERIALSHADERMAP_DERIVEDDATA_VER		TEXT("E426546F565C4650A2D1BC1D2EF145FD")
diff --git a/Engine/Source/Runtime/Media/Private/MediaModule.cpp b/Engine/Source/Runtime/Media/Private/MediaModule.cpp
--- a/Engine/Source/Runtime/Media/Private/MediaModule.cpp
+++ b/Engine/Source/Runtime/Media/Private/MediaModule.cpp
@@ -43,7 +43,7 @@ public:
 	{
 		for (IMediaPlayerFactory* Factory : PlayerFactories)
 		{
-			if (Factory->GetName() == FactoryName)
+			if (Factory->GetPlayerName() == FactoryName)
 			{
 				return Factory;
 			}
diff --git a/Engine/Source/Runtime/Media/Public/IMediaPlayer.h b/Engine/Source/Runtime/Media/Public/IMediaPlayer.h
--- a/Engine/Source/Runtime/Media/Public/IMediaPlayer.h
+++ b/Engine/Source/Runtime/Media/Public/IMediaPlayer.h
@@ -4,6 +4,7 @@
 
 #include "Containers/UnrealString.h"
 #include "Internationalization/Text.h"
+#include "Misc/Paths.h"
 #include "Misc/Timespan.h"
 
 class FArchive;
@@ -75,9 +76,10 @@ public:
 	/**
 	 * Get the name of this player.
 	 *
-	 * @return Player name, i.e. 'AndroidMedia' or 'WmfMedia'.
+	 * @return Media player name, i.e. 'AndroidMedia' or 'WmfMedia'.
+	 * @see GetMediaName
 	 */
-	virtual FName GetName() const = 0;
+	virtual FName GetPlayerName() const = 0;
 
 	/**
 	 * Get the player's sample queue.
@@ -167,6 +169,29 @@ public:
 
 	//~ The following methods are optional
 
+	/**
+	 * Get the human readable name of the currently loaded media source.
+	 *
+	 * Depending on the type of media source, this might be the name of a file,
+	 * the display name of a capture device, or some other identifying string.
+	 * If the player does not provide a specialized implementation for this
+	 * method, the media name will be derived from the current media URL.
+	 *
+	 * @return Media source name, or empty text if no media is opened
+	 * @see GetPlayerName, GetUrl
+	 */
+	virtual FText GetMediaName() const
+	{
+		const FString Url = GetUrl();
+
+		if (Url.IsEmpty())
+		{
+			return FText::GetEmpty();
+		}
+
+		return FText::FromString(FPaths::GetBaseFilename(Url));
+	}
+
 	/**
 	 * Set the player's globally unique identifier.
 	 *
@@ -197,9 +222,10 @@ public:
 	 * samples before they are rendered on textures or audio components.
 	 *
 	 * @param DeltaTime Time since last tick.
+	 * @param Timecode The current media time code.
 	 * @see TickAudio, TickInput
 	 */
-	virtual void TickFetch(FTimespan DeltaTime)
+	virtual void TickFetch(FTimespan DeltaTime, FTimespan Timecode)
 	{
 		// override in child class if needed
 	}
@@ -211,9 +237,10 @@ public:
 	 * Engine is being ticked, or to initiate the processing of input samples.
 	 *
 	 * @param DeltaTime Time since last tick.
+	 * @param Timecode The current media time code.
 	 * @see TickAudio, TickFetch
 	 */
-	virtual void TickInput(FTimespan DeltaTime)
+	virtual void TickInput(FTimespan DeltaTime, FTimespan Timecode)
 	{
 		// override in child class if needed
 	}
diff --git a/Engine/Source/Runtime/Media/Public/IMediaPlayerFactory.h b/Engine/Source/Runtime/Media/Public/IMediaPlayerFactory.h
--- a/Engine/Source/Runtime/Media/Public/IMediaPlayerFactory.h
+++ b/Engine/Source/Runtime/Media/Public/IMediaPlayerFactory.h
@@ -97,10 +97,10 @@ public:
 	/**
 	 * Get the unique name of the media player.
 	 *
-	 * @return Player name.
+	 * @return Media player name, i.e. 'AndroidMedia' or 'WmfMedia'.
 	 * @see GetDisplayName
 	 */
-	virtual FName GetName() const = 0;
+	virtual FName GetPlayerName() const = 0;
 
 	/**
 	 * Get the names of platforms that the media player supports.
diff --git a/Engine/Source/Runtime/MediaAssets/Private/Assets/MediaPlayer.cpp b/Engine/Source/Runtime/MediaAssets/Private/Assets/MediaPlayer.cpp
--- a/Engine/Source/Runtime/MediaAssets/Private/Assets/MediaPlayer.cpp
+++ b/Engine/Source/Runtime/MediaAssets/Private/Assets/MediaPlayer.cpp
@@ -135,6 +135,12 @@ float UMediaPlayer::GetHorizontalFieldOfView() const
 }
 
 
+FText UMediaPlayer::GetMediaName() const
+{
+	return PlayerFacade->GetMediaName();
+}
+
+
 int32 UMediaPlayer::GetNumTracks(EMediaPlayerTrack TrackType) const
 {
 	return PlayerFacade->GetNumTracks((EMediaTrackType)TrackType);
@@ -604,9 +610,9 @@ void UMediaPlayer::PostDuplicate(bool bDuplicateForPIE)
 }
 
 
-void UMediaPlayer::PostLoad()
+void UMediaPlayer::PostInitProperties()
 {
-	Super::PostLoad();
+	Super::PostInitProperties();
 
 	PlayerFacade->SetGuid(PlayerGuid);
 
@@ -697,5 +703,9 @@ void UMediaPlayer::HandlePlayerMediaEvent(EMediaEvent Event)
 	case EMediaEvent::SeekCompleted:
 		OnSeekCompleted.Broadcast();
 		break;
+
+	case EMediaEvent::TracksChanged:
+		OnTracksChanged.Broadcast();
+		break;
 	}
 }
diff --git a/Engine/Source/Runtime/MediaAssets/Private/Misc/MediaBlueprintFunctionLibrary.cpp b/Engine/Source/Runtime/MediaAssets/Private/Misc/MediaBlueprintFunctionLibrary.cpp
--- a/Engine/Source/Runtime/MediaAssets/Private/Misc/MediaBlueprintFunctionLibrary.cpp
+++ b/Engine/Source/Runtime/MediaAssets/Private/Misc/MediaBlueprintFunctionLibrary.cpp
@@ -17,7 +17,7 @@ void UMediaBlueprintFunctionLibrary::EnumerateAudioCaptureDevices(int32 Filter,
 
 	for (const auto& DeviceInfo : DeviceInfos)
 	{
-		if (((DeviceInfo.Type == EMediaCaptureDeviceType::Audio) && EnumHasAnyFlags(FilterEnum, EMediaAudioCaptureDeviceFilter::Unspecified)) ||
+		if (((DeviceInfo.Type == EMediaCaptureDeviceType::Audio) && EnumHasAnyFlags(FilterEnum, EMediaAudioCaptureDeviceFilter::Unknown)) ||
 			((DeviceInfo.Type == EMediaCaptureDeviceType::AudioCard) && EnumHasAnyFlags(FilterEnum, EMediaAudioCaptureDeviceFilter::Card)) ||
 			((DeviceInfo.Type == EMediaCaptureDeviceType::AudioSoftware) && EnumHasAnyFlags(FilterEnum, EMediaAudioCaptureDeviceFilter::Software)) ||
 			((DeviceInfo.Type == EMediaCaptureDeviceType::Microphone) && EnumHasAnyFlags(FilterEnum, EMediaAudioCaptureDeviceFilter::Microphone)))
@@ -38,7 +38,7 @@ void UMediaBlueprintFunctionLibrary::EnumerateVideoCaptureDevices(int32 Filter,
 	for (const auto& DeviceInfo : DeviceInfos)
 	{
 		if (((DeviceInfo.Type == EMediaCaptureDeviceType::DepthSensor) && EnumHasAnyFlags(FilterEnum, EMediaVideoCaptureDeviceFilter::Webcam)) ||
-			((DeviceInfo.Type == EMediaCaptureDeviceType::Video) && EnumHasAnyFlags(FilterEnum, EMediaVideoCaptureDeviceFilter::Unspecified)) ||
+			((DeviceInfo.Type == EMediaCaptureDeviceType::Video) && EnumHasAnyFlags(FilterEnum, EMediaVideoCaptureDeviceFilter::Unknown)) ||
 			((DeviceInfo.Type == EMediaCaptureDeviceType::VideoCard) && EnumHasAnyFlags(FilterEnum, EMediaVideoCaptureDeviceFilter::Card)) ||
 			((DeviceInfo.Type == EMediaCaptureDeviceType::VideoSoftware) && EnumHasAnyFlags(FilterEnum, EMediaVideoCaptureDeviceFilter::Software)) ||
 			((DeviceInfo.Type == EMediaCaptureDeviceType::Webcam) && EnumHasAnyFlags(FilterEnum, EMediaVideoCaptureDeviceFilter::Webcam)) ||
@@ -61,7 +61,7 @@ void UMediaBlueprintFunctionLibrary::EnumerateWebcamCaptureDevices(int32 Filter,
 	for (const auto& DeviceInfo : DeviceInfos)
 	{
 		if (((DeviceInfo.Type == EMediaCaptureDeviceType::DepthSensor) && EnumHasAnyFlags(FilterEnum, EMediaWebcamCaptureDeviceFilter::DepthSensor)) ||
-			((DeviceInfo.Type == EMediaCaptureDeviceType::Webcam) && EnumHasAnyFlags(FilterEnum, EMediaWebcamCaptureDeviceFilter::Unspecified)) ||
+			((DeviceInfo.Type == EMediaCaptureDeviceType::Webcam) && EnumHasAnyFlags(FilterEnum, EMediaWebcamCaptureDeviceFilter::Unknown)) ||
 			((DeviceInfo.Type == EMediaCaptureDeviceType::WebcamFront) && EnumHasAnyFlags(FilterEnum, EMediaWebcamCaptureDeviceFilter::Front)) ||
 			((DeviceInfo.Type == EMediaCaptureDeviceType::WebcamRear) && EnumHasAnyFlags(FilterEnum, EMediaWebcamCaptureDeviceFilter::Rear)))
 		{
diff --git a/Engine/Source/Runtime/MediaAssets/Private/Misc/MediaBlueprintFunctionLibrary.h b/Engine/Source/Runtime/MediaAssets/Private/Misc/MediaBlueprintFunctionLibrary.h
--- a/Engine/Source/Runtime/MediaAssets/Private/Misc/MediaBlueprintFunctionLibrary.h
+++ b/Engine/Source/Runtime/MediaAssets/Private/Misc/MediaBlueprintFunctionLibrary.h
@@ -9,6 +9,9 @@
 #include "MediaBlueprintFunctionLibrary.generated.h"
 
 
+/**
+ * Filter flags for the EnumerateAudioCaptureDevices BP function.
+ */
 UENUM(BlueprintType, meta=(BitFlags, UseEnumValuesAsMaskValuesInEditor = "true"))
 enum class EMediaAudioCaptureDeviceFilter : uint8
 {
@@ -21,13 +24,16 @@ enum class EMediaAudioCaptureDeviceFilter : uint8
 	/** Software device. */
 	Software = 0x4,
 
-	/** Unspecified audio capture device. */
-	Unspecified = 0x8
+	/** Unknown audio capture device types. */
+	Unknown = 0x8
 };
 
 ENUM_CLASS_FLAGS(EMediaAudioCaptureDeviceFilter)
 
 
+/**
+ * Filter flags for the EnumerateVideoCaptureDevices BP function.
+ */
 UENUM(BlueprintType, meta=(BitFlags, UseEnumValuesAsMaskValuesInEditor = "true"))
 enum class EMediaVideoCaptureDeviceFilter : uint8
 {
@@ -37,8 +43,8 @@ enum class EMediaVideoCaptureDeviceFilter : uint8
 	/** Software video capture device. */
 	Software = 0x2,
 
-	/** Unspecified video capture device. */
-	Unspecified = 0x4,
+	/** Unknown video capture device types. */
+	Unknown = 0x4,
 
 	/** Web cam. */
 	Webcam = 0x8
@@ -47,6 +53,9 @@ enum class EMediaVideoCaptureDeviceFilter : uint8
 ENUM_CLASS_FLAGS(EMediaVideoCaptureDeviceFilter)
 
 
+/**
+ * Filter flags for the EnumerateWebcamCaptureDevices BP function.
+ */
 UENUM(BlueprintType, meta=(BitFlags, UseEnumValuesAsMaskValuesInEditor = "true"))
 enum class EMediaWebcamCaptureDeviceFilter : uint8
 {
@@ -59,8 +68,8 @@ enum class EMediaWebcamCaptureDeviceFilter : uint8
 	/** Rear facing web cam. */
 	Rear = 0x4,
 
-	/** Unspecified web cam type. */
-	Unspecified = 0x8
+	/** Unknown web cam types. */
+	Unknown = 0x8
 };
 
 ENUM_CLASS_FLAGS(EMediaWebcamCaptureDeviceFilter)
diff --git a/Engine/Source/Runtime/MediaAssets/Private/Misc/MediaTextureResource.cpp b/Engine/Source/Runtime/MediaAssets/Private/Misc/MediaTextureResource.cpp
--- a/Engine/Source/Runtime/MediaAssets/Private/Misc/MediaTextureResource.cpp
+++ b/Engine/Source/Runtime/MediaAssets/Private/Misc/MediaTextureResource.cpp
@@ -177,7 +177,7 @@ void FMediaTextureResource::Render(const FRenderParams& Params)
 		if (Sample->GetOutputDim().GetMin() <= 0)
 		{
 			#if MEDIATEXTURERESOURCE_TRACE_RENDER
-				UE_LOG(LogMediaAssets, VeryVerbose, TEXT("TextureResource %llx: Corrupt sample with time %s at time %s"), this, *Sample->GetTime().ToString(), *Params.Time.ToString());
+				UE_LOG(LogMediaAssets, VeryVerbose, TEXT("TextureResource %p: Corrupt sample with time %s at time %s"), this, *Sample->GetTime().ToString(), *Params.Time.ToString());
 			#endif
 
 			ClearTexture(FLinearColor::Red, Params.SrgbOutput); // mark corrupt sample
@@ -185,7 +185,7 @@ void FMediaTextureResource::Render(const FRenderParams& Params)
 		else if (MediaTextureResource::RequiresConversion(Sample, Params.SrgbOutput))
 		{
 			#if MEDIATEXTURERESOURCE_TRACE_RENDER
-				UE_LOG(LogMediaAssets, VeryVerbose, TEXT("TextureResource %llx: Converting sample with time %s at time %s"), this, *Sample->GetTime().ToString(), *Params.Time.ToString());
+				UE_LOG(LogMediaAssets, VeryVerbose, TEXT("TextureResource %p: Converting sample with time %s at time %s"), this, *Sample->GetTime().ToString(), *Params.Time.ToString());
 			#endif
 
 			ConvertSample(Sample, Params.ClearColor, Params.SrgbOutput);
@@ -193,7 +193,7 @@ void FMediaTextureResource::Render(const FRenderParams& Params)
 		else
 		{
 			#if MEDIATEXTURERESOURCE_TRACE_RENDER
-				UE_LOG(LogMediaAssets, VeryVerbose, TEXT("TextureResource %llx: Copying sample with time %s at time %s"), this, *Sample->GetTime().ToString(), *Params.Time.ToString());
+				UE_LOG(LogMediaAssets, VeryVerbose, TEXT("TextureResource %p: Copying sample with time %s at time %s"), this, *Sample->GetTime().ToString(), *Params.Time.ToString());
 			#endif
 
 			CopySample(Sample, Params.ClearColor, Params.SrgbOutput);
@@ -208,7 +208,7 @@ void FMediaTextureResource::Render(const FRenderParams& Params)
 	else if (!Cleared)
 	{
 		#if MEDIATEXTURERESOURCE_TRACE_RENDER
-			UE_LOG(LogMediaAssets, VeryVerbose, TEXT("TextureResource %llx: Clearing texture at time %s"), this, *Params.Time.ToString());
+			UE_LOG(LogMediaAssets, VeryVerbose, TEXT("TextureResource %p: Clearing texture at time %s"), this, *Params.Time.ToString());
 		#endif
 
 		ClearTexture(Params.ClearColor, Params.SrgbOutput);
diff --git a/Engine/Source/Runtime/MediaAssets/Public/MediaPlayer.h b/Engine/Source/Runtime/MediaAssets/Public/MediaPlayer.h
--- a/Engine/Source/Runtime/MediaAssets/Public/MediaPlayer.h
+++ b/Engine/Source/Runtime/MediaAssets/Public/MediaPlayer.h
@@ -179,6 +179,15 @@ public:
 	UFUNCTION(BlueprintCallable, Category="Media|MediaPlayer")
 	float GetHorizontalFieldOfView() const;
 
+	/**
+	 * Get the human readable name of the currently loaded media source.
+	 *
+	 * @return Media source name, or empty text if no media is opened
+	 * @see GetPlayerName, GetUrl
+	 */
+	UFUNCTION(BlueprintCallable, Category="Media|MediaPlayer")
+	virtual FText GetMediaName() const;
+
 	/**
 	 * Get the number of tracks of the given type.
 	 *
@@ -204,6 +213,7 @@ public:
 	 * Get the name of the current native media player.
 	 *
 	 * @return Player name, or NAME_None if not available.
+	 * @see GetMediaName
 	 */
 	UFUNCTION(BlueprintCallable, Category="Media|MediaPlayer")
 	FName GetPlayerName() const;
@@ -752,7 +762,7 @@ public:
 	 * synchronously or asynchronously, this event may be executed before or
 	 * after the call to OpenSource / OpenUrl returns.
 	 *
-	 * @see OnMediaOpenFailed
+	 * @see OnMediaOpenFailed, OnTracksChanged
 	 */
 	UPROPERTY(BlueprintAssignable, Category="Media|MediaPlayer")
 	FOnMediaPlayerMediaOpened OnMediaOpened;
@@ -790,11 +800,19 @@ public:
 	 *
 	 * Depending on whether the underlying player implementation performs seeks
 	 * synchronously or asynchronously, this event may be executed before or
-	 * after the call to Seek.
+	 * after the call to Seek returns.
 	 */
-	UPROPERTY(BlueprintAssignable, Category = "Media|MediaPlayer")
+	UPROPERTY(BlueprintAssignable, Category="Media|MediaPlayer")
 	FOnMediaPlayerMediaEvent OnSeekCompleted;
 
+	/**
+	 * A delegate that is invoked when the media track collection changed.
+	 *
+	 * @see OnMediaOpened
+	 */
+	UPROPERTY(BlueprintAssignable, Category="Media|MediaPlayer")
+	FOnMediaPlayerMediaEvent OnTracksChanged;
+
 public:
 
 	/**
@@ -861,7 +879,7 @@ public:
 	virtual bool CanBeInCluster() const override;
 	virtual FString GetDesc() override;
 	virtual void PostDuplicate(bool bDuplicateForPIE) override;
-	virtual void PostLoad() override;
+	virtual void PostInitProperties() override;
 
 #if WITH_EDITOR
 	virtual void PostEditChangeProperty(FPropertyChangedEvent& PropertyChangedEvent) override;
diff --git a/Engine/Source/Runtime/MediaUtils/Private/MediaAudioResampler.cpp b/Engine/Source/Runtime/MediaUtils/Private/MediaAudioResampler.cpp
--- a/Engine/Source/Runtime/MediaUtils/Private/MediaAudioResampler.cpp
+++ b/Engine/Source/Runtime/MediaUtils/Private/MediaAudioResampler.cpp
@@ -289,11 +289,13 @@ namespace MediaAudioResampler
 			{
 				*Dest = 0.0f;
 
-				for (uint32 SrcIndex = 0; SrcIndex < SrcChannels; ++Coeff, ++Src, ++SrcIndex)
+				for (uint32 SrcIndex = 0; SrcIndex < SrcChannels; ++Coeff, ++SrcIndex)
 				{
-					*Dest += *Src * *Coeff;
+					*Dest += Src[SrcIndex] * *Coeff;
 				}
 			}
+
+			Src += SrcChannels;
 		}
 
 		return true;
@@ -558,7 +560,7 @@ bool FMediaAudioResampler::SetInput(const TSharedPtr<IMediaAudioSample, ESPMode:
 		case EMediaAudioSampleFormat::Int16:
 			while (FloatData < FloatDataEnd)
 			{
-				*FloatData++ = (float)(*(const int16*)Buffer / 32767.0f);
+				*FloatData++ = (float)(*(const int16*)Buffer / 32768.0f);
 				Buffer += sizeof(int16);
 			}
 			break;
@@ -566,7 +568,7 @@ bool FMediaAudioResampler::SetInput(const TSharedPtr<IMediaAudioSample, ESPMode:
 		case EMediaAudioSampleFormat::Int32:
 			while (FloatData < FloatDataEnd)
 			{
-				*FloatData++ = (float)*(const int32*)Buffer / 4294967295.0f;
+				*FloatData++ = (float)*(const int32*)Buffer / 2147483648.0f;
 				Buffer += sizeof(int32);
 			}
 			break;
@@ -574,7 +576,7 @@ bool FMediaAudioResampler::SetInput(const TSharedPtr<IMediaAudioSample, ESPMode:
 		case EMediaAudioSampleFormat::Int8:
 			while (FloatData < FloatDataEnd)
 			{
-				*FloatData++ = (float)*(const int8*)Buffer / 255.0f;
+				*FloatData++ = (float)*(const int8*)Buffer / 128.0f;
 				Buffer += sizeof(int8);
 			}
 			break;
@@ -591,7 +593,7 @@ bool FMediaAudioResampler::SetInput(const TSharedPtr<IMediaAudioSample, ESPMode:
 		else
 		{
 			// down-mix channels
-			Input.AddUninitialized(NumSamples);
+			Input.AddDefaulted(NumFrames * OutputChannels);
 
 			if (!MediaAudioResampler::Downmix(FloatSamples.GetData(), NumChannels, NumFrames, Input.GetData(), OutputChannels))
 			{
diff --git a/Engine/Source/Runtime/MediaUtils/Private/MediaPlayerFacade.cpp b/Engine/Source/Runtime/MediaUtils/Private/MediaPlayerFacade.cpp
--- a/Engine/Source/Runtime/MediaUtils/Private/MediaPlayerFacade.cpp
+++ b/Engine/Source/Runtime/MediaUtils/Private/MediaPlayerFacade.cpp
@@ -18,9 +18,13 @@
 #include "Misc/ScopeLock.h"
 #include "Modules/ModuleManager.h"
 
+#include "MediaHelpers.h"
 #include "MediaSampleCache.h"
 
 
+#define MEDIAPLAYERFACADE_TRACE_SINKOVERFLOWS 0
+
+
 /* Local helpers
 *****************************************************************************/
 
@@ -203,6 +207,12 @@ FString FMediaPlayerFacade::GetInfo() const
 }
 
 
+FText FMediaPlayerFacade::GetMediaName() const
+{
+	return Player.IsValid() ? Player->GetMediaName() : FText::GetEmpty();
+}
+
+
 int32 FMediaPlayerFacade::GetNumTracks(EMediaTrackType TrackType) const
 {
 	return Player.IsValid() ? Player->GetTracks().GetNumTracks(TrackType) : 0;
@@ -217,7 +227,7 @@ int32 FMediaPlayerFacade::GetNumTrackFormats(EMediaTrackType TrackType, int32 Tr
 
 FName FMediaPlayerFacade::GetPlayerName() const
 {
-	return Player.IsValid() ? Player->GetName() : NAME_None;
+	return Player.IsValid() ? Player->GetPlayerName() : NAME_None;
 }
 
 
@@ -540,7 +550,7 @@ bool FMediaPlayerFacade::SupportsRate(float Rate, bool Unthinned) const
 
 void FMediaPlayerFacade::FlushSinks()
 {
-	UE_LOG(LogMediaUtils, Verbose, TEXT("PlayerFacade %llx: Flushing sinks"), this);
+	UE_LOG(LogMediaUtils, Verbose, TEXT("PlayerFacade %p: Flushing sinks"), this);
 
 	AudioSampleSinks.Flush();
 	CaptionSampleSinks.Flush();
@@ -589,7 +599,7 @@ TSharedPtr<IMediaPlayer, ESPMode::ThreadSafe> FMediaPlayerFacade::GetPlayerForUr
 	}
 
 	// reuse existing player if requested
-	if (Player.IsValid() && (PlayerName == Player->GetName()))
+	if (Player.IsValid() && (PlayerName == Player->GetPlayerName()))
 	{
 		return Player;
 	}
@@ -628,7 +638,7 @@ TSharedPtr<IMediaPlayer, ESPMode::ThreadSafe> FMediaPlayerFacade::GetPlayerForUr
 	// try to reuse existing player
 	if (Player.IsValid())
 	{
-		IMediaPlayerFactory* Factory = MediaModule->GetPlayerFactory(Player->GetName());
+		IMediaPlayerFactory* Factory = MediaModule->GetPlayerFactory(Player->GetPlayerName());
 
 		if ((Factory != nullptr) && Factory->CanPlayUrl(Url, Options))
 		{
@@ -665,11 +675,11 @@ TSharedPtr<IMediaPlayer, ESPMode::ThreadSafe> FMediaPlayerFacade::GetPlayerForUr
 		{
 			if (Factory->SupportsPlatform(RunningPlatformName))
 			{
-				UE_LOG(LogMediaUtils, Log, TEXT("| %s (URI scheme or file extension not supported)"), *Factory->GetName().ToString());
+				UE_LOG(LogMediaUtils, Log, TEXT("| %s (URI scheme or file extension not supported)"), *Factory->GetPlayerName().ToString());
 			}
 			else
 			{
-				UE_LOG(LogMediaUtils, Log, TEXT("| %s (only available on %s, but not on %s)"), *Factory->GetName().ToString(), *FString::Join(Factory->GetSupportedPlatforms(), TEXT(", ")), *RunningPlatformName);
+				UE_LOG(LogMediaUtils, Log, TEXT("| %s (only available on %s, but not on %s)"), *Factory->GetPlayerName().ToString(), *FString::Join(Factory->GetSupportedPlatforms(), TEXT(", ")), *RunningPlatformName);
 			}	
 		}
 	}
@@ -715,11 +725,11 @@ void FMediaPlayerFacade::ProcessEvent(EMediaEvent Event)
 
 		if (MediaInfo.IsEmpty())
 		{
-			UE_LOG(LogMediaUtils, Verbose, TEXT("PlayerFacade %llx: Media Info: n/a"), this);
+			UE_LOG(LogMediaUtils, Verbose, TEXT("PlayerFacade %p: Media Info: n/a"), this);
 		}
 		else
 		{
-			UE_LOG(LogMediaUtils, Verbose, TEXT("PlayerFacade %llx: Media Info:\n%s"), this, *MediaInfo);
+			UE_LOG(LogMediaUtils, Verbose, TEXT("PlayerFacade %p: Media Info:\n%s"), this, *MediaInfo);
 		}
 	}
 
@@ -751,12 +761,12 @@ void FMediaPlayerFacade::SelectDefaultTracks()
 /* IMediaClockSink interface
 *****************************************************************************/
 
-void FMediaPlayerFacade::TickFetch(FTimespan DeltaTime, FTimespan /*Timecode*/)
+void FMediaPlayerFacade::TickFetch(FTimespan DeltaTime, FTimespan Timecode)
 {
 	// let the player generate samples & process events
 	if (Player.IsValid())
 	{
-		Player->TickFetch(DeltaTime);
+		Player->TickFetch(DeltaTime, Timecode);
 	}
 
 	// process deferred events
@@ -811,11 +821,11 @@ void FMediaPlayerFacade::TickFetch(FTimespan DeltaTime, FTimespan /*Timecode*/)
 }
 
 
-void FMediaPlayerFacade::TickInput(FTimespan DeltaTime, FTimespan /*Timecode*/)
+void FMediaPlayerFacade::TickInput(FTimespan DeltaTime, FTimespan Timecode)
 {
 	if (Player.IsValid())
 	{
-		Player->TickInput(DeltaTime);
+		Player->TickInput(DeltaTime, Timecode);
 	}
 }
 
@@ -881,14 +891,8 @@ void FMediaPlayerFacade::TickTickable()
 
 void FMediaPlayerFacade::ReceiveMediaEvent(EMediaEvent Event)
 {
-	if (IsInGameThread())
-	{
-		ProcessEvent(Event);
-	}
-	else
-	{
-		QueuedEvents.Enqueue(Event);
-	}
+	UE_LOG(LogMediaUtils, VeryVerbose, TEXT("PlayerFacade %p: Received media event %s"), this, *MediaUtils::EventToString(Event));
+	QueuedEvents.Enqueue(Event);
 }
 
 
@@ -903,7 +907,9 @@ void FMediaPlayerFacade::ProcessAudioSamples(IMediaSamples& Samples, TRange<FTim
 	{
 		if (!AudioSampleSinks.Enqueue(Sample, MediaPlayerFacade::MaxAudioSinkDepth))
 		{
-			UE_LOG(LogMediaUtils, VeryVerbose, TEXT("PlayerFacade %llx: Audio sample sink overflow"), this);
+			#if MEDIAPLAYERFACADE_TRACE_SINKOVERFLOWS
+				UE_LOG(LogMediaUtils, VeryVerbose, TEXT("PlayerFacade %p: Audio sample sink overflow"), this);
+			#endif
 		}
 	}
 }
@@ -917,7 +923,9 @@ void FMediaPlayerFacade::ProcessCaptionSamples(IMediaSamples& Samples, TRange<FT
 	{
 		if (!CaptionSampleSinks.Enqueue(Sample, MediaPlayerFacade::MaxCaptionSinkDepth))
 		{
-			UE_LOG(LogMediaUtils, VeryVerbose, TEXT("PlayerFacade %llx: Caption sample sink overflow"), this);
+			#if MEDIAPLAYERFACADE_TRACE_SINKOVERFLOWS
+				UE_LOG(LogMediaUtils, VeryVerbose, TEXT("PlayerFacade %p: Caption sample sink overflow"), this);
+			#endif
 		}
 	}
 }
@@ -931,7 +939,9 @@ void FMediaPlayerFacade::ProcessMetadataSamples(IMediaSamples& Samples, TRange<F
 	{
 		if (!MetadataSampleSinks.Enqueue(Sample, MediaPlayerFacade::MaxMetadataSinkDepth))
 		{
-			UE_LOG(LogMediaUtils, VeryVerbose, TEXT("PlayerFacade %llx: Metadata sample sink overflow"), this);
+			#if MEDIAPLAYERFACADE_TRACE_SINKOVERFLOWS
+				UE_LOG(LogMediaUtils, VeryVerbose, TEXT("PlayerFacade %p: Metadata sample sink overflow"), this);
+			#endif
 		}
 	}
 }
@@ -945,7 +955,9 @@ void FMediaPlayerFacade::ProcessSubtitleSamples(IMediaSamples& Samples, TRange<F
 	{
 		if (!SubtitleSampleSinks.Enqueue(Sample, MediaPlayerFacade::MaxSubtitleSinkDepth))
 		{
-			UE_LOG(LogMediaUtils, VeryVerbose, TEXT("PlayerFacade %llx: Subtitle sample sink overflow"), this);
+			#if MEDIAPLAYERFACADE_TRACE_SINKOVERFLOWS
+				UE_LOG(LogMediaUtils, VeryVerbose, TEXT("PlayerFacade %p: Subtitle sample sink overflow"), this);
+			#endif
 		}
 	}
 }
@@ -959,7 +971,9 @@ void FMediaPlayerFacade::ProcessVideoSamples(IMediaSamples& Samples, TRange<FTim
 	{
 		if (!VideoSampleSinks.Enqueue(Sample, MediaPlayerFacade::MaxVideoSinkDepth))
 		{
-			UE_LOG(LogMediaUtils, VeryVerbose, TEXT("PlayerFacade %llx: Video sample sink overflow"), this);
+			#if MEDIAPLAYERFACADE_TRACE_SINKOVERFLOWS
+				UE_LOG(LogMediaUtils, VeryVerbose, TEXT("PlayerFacade %p: Video sample sink overflow"), this);
+			#endif
 		}
 	}
 }
diff --git a/Engine/Source/Runtime/MediaUtils/Private/MediaSamples.cpp b/Engine/Source/Runtime/MediaUtils/Private/MediaSamples.cpp
--- a/Engine/Source/Runtime/MediaUtils/Private/MediaSamples.cpp
+++ b/Engine/Source/Runtime/MediaUtils/Private/MediaSamples.cpp
@@ -27,11 +27,7 @@ bool FMediaSamples::FetchAudio(TRange<FTimespan> TimeRange, TSharedPtr<IMediaAud
 		return false;
 	}
 
-	if (!AudioSampleQueue.Dequeue(Sample))
-	{
-		return false;
-	}
-
+	check(AudioSampleQueue.Pop());
 	OutSample = Sample;
 
 	return true;
@@ -54,11 +50,7 @@ bool FMediaSamples::FetchCaption(TRange<FTimespan> TimeRange, TSharedPtr<IMediaO
 		return false;
 	}
 
-	if (!CaptionSampleQueue.Dequeue(Sample))
-	{
-		return false;
-	}
-
+	check(CaptionSampleQueue.Pop());
 	OutSample = Sample;
 
 	return true;
@@ -81,11 +73,7 @@ bool FMediaSamples::FetchMetadata(TRange<FTimespan> TimeRange, TSharedPtr<IMedia
 		return false;
 	}
 
-	if (!MetadataSampleQueue.Dequeue(Sample))
-	{
-		return false;
-	}
-
+	check(MetadataSampleQueue.Pop());
 	OutSample = Sample;
 
 	return true;
@@ -108,11 +96,7 @@ bool FMediaSamples::FetchSubtitle(TRange<FTimespan> TimeRange, TSharedPtr<IMedia
 		return false;
 	}
 
-	if (!SubtitleSampleQueue.Dequeue(Sample))
-	{
-		return false;
-	}
-
+	check(SubtitleSampleQueue.Pop());
 	OutSample = Sample;
 
 	return true;
@@ -135,11 +119,7 @@ bool FMediaSamples::FetchVideo(TRange<FTimespan> TimeRange, TSharedPtr<IMediaTex
 		return false;
 	}
 
-	if (!VideoSampleQueue.Dequeue(Sample))
-	{
-		return false;
-	}
-
+	check(VideoSampleQueue.Pop());
 	OutSample = Sample;
 
 	return true;
diff --git a/Engine/Source/Runtime/MediaUtils/Public/MediaPlayerFacade.h b/Engine/Source/Runtime/MediaUtils/Public/MediaPlayerFacade.h
--- a/Engine/Source/Runtime/MediaUtils/Public/MediaPlayerFacade.h
+++ b/Engine/Source/Runtime/MediaUtils/Public/MediaPlayerFacade.h
@@ -184,6 +184,14 @@ public:
 	 */
 	FString GetInfo() const;
 
+	/**
+	 * Get the human readable name of the currently loaded media source.
+	 *
+	 * @return Media source name, or empty text if no media is opened
+	 * @see GetPlayerName, GetUrl
+	 */
+	FText GetMediaName() const;
+
 	/**
 	 * Get the number of tracks of the given type.
 	 *
@@ -217,6 +225,7 @@ public:
 	 * Get the name of the current native media player.
 	 *
 	 * @return Player name, or NAME_None if not available.
+	 * @see GetMediaName
 	 */
 	FName GetPlayerName() const;
 
diff --git a/Engine/Source/Runtime/MediaUtils/Public/MediaSampleQueue.h b/Engine/Source/Runtime/MediaUtils/Public/MediaSampleQueue.h
--- a/Engine/Source/Runtime/MediaUtils/Public/MediaSampleQueue.h
+++ b/Engine/Source/Runtime/MediaUtils/Public/MediaSampleQueue.h
@@ -65,7 +65,7 @@ public:
 			return false; // pending flush
 		}
 
-		check(Samples.Dequeue(Sample));
+		check(Samples.Pop());
 		check(Sample.IsValid());
 
 		FPlatformAtomics::InterlockedDecrement(&NumSamples);
@@ -97,20 +97,48 @@ public:
 		return true;
 	}
 
+	virtual bool Pop() override
+	{
+		TSharedPtr<SampleType, ESPMode::ThreadSafe> Sample;
+
+		if (!Samples.Peek(Sample))
+		{
+			return false; // empty queue
+		}
+
+		if (!Sample.IsValid())
+		{
+			return false; // pending flush
+		}
+
+		check(Samples.Pop());
+		check(Sample.IsValid());
+
+		FPlatformAtomics::InterlockedDecrement(&NumSamples);
+		check(NumSamples >= 0);
+
+		return true;
+	}
+
 public:
 
 	//~ TMediaSampleSink interface (to be called only from producer thread)
 
 	virtual bool Enqueue(const TSharedPtr<SampleType, ESPMode::ThreadSafe>& Sample) override
 	{
-		if (!Samples.Enqueue(Sample))
+		if (Sample.IsValid())
 		{
-			return false;
+			FPlatformAtomics::InterlockedIncrement(&NumSamples);
 		}
 
-		if (Sample.IsValid())
+		if (!Samples.Enqueue(Sample))
 		{
-			FPlatformAtomics::InterlockedIncrement(&NumSamples);
+			if (Sample.IsValid())
+			{
+				FPlatformAtomics::InterlockedDecrement(&NumSamples);
+			}
+
+			return false;
 		}
 
 		return true;
diff --git a/Engine/Source/Runtime/MediaUtils/Public/MediaSampleSource.h b/Engine/Source/Runtime/MediaUtils/Public/MediaSampleSource.h
--- a/Engine/Source/Runtime/MediaUtils/Public/MediaSampleSource.h
+++ b/Engine/Source/Runtime/MediaUtils/Public/MediaSampleSource.h
@@ -20,11 +20,11 @@ class TMediaSampleSource
 public:
 
 	/**
-	 * Dequeue the next sample in the queue.
+	 * Remove and return the next sample in the queue.
 	 *
 	 * @param OutSample Will contain the sample if the queue is not empty.
 	 * @return true if a sample has been returned, false if the queue was empty.
-	 * @see Peek
+	 * @see Peek, Pop
 	 */
 	virtual bool Dequeue(TSharedPtr<SampleType, ESPMode::ThreadSafe>& OutSample) = 0;
 
@@ -33,10 +33,18 @@ public:
 	 *
 	 * @param OutSample Will contain the sample if the queue is not empty.
 	 * @return true if a sample has been returned, false if the queue was empty.
-	 * @see Peek
+	 * @see Dequeue, Pop
 	 */
 	virtual bool Peek(TSharedPtr<SampleType, ESPMode::ThreadSafe>& OutSample) = 0;
 
+	/**
+	 * Remove the next sample from the queue.
+	 *
+	 * @return true if a sample was removed, false otherwise.
+	 * @see Dequeue, Peek
+	 */
+	virtual bool Pop() = 0;
+
 public:
 
 	/** Virtual destructor. */
diff --git a/Engine/Source/Runtime/MovieScene/Public/MovieSceneSection.h b/Engine/Source/Runtime/MovieScene/Public/MovieSceneSection.h
--- a/Engine/Source/Runtime/MovieScene/Public/MovieSceneSection.h
+++ b/Engine/Source/Runtime/MovieScene/Public/MovieSceneSection.h
@@ -411,6 +411,9 @@ public:
 	void SetPostRollTime(float InPostRollTime) { PostRollTime = InPostRollTime; }
 	float GetPostRollTime() const { return PostRollTime; }
 
+	/** The optional offset time of this section */
+	virtual TOptional<float> GetOffsetTime() const { return TOptional<float>(); }
+
 	/** Gets the time for the key referenced by the supplied key handle. */
 	virtual TOptional<float> GetKeyTime( FKeyHandle KeyHandle ) const PURE_VIRTUAL( UAISenseEvent::GetKeyTime, return TOptional<float>(); );
 
