Sequencer: Fixed crash caused by erroneous persistent references to FSequencer.

diff --git a/Engine/Source/Editor/Sequencer/Private/SSequencer.cpp b/Engine/Source/Editor/Sequencer/Private/SSequencer.cpp
--- a/Engine/Source/Editor/Sequencer/Private/SSequencer.cpp
+++ b/Engine/Source/Editor/Sequencer/Private/SSequencer.cpp
@@ -1037,15 +1037,14 @@ TSharedRef<SWidget> SSequencer::MakeGeneralMenu()
 TSharedRef<SWidget> SSequencer::MakePlaybackMenu()
 {
 	FMenuBuilder MenuBuilder( true, SequencerPtr.Pin()->GetCommandBindings() );
-	TSharedPtr<FSequencer> Sequencer = SequencerPtr.Pin();
 
 	// playback range options
 	MenuBuilder.BeginSection("PlaybackThisSequence", LOCTEXT("PlaybackThisSequenceHeader", "Playback - This Sequence"));
 	{
 		// Menu entry for the start position
 		auto OnStartChanged = [=](float NewValue){
-			float Upper = Sequencer->GetPlaybackRange().GetUpperBoundValue();
-			Sequencer->SetPlaybackRange(TRange<float>(FMath::Min(NewValue, Upper), Upper));
+			float Upper = SequencerPtr.Pin()->GetPlaybackRange().GetUpperBoundValue();
+			SequencerPtr.Pin()->SetPlaybackRange(TRange<float>(FMath::Min(NewValue, Upper), Upper));
 		};
 
 		MenuBuilder.AddWidget(
@@ -1060,7 +1059,7 @@ TSharedRef<SWidget> SSequencer::MakePlaybackMenu()
 					SNew(SSpinBox<float>)
 						.TypeInterface(NumericTypeInterface)
 						.IsEnabled_Lambda([=]() {
-							return !Sequencer->IsPlaybackRangeLocked();
+							return !SequencerPtr.Pin()->IsPlaybackRangeLocked();
 						})
 						.Style(&FEditorStyle::GetWidgetStyle<FSpinBoxStyle>("Sequencer.HyperlinkSpinBox"))
 						.OnValueCommitted_Lambda([=](float Value, ETextCommit::Type){ OnStartChanged(Value); })
@@ -1068,21 +1067,21 @@ TSharedRef<SWidget> SSequencer::MakePlaybackMenu()
 						.OnBeginSliderMovement(OnPlaybackRangeBeginDrag)
 						.OnEndSliderMovement_Lambda([=](float Value){ OnStartChanged(Value); OnPlaybackRangeEndDrag.ExecuteIfBound(); })
 						.MinValue_Lambda([=]() -> float {
-							return Sequencer->GetClampRange().GetLowerBoundValue(); 
+							return SequencerPtr.Pin()->GetClampRange().GetLowerBoundValue(); 
 						})
 						.MaxValue_Lambda([=]() -> float {
-							return Sequencer->GetPlaybackRange().GetUpperBoundValue(); 
+							return SequencerPtr.Pin()->GetPlaybackRange().GetUpperBoundValue(); 
 						})
 						.Value_Lambda([=]() -> float {
-							return Sequencer->GetPlaybackRange().GetLowerBoundValue();
+							return SequencerPtr.Pin()->GetPlaybackRange().GetLowerBoundValue();
 						})
 				],
 			LOCTEXT("PlaybackStartLabel", "Start"));
 
 		// Menu entry for the end position
 		auto OnEndChanged = [=](float NewValue){
-			float Lower = Sequencer->GetPlaybackRange().GetLowerBoundValue();
-			Sequencer->SetPlaybackRange(TRange<float>(Lower, FMath::Max(NewValue, Lower)));
+			float Lower = SequencerPtr.Pin()->GetPlaybackRange().GetLowerBoundValue();
+			SequencerPtr.Pin()->SetPlaybackRange(TRange<float>(Lower, FMath::Max(NewValue, Lower)));
 		};
 
 		MenuBuilder.AddWidget(
@@ -1097,7 +1096,7 @@ TSharedRef<SWidget> SSequencer::MakePlaybackMenu()
 					SNew(SSpinBox<float>)
 						.TypeInterface(NumericTypeInterface)
 						.IsEnabled_Lambda([=]() {
-							return !Sequencer->IsPlaybackRangeLocked();
+							return !SequencerPtr.Pin()->IsPlaybackRangeLocked();
 						})
 						.Style(&FEditorStyle::GetWidgetStyle<FSpinBoxStyle>("Sequencer.HyperlinkSpinBox"))
 						.OnValueCommitted_Lambda([=](float Value, ETextCommit::Type){ OnEndChanged(Value); })
@@ -1105,13 +1104,13 @@ TSharedRef<SWidget> SSequencer::MakePlaybackMenu()
 						.OnBeginSliderMovement(OnPlaybackRangeBeginDrag)
 						.OnEndSliderMovement_Lambda([=](float Value){ OnEndChanged(Value); OnPlaybackRangeEndDrag.ExecuteIfBound(); })
 						.MinValue_Lambda([=]() -> float {
-							return Sequencer->GetPlaybackRange().GetLowerBoundValue(); 
+							return SequencerPtr.Pin()->GetPlaybackRange().GetLowerBoundValue(); 
 						})
 						.MaxValue_Lambda([=]() -> float {
-							return Sequencer->GetClampRange().GetUpperBoundValue(); 
+							return SequencerPtr.Pin()->GetClampRange().GetUpperBoundValue(); 
 						})
 						.Value_Lambda([=]() -> float {
-							return Sequencer->GetPlaybackRange().GetUpperBoundValue();
+							return SequencerPtr.Pin()->GetPlaybackRange().GetUpperBoundValue();
 						})
 				],
 			LOCTEXT("PlaybackStartEnd", "End"));
