Copying //UE4/Dev-Sequencer to //UE4/Dev-Main (Source: //UE4/Dev-Sequencer @ 3800792)

diff --git a/Engine/Source/Editor/MovieSceneTools/Private/TrackEditors/CinematicShotTrackEditor.cpp b/Engine/Source/Editor/MovieSceneTools/Private/TrackEditors/CinematicShotTrackEditor.cpp
--- a/Engine/Source/Editor/MovieSceneTools/Private/TrackEditors/CinematicShotTrackEditor.cpp
+++ b/Engine/Source/Editor/MovieSceneTools/Private/TrackEditors/CinematicShotTrackEditor.cpp
@@ -136,6 +136,10 @@ bool FCinematicShotTrackEditor::HandleAssetAdded(UObject* Asset, const FGuid& Ta
 		return false;
 	}
 
+	if (!SupportsSequence(Sequence))
+	{
+		return false;
+	}
 
 	//@todo If there's already a subscenes track, allow that track to handle this asset
 	UMovieScene* FocusedMovieScene = GetFocusedMovieScene();
@@ -326,31 +330,6 @@ UMovieSceneSubSection* FCinematicShotTrackEditor::CreateShotInternal(FString& Ne
 	return NewSection;
 }
 
-int32 FindAvailableRowIndex(UMovieSceneCinematicShotTrack* InTrack, UMovieSceneSubSection* InSection)
-{
-	for (int32 RowIndex = 0; RowIndex <= InTrack->GetMaxRowIndex(); ++RowIndex)
-	{
-		bool bFoundIntersect = false;
-		for (auto Section : InTrack->GetAllSections())
-		{
-			TRange<float> InRange(InSection->GetStartTime(), InSection->GetEndTime());
-			TRange<float> Range(Section->GetStartTime(), Section->GetEndTime());
-
-			if (Section != InSection && Section->GetRowIndex() == RowIndex && Range.Overlaps(InRange))
-			{
-				bFoundIntersect = true;
-				break;
-			}
-		}
-		if (!bFoundIntersect)
-		{
-			return RowIndex;
-		}
-	}
-
-	return InTrack->GetMaxRowIndex() + 1;
-}
-
 void FCinematicShotTrackEditor::InsertShot()
 {
 	const FScopedTransaction Transaction(LOCTEXT("InsertShot_Transaction", "Insert Shot"));
@@ -363,7 +342,7 @@ void FCinematicShotTrackEditor::InsertShot()
 	UMovieSceneSubSection* NewShot = CreateShotInternal(NewShotName, NewShotStartTime);
 	if (NewShot)
 	{
-		NewShot->SetRowIndex(FindAvailableRowIndex(CinematicShotTrack, NewShot));
+		NewShot->SetRowIndex(MovieSceneToolHelpers::FindAvailableRowIndex(CinematicShotTrack, NewShot));
 	}
 
 	GetSequencer()->NotifyMovieSceneDataChanged( EMovieSceneDataChangeType::MovieSceneStructureItemAdded );
@@ -388,8 +367,8 @@ void FCinematicShotTrackEditor::InsertFiller()
 
 	UMovieSceneCinematicShotSection* NewCinematicShotSection = Cast<UMovieSceneCinematicShotSection>(NewSection);
 
-	NewCinematicShotSection->SetShotDisplayName(FText(LOCTEXT("Filler", "Filler")));
-	NewCinematicShotSection->SetRowIndex(FindAvailableRowIndex(CinematicShotTrack, NewSection));
+	NewCinematicShotSection->SetShotDisplayName(FText(LOCTEXT("Filler", "Filler")).ToString());
+	NewCinematicShotSection->SetRowIndex(MovieSceneToolHelpers::FindAvailableRowIndex(CinematicShotTrack, NewSection));
 
 	GetSequencer()->NotifyMovieSceneDataChanged( EMovieSceneDataChangeType::MovieSceneStructureItemAdded );
 }
@@ -408,7 +387,7 @@ void FCinematicShotTrackEditor::DuplicateShot(UMovieSceneCinematicShotSection* S
 	{
 		NewShot->SetStartTime(Section->GetStartTime());
 		NewShot->SetEndTime(Section->GetEndTime());
-		NewShot->SetRowIndex(FindAvailableRowIndex(CinematicShotTrack, NewShot));
+		NewShot->SetRowIndex(MovieSceneToolHelpers::FindAvailableRowIndex(CinematicShotTrack, NewShot));
 		NewShot->Parameters.StartOffset = Section->Parameters.StartOffset;
 		NewShot->Parameters.TimeScale = Section->Parameters.TimeScale;
 		NewShot->SetPreRollTime(Section->GetPreRollTime());
@@ -438,7 +417,7 @@ void FCinematicShotTrackEditor::NewTake(UMovieSceneCinematicShotSection* Section
 	FString ShotPrefix;
 	uint32 ShotNumber = INDEX_NONE;
 	uint32 TakeNumber = INDEX_NONE;
-	if (MovieSceneToolHelpers::ParseShotName(Section->GetShotDisplayName().ToString(), ShotPrefix, ShotNumber, TakeNumber))
+	if (MovieSceneToolHelpers::ParseShotName(Section->GetShotDisplayName(), ShotPrefix, ShotNumber, TakeNumber))
 	{
 		TArray<uint32> TakeNumbers;
 		uint32 CurrentTakeNumber;
@@ -456,6 +435,7 @@ void FCinematicShotTrackEditor::NewTake(UMovieSceneCinematicShotSection* Section
 		float NewShotStartOffset = Section->Parameters.StartOffset;
 		float NewShotTimeScale = Section->Parameters.TimeScale;
 		float NewShotPrerollTime = Section->GetPreRollTime();
+		int32 NewRowIndex = Section->GetRowIndex();
 
 		UMovieSceneSubSection* NewShot = CreateShotInternal(NewShotName, NewShotStartTime, Section);
 
@@ -469,6 +449,7 @@ void FCinematicShotTrackEditor::NewTake(UMovieSceneCinematicShotSection* Section
 			NewShot->Parameters.StartOffset = NewShotStartOffset;
 			NewShot->Parameters.TimeScale = NewShotTimeScale;
 			NewShot->SetPreRollTime(NewShotPrerollTime);
+			NewShot->SetRowIndex(NewRowIndex);
 
 			GetSequencer()->NotifyMovieSceneDataChanged( EMovieSceneDataChangeType::MovieSceneStructureItemsChanged );
 		}
@@ -596,10 +577,7 @@ FKeyPropertyResult FCinematicShotTrackEditor::AddKeyInternal(float KeyTime, UMov
 		UMovieSceneCinematicShotTrack* CinematicShotTrack = FindOrCreateCinematicShotTrack();
 		float EndTime = InMovieSceneSequence->GetMovieScene()->GetPlaybackRange().Size<float>();
 		UMovieSceneSubSection* NewSection = CinematicShotTrack->AddSequence(InMovieSceneSequence, KeyTime, EndTime);
-
-		UMovieSceneCinematicShotSection* NewCinematicShotSection = Cast<UMovieSceneCinematicShotSection>(NewSection);
-		NewCinematicShotSection->SetRowIndex(FindAvailableRowIndex(CinematicShotTrack, NewCinematicShotSection));
-
+		NewSection->SetRowIndex(MovieSceneToolHelpers::FindAvailableRowIndex(CinematicShotTrack, NewSection));
 		KeyPropertyResult.bTrackModified = true;
 	}
 
@@ -727,10 +705,7 @@ FKeyPropertyResult FCinematicShotTrackEditor::HandleSequenceAdded(float KeyTime,
 	auto CinematicShotTrack = FindOrCreateCinematicShotTrack();
 	float Duration = Sequence->GetMovieScene()->GetPlaybackRange().Size<float>();
 	UMovieSceneSubSection* NewSection = CinematicShotTrack->AddSequence(Sequence, KeyTime, Duration);
-		
-	UMovieSceneCinematicShotSection* NewCinematicShotSection = Cast<UMovieSceneCinematicShotSection>(NewSection);
-	NewCinematicShotSection->SetRowIndex(FindAvailableRowIndex(CinematicShotTrack, NewCinematicShotSection));
-
+	NewSection->SetRowIndex(MovieSceneToolHelpers::FindAvailableRowIndex(CinematicShotTrack, NewSection));
 	KeyPropertyResult.bTrackModified = true;
 
 	return KeyPropertyResult;
diff --git a/Engine/Source/Editor/MovieSceneTools/Private/TrackEditors/EventTrackEditor.cpp b/Engine/Source/Editor/MovieSceneTools/Private/TrackEditors/EventTrackEditor.cpp
--- a/Engine/Source/Editor/MovieSceneTools/Private/TrackEditors/EventTrackEditor.cpp
+++ b/Engine/Source/Editor/MovieSceneTools/Private/TrackEditors/EventTrackEditor.cpp
@@ -78,24 +78,6 @@ void FEventTrackEditor::BuildTrackContextMenu(FMenuBuilder& MenuBuilder, UMovieS
 	class FEventTrackCustomization : public IDetailCustomization
 	{
 	public:
-		FEventTrackCustomization(TWeakPtr<IDetailsView> InDetailsView, TSharedPtr<ISequencer> InSequencer)
-			: WeakDetailsView(InDetailsView)
-		{
-			FOnGetPropertyTypeCustomizationInstance Factory = FOnGetPropertyTypeCustomizationInstance::CreateLambda([=]{ return MakeShared<FMovieSceneObjectBindingIDCustomization>(InSequencer->GetFocusedTemplateID(), InSequencer); });
-
-			// Register an object binding ID customization that can use the current sequencer interface
-			WeakDetailsView.Pin()->RegisterInstancedCustomPropertyTypeLayout("MovieSceneObjectBindingID", Factory);
-		}
-
-		~FEventTrackCustomization()
-		{
-			auto PinnedDetailsView = WeakDetailsView.Pin();
-			if (PinnedDetailsView.IsValid())
-			{
-				PinnedDetailsView->UnregisterInstancedCustomPropertyTypeLayout("MovieSceneObjectBindingID");
-			}
-		}
-
 		virtual void CustomizeDetails(IDetailLayoutBuilder& DetailBuilder) override
 		{
 			DetailBuilder.HideCategory("Track");
@@ -104,8 +86,6 @@ void FEventTrackEditor::BuildTrackContextMenu(FMenuBuilder& MenuBuilder, UMovieS
 			IDetailCategoryBuilder& Category = DetailBuilder.EditCategory("TrackEvent");
 			Category.AddProperty("EventReceivers").ShouldAutoExpand(true);
 		}
-
-		TWeakPtr<IDetailsView> WeakDetailsView;
 	};
 
 	auto PopulateSubMenu = [this, EventTrack](FMenuBuilder& SubMenuBuilder)
@@ -120,10 +100,11 @@ void FEventTrackEditor::BuildTrackContextMenu(FMenuBuilder& MenuBuilder, UMovieS
 		TSharedRef<IDetailsView> DetailsView = PropertyEditor.CreateDetailView(DetailsViewArgs);
 		
 		// Register the custom type layout for the class
-		TWeakPtr<IDetailsView> WeakDetailsView = DetailsView;
-		FOnGetDetailCustomizationInstance CreateInstance = FOnGetDetailCustomizationInstance::CreateLambda([=]{ return MakeShared<FEventTrackCustomization>(WeakDetailsView, GetSequencer()); });
+		FOnGetDetailCustomizationInstance CreateInstance = FOnGetDetailCustomizationInstance::CreateLambda(&MakeShared<FEventTrackCustomization>);
 		DetailsView->RegisterInstancedCustomPropertyLayout(UMovieSceneEventTrack::StaticClass(), CreateInstance);
 
+		GetSequencer()->OnInitializeDetailsPanel().Broadcast(DetailsView, GetSequencer().ToSharedRef());
+
 		// Assign the object
 		DetailsView->SetObject(EventTrack, true);
 
diff --git a/Engine/Source/Editor/MovieSceneTools/Private/TrackEditors/MaterialTrackEditor.cpp b/Engine/Source/Editor/MovieSceneTools/Private/TrackEditors/MaterialTrackEditor.cpp
--- a/Engine/Source/Editor/MovieSceneTools/Private/TrackEditors/MaterialTrackEditor.cpp
+++ b/Engine/Source/Editor/MovieSceneTools/Private/TrackEditors/MaterialTrackEditor.cpp
@@ -196,6 +196,11 @@ bool FComponentMaterialTrackEditor::SupportsType( TSubclassOf<UMovieSceneTrack>
 	return Type == UMovieSceneComponentMaterialTrack::StaticClass();
 }
 
+bool FComponentMaterialTrackEditor::GetDefaultExpansionState(UMovieSceneTrack* InTrack) const
+{
+	return true;
+}
+
 
 UMaterialInterface* FComponentMaterialTrackEditor::GetMaterialInterfaceForTrack( FGuid ObjectBinding, UMovieSceneMaterialTrack* MaterialTrack )
 {
diff --git a/Engine/Source/Editor/MovieSceneTools/Private/TrackEditors/PathTrackEditor.cpp b/Engine/Source/Editor/MovieSceneTools/Private/TrackEditors/PathTrackEditor.cpp
--- a/Engine/Source/Editor/MovieSceneTools/Private/TrackEditors/PathTrackEditor.cpp
+++ b/Engine/Source/Editor/MovieSceneTools/Private/TrackEditors/PathTrackEditor.cpp
@@ -41,7 +41,7 @@ public:
 			TSharedPtr<ISequencer> Sequencer = PathTrackEditor->GetSequencer();
 			if (Sequencer.IsValid())
 			{
-				TArrayView<TWeakObjectPtr<UObject>> RuntimeObjects = Sequencer->FindBoundObjects(PathSection->GetConstraintId(), Sequencer->GetFocusedTemplateID());
+				TArrayView<TWeakObjectPtr<UObject>> RuntimeObjects = Sequencer->FindBoundObjects(PathSection->GetConstraintBindingID().GetGuid(), PathSection->GetConstraintBindingID().GetSequenceID());
 				if (RuntimeObjects.Num() == 1 && RuntimeObjects[0].IsValid())
 				{
 					if (AActor* Actor = Cast<AActor>(RuntimeObjects[0].Get()))
@@ -135,14 +135,11 @@ bool F3DPathTrackEditor::IsActorPickable(const AActor* const ParentActor, FGuid
 	UMovieScene3DPathSection* PathSection = Cast<UMovieScene3DPathSection>(InSection);
 	if (PathSection != nullptr)
 	{
-		FGuid ConstraintId = PathSection->GetConstraintId();
-
-		if (ConstraintId.IsValid())
+		TArrayView<TWeakObjectPtr<UObject>> RuntimeObjects = GetSequencer()->FindBoundObjects(PathSection->GetConstraintBindingID().GetGuid(), PathSection->GetConstraintBindingID().GetSequenceID());
+		
+		if (RuntimeObjects.Contains(ParentActor))
 		{
-			if (GetSequencer()->FindObjectsInCurrentSequence(ConstraintId).Contains(ParentActor))
-			{
-				return false;
-			}
+			return false;
 		}
 	}
 
@@ -162,18 +159,29 @@ bool F3DPathTrackEditor::IsActorPickable(const AActor* const ParentActor, FGuid
 }
 
 
-void F3DPathTrackEditor::ActorSocketPicked(const FName SocketName, USceneComponent* Component, AActor* ParentActor, FGuid ObjectGuid, UMovieSceneSection* Section)
+void F3DPathTrackEditor::ActorSocketPicked(const FName SocketName, USceneComponent* Component, FActorPickerID ActorPickerID, FGuid ObjectGuid, UMovieSceneSection* Section)
 {
 	if (Section != nullptr)
 	{
 		const FScopedTransaction Transaction(LOCTEXT("UndoSetPath", "Set Path"));
 
 		UMovieScene3DPathSection* PathSection = (UMovieScene3DPathSection*)(Section);
-		FGuid SplineId = FindOrCreateHandleToObject(ParentActor).Handle;
 
-		if (SplineId.IsValid())
+		FMovieSceneObjectBindingID ConstraintBindingID;
+
+		if (ActorPickerID.ExistingBindingID.IsValid())
 		{
-			PathSection->SetConstraintId(SplineId);
+			ConstraintBindingID = ActorPickerID.ExistingBindingID;
+		}
+		else if (ActorPickerID.ActorPicked.IsValid())
+		{
+			FGuid ParentActorId = FindOrCreateHandleToObject(ActorPickerID.ActorPicked.Get()).Handle;
+			ConstraintBindingID = FMovieSceneObjectBindingID(ParentActorId, MovieSceneSequenceID::Root);
+		}
+
+		if (ConstraintBindingID.IsValid())
+		{
+			PathSection->SetConstraintBindingID(ConstraintBindingID);
 		}
 	}
 	else if (ObjectGuid.IsValid())
@@ -183,26 +191,29 @@ void F3DPathTrackEditor::ActorSocketPicked(const FName SocketName, USceneCompone
 		{
 			OutObjects.Add(Object);
 		}
-		AnimatablePropertyChanged( FOnKeyProperty::CreateRaw( this, &F3DPathTrackEditor::AddKeyInternal, OutObjects, ParentActor) );
+		AnimatablePropertyChanged( FOnKeyProperty::CreateRaw( this, &F3DPathTrackEditor::AddKeyInternal, OutObjects, ActorPickerID) );
 	}
 }
 
-FKeyPropertyResult F3DPathTrackEditor::AddKeyInternal( float KeyTime, const TArray<TWeakObjectPtr<UObject>> Objects, AActor* ParentActor)
+FKeyPropertyResult F3DPathTrackEditor::AddKeyInternal( float KeyTime, const TArray<TWeakObjectPtr<UObject>> Objects, FActorPickerID ActorPickerID)
 {
 	FKeyPropertyResult KeyPropertyResult;
 
-	AActor* ActorWithSplineComponent = ParentActor;
-	
-	FGuid SplineId;
+	FMovieSceneObjectBindingID ConstraintBindingID;
 
-	if (ActorWithSplineComponent != nullptr)
+	if (ActorPickerID.ExistingBindingID.IsValid())
 	{
-		FFindOrCreateHandleResult HandleResult = FindOrCreateHandleToObject( ActorWithSplineComponent );
-		SplineId = HandleResult.Handle;
+		ConstraintBindingID = ActorPickerID.ExistingBindingID;
+	}
+	else if (ActorPickerID.ActorPicked.IsValid())
+	{
+		FFindOrCreateHandleResult HandleResult = FindOrCreateHandleToObject(ActorPickerID.ActorPicked.Get());
+		FGuid ParentActorId = HandleResult.Handle;
 		KeyPropertyResult.bHandleCreated |= HandleResult.bWasCreated;
+		ConstraintBindingID = FMovieSceneObjectBindingID(ParentActorId, MovieSceneSequenceID::Root);
 	}
 
-	if (!SplineId.IsValid())
+	if (!ConstraintBindingID.IsValid())
 	{
 		return KeyPropertyResult;
 	}
@@ -238,7 +249,7 @@ FKeyPropertyResult F3DPathTrackEditor::AddKeyInternal( float KeyTime, const TArr
 					}
 				}
 
-				Cast<UMovieScene3DPathTrack>(Track)->AddConstraint( KeyTime, PathEndTime, NAME_None, NAME_None, SplineId );
+				Cast<UMovieScene3DPathTrack>(Track)->AddConstraint( KeyTime, PathEndTime, NAME_None, NAME_None, ConstraintBindingID );
 				KeyPropertyResult.bTrackModified = true;
 			}
 		}
@@ -247,5 +258,4 @@ FKeyPropertyResult F3DPathTrackEditor::AddKeyInternal( float KeyTime, const TArr
 	return KeyPropertyResult;
 }
 
-
 #undef LOCTEXT_NAMESPACE
diff --git a/Engine/Source/Editor/MovieSceneTools/Private/TrackEditors/PathTrackEditor.h b/Engine/Source/Editor/MovieSceneTools/Private/TrackEditors/PathTrackEditor.h
--- a/Engine/Source/Editor/MovieSceneTools/Private/TrackEditors/PathTrackEditor.h
+++ b/Engine/Source/Editor/MovieSceneTools/Private/TrackEditors/PathTrackEditor.h
@@ -51,12 +51,12 @@ public:
 
 	// FTrackEditorActorPicker
 	virtual bool IsActorPickable( const AActor* const ParentActor, FGuid ObjectBinding, UMovieSceneSection* InSection ) override;
-	virtual void ActorSocketPicked(const FName SocketName, USceneComponent* Component, AActor* ParentActor, FGuid ObjectBinding, UMovieSceneSection* Section) override;
+	virtual void ActorSocketPicked(const FName SocketName, USceneComponent* Component, FActorPickerID ActorPickerID, FGuid ObjectBinding, UMovieSceneSection* Section) override;
 
 private:
 
 	/** Delegate for AnimatablePropertyChanged in AddKey */
-	FKeyPropertyResult AddKeyInternal(float KeyTime, const TArray<TWeakObjectPtr<UObject>> Objects, AActor* ParentActor);
+	FKeyPropertyResult AddKeyInternal(float KeyTime, const TArray<TWeakObjectPtr<UObject>> Objects, FActorPickerID ActorPickerID);
 };
 
 
diff --git a/Engine/Source/Editor/MovieSceneTools/Private/TrackEditors/SubTrackEditor.cpp b/Engine/Source/Editor/MovieSceneTools/Private/TrackEditors/SubTrackEditor.cpp
--- a/Engine/Source/Editor/MovieSceneTools/Private/TrackEditors/SubTrackEditor.cpp
+++ b/Engine/Source/Editor/MovieSceneTools/Private/TrackEditors/SubTrackEditor.cpp
@@ -23,6 +23,7 @@
 #include "ISequenceRecorder.h"
 #include "SequenceRecorderSettings.h"
 #include "DragAndDrop/AssetDragDropOp.h"
+#include "MovieSceneToolHelpers.h"
 
 namespace SubTrackEditorConstants
 {
@@ -377,6 +378,11 @@ bool FSubTrackEditor::HandleAssetAdded(UObject* Asset, const FGuid& TargetObject
 		return false;
 	}
 
+	if (!SupportsSequence(Sequence))
+	{
+		return false;
+	}
+
 	//@todo If there's already a cinematic shot track, allow that track to handle this asset
 	UMovieScene* FocusedMovieScene = GetFocusedMovieScene();
 
@@ -626,7 +632,8 @@ FKeyPropertyResult FSubTrackEditor::AddKeyInternal(float KeyTime, UMovieSceneSeq
 	{
 		UMovieSceneSubTrack* SubTrack = Cast<UMovieSceneSubTrack>(InTrack);
 		float Duration = InMovieSceneSequence->GetMovieScene()->GetPlaybackRange().Size<float>();
-		SubTrack->AddSequence(InMovieSceneSequence, KeyTime, Duration);
+		UMovieSceneSubSection* NewSection = SubTrack->AddSequence(InMovieSceneSequence, KeyTime, Duration);
+		NewSection->SetRowIndex(MovieSceneToolHelpers::FindAvailableRowIndex(SubTrack, NewSection));
 		KeyPropertyResult.bTrackModified = true;
 
 		return KeyPropertyResult;
@@ -645,7 +652,8 @@ FKeyPropertyResult FSubTrackEditor::HandleSequenceAdded(float KeyTime, UMovieSce
 
 	auto SubTrack = FindOrCreateMasterTrack<UMovieSceneSubTrack>().Track;
 	float Duration = Sequence->GetMovieScene()->GetPlaybackRange().Size<float>();
-	SubTrack->AddSequence(Sequence, KeyTime, Duration);
+	UMovieSceneSubSection* NewSection = SubTrack->AddSequence(Sequence, KeyTime, Duration);
+	NewSection->SetRowIndex(MovieSceneToolHelpers::FindAvailableRowIndex(SubTrack, NewSection));
 	KeyPropertyResult.bTrackModified = true;
 
 	return KeyPropertyResult;
diff --git a/Engine/Source/Editor/MovieSceneTools/Private/TrackEditors/TransformTrackEditor.cpp b/Engine/Source/Editor/MovieSceneTools/Private/TrackEditors/TransformTrackEditor.cpp
--- a/Engine/Source/Editor/MovieSceneTools/Private/TrackEditors/TransformTrackEditor.cpp
+++ b/Engine/Source/Editor/MovieSceneTools/Private/TrackEditors/TransformTrackEditor.cpp
@@ -81,6 +81,27 @@ protected:
 		}
 		return TOptional<FTransform>();
 	}
+	virtual TOptional<FRotator> GetCurrentRotator() const
+	{
+		TSharedPtr<ISequencer> Sequencer = WeakSequencer.Pin();
+		if (Sequencer.IsValid())
+		{
+			for (TWeakObjectPtr<> WeakObject : Sequencer->FindBoundObjects(ObjectBinding, Sequencer->GetFocusedTemplateID()))
+			{
+				if (UObject* Object = WeakObject.Get())
+				{
+					AActor* Actor = nullptr;
+					USceneComponent* Component = nullptr;
+					GetActorAndSceneComponentFromObject(Object, Actor, Component);
+					if (Component)
+					{
+						return Component->RelativeRotation;
+					}
+				}
+			}
+		}
+		return TOptional<FRotator>();
+	}
 };
 
 
@@ -699,6 +720,19 @@ void GetKeysForVector( bool LastVectorIsValid, const FVector& LastVector, const
 	ZKeys.Add( FTransformKey( VectorChannel, EAxis::Z, CurrentVector.Z, bUnwindRotation ) );
 }
 
+float UnwindChannel(const float& OldValue, float NewValue)
+{
+	while( NewValue - OldValue > 180.0f )
+	{
+		NewValue -= 360.0f;
+	}
+	while( NewValue - OldValue < -180.0f )
+	{
+		NewValue += 360.0f;
+	}
+	return NewValue;
+}
+
 void F3DTransformTrackEditor::GetTransformKeys( const FTransformData& LastTransform, const FTransformData& CurrentTransform, EKey3DTransformChannel::Type ChannelsToKey, bool bUnwindRotation, TArray<FTransformKey>& OutNewKeys, TArray<FTransformKey>& OutDefaultKeys )
 {
 	bool bLastVectorIsValid = LastTransform.IsValid();
@@ -710,8 +744,18 @@ void F3DTransformTrackEditor::GetTransformKeys( const FTransformData& LastTransf
 		ChannelsToKey = EKey3DTransformChannel::All;
 	}
 
+	// Unwind rotation
+	FRotator CurrentRotator = CurrentTransform.Rotation;
+	FRotator LastRotator = LastTransform.Rotation;
+	if (LastTransform.IsValid())
+	{
+		CurrentRotator.Yaw = UnwindChannel(LastRotator.Yaw, CurrentRotator.Yaw);
+		CurrentRotator.Pitch = UnwindChannel(LastRotator.Pitch, CurrentRotator.Pitch);
+		CurrentRotator.Roll = UnwindChannel(LastRotator.Roll, CurrentRotator.Roll);
+	}
+
 	GetKeysForVector( bLastVectorIsValid, LastTransform.Translation, CurrentTransform.Translation, EKey3DTransformChannel::Translation, ChannelsToKey, bUnwindRotation, OutNewKeys, OutDefaultKeys );
-	GetKeysForVector( bLastVectorIsValid, LastTransform.Rotation.Euler(), CurrentTransform.Rotation.Euler(), EKey3DTransformChannel::Rotation, ChannelsToKey, bUnwindRotation, OutNewKeys, OutDefaultKeys );
+	GetKeysForVector( bLastVectorIsValid, LastRotator.Euler(), CurrentRotator.Euler(), EKey3DTransformChannel::Rotation, ChannelsToKey, bUnwindRotation, OutNewKeys, OutDefaultKeys );
 	GetKeysForVector( bLastVectorIsValid, LastTransform.Scale, CurrentTransform.Scale, EKey3DTransformChannel::Scale, ChannelsToKey, bUnwindRotation, OutNewKeys, OutDefaultKeys );
 }
 
diff --git a/Engine/Source/Editor/MovieSceneTools/Public/GenericKeyArea.h b/Engine/Source/Editor/MovieSceneTools/Public/GenericKeyArea.h
--- a/Engine/Source/Editor/MovieSceneTools/Public/GenericKeyArea.h
+++ b/Engine/Source/Editor/MovieSceneTools/Public/GenericKeyArea.h
@@ -296,6 +296,10 @@ private:
 					.Value_Lambda([=]{ return this->CurveInterface.GetKeyTime(KeyHandle).Get(TimeType()); })
 					.OnValueChanged_Lambda(OnValueChanged)
 					.OnValueCommitted_Lambda([=](TimeType InTime, ETextCommit::Type){ OnValueChanged(InTime); })
+					.MinValue(TOptional<float>())
+					.MaxValue(TOptional<float>())
+					.MaxSliderValue(TOptional<float>())
+					.MinSliderValue(TOptional<float>())
 					.ToolTipText(TimeTooltipText)
 				];
 
diff --git a/Engine/Source/Editor/MovieSceneTools/Public/K2Node_GetSequenceBinding.h b/Engine/Source/Editor/MovieSceneTools/Public/K2Node_GetSequenceBinding.h
--- a/Engine/Source/Editor/MovieSceneTools/Public/K2Node_GetSequenceBinding.h
+++ b/Engine/Source/Editor/MovieSceneTools/Public/K2Node_GetSequenceBinding.h
@@ -6,6 +6,7 @@
 #include "K2Node.h"
 #include "UObject/ObjectMacros.h"
 #include "MovieScene.h"
+#include "Evaluation/MovieSceneSequenceHierarchy.h"
 
 #include "K2Node_GetSequenceBinding.generated.h"
 
@@ -52,4 +53,8 @@ public:
 	FText GetSequenceName() const;
 	FText GetBindingName() const;
 	UMovieScene* GetObjectMovieScene() const;
+
+private:
+	mutable FMovieSceneSequenceHierarchy SequenceHierarchyCache;
+	mutable TMap<FMovieSceneSequenceID, FGuid> SequenceSignatureCache;
 };
diff --git a/Engine/Source/Editor/MovieSceneTools/Public/KeyframeTrackEditor.h b/Engine/Source/Editor/MovieSceneTools/Public/KeyframeTrackEditor.h
--- a/Engine/Source/Editor/MovieSceneTools/Public/KeyframeTrackEditor.h
+++ b/Engine/Source/Editor/MovieSceneTools/Public/KeyframeTrackEditor.h
@@ -97,7 +97,7 @@ protected:
 
 			if ( ObjectHandle.IsValid() )
 			{
-				KeyPropertyResult.bTrackCreated |= AddKeysToHandle( ObjectHandle, KeyTime, NewKeys, DefaultKeys, KeyMode, TrackClass, PropertyName, OnInitializeNewTrack );
+				KeyPropertyResult |= AddKeysToHandle( ObjectHandle, KeyTime, NewKeys, DefaultKeys, KeyMode, TrackClass, PropertyName, OnInitializeNewTrack );
 			}
 		}
 		return KeyPropertyResult;
@@ -144,7 +144,7 @@ private:
 	 *        track is created, but before any sections or keys have been added.
 	 * @return Whether or not a track was created. Note this does not return true if keys were added or modified.
 	*/
-	bool AddKeysToHandle(
+	FKeyPropertyResult AddKeysToHandle(
 		FGuid ObjectHandle, float KeyTime,
 		const TArray<KeyDataType>& NewKeys, const TArray<KeyDataType>& DefaultKeys,
 		ESequencerKeyMode KeyMode, TSubclassOf<UMovieSceneTrack> TrackClass, FName PropertyName,
@@ -175,21 +175,24 @@ private:
 			}
 		}
 
+		FKeyPropertyResult KeyPropertyResult;
 		if ( Track )
 		{
-			bSectionCreated = AddKeysToTrack( Track, KeyTime, NewKeys, DefaultKeys, KeyMode, bTrackCreated );
+			KeyPropertyResult |= AddKeysToTrack( Track, KeyTime, NewKeys, DefaultKeys, KeyMode, bTrackCreated );
 		}
 
-		return bTrackCreated || bSectionCreated;
+		KeyPropertyResult.bTrackCreated |= bTrackCreated || bSectionCreated;
+
+		return KeyPropertyResult;
 	}
 
-	/* Returns whether a section was added */
-	bool AddKeysToTrack(
+	FKeyPropertyResult AddKeysToTrack(
 		TrackType* Track, float KeyTime,
 		const TArray<KeyDataType>& NewKeys, const TArray<KeyDataType>& DefaultKeys,
 		ESequencerKeyMode KeyMode, bool bNewTrack)
 	{
 		bool bSectionCreated = false;
+		bool bKeyCreated = false;
 		bool bInfiniteKeyAreas = GetSequencer()->GetInfiniteKeyAreas();
 
 		EAutoChangeMode AutoChangeMode = GetSequencer()->GetAutoChangeMode();
@@ -213,6 +216,7 @@ private:
 					if ( HasKeys( Track, NewKey ) || bKeyEvenIfEmpty )
 					{
 						bSectionCreated |= AddKey( Track, KeyTime, NewKey, InterpolationMode, bInfiniteKeyAreas );
+						bKeyCreated = true;
 					}
 				}
 			}
@@ -239,7 +243,11 @@ private:
 			NewSection->SetIsInfinite(bInfiniteKeyAreas);
 		}
 
-		return bSectionCreated;
+		FKeyPropertyResult KeyPropertyResult;
+		KeyPropertyResult.bTrackCreated |= bSectionCreated;
+		KeyPropertyResult.bKeyCreated = bKeyCreated;
+
+		return KeyPropertyResult;
 	}
 
 	bool NewKeyIsNewData( TrackType* Track, float Time, const KeyDataType& KeyData ) const
diff --git a/Engine/Source/Editor/MovieSceneTools/Public/MovieSceneObjectBindingIDCustomization.h b/Engine/Source/Editor/MovieSceneTools/Public/MovieSceneObjectBindingIDCustomization.h
--- a/Engine/Source/Editor/MovieSceneTools/Public/MovieSceneObjectBindingIDCustomization.h
+++ b/Engine/Source/Editor/MovieSceneTools/Public/MovieSceneObjectBindingIDCustomization.h
@@ -27,6 +27,8 @@ public:
 		: FMovieSceneObjectBindingIDPicker(InLocalSequenceID, InSequencer)
 	{}
 
+	MOVIESCENETOOLS_API static void BindTo(TSharedRef<ISequencer> InSequencer);
+
 	virtual void CustomizeHeader(TSharedRef<IPropertyHandle> PropertyHandle, FDetailWidgetRow& HeaderRow, IPropertyTypeCustomizationUtils& CustomizationUtils) override;
 	virtual void CustomizeChildren(TSharedRef<IPropertyHandle> StructPropertyHandle, class IDetailChildrenBuilder& ChildBuilder, IPropertyTypeCustomizationUtils& StructCustomizationUtils) override {}
 
diff --git a/Engine/Source/Editor/MovieSceneTools/Public/MovieSceneObjectBindingIDPicker.h b/Engine/Source/Editor/MovieSceneTools/Public/MovieSceneObjectBindingIDPicker.h
--- a/Engine/Source/Editor/MovieSceneTools/Public/MovieSceneObjectBindingIDPicker.h
+++ b/Engine/Source/Editor/MovieSceneTools/Public/MovieSceneObjectBindingIDPicker.h
@@ -39,6 +39,12 @@ public:
 		, bIsCurrentItemSpawnable(false)
 	{}
 
+
+	/**
+	 * Check whether this picker actually has anything to pick
+	 */
+	bool IsEmpty() const;
+
 protected:
 	virtual ~FMovieSceneObjectBindingIDPicker() { }
 
@@ -75,6 +81,9 @@ protected:
 	/** Build menu content that allows the user to choose a binding from inside the source sequence */
 	TSharedRef<SWidget> GetPickerMenu();
 
+	/** Build menu content that allows the user to choose a binding from inside the source sequence */
+	void GetPickerMenu(FMenuBuilder& MenuBuilder);
+
 	/** Get a widget that represents the currently chosen item */
 	TSharedRef<SWidget> GetCurrentItemWidget(TSharedRef<STextBlock> TextContent);
 
@@ -109,4 +118,7 @@ private:
 
 	/** Data tree that stores all the available bindings for the current sequence, and their identifiers */
 	TSharedPtr<FSequenceBindingTree> DataTree;
+
+	/** Weak ptr to a widget used to dismiss menus to */
+	TWeakPtr<SWidget> DismissWidget;
 };
\ No newline at end of file
diff --git a/Engine/Source/Editor/MovieSceneTools/Public/MovieSceneToolHelpers.h b/Engine/Source/Editor/MovieSceneTools/Public/MovieSceneToolHelpers.h
--- a/Engine/Source/Editor/MovieSceneTools/Public/MovieSceneToolHelpers.h
+++ b/Engine/Source/Editor/MovieSceneTools/Public/MovieSceneToolHelpers.h
@@ -6,10 +6,18 @@
 #include "Misc/Guid.h"
 #include "Misc/Attribute.h"
 #include "Widgets/SWidget.h"
+#include "Curves/RichCurve.h"
+#include "Math/InterpCurvePoint.h"
+#include "MovieSceneObjectBindingID.h"
+#include "MovieSceneObjectBindingIDPicker.h"
+#include "ISequencer.h"
 
 class ISequencer;
 class UMovieScene;
 class UMovieSceneSection;
+class UInterpTrackMoveAxis;
+struct FMovieSceneObjectBindingID;
+class UMovieSceneTrack;
 
 DECLARE_DELEGATE_TwoParams(FOnEnumSelectionChanged, int32 /*Selection*/, ESelectInfo::Type /*SelectionType*/);
 
@@ -92,6 +100,15 @@ public:
 	 */
 	static UObject* GetTake(const UMovieSceneSection* Section, uint32 TakeNumber);
 
+	/**
+	 * Get the next available row index for the section so that it doesn't overlap any other sections in time.
+	 *
+	 * @param InTrack The track to find the next available row on
+	 * @param InSection The section
+	 * @return The next available row index
+	 */
+	static int32 FindAvailableRowIndex(UMovieSceneTrack* InTrack, UMovieSceneSection* InSection);
+
 	/**
 	 * Generate a combobox for editing enum values
 	 *
@@ -132,4 +149,56 @@ public:
 	 * @return Whether the import was successful
 	 */
 	static bool ImportFBX(UMovieScene* InMovieScene, ISequencer& InSequencer, const TMap<FGuid, FString>& InObjectBindingNameMap);
+
+	/*
+	 * Rich curve interpolation to matinee interpolation
+	 *
+	 * @param InterpMode The rich curve interpolation to convert
+	 * @return The converted matinee interpolation
+	 */
+	static EInterpCurveMode RichCurveInterpolationToMatineeInterpolation( ERichCurveInterpMode InterpMode );
+
+	/*
+	 * Copy rich curve to move axis
+	 *
+	 * @param RichCurve The rich curve to copy from
+	 * @param MoveAxis The move axis to copy to
+	 */
+	static void CopyRichCurveToMoveAxis(const FRichCurve& RichCurve, UInterpTrackMoveAxis* MoveAxis);
+
+	/*
+	 * Export the object binding to a camera anim
+	 *
+	 * @param InMovieScene The movie scene to export the object binding from
+	 * @param InObjectBinding The object binding to export
+	 * @return The exported camera anim asset
+	 */
+	static UObject* ExportToCameraAnim(UMovieScene* InMovieScene, FGuid& InObjectBinding);
 };
+
+class FTrackEditorBindingIDPicker : public FMovieSceneObjectBindingIDPicker
+{
+public:
+	FTrackEditorBindingIDPicker(FMovieSceneSequenceID InLocalSequenceID, TWeakPtr<ISequencer> InSequencer)
+		: FMovieSceneObjectBindingIDPicker(InLocalSequenceID, InSequencer)
+	{
+		Initialize();
+	}
+
+	DECLARE_EVENT_OneParam(FTrackEditorBindingIDPicker, FOnBindingPicked, FMovieSceneObjectBindingID)
+	FOnBindingPicked& OnBindingPicked()
+	{
+		return OnBindingPickedEvent;
+	}
+
+	using FMovieSceneObjectBindingIDPicker::GetPickerMenu;
+
+private:
+
+	virtual UMovieSceneSequence* GetSequence() const override { return WeakSequencer.Pin()->GetFocusedMovieSceneSequence(); }
+	virtual void SetCurrentValue(const FMovieSceneObjectBindingID& InBindingId) override { OnBindingPickedEvent.Broadcast(InBindingId); }
+	virtual FMovieSceneObjectBindingID GetCurrentValue() const override { return FMovieSceneObjectBindingID(); }
+
+	FOnBindingPicked OnBindingPickedEvent;
+};
+
diff --git a/Engine/Source/Editor/MovieSceneTools/Public/MovieSceneToolsUserSettings.h b/Engine/Source/Editor/MovieSceneTools/Public/MovieSceneToolsUserSettings.h
--- a/Engine/Source/Editor/MovieSceneTools/Public/MovieSceneToolsUserSettings.h
+++ b/Engine/Source/Editor/MovieSceneTools/Public/MovieSceneToolsUserSettings.h
@@ -64,4 +64,12 @@ public:
 	/** Whether to create cameras if they don't already exist in the level. */
 	UPROPERTY(EditAnywhere, config, Category=Import)
 	bool bCreateCameras;
+
+	/** Whether to remove keyframes within a tolerance from the imported tracks */
+	UPROPERTY(EditAnywhere, config, Category=Import)
+	bool bReduceKeys;
+
+	/** The tolerance for reduce keys */
+	UPROPERTY(EditAnywhere, config, Category=Import, meta=(EditCondition=bReduceKeys))
+	float ReduceKeysTolerance;
 };
\ No newline at end of file
diff --git a/Engine/Source/Editor/MovieSceneTools/Public/TrackEditors/MaterialTrackEditor.h b/Engine/Source/Editor/MovieSceneTools/Public/TrackEditors/MaterialTrackEditor.h
--- a/Engine/Source/Editor/MovieSceneTools/Public/TrackEditors/MaterialTrackEditor.h
+++ b/Engine/Source/Editor/MovieSceneTools/Public/TrackEditors/MaterialTrackEditor.h
@@ -83,6 +83,7 @@ public:
 	// ISequencerTrackEditor interface
 
 	virtual bool SupportsType( TSubclassOf<UMovieSceneTrack> Type ) const override;
+	virtual bool GetDefaultExpansionState(UMovieSceneTrack* InTrack) const override;
 
 protected:
 
diff --git a/Engine/Source/Editor/SequenceRecorder/Private/ActorRecording.cpp b/Engine/Source/Editor/SequenceRecorder/Private/ActorRecording.cpp
--- a/Engine/Source/Editor/SequenceRecorder/Private/ActorRecording.cpp
+++ b/Engine/Source/Editor/SequenceRecorder/Private/ActorRecording.cpp
@@ -23,6 +23,7 @@ static const FName MovieSceneSectionRecorderFactoryName("MovieSceneSectionRecord
 UActorRecording::UActorRecording(const FObjectInitializer& ObjectInitializer)
 	: Super(ObjectInitializer)
 {
+	bActive = true;
 	bWasSpawnedPostRecord = false;
 	Guid.Invalidate();
 	bNewComponentAddedWhileRecording = false;
@@ -61,6 +62,11 @@ bool UActorRecording::IsRelevantForRecording(AActor* Actor)
 
 bool UActorRecording::StartRecording(ULevelSequence* CurrentSequence, float CurrentSequenceTime)
 {
+	if (!bActive)
+	{
+		return false;
+	}
+
 	bNewComponentAddedWhileRecording = false;
 	DuplicatedDynamicComponents.Reset();
 
@@ -221,7 +227,7 @@ bool UActorRecording::ValidComponent(USceneComponent* SceneComponent) const
 		const USequenceRecorderSettings* Settings = GetDefault<USequenceRecorderSettings>();
 		for (const FPropertiesToRecordForClass& PropertiesToRecordForClass : Settings->ClassesAndPropertiesToRecord)
 		{			
-			if (PropertiesToRecordForClass.Class != nullptr && SceneComponent->IsA(PropertiesToRecordForClass.Class))
+			if (PropertiesToRecordForClass.Class != nullptr && SceneComponent->IsA(PropertiesToRecordForClass.Class) && !SceneComponent->bIsEditorOnly)
 			{
 				return true;
 			}
diff --git a/Engine/Source/Editor/SequenceRecorder/Private/ActorRecording.h b/Engine/Source/Editor/SequenceRecorder/Private/ActorRecording.h
--- a/Engine/Source/Editor/SequenceRecorder/Private/ActorRecording.h
+++ b/Engine/Source/Editor/SequenceRecorder/Private/ActorRecording.h
@@ -85,6 +85,10 @@ public:
 	UPROPERTY(EditAnywhere, Category = "Actor Recording")
 	FActorRecordingSettings ActorSettings;
 
+	/** Whether this recording is active and to be recorded when the 'Record' button is pressed. */
+	UPROPERTY(EditAnywhere, Category = "Actor Recording")
+	bool bActive;
+
 	/** Whether we should specify the target animation or auto-create it */
 	UPROPERTY(EditAnywhere, Category = "Animation Recording")
 	bool bSpecifyTargetAnimation;
diff --git a/Engine/Source/Editor/SequenceRecorder/Private/SSequenceRecorder.cpp b/Engine/Source/Editor/SequenceRecorder/Private/SSequenceRecorder.cpp
--- a/Engine/Source/Editor/SequenceRecorder/Private/SSequenceRecorder.cpp
+++ b/Engine/Source/Editor/SequenceRecorder/Private/SSequenceRecorder.cpp
@@ -8,6 +8,8 @@
 #include "Framework/Commands/UICommandList.h"
 #include "Widgets/Notifications/SProgressBar.h"
 #include "Widgets/Text/STextBlock.h"
+#include "Widgets/Images/SImage.h"
+#include "Widgets/Input/SButton.h"
 #include "Framework/MultiBox/MultiBoxDefs.h"
 #include "Framework/MultiBox/MultiBoxBuilder.h"
 #include "Widgets/Views/SListView.h"
@@ -24,6 +26,7 @@
 
 #define LOCTEXT_NAMESPACE "SequenceRecorder"
 
+static const FName ActiveColumnName(TEXT("Active"));
 static const FName ActorColumnName(TEXT("Actor"));
 static const FName AnimationColumnName(TEXT("Animation"));
 static const FName LengthColumnName(TEXT("Length"));
@@ -53,7 +56,23 @@ public:
 	/** Overridden from SMultiColumnTableRow.  Generates a widget for this column of the list view. */
 	virtual TSharedRef<SWidget> GenerateWidgetForColumn(const FName& ColumnName) override
 	{
-		if (ColumnName == ActorColumnName)
+		if (ColumnName == ActiveColumnName)
+		{
+			return
+				SNew(SButton)
+				.ContentPadding(0)
+				.OnClicked(this, &SSequenceRecorderListRow::ToggleRecordingActive)
+				.ButtonStyle(FEditorStyle::Get(), "NoBorder")
+				.ToolTipText(LOCTEXT("ActiveButtonToolTip", "Toggle Recording Active"))
+				.HAlign(HAlign_Center)
+				.VAlign(VAlign_Center)
+				.Content()
+				[
+					SNew(SImage)
+					.Image(this, &SSequenceRecorderListRow::GetActiveBrushForRecording)
+				];
+		}
+		else if (ColumnName == ActorColumnName)
 		{
 			return
 				SNew(STextBlock)
@@ -76,6 +95,28 @@ public:
 	}
 
 private:
+
+	FReply ToggleRecordingActive()
+	{
+		if (RecordingPtr.IsValid())
+		{
+			 RecordingPtr.Get()->bActive = !RecordingPtr.Get()->bActive;
+		}
+		return FReply::Handled();
+	}
+
+	const FSlateBrush* GetActiveBrushForRecording() const
+	{
+		if (RecordingPtr.IsValid() && RecordingPtr.Get()->bActive)
+		{
+			return FEditorStyle::GetBrush("SequenceRecorder.Common.RecordingActive");
+		}
+		else
+		{
+			return FEditorStyle::GetBrush("SequenceRecorder.Common.RecordingInactive");
+		}
+	}
+
 	FText GetRecordingActorName() const
 	{
 		FText ActorName(LOCTEXT("InvalidActorName", "None"));
@@ -195,6 +236,9 @@ void SSequenceRecorder::Construct(const FArguments& Args)
 								.HeaderRow
 								(
 									SNew(SHeaderRow)
+									+ SHeaderRow::Column(ActiveColumnName)
+									.FillWidth(14.0f)
+									.DefaultLabel(LOCTEXT("ActiveColumnName", "Active"))
 									+ SHeaderRow::Column(ActorColumnName)
 									.FillWidth(43.0f)
 									.DefaultLabel(LOCTEXT("ActorHeaderName", "Actor"))
@@ -341,7 +385,7 @@ bool SSequenceRecorder::IsStopAllVisible() const
 
 void SSequenceRecorder::HandleAddRecording()
 {
-	FSequenceRecorder::Get().AddNewQueuedRecording();
+	FSequenceRecorder::Get().AddNewQueuedRecordingsForSelectedActors();
 }
 
 bool SSequenceRecorder::CanAddRecording() const
diff --git a/Engine/Source/Editor/SequenceRecorder/Private/Sections/MovieScene3DTransformSectionRecorder.cpp b/Engine/Source/Editor/SequenceRecorder/Private/Sections/MovieScene3DTransformSectionRecorder.cpp
--- a/Engine/Source/Editor/SequenceRecorder/Private/Sections/MovieScene3DTransformSectionRecorder.cpp
+++ b/Engine/Source/Editor/SequenceRecorder/Private/Sections/MovieScene3DTransformSectionRecorder.cpp
@@ -40,9 +40,10 @@ bool FMovieScene3DTransformSectionRecorderFactory::CanRecordObject(UObject* InOb
 	}
 }
 
-void FMovieScene3DTransformSectionRecorder::CreateSection(UObject* InObjectToRecord, UMovieScene* InMovieScene, const FGuid& Guid, float Time) 
+void FMovieScene3DTransformSectionRecorder::CreateSection(UObject* InObjectToRecord, UMovieScene* InMovieScene, const FGuid& InGuid, float Time) 
 {
 	ObjectToRecord = InObjectToRecord;
+	Guid = InGuid;
 	bWasAttached = false;
 		
 	MovieScene = InMovieScene;
@@ -55,7 +56,8 @@ void FMovieScene3DTransformSectionRecorder::CreateSection(UObject* InObjectToRec
 
 		const bool bUnwindRotation = false;
 
-		DefaultTransform = GetTransformToRecord();
+		DefaultTransform = FTransform::Identity;
+		GetTransformToRecord(DefaultTransform);
 		FVector EulerRotation = DefaultTransform.GetRotation().Rotator().Euler();
 
 		MovieSceneSection->SetDefault(FTransformKey(EKey3DTransformChannel::Translation, EAxis::X, DefaultTransform.GetTranslation().X, bUnwindRotation));
@@ -228,6 +230,12 @@ void FMovieScene3DTransformSectionRecorder::FinalizeSection()
 	MovieSceneSection->GetScaleCurve(EAxis::Y).SetKeys(ScaleYKeys);
 	MovieSceneSection->GetScaleCurve(EAxis::Z).SetKeys(ScaleZKeys);
 
+	FTransform FirstTransform = FTransform::Identity;
+	if (BufferedTransforms.Num())
+	{
+		FirstTransform = BufferedTransforms[0].Transform;
+	}
+
 	BufferedTransforms.Empty();
 
 	SlowTask.EnterProgressFrame();
@@ -278,6 +286,16 @@ void FMovieScene3DTransformSectionRecorder::FinalizeSection()
 	}
 
 	SlowTask.EnterProgressFrame();
+
+	// If recording a spawnable, update the spawnable object template to the first keyframe
+	if (MovieScene.IsValid() && Guid.IsValid())
+	{
+		FMovieSceneSpawnable* Spawnable = MovieScene->FindSpawnable(Guid);
+		if (Spawnable)
+		{
+			Spawnable->SpawnTransform = FirstTransform;
+		}
+	}
 }
 
 void FMovieScene3DTransformSectionRecorder::Record(float CurrentTime)
@@ -298,19 +316,24 @@ void FMovieScene3DTransformSectionRecorder::Record(float CurrentTime)
 		if(bRecording)
 		{
 			// don't record from the transform of the component/actor if we are synchronizing with an animation
-			if(!AnimRecorder.IsValid())
+			if (!AnimRecorder.IsValid())
 			{
-				BufferedTransforms.Add(FBufferedTransformKey(GetTransformToRecord(), CurrentTime));
+				FTransform TransformToRecord;
+				if (GetTransformToRecord(TransformToRecord))
+				{
+					BufferedTransforms.Add(FBufferedTransformKey(TransformToRecord, CurrentTime));
+				}
 			}
 		}
 	}
 }
 
-FTransform FMovieScene3DTransformSectionRecorder::GetTransformToRecord()
+bool FMovieScene3DTransformSectionRecorder::GetTransformToRecord(FTransform& TransformToRecord)
 {
 	if(USceneComponent* SceneComponent = Cast<USceneComponent>(ObjectToRecord.Get()))
 	{
-		return SceneComponent->GetRelativeTransform();
+		TransformToRecord = SceneComponent->GetRelativeTransform();
+		return true;
 	}
 	else if(AActor* Actor = Cast<AActor>(ObjectToRecord.Get()))
 	{
@@ -326,8 +349,22 @@ FTransform FMovieScene3DTransformSectionRecorder::GetTransformToRecord()
 			bCaptureWorldSpaceTransform = !FSequenceRecorder::Get().FindRecording(AttachParent->GetOwner());
 		}
 
-		return (bCaptureWorldSpaceTransform || !RootComponent) ? Actor->ActorToWorld() : RootComponent->GetRelativeTransform();
+		if (!RootComponent)
+		{
+			return false;
+		}
+
+		if (bCaptureWorldSpaceTransform)
+		{
+			TransformToRecord = Actor->ActorToWorld();
+			return true;
+		}
+		else
+		{
+			TransformToRecord = RootComponent->GetRelativeTransform();
+			return true;
+		}
 	}
 
-	return FTransform::Identity;
+	return false;
 }
diff --git a/Engine/Source/Editor/SequenceRecorder/Private/Sections/MovieScene3DTransformSectionRecorder.h b/Engine/Source/Editor/SequenceRecorder/Private/Sections/MovieScene3DTransformSectionRecorder.h
--- a/Engine/Source/Editor/SequenceRecorder/Private/Sections/MovieScene3DTransformSectionRecorder.h
+++ b/Engine/Source/Editor/SequenceRecorder/Private/Sections/MovieScene3DTransformSectionRecorder.h
@@ -66,7 +66,7 @@ public:
 	}
 
 private:
-	FTransform GetTransformToRecord();
+	bool GetTransformToRecord(FTransform& TransformToRecord);
 
 private:
 	/** Object to record from */
@@ -95,4 +95,7 @@ private:
 
 	/** Flag indicating that some time while this recorder was active an attachment was also in place */
 	bool bWasAttached;
+
+	/** The guid being recorded to */
+	FGuid Guid;
 };
diff --git a/Engine/Source/Editor/SequenceRecorder/Private/SequenceRecorder.cpp b/Engine/Source/Editor/SequenceRecorder/Private/SequenceRecorder.cpp
--- a/Engine/Source/Editor/SequenceRecorder/Private/SequenceRecorder.cpp
+++ b/Engine/Source/Editor/SequenceRecorder/Private/SequenceRecorder.cpp
@@ -23,6 +23,7 @@
 #include "ObjectTools.h"
 #include "Features/IModularFeatures.h"
 #include "Engine/DemoNetDriver.h"
+#include "Engine/Selection.h"
 #include "Framework/Notifications/NotificationManager.h"
 #include "Widgets/Notifications/SNotificationList.h"
 #include "LevelSequenceActor.h"
@@ -125,6 +126,35 @@ void FSequenceRecorder::StopAllQueuedRecordings()
 	}
 }
 
+void FSequenceRecorder::AddNewQueuedRecordingsForSelectedActors()
+{
+	bool bAnySelectedActorsAdded = false;
+	TArray<AActor*> EntireSelection;
+
+	GEditor->GetSelectedActors()->GetSelectedObjects(EntireSelection);
+
+	TArray<AActor*> ActorsToRecord;
+	for (AActor* Actor : EntireSelection)
+	{
+		if (!FindRecording(Actor))
+		{
+			if (UActorRecording* NewRecording = AddNewQueuedRecording(Actor))
+			{
+				bAnySelectedActorsAdded = true;
+			}
+		}
+		else
+		{
+			bAnySelectedActorsAdded = true;
+		}
+	}
+
+	if (!bAnySelectedActorsAdded)
+	{
+		AddNewQueuedRecording();
+	}
+}
+
 UActorRecording* FSequenceRecorder::AddNewQueuedRecording(AActor* Actor, UAnimSequence* AnimSequence, float Length)
 {
 	const USequenceRecorderSettings* Settings = GetDefault<USequenceRecorderSettings>();
@@ -230,27 +260,6 @@ void FSequenceRecorder::Tick(float DeltaSeconds)
 
 	const USequenceRecorderSettings* Settings = GetDefault<USequenceRecorderSettings>();
 
-	// check for spawned actors and if they have begun playing yet
-	for(TWeakObjectPtr<AActor>& QueuedSpawnedActor : QueuedSpawnedActors)
-	{
-		if(AActor* Actor = QueuedSpawnedActor.Get())
-		{
-			if(Actor->HasActorBegunPlay())
-			{
-				if(UActorRecording::IsRelevantForRecording(Actor) && IsActorValidForRecording(Actor))
-				{
-					UActorRecording* NewRecording = AddNewQueuedRecording(Actor);
-					NewRecording->bWasSpawnedPostRecord = true;
-					NewRecording->StartRecording(CurrentSequence.Get(), CurrentTime);
-				}
-
-				QueuedSpawnedActor.Reset();
-			}
-		}
-	}
-
-	QueuedSpawnedActors.RemoveAll([](TWeakObjectPtr<AActor>& QueuedSpawnedActor) { return !QueuedSpawnedActor.IsValid(); } );
-
 	FAnimationRecorderManager::Get().Tick(DeltaSeconds);
 
 	for(UActorRecording* Recording : QueuedRecordings)
@@ -422,11 +431,8 @@ void FSequenceRecorder::DrawDebug(UCanvas* InCanvas, APlayerController* InPlayer
 	}
 }
 
-bool FSequenceRecorder::StartRecording(const FOnRecordingStarted& OnRecordingStarted, const FOnRecordingFinished& OnRecordingFinished, const FString& InPathToRecordTo, const FString& InSequenceName)
+bool FSequenceRecorder::StartRecording(const FString& InPathToRecordTo, const FString& InSequenceName)
 {
-	OnRecordingStartedDelegate = OnRecordingStarted;
-	OnRecordingFinishedDelegate = OnRecordingFinished;
-
 	const USequenceRecorderSettings* Settings = GetDefault<USequenceRecorderSettings>();
 
 	if(InPathToRecordTo.Len() > 0)
@@ -506,14 +512,18 @@ bool FSequenceRecorder::StartRecordingInternal(UWorld* World)
 	const USequenceRecorderSettings* Settings = GetDefault<USequenceRecorderSettings>();
 
 	UWorld* ActorWorld = nullptr;
+	if(World != nullptr || (QueuedRecordings.Num() > 0 && QueuedRecordings[0]->GetActorToRecord() != nullptr))
+	{
+		ActorWorld = World != nullptr ? World : QueuedRecordings[0]->GetActorToRecord()->GetWorld();
+	}
+
 	if(Settings->bRecordWorldSettingsActor)
 	{
-		if(World != nullptr || (QueuedRecordings.Num() > 0 && QueuedRecordings[0]->GetActorToRecord() != nullptr))
+		if(ActorWorld)
 		{
-			ActorWorld = World != nullptr ? World : QueuedRecordings[0]->GetActorToRecord()->GetWorld();
-			if(ActorWorld)
+			AWorldSettings* WorldSettings = ActorWorld->GetWorldSettings();
+			if (!IsRecordingQueued(WorldSettings))
 			{
-				AWorldSettings* WorldSettings = ActorWorld->GetWorldSettings();
 				AddNewQueuedRecording(WorldSettings);
 			}
 		}
@@ -526,14 +536,39 @@ bool FSequenceRecorder::StartRecordingInternal(UWorld* World)
 		{
 			// find a counterpart in the PIE world if this actor is not
 			ALevelSequenceActor* ActorToTrigger = LevelSequenceActor.Get();
-			if(ActorWorld && ActorWorld->IsPlayInEditor() && !ActorToTrigger->GetWorld()->IsPlayInEditor())
+			if(!ActorToTrigger->GetWorld()->IsPlayInEditor())
 			{
 				ActorToTrigger = Cast<ALevelSequenceActor>(EditorUtilities::GetSimWorldCounterpartActor(ActorToTrigger));
 			}
 
-			if(ActorToTrigger && ActorToTrigger->SequencePlayer)
+			if(ActorToTrigger)
 			{
-				ActorToTrigger->SequencePlayer->Play();
+				// Duplicate the level sequence we want to trigger so that we can playback the level sequence and record to it at the same time
+				ALevelSequenceActor* DupActorToTrigger = ActorToTrigger->GetWorld()->SpawnActor<ALevelSequenceActor>();
+				if (!DupActorToTrigger)
+				{
+					UE_LOG(LogAnimation, Display, TEXT("Unabled to spawn actor to trigger: (%s)"), *ActorToTrigger->GetPathName());
+					continue;
+				}
+
+				ULevelSequence* DupLevelSequence = CastChecked<ULevelSequence>(StaticDuplicateObject(ActorToTrigger->GetSequence(), ActorToTrigger->GetOuter(), NAME_None, RF_AllFlags & ~RF_Transactional));
+				DupActorToTrigger->SetSequence(DupLevelSequence);
+
+				// Always initialize the player so that the playback settings/range can be initialized from editor.
+				DupActorToTrigger->InitializePlayer();
+				
+				if (DupActorToTrigger->SequencePlayer)
+				{
+					DupActorToTrigger->SequencePlayer->SetDisableCameraCuts(true);
+					DupActorToTrigger->SequencePlayer->SetPlaybackPosition(0.f);
+					DupActorToTrigger->SequencePlayer->Play();
+				}
+				else
+				{
+					UE_LOG(LogAnimation, Display, TEXT("Level sequence (%s) is not initialized for playback"), *ActorToTrigger->GetPathName());
+				}
+
+				DupActorsToTrigger.Add(DupActorToTrigger);
 			}
 		}
 	}
@@ -567,7 +602,7 @@ bool FSequenceRecorder::StartRecordingInternal(UWorld* World)
 					FDelegateHandle* FoundHandle = ActorSpawningDelegateHandles.Find(ActorToRecordWorld);
 					if(FoundHandle == nullptr)
 					{
-						FDelegateHandle NewHandle = ActorToRecordWorld->AddOnActorSpawnedHandler(FOnActorSpawned::FDelegate::CreateStatic(&FSequenceRecorder::HandleActorSpawned));
+						FDelegateHandle NewHandle = ActorToRecordWorld->AddOnActorSpawnedHandler(FOnActorSpawned::FDelegate::CreateRaw(this, &FSequenceRecorder::HandleActorSpawned));
 						ActorSpawningDelegateHandles.Add(ActorToRecordWorld, NewHandle);
 					}
 				}
@@ -575,17 +610,16 @@ bool FSequenceRecorder::StartRecordingInternal(UWorld* World)
 		}
 
 		// start recording 
-		bool bStartedRecordingAllActors = true;
+		bool bAnyRecordingsStarted = false;
 		for(UActorRecording* Recording : QueuedRecordings)
 		{
-			if(!Recording->StartRecording(CurrentSequence.Get(), CurrentTime))
+			if(Recording->StartRecording(CurrentSequence.Get(), CurrentTime))
 			{
-				bStartedRecordingAllActors = false;
-				break;
+				bAnyRecordingsStarted = true;
 			}
 		}
 
-		if(!bStartedRecordingAllActors)
+		if(!bAnyRecordingsStarted)
 		{
 			// if we couldnt start a recording, stop all others
 			TArray<FAssetData> AssetsToCleanUp;
@@ -653,7 +687,11 @@ bool FSequenceRecorder::StartRecordingInternal(UWorld* World)
 			}
 		}
 
-		OnRecordingStartedDelegate.ExecuteIfBound(CurrentSequence.Get());
+		if (OnRecordingStartedDelegate.IsBound())
+		{
+			OnRecordingStartedDelegate.Broadcast(CurrentSequence.Get());
+		}
+
 		return true;
 	}
 	else
@@ -675,6 +713,11 @@ void FSequenceRecorder::HandleEndPIE(bool bSimulating)
 
 bool FSequenceRecorder::StopRecording()
 {
+	if (!IsRecording())
+	{
+		return false;
+	}
+
 	const USequenceRecorderSettings* Settings = GetDefault<USequenceRecorderSettings>();
 
 	if (Settings->bImmersiveMode)
@@ -773,6 +816,41 @@ bool FSequenceRecorder::StopRecording()
 
 	DeadRecordings.Empty();
 
+	// Remove any spawned recordings
+	for (int32 QueuedRecordingIndex = 0; QueuedRecordingIndex < QueuedRecordings.Num(); )
+	{
+		if (QueuedRecordings[QueuedRecordingIndex]->bWasSpawnedPostRecord)
+		{
+			RemoveQueuedRecording(QueuedRecordings[QueuedRecordingIndex]);
+		}
+		else
+		{
+			QueuedRecordingIndex++;
+		}
+	}
+
+
+	// Stop any level sequences that were triggered
+	for(int32 DupActorToTriggerIndex = 0; DupActorToTriggerIndex < DupActorsToTrigger.Num(); ++DupActorToTriggerIndex )
+	{
+		if (DupActorsToTrigger[DupActorToTriggerIndex].IsValid())
+		{
+			ALevelSequenceActor* DupActorToTrigger = DupActorsToTrigger[DupActorToTriggerIndex].Get();
+			ULevelSequencePlayer* SequencePlayer = DupActorToTrigger->SequencePlayer;
+			if (SequencePlayer)
+			{
+				SequencePlayer->SetDisableCameraCuts(false);
+				SequencePlayer->Stop();
+			}
+			if (DupActorToTrigger->GetWorld())
+			{
+				DupActorToTrigger->GetWorld()->DestroyActor(DupActorToTrigger);
+			}
+		}
+	}
+
+	DupActorsToTrigger.Empty();
+
 	if(Settings->bCreateLevelSequence)
 	{
 		if(LevelSequence)
@@ -816,10 +894,10 @@ bool FSequenceRecorder::StopRecording()
 
 			UE_LOG(LogAnimation, Display, TEXT("Stopped recording sequence %s"), *LevelSequence->GetPathName());
 
-			OnRecordingFinishedDelegate.ExecuteIfBound(LevelSequence);
-
-			OnRecordingFinishedDelegate = FOnRecordingFinished();
-			OnRecordingStartedDelegate = FOnRecordingStarted();
+			if (OnRecordingFinishedDelegate.IsBound())
+			{
+				OnRecordingFinishedDelegate.Broadcast(LevelSequence);
+			}
 
 			return true;
 		}
@@ -922,8 +1000,12 @@ void FSequenceRecorder::HandleActorSpawned(AActor* Actor)
 
 	if(Actor && FSequenceRecorder::Get().IsRecording())
 	{
-		// Queue a possible recording - we need to wait until the actor has begun playing
-		FSequenceRecorder::Get().QueuedSpawnedActors.Add(Actor);
+		if (UActorRecording::IsRelevantForRecording(Actor) || IsActorValidForRecording(Actor))
+		{
+			UActorRecording* NewRecording = AddNewQueuedRecording(Actor);
+			NewRecording->bWasSpawnedPostRecord = true;
+			NewRecording->StartRecording(CurrentSequence.Get(), CurrentTime);
+		}
 	}
 }
 
diff --git a/Engine/Source/Editor/SequenceRecorder/Private/SequenceRecorder.h b/Engine/Source/Editor/SequenceRecorder/Private/SequenceRecorder.h
--- a/Engine/Source/Editor/SequenceRecorder/Private/SequenceRecorder.h
+++ b/Engine/Source/Editor/SequenceRecorder/Private/SequenceRecorder.h
@@ -10,6 +10,7 @@
 #include "Sections/MovieSceneMultiPropertyRecorder.h"
 
 class APlayerController;
+class ALevelSequenceActor;
 class ISequenceAudioRecorder;
 class UCanvas;
 class UTexture;
@@ -29,7 +30,7 @@ public:
 	void Shutdown();
 
 	/** Starts recording a sequence. */
-	bool StartRecording(const FOnRecordingStarted& OnRecordingStarted = FOnRecordingStarted(), const FOnRecordingFinished& OnRecordingFinished = FOnRecordingFinished(), const FString& InPathToRecordTo = FString(), const FString& InSequenceName = FString());
+	bool StartRecording(const FString& InPathToRecordTo = FString(), const FString& InSequenceName = FString());
 
 	/** Starts recording a sequence for the specified world. */
 	bool StartRecordingForReplay(UWorld* World, const struct FSequenceRecorderActorFilter& ActorFilter);
@@ -58,6 +59,8 @@ public:
 
 	void StopRecordingDeadAnimations();
 
+	void AddNewQueuedRecordingsForSelectedActors();
+
 	class UActorRecording* AddNewQueuedRecording(AActor* Actor = nullptr, UAnimSequence* AnimSequence = nullptr, float Length = 0.0f);
 
 	void RemoveQueuedRecording(AActor* Actor);
@@ -80,7 +83,7 @@ public:
 	void DrawDebug(UCanvas* InCanvas, APlayerController* InPlayerController);
 
 	/** Handle actors being spawned */
-	static void HandleActorSpawned(AActor* Actor);
+	void HandleActorSpawned(AActor* Actor);
 
 	/** Handle actors being de-spawned */
 	void HandleActorDespawned(AActor* Actor);
@@ -103,6 +106,12 @@ public:
 	/** Refresh the name of the next sequence we will be recording */
 	void RefreshNextSequence();
 
+	/** Multicast delegate fired when recording is started */
+	FOnRecordingStarted OnRecordingStartedDelegate;
+
+	/** Multicast delegate fired when recording has finished */
+	FOnRecordingFinished OnRecordingFinishedDelegate;
+
 private:
 	/** Starts recording a sequence, possibly delayed */
 	bool StartRecordingInternal(UWorld* World);
@@ -130,8 +139,6 @@ private:
 
 	TArray<class UActorRecording*> DeadRecordings;
 
-	TArray<TWeakObjectPtr<AActor>> QueuedSpawnedActors;
-
 	bool bQueuedRecordingsDirty;
 
 	bool bWasImmersive;
@@ -151,12 +158,6 @@ private:
 	/** Texture we use for the recording indicator */
 	TWeakObjectPtr<UTexture> RecordingIndicatorTexture;
 
-	/** Delegate fired when recording is started */
-	FOnRecordingStarted OnRecordingStartedDelegate;
-
-	/** Delegate fired when recording has finished */
-	FOnRecordingFinished OnRecordingFinishedDelegate;
-
 	/** Cached sequence name to record to */
 	FString SequenceName;
 
@@ -177,4 +178,7 @@ private:
 
 	/** Built in multi-property recorder */
 	FMovieSceneMultiPropertyRecorderFactory MultiPropertySectionRecorder;
+
+	/** Duplicated level sequence actors to trigger, to be stopped at the end of recording */
+	TArray<TWeakObjectPtr<ALevelSequenceActor>> DupActorsToTrigger;
 };
diff --git a/Engine/Source/Editor/SequenceRecorder/Private/SequenceRecorderCommands.cpp b/Engine/Source/Editor/SequenceRecorder/Private/SequenceRecorderCommands.cpp
--- a/Engine/Source/Editor/SequenceRecorder/Private/SequenceRecorderCommands.cpp
+++ b/Engine/Source/Editor/SequenceRecorder/Private/SequenceRecorderCommands.cpp
@@ -14,7 +14,7 @@ void FSequenceRecorderCommands::RegisterCommands()
 {
 	UI_COMMAND(RecordAll, "Record", "Record all recordings", EUserInterfaceActionType::Button, FInputChord(EModifierKey::Control | EModifierKey::Shift, EKeys::R));
 	UI_COMMAND(StopAll, "StopAll", "Stop all recordings", EUserInterfaceActionType::Button, FInputChord(EModifierKey::Control | EModifierKey::Shift, EKeys::S));
-	UI_COMMAND(AddRecording, "Add", "Add a new recording", EUserInterfaceActionType::Button, FInputChord());
+	UI_COMMAND(AddRecording, "Add", "Add a new recording for selected actors. If nothing selected, an empty recording will be added", EUserInterfaceActionType::Button, FInputChord());
 	UI_COMMAND(RemoveRecording, "Remove", "Remove selected recording", EUserInterfaceActionType::Button, FInputChord());
 	UI_COMMAND(RemoveAllRecordings, "Remove All", "Remove all recordings", EUserInterfaceActionType::Button, FInputChord());
 }
diff --git a/Engine/Source/Editor/SequenceRecorder/Private/SequenceRecorderModule.cpp b/Engine/Source/Editor/SequenceRecorder/Private/SequenceRecorderModule.cpp
--- a/Engine/Source/Editor/SequenceRecorder/Private/SequenceRecorderModule.cpp
+++ b/Engine/Source/Editor/SequenceRecorder/Private/SequenceRecorderModule.cpp
@@ -558,7 +558,7 @@ class FSequenceRecorderModule : public ISequenceRecorder, private FSelfRegisteri
 		return CurrentSequence.IsValid() ? CurrentSequence->GetMovieScene()->GetPlaybackRange().Size<float>() : 0.0f;
 	}
 
-	virtual bool StartRecording(TArrayView<AActor* const> ActorsToRecord, const FOnRecordingStarted& OnRecordingStarted, const FOnRecordingFinished& OnRecordingFinished, const FString& PathToRecordTo, const FString& SequenceName) override
+	virtual bool StartRecording(TArrayView<AActor* const> ActorsToRecord, const FString& PathToRecordTo, const FString& SequenceName) override
 	{
 		if(ActorsToRecord.Num() != 0)
 		{
@@ -581,12 +581,12 @@ class FSequenceRecorderModule : public ISequenceRecorder, private FSelfRegisteri
 			UE_LOG(LogAnimation, Display, TEXT("Couldn't find actor to record"));
 		}
 
-		return FSequenceRecorder::Get().StartRecording(OnRecordingStarted, OnRecordingFinished, PathToRecordTo, SequenceName);
+		return FSequenceRecorder::Get().StartRecording(PathToRecordTo, SequenceName);
 	}
 
 	virtual void NotifyActorStartRecording(AActor* Actor)
 	{
-		FSequenceRecorder::HandleActorSpawned(Actor);
+		FSequenceRecorder::Get().HandleActorSpawned(Actor);
 	}
 
 	virtual void NotifyActorStopRecording(AActor* Actor)
@@ -633,6 +633,10 @@ class FSequenceRecorderModule : public ISequenceRecorder, private FSelfRegisteri
 		return AudioFactory ? AudioFactory() : TUniquePtr<ISequenceAudioRecorder>();
 	}
 
+	virtual FOnRecordingStarted& OnRecordingStarted() override { return FSequenceRecorder::Get().OnRecordingStartedDelegate; }
+
+	virtual FOnRecordingFinished& OnRecordingFinished() override { return FSequenceRecorder::Get().OnRecordingFinishedDelegate; }
+
 	static void TickSequenceRecorder(float DeltaSeconds)
 	{
 		if (!IsRunningDedicatedServer() && !IsRunningCommandlet())
diff --git a/Engine/Source/Editor/SequenceRecorder/Public/ISequenceRecorder.h b/Engine/Source/Editor/SequenceRecorder/Public/ISequenceRecorder.h
--- a/Engine/Source/Editor/SequenceRecorder/Public/ISequenceRecorder.h
+++ b/Engine/Source/Editor/SequenceRecorder/Public/ISequenceRecorder.h
@@ -10,9 +10,9 @@
 class AActor;
 class ISequenceAudioRecorder;
 
-DECLARE_DELEGATE_OneParam(FOnRecordingStarted, class UMovieSceneSequence* /*Sequence*/);
+DECLARE_MULTICAST_DELEGATE_OneParam(FOnRecordingStarted, class UMovieSceneSequence* /*Sequence*/);
 
-DECLARE_DELEGATE_OneParam(FOnRecordingFinished, class UMovieSceneSequence* /*Sequence*/);
+DECLARE_MULTICAST_DELEGATE_OneParam(FOnRecordingFinished, class UMovieSceneSequence* /*Sequence*/);
 
 class ISequenceRecorder : public IModuleInterface
 {
@@ -37,26 +37,22 @@ public:
 	/**
 	 * Start a recording, possibly with some delay (specified by the sequence recording settings).
 	 * @param	ActorsToRecord		Actors to record.
-	 * @param	OnRecordingStarted	Delegate fired when recording has commenced.
-	 * @param	OnRecordingFinished	Delegate fired when recording has finished.
 	 * @param	PathToRecordTo		Optional path to a sequence to record to. If none is specified we use the defaults in the settings.
 	 * @param	SequenceName		Optional name of a sequence to record to. If none is specified we use the defaults in the settings.
 	 * @return true if recording was successfully started
 	*/
-	virtual bool StartRecording(TArrayView<AActor* const> ActorsToRecord, const FOnRecordingStarted& OnRecordingStarted, const FOnRecordingFinished& OnRecordingFinished, const FString& PathToRecordTo = FString(), const FString& SequenceName = FString()) = 0;
+	virtual bool StartRecording(TArrayView<AActor* const> ActorsToRecord, const FString& PathToRecordTo = FString(), const FString& SequenceName = FString()) = 0;
 
 	/**
 	 * Start a recording, possibly with some delay (specified by the sequence recording settings).
 	 * @param	ActorToRecord		Actor to record.
-	 * @param	OnRecordingStarted	Delegate fired when recording has commenced.
-	 * @param	OnRecordingFinished	Delegate fired when recording has finished.
 	 * @param	PathToRecordTo		Optional path to a sequence to record to. If none is specified we use the defaults in the settings.
 	 * @param	SequenceName		Optional name of a sequence to record to. If none is specified we use the defaults in the settings.
 	 * @return true if recording was successfully started
 	*/
-	bool StartRecording(AActor* ActorToRecord, const FOnRecordingStarted& OnRecordingStarted, const FOnRecordingFinished& OnRecordingFinished, const FString& PathToRecordTo = FString(), const FString& SequenceName = FString())
+	bool StartRecording(AActor* ActorToRecord, const FString& PathToRecordTo = FString(), const FString& SequenceName = FString())
 	{
-		return StartRecording(MakeArrayView(&ActorToRecord, 1), OnRecordingStarted, OnRecordingFinished, PathToRecordTo, SequenceName);
+		return StartRecording(MakeArrayView(&ActorToRecord, 1), PathToRecordTo, SequenceName);
 	}
 
 	/**
@@ -105,4 +101,10 @@ public:
 	 * @return A valid ptr to an audio recorder or null
 	 */
 	virtual TUniquePtr<ISequenceAudioRecorder> CreateAudioRecorder() const = 0;
+
+	/** Get the sequence recorder started delegate */
+	virtual FOnRecordingStarted& OnRecordingStarted() = 0;
+
+	/** Get the sequence recorder finished delegate */
+	virtual FOnRecordingFinished& OnRecordingFinished() = 0;
 };
diff --git a/Engine/Source/Editor/SequenceRecorderSections/Private/MovieSceneSpawnSectionRecorder.cpp b/Engine/Source/Editor/SequenceRecorderSections/Private/MovieSceneSpawnSectionRecorder.cpp
--- a/Engine/Source/Editor/SequenceRecorderSections/Private/MovieSceneSpawnSectionRecorder.cpp
+++ b/Engine/Source/Editor/SequenceRecorderSections/Private/MovieSceneSpawnSectionRecorder.cpp
@@ -6,6 +6,7 @@
 #include "Sections/MovieSceneBoolSection.h"
 #include "Tracks/MovieSceneSpawnTrack.h"
 #include "MovieScene.h"
+#include "SequenceRecorderSettings.h"
 
 TSharedPtr<IMovieSceneSectionRecorder> FMovieSceneSpawnSectionRecorderFactory::CreateSectionRecorder(const struct FActorRecordingSettings& InActorRecordingSettings) const
 {
@@ -46,6 +47,14 @@ void FMovieSceneSpawnSectionRecorder::FinalizeSection()
 	{
 		MovieSceneSection->AddKey(MovieSceneSection->GetEndTime(), bSpawned, EMovieSceneKeyInterpolation::Break);
 	}
+
+	// If the section is degenerate, assume the actor was spawned and destroyed. Give it a 1 frame spawn section.
+	if (MovieSceneSection->GetRange().IsDegenerate())
+	{
+		float OneFrameInterval = 1.0f/GetDefault<USequenceRecorderSettings>()->DefaultAnimationSettings.SampleRate;
+		MovieSceneSection->AddKey(MovieSceneSection->GetEndTime() - OneFrameInterval, true, EMovieSceneKeyInterpolation::Break);
+		MovieSceneSection->SetStartTime(MovieSceneSection->GetEndTime() - OneFrameInterval);
+	}
 }
 
 void FMovieSceneSpawnSectionRecorder::Record(float CurrentTime)
diff --git a/Engine/Source/Editor/Sequencer/Private/DisplayNodes/SequencerDisplayNode.cpp b/Engine/Source/Editor/Sequencer/Private/DisplayNodes/SequencerDisplayNode.cpp
--- a/Engine/Source/Editor/Sequencer/Private/DisplayNodes/SequencerDisplayNode.cpp
+++ b/Engine/Source/Editor/Sequencer/Private/DisplayNodes/SequencerDisplayNode.cpp
@@ -632,17 +632,19 @@ void FSequencerDisplayNode::BuildContextMenu(FMenuBuilder& MenuBuilder)
 		
 		MenuBuilder.AddMenuEntry(FGenericCommands::Get().Paste);
 		
+		MenuBuilder.AddMenuEntry(FGenericCommands::Get().Duplicate);
+		
 		MenuBuilder.AddMenuEntry(
 			LOCTEXT("DeleteNode", "Delete"),
 			LOCTEXT("DeleteNodeTooltip", "Delete this or selected tracks"),
-			FSlateIcon(),
+			FSlateIcon(FEditorStyle::GetStyleSetName(), "ContentBrowser.AssetActions.Delete"),
 			FUIAction(FExecuteAction::CreateSP(&GetSequencer(), &FSequencer::DeleteNode, ThisNode), CanExecute)
 		);
 
 		MenuBuilder.AddMenuEntry(
 			LOCTEXT("RenameNode", "Rename"),
 			LOCTEXT("RenameNodeTooltip", "Rename this track"),
-			FSlateIcon(),
+			FSlateIcon(FEditorStyle::GetStyleSetName(), "ContentBrowser.AssetActions.Rename"),
 			FUIAction(
 				FExecuteAction::CreateSP(this, &FSequencerDisplayNode::HandleContextMenuRenameNodeExecute),
 				FCanExecuteAction::CreateSP(this, &FSequencerDisplayNode::HandleContextMenuRenameNodeCanExecute)
diff --git a/Engine/Source/Editor/Sequencer/Private/DisplayNodes/SequencerObjectBindingNode.cpp b/Engine/Source/Editor/Sequencer/Private/DisplayNodes/SequencerObjectBindingNode.cpp
--- a/Engine/Source/Editor/Sequencer/Private/DisplayNodes/SequencerObjectBindingNode.cpp
+++ b/Engine/Source/Editor/Sequencer/Private/DisplayNodes/SequencerObjectBindingNode.cpp
@@ -196,6 +196,14 @@ void FSequencerObjectBindingNode::BuildContextMenu(FMenuBuilder& MenuBuilder)
 			FUIAction(
 				FExecuteAction::CreateLambda([=]{ GetSequencer().ExportFBX(); })
 			));
+
+		MenuBuilder.AddMenuEntry(
+			LOCTEXT("Export to Camera Anim", "Export to Camera Anim..."),
+			LOCTEXT("ExportToCameraAnimTooltip", "Exports the animation to a camera anim asset"),
+			FSlateIcon(),
+			FUIAction(
+				FExecuteAction::CreateLambda([=]{ GetSequencer().ExportToCameraAnim(); })
+			));
 	}
 
 	MenuBuilder.BeginSection("Organize", LOCTEXT("OrganizeContextMenuSectionName", "Organize"));
diff --git a/Engine/Source/Editor/Sequencer/Private/DisplayNodes/SequencerTrackNode.cpp b/Engine/Source/Editor/Sequencer/Private/DisplayNodes/SequencerTrackNode.cpp
--- a/Engine/Source/Editor/Sequencer/Private/DisplayNodes/SequencerTrackNode.cpp
+++ b/Engine/Source/Editor/Sequencer/Private/DisplayNodes/SequencerTrackNode.cpp
@@ -324,68 +324,56 @@ void FSequencerTrackNode::SetDisplayName(const FText& NewDisplayName)
 	}
 }
 
-struct FOverlappingCompileRules : FMovieSceneSegmentCompilerRules
-{
-	int32 PredicatePriority;
-	FOverlappingCompileRules(int32 InPredicatePriority) : PredicatePriority(InPredicatePriority) {}
-
-	virtual void BlendSegment(FMovieSceneSegment& Segment, const TArrayView<const FMovieSceneSectionData>& SourceData) const
-	{
-		// If there is anything on top of this section in this range, ignore it completely (that section will render everything underneath it) 
-		auto IsUnderneathAnything = [=](const FSectionEvaluationData& In) { return SourceData[In.ImplIndex].Priority > PredicatePriority; };
-
-		if (Segment.Impls.ContainsByPredicate(IsUnderneathAnything))
-		{
-			Segment.Impls.Reset();
-		}
-		else if (Segment.Impls.Num() > 1)
-		{
-			// Sort lowest to highest
-			Segment.Impls.Sort([=](const FSectionEvaluationData& A, const FSectionEvaluationData& B){
-				return SourceData[A.ImplIndex].Priority < SourceData[B.ImplIndex].Priority;
-			});
-		}
-	}
-};
-
 TArray<FSequencerOverlapRange> FSequencerTrackNode::GetUnderlappingSections(UMovieSceneSection* InSection)
 {
 	TRange<float> InSectionRange = InSection->IsInfinite() ? TRange<float>::All() : InSection->GetRange();
 
-	TArray<FMovieSceneSectionData> CompileData;
+	TMovieSceneEvaluationTree<int32> SectionIndexTree;
+
+	// Iterate all other sections on the same row with <= overlap priority
 	for (int32 SectionIndex = 0; SectionIndex < Sections.Num(); ++SectionIndex)
 	{
 		UMovieSceneSection* SectionObj = Sections[SectionIndex]->GetSectionObject();
-		if (!SectionObj || SectionObj == InSection || SectionObj->GetRowIndex() != InSection->GetRowIndex())
+		if (!SectionObj || SectionObj == InSection || SectionObj->GetRowIndex() != InSection->GetRowIndex() || SectionObj->GetOverlapPriority() > InSection->GetOverlapPriority())
 		{
 			continue;
 		}
 
-		
 		TRange<float> OtherSectionRange = SectionObj->IsInfinite() ? TRange<float>::All() : SectionObj->GetRange();
 		TRange<float> Intersection = TRange<float>::Intersection(OtherSectionRange, InSectionRange);
 		if (!Intersection.IsEmpty())
 		{
-			CompileData.Add(FMovieSceneSectionData(Intersection, FSectionEvaluationData(SectionIndex), FOptionalMovieSceneBlendType(), SectionObj->GetOverlapPriority()));
+			SectionIndexTree.Add(Intersection, SectionIndex);
 		}
 	}
 
-	FOverlappingCompileRules Rules(InSection->GetOverlapPriority());
+	TSharedRef<FSequencerTrackNode> TrackNode = SharedThis(this);;
 
 	TArray<FSequencerOverlapRange> Result;
-
-	for (const FMovieSceneSegment& Segment : FMovieSceneSegmentCompiler().Compile(CompileData, &Rules))
+	for (FMovieSceneEvaluationTreeRangeIterator It(SectionIndexTree); It; ++It)
 	{
-		Result.Emplace();
+		FSequencerOverlapRange NewRange;
+
+		NewRange.Range = It.Range();
+
+		for (int32 SectionIndex : SectionIndexTree.GetAllData(It.Node()))
+		{
+			NewRange.Sections.Add(FSectionHandle(TrackNode, SectionIndex));
+		}
 
-		FSequencerOverlapRange& NewRange = Result.Last();
-		NewRange.Range = Segment.Range;
-		for (FSectionEvaluationData EvalData : Segment.Impls)
+		if (!NewRange.Sections.Num())
 		{
-			NewRange.Sections.Add(FSectionHandle(SharedThis(this), EvalData.ImplIndex));
+			continue;
 		}
+
+		// Sort lowest to highest
+		NewRange.Sections.Sort([](const FSectionHandle& A, const FSectionHandle& B){
+			return A.GetSectionObject()->GetOverlapPriority() < B.GetSectionObject()->GetOverlapPriority();
+		});
+
+		Result.Add(MoveTemp(NewRange));
 	}
-	
+
 	return Result;
 }
 
@@ -395,10 +383,13 @@ TArray<FSequencerOverlapRange> FSequencerTrackNode::GetEasingSegmentsForSection(
 
 	TArray<FMovieSceneSectionData> CompileData;
 
+	TMovieSceneEvaluationTree<int32> SectionIndexTree;
+
+	// Iterate all active sections on the same row with <= overlap priority
 	for (int32 SectionIndex = 0; SectionIndex < Sections.Num(); ++SectionIndex)
 	{
 		UMovieSceneSection* SectionObj = Sections[SectionIndex]->GetSectionObject();
-		if (!SectionObj || !SectionObj->IsActive() || SectionObj->GetRowIndex() != InSection->GetRowIndex())
+		if (!SectionObj || !SectionObj->IsActive() || SectionObj->GetRowIndex() != InSection->GetRowIndex() || SectionObj->GetOverlapPriority() > InSection->GetOverlapPriority())
 		{
 			continue;
 		}
@@ -406,32 +397,43 @@ TArray<FSequencerOverlapRange> FSequencerTrackNode::GetEasingSegmentsForSection(
 		TRange<float> Intersection = TRange<float>::Intersection(SectionObj->GetEaseInRange(), InSectionRange);
 		if (!Intersection.IsEmpty())
 		{
-			CompileData.Add(FMovieSceneSectionData(Intersection, FSectionEvaluationData(SectionIndex), FOptionalMovieSceneBlendType(), SectionObj->GetOverlapPriority()));
+			SectionIndexTree.Add(Intersection, SectionIndex);
 		}
 
 		Intersection = TRange<float>::Intersection(SectionObj->GetEaseOutRange(), InSectionRange);
 		if (!Intersection.IsEmpty())
 		{
-			CompileData.Add(FMovieSceneSectionData(Intersection, FSectionEvaluationData(SectionIndex), FOptionalMovieSceneBlendType(), SectionObj->GetOverlapPriority()));
+			SectionIndexTree.Add(Intersection, SectionIndex);
 		}
 	}
 
-	FOverlappingCompileRules Rules(InSection->GetOverlapPriority());
+	TSharedRef<FSequencerTrackNode> TrackNode = SharedThis(this);;
 
 	TArray<FSequencerOverlapRange> Result;
-
-	for (const FMovieSceneSegment& Segment : FMovieSceneSegmentCompiler().Compile(CompileData, &Rules))
+	for (FMovieSceneEvaluationTreeRangeIterator It(SectionIndexTree); It; ++It)
 	{
-		Result.Emplace();
+		FSequencerOverlapRange NewRange;
+
+		NewRange.Range = It.Range();
 
-		FSequencerOverlapRange& NewRange = Result.Last();
-		NewRange.Range = Segment.Range;
-		for (FSectionEvaluationData EvalData : Segment.Impls)
+		for (int32 SectionIndex : SectionIndexTree.GetAllData(It.Node()))
 		{
-			NewRange.Sections.Add(FSectionHandle(SharedThis(this), EvalData.ImplIndex));
+			NewRange.Sections.Add(FSectionHandle(TrackNode, SectionIndex));
 		}
+
+		if (!NewRange.Sections.Num())
+		{
+			continue;
+		}
+
+		// Sort lowest to highest
+		NewRange.Sections.Sort([=](const FSectionHandle& A, const FSectionHandle& B){
+			return A.GetSectionObject()->GetOverlapPriority() < B.GetSectionObject()->GetOverlapPriority();
+		});
+
+		Result.Add(MoveTemp(NewRange));
 	}
-	
+
 	return Result;
 }
 
diff --git a/Engine/Source/Editor/Sequencer/Private/LevelEditorSequencerIntegration.cpp b/Engine/Source/Editor/Sequencer/Private/LevelEditorSequencerIntegration.cpp
--- a/Engine/Source/Editor/Sequencer/Private/LevelEditorSequencerIntegration.cpp
+++ b/Engine/Source/Editor/Sequencer/Private/LevelEditorSequencerIntegration.cpp
@@ -40,6 +40,9 @@
 #include "Widgets/Layout/SBox.h"
 #include "Widgets/Text/STextBlock.h"
 #include "UObject/ObjectKey.h"
+#include "UnrealEdGlobals.h"
+#include "Editor/UnrealEdEngine.h"
+#include "EditorSupportDelegates.h"
 
 #define LOCTEXT_NAMESPACE "LevelEditorSequencerIntegration"
 
@@ -108,7 +111,7 @@ private:
 static const FName DetailsTabIdentifiers[] = { "LevelEditorSelectionDetails", "LevelEditorSelectionDetails2", "LevelEditorSelectionDetails3", "LevelEditorSelectionDetails4" };
 
 FLevelEditorSequencerIntegration::FLevelEditorSequencerIntegration()
-	: bScrubbing(false)
+	: bDeferUpdates(false)
 {
 	KeyFrameHandler = MakeShared<FDetailKeyframeHandlerWrapper>();
 }
@@ -243,33 +246,41 @@ void FLevelEditorSequencerIntegration::GetActorRecordingState( bool& bIsRecordin
 	);
 }
 
-void RenameSpawnable(FSequencer* Sequencer, UMovieSceneSequence* Sequence, FMovieSceneSequenceIDRef SequenceID, AActor* ChangedActor)
+void RenameSpawnableRecursive(FSequencer* Sequencer, FMovieSceneSequenceIDRef SequenceID, AActor* ChangedActor)
 {
-	UMovieScene* MovieScene = Sequence->GetMovieScene();
-	if (!MovieScene)
-	{
-		return;
-	}
-		
-	for (int32 Index = 0; Index < MovieScene->GetSpawnableCount(); ++Index)
-	{
-		FGuid ThisGuid = MovieScene->GetSpawnable(Index).GetGuid();
+	FMovieSceneRootEvaluationTemplateInstance& RootInstance = Sequencer->GetEvaluationTemplate();
+
+	// Find the sequence that corresponds to the sequence ID
+	UMovieSceneSequence* Sequence = RootInstance.GetSequence(SequenceID);
+	UMovieScene* MovieScene = Sequence ? Sequence->GetMovieScene() : nullptr;
 
-		for (TWeakObjectPtr<> WeakObject : Sequencer->FindBoundObjects(ThisGuid, SequenceID))
+	if (MovieScene)
+	{
+		// Iterate all this movie scene's spawnables, renaming as appropriate
+		for (int32 Index = 0; Index < MovieScene->GetSpawnableCount(); ++Index)
 		{
-			if (WeakObject.IsValid())
+			FGuid ThisGuid = MovieScene->GetSpawnable(Index).GetGuid();
+
+			for (TWeakObjectPtr<> WeakObject : Sequencer->FindBoundObjects(ThisGuid, SequenceID))
 			{
 				AActor* Actor = Cast<AActor>(WeakObject.Get());
-				if (Actor != nullptr)
+				if (Actor && Actor == ChangedActor)
 				{
-					if (Actor == ChangedActor)
-					{
-						MovieScene->GetSpawnable(Index).SetName(ChangedActor->GetActorLabel());
-					}
+					MovieScene->GetSpawnable(Index).SetName(ChangedActor->GetActorLabel());
 				}
 			}
 		}
 	}
+
+	// Recurse into child nodes
+	const FMovieSceneSequenceHierarchyNode* Node = RootInstance.GetHierarchy().FindNode(SequenceID);
+	if (Node)
+	{
+		for (FMovieSceneSequenceIDRef ChildID : Node->Children)
+		{
+			RenameSpawnableRecursive(Sequencer, ChildID, ChangedActor);
+		}
+	}
 }
 
 void FLevelEditorSequencerIntegration::OnActorLabelChanged(AActor* ChangedActor)
@@ -279,14 +290,7 @@ void FLevelEditorSequencerIntegration::OnActorLabelChanged(AActor* ChangedActor)
 		TSharedPtr<FSequencer> Pinned = SequencerAndOptions.Sequencer.Pin();
 		if (Pinned.IsValid())
 		{
-			FMovieSceneRootEvaluationTemplateInstance& RootTemplate = Pinned.Get()->GetEvaluationTemplate();
-	
-			RenameSpawnable(Pinned.Get(), Pinned.Get()->GetRootMovieSceneSequence(), MovieSceneSequenceID::Root, ChangedActor);
-			
-			for (auto& SubInstance : RootTemplate.GetSubInstances())
-			{
-				RenameSpawnable(Pinned.Get(), SubInstance.Value.Sequence.Get(), SubInstance.Key, ChangedActor);
-			}
+			RenameSpawnableRecursive(Pinned.Get(), MovieSceneSequenceID::Root, ChangedActor);
 		}
 	}
 }
@@ -362,6 +366,11 @@ void FLevelEditorSequencerIntegration::OnActorSelectionChanged( UObject* )
 		{
 			if (Options.bRequiresActorEvents)
 			{
+				if (In.GetSequencerSettings()->GetShowSelectedNodesOnly())
+				{
+					In.RefreshTree();
+				}
+
 				In.ExternalSelectionHasChanged();
 			}
 		}
@@ -399,23 +408,66 @@ void FLevelEditorSequencerIntegration::OnSequencerEvaluated()
 		}
 	}
 
-	if (!bScrubbing)
+	if (!bDeferUpdates)
 	{
 		UpdateDetails();
 	}
+
+	// If realtime is off, this needs to be called to update the pivot location when scrubbing.
+	GUnrealEd->UpdatePivotLocationForSelection();
+
+	// Redraw
+	FEditorSupportDelegates::RedrawAllViewports.Broadcast();
 }
 
-void FLevelEditorSequencerIntegration::OnBeginScrubbing()
+void FLevelEditorSequencerIntegration::OnBeginDeferUpdates()
 {
-	bScrubbing = true;
+	bDeferUpdates = true;
 }
 
-void FLevelEditorSequencerIntegration::OnEndScrubbing()
+void FLevelEditorSequencerIntegration::OnEndDeferUpdates()
 {
-	bScrubbing = false;
+	bDeferUpdates = false;
 	UpdateDetails();
 }
 
+bool FLevelEditorSequencerIntegration::IsBindingVisible(const FMovieSceneBinding& InBinding)
+{
+	// If nothing selected, show all nodes
+	if (GEditor->GetSelectedActorCount() == 0)
+	{
+		return true;
+	}
+
+	// Disregard if not a level sequence (ie. a control rig sequence)
+	for (FSequencerAndOptions& SequencerAndOptions : BoundSequencers)
+	{
+		TSharedPtr<FSequencer> Pinned = SequencerAndOptions.Sequencer.Pin();
+		if (Pinned.IsValid())
+		{
+			if (UMovieSceneSequence* RootSequence = Pinned->GetRootMovieSceneSequence())
+			{
+				if (RootSequence->GetClass()->GetName() != TEXT("LevelSequence"))
+				{
+					return true;
+				}
+			}
+		}
+	}
+
+	for( FSelectionIterator SelectionIt( *GEditor->GetSelectedActors() ); SelectionIt; ++SelectionIt )
+	{
+		AActor* SelectedActor = CastChecked<AActor>( *SelectionIt );
+		
+		if (SelectedActor->GetActorLabel() == InBinding.GetName())
+		{
+			return true;
+		}
+	}
+
+	return false;
+}
+
 void FLevelEditorSequencerIntegration::OnMovieSceneBindingsChanged()
 {
 	for (FSequencerAndOptions& SequencerAndOptions : BoundSequencers)
@@ -1255,8 +1307,12 @@ void FLevelEditorSequencerIntegration::AddSequencer(TSharedRef<ISequencer> InSeq
 		FDelegateHandle DataHandle = InSequencer->OnMovieSceneDataChanged().AddRaw(this, &FLevelEditorSequencerIntegration::OnMovieSceneDataChanged);
 		FDelegateHandle AllowEditsModeHandle = InSequencer->GetSequencerSettings()->GetOnAllowEditsModeChanged().AddRaw(this, &FLevelEditorSequencerIntegration::OnAllowEditsModeChanged);
 
-		FDelegateHandle BeginScrubbingHandle = InSequencer->OnBeginScrubbingEvent().AddRaw(this, &FLevelEditorSequencerIntegration::OnBeginScrubbing);
-		FDelegateHandle EndScrubbingHandle = InSequencer->OnEndScrubbingEvent().AddRaw(this, &FLevelEditorSequencerIntegration::OnEndScrubbing);
+		FDelegateHandle PlayHandle = InSequencer->OnPlayEvent().AddRaw(this, &FLevelEditorSequencerIntegration::OnBeginDeferUpdates);
+		FDelegateHandle StopHandle = InSequencer->OnStopEvent().AddRaw(this, &FLevelEditorSequencerIntegration::OnEndDeferUpdates);
+		FDelegateHandle BeginScrubbingHandle = InSequencer->OnBeginScrubbingEvent().AddRaw(this, &FLevelEditorSequencerIntegration::OnBeginDeferUpdates);
+		FDelegateHandle EndScrubbingHandle = InSequencer->OnEndScrubbingEvent().AddRaw(this, &FLevelEditorSequencerIntegration::OnEndDeferUpdates);
+
+		InSequencer->OnGetIsBindingVisible().BindRaw(this, &FLevelEditorSequencerIntegration::IsBindingVisible);
 
 		BoundSequencers.Last().AcquiredResources.Add(
 			[=]
@@ -1268,6 +1324,8 @@ void FLevelEditorSequencerIntegration::AddSequencer(TSharedRef<ISequencer> InSeq
 					Pinned->OnMovieSceneBindingsChanged().Remove(BindingsHandle);
 					Pinned->OnMovieSceneDataChanged().Remove(DataHandle);
 					Pinned->GetSequencerSettings()->GetOnAllowEditsModeChanged().Remove(AllowEditsModeHandle);
+					Pinned->OnPlayEvent().Remove(PlayHandle);
+					Pinned->OnStopEvent().Remove(StopHandle);
 					Pinned->OnBeginScrubbingEvent().Remove(BeginScrubbingHandle);
 					Pinned->OnEndScrubbingEvent().Remove(EndScrubbingHandle);
 				}
@@ -1341,20 +1399,24 @@ void FLevelEditorSequencerIntegration::RemoveSequencer(TSharedRef<ISequencer> In
 	}
 }
 
-void AddActorsToBindingsMap(TWeakPtr<FSequencer> Sequencer, UMovieSceneSequence* Sequence, FMovieSceneSequenceIDRef SequenceID, TMap<FObjectKey, FString>& ActorBindingsMap)
+void AddActorsToBindingsMapRecursive(FSequencer& Sequencer, FMovieSceneSequenceIDRef SequenceID, TMap<FObjectKey, FString>& ActorBindingsMap)
 {
-	UMovieScene* MovieScene = Sequence->GetMovieScene();
+	FMovieSceneRootEvaluationTemplateInstance& RootInstance = Sequencer.GetEvaluationTemplate();
 
-	FString SequenceName = Sequence->GetDisplayName().ToString();
+	// Find the sequence that corresponds to the sequence ID
+	UMovieSceneSequence* Sequence = RootInstance.GetSequence(SequenceID);
+	UMovieScene* MovieScene = Sequence ? Sequence->GetMovieScene() : nullptr;
 
-	// Search all possessables
-	for (int32 Index = 0; Index < MovieScene->GetPossessableCount(); ++Index)
+	if (MovieScene)
 	{
-		FGuid ThisGuid = MovieScene->GetPossessable(Index).GetGuid();
+		FString SequenceName = Sequence->GetDisplayName().ToString();
 
-		for (TWeakObjectPtr<> WeakObject : Sequencer.Pin()->FindBoundObjects(ThisGuid, SequenceID))
+		// Search all possessables
+		for (int32 Index = 0; Index < MovieScene->GetPossessableCount(); ++Index)
 		{
-			if (WeakObject.IsValid())
+			FGuid ThisGuid = MovieScene->GetPossessable(Index).GetGuid();
+
+			for (TWeakObjectPtr<> WeakObject : Sequencer.FindBoundObjects(ThisGuid, SequenceID))
 			{
 				AActor* Actor = Cast<AActor>(WeakObject.Get());
 				if (Actor != nullptr)
@@ -1372,16 +1434,13 @@ void AddActorsToBindingsMap(TWeakPtr<FSequencer> Sequencer, UMovieSceneSequence*
 				}
 			}
 		}
-	}
-
-	// Search all spawnables
-	for (int32 Index = 0; Index < MovieScene->GetSpawnableCount(); ++Index)
-	{
-		FGuid ThisGuid = MovieScene->GetSpawnable(Index).GetGuid();
 
-		for (TWeakObjectPtr<> WeakObject : Sequencer.Pin()->FindBoundObjects(ThisGuid, SequenceID))
+		// Search all spawnables
+		for (int32 Index = 0; Index < MovieScene->GetSpawnableCount(); ++Index)
 		{
-			if (WeakObject.IsValid())
+			FGuid ThisGuid = MovieScene->GetSpawnable(Index).GetGuid();
+
+			for (TWeakObjectPtr<> WeakObject : Sequencer.FindBoundObjects(ThisGuid, SequenceID))
 			{
 				AActor* Actor = Cast<AActor>(WeakObject.Get());
 				if (Actor != nullptr)
@@ -1400,6 +1459,16 @@ void AddActorsToBindingsMap(TWeakPtr<FSequencer> Sequencer, UMovieSceneSequence*
 			}
 		}
 	}
+
+	// Recurse into child nodes
+	const FMovieSceneSequenceHierarchyNode* Node = RootInstance.GetHierarchy().FindNode(SequenceID);
+	if (Node)
+	{
+		for (FMovieSceneSequenceIDRef ChildID : Node->Children)
+		{
+			AddActorsToBindingsMapRecursive(Sequencer, ChildID, ActorBindingsMap);
+		}
+	}
 }
 
 void AddPropertiesToBindingsMap(TWeakPtr<FSequencer> Sequencer, UMovieSceneSequence* Sequence, FMovieSceneSequenceIDRef SequenceID, TMap<FObjectKey, TArray<FString> >& PropertyBindingsMap)
@@ -1483,17 +1552,7 @@ void FLevelEditorSequencerBindingData::UpdateActorBindingsData(TWeakPtr<FSequenc
 	
 		// Finding the bound objects can cause bindings to be evaluated and changed, causing this to be invoked again
 		TGuardValue<bool> ReentrantGuard(bIsReentrant, true);
-
-		FMovieSceneRootEvaluationTemplateInstance& RootTemplate = InSequencer.Pin()->GetEvaluationTemplate();
-		
-		UMovieSceneSequence* Sequence = RootTemplate.GetSequence(MovieSceneSequenceID::Root);
-
-		AddActorsToBindingsMap(InSequencer, Sequence, MovieSceneSequenceID::Root, ActorBindingsMap);
-
-		for (auto& SubInstance : RootTemplate.GetSubInstances())
-		{
-			AddActorsToBindingsMap(InSequencer, SubInstance.Value.Sequence.Get(), SubInstance.Key, ActorBindingsMap);
-		}
+		AddActorsToBindingsMapRecursive(*InSequencer.Pin(), MovieSceneSequenceID::Root, ActorBindingsMap);
 
 		bActorBindingsDirty = false;
 
diff --git a/Engine/Source/Editor/Sequencer/Private/MovieSceneTrackEditor.cpp b/Engine/Source/Editor/Sequencer/Private/MovieSceneTrackEditor.cpp
--- a/Engine/Source/Editor/Sequencer/Private/MovieSceneTrackEditor.cpp
+++ b/Engine/Source/Editor/Sequencer/Private/MovieSceneTrackEditor.cpp
@@ -74,6 +74,10 @@ void FMovieSceneTrackEditor::AnimatablePropertyChanged( FOnKeyProperty OnKeyProp
 			{
 				Sequencer.Pin()->NotifyMovieSceneDataChanged( EMovieSceneDataChangeType::MovieSceneStructureItemAdded );
 			}
+			else if (KeyPropertyResult.bKeyCreated)
+			{
+				Sequencer.Pin()->NotifyMovieSceneDataChanged( EMovieSceneDataChangeType::TrackValueChanged );
+			}
 
 			UpdatePlaybackRange();
 
diff --git a/Engine/Source/Editor/Sequencer/Private/SKeyNavigationButtons.h b/Engine/Source/Editor/Sequencer/Private/SKeyNavigationButtons.h
--- a/Engine/Source/Editor/Sequencer/Private/SKeyNavigationButtons.h
+++ b/Engine/Source/Editor/Sequencer/Private/SKeyNavigationButtons.h
@@ -19,8 +19,9 @@
 #include "IKeyArea.h"
 #include "MovieSceneCommonHelpers.h"
 #include "ScopedTransaction.h"
+#include "SequencerCommands.h"
 
-#define LOCTEXT_NAMESPACE "Arse"
+#define LOCTEXT_NAMESPACE "SKeyNavigationButtons"
 
 /**
  * A widget for navigating between keys on a sequencer track
@@ -35,6 +36,10 @@ public:
 
 	void Construct(const FArguments& InArgs, const TSharedRef<FSequencerDisplayNode>& InDisplayNode)
 	{
+		FText SetKeyToolTip = FText::Format(LOCTEXT("AddKeyButton", "Add a new key at the current time ({0})"), FSequencerCommands::Get().SetKey->GetInputText());
+		FText PreviousKeyToolTip = FText::Format(LOCTEXT("PreviousKeyButton", "Set the time to the previous key ({0})"), FSequencerCommands::Get().StepToPreviousKey->GetInputText());
+		FText NextKeyToolTip = FText::Format(LOCTEXT("NextKeyButton", "Set the time to the next key ({0})"), FSequencerCommands::Get().StepToNextKey->GetInputText());
+
 		DisplayNode = InDisplayNode;
 
 		const FSlateBrush* NoBorder = FEditorStyle::GetBrush( "NoBorder" );
@@ -58,7 +63,7 @@ public:
 				[
 					SNew(SButton)
 					.ButtonStyle(FEditorStyle::Get(), "FlatButton")
-					.ToolTipText(LOCTEXT("PreviousKeyButton", "Set the time to the previous key"))
+					.ToolTipText(PreviousKeyToolTip)
 					.OnClicked(this, &SKeyNavigationButtons::OnPreviousKeyClicked)
 					.ForegroundColor( FSlateColor::UseForeground() )
 					.ContentPadding(0)
@@ -83,7 +88,7 @@ public:
 				[
 					SNew(SButton)
 					.ButtonStyle(FEditorStyle::Get(), "FlatButton")
-					.ToolTipText(LOCTEXT("AddKeyButton", "Add a new key at the current time"))
+					.ToolTipText(SetKeyToolTip)
 					.OnClicked(this, &SKeyNavigationButtons::OnAddKeyClicked)
 					.ForegroundColor( FSlateColor::UseForeground() )
 					.ContentPadding(0)
@@ -107,7 +112,7 @@ public:
 				[
 					SNew(SButton)
 					.ButtonStyle(FEditorStyle::Get(), "FlatButton")
-					.ToolTipText(LOCTEXT("NextKeyButton", "Set the time to the next key"))
+					.ToolTipText(NextKeyToolTip)
 					.OnClicked(this, &SKeyNavigationButtons::OnNextKeyClicked)
 					.ContentPadding(0)
 					.ForegroundColor( FSlateColor::UseForeground() )
diff --git a/Engine/Source/Editor/Sequencer/Private/SSequencer.cpp b/Engine/Source/Editor/Sequencer/Private/SSequencer.cpp
--- a/Engine/Source/Editor/Sequencer/Private/SSequencer.cpp
+++ b/Engine/Source/Editor/Sequencer/Private/SSequencer.cpp
@@ -63,6 +63,7 @@
 #include "IVREditorModule.h"
 #include "EditorFontGlyphs.h"
 #include "HAL/PlatformApplicationMisc.h"
+#include "Camera/CameraActor.h"
 
 #define LOCTEXT_NAMESPACE "Sequencer"
 
@@ -996,6 +997,7 @@ TSharedRef<SWidget> SSequencer::MakeGeneralMenu()
 	MenuBuilder.BeginSection( "ViewOptions", LOCTEXT( "ViewMenuHeader", "View" ) );
 	{
 		MenuBuilder.AddMenuEntry( FSequencerCommands::Get().ToggleLabelBrowser );
+		MenuBuilder.AddMenuEntry( FSequencerCommands::Get().ToggleShowSelectedNodesOnly );
 		MenuBuilder.AddMenuEntry( FSequencerCommands::Get().ToggleCombinedKeyframes );
 		MenuBuilder.AddMenuEntry( FSequencerCommands::Get().ToggleChannelColors );
 		MenuBuilder.AddMenuEntry( FSequencerCommands::Get().ToggleShowPreAndPostRoll );
@@ -1034,6 +1036,29 @@ TSharedRef<SWidget> SSequencer::MakeGeneralMenu()
 	return MenuBuilder.MakeWidget();
 }
 
+void SSequencer::FillPlaybackSpeedMenu(FMenuBuilder& InMenuBarBuilder)
+{
+	const int32 NumPlaybackSpeeds = 7;
+	float PlaybackSpeeds[NumPlaybackSpeeds] = { 0.1f, 0.25f, 0.5f, 1.0f, 2.0f, 5.0f, 10.0f };
+
+	InMenuBarBuilder.BeginSection("PlaybackSpeed");
+	for( uint32 PlaybackSpeedIndex = 1; PlaybackSpeedIndex < NumPlaybackSpeeds; ++PlaybackSpeedIndex )
+	{
+		float PlaybackSpeed = PlaybackSpeeds[PlaybackSpeedIndex];
+		const FText MenuStr = FText::Format( LOCTEXT("PlaybackSpeedStr", "x{0}"), FText::AsNumber( PlaybackSpeed ) );
+		InMenuBarBuilder.AddMenuEntry(MenuStr, FText(), FSlateIcon(),
+			FUIAction(
+				FExecuteAction::CreateLambda( [this, PlaybackSpeed]{ SequencerPtr.Pin()->SetPlaybackSpeed(PlaybackSpeed); }),
+				FCanExecuteAction::CreateLambda([] { return true; }),
+				FIsActionChecked::CreateLambda( [this, PlaybackSpeed]{ return SequencerPtr.Pin()->GetPlaybackSpeed() == PlaybackSpeed; })
+				),
+			NAME_None,
+			EUserInterfaceActionType::RadioButton
+			);
+	}
+	InMenuBarBuilder.EndSection();
+}
+
 TSharedRef<SWidget> SSequencer::MakePlaybackMenu()
 {
 	FMenuBuilder MenuBuilder( true, SequencerPtr.Pin()->GetCommandBindings() );
@@ -1115,6 +1140,9 @@ TSharedRef<SWidget> SSequencer::MakePlaybackMenu()
 				],
 			LOCTEXT("PlaybackStartEnd", "End"));
 
+
+		MenuBuilder.AddSubMenu(LOCTEXT("PlaybackSpeedHeader", "Playback Speed"), FText::GetEmpty(), FNewMenuDelegate::CreateRaw(this, &SSequencer::FillPlaybackSpeedMenu));
+
 		MenuBuilder.AddMenuEntry( FSequencerCommands::Get().TogglePlaybackRangeLocked );
 		MenuBuilder.AddMenuEntry( FSequencerCommands::Get().ToggleForceFixedFrameIntervalPlayback );
 
@@ -1174,6 +1202,11 @@ TSharedRef<SWidget> SSequencer::MakeSelectEditMenu()
 	TSharedPtr<FSequencer> Sequencer = SequencerPtr.Pin();
 
 	MenuBuilder.AddMenuEntry(FSequencerCommands::Get().ToggleShowTransformBox);
+	
+	if (SequencerPtr.Pin()->IsLevelEditorSequencer())
+	{
+		MenuBuilder.AddMenuEntry(FSequencerCommands::Get().BakeTransform);
+	}
 
 	// selection range options
 	MenuBuilder.BeginSection("SelectionRange", LOCTEXT("SelectionRangeHeader", "Selection Range"));
@@ -1686,7 +1719,18 @@ void SSequencer::OnAssetsDropped( const FAssetDragDropOp& DragDropOp )
 
 		if (!SequencerRef.OnHandleAssetDropped(CurObject, TargetObjectGuid))
 		{
-			SequencerRef.MakeNewSpawnable( *CurObject );
+			FGuid NewGuid = SequencerRef.MakeNewSpawnable( *CurObject, DragDropOp.GetActorFactory() );
+
+			UMovieScene* MovieScene = SequencerRef.GetFocusedMovieSceneSequence()->GetMovieScene();
+			if (MovieScene)
+			{
+				FMovieSceneSpawnable* Spawnable = MovieScene->FindSpawnable(NewGuid);
+
+				if (Spawnable && Spawnable->GetObjectTemplate()->IsA<ACameraActor>())
+				{
+					SequencerRef.NewCameraAdded(NewGuid);
+				}
+			}
 		}
 		bObjectAdded = true;
 	}
diff --git a/Engine/Source/Editor/Sequencer/Private/SSequencer.h b/Engine/Source/Editor/Sequencer/Private/SSequencer.h
--- a/Engine/Source/Editor/Sequencer/Private/SSequencer.h
+++ b/Engine/Source/Editor/Sequencer/Private/SSequencer.h
@@ -314,6 +314,9 @@ private:
 	/** Makes the allow edits menu for the toolbar. */
 	TSharedRef<SWidget> MakeAllowEditsMenu();
 
+	/** Makes the playback speed menu for the toolbar. */
+	void FillPlaybackSpeedMenu(FMenuBuilder& InMenuBuilder);
+
 public:	
 
 	/** Makes a time range widget with the specified inner content */
@@ -430,6 +433,9 @@ private:
 
 	EVisibility GetDebugVisualizerVisibility() const;
 	
+	void SetPlaybackSpeed(float InPlaybackSpeed);
+	float GetPlaybackSpeed() const;
+
 public:
 	/** On Paste Command */
 	void OnPaste();
diff --git a/Engine/Source/Editor/Sequencer/Private/SSequencerCurveEditor.cpp b/Engine/Source/Editor/Sequencer/Private/SSequencerCurveEditor.cpp
--- a/Engine/Source/Editor/Sequencer/Private/SSequencerCurveEditor.cpp
+++ b/Engine/Source/Editor/Sequencer/Private/SSequencerCurveEditor.cpp
@@ -48,7 +48,7 @@ FReply SSequencerCurveEditor::OnMouseButtonDown( const FGeometry& MyGeometry, co
 	FReply Reply = SCurveEditor::OnMouseButtonDown( MyGeometry, MouseEvent );
 	if ( Reply.IsEventHandled() )
 	{
-		return Reply;
+		return Reply.PreventThrottling();
 	}
 
 	return TimeSliderController->OnMouseButtonDown( *this, MyGeometry, MouseEvent );
