Copying //UE4/Dev-Sequencer to //UE4/Dev-Main (Source: //UE4/Dev-Sequencer @ 3800792)

diff --git a/Engine/Source/Editor/Sequencer/Private/SSequencerDebugVisualizer.cpp b/Engine/Source/Editor/Sequencer/Private/SSequencerDebugVisualizer.cpp
--- a/Engine/Source/Editor/Sequencer/Private/SSequencerDebugVisualizer.cpp
+++ b/Engine/Source/Editor/Sequencer/Private/SSequencerDebugVisualizer.cpp
@@ -14,9 +14,6 @@ void SSequencerDebugVisualizer::Construct(const FArguments& InArgs, TSharedRef<F
 {
 	WeakSequencer = InSequencer;
 
-	InSequencer->GetSequenceInstance().OnUpdated().AddSP(this, &SSequencerDebugVisualizer::Refresh);
-	InSequencer->OnActivateSequence().AddSP(this, &SSequencerDebugVisualizer::OnSequenceActivated);
-
 	SetClipping(EWidgetClipping::ClipToBounds);
 
 	ViewRange = InArgs._ViewRange;
@@ -24,15 +21,10 @@ void SSequencerDebugVisualizer::Construct(const FArguments& InArgs, TSharedRef<F
 	Refresh();
 }
 
-void SSequencerDebugVisualizer::OnSequenceActivated(FMovieSceneSequenceIDRef)
-{
-	Refresh();
-}
-
-const FMovieSceneEvaluationTemplateInstance* SSequencerDebugVisualizer::GetTemplate() const
+const FMovieSceneEvaluationTemplate* SSequencerDebugVisualizer::GetTemplate() const
 {
 	TSharedPtr<FSequencer> Sequencer = WeakSequencer.Pin();
-	return Sequencer.IsValid() ? Sequencer->GetSequenceInstance().GetInstance(FocusedSequenceID) : nullptr;
+	return Sequencer.IsValid() ? Sequencer->GetEvaluationTemplate().FindTemplate(Sequencer->GetFocusedTemplateID()) : nullptr;
 }
 
 void SSequencerDebugVisualizer::Refresh()
@@ -45,25 +37,27 @@ void SSequencerDebugVisualizer::Refresh()
 		return;
 	}
 
-	FocusedSequenceID = Sequencer->GetFocusedTemplateID();
-
-	const FMovieSceneEvaluationTemplateInstance* ActiveInstance = GetTemplate();
-	if (!ActiveInstance)
+	const FMovieSceneEvaluationTemplate* ActiveTemplate = GetTemplate();
+	if (!ActiveTemplate)
 	{
 		return;
 	}
 
-	const FMovieSceneEvaluationField& EvaluationField = ActiveInstance->Template->EvaluationField;
+	const FMovieSceneEvaluationField& EvaluationField = ActiveTemplate->EvaluationField;
+
+	CachedSignature = EvaluationField.GetSignature();
 
 	TArray<int32> SegmentComplexity;
-	SegmentComplexity.Reserve(EvaluationField.Groups.Num());
+	SegmentComplexity.Reserve(EvaluationField.Size());
 
 	// Heatmap complexity
 	float AverageComplexity = 0;
 	int32 MaxComplexity = 0;
 
-	for (const FMovieSceneEvaluationGroup& Group : EvaluationField.Groups)
+	for (int32 Index = 0; Index < EvaluationField.Size(); ++Index)
 	{
+		const FMovieSceneEvaluationGroup& Group = EvaluationField.GetGroup(Index);
+
 		int32 Complexity = 0;
 		for (const FMovieSceneEvaluationGroupLUTIndex& LUTIndex : Group.LUTIndices)
 		{
@@ -83,9 +77,9 @@ void SSequencerDebugVisualizer::Refresh()
 	static const FSlateBrush* SectionBackgroundBrush = FEditorStyle::GetBrush("Sequencer.Section.Background");
 	static const FSlateBrush* SectionBackgroundTintBrush = FEditorStyle::GetBrush("Sequencer.Section.BackgroundTint");
 
-	for (int32 Index = 0; Index < EvaluationField.Ranges.Num(); ++Index)
+	for (int32 Index = 0; Index < EvaluationField.Size(); ++Index)
 	{
-		TRange<float> Range = EvaluationField.Ranges[Index];
+		TRange<float> Range = EvaluationField.GetRange(Index);
 
 		const float Complexity = SegmentComplexity[Index];
 		
@@ -126,36 +120,28 @@ void SSequencerDebugVisualizer::Refresh()
 FVector2D SSequencerDebugVisualizer::ComputeDesiredSize(float) const
 {
 	// Note: X Size is not used
-	FVector2D Size(100, 0.f);
-	if (Children.Num())
-	{
-		for (int32 Index = 0; Index < Children.Num(); ++Index)
-		{
-			Size.Y = FMath::Max(Size.Y, Children[Index]->GetDesiredSize().Y);
-		}
-	}
-	return Size;
+	return FVector2D(100, 20.f);
 }
 
 FGeometry SSequencerDebugVisualizer::GetSegmentGeometry(const FGeometry& AllottedGeometry, const SSequencerDebugSlot& Slot, const FTimeToPixel& TimeToPixelConverter) const
 {
-	const FMovieSceneEvaluationTemplateInstance* ActiveInstance = GetTemplate();
-	if (!ActiveInstance)
+	const FMovieSceneEvaluationTemplate* ActiveTemplate = GetTemplate();
+	if (!ActiveTemplate || ActiveTemplate->EvaluationField.GetSignature() != CachedSignature)
 	{
 		return AllottedGeometry.MakeChild(FVector2D(0,0), FVector2D(0,0));
 	}
 
-	TRange<float> SegmentRange = ActiveInstance->Template->EvaluationField.Ranges[Slot.GetSegmentIndex()];
+	TRange<float> SegmentRange = ActiveTemplate->EvaluationField.GetRange(Slot.GetSegmentIndex());
 
 	float PixelStartX = SegmentRange.GetLowerBound().IsOpen() ? 0.f : TimeToPixelConverter.TimeToPixel(SegmentRange.GetLowerBoundValue());
 	float PixelEndX = SegmentRange.GetUpperBound().IsOpen() ? AllottedGeometry.GetDrawSize().X : TimeToPixelConverter.TimeToPixel(SegmentRange.GetUpperBoundValue());
 
-	const float MinSectionWidth = 1.f;
+	const float MinSectionWidth = 0.f;
 	float SectionLength = FMath::Max(MinSectionWidth, PixelEndX - PixelStartX);
 
 	return AllottedGeometry.MakeChild(
 		FVector2D(PixelStartX, 0),
-		FVector2D(SectionLength, Slot.GetDesiredSize().Y)
+		FVector2D(SectionLength, FMath::Max(Slot.GetDesiredSize().Y, 20.f))
 		);
 }
 
@@ -166,13 +152,13 @@ EVisibility SSequencerDebugVisualizer::GetSegmentVisibility(TRange<float> Range)
 
 TSharedRef<SWidget> SSequencerDebugVisualizer::GetTooltipForSegment(int32 SegmentIndex) const
 {
-	const FMovieSceneEvaluationTemplateInstance* ActiveInstance = GetTemplate();
-	if (!ActiveInstance)
+	const FMovieSceneEvaluationTemplate* ActiveTemplate = GetTemplate();
+	if (!ActiveTemplate)
 	{
 		return SNullWidget::NullWidget;
 	}
 
-	const FMovieSceneEvaluationGroup& Group = ActiveInstance->Template->EvaluationField.Groups[SegmentIndex];
+	const FMovieSceneEvaluationGroup& Group = ActiveTemplate->EvaluationField.GetGroup(SegmentIndex);
 
 	TSharedRef<SVerticalBox> VerticalBox = SNew(SVerticalBox);
 
@@ -206,6 +192,21 @@ TSharedRef<SWidget> SSequencerDebugVisualizer::GetTooltipForSegment(int32 Segmen
 	return VerticalBox;
 }
 
+void SSequencerDebugVisualizer::Tick( const FGeometry& AllottedGeometry, const double InCurrentTime, const float InDeltaTime )
+{
+	SPanel::Tick(AllottedGeometry, InCurrentTime, InDeltaTime);
+
+	const FMovieSceneEvaluationTemplate* Template = GetTemplate();
+	if (!Template)
+	{
+		Children.Empty();
+	}
+	else if (Template->EvaluationField.GetSignature() != CachedSignature)
+	{
+		Refresh();
+	}
+}
+
 void SSequencerDebugVisualizer::OnArrangeChildren( const FGeometry& AllottedGeometry, FArrangedChildren& ArrangedChildren ) const
 {
 	FTimeToPixel TimeToPixelConverter = FTimeToPixel(AllottedGeometry, ViewRange.Get());
@@ -218,11 +219,13 @@ void SSequencerDebugVisualizer::OnArrangeChildren( const FGeometry& AllottedGeom
 		if( ArrangedChildren.Accepts( WidgetVisibility ) )
 		{
 			FGeometry SegmentGeometry = GetSegmentGeometry(AllottedGeometry, *Child, TimeToPixelConverter);
-			
-			ArrangedChildren.AddWidget( 
-				WidgetVisibility, 
-				AllottedGeometry.MakeChild(Child, SegmentGeometry.Position, SegmentGeometry.GetDrawSize())
-				);
+			if (SegmentGeometry.GetLocalSize().X >= 1)
+			{
+				ArrangedChildren.AddWidget( 
+					WidgetVisibility, 
+					AllottedGeometry.MakeChild(Child, SegmentGeometry.Position, SegmentGeometry.GetDrawSize())
+					);
+			}
 		}
 	}
 }
diff --git a/Engine/Source/Editor/Sequencer/Private/SSequencerDebugVisualizer.h b/Engine/Source/Editor/Sequencer/Private/SSequencerDebugVisualizer.h
--- a/Engine/Source/Editor/Sequencer/Private/SSequencerDebugVisualizer.h
+++ b/Engine/Source/Editor/Sequencer/Private/SSequencerDebugVisualizer.h
@@ -51,11 +51,14 @@ public:
 protected:
 	/** SPanel Interface */
 	virtual void OnArrangeChildren( const FGeometry& AllottedGeometry, FArrangedChildren& ArrangedChildren ) const override;
+	virtual void Tick( const FGeometry& AllottedGeometry, const double InCurrentTime, const float InDeltaTime ) override;
 	virtual FVector2D ComputeDesiredSize(float) const override;
 	virtual FChildren* GetChildren() override { return &Children; }
 
 protected:
-	
+
+	void RequestRefresh();
+
 	void Refresh();
 
 	FGeometry GetSegmentGeometry(const FGeometry& AllottedGeometry, const SSequencerDebugSlot& Slot, const FTimeToPixel& TimeToPixelConverter) const;
@@ -66,12 +69,10 @@ protected:
 
 	void OnSequenceActivated(FMovieSceneSequenceIDRef);
 
-	const FMovieSceneEvaluationTemplateInstance* GetTemplate() const;
+	const FMovieSceneEvaluationTemplate* GetTemplate() const;
 
 private:
 
-	FMovieSceneSequenceID FocusedSequenceID;
-
 	/** The current view range */
 	TAttribute<TRange<float>> ViewRange;
 
@@ -79,4 +80,6 @@ private:
 	TSlotlessChildren<SSequencerDebugSlot> Children;
 
 	TWeakPtr<FSequencer> WeakSequencer;
+
+	FGuid CachedSignature;
 };
diff --git a/Engine/Source/Editor/Sequencer/Private/SSequencerSection.cpp b/Engine/Source/Editor/Sequencer/Private/SSequencerSection.cpp
--- a/Engine/Source/Editor/Sequencer/Private/SSequencerSection.cpp
+++ b/Engine/Source/Editor/Sequencer/Private/SSequencerSection.cpp
@@ -1209,7 +1209,9 @@ void SSequencerSection::PaintKeys( FSequencerSectionPainter& InPainter, const FW
 
 			// Count the number of overlapping keys
 			int32 NumOverlaps = 0;
-			while (KeyIndex + 1 < KeysInRange.Num() && FMath::IsNearlyEqual(TimeToPixelConverter.TimeToPixel(KeysInRange[KeyIndex+1].Time), KeyPosition, PixelOverlapThreshold))
+			while (KeyIndex < KeysInRange.Num() && 
+					( (KeyIndex + 1 < KeysInRange.Num() && FMath::IsNearlyEqual(TimeToPixelConverter.TimeToPixel(KeysInRange[KeyIndex+1].Time), KeyPosition, PixelOverlapThreshold)) ||
+					  (KeyIndex > 1 && FMath::IsNearlyEqual(TimeToPixelConverter.TimeToPixel(KeysInRange[KeyIndex-1].Time), KeyPosition, PixelOverlapThreshold)) ))
 			{
 				++KeyIndex;
 				++NumOverlaps;
diff --git a/Engine/Source/Editor/Sequencer/Private/SSequencerTrackArea.cpp b/Engine/Source/Editor/Sequencer/Private/SSequencerTrackArea.cpp
--- a/Engine/Source/Editor/Sequencer/Private/SSequencerTrackArea.cpp
+++ b/Engine/Source/Editor/Sequencer/Private/SSequencerTrackArea.cpp
@@ -356,7 +356,16 @@ FReply SSequencerTrackArea::OnMouseWheel( const FGeometry& MyGeometry, const FPo
 	{
 		// Always ensure the edit tool is set up
 		InputStack.SetHandlerAt(0, EditTool.Get());
-		return InputStack.HandleMouseWheel(*this, MyGeometry, MouseEvent);
+		FReply EditToolHandle = InputStack.HandleMouseWheel(*this, MyGeometry, MouseEvent);
+		if (EditToolHandle.IsEventHandled())
+		{
+			return EditToolHandle;
+		}
+
+		const float ScrollAmount = -MouseEvent.GetWheelDelta() * GetGlobalScrollAmount();
+		Sequencer.Pin()->VerticalScroll(ScrollAmount);
+
+		return FReply::Handled();
 	}
 	return FReply::Unhandled();
 }
@@ -464,6 +473,7 @@ void SSequencerTrackArea::OnDragEnter(const FGeometry& MyGeometry, const FDragDr
 	
 void SSequencerTrackArea::OnDragLeave(const FDragDropEvent& DragDropEvent)
 {
+	DroppedNode.Reset();
 	SPanel::OnDragLeave(DragDropEvent);
 }
 	 
diff --git a/Engine/Source/Editor/Sequencer/Private/Sequencer.cpp b/Engine/Source/Editor/Sequencer/Private/Sequencer.cpp
--- a/Engine/Source/Editor/Sequencer/Private/Sequencer.cpp
+++ b/Engine/Source/Editor/Sequencer/Private/Sequencer.cpp
@@ -60,7 +60,10 @@
 #include "EditorSupportDelegates.h"
 #include "SSequencerTreeView.h"
 #include "ScopedTransaction.h"
+#include "Tracks/MovieScene3DTransformTrack.h"
 #include "Tracks/MovieScene3DAttachTrack.h"
+#include "Tracks/MovieSceneCameraAnimTrack.h"
+#include "Tracks/MovieSceneCameraShakeTrack.h"
 #include "Tracks/MovieSceneCameraCutTrack.h"
 #include "ISequencerTrackEditor.h"
 #include "MovieSceneToolHelpers.h"
@@ -71,6 +74,7 @@
 #include "Sections/MovieSceneSubSection.h"
 #include "Tracks/MovieSceneSubTrack.h"
 #include "Sections/MovieSceneCinematicShotSection.h"
+#include "MovieSceneObjectBindingIDCustomization.h"
 #include "ISettingsModule.h"
 #include "Framework/Commands/GenericCommands.h"
 #include "Tracks/MovieSceneSpawnTrack.h"
@@ -104,20 +108,16 @@
 #include "LevelSequence.h"
 #include "IVREditorModule.h"
 #include "HAL/PlatformApplicationMisc.h"
+#include "Compilation/MovieSceneCompiler.h"
 #include "SequencerKeyActor.h"
+#include "MovieSceneCopyableBinding.h"
 
 #define LOCTEXT_NAMESPACE "Sequencer"
 
 DEFINE_LOG_CATEGORY(LogSequencer);
 
-struct FSequencerTemplateStore : FMovieSceneSequenceTemplateStore
+struct FSequencerTemplateStore : IMovieSceneSequenceTemplateStore
 {
-	FSequencerTemplateStore()
-	{
-		TemplateParameters.bForEditorPreview = true;
-		bTemplatesAreVolatile = true;
-	}
-
 	void Reset()
 	{
 		Templates.Reset();
@@ -131,28 +131,24 @@ struct FSequencerTemplateStore : FMovieSceneSequenceTemplateStore
 		}
 	}
 
-	virtual FMovieSceneEvaluationTemplate& GetCompiledTemplate(UMovieSceneSequence& Sequence, FObjectKey InSequenceKey)
+	virtual FMovieSceneEvaluationTemplate& AccessTemplate(UMovieSceneSequence& Sequence)
 	{
-		if (TUniquePtr<FCachedMovieSceneEvaluationTemplate>* ExistingTemplate = Templates.Find(InSequenceKey))
+		FObjectKey SequenceKey(&Sequence);
+		if (TUniquePtr<FMovieSceneEvaluationTemplate>* ExistingTemplate = Templates.Find(SequenceKey))
 		{
-			FCachedMovieSceneEvaluationTemplate* Template = ExistingTemplate->Get();
-			Template->Regenerate(TemplateParameters);
+			FMovieSceneEvaluationTemplate* Template = ExistingTemplate->Get();
 			return *Template;
 		}
 		else
 		{
-			FCachedMovieSceneEvaluationTemplate* NewTemplate = new FCachedMovieSceneEvaluationTemplate;
-			NewTemplate->Initialize(Sequence, this);
-			NewTemplate->Regenerate(TemplateParameters);
-
-			Templates.Add(InSequenceKey, TUniquePtr<FCachedMovieSceneEvaluationTemplate>(NewTemplate));
+			FMovieSceneEvaluationTemplate* NewTemplate = new FMovieSceneEvaluationTemplate;
+			Templates.Add(SequenceKey, TUniquePtr<FMovieSceneEvaluationTemplate>(NewTemplate));
 			return *NewTemplate;
 		}
 	}
 
 	// Store templates as unique ptrs to ensure that external pointers don't become invalid when the array is reallocated
-	TMap<FObjectKey, TUniquePtr<FCachedMovieSceneEvaluationTemplate>> Templates;
-	FMovieSceneTrackCompilationParams TemplateParameters;
+	TMap<FObjectKey, TUniquePtr<FMovieSceneEvaluationTemplate>> Templates;
 };
 
 
@@ -194,16 +190,23 @@ void FSequencer::InitSequencer(const FSequencerInitParams& InitParams, const TSh
 	ResetTimingManager(Settings->ShouldLockPlaybackToAudioClock());
 
 	Settings->GetOnEvaluateSubSequencesInIsolationChanged().AddSP(this, &FSequencer::RestorePreAnimatedState);
-	
+	Settings->GetOnShowSelectedNodesOnlyChanged().AddSP(this, &FSequencer::OnSelectedNodesOnlyChanged);
+
 	{
 		FDelegateHandle OnBlueprintCompiledHandle = GEditor->OnBlueprintCompiled().AddLambda([&]{ State.InvalidateExpiredObjects(); });
 		AcquiredResources.Add([=]{ GEditor->OnBlueprintCompiled().Remove(OnBlueprintCompiledHandle); });
 	}
 
+	{
+		ISequenceRecorder* Recorder = FModuleManager::Get().GetModulePtr<ISequenceRecorder>("SequenceRecorder");
+		
+		Recorder->OnRecordingStarted().AddSP(this, &FSequencer::HandleRecordingStarted);
+		Recorder->OnRecordingFinished().AddSP(this, &FSequencer::HandleRecordingFinished);
+	}
 	ToolkitHost = InitParams.ToolkitHost;
 
 	ScrubPosition = InitParams.ViewParams.InitialScrubPosition;
-	PlayRate = 1.f;
+	PlaybackSpeed = 1.f;
 	ShuttleMultiplier = 0;
 	ObjectChangeListener = InObjectChangeListener;
 
@@ -265,6 +268,8 @@ void FSequencer::InitSequencer(const FSequencerInitParams& InitParams, const TSh
 		ObjectBindings.Add(ObjectBinding);
 	}
 
+	FMovieSceneObjectBindingIDCustomization::BindTo(AsShared());
+
 	ZoomAnimation = FCurveSequence();
 	ZoomCurve = ZoomAnimation.AddCurve(0.f, 0.2f, ECurveEaseFunction::QuadIn);
 	OverlayAnimation = FCurveSequence();
@@ -300,7 +305,6 @@ FSequencer::FSequencer()
 	, bIsEditingWithinLevelEditor( false )
 	, bShowCurveEditor( false )
 	, bNeedTreeRefresh( false )
-	, bNeedInstanceRefresh( false )
 	, StoredPlaybackState( EMovieScenePlayerStatus::Stopped )
 	, NodeTree( MakeShareable( new FSequencerNodeTree( *this ) ) )
 	, bUpdatingSequencerSelection( false )
@@ -343,6 +347,8 @@ void FSequencer::Close()
 
 	SequencerWidget.Reset();
 	TrackEditors.Empty();
+
+	GUnrealEd->UpdatePivotLocationForSelection();
 }
 
 
@@ -365,6 +371,8 @@ void FSequencer::Tick(float InDeltaTime)
 		{
 			PrePossessionViewTargets.Reset();
 			State.ClearObjectCaches(*this);
+			RestorePreAnimatedState();
+			NotifyMovieSceneDataChanged(EMovieSceneDataChangeType::RefreshAllImmediately);
 			CachedPlaybackContext = NewPlaybackContext;
 		}
 	}
@@ -374,18 +382,11 @@ void FSequencer::Tick(float InDeltaTime)
 		bNeedsEvaluate = true;
 	}
 	
-	if ( bNeedInstanceRefresh )
-	{
-		UpdateRuntimeInstances();
-		bNeedInstanceRefresh = false;
-	}
-
 	if (bNeedTreeRefresh)
 	{
 		SelectionPreview.Empty();
 
-		SequencerWidget->UpdateLayoutTree();
-		bNeedTreeRefresh = false;
+		RefreshTree();
 
 		SetPlaybackStatus(StoredPlaybackState);
 	}
@@ -393,7 +394,7 @@ void FSequencer::Tick(float InDeltaTime)
 	UObject* PlaybackContext = GetPlaybackContext();
 	UWorld* World = PlaybackContext ? PlaybackContext->GetWorld() : nullptr;
 	float Dilation = World ? World->GetWorldSettings()->MatineeTimeDilation : 1.f;
-	FTimeAndDelta TimeAndDelta = TimingManager->AdjustTime(GetGlobalTime(), InDeltaTime, PlayRate, Dilation);
+	FTimeAndDelta TimeAndDelta = TimingManager->AdjustTime(GetGlobalTime(), InDeltaTime, PlaybackSpeed, Dilation);
 
 	static const float AutoScrollFactor = 0.1f;
 
@@ -453,7 +454,7 @@ void FSequencer::Tick(float InDeltaTime)
 	// Tick all the tools we own as well
 	for (int32 EditorIndex = 0; EditorIndex < TrackEditors.Num(); ++EditorIndex)
 	{
-		TrackEditors[EditorIndex]->Tick(TimeAndDelta.Delta * PlayRate);
+		TrackEditors[EditorIndex]->Tick(TimeAndDelta.Delta * PlaybackSpeed);
 	}
 
 	if (!IsInSilentMode())
@@ -542,23 +543,17 @@ void FSequencer::ResetToNewRootSequence(UMovieSceneSequence& NewSequence)
 
 void FSequencer::FocusSequenceInstance(UMovieSceneSubSection& InSubSection)
 {
-	// Root out the SequenceID for the sub section
-	FMovieSceneSequenceID SequenceID = MovieSceneSequenceID::Root;
+	FMovieSceneRootOverridePath Path;
+	Path.Set(ActiveTemplateIDs.Last(), RootTemplateInstance.GetHierarchy());
 
-	const FMovieSceneSequenceHierarchyNode* Node = RootTemplateInstance.GetHierarchy().FindNode(ActiveTemplateIDs.Last());
+	// Root out the SequenceID for the sub section
+	FMovieSceneSequenceID SequenceID = Path.Remap(InSubSection.GetSequenceID());
 
-	FName SearchForName(*InSubSection.GetPathNameInMovieScene());
-	for (FMovieSceneSequenceIDRef ChildID : Node->Children)
-	{
-		const FMovieSceneSubSequenceData* SubSequence = RootTemplateInstance.GetHierarchy().FindSubData(ChildID);
-		if (SearchForName == SubSequence->SectionPath)
-		{
-			SequenceID = ChildID;
-			break;
-		}
-	}
+	// Ensure the hierarchy is up to date for this level
+	int32 MaxDepth = 1;
+	FMovieSceneCompiler::CompileHierarchy(*GetFocusedMovieSceneSequence(), RootTemplateInstance.GetHierarchy(), ActiveTemplateIDs.Last(), MaxDepth);
 
-	if (!ensure(SequenceID != MovieSceneSequenceID::Root))
+	if (!ensure(RootTemplateInstance.GetHierarchy().FindSubData(SequenceID)))
 	{
 		return;
 	}
@@ -680,7 +675,9 @@ void FSequencer::PopToSequenceInstance(FMovieSceneSequenceIDRef SequenceID)
 				LevelVC->Invalidate();
 			}
 		}
-		SetPerspectiveViewportCameraCutEnabled(true);
+
+		UMovieSceneSequence* Sequence = GetFocusedMovieSceneSequence();
+		SetPerspectiveViewportCameraCutEnabled(Sequence && Sequence->GetMovieScene()->GetCameraCutTrack());
 
 		ResetPerMovieSceneData();
 		SequencerWidget->UpdateBreadcrumbs();
@@ -697,7 +694,7 @@ void FSequencer::UpdateSubSequenceData()
 	RootToLocalTransform = FMovieSceneSequenceTransform();
 
 	// Find the parent sub section and set up the sub sequence range, if necessary
-	if (ActiveTemplateIDs.Num() <= 1)
+	if (ActiveTemplateIDs.Num() <= 1 || Settings->ShouldEvaluateSubSequencesInIsolation())
 	{
 		return;
 	}
@@ -706,7 +703,7 @@ void FSequencer::UpdateSubSequenceData()
 
 	if (SubSequenceData)
 	{
-		SubSequenceRange = SubSequenceData->ValidPlayRange;
+		SubSequenceRange = SubSequenceData->PlayRange;
 		RootToLocalTransform = SubSequenceData->RootToSequenceTransform;
 	}
 }
@@ -721,15 +718,12 @@ void FSequencer::RerunConstructionScripts()
 
 	GetConstructionScriptActors(Sequence->GetMovieScene(), MovieSceneSequenceID::Root, BoundActors);
 
-	for (auto& SubInstance : RootTemplate.GetSubInstances())
+	for (FMovieSceneSequenceIDRef SequenceID : RootTemplateInstance.GetThisFrameMetaData().ActiveSequences)
 	{
-		if (RootTemplateInstance.GetThisFrameMetaData().ActiveSequences.Contains(SubInstance.Key))
+		UMovieSceneSequence* SubSequence = RootTemplateInstance.GetSequence(SequenceID);
+		if (SubSequence)
 		{
-			UMovieSceneSequence* SubSequence = SubInstance.Value.Sequence.Get();
-			if (SubSequence)
-			{
-				GetConstructionScriptActors(SubSequence->GetMovieScene(), SubInstance.Key, BoundActors);
-			}
+			GetConstructionScriptActors(SubSequence->GetMovieScene(), SequenceID, BoundActors);
 		}
 	}
 
@@ -1042,6 +1036,176 @@ void FSequencer::TranslateSelectedKeysAndSections(bool bTranslateLeft)
 	TransformSelectedKeysAndSections(bTranslateLeft ? -GetFixedFrameInterval() : GetFixedFrameInterval(), 1.f);
 }
 
+void FSequencer::BakeTransform()
+{
+	FScopedTransaction BakeTransform(NSLOCTEXT("Sequencer", "BakeTransform", "Bake Transform"));
+
+	UMovieScene* FocusedMovieScene = GetFocusedMovieSceneSequence()->GetMovieScene();
+	FocusedMovieScene->Modify();
+
+	for (const TSharedRef<FSequencerDisplayNode>& Node : Selection.GetSelectedOutlinerNodes())
+	{
+		if (Node->GetType() != ESequencerNode::Object)
+		{
+			continue;
+		}
+		auto ObjectBindingNode = StaticCastSharedRef<FSequencerObjectBindingNode>(Node);
+
+		FGuid Guid = ObjectBindingNode->GetObjectBinding();
+		for (auto RuntimeObject : FindBoundObjects(Guid, ActiveTemplateIDs.Top()) )
+		{
+			AActor* Actor = Cast<AActor>(RuntimeObject.Get());
+
+			UCameraComponent* CameraComponent = MovieSceneHelpers::CameraComponentFromRuntimeObject(RuntimeObject.Get());
+
+			FVector Location = Actor->GetActorLocation();
+			FRotator Rotation = Actor->GetActorRotation();
+			FVector Scale = Actor->GetActorScale();
+
+			// Cache transforms
+			TArray<FVector> Locations;
+			TArray<FRotator> Rotations;
+			TArray<float> KeyTimes;
+			float InRange = GetPlaybackRange().GetLowerBoundValue();
+			float OutRange = GetPlaybackRange().GetUpperBoundValue();
+			float CurrentTime = ScrubPosition;
+			for (float EvalTime = InRange; EvalTime <= OutRange; EvalTime += Settings->GetTimeSnapInterval())
+			{
+				FMovieSceneEvaluationRange Range = PlayPosition.JumpTo(EvalTime * RootToLocalTransform.Inverse(), FocusedMovieScene->GetOptionalFixedFrameInterval());
+				EvaluateInternal(Range);
+
+				if (CameraComponent)
+				{
+					FTransform AdditiveOffset;
+					float AdditiveFOVOffset;
+					CameraComponent->GetAdditiveOffset(AdditiveOffset, AdditiveFOVOffset);
+
+					FTransform Transform(Actor->GetActorRotation(), Actor->GetActorLocation());
+					FTransform TransformWithAdditiveOffset = Transform * AdditiveOffset;
+					FVector LocalTranslation = TransformWithAdditiveOffset.GetTranslation();
+					FRotator LocalRotation = TransformWithAdditiveOffset.GetRotation().Rotator();
+
+					Locations.Add(LocalTranslation);
+					Rotations.Add(LocalRotation);
+				}
+				else
+				{
+					Locations.Add(Actor->GetActorLocation());
+					Rotations.Add(Actor->GetActorRotation());
+				}
+
+				KeyTimes.Add(EvalTime);
+			}
+
+			// Delete any attach tracks
+			// cbb: this only operates on a single attach section.
+			AActor* AttachParentActor = nullptr;
+			UMovieScene3DAttachTrack* AttachTrack = Cast<UMovieScene3DAttachTrack>(FocusedMovieScene->FindTrack(UMovieScene3DAttachTrack::StaticClass(), Guid));
+			if (AttachTrack)
+			{
+				for (auto AttachSection : AttachTrack->GetAllSections())
+				{
+					FMovieSceneObjectBindingID ConstraintBindingID = (Cast<UMovieScene3DAttachSection>(AttachSection))->GetConstraintBindingID();
+					for (auto ParentObject : FindBoundObjects(ConstraintBindingID.GetGuid(), ConstraintBindingID.GetSequenceID()) )
+					{
+						AttachParentActor = Cast<AActor>(ParentObject.Get());
+						break;
+					}
+				}
+
+				FocusedMovieScene->RemoveTrack(*AttachTrack);
+			}
+
+			// Delete any transform tracks
+			UMovieScene3DTransformTrack* TransformTrack = Cast<UMovieScene3DTransformTrack>(FocusedMovieScene->FindTrack(UMovieScene3DTransformTrack::StaticClass(), Guid, "Transform"));
+			if (TransformTrack)
+			{
+				FocusedMovieScene->RemoveTrack(*TransformTrack);
+			}
+
+			// Delete any camera anim tracks
+			UMovieSceneCameraAnimTrack* CameraAnimTrack = Cast<UMovieSceneCameraAnimTrack>(FocusedMovieScene->FindTrack(UMovieSceneCameraAnimTrack::StaticClass(), Guid));
+			if (CameraAnimTrack)
+			{
+				FocusedMovieScene->RemoveTrack(*CameraAnimTrack);
+			}
+
+			// Delete any camera shake tracks
+			UMovieSceneCameraShakeTrack* CameraShakeTrack = Cast<UMovieSceneCameraShakeTrack>(FocusedMovieScene->FindTrack(UMovieSceneCameraShakeTrack::StaticClass(), Guid));
+			if (CameraShakeTrack)
+			{
+				FocusedMovieScene->RemoveTrack(*CameraShakeTrack);
+			}
+
+			// Reset position
+			FMovieSceneEvaluationRange Range = PlayPosition.JumpTo(ScrubPosition, FocusedMovieScene->GetOptionalFixedFrameInterval());
+			EvaluateInternal(Range);
+
+			// If there was an attach track that was the parent, detach and attach to that actor's parent if it exists
+			FTransform ParentInverseTransform;
+			ParentInverseTransform.SetIdentity();
+			if (AttachParentActor)
+			{
+				Actor->DetachFromActor(FDetachmentTransformRules::KeepRelativeTransform);
+
+				AActor* ExistingParentActor = AttachParentActor->GetAttachParentActor();
+				if (ExistingParentActor)
+				{
+					Actor->AttachToActor(ExistingParentActor, FAttachmentTransformRules::KeepRelativeTransform);
+					ParentInverseTransform = ExistingParentActor->GetActorTransform().Inverse();
+				}
+			}
+
+			// Create new transform track and section
+			TransformTrack = Cast<UMovieScene3DTransformTrack>(FocusedMovieScene->AddTrack(UMovieScene3DTransformTrack::StaticClass(), Guid));
+
+			if (TransformTrack)
+			{
+				UMovieScene3DTransformSection* TransformSection = CastChecked<UMovieScene3DTransformSection>(TransformTrack->CreateNewSection());
+				TransformTrack->AddSection(*TransformSection);
+			
+				TransformSection->SetIsInfinite(true);
+				TransformSection->SetStartTime(InRange);
+				TransformSection->SetEndTime(OutRange);
+
+				TransformSection->SetDefault(FTransformKey(EKey3DTransformChannel::Translation, EAxis::X, Location.X, false /*bUnwindRotation*/));
+				TransformSection->SetDefault(FTransformKey(EKey3DTransformChannel::Translation, EAxis::Y, Location.Y, false /*bUnwindRotation*/));
+				TransformSection->SetDefault(FTransformKey(EKey3DTransformChannel::Translation, EAxis::Z, Location.Z, false /*bUnwindRotation*/));
+
+				TransformSection->SetDefault(FTransformKey(EKey3DTransformChannel::Rotation, EAxis::X, Rotation.Euler().X, false /*bUnwindRotation*/));
+				TransformSection->SetDefault(FTransformKey(EKey3DTransformChannel::Rotation, EAxis::Y, Rotation.Euler().Y, false /*bUnwindRotation*/));
+				TransformSection->SetDefault(FTransformKey(EKey3DTransformChannel::Rotation, EAxis::Z, Rotation.Euler().Z, false /*bUnwindRotation*/));
+
+				TransformSection->SetDefault(FTransformKey(EKey3DTransformChannel::Scale, EAxis::X, Scale.X, false /*bUnwindRotation*/));
+				TransformSection->SetDefault(FTransformKey(EKey3DTransformChannel::Scale, EAxis::Y, Scale.Y, false /*bUnwindRotation*/));
+				TransformSection->SetDefault(FTransformKey(EKey3DTransformChannel::Scale, EAxis::Z, Scale.Z, false /*bUnwindRotation*/));
+
+				for (int32 Counter = 0; Counter < KeyTimes.Num(); ++Counter)
+				{
+					float KeyTime = KeyTimes[Counter];
+					const bool bUnwindRotation = false;
+
+					FTransform Transform(Rotations[Counter], Locations[Counter]);
+					FTransform LocalTransform = ParentInverseTransform * Transform;
+					FVector LocalTranslation = LocalTransform.GetTranslation();
+					FVector LocalRotation = LocalTransform.GetRotation().Euler();
+
+					TransformSection->AddKey(KeyTime, FTransformKey(EKey3DTransformChannel::Translation, EAxis::X, LocalTranslation.X, bUnwindRotation), EMovieSceneKeyInterpolation::Linear);
+					TransformSection->AddKey(KeyTime, FTransformKey(EKey3DTransformChannel::Translation, EAxis::Y, LocalTranslation.Y, bUnwindRotation), EMovieSceneKeyInterpolation::Linear);
+					TransformSection->AddKey(KeyTime, FTransformKey(EKey3DTransformChannel::Translation, EAxis::Z, LocalTranslation.Z, bUnwindRotation), EMovieSceneKeyInterpolation::Linear);
+
+					TransformSection->AddKey(KeyTime, FTransformKey(EKey3DTransformChannel::Rotation, EAxis::X, LocalRotation.X, bUnwindRotation), EMovieSceneKeyInterpolation::Linear);
+					TransformSection->AddKey(KeyTime, FTransformKey(EKey3DTransformChannel::Rotation, EAxis::Y, LocalRotation.Y, bUnwindRotation), EMovieSceneKeyInterpolation::Linear);
+					TransformSection->AddKey(KeyTime, FTransformKey(EKey3DTransformChannel::Rotation, EAxis::Z, LocalRotation.Z, bUnwindRotation), EMovieSceneKeyInterpolation::Linear);
+				}
+			}
+		}
+	}
+	
+	NotifyMovieSceneDataChanged( EMovieSceneDataChangeType::MovieSceneStructureItemsChanged );
+}
+
+
 void FSequencer::OnActorsDropped( const TArray<TWeakObjectPtr<AActor> >& Actors )
 {
 	AddActors(Actors);
@@ -1071,27 +1235,18 @@ void FSequencer::NotifyMovieSceneDataChanged( EMovieSceneDataChangeType DataChan
 		// When structure items are removed, or we don't know what may have changed, refresh the tree and instances immediately so that the data
 		// is in a consistent state when the UI is updated during the next tick.
 		SetPlaybackStatus( EMovieScenePlayerStatus::Stopped );
-		UpdateRuntimeInstances();
 		SelectionPreview.Empty();
-		SequencerWidget->UpdateLayoutTree();
-		bNeedInstanceRefresh = false;
-		bNeedTreeRefresh = false;
+		RefreshTree();
 		SetPlaybackStatus( StoredPlaybackState );
 	}
 	else if (DataChangeType == EMovieSceneDataChangeType::TrackValueChangedRefreshImmediately)
 	{
-		UpdateRuntimeInstances();
-
 		// Evaluate now
 		EvaluateInternal(FMovieSceneEvaluationRange(ScrubPosition));
 	}
 	else if (DataChangeType == EMovieSceneDataChangeType::RefreshAllImmediately)
 	{
-		UpdateRuntimeInstances();
-
-		SequencerWidget->UpdateLayoutTree();
-		bNeedInstanceRefresh = false;
-		bNeedTreeRefresh = false;
+		RefreshTree();
 
 		// Evaluate now
 		EvaluateInternal(FMovieSceneEvaluationRange(ScrubPosition));
@@ -1104,7 +1259,6 @@ void FSequencer::NotifyMovieSceneDataChanged( EMovieSceneDataChangeType DataChan
 			SetPlaybackStatus( EMovieScenePlayerStatus::Stopped );
 			bNeedTreeRefresh = true;
 		}
-		bNeedInstanceRefresh = true;
 	}
 
 	if (DataChangeType == EMovieSceneDataChangeType::TrackValueChanged || 
@@ -1126,6 +1280,11 @@ void FSequencer::NotifyMovieSceneDataChanged( EMovieSceneDataChangeType DataChan
 	OnMovieSceneDataChangedDelegate.Broadcast(DataChangeType);
 }
 
+void FSequencer::RefreshTree()
+{
+	SequencerWidget->UpdateLayoutTree();
+	bNeedTreeRefresh = false;
+}
 
 FAnimatedRange FSequencer::GetViewRange() const
 {
@@ -1790,9 +1949,6 @@ void FSequencer::EvaluateInternal(FMovieSceneEvaluationRange InRange, bool bHasJ
 		RerunConstructionScripts();
 	}
 
-	// If realtime is off, this needs to be called to update the pivot location when scrubbing.
-	GUnrealEd->UpdatePivotLocationForSelection();
-
 	if (!IsInSilentMode())
 	{
 		OnGlobalTimeChangedDelegate.Broadcast();
@@ -2312,7 +2468,6 @@ void FSequencer::SetPlaybackStatus(EMovieScenePlayerStatus::Type InPlaybackStatu
 	{
 		GEngine->SetMaxFPS(OldMaxTickRate);
 
-		PlayRate = 1.f;
 		ShuttleMultiplier = 0;
 	}
 
@@ -2331,10 +2486,10 @@ void FSequencer::AddReferencedObjects( FReferenceCollector& Collector )
 
 	if (RootTemplateInstance.IsValid())
 	{
-		// Sequencer references all sub movie scene sequences contained within the root
-		for (auto& Pair : RootTemplateInstance.GetSubInstances())
+		// Sequencer references all active sub sequences
+		for (FMovieSceneSequenceIDRef SequenceID : RootTemplateInstance.GetThisFrameMetaData().ActiveSequences)
 		{
-			if (UMovieSceneSequence* Sequence = Pair.Value.Sequence.Get())
+			if (UMovieSceneSequence* Sequence = RootTemplateInstance.GetSequence(SequenceID))
 			{
 				Collector.AddReferencedObject(Sequence);
 			}
@@ -2348,10 +2503,9 @@ void FSequencer::ResetPerMovieSceneData()
 	//@todo Sequencer - We may want to preserve selections when moving between movie scenes
 	Selection.Empty();
 
-	SequencerWidget->UpdateLayoutTree();
+	RefreshTree();
 
 	UpdateTimeBoundsToFocusedMovieScene();
-	UpdateRuntimeInstances();
 
 	LabelManager.SetMovieScene( GetFocusedMovieSceneSequence()->GetMovieScene() );
 
@@ -2360,15 +2514,6 @@ void FSequencer::ResetPerMovieSceneData()
 }
 
 
-void FSequencer::UpdateRuntimeInstances()
-{
-	// If realtime is off, this needs to be called to update the pivot location when scrubbing.
-	GUnrealEd->UpdatePivotLocationForSelection();
-	
-	// Redraw
-	FEditorSupportDelegates::RedrawAllViewports.Broadcast();
-}
-
 void FSequencer::RecordSelectedActors()
 {
 	ISequenceRecorder& SequenceRecorder = FModuleManager::LoadModuleChecked<ISequenceRecorder>("SequenceRecorder");
@@ -2453,8 +2598,6 @@ void FSequencer::RecordSelectedActors()
 		const FString& SequenceName = UMovieSceneSubSection::GetRecordingSection()->GetTargetSequenceName();
 		SequenceRecorder.StartRecording(
 			ActorsToRecord,
-			FOnRecordingStarted::CreateSP(this, &FSequencer::HandleRecordingStarted),
-			FOnRecordingFinished::CreateSP(this, &FSequencer::HandleRecordingFinished),
 			PathToRecordTo,
 			SequenceName);
 	}
@@ -2468,8 +2611,8 @@ TSharedRef<SWidget> FSequencer::MakeTransportControls(bool bExtended)
 	{
 		TransportControlArgs.OnBackwardEnd.BindSP( this, &FSequencer::OnJumpToStart );
 		TransportControlArgs.OnBackwardStep.BindSP( this, &FSequencer::OnStepBackward );
-		TransportControlArgs.OnForwardPlay.BindSP( this, &FSequencer::OnPlay, true, 1.f );
-		TransportControlArgs.OnBackwardPlay.BindSP( this, &FSequencer::OnPlay, true, -1.f );
+		TransportControlArgs.OnForwardPlay.BindSP( this, &FSequencer::OnPlayForward, true );
+		TransportControlArgs.OnBackwardPlay.BindSP( this, &FSequencer::OnPlayBackward, true );
 		TransportControlArgs.OnForwardStep.BindSP( this, &FSequencer::OnStepForward );
 		TransportControlArgs.OnForwardEnd.BindSP( this, &FSequencer::OnJumpToEnd );
 		TransportControlArgs.OnGetPlaybackMode.BindSP( this, &FSequencer::GetPlaybackMode );
@@ -2507,36 +2650,44 @@ TSharedRef<SWidget> FSequencer::MakeTransportControls(bool bExtended)
 
 TSharedRef<SWidget> FSequencer::OnCreateTransportSetPlaybackStart()
 {
+	FText SetPlaybackStartToolTip = FText::Format(LOCTEXT("SetPlayStart_Tooltip", "Set playback start to the current position ({0})"), FSequencerCommands::Get().SetStartPlaybackRange->GetInputText());
+
 	return SNew(SButton)
 		.OnClicked(this, &FSequencer::SetPlaybackStart)
-		.ToolTipText(LOCTEXT("SetPlayStart_Tooltip", "Set playback start to the current position"))
+		.ToolTipText(SetPlaybackStartToolTip)
 		.ButtonStyle(FEditorStyle::Get(), "Sequencer.Transport.SetPlayStart")
 		.ContentPadding(2.0f);
 }
 
 TSharedRef<SWidget> FSequencer::OnCreateTransportJumpToPreviousKey()
 {
+	FText JumpToPreviousKeyToolTip = FText::Format(LOCTEXT("JumpToPreviousKey_Tooltip", "Jump to the previous key in the selected track(s) ({0})"), FSequencerCommands::Get().StepToPreviousKey->GetInputText());
+
 	return SNew(SButton)
 		.OnClicked(this, &FSequencer::JumpToPreviousKey)
-		.ToolTipText(LOCTEXT("JumpToPreviousKey_Tooltip", "Jump to the previous key in the selected track(s)"))
+		.ToolTipText(JumpToPreviousKeyToolTip)
 		.ButtonStyle(FEditorStyle::Get(), "Sequencer.Transport.JumpToPreviousKey")
 		.ContentPadding(2.0f);
 }
 
 TSharedRef<SWidget> FSequencer::OnCreateTransportJumpToNextKey()
 {
+	FText JumpToNextKeyToolTip = FText::Format(LOCTEXT("JumpToNextKey_Tooltip", "Jump to the next key in the selected track(s) ({0})"), FSequencerCommands::Get().StepToNextKey->GetInputText());
+
 	return SNew(SButton)
 		.OnClicked(this, &FSequencer::JumpToNextKey)
-		.ToolTipText(LOCTEXT("JumpToNextKey_Tooltip", "Jump to the next key in the selected track(s)"))
+		.ToolTipText(JumpToNextKeyToolTip)
 		.ButtonStyle(FEditorStyle::Get(), "Sequencer.Transport.JumpToNextKey")
 		.ContentPadding(2.0f);
 }
 
 TSharedRef<SWidget> FSequencer::OnCreateTransportSetPlaybackEnd()
 {
+	FText SetPlaybackEndToolTip = FText::Format(LOCTEXT("SetPlayEnd_Tooltip", "Set playback end to the current position ({0})"), FSequencerCommands::Get().SetEndPlaybackRange->GetInputText());
+
 	return SNew(SButton)
 		.OnClicked(this, &FSequencer::SetPlaybackEnd)
-		.ToolTipText(LOCTEXT("SetPlayEnd_Tooltip", "Set playback end to the current position"))
+		.ToolTipText(SetPlaybackEndToolTip)
 		.ButtonStyle(FEditorStyle::Get(), "Sequencer.Transport.SetPlayEnd")
 		.ContentPadding(2.0f);
 }
@@ -2668,21 +2819,21 @@ UObject* FSequencer::FindSpawnedObjectOrTemplate(const FGuid& BindingId)
 	return nullptr;
 }
 
-FReply FSequencer::OnPlay(bool bTogglePlay, float InPlayRate)
+FReply FSequencer::OnPlay(bool bTogglePlay)
 {
 	if( (PlaybackState == EMovieScenePlayerStatus::Playing ||
-		 PlaybackState == EMovieScenePlayerStatus::Recording) && bTogglePlay && (FMath::Sign(InPlayRate) == FMath::Sign(PlayRate) ) )
+		 PlaybackState == EMovieScenePlayerStatus::Recording) && bTogglePlay )
 	{
 		Pause();
 	}
 	else
 	{
-		PlayRate = InPlayRate;
-
 		SetPlaybackStatus(EMovieScenePlayerStatus::Playing);
 
 		// Make sure Slate ticks during playback
 		SequencerWidget->RegisterActiveTimerForPlayback();
+
+		OnPlayDelegate.Broadcast();
 	}
 
 	return FReply::Handled();
@@ -2713,7 +2864,7 @@ FReply FSequencer::OnRecord()
 
 		const FString& PathToRecordTo = UMovieSceneSubSection::GetRecordingSection()->GetTargetPathToRecordTo();
 		const FString& SequenceName = UMovieSceneSubSection::GetRecordingSection()->GetTargetSequenceName();
-		SequenceRecorder.StartRecording(ActorToRecord, FOnRecordingStarted::CreateSP(this, &FSequencer::HandleRecordingStarted), FOnRecordingFinished::CreateSP(this, &FSequencer::HandleRecordingFinished), PathToRecordTo, SequenceName);
+		SequenceRecorder.StartRecording(ActorToRecord, PathToRecordTo, SequenceName);
 	}
 	else if(SequenceRecorder.IsRecording())
 	{
@@ -2725,7 +2876,7 @@ FReply FSequencer::OnRecord()
 
 void FSequencer::HandleRecordingStarted(UMovieSceneSequence* Sequence)
 {
-	OnPlay(false);
+	OnPlayForward(false);
 		
 	// Make sure Slate ticks during playback
 	SequencerWidget->RegisterActiveTimerForPlayback();
@@ -2747,7 +2898,7 @@ void FSequencer::HandleRecordingFinished(UMovieSceneSequence* Sequence)
 	// as the post processing takes such a long time we don't really care if the sequence doesnt carry on
 	if(PlaybackState == EMovieScenePlayerStatus::Playing)
 	{
-		OnPlay(true);
+		OnPlayForward(true);
 	}
 
 	// now patchup the section that was recorded to
@@ -2778,7 +2929,42 @@ void FSequencer::HandleRecordingFinished(UMovieSceneSequence* Sequence)
 	}
 
 	bNeedTreeRefresh = true;
-	bNeedInstanceRefresh = true;
+
+	// If viewing the same sequence, rebuild
+	if (RootSequence.IsValid() && RootSequence.Get() == Sequence)
+	{
+		ResetToNewRootSequence(*RootSequence.Get());
+
+		NotifyMovieSceneDataChanged(EMovieSceneDataChangeType::RefreshAllImmediately);
+	}
+}
+
+FReply FSequencer::OnPlayForward(bool bTogglePlay)
+{
+	if (PlaybackSpeed < 0)
+	{
+		PlaybackSpeed = -PlaybackSpeed;
+		OnPlay(false);
+	}
+	else
+	{
+		OnPlay(bTogglePlay);
+	}
+	return FReply::Handled();
+}
+
+FReply FSequencer::OnPlayBackward(bool bTogglePlay)
+{
+	if (PlaybackSpeed > 0)
+	{
+		PlaybackSpeed = -PlaybackSpeed;
+		OnPlay(false);
+	}
+	else
+	{
+		OnPlay(bTogglePlay);
+	}
+	return FReply::Handled();
 }
 
 FReply FSequencer::OnStepForward()
@@ -2920,7 +3106,7 @@ void FSequencer::SetLocalTimeLooped(float NewLocalTime)
 		{
 			if (NewLocalTime <= TimeBounds.GetLowerBoundValue() || NewLocalTime >= TimeBounds.GetUpperBoundValue())
 			{
-				NewGlobalTime = (PlayRate > 0 ? TimeBounds.GetLowerBoundValue() : TimeBounds.GetUpperBoundValue()) * RootToLocalTransform.Inverse();
+				NewGlobalTime = (PlaybackSpeed > 0 ? TimeBounds.GetLowerBoundValue() : TimeBounds.GetUpperBoundValue()) * RootToLocalTransform.Inverse();
 				TimingManager->OnStartPlaying(NewGlobalTime);
 
 				// Always evaluate from the start/end when looping
@@ -2935,7 +3121,7 @@ void FSequencer::SetLocalTimeLooped(float NewLocalTime)
 		TRange<float> WorkingRange = GetClampRange();
 
 		bool bReachedEnd = false;
-		if (PlayRate > 0)
+		if (PlaybackSpeed > 0)
 		{
 			bReachedEnd = GetLocalTime() < TimeBounds.GetUpperBoundValue() && NewLocalTime >= TimeBounds.GetUpperBoundValue();
 		}
@@ -2947,7 +3133,7 @@ void FSequencer::SetLocalTimeLooped(float NewLocalTime)
 		// Stop if we hit the playback range end
 		if (bReachedEnd)
 		{
-			NewGlobalTime = (PlayRate > 0 ? TimeBounds.GetUpperBoundValue() : TimeBounds.GetLowerBoundValue()) * RootToLocalTransform.Inverse();
+			NewGlobalTime = (PlaybackSpeed > 0 ? TimeBounds.GetUpperBoundValue() : TimeBounds.GetLowerBoundValue()) * RootToLocalTransform.Inverse();
 			TimingManager->OnStartPlaying(NewGlobalTime);
 			NewPlaybackStatus = EMovieScenePlayerStatus::Stopped;
 		}
@@ -2958,7 +3144,7 @@ void FSequencer::SetLocalTimeLooped(float NewLocalTime)
 			if (NewLocalTime <= TimeBounds.GetLowerBoundValue() || 
 				NewLocalTime >= TimeBounds.GetUpperBoundValue())
 			{
-				NewGlobalTime = (PlayRate > 0 ? TimeBounds.GetLowerBoundValue() : TimeBounds.GetUpperBoundValue()) * RootToLocalTransform.Inverse();
+				NewGlobalTime = (PlaybackSpeed > 0 ? TimeBounds.GetLowerBoundValue() : TimeBounds.GetUpperBoundValue()) * RootToLocalTransform.Inverse();
 				TimingManager->OnStartPlaying(NewGlobalTime);
 
 				// Always evaluate from the start/end when looping
@@ -3009,7 +3195,7 @@ EPlaybackMode::Type FSequencer::GetPlaybackMode() const
 {
 	if (PlaybackState == EMovieScenePlayerStatus::Playing)
 	{
-		if (PlayRate > 0)
+		if (PlaybackSpeed > 0)
 		{
 			return EPlaybackMode::PlayingForward;
 		}
@@ -3241,8 +3427,6 @@ void FSequencer::OnEndScrubbing()
 	StopAutoscroll();
 
 	OnEndScrubbingDelegate.Broadcast();
-
-	ForceEvaluate();
 }
 
 
@@ -3323,7 +3507,7 @@ void FSequencer::VerticalScroll(float ScrollAmountUnits)
 	SequencerWidget->GetTreeView()->ScrollByDelta(ScrollAmountUnits);
 }
 
-FGuid FSequencer::AddSpawnable(UObject& Object)
+FGuid FSequencer::AddSpawnable(UObject& Object, UActorFactory* ActorFactory)
 {
 	UMovieSceneSequence* Sequence = GetFocusedMovieSceneSequence();
 	if (!Sequence->AllowsSpawnableObjects())
@@ -3335,7 +3519,7 @@ FGuid FSequencer::AddSpawnable(UObject& Object)
 	// MovieScene asset.
 	UMovieScene* OwnerMovieScene = Sequence->GetMovieScene();
 
-	TValueOrError<FNewSpawnable, FText> Result = SpawnRegister->CreateNewSpawnableType(Object, *OwnerMovieScene);
+	TValueOrError<FNewSpawnable, FText> Result = SpawnRegister->CreateNewSpawnableType(Object, *OwnerMovieScene, ActorFactory);
 	if (!Result.IsValid())
 	{
 		FNotificationInfo Info(Result.GetError());
@@ -3364,12 +3548,10 @@ FGuid FSequencer::AddSpawnable(UObject& Object)
 
 	ForceEvaluate();
 
-	UpdateRuntimeInstances();
-	
 	return NewGuid;
 }
 
-FGuid FSequencer::MakeNewSpawnable( UObject& Object )
+FGuid FSequencer::MakeNewSpawnable( UObject& Object, UActorFactory* ActorFactory )
 {
 	// @todo sequencer: Undo doesn't seem to be working at all
 	const FScopedTransaction Transaction( LOCTEXT("UndoAddingObject", "Add Object to MovieScene") );
@@ -3377,7 +3559,7 @@ FGuid FSequencer::MakeNewSpawnable( UObject& Object )
 	UMovieSceneSequence* Sequence = GetFocusedMovieSceneSequence();
 	UMovieScene* MovieScene = Sequence->GetMovieScene();
 
-	FGuid NewGuid = AddSpawnable(Object);
+	FGuid NewGuid = AddSpawnable(Object, ActorFactory);
 	if (!NewGuid.IsValid())
 	{
 		return FGuid();
@@ -3710,7 +3892,8 @@ void FSequencer::OnNewActorsDropped(const TArray<UObject*>& DroppedObjects, cons
 
 					// Create an attach track
 					UMovieScene3DAttachTrack* AttachTrack = Cast<UMovieScene3DAttachTrack>(OwnerMovieScene->AddTrack(UMovieScene3DAttachTrack::StaticClass(), NewCameraGuid));
-					AttachTrack->AddConstraint(GetPlaybackRange().GetLowerBoundValue(), GetPlaybackRange().GetUpperBoundValue(), NAME_None, NAME_None, NewGuid);
+					FMovieSceneObjectBindingID AttachBindingID(NewGuid, MovieSceneSequenceID::Root);
+					AttachTrack->AddConstraint(GetPlaybackRange().GetLowerBoundValue(), GetPlaybackRange().GetUpperBoundValue(), NAME_None, NAME_None, AttachBindingID);
 				}
 				else
 				{
@@ -3754,8 +3937,7 @@ void FSequencer::OnNewActorsDropped(const TArray<UObject*>& DroppedObjects, cons
 				}
 
 				// New camera added, don't lock the view to the camera because we want to see where the camera rig was placed
-				const bool bLockToCamera = false;
-				NewCameraAdded(NewCamera, NewCameraGuid, bLockToCamera);
+				NewCameraAdded(NewCameraGuid);
 			}
 		}
 
@@ -3941,7 +4123,6 @@ void FSequencer::SaveCurrentMovieScene()
 
 	FEditorFileUtils::PromptForCheckoutAndSave( PackagesToSave, bCheckDirty, bPromptToSave );
 
-	UpdateRuntimeInstances();
 	FMovieSceneEvaluationRange Range = PlayPosition.JumpTo(ScrubPosition, GetFocusedMovieSceneSequence()->GetMovieScene()->GetOptionalFixedFrameInterval());
 	EvaluateInternal(Range);
 
@@ -3996,7 +4177,10 @@ TArray<FGuid> FSequencer::AddActors(const TArray<TWeakObjectPtr<AActor> >& InAct
 				FGuid PossessableGuid = CreateBinding(*Actor, Actor->GetActorLabel());
 				PossessableGuids.Add(PossessableGuid);
 
-				UpdateRuntimeInstances();
+				if (Actor->IsA<ACameraActor>())
+				{
+					NewCameraAdded(PossessableGuid);
+				}
 
 				OnActorAddedToSequencerEvent.Broadcast(Actor, PossessableGuid);
 			}
@@ -4006,7 +4190,7 @@ TArray<FGuid> FSequencer::AddActors(const TArray<TWeakObjectPtr<AActor> >& InAct
 
 	if (bPossessableAdded)
 	{
-		SequencerWidget->UpdateLayoutTree();
+		RefreshTree();
 
 		SynchronizeSequencerSelectionWithExternalSelection();
 	}
@@ -4266,6 +4450,32 @@ void FSequencer::SynchronizeSequencerSelectionWithExternalSelection()
 	}
 }
 
+bool FSequencer::IsBindingVisible(const FMovieSceneBinding& InBinding)
+{
+	if (Settings->GetShowSelectedNodesOnly() && OnGetIsBindingVisible().IsBound())
+	{
+		return OnGetIsBindingVisible().Execute(InBinding);
+	}
+
+	return true;
+}
+
+bool FSequencer::IsTrackVisible(const UMovieSceneTrack* InTrack)
+{
+	if (Settings->GetShowSelectedNodesOnly() && OnGetIsTrackVisible().IsBound())
+	{
+		return OnGetIsTrackVisible().Execute(InTrack);
+	}
+
+	return true;
+}
+
+void FSequencer::OnSelectedNodesOnlyChanged()
+{
+	RefreshTree();
+	
+	SynchronizeSequencerSelectionWithExternalSelection();
+}
 
 void FSequencer::ZoomToSelectedSections()
 {
@@ -4693,21 +4903,7 @@ void FSequencer::DeleteSelectedNodes()
 }
 
 
-void FSequencer::CopySelectedTracks(TArray<TSharedPtr<FSequencerTrackNode>>& TrackNodes)
-{
-	TArray<UMovieSceneTrack*> TracksToCopy;
-	for (TSharedPtr<FSequencerTrackNode> TrackNode : TrackNodes)
-	{
-		TracksToCopy.Add(TrackNode->GetTrack());
-	}
-
-	FString ExportedText;
-	FSequencer::ExportTracksToText(TracksToCopy, /*out*/ ExportedText);
-	FPlatformApplicationMisc::ClipboardCopy(*ExportedText);
-}
-
-
-void FSequencer::ExportTracksToText(TArray<UMovieSceneTrack*> TracksToExport, FString& ExportedText)
+void ExportObjectsToText(const TArray<UMovieSceneCopyableBinding*>& ObjectsToExport, FString& ExportedText)
 {
 	// Clear the mark state for saving.
 	UnMarkAllObjects(EObjectMark(OBJECTMARK_TagExp | OBJECTMARK_TagImp));
@@ -4718,76 +4914,334 @@ void FSequencer::ExportTracksToText(TArray<UMovieSceneTrack*> TracksToExport, FS
 	// Export each of the selected nodes
 	UObject* LastOuter = nullptr;
 
-	for (UMovieSceneTrack* TrackToExport : TracksToExport)
+	for (UMovieSceneCopyableBinding* ObjectToExport : ObjectsToExport)
 	{
 		// The nodes should all be from the same scope
-		UObject* ThisOuter = TrackToExport->GetOuter();
+		UObject* ThisOuter = ObjectToExport->GetOuter();
 		check((LastOuter == ThisOuter) || (LastOuter == nullptr));
 		LastOuter = ThisOuter;
 
-		UExporter::ExportToOutputDevice(&Context, TrackToExport, nullptr, Archive, TEXT("copy"), 0, PPF_ExportsNotFullyQualified | PPF_Copy | PPF_Delimited, false, ThisOuter);
+		UExporter::ExportToOutputDevice(&Context, ObjectToExport, nullptr, Archive, TEXT("copy"), 0, PPF_ExportsNotFullyQualified | PPF_Copy | PPF_Delimited, false, ThisOuter);
+
+		// cbb: Copy the object template separately for now. There's a bug in reimporting the object template from the CopyableBinding.
+		if (ObjectToExport->SpawnableObjectTemplate)
+		{
+			UExporter::ExportToOutputDevice(&Context, ObjectToExport->SpawnableObjectTemplate, nullptr, Archive, TEXT("copy"), 0, PPF_ExportsNotFullyQualified | PPF_Copy | PPF_Delimited, false, ThisOuter);
+		}
 	}
 
 	ExportedText = Archive;
 }
 
-void FSequencer::PasteCopiedTracks()
+class FObjectBindingTextFactory : public FCustomizableTextObjectFactory
 {
-	TSet<TSharedRef<FSequencerDisplayNode>> SelectedNodes = Selection.GetSelectedOutlinerNodes();
+public:
+	FObjectBindingTextFactory()
+		: FCustomizableTextObjectFactory(GWarn)
+	{
+	}
 
-	TArray<TSharedPtr<FSequencerObjectBindingNode>> ObjectNodes;
-	for (TSharedRef<FSequencerDisplayNode> Node : SelectedNodes)
+	// FCustomizableTextObjectFactory implementation
+	virtual bool CanCreateClass(UClass* InObjectClass, bool& bOmitSubObjs) const override
 	{
-		if (Node->GetType() != ESequencerNode::Object)
+		return true;
+	}
+	
+
+	virtual void ProcessConstructedObject(UObject* NewObject) override
+	{
+		check(NewObject);
+
+		if (NewObject->IsA<UMovieSceneCopyableBinding>())
 		{
-			continue;
+			UMovieSceneCopyableBinding* CopyableBinding = Cast<UMovieSceneCopyableBinding>(NewObject);
+			NewCopyableBindings.Add(CopyableBinding);
 		}
-
-		TSharedPtr<FSequencerObjectBindingNode> ObjectNode = StaticCastSharedRef<FSequencerObjectBindingNode>(Node);
-		if (ObjectNode.IsValid())
+		else
 		{
-			ObjectNodes.Add(ObjectNode);
+			NewSpawnableObjectTemplates.Add(NewObject);
 		}
 	}
 
-	FScopedTransaction Transaction(FGenericCommands::Get().Paste->GetDescription());
-	// Grab the text to paste from the clipboard
-	FString TextToImport;
-	FPlatformApplicationMisc::ClipboardPaste(TextToImport);
+public:
+	TArray<UMovieSceneCopyableBinding*> NewCopyableBindings;
+	TArray<UObject*> NewSpawnableObjectTemplates;
+};
 
-	TArray<UMovieSceneTrack*> ImportedTracks;
-	FSequencer::ImportTracksFromText(TextToImport, ImportedTracks);
 
-	if (ImportedTracks.Num() == 0)
-	{
-		Transaction.Cancel();
-		return;
-	}
-	
-	if (ObjectNodes.Num())
+void ImportObjectsFromText(const FString& TextToImport, /*out*/ TArray<UMovieSceneCopyableBinding*>& ImportedObjects)
+{
+	UPackage* TempPackage = NewObject<UPackage>(nullptr, TEXT("/Engine/Sequencer/Editor/Transient"), RF_Transient);
+	TempPackage->AddToRoot();
+
+	// Turn the text buffer into objects
+	FObjectBindingTextFactory Factory;
+	Factory.ProcessBuffer(TempPackage, RF_Transactional, TextToImport);
+	ImportedObjects = Factory.NewCopyableBindings;
+
+	// cbb: Process the object templates separately for now. There's a bug in reimporting the object template from the CopyableBinding.
+	int32 SpawnableObjectTemplateIndex = 0;
+	for (auto ImportedObject : ImportedObjects)
 	{
-		for (TSharedPtr<FSequencerObjectBindingNode> ObjectNode : ObjectNodes)
+		if (ImportedObject->Spawnable.GetGuid().IsValid() && SpawnableObjectTemplateIndex < Factory.NewSpawnableObjectTemplates.Num())
 		{
-			FGuid ObjectGuid = ObjectNode->GetObjectBinding();
+			ImportedObject->SpawnableObjectTemplate = Factory.NewSpawnableObjectTemplates[SpawnableObjectTemplateIndex++];
+		}
+	}
 
-			TArray<UMovieSceneTrack*> NewTracks;
-			FSequencer::ImportTracksFromText(TextToImport, NewTracks);
+	// Remove the temp package from the root now that it has served its purpose
+	TempPackage->RemoveFromRoot();
+}
 
-			for (UMovieSceneTrack* NewTrack : NewTracks)
-			{
-				if (!GetFocusedMovieSceneSequence()->GetMovieScene()->AddGivenTrack(NewTrack, ObjectGuid))
-				{
-					FNotificationInfo Info(LOCTEXT("TrackAlreadyBound", "Can't Paste: Binding doesn't exist"));
-					Info.FadeInDuration = 0.1f;
-					Info.FadeOutDuration = 0.5f;
-					Info.ExpireDuration = 2.5f;
-					auto NotificationItem = FSlateNotificationManager::Get().AddNotification(Info);
 
-					NotificationItem->SetCompletionState(SNotificationItem::CS_Success);
-					NotificationItem->ExpireAndFadeout();
+void FSequencer::CopySelectedObjects(TArray<TSharedPtr<FSequencerObjectBindingNode>>& ObjectNodes)
+{
+	// Gather guids for the object nodes and any child object nodes
+	TSet<FGuid> GuidsToCopy;
+	for (TSharedPtr<FSequencerObjectBindingNode> ObjectNode : ObjectNodes)
+	{
+		GuidsToCopy.Add(ObjectNode->GetObjectBinding());
 
-					continue;
-				}
+		TSet<TSharedRef<FSequencerDisplayNode> > DescendantNodes;
+
+		SequencerHelpers::GetDescendantNodes(ObjectNode.ToSharedRef(), DescendantNodes);
+
+		for (auto DescendantNode : DescendantNodes)
+		{
+			if (DescendantNode->GetType() == ESequencerNode::Object)
+			{
+				GuidsToCopy.Add((StaticCastSharedRef<FSequencerObjectBindingNode>(DescendantNode))->GetObjectBinding());
+			}
+		}
+	}
+
+	UMovieScene* MovieScene = GetFocusedMovieSceneSequence()->GetMovieScene();
+
+	// Export each of the bindings
+	TArray<UMovieSceneCopyableBinding*> CopyableBindings;
+
+	for (auto ObjectBinding : GuidsToCopy)
+	{
+		UMovieSceneCopyableBinding *CopyableBinding = NewObject<UMovieSceneCopyableBinding>(GetTransientPackage(), UMovieSceneCopyableBinding::StaticClass(), NAME_None, RF_Transient);
+		CopyableBindings.Add(CopyableBinding);
+
+		FMovieScenePossessable* Possessable = MovieScene->FindPossessable(ObjectBinding);
+		if (Possessable)
+		{
+			CopyableBinding->Possessable = *Possessable;
+		}
+		else
+		{
+			FMovieSceneSpawnable* Spawnable = MovieScene->FindSpawnable(ObjectBinding);
+			if (Spawnable)
+			{
+				CopyableBinding->Spawnable = *Spawnable;
+				CopyableBinding->SpawnableObjectTemplate = Spawnable->GetObjectTemplate();
+			}
+		}
+
+		const FMovieSceneBinding* Binding = MovieScene->GetBindings().FindByPredicate([=](const FMovieSceneBinding& InBinding){ return InBinding.GetObjectGuid() == ObjectBinding; });
+		if (Binding)
+		{
+			CopyableBinding->Binding = *Binding;
+			for (auto Track : Binding->GetTracks())
+			{
+				UMovieSceneTrack* DuplicatedTrack = Cast<UMovieSceneTrack>(StaticDuplicateObject(Track, CopyableBinding));
+
+				CopyableBinding->Tracks.Add(DuplicatedTrack);
+			}
+		}
+	}
+
+	FString ExportedText;
+	ExportObjectsToText(CopyableBindings, /*out*/ ExportedText);
+	FPlatformApplicationMisc::ClipboardCopy(*ExportedText);
+}
+
+
+void FSequencer::CopySelectedTracks(TArray<TSharedPtr<FSequencerTrackNode>>& TrackNodes)
+{
+	TArray<UMovieSceneTrack*> TracksToCopy;
+	for (TSharedPtr<FSequencerTrackNode> TrackNode : TrackNodes)
+	{
+		TracksToCopy.Add(TrackNode->GetTrack());
+	}
+
+	FString ExportedText;
+	FSequencer::ExportTracksToText(TracksToCopy, /*out*/ ExportedText);
+	FPlatformApplicationMisc::ClipboardCopy(*ExportedText);
+}
+
+
+void FSequencer::ExportTracksToText(TArray<UMovieSceneTrack*> TracksToExport, FString& ExportedText)
+{
+	// Clear the mark state for saving.
+	UnMarkAllObjects(EObjectMark(OBJECTMARK_TagExp | OBJECTMARK_TagImp));
+
+	FStringOutputDevice Archive;
+	const FExportObjectInnerContext Context;
+
+	// Export each of the selected nodes
+	UObject* LastOuter = nullptr;
+
+	for (UMovieSceneTrack* TrackToExport : TracksToExport)
+	{
+		// The nodes should all be from the same scope
+		UObject* ThisOuter = TrackToExport->GetOuter();
+		check((LastOuter == ThisOuter) || (LastOuter == nullptr));
+		LastOuter = ThisOuter;
+
+		UExporter::ExportToOutputDevice(&Context, TrackToExport, nullptr, Archive, TEXT("copy"), 0, PPF_ExportsNotFullyQualified | PPF_Copy | PPF_Delimited, false, ThisOuter);
+	}
+
+	ExportedText = Archive;
+}
+
+void FSequencer::PasteCopiedTracks()
+{
+	FScopedTransaction Transaction(FGenericCommands::Get().Paste->GetDescription());
+	// Grab the text to paste from the clipboard
+	FString TextToImport;
+	FPlatformApplicationMisc::ClipboardPaste(TextToImport);
+
+	UMovieSceneSequence* OwnerSequence = GetFocusedMovieSceneSequence();
+	UObject* BindingContext = GetPlaybackContext();
+
+	TArray<UMovieSceneCopyableBinding*> ImportedBindings;
+	ImportObjectsFromText(TextToImport, ImportedBindings);
+
+	if (ImportedBindings.Num() != 0)
+	{
+		UMovieScene* MovieScene = GetFocusedMovieSceneSequence()->GetMovieScene();
+		TMap<FGuid, FGuid> OldToNewGuidMap;
+		TArray<FGuid> PossessableGuids;
+
+		TArray<FMovieSceneBinding> BindingsPasted;
+		for (UMovieSceneCopyableBinding* CopyableBinding : ImportedBindings)
+		{
+			FGuid NewGuid = FGuid::NewGuid();
+
+			FMovieSceneBinding NewBinding(NewGuid, CopyableBinding->Binding.GetName(), CopyableBinding->Tracks);
+			
+			if (CopyableBinding->Possessable.GetGuid().IsValid())
+			{
+				FMovieScenePossessable NewPossessable = CopyableBinding->Possessable;
+				NewPossessable.SetGuid(NewGuid);
+
+				MovieScene->AddPossessable(NewPossessable, NewBinding);
+
+				OldToNewGuidMap.Add(CopyableBinding->Possessable.GetGuid(), NewGuid);
+
+				BindingsPasted.Add(NewBinding);
+
+				PossessableGuids.Add(NewGuid);
+			}
+			else if (CopyableBinding->Spawnable.GetGuid().IsValid())
+			{
+				FMovieSceneSpawnable NewSpawnable = CopyableBinding->Spawnable;
+				NewSpawnable.SetGuid(NewGuid);
+
+				MovieScene->AddSpawnable(NewSpawnable, NewBinding);
+
+				FMovieSceneSpawnable* Spawnable = MovieScene->FindSpawnable(NewGuid);
+				Spawnable->CopyObjectTemplate(*CopyableBinding->SpawnableObjectTemplate, *GetFocusedMovieSceneSequence());
+
+				OldToNewGuidMap.Add(CopyableBinding->Spawnable.GetGuid(), NewGuid);
+
+				BindingsPasted.Add(NewBinding);
+			}
+		}
+
+		// Fix up parent guids
+		for (auto PossessableGuid : PossessableGuids)
+		{
+			FMovieScenePossessable* Possessable = MovieScene->FindPossessable(PossessableGuid);
+			if (Possessable && OldToNewGuidMap.Contains(Possessable->GetParent()))
+			{
+				Possessable->SetParent(OldToNewGuidMap[Possessable->GetParent()]);
+			}
+		}
+
+		// Fix possessable bindings
+		for (auto PossessableGuid : PossessableGuids)
+		{
+			FMovieScenePossessable* Possessable = MovieScene->FindPossessable(PossessableGuid);
+			UWorld* PlaybackContext = Cast<UWorld>(GetPlaybackContext());
+			if (PlaybackContext)
+			{
+				for ( TActorIterator<AActor> ActorItr( PlaybackContext ); ActorItr; ++ActorItr )
+				{
+					AActor *Actor = *ActorItr;
+					if (Actor && Actor->GetActorLabel() == *Possessable->GetName())
+					{
+						FGuid ExistingGuid = FindObjectId( *Actor, ActiveTemplateIDs.Top() );
+
+						if (!ExistingGuid.IsValid())
+						{
+							DoAssignActor( &Actor, 1, Possessable->GetGuid());
+						}
+					}
+				}
+			}
+		}
+		
+		OnMovieSceneBindingsPastedDelegate.Broadcast(BindingsPasted);
+
+		NotifyMovieSceneDataChanged(EMovieSceneDataChangeType::MovieSceneStructureItemsChanged);
+
+		return;
+	}
+
+	TArray<UMovieSceneTrack*> ImportedTracks;
+	FSequencer::ImportTracksFromText(TextToImport, ImportedTracks);
+
+	if (ImportedTracks.Num() == 0)
+	{
+		Transaction.Cancel();
+		return;
+	}
+	
+	TSet<TSharedRef<FSequencerDisplayNode>> SelectedNodes = Selection.GetSelectedOutlinerNodes();
+
+	TArray<TSharedPtr<FSequencerObjectBindingNode>> ObjectNodes;
+	for (TSharedRef<FSequencerDisplayNode> Node : SelectedNodes)
+	{
+		if (Node->GetType() != ESequencerNode::Object)
+		{
+			continue;
+		}
+
+		TSharedPtr<FSequencerObjectBindingNode> ObjectNode = StaticCastSharedRef<FSequencerObjectBindingNode>(Node);
+		if (ObjectNode.IsValid())
+		{
+			ObjectNodes.Add(ObjectNode);
+		}
+	}
+
+	if (ObjectNodes.Num())
+	{
+		for (TSharedPtr<FSequencerObjectBindingNode> ObjectNode : ObjectNodes)
+		{
+			FGuid ObjectGuid = ObjectNode->GetObjectBinding();
+
+			TArray<UMovieSceneTrack*> NewTracks;
+			FSequencer::ImportTracksFromText(TextToImport, NewTracks);
+
+			for (UMovieSceneTrack* NewTrack : NewTracks)
+			{
+				if (!GetFocusedMovieSceneSequence()->GetMovieScene()->AddGivenTrack(NewTrack, ObjectGuid))
+				{
+					FNotificationInfo Info(LOCTEXT("TrackAlreadyBound", "Can't Paste: Binding doesn't exist"));
+					Info.FadeInDuration = 0.1f;
+					Info.FadeOutDuration = 0.5f;
+					Info.ExpireDuration = 2.5f;
+					auto NotificationItem = FSlateNotificationManager::Get().AddNotification(Info);
+
+					NotificationItem->SetCompletionState(SNotificationItem::CS_Success);
+					NotificationItem->ExpireAndFadeout();
+
+					continue;
+				}
 				else
 				{
 					NotifyMovieSceneDataChanged(EMovieSceneDataChangeType::MovieSceneStructureItemsChanged);
@@ -4849,13 +5303,19 @@ public:
 
 bool FSequencer::CanPaste(const FString& TextToImport) const
 {
-	FTrackObjectTextFactory Factory;
-	if (!Factory.CanCreateObjectsFromText(TextToImport))
+	FObjectBindingTextFactory ObjectBindingFactory;
+	if (ObjectBindingFactory.CanCreateObjectsFromText(TextToImport))
 	{
-		return false;
+		return true;
+	}
+		
+	FTrackObjectTextFactory TrackFactory;
+	if (TrackFactory.CanCreateObjectsFromText(TextToImport))
+	{
+		return true;
 	}
 
-	return true;
+	return false;
 }
 
 void FSequencer::ImportTracksFromText(const FString& TextToImport, /*out*/ TArray<UMovieSceneTrack*>& ImportedTracks)
@@ -5383,16 +5843,9 @@ void FSequencer::OnAddFolder()
 
 void FSequencer::TogglePlay()
 {
-	OnPlay();
+	OnPlayForward(true);
 }
 
-
-void FSequencer::PlayForward()
-{
-	OnPlay(false);
-}
-
-
 void FSequencer::JumpToStart()
 {
 	OnJumpToStart();
@@ -5405,34 +5858,36 @@ void FSequencer::JumpToEnd()
 
 void FSequencer::ShuttleForward()
 {
-	float NewPlayRate = PlayRate;
-	if (ShuttleMultiplier == 0 || PlayRate < 0) 
+	float NewPlaybackSpeed = PlaybackSpeed;
+	if (ShuttleMultiplier == 0 || PlaybackSpeed < 0) 
 	{
 		ShuttleMultiplier = 2.f;
-		NewPlayRate = 1.f;
+		NewPlaybackSpeed = 1.f;
 	}
 	else
 	{
-		NewPlayRate *= ShuttleMultiplier;
+		NewPlaybackSpeed *= ShuttleMultiplier;
 	}
 
-	OnPlay(false, NewPlayRate);
+	PlaybackSpeed = NewPlaybackSpeed;
+	OnPlayForward(false);
 }
 
 void FSequencer::ShuttleBackward()
 {
-	float NewPlayRate = PlayRate;
-	if (ShuttleMultiplier == 0 || PlayRate > 0)
+	float NewPlaybackSpeed = PlaybackSpeed;
+	if (ShuttleMultiplier == 0 || PlaybackSpeed > 0)
 	{
 		ShuttleMultiplier = 2.f;
-		NewPlayRate = -1.f;
+		NewPlaybackSpeed = -1.f;
 	}
 	else
 	{
-		NewPlayRate *= ShuttleMultiplier;
+		NewPlaybackSpeed *= ShuttleMultiplier;
 	}
 
-	OnPlay(false, NewPlayRate);
+	PlaybackSpeed = NewPlaybackSpeed;
+	OnPlayBackward(false);
 }
 
 void FSequencer::Pause()
@@ -5455,6 +5910,8 @@ void FSequencer::Pause()
 		FMovieSceneEvaluationRange Range = LastRange.IsSet() ? LastRange.GetValue() : PlayPosition.JumpTo(ScrubPosition, GetFocusedMovieSceneSequence()->GetMovieScene()->GetOptionalFixedFrameInterval());
 		EvaluateInternal(Range);
 	}
+
+	OnStopDelegate.Broadcast();
 }
 
 void FSequencer::StepForward()
@@ -5500,10 +5957,11 @@ void FSequencer::StepToNextShot()
 		return;
 	}
 
-	UMovieSceneSequence* Sequence = RootTemplateInstance.GetSequence(ActiveTemplateIDs[ActiveTemplateIDs.Num()-2]);
+	FMovieSceneSequenceID OuterSequenceID = ActiveTemplateIDs[ActiveTemplateIDs.Num()-2];
+	UMovieSceneSequence* Sequence = RootTemplateInstance.GetSequence(OuterSequenceID);
 
 	float StartTime = 0.f * RootToLocalTransform.Inverse();
-	float CurrentTime = StartTime * RootTemplateInstance.GetInstance(ActiveTemplateIDs[ActiveTemplateIDs.Num()-2])->RootToSequenceTransform;
+	float CurrentTime = StartTime * RootTemplateInstance.GetHierarchy().FindSubData(OuterSequenceID)->RootToSequenceTransform;
 
 	UMovieSceneSubSection* NextShot = Cast<UMovieSceneSubSection>(FindNextOrPreviousShot(Sequence, CurrentTime, true));
 	if (!NextShot)
@@ -5528,10 +5986,11 @@ void FSequencer::StepToPreviousShot()
 		return;
 	}
 
-	UMovieSceneSequence* Sequence = RootTemplateInstance.GetSequence(ActiveTemplateIDs[ActiveTemplateIDs.Num()-2]);
+	FMovieSceneSequenceID OuterSequenceID = ActiveTemplateIDs[ActiveTemplateIDs.Num()-2];
+	UMovieSceneSequence* Sequence = RootTemplateInstance.GetSequence(OuterSequenceID);
 
 	float StartTime = 0.f * RootToLocalTransform.Inverse();
-	float CurrentTime = StartTime * RootTemplateInstance.GetInstance(ActiveTemplateIDs[ActiveTemplateIDs.Num()-2])->RootToSequenceTransform;
+	float CurrentTime = StartTime * RootTemplateInstance.GetHierarchy().FindSubData(OuterSequenceID)->RootToSequenceTransform;
 
 	UMovieSceneSubSection* PreviousShot = Cast<UMovieSceneSubSection>(FindNextOrPreviousShot(Sequence, CurrentTime, false));
 	if (!PreviousShot)
@@ -5754,6 +6213,7 @@ void FSequencer::CopySelection()
 	if (Selection.GetSelectedKeys().Num() == 0)
 	{
 		TArray<TSharedPtr<FSequencerTrackNode>> TracksToCopy;
+		TArray<TSharedPtr<FSequencerObjectBindingNode>> ObjectsToCopy;
 		TSet<TSharedRef<FSequencerDisplayNode>> SelectedNodes = Selection.GetNodesWithSelectedKeysOrSections();
 		if (SelectedNodes.Num() == 0)
 		{
@@ -5761,15 +6221,21 @@ void FSequencer::CopySelection()
 		}
 		for (TSharedRef<FSequencerDisplayNode> Node : SelectedNodes)
 		{
-			if (Node->GetType() != ESequencerNode::Track)
+			if (Node->GetType() == ESequencerNode::Track)
 			{
-				continue;
+				TSharedPtr<FSequencerTrackNode> TrackNode = StaticCastSharedRef<FSequencerTrackNode>(Node);
+				if (TrackNode.IsValid())
+				{
+					TracksToCopy.Add(TrackNode);
+				}
 			}
-
-			TSharedPtr<FSequencerTrackNode> TrackNode = StaticCastSharedRef<FSequencerTrackNode>(Node);
-			if (TrackNode.IsValid())
+			else if (Node->GetType() == ESequencerNode::Object)
 			{
-				TracksToCopy.Add(TrackNode);
+				TSharedPtr<FSequencerObjectBindingNode> ObjectNode = StaticCastSharedRef<FSequencerObjectBindingNode>(Node);
+				if (ObjectNode.IsValid())
+				{
+					ObjectsToCopy.Add(ObjectNode);
+				}
 			}
 		}
 
@@ -5779,7 +6245,15 @@ void FSequencer::CopySelection()
 			TSharedRef<FMovieSceneClipboard> NullClipboard = MakeShareable(new FMovieSceneClipboard());
 			GClipboardStack.Push(NullClipboard);
 		}
-		CopySelectedTracks(TracksToCopy);
+
+		if (ObjectsToCopy.Num())
+		{
+			CopySelectedObjects(ObjectsToCopy);
+		}
+		else
+		{
+			CopySelectedTracks(TracksToCopy);
+		}
 	}
 	else
 	{
@@ -5787,46 +6261,12 @@ void FSequencer::CopySelection()
 	}
 }
 
-
 void FSequencer::CutSelection()
 {
-	FScopedTransaction CutSelectionTransaction(LOCTEXT("CutSelection_Transaction", "Cut Selection(s)"));
 	if (Selection.GetSelectedKeys().Num() == 0)
 	{
-		TArray<TSharedPtr<FSequencerTrackNode>> TracksToCopy;
-		TSet<TSharedRef<FSequencerDisplayNode>> SelectedNodes = Selection.GetNodesWithSelectedKeysOrSections();
-		if (SelectedNodes.Num() == 0)
-		{
-			SelectedNodes = Selection.GetSelectedOutlinerNodes();
-		}
-		for (TSharedRef<FSequencerDisplayNode> Node : SelectedNodes)
-		{
-			if (Node->GetType() != ESequencerNode::Track)
-			{
-				FNotificationInfo Info(LOCTEXT("InvalidCut", "Warning: One of the selected node is not a track node"));
-				Info.FadeInDuration = 0.1f;
-				Info.FadeOutDuration = 0.5f;
-				Info.ExpireDuration = 2.5f;
-				auto NotificationItem = FSlateNotificationManager::Get().AddNotification(Info);
-
-				NotificationItem->SetCompletionState(SNotificationItem::CS_Success);
-				NotificationItem->ExpireAndFadeout();
-				return;
-			}
-
-			TSharedPtr<FSequencerTrackNode> TrackNode = StaticCastSharedRef<FSequencerTrackNode>(Node);
-			if (TrackNode.IsValid())
-			{
-				TracksToCopy.Add(TrackNode);
-			}
-		}
-		// Make a empty clipboard if the stack is empty
-		if (GClipboardStack.Num() == 0)
-		{
-			TSharedRef<FMovieSceneClipboard> NullClipboard = MakeShareable(new FMovieSceneClipboard());
-			GClipboardStack.Push(NullClipboard);
-		}
-		CopySelectedTracks(TracksToCopy);
+		FScopedTransaction CutSelectionTransaction(LOCTEXT("CutSelection_Transaction", "Cut Selection"));
+		CopySelection();
 		DeleteSelectedItems();
 	}
 	else
@@ -5835,6 +6275,12 @@ void FSequencer::CutSelection()
 	}
 }
 
+void FSequencer::DuplicateSelection()
+{
+	CopySelection();
+	PasteCopiedTracks();
+}
+
 void FSequencer::CopySelectedKeys()
 {
 	TOptional<float> CopyRelativeTo;
@@ -6100,8 +6546,7 @@ void FSequencer::CreateCamera()
 
 	OnActorAddedToSequencerEvent.Broadcast(NewCamera, CameraGuid);
 
-	const bool bLockToCamera = true;
-	NewCameraAdded(NewCamera, CameraGuid, bLockToCamera);
+	NewCameraAdded(CameraGuid, NewCamera);
 
 	if (ensure(Spawnable))
 	{
@@ -6111,12 +6556,12 @@ void FSequencer::CreateCamera()
 	NotifyMovieSceneDataChanged(EMovieSceneDataChangeType::MovieSceneStructureItemAdded);
 }
 
-void FSequencer::NewCameraAdded(ACineCameraActor* NewCamera, FGuid CameraGuid, bool bLockToCamera)
+void FSequencer::NewCameraAdded(FGuid CameraGuid, ACameraActor* NewCamera)
 {
 	SetPerspectiveViewportCameraCutEnabled(false);
 
 	// Lock the viewport to this camera
-	if (bLockToCamera && NewCamera && NewCamera->GetLevel())
+	if (NewCamera && NewCamera->GetLevel())
 	{
 		GCurrentLevelEditingViewportClient->SetMatineeActorLock(nullptr);
 		GCurrentLevelEditingViewportClient->SetActorLock(NewCamera);
@@ -6135,9 +6580,14 @@ void FSequencer::NewCameraAdded(ACineCameraActor* NewCamera, FGuid CameraGuid, b
 		return;
 	}
 
-	// If there's a camera cut track, create or set the camera section to this new camera
 	UMovieSceneTrack* CameraCutTrack = OwnerMovieScene->GetCameraCutTrack();
 
+	// If there's a camera cut track with at least one section, no need to change the section
+	if (CameraCutTrack && CameraCutTrack->GetAllSections().Num() > 0)
+	{
+		return;
+	}
+
 	if (!CameraCutTrack)
 	{
 		CameraCutTrack = OwnerMovieScene->AddCameraCutTrack(UMovieSceneCameraCutTrack::StaticClass());
@@ -6169,12 +6619,19 @@ void FSequencer::NewCameraAdded(ACineCameraActor* NewCamera, FGuid CameraGuid, b
 
 void FSequencer::FixActorReferences()
 {
+	UWorld* PlaybackContext = Cast<UWorld>(GetPlaybackContext());
+
+	if (!PlaybackContext)
+	{
+		return;
+	}
+
 	FScopedTransaction FixActorReferencesTransaction( NSLOCTEXT( "Sequencer", "FixActorReferences", "Fix Actor References" ) );
 
 	UMovieScene* FocusedMovieScene = GetFocusedMovieSceneSequence()->GetMovieScene();
 
 	TMap<FString, AActor*> ActorNameToActorMap;
-	UWorld* PlaybackContext = Cast<UWorld>(GetPlaybackContext());
+
 	for ( TActorIterator<AActor> ActorItr( PlaybackContext ); ActorItr; ++ActorItr )
 	{
 		// Same as with the Object Iterator, access the subclass instance with the * or -> operators.
@@ -6520,6 +6977,23 @@ void FSequencer::ExportFBX()
 }
 
 
+void FSequencer::ExportToCameraAnim()
+{
+	for (const TSharedRef<FSequencerDisplayNode>& Node : Selection.GetSelectedOutlinerNodes())
+	{
+		if (Node->GetType() != ESequencerNode::Object)
+		{
+			continue;
+		}
+		auto ObjectBindingNode = StaticCastSharedRef<FSequencerObjectBindingNode>(Node);
+
+		FGuid Guid = ObjectBindingNode->GetObjectBinding();
+		
+		MovieSceneToolHelpers::ExportToCameraAnim(GetFocusedMovieSceneSequence()->GetMovieScene(), Guid);
+	}
+}
+
+
 void FSequencer::GenericTextEntryModeless(const FText& DialogText, const FText& DefaultText, FOnTextCommitted OnTextComitted)
 {
 	TSharedRef<STextEntryPopup> TextEntryPopup = 
@@ -6603,55 +7077,6 @@ void FSequencer::BindCommands()
 		Commands.StepToPreviousCameraKey,
 		FExecuteAction::CreateSP( this, &FSequencer::StepToPreviousCameraKey ) );
 
-	SequencerCommandBindings->MapAction(
-		Commands.StepToNextShot,
-		FExecuteAction::CreateSP( this, &FSequencer::StepToNextShot ) );
-
-	SequencerCommandBindings->MapAction(
-		Commands.StepToPreviousShot,
-		FExecuteAction::CreateSP( this, &FSequencer::StepToPreviousShot ) );
-
-	SequencerCommandBindings->MapAction(
-		Commands.SetStartPlaybackRange,
-		FExecuteAction::CreateSP( this, &FSequencer::SetPlaybackRangeStart ),
-		FCanExecuteAction::CreateSP( this, &FSequencer::IsViewingMasterSequence ) );
-
-	SequencerCommandBindings->MapAction(
-		Commands.ResetViewRange,
-		FExecuteAction::CreateSP( this, &FSequencer::ResetViewRange ) );
-
-	SequencerCommandBindings->MapAction(
-		Commands.ZoomInViewRange,
-		FExecuteAction::CreateSP( this, &FSequencer::ZoomInViewRange ),
-		FCanExecuteAction(),
-		EUIActionRepeatMode::RepeatEnabled );
-
-	SequencerCommandBindings->MapAction(
-		Commands.ZoomOutViewRange,
-		FExecuteAction::CreateSP( this, &FSequencer::ZoomOutViewRange ),		
-		FCanExecuteAction(),
-		EUIActionRepeatMode::RepeatEnabled );
-
-	SequencerCommandBindings->MapAction(
-		Commands.SetEndPlaybackRange,
-		FExecuteAction::CreateSP( this, &FSequencer::SetPlaybackRangeEnd ),
-		FCanExecuteAction::CreateSP( this, &FSequencer::IsViewingMasterSequence ) );
-
-	SequencerCommandBindings->MapAction(
-		Commands.SetSelectionRangeToNextShot,
-		FExecuteAction::CreateSP( this, &FSequencer::SetSelectionRangeToShot, true ),
-		FCanExecuteAction::CreateSP( this, &FSequencer::IsViewingMasterSequence ) );
-
-	SequencerCommandBindings->MapAction(
-		Commands.SetSelectionRangeToPreviousShot,
-		FExecuteAction::CreateSP( this, &FSequencer::SetSelectionRangeToShot, false ),
-		FCanExecuteAction::CreateSP( this, &FSequencer::IsViewingMasterSequence ) );
-
-	SequencerCommandBindings->MapAction(
-		Commands.SetPlaybackRangeToAllShots,
-		FExecuteAction::CreateSP( this, &FSequencer::SetPlaybackRangeToAllShots ),
-		FCanExecuteAction::CreateSP( this, &FSequencer::IsViewingMasterSequence ) );
-
 	SequencerCommandBindings->MapAction(
 		Commands.ExpandAllNodesAndDescendants,
 		FExecuteAction::CreateSP(this, &FSequencer::ExpandAllNodesAndDescendants));
@@ -6854,6 +7279,14 @@ void FSequencer::BindCommands()
 		FCanExecuteAction::CreateLambda( []{ return true; } ),
 		FIsActionChecked::CreateLambda( [this]{ return Settings->GetLabelBrowserVisible(); } ) );
 
+	SequencerCommandBindings->MapAction(
+		Commands.ToggleShowSelectedNodesOnly,
+		FExecuteAction::CreateLambda( [this]{
+			Settings->SetShowSelectedNodesOnly( !Settings->GetShowSelectedNodesOnly() );
+		} ),
+		FCanExecuteAction::CreateLambda( []{ return true; } ),
+		FIsActionChecked::CreateLambda( [this]{ return Settings->GetShowSelectedNodesOnly(); } ) );
+
 	SequencerCommandBindings->MapAction(
 		Commands.ToggleShowFrameNumbers,
 		FExecuteAction::CreateLambda( [this]{ Settings->SetShowFrameNumbers( !Settings->GetShowFrameNumbers() ); } ),
@@ -6946,7 +7379,6 @@ void FSequencer::BindCommands()
 
 	auto CanCutOrCopy = [this]{
 		// For copy tracks
-		TArray<TSharedPtr<FSequencerTrackNode>> TracksToCopy;
 		TSet<TSharedRef<FSequencerDisplayNode>> SelectedNodes = Selection.GetNodesWithSelectedKeysOrSections();
 		// If this is empty then we are selecting display nodes
 		if (SelectedNodes.Num() == 0)
@@ -6954,7 +7386,7 @@ void FSequencer::BindCommands()
 			SelectedNodes = Selection.GetSelectedOutlinerNodes();
 			for (TSharedRef<FSequencerDisplayNode> Node : SelectedNodes)
 			{
-				if (Node->GetType() == ESequencerNode::Track)
+				if (Node->GetType() == ESequencerNode::Track || Node->GetType() == ESequencerNode::Object)
 				{
 					// if contains one node that can be copied we allow the action
 					// later on we will filter out the invalid nodes in CopySelection() or CutSelection()
@@ -6989,6 +7421,25 @@ void FSequencer::BindCommands()
 		return Selection.GetSelectedKeys().Num() || Selection.GetSelectedSections().Num() || Selection.GetSelectedOutlinerNodes().Num();
 	};
 
+	auto CanDuplicate = [this]{
+		// For duplicate object tracks
+		TSet<TSharedRef<FSequencerDisplayNode>> SelectedNodes = Selection.GetNodesWithSelectedKeysOrSections();
+		if (SelectedNodes.Num() == 0)
+		{
+			SelectedNodes = Selection.GetSelectedOutlinerNodes();
+			for (TSharedRef<FSequencerDisplayNode> Node : SelectedNodes)
+			{
+				if (Node->GetType() == ESequencerNode::Object)
+				{
+					// if contains one node that can be copied we allow the action
+					return true;
+				}
+			}
+			return false;
+		}
+		return false;
+	};
+
 	SequencerCommandBindings->MapAction(
 		FGenericCommands::Get().Cut,
 		FExecuteAction::CreateSP(this, &FSequencer::CutSelection),
@@ -7001,6 +7452,12 @@ void FSequencer::BindCommands()
 		FCanExecuteAction::CreateLambda(CanCutOrCopy)
 	);
 
+	SequencerCommandBindings->MapAction(
+		FGenericCommands::Get().Duplicate,
+		FExecuteAction::CreateSP(this, &FSequencer::DuplicateSelection),
+		FCanExecuteAction::CreateLambda(CanDuplicate)
+	);
+
 	SequencerCommandBindings->MapAction(
 		FGenericCommands::Get().Delete,
 		FExecuteAction::CreateSP( this, &FSequencer::DeleteSelectedItems ),
@@ -7104,6 +7561,11 @@ void FSequencer::BindCommands()
 		})
 	);
 
+	SequencerCommandBindings->MapAction(
+		Commands.BakeTransform,
+		FExecuteAction::CreateSP( this, &FSequencer::BakeTransform ),
+		FCanExecuteAction::CreateLambda( []{ return true; } ) );
+
 	SequencerCommandBindings->MapAction(
 		Commands.FixActorReferences,
 		FExecuteAction::CreateSP( this, &FSequencer::FixActorReferences ),
@@ -7129,6 +7591,11 @@ void FSequencer::BindCommands()
 		FExecuteAction::CreateSP( this, &FSequencer::ExportFBX ),
 		FCanExecuteAction::CreateLambda( [] { return true; } ) );
 
+	SequencerCommandBindings->MapAction(
+		Commands.ExportToCameraAnim,
+		FExecuteAction::CreateSP( this, &FSequencer::ExportToCameraAnim ),
+		FCanExecuteAction::CreateLambda( [] { return true; } ) );
+
 	for (int32 i = 0; i < TrackEditors.Num(); ++i)
 	{
 		TrackEditors[i]->BindCommands(SequencerCommandBindings);
@@ -7164,7 +7631,7 @@ void FSequencer::BindCommands()
 
 	SequencerCommandBindings->MapAction(
 		Commands.PlayForward,
-		FExecuteAction::CreateSP( this, &FSequencer::PlayForward ));
+		FExecuteAction::CreateLambda( [this] { OnPlayForward(false); }));
 
 	SequencerCommandBindings->MapAction(
 		Commands.JumpToStart,
@@ -7220,6 +7687,55 @@ void FSequencer::BindCommands()
 		Commands.SelectAllInSelectionRange,
 		FExecuteAction::CreateSP(this, &FSequencer::SelectInSelectionRange, true, true));
 
+	SequencerCommandBindings->MapAction(
+		Commands.StepToNextShot,
+		FExecuteAction::CreateSP( this, &FSequencer::StepToNextShot ) );
+
+	SequencerCommandBindings->MapAction(
+		Commands.StepToPreviousShot,
+		FExecuteAction::CreateSP( this, &FSequencer::StepToPreviousShot ) );
+
+	SequencerCommandBindings->MapAction(
+		Commands.SetStartPlaybackRange,
+		FExecuteAction::CreateSP( this, &FSequencer::SetPlaybackRangeStart ),
+		FCanExecuteAction::CreateSP( this, &FSequencer::IsViewingMasterSequence ) );
+
+	SequencerCommandBindings->MapAction(
+		Commands.ResetViewRange,
+		FExecuteAction::CreateSP( this, &FSequencer::ResetViewRange ) );
+
+	SequencerCommandBindings->MapAction(
+		Commands.ZoomInViewRange,
+		FExecuteAction::CreateSP( this, &FSequencer::ZoomInViewRange ),
+		FCanExecuteAction(),
+		EUIActionRepeatMode::RepeatEnabled );
+
+	SequencerCommandBindings->MapAction(
+		Commands.ZoomOutViewRange,
+		FExecuteAction::CreateSP( this, &FSequencer::ZoomOutViewRange ),		
+		FCanExecuteAction(),
+		EUIActionRepeatMode::RepeatEnabled );
+
+	SequencerCommandBindings->MapAction(
+		Commands.SetEndPlaybackRange,
+		FExecuteAction::CreateSP( this, &FSequencer::SetPlaybackRangeEnd ),
+		FCanExecuteAction::CreateSP( this, &FSequencer::IsViewingMasterSequence ) );
+
+	SequencerCommandBindings->MapAction(
+		Commands.SetSelectionRangeToNextShot,
+		FExecuteAction::CreateSP( this, &FSequencer::SetSelectionRangeToShot, true ),
+		FCanExecuteAction::CreateSP( this, &FSequencer::IsViewingMasterSequence ) );
+
+	SequencerCommandBindings->MapAction(
+		Commands.SetSelectionRangeToPreviousShot,
+		FExecuteAction::CreateSP( this, &FSequencer::SetSelectionRangeToShot, false ),
+		FCanExecuteAction::CreateSP( this, &FSequencer::IsViewingMasterSequence ) );
+
+	SequencerCommandBindings->MapAction(
+		Commands.SetPlaybackRangeToAllShots,
+		FExecuteAction::CreateSP( this, &FSequencer::SetPlaybackRangeToAllShots ),
+		FCanExecuteAction::CreateSP( this, &FSequencer::IsViewingMasterSequence ) );
+
 	// bind widget specific commands
 	SequencerWidget->BindCommands(SequencerCommandBindings);
 }
diff --git a/Engine/Source/Editor/Sequencer/Private/Sequencer.h b/Engine/Source/Editor/Sequencer/Private/Sequencer.h
--- a/Engine/Source/Editor/Sequencer/Private/Sequencer.h
+++ b/Engine/Source/Editor/Sequencer/Private/Sequencer.h
@@ -39,7 +39,7 @@
 #include "SequencerSettings.h"
 
 class AActor;
-class ACineCameraActor;
+class ACameraActor;
 class APlayerController;
 class FLevelEditorViewportClient;
 class FMenuBuilder;
@@ -262,6 +262,9 @@ public:
 	/** Translate the selected keys and section by the time snap interval */
 	void TranslateSelectedKeysAndSections(bool bTranslateLeft);
 
+	/** Bake transform */
+	void BakeTransform();
+
 	/**
 	 * @return Movie scene tools used by the sequencer
 	 */
@@ -302,10 +305,11 @@ protected:
 	 * Attempts to add a new spawnable to the MovieScene for the specified asset, class, or actor
 	 *
 	 * @param	Object	The asset, class, or actor to add a spawnable for
+	 * @param	ActorFactory	Optional actor factory to use to create spawnable type
 	 *
 	 * @return	The spawnable ID, or invalid ID on failure
 	 */
-	FGuid AddSpawnable( UObject& Object );
+	FGuid AddSpawnable( UObject& Object, UActorFactory* ActorFactory = nullptr );
 
 	/**
 	 * Save default spawnable state for the currently selected objects
@@ -392,6 +396,8 @@ public:
 
 	/** Functions to push on to the transport controls we use */
 	FReply OnRecord();
+	FReply OnPlayForward(bool bTogglePlay);
+	FReply OnPlayBackward(bool bTogglePlay);
 	FReply OnStepForward();
 	FReply OnStepBackward();
 	FReply OnJumpToStart();
@@ -414,8 +420,6 @@ public:
 	/** Set the new global time, accounting for looping options */
 	void SetLocalTimeLooped(float InTime);
 
-	float AutoScroll(float InTime, ESnapTimeMode SnapTimeMode);
-
 	ESequencerLoopMode GetLoopMode() const;
 
 	EPlaybackMode::Type GetPlaybackMode() const;
@@ -450,6 +454,7 @@ public:
 	void DeleteSelectedNodes();
 
 	/** Called when a user executes the copy track menu item */
+	void CopySelectedObjects(TArray<TSharedPtr<FSequencerObjectBindingNode>>& ObjectNodes);
 	void CopySelectedTracks(TArray<TSharedPtr<FSequencerTrackNode>>& TrackNodes);
 	void ExportTracksToText(TArray<UMovieSceneTrack*> TrackToExport, /*out*/ FString& ExportedText);
 
@@ -484,6 +489,14 @@ public:
 
 	/** Updates the sequencer selection to match the current external selection. */
 	void SynchronizeSequencerSelectionWithExternalSelection();
+		
+	/** Whether the binding is visible in the tree view */
+	bool IsBindingVisible(const FMovieSceneBinding& InBinding);
+
+	/** Whether the track is visible in the tree view */
+	bool IsTrackVisible(const UMovieSceneTrack* InTrack);
+
+	void OnSelectedNodesOnlyChanged();
 
 public:
 
@@ -493,6 +506,9 @@ public:
 	/** Cut the selection, whether it's keys or tracks */
 	void CutSelection();
 
+	/** Duplicate the selection */
+	void DuplicateSelection();
+
 	/** Copy the selected keys to the clipboard */
 	void CopySelectedKeys();
 
@@ -511,8 +527,8 @@ public:
 	/** Create camera and set it as the current camera cut. */
 	void CreateCamera();
 
-	/** Called when a new camera is added */
-	void NewCameraAdded(ACineCameraActor* NewCamera, FGuid CameraGuid, bool bLockToCamera);
+	/** Called when a new camera is added. Locks the viewport to the NewCamera is not null. */
+	void NewCameraAdded(FGuid CameraGuid, ACameraActor* NewCamera = nullptr);
 
 	/** Attempts to automatically fix up broken actor references in the current scene. */
 	void FixActorReferences();
@@ -529,6 +545,9 @@ public:
 	/** Exports the animation to an fbx file. */
 	void ExportFBX();
 
+	/** Exports the animation to a camera anim asset. */
+	void ExportToCameraAnim();
+
 public:
 	
 	/** Access the currently active track area edit tool */
@@ -556,6 +575,11 @@ public:
 	/** Scroll the sequencer vertically by the specified number of slate units */
 	void VerticalScroll(float ScrollAmountUnits);
 
+	/**
+	 * Update auto-scroll mechanics as a result of a new time position
+	 */
+	void UpdateAutoScroll(float NewTime);
+
 public:
 
 	//~ FGCObject Interface
@@ -615,8 +639,10 @@ protected:
 	virtual void NotifyMovieSceneDataChangedInternal() override;
 public:
 	virtual void NotifyMovieSceneDataChanged( EMovieSceneDataChangeType DataChangeType ) override;
-	virtual void UpdateRuntimeInstances() override;
+	virtual void RefreshTree() override;
 	virtual void UpdatePlaybackRange() override;
+	virtual void SetPlaybackSpeed(float InPlaybackSpeed) override { PlaybackSpeed = InPlaybackSpeed; }
+	virtual float GetPlaybackSpeed() const override { return PlaybackSpeed; }
 	virtual TArray<FGuid> AddActors(const TArray<TWeakObjectPtr<AActor> >& InActors) override;
 	virtual void AddSubSequence(UMovieSceneSequence* Sequence) override;
 	virtual bool CanKeyProperty(FCanKeyPropertyParams CanKeyPropertyParams) const override;
@@ -631,10 +657,13 @@ public:
 	virtual void SelectByPropertyPaths(const TArray<FString>& InPropertyPaths) override;
 	virtual void EmptySelection() override;
 	virtual FOnGlobalTimeChanged& OnGlobalTimeChanged() override { return OnGlobalTimeChangedDelegate; }
+	virtual FOnPlayEvent& OnPlayEvent() override { return OnPlayDelegate; }
+	virtual FOnStopEvent& OnStopEvent() override { return OnStopDelegate; }
 	virtual FOnBeginScrubbingEvent& OnBeginScrubbingEvent() override { return OnBeginScrubbingDelegate; }
 	virtual FOnEndScrubbingEvent& OnEndScrubbingEvent() override { return OnEndScrubbingDelegate; }
 	virtual FOnMovieSceneDataChanged& OnMovieSceneDataChanged() override { return OnMovieSceneDataChangedDelegate; }
 	virtual FOnMovieSceneBindingsChanged& OnMovieSceneBindingsChanged() override { return OnMovieSceneBindingsChangedDelegate; }
+	virtual FOnMovieSceneBindingsPasted& OnMovieSceneBindingsPasted() override { return OnMovieSceneBindingsPastedDelegate; }
 	virtual FOnSelectionChangedObjectGuids& GetSelectionChangedObjectGuids() override { return OnSelectionChangedObjectGuidsDelegate; }
 	virtual FOnSelectionChangedTracks& GetSelectionChangedTracks() override { return OnSelectionChangedTracksDelegate; }
 	virtual FOnSelectionChangedSections& GetSelectionChangedSections() override { return OnSelectionChangedSectionsDelegate; }
@@ -649,11 +678,11 @@ public:
 	virtual TSharedRef<INumericTypeInterface<float>> GetNumericTypeInterface() override;
 	virtual TSharedRef<INumericTypeInterface<float>> GetZeroPadNumericTypeInterface() override;
 	virtual TSharedRef<SWidget> MakeTransportControls(bool bExtended) override;
-	virtual FReply OnPlay(bool bTogglePlay = true, float InPlayRate = 1.f) override;
+	virtual FReply OnPlay(bool bTogglePlay = true) override;
 	virtual void Pause() override;
 	virtual TSharedRef<SWidget> MakeTimeRange(const TSharedRef<SWidget>& InnerContent, bool bShowWorkingRange, bool bShowViewRange, bool bShowPlaybackRange) override;
 	virtual UObject* FindSpawnedObjectOrTemplate(const FGuid& BindingId) override;
-	virtual FGuid MakeNewSpawnable(UObject& SourceObject) override;
+	virtual FGuid MakeNewSpawnable(UObject& SourceObject, UActorFactory* ActorFactory = nullptr) override;
 	virtual bool IsReadOnly() const override;
 	virtual void ExternalSelectionHasChanged() override { SynchronizeSequencerSelectionWithExternalSelection(); }
 	virtual USequencerSettings* GetSequencerSettings() override { return Settings; }
@@ -672,10 +701,6 @@ public:
 	virtual FMovieSceneSpawnRegister& GetSpawnRegister() override { return *SpawnRegister; }
 	virtual bool IsPreview() const override { return SilentModeCount != 0; }
 
-public:
-
-	FMovieSceneRootEvaluationTemplateInstance& GetSequenceInstance() { return RootTemplateInstance; }
-
 protected:
 
 	/** Reevaluate the sequence at the current time */
@@ -743,11 +768,6 @@ protected:
 
 protected:
 
-	/**
-	 * Update auto-scroll mechanics as a result of a new time position
-	 */
-	void UpdateAutoScroll(float NewTime);
-
 	/**
 	 * Ensure that the specified local time is in the view
 	 */
@@ -788,7 +808,6 @@ protected:
 	
 	/** Transport controls */
 	void TogglePlay();
-	void PlayForward();
 	void JumpToStart();
 	void JumpToEnd();
 	void ShuttleForward();
@@ -959,8 +978,8 @@ private:
 	/** Current play position */
 	FMovieScenePlaybackPosition PlayPosition;
 
-	/** The playback rate */
-	float PlayRate;
+	/** The playback speed */
+	float PlaybackSpeed;
 
 	/** The shuttle multiplier */
 	float ShuttleMultiplier;
@@ -984,9 +1003,6 @@ private:
 		the MovieScene data can change many times per frame.) */
 	bool bNeedTreeRefresh;
 
-	/** When true, the runtime instances need to be updated next frame. */
-	bool bNeedInstanceRefresh;
-
 	/** Stores the playback status to be restored on refresh. */
 	EMovieScenePlayerStatus::Type StoredPlaybackState;
 
@@ -1000,6 +1016,12 @@ private:
 	/** A delegate which is called any time the global time changes. */
 	FOnGlobalTimeChanged OnGlobalTimeChangedDelegate;
 
+	/** A delegate which is called whenever the user begins playing the sequence. */
+	FOnPlayEvent OnPlayDelegate;
+
+	/** A delegate which is called whenever the user stops playing the sequence. */
+	FOnStopEvent OnStopDelegate;
+
 	/** A delegate which is called whenever the user begins scrubbing. */
 	FOnBeginScrubbingEvent OnBeginScrubbingDelegate;
 
@@ -1012,6 +1034,9 @@ private:
 	/** A delegate which is called any time the movie scene bindings are changed. */
 	FOnMovieSceneBindingsChanged OnMovieSceneBindingsChangedDelegate;
 
+	/** A delegate which is called any time a binding is pasted. */
+	FOnMovieSceneBindingsPasted OnMovieSceneBindingsPastedDelegate;
+
 	/** A delegate which is called any time the sequencer selection changes. */
 	FOnSelectionChangedObjectGuids OnSelectionChangedObjectGuidsDelegate;
 
diff --git a/Engine/Source/Editor/Sequencer/Private/SequencerCommands.cpp b/Engine/Source/Editor/Sequencer/Private/SequencerCommands.cpp
--- a/Engine/Source/Editor/Sequencer/Private/SequencerCommands.cpp
+++ b/Engine/Source/Editor/Sequencer/Private/SequencerCommands.cpp
@@ -82,6 +82,7 @@ void FSequencerCommands::RegisterCommands()
 	UI_COMMAND( ToggleShowFrameNumbers, "Show Frame Numbers", "Enables and disables showing frame numbers", EUserInterfaceActionType::ToggleButton, FInputChord(EModifierKey::Control, EKeys::T) );
 	UI_COMMAND( ToggleShowGotoBox, "Go to Time...", "Go to a particular point on the timeline", EUserInterfaceActionType::Button, FInputChord(EModifierKey::Control, EKeys::G) );
 	UI_COMMAND( ToggleShowTransformBox, "Transform...", "Transform the selected keys and sections by a given amount", EUserInterfaceActionType::Button, FInputChord(EModifierKey::Control, EKeys::M) );
+	UI_COMMAND( BakeTransform, "Bake Transform", "Bake transform in world space, removing any existing transform and attach tracks", EUserInterfaceActionType::Button, FInputChord() );
 	UI_COMMAND( ToggleShowRangeSlider, "Show Range Slider", "Enables and disables showing the time range slider", EUserInterfaceActionType::ToggleButton, FInputChord() );
 	UI_COMMAND( ToggleIsSnapEnabled, "Enable Snapping", "Enables and disables snapping", EUserInterfaceActionType::ToggleButton, FInputChord() );
 
@@ -102,6 +103,7 @@ void FSequencerCommands::RegisterCommands()
 	UI_COMMAND( ToggleCombinedKeyframes, "Combined Keyframes", "Show/hide the combined keyframes at the top node level", EUserInterfaceActionType::ToggleButton, FInputChord() );
 	UI_COMMAND( ToggleChannelColors, "Channel Colors", "Show/hide the channel colors in the track area", EUserInterfaceActionType::ToggleButton, FInputChord() );
 	UI_COMMAND( ToggleLabelBrowser, "Label Browser", "Show/hide the track label browser", EUserInterfaceActionType::ToggleButton, FInputChord() );
+	UI_COMMAND( ToggleShowSelectedNodesOnly, "Selected Nodes Only", "Show selected nodes only", EUserInterfaceActionType::ToggleButton, FInputChord() );
 	
 	UI_COMMAND( ToggleShowCurveEditor, "Curve Editor", "Show the animation keys in a curve editor", EUserInterfaceActionType::ToggleButton, FInputChord() );
 	UI_COMMAND( ToggleLinkCurveEditorTimeRange, "Link Curve Editor Time Range", "Link the curve editor time range to the sequence", EUserInterfaceActionType::ToggleButton, FInputChord() );
@@ -131,6 +133,8 @@ void FSequencerCommands::RegisterCommands()
 	UI_COMMAND( ImportFBX, "Import...", "Imports the animation from an FBX file.", EUserInterfaceActionType::Button, FInputChord() );
 	UI_COMMAND( ExportFBX, "Export...", "Exports the animation to an FBX file. (Shots and sub-scenes not supported)", EUserInterfaceActionType::Button, FInputChord() );
 
+	UI_COMMAND( ExportToCameraAnim, "Export to Camera Anim...", "Exports the animation to a camera anim asset.", EUserInterfaceActionType::Button, FInputChord() );
+
 	UI_COMMAND( ToggleEvaluateSubSequencesInIsolation, "Evaluate Sub Sequences In Isolation", "When enabled, will only evaluate the currently focused sequence; otherwise evaluate from the master sequence.", EUserInterfaceActionType::ToggleButton, FInputChord() );
 }
 
diff --git a/Engine/Source/Editor/Sequencer/Private/SequencerCommonHelpers.cpp b/Engine/Source/Editor/Sequencer/Private/SequencerCommonHelpers.cpp
--- a/Engine/Source/Editor/Sequencer/Private/SequencerCommonHelpers.cpp
+++ b/Engine/Source/Editor/Sequencer/Private/SequencerCommonHelpers.cpp
@@ -308,6 +308,7 @@ void SequencerHelpers::UpdateHoveredNodeFromSelectedKeys(FSequencer& Sequencer)
 void SequencerHelpers::PerformDefaultSelection(FSequencer& Sequencer, const FPointerEvent& MouseEvent)
 {
 	FSequencerSelection& Selection = Sequencer.GetSelection();
+	Selection.SuspendBroadcast();
 
 	// @todo: selection in transactions
 	auto ConditionallyClearSelection = [&]{
@@ -323,6 +324,8 @@ void SequencerHelpers::PerformDefaultSelection(FSequencer& Sequencer, const FPoi
 	if (!Hotspot.IsValid())
 	{
 		ConditionallyClearSelection();
+		Selection.ResumeBroadcast();
+		Selection.GetOnOutlinerNodeSelectionChanged().Broadcast();
 		return;
 	}
 
@@ -366,6 +369,8 @@ void SequencerHelpers::PerformDefaultSelection(FSequencer& Sequencer, const FPoi
 			UpdateHoveredNodeFromSelectedSections(Sequencer);
 		}
 		
+		Selection.ResumeBroadcast();
+		Selection.GetOnOutlinerNodeSelectionChanged().Broadcast();
 		return;
 	}
 
@@ -412,6 +417,9 @@ void SequencerHelpers::PerformDefaultSelection(FSequencer& Sequencer, const FPoi
 	{
 		UpdateHoveredNodeFromSelectedSections(Sequencer);
 	}
+
+	Selection.ResumeBroadcast();
+	Selection.GetOnOutlinerNodeSelectionChanged().Broadcast();
 }
 
 TSharedPtr<SWidget> SequencerHelpers::SummonContextMenu(FSequencer& Sequencer, const FGeometry& MyGeometry, const FPointerEvent& MouseEvent)
diff --git a/Engine/Source/Editor/Sequencer/Private/SequencerContextMenus.cpp b/Engine/Source/Editor/Sequencer/Private/SequencerContextMenus.cpp
--- a/Engine/Source/Editor/Sequencer/Private/SequencerContextMenus.cpp
+++ b/Engine/Source/Editor/Sequencer/Private/SequencerContextMenus.cpp
@@ -692,9 +692,8 @@ void FSectionContextMenu::AddPropertiesMenu(FMenuBuilder& MenuBuilder)
 	}
 
 	TSharedRef<IDetailsView> DetailsView = FModuleManager::GetModuleChecked<FPropertyEditorModule>("PropertyEditor").CreateDetailView(DetailsViewArgs);
-	{
-		DetailsView->SetObjects(Sections);
-	}
+	Sequencer->OnInitializeDetailsPanel().Broadcast(DetailsView, Sequencer);
+	DetailsView->SetObjects(Sections);
 
 	DetailsNotifyWrapper->SetDetailsAndSequencer(DetailsView, Sequencer);
 	MenuBuilder.AddWidget(DetailsNotifyWrapper, FText::GetEmpty(), true);
@@ -1777,15 +1776,22 @@ void FEasingContextMenu::PopulateMenu(FMenuBuilder& MenuBuilder)
 
 		auto OnBeginSliderMovement = [=]
 		{
-			if (ensure(!Shared->ScopedTransaction.IsValid()))
+			GEditor->BeginTransaction(LOCTEXT("SetEasingTimeText", "Set Easing Length"));
+		};
+		auto OnEndSliderMovement = [=](float NewLength)
+		{
+			if (GEditor->IsTransactionActive())
 			{
-				Shared->ScopedTransaction.Reset(new FScopedTransaction(LOCTEXT("SetEasingTimeText", "Set Easing Length")));
+				GEditor->EndTransaction();
 			}
 		};
-		auto OnEndSliderMovement = [=](float NewLength)
+		auto OnValueCommitted = [=](float NewLength, ETextCommit::Type CommitInfo)
 		{
-			Shared->OnUpdateLength(NewLength);
-			Shared->ScopedTransaction.Reset();
+			if (CommitInfo == ETextCommit::OnEnter || CommitInfo == ETextCommit::OnUserMovedFocus)
+			{
+				FScopedTransaction Transaction(LOCTEXT("SetEasingTimeText", "Set Easing Length"));
+				Shared->OnUpdateLength(NewLength);
+			}
 		};
 
 		TSharedRef<SWidget> SpinBox = SNew(SHorizontalBox)
@@ -1809,7 +1815,7 @@ void FEasingContextMenu::PopulateMenu(FMenuBuilder& MenuBuilder)
 					.Delta(0.001f)
 					.Value_Lambda([=]{ return Shared->GetCurrentLength(); })
 					.OnValueChanged_Lambda([=](float NewLength){ Shared->OnUpdateLength(NewLength); })
-					.OnValueCommitted_Lambda([=](float NewLength, ETextCommit::Type){ Shared->OnUpdateLength(NewLength); })
+					.OnValueCommitted_Lambda(OnValueCommitted)
 					.OnBeginSliderMovement_Lambda(OnBeginSliderMovement)
 					.OnEndSliderMovement_Lambda(OnEndSliderMovement)
 					.BorderForegroundColor(FEditorStyle::GetSlateColor("DefaultForeground"))
@@ -2065,11 +2071,15 @@ void FEasingContextMenu::EasingOptionsMenu(FMenuBuilder& MenuBuilder)
 		{
 			if (Handle.EasingType == ESequencerEasingType::In)
 			{
-				Objects.AddUnique(Section->Easing.EaseIn.GetObject());
+				UObject* EaseInObject = Section->Easing.EaseIn.GetObject();
+				EaseInObject->SetFlags(RF_Transactional);
+				Objects.AddUnique(EaseInObject);
 			}
 			else
 			{
-				Objects.AddUnique(Section->Easing.EaseOut.GetObject());
+				UObject* EaseOutObject = Section->Easing.EaseOut.GetObject();
+				EaseOutObject->SetFlags(RF_Transactional);
+				Objects.AddUnique(EaseOutObject);
 			}
 		}
 	}
diff --git a/Engine/Source/Editor/Sequencer/Private/SequencerCurveOwner.cpp b/Engine/Source/Editor/Sequencer/Private/SequencerCurveOwner.cpp
--- a/Engine/Source/Editor/Sequencer/Private/SequencerCurveOwner.cpp
+++ b/Engine/Source/Editor/Sequencer/Private/SequencerCurveOwner.cpp
@@ -162,6 +162,16 @@ void FSequencerCurveOwner::ModifyOwner()
 	}
 }
 
+void FSequencerCurveOwner::ModifyOwnerChange()
+{
+	TArray<UMovieSceneSection*> Owners;
+	EditInfoToSectionMap.GenerateValueArray( Owners );
+	for ( auto Owner : Owners )
+	{
+		Owner->MarkAsChanged();
+	}
+}
+
 TArray<const UObject*> FSequencerCurveOwner::GetOwners() const
 {
 	TArray<UMovieSceneSection*> Owners;
diff --git a/Engine/Source/Editor/Sequencer/Private/SequencerCurveOwner.h b/Engine/Source/Editor/Sequencer/Private/SequencerCurveOwner.h
--- a/Engine/Source/Editor/Sequencer/Private/SequencerCurveOwner.h
+++ b/Engine/Source/Editor/Sequencer/Private/SequencerCurveOwner.h
@@ -22,6 +22,7 @@ public:
 	virtual TArray<FRichCurveEditInfoConst> GetCurves() const override;
 	virtual TArray<FRichCurveEditInfo> GetCurves() override;
 	virtual void ModifyOwner() override;
+	virtual void ModifyOwnerChange() override;
 	virtual TArray<const UObject*> GetOwners() const override;
 	virtual void MakeTransactional() override;
 	virtual void OnCurveChanged( const TArray<FRichCurveEditInfo>& ChangedCurveEditInfos ) override;
