Copying //UE4/Dev-Sequencer to //UE4/Dev-Main (Source: //UE4/Dev-Sequencer @ 3946692)

diff --git a/Engine/Source/Editor/SequenceRecorder/Private/SequenceRecorder.cpp b/Engine/Source/Editor/SequenceRecorder/Private/SequenceRecorder.cpp
--- a/Engine/Source/Editor/SequenceRecorder/Private/SequenceRecorder.cpp
+++ b/Engine/Source/Editor/SequenceRecorder/Private/SequenceRecorder.cpp
@@ -31,6 +31,8 @@
 #include "Tracks/MovieSceneAudioTrack.h"
 #include "Sections/MovieSceneAudioSection.h"
 #include "Sound/SoundWave.h"
+#include "SequenceRecorderActorGroup.h"
+#include "MovieSceneTimeHelpers.h"
 
 #define LOCTEXT_NAMESPACE "SequenceRecorder"
 
@@ -143,10 +145,6 @@ void FSequenceRecorder::AddNewQueuedRecordingsForSelectedActors()
 				bAnySelectedActorsAdded = true;
 			}
 		}
-		else
-		{
-			bAnySelectedActorsAdded = true;
-		}
 	}
 
 	if (!bAnySelectedActorsAdded)
@@ -159,7 +157,7 @@ UActorRecording* FSequenceRecorder::AddNewQueuedRecording(AActor* Actor, UAnimSe
 {
 	const USequenceRecorderSettings* Settings = GetDefault<USequenceRecorderSettings>();
 
-	UActorRecording* ActorRecording = NewObject<UActorRecording>();
+	UActorRecording* ActorRecording = NewObject<UActorRecording>(CurrentRecorderGroup.IsValid() ? CurrentRecorderGroup.Get() : (UObject*)GetTransientPackage());
 	ActorRecording->AddToRoot();
 	ActorRecording->SetActorToRecord(Actor);
 	ActorRecording->TargetAnimation = AnimSequence;
@@ -179,6 +177,10 @@ UActorRecording* FSequenceRecorder::AddNewQueuedRecording(AActor* Actor, UAnimSe
 	}
 
 	QueuedRecordings.Add(ActorRecording);
+	if (CurrentRecorderGroup.IsValid() && !CurrentRecorderGroup->RecordedActors.Contains(ActorRecording))
+	{
+		CurrentRecorderGroup->RecordedActors.Add(ActorRecording);
+	}
 
 	bQueuedRecordingsDirty = true;
 
@@ -238,6 +240,18 @@ bool FSequenceRecorder::HasQueuedRecordings() const
 	return QueuedRecordings.Num() > 0;
 }
 
+bool FSequenceRecorder::IsRecording() const
+{
+	for(UActorRecording* Recording : QueuedRecordings)
+	{
+		if (Recording->IsRecording())
+		{
+			return true;
+		}
+	}
+	return false;
+}
+
 
 void FSequenceRecorder::Tick(float DeltaSeconds)
 {
@@ -277,7 +291,7 @@ void FSequenceRecorder::Tick(float DeltaSeconds)
 		}
 	}
 
-	if(Settings->bCreateLevelSequence && CurrentSequence.IsValid())
+	if (IsRecording())
 	{
 		CurrentTime += DeltaSeconds;
 
@@ -364,14 +378,19 @@ void FSequenceRecorder::DrawDebug(UCanvas* InCanvas, APlayerController* InPlayer
 
 	if(bCountingDown || IsRecording())
 	{
+		const USequenceRecorderSettings* Settings = GetDefault<USequenceRecorderSettings>();
+	
 		FText LabelText;
-		if(IsRecording())
-		{
-			LabelText = FText::Format(LOCTEXT("RecordingIndicatorFormat", "{0}"), FText::FromName(CurrentSequence.Get()->GetFName()));
-		}
-		else
+		if (Settings->bCreateLevelSequence)
 		{
-			LabelText = FText::Format(LOCTEXT("RecordingIndicatorPending", "Pending recording: {0}"), FText::FromString(NextSequenceName));
+			if (CurrentSequence.IsValid())
+			{
+				LabelText = FText::Format(LOCTEXT("RecordingIndicatorFormat", "{0}"), FText::FromName(CurrentSequence.Get()->GetFName()));
+			}
+			else
+			{
+				LabelText = FText::Format(LOCTEXT("RecordingIndicatorPending", "Pending recording: {0}"), FText::FromString(NextSequenceName));
+			}
 		}
 
 
@@ -441,7 +460,7 @@ bool FSequenceRecorder::StartRecording(const FString& InPathToRecordTo, const FS
 	}
 	else
 	{
-		PathToRecordTo = Settings->SequenceRecordingBasePath.Path;
+		PathToRecordTo = GetSequenceRecordingBasePath();
 	}
 
 	if(InSequenceName.Len() > 0)
@@ -450,7 +469,7 @@ bool FSequenceRecorder::StartRecording(const FString& InPathToRecordTo, const FS
 	}
 	else
 	{
-		SequenceName = Settings->SequenceName.Len() > 0 ? Settings->SequenceName : TEXT("RecordedSequence");
+		SequenceName = GetSequenceRecordingName().Len() > 0 ? GetSequenceRecordingName() : TEXT("RecordedSequence");
 	}
 
 	CurrentTime = 0.0f;
@@ -560,7 +579,7 @@ bool FSequenceRecorder::StartRecordingInternal(UWorld* World)
 				if (DupActorToTrigger->SequencePlayer)
 				{
 					DupActorToTrigger->SequencePlayer->SetDisableCameraCuts(true);
-					DupActorToTrigger->SequencePlayer->SetPlaybackPosition(0.f);
+					DupActorToTrigger->SequencePlayer->JumpToFrame(0);
 					DupActorToTrigger->SequencePlayer->Play();
 				}
 				else
@@ -660,7 +679,7 @@ bool FSequenceRecorder::StartRecordingInternal(UWorld* World)
 			if (LevelSequence)
 			{
 				FDirectoryPath AudioDirectory;
-				AudioDirectory.Path = Settings->SequenceRecordingBasePath.Path;
+				AudioDirectory.Path = GetSequenceRecordingBasePath();
 				if (Settings->AudioSubDirectory.Len())
 				{
 					AudioDirectory.Path /= Settings->AudioSubDirectory;
@@ -715,6 +734,8 @@ bool FSequenceRecorder::StopRecording()
 {
 	if (!IsRecording())
 	{
+		CurrentDelay = 0.0f;
+	
 		return false;
 	}
 
@@ -771,10 +792,11 @@ bool FSequenceRecorder::StopRecording()
 
 			UMovieSceneAudioSection* NewAudioSection = NewObject<UMovieSceneAudioSection>(AudioTrack, UMovieSceneAudioSection::StaticClass());
 
+			FFrameRate FrameResolution = AudioTrack->GetTypedOuter<UMovieScene>()->GetFrameResolution();
+
 			NewAudioSection->SetRowIndex(RowIndex + 1);
 			NewAudioSection->SetSound(RecordedAudio);
-			NewAudioSection->SetStartTime(0);
-			NewAudioSection->SetEndTime(RecordedAudio->GetDuration());
+			NewAudioSection->SetRange(TRange<FFrameNumber>(FFrameNumber(0), (RecordedAudio->GetDuration() * FrameResolution).CeilToFrame()));
 
 			AudioTrack->AddSection(*NewAudioSection);
 		}
@@ -918,24 +940,29 @@ void FSequenceRecorder::UpdateSequencePlaybackRange()
 		UMovieScene* MovieScene = CurrentSequence->GetMovieScene();
 		if(MovieScene)
 		{
-			float MaxRange = 0.0f;
-			float MinRange = 0.0f;
+			TRange<FFrameNumber> PlayRange = TRange<FFrameNumber>::Empty();
 
 			TArray<UMovieSceneSection*> MovieSceneSections = MovieScene->GetAllSections();
 			for(UMovieSceneSection* Section : MovieSceneSections)
 			{
-				MaxRange = FMath::Max(MaxRange, Section->GetEndTime());
-				MinRange = FMath::Min(MinRange, Section->GetStartTime());
+				TRange<FFrameNumber> SectionRange = Section->GetRange();
+				if (SectionRange.GetLowerBound().IsClosed() && SectionRange.GetUpperBound().IsClosed())
+				{
+					PlayRange = TRange<FFrameNumber>::Hull(PlayRange, SectionRange);
+				}
 			}
 
-			MovieScene->SetPlaybackRange(MinRange, MaxRange);
+			MovieScene->SetPlaybackRange(PlayRange);
 
 			// Initialize the working and view range with a little bit more space
-			const float OutputViewSize = MaxRange - MinRange;
-			const float OutputChange = OutputViewSize * 0.1f;
-
-			MovieScene->SetWorkingRange(MinRange - OutputChange, MaxRange + OutputChange);
-			MovieScene->SetViewRange(MinRange - OutputChange, MaxRange + OutputChange);
+			FFrameRate  FrameResolution = MovieScene->GetFrameResolution();
+			const double OutputViewSize = PlayRange.Size<FFrameNumber>() / FrameResolution;
+			const double OutputChange   = OutputViewSize * 0.1;
+
+			TRange<double> NewRange = MovieScene::ExpandRange(PlayRange / FrameResolution, OutputChange);
+			FMovieSceneEditorData& EditorData = MovieScene->GetEditorData();
+			EditorData.ViewStart = EditorData.WorkStart = NewRange.GetLowerBoundValue();
+			EditorData.ViewEnd   = EditorData.WorkEnd   = NewRange.GetUpperBoundValue();
 		}
 	}
 }
@@ -1029,14 +1056,191 @@ void FSequenceRecorder::HandleActorDespawned(AActor* Actor)
 
 void FSequenceRecorder::RefreshNextSequence()
 {
-	const USequenceRecorderSettings* Settings = GetDefault<USequenceRecorderSettings>();
-	if (SequenceName.IsEmpty())
+	SequenceName = GetSequenceRecordingName().Len() > 0 ? GetSequenceRecordingName() : TEXT("RecordedSequence");
+
+	// Cache the name of the next sequence we will try to record to
+	NextSequenceName = SequenceRecorderUtils::MakeNewAssetName<ULevelSequence>(GetSequenceRecordingBasePath(), SequenceName);
+}
+
+TWeakObjectPtr<ASequenceRecorderGroup> FSequenceRecorder::GetRecordingGroupActor()
+{
+	if (CachedRecordingActor.IsValid())
 	{
-		SequenceName = Settings->SequenceName.Len() > 0 ? Settings->SequenceName : TEXT("RecordedSequence");
+		return CachedRecordingActor;
 	}
 
-	// Cache the name of the next sequence we will try to record to
-	NextSequenceName = SequenceRecorderUtils::MakeNewAssetName<ULevelSequence>(Settings->SequenceRecordingBasePath.Path, SequenceName);
+	// Check the map for one
+	UWorld* EditorWorld = GEditor->GetEditorWorldContext().World();
+	ASequenceRecorderGroup* GroupActor = nullptr;
+
+	if (EditorWorld && EditorWorld->PersistentLevel)
+	{
+		for (int32 ActorIndex = 0; ActorIndex < EditorWorld->PersistentLevel->Actors.Num(); ++ActorIndex)
+		{
+			AActor* Actor = EditorWorld->PersistentLevel->Actors[ActorIndex];
+			GroupActor = Cast<ASequenceRecorderGroup>(Actor);
+			if (GroupActor)
+			{
+				// We want to find the first actor
+				break;
+			}
+		}
+	}
+
+	// We may not have one, or we may be in a situation where we can't safely create
+	// an actor, calling functions should expect this to possibly be null.
+	CachedRecordingActor = GroupActor;
+	return CachedRecordingActor;
+}
+
+void FSequenceRecorder::AddRecordingGroup()
+{
+	TWeakObjectPtr<ASequenceRecorderGroup> GroupActor = GetRecordingGroupActor();
+	UWorld* EditorWorld = GEditor->GetEditorWorldContext().World();
+
+	// There may not be a group actor in the level yet, so we'll spawn a new one.
+	if (!GroupActor.IsValid())
+	{
+		GroupActor = (ASequenceRecorderGroup*)GEditor->AddActor(EditorWorld->PersistentLevel, ASequenceRecorderGroup::StaticClass(), FTransform::Identity);
+		CachedRecordingActor = GroupActor;
+	}
+
+	// Now add a new actor group to this actor
+	check(GroupActor.IsValid());
+	USequenceRecorderActorGroup* ActorGroup = NewObject<USequenceRecorderActorGroup>(GroupActor.Get());
+	ActorGroup->GroupName = FName(*FGuid::NewGuid().ToString());
+	GroupActor->ActorGroups.Add(ActorGroup);
+
+	// Remove the existing queued recordings which marks us as dirty so the UI will refresh too.
+	ClearQueuedRecordings();
+
+	// And then select our new object by default
+	CurrentRecorderGroup = ActorGroup;
+}
+
+void FSequenceRecorder::RemoveCurrentRecordingGroup()
+{
+	if (!GetRecordingGroup().IsValid())
+	{
+		return;
+	}
+
+	ClearQueuedRecordings();
+	TWeakObjectPtr<ASequenceRecorderGroup> GroupActor = GetRecordingGroupActor();
+	if (GroupActor.IsValid())
+	{
+		GroupActor->ActorGroups.Remove(GetRecordingGroup().Get());
+	}
+}
+
+void FSequenceRecorder::DuplicateRecordingGroup()
+{
+	check(GetRecordingGroup().IsValid());
+	check(GetRecordingGroupActor().IsValid());
+
+	USequenceRecorderActorGroup* DuplicatedGroup = DuplicateObject<USequenceRecorderActorGroup>(GetRecordingGroup().Get(), GetRecordingGroupActor().Get());
+	DuplicatedGroup->GroupName = FName(*FGuid::NewGuid().ToString());
+	GetRecordingGroupActor().Get()->ActorGroups.Add(DuplicatedGroup);
+
+	// We'll invoke the standard load function so that it triggers everything to clear/update correctly.
+	LoadRecordingGroup(DuplicatedGroup->GroupName);
+}
+
+TArray<FName> FSequenceRecorder::GetRecordingGroupNames() const
+{
+	TArray<FName> GroupNames;
+	UWorld* EditorWorld = GEditor->GetEditorWorldContext().World();
+	if (EditorWorld && EditorWorld->PersistentLevel)
+	{
+		for (int32 ActorIndex = 0; ActorIndex < EditorWorld->PersistentLevel->Actors.Num(); ++ActorIndex)
+		{
+			AActor* Actor = EditorWorld->PersistentLevel->Actors[ActorIndex];
+			ASequenceRecorderGroup* GroupActor = Cast<ASequenceRecorderGroup>(Actor);
+			if (GroupActor)
+			{
+				for (USequenceRecorderActorGroup* ActorGroup : GroupActor->ActorGroups)
+				{
+					GroupNames.Add(ActorGroup->GroupName);
+				}
+
+				// We only examine the first actor group in the map as it should contain all of our groups.
+				break;
+			}
+		}
+	}
+
+	return GroupNames;
+}
+
+void FSequenceRecorder::LoadRecordingGroup(const FName Name)
+{
+	TWeakObjectPtr<ASequenceRecorderGroup> GroupActor = nullptr;
+	UWorld* EditorWorld = GEditor->GetEditorWorldContext().World();
+
+	if (EditorWorld && EditorWorld->PersistentLevel)
+	{
+		for (int32 ActorIndex = 0; ActorIndex < EditorWorld->PersistentLevel->Actors.Num(); ++ActorIndex)
+		{
+			AActor* Actor = EditorWorld->PersistentLevel->Actors[ActorIndex];
+			GroupActor = Cast<ASequenceRecorderGroup>(Actor);
+			if (GroupActor.IsValid())
+			{
+				// We only examine the first actor group in the map
+				break;
+			}
+		}
+	}
+
+	if (GroupActor.IsValid())
+	{
+		// Remove the existing queued recordings to mark us as dirty (this causes the UI to refresh)
+		ClearQueuedRecordings();
+
+		TWeakObjectPtr<USequenceRecorderActorGroup> ActorGroup = GroupActor->FindActorGroup(Name);
+		if (ActorGroup.IsValid())
+		{
+			CurrentRecorderGroup = ActorGroup;
+			for (UActorRecording* ActorRecording : ActorGroup->RecordedActors)
+			{
+				if (ActorRecording != nullptr)
+				{
+					ActorRecording->AddToRoot();
+					QueuedRecordings.Add(ActorRecording);
+				}
+			}
+			RefreshNextSequence();
+			return;
+		}
+	}
+
+	// We either don't have a group actor or we can't find a group by that name, clear anything we have loaded.
+	// This lets the UI handle switching back to profile "None".
+	RefreshNextSequence();
+	ClearQueuedRecordings();
+	CurrentRecorderGroup = nullptr;
+}
+
+FString FSequenceRecorder::GetSequenceRecordingBasePath() const
+{
+	TWeakObjectPtr<USequenceRecorderActorGroup> RecordingGroup = GetRecordingGroup();
+	if (RecordingGroup.IsValid())
+	{
+		return RecordingGroup->SequenceRecordingBasePath.Path;
+	}
+
+	// If no profile is loaded, we just return the default.
+	return GetDefault<USequenceRecorderActorGroup>()->SequenceRecordingBasePath.Path;
 }
 
+FString FSequenceRecorder::GetSequenceRecordingName() const
+{
+	TWeakObjectPtr<USequenceRecorderActorGroup> RecordingGroup = GetRecordingGroup();
+	if (RecordingGroup.IsValid())
+	{
+		return RecordingGroup->SequenceName;
+	}
+
+	// If no profile is loaded, just return the default value.
+	return GetDefault<USequenceRecorderActorGroup>()->SequenceName;
+}
 #undef LOCTEXT_NAMESPACE
diff --git a/Engine/Source/Editor/SequenceRecorder/Private/SequenceRecorder.h b/Engine/Source/Editor/SequenceRecorder/Private/SequenceRecorder.h
--- a/Engine/Source/Editor/SequenceRecorder/Private/SequenceRecorder.h
+++ b/Engine/Source/Editor/SequenceRecorder/Private/SequenceRecorder.h
@@ -14,6 +14,8 @@ class ALevelSequenceActor;
 class ISequenceAudioRecorder;
 class UCanvas;
 class UTexture;
+class ASequenceRecorderGroup;
+class USequenceRecorderActorGroup;
 
 struct FSequenceRecorder
 {
@@ -77,7 +79,7 @@ public:
 
 	void ResetQueuedRecordingsDirty() { bQueuedRecordingsDirty = false; }
 
-	bool IsRecording() const { return CurrentSequence.IsValid(); }
+	bool IsRecording() const;
 
 	/** Draw the countdown to the screen */
 	void DrawDebug(UCanvas* InCanvas, APlayerController* InPlayerController);
@@ -88,6 +90,23 @@ public:
 	/** Handle actors being de-spawned */
 	void HandleActorDespawned(AActor* Actor);
 
+	TWeakObjectPtr<USequenceRecorderActorGroup> GetRecordingGroup() const
+	{
+		return CurrentRecorderGroup;
+	}
+
+	TWeakObjectPtr<ASequenceRecorderGroup> GetRecordingGroupActor();
+
+	void AddRecordingGroup();
+	void RemoveCurrentRecordingGroup();
+	void DuplicateRecordingGroup();
+	void LoadRecordingGroup(const FName Name);
+
+	TArray<FName> GetRecordingGroupNames() const;
+
+	FString GetSequenceRecordingBasePath() const;
+	FString GetSequenceRecordingName() const;
+
 	/** Get the built-in animation factory (as this uses special case handling) */
 	const FMovieSceneAnimationSectionRecorderFactory& GetAnimationRecorderFactory() const
 	{
@@ -135,6 +154,12 @@ private:
 	/** World we are recording a replay for, if any */
 	TLazyObjectPtr<class UWorld> CurrentReplayWorld;
 
+	/** Recorder Group that our actor recordings go into.  */
+	TWeakObjectPtr<USequenceRecorderActorGroup> CurrentRecorderGroup;
+
+	/** Cached actor for this level who holds the recording group. */
+	TWeakObjectPtr<ASequenceRecorderGroup> CachedRecordingActor;
+
 	TArray<class UActorRecording*> QueuedRecordings;
 
 	TArray<class UActorRecording*> DeadRecordings;
diff --git a/Engine/Source/Editor/SequenceRecorder/Private/SequenceRecorderActorGroup.h b/Engine/Source/Editor/SequenceRecorder/Private/SequenceRecorderActorGroup.h
new file mode 100644
--- /dev/null
+++ b/Engine/Source/Editor/SequenceRecorder/Private/SequenceRecorderActorGroup.h
@@ -0,0 +1,90 @@
+// Copyright 1998-2018 Epic Games, Inc. All Rights Reserved.
+
+#pragma once
+
+#include "CoreMinimal.h"
+#include "GameFramework/Info.h"
+#include "SequenceRecorder.h"
+#include "SequenceRecorderActorGroup.generated.h"
+
+UCLASS()
+class SEQUENCERECORDER_API USequenceRecorderActorGroup : public UObject
+{
+	GENERATED_BODY()
+public:
+	USequenceRecorderActorGroup()
+	{
+		SequenceName = TEXT("RecordedSequence");
+		SequenceRecordingBasePath.Path = TEXT("/Game/Cinematics/Sequences");
+	}
+
+	UPROPERTY(EditAnywhere, Category = "Recording Groups")
+	FName GroupName;
+
+	/** The base name of the sequence to record to. This name will also be used to auto-generate any assets created by this recording. */
+	UPROPERTY(EditAnywhere, Category = "Recording Groups")
+	FString SequenceName;
+
+	/** Base path for this recording. Sub-assets will be created in subdirectories as specified */
+	UPROPERTY(EditAnywhere, Category = "Recording Groups", meta = (ContentDir))
+	FDirectoryPath SequenceRecordingBasePath;
+
+	/** A list of actor recordings in this group which contains both the actors to record as well as settings for each one. */
+	UPROPERTY(VisibleAnywhere, Category = "Recording Groups")
+	TArray<UActorRecording*> RecordedActors;
+
+	void PostEditChangeChainProperty(struct FPropertyChangedChainEvent& PropertyChangedEvent) override
+	{
+		Super::PostEditChangeChainProperty(PropertyChangedEvent);
+
+		if (PropertyChangedEvent.Property)
+		{
+			if (PropertyChangedEvent.Property->GetFName() == GET_MEMBER_NAME_CHECKED(USequenceRecorderActorGroup, SequenceName) ||
+				PropertyChangedEvent.Property->GetFName() == GET_MEMBER_NAME_CHECKED(USequenceRecorderActorGroup, SequenceRecordingBasePath))
+			{
+				FSequenceRecorder::Get().RefreshNextSequence();
+			}
+		}
+	}
+
+};
+
+/**
+ * 
+ */
+UCLASS(hidedropdown, NotBlueprintable)
+class SEQUENCERECORDER_API ASequenceRecorderGroup : public AActor
+{
+	GENERATED_BODY()
+public:
+	ASequenceRecorderGroup()
+		: AActor()
+	{
+		bListedInSceneOutliner = false;
+	}
+
+	/** AActor Interface */
+	virtual bool IsEditorOnly() const final
+	{
+		return true;
+	}
+	/** End AActor Interface */
+
+	TWeakObjectPtr<USequenceRecorderActorGroup> FindActorGroup(const FName& Name)
+	{
+		for (USequenceRecorderActorGroup* Group : ActorGroups)
+		{
+			if (Group->GroupName == Name)
+			{
+				return Group;
+			}
+		}
+
+		return nullptr;
+	}
+
+public:
+
+	UPROPERTY()
+	TArray<USequenceRecorderActorGroup*> ActorGroups;
+};
diff --git a/Engine/Source/Editor/SequenceRecorder/Private/SequenceRecorderCommands.cpp b/Engine/Source/Editor/SequenceRecorder/Private/SequenceRecorderCommands.cpp
--- a/Engine/Source/Editor/SequenceRecorder/Private/SequenceRecorderCommands.cpp
+++ b/Engine/Source/Editor/SequenceRecorder/Private/SequenceRecorderCommands.cpp
@@ -17,6 +17,9 @@ void FSequenceRecorderCommands::RegisterCommands()
 	UI_COMMAND(AddRecording, "Add", "Add a new recording for selected actors. If nothing selected, an empty recording will be added", EUserInterfaceActionType::Button, FInputChord());
 	UI_COMMAND(RemoveRecording, "Remove", "Remove selected recording", EUserInterfaceActionType::Button, FInputChord());
 	UI_COMMAND(RemoveAllRecordings, "Remove All", "Remove all recordings", EUserInterfaceActionType::Button, FInputChord());
+	UI_COMMAND(AddRecordingGroup, "Add Group", "Adds a new group for Actor recordings to be stored with", EUserInterfaceActionType::Button, FInputChord());
+	UI_COMMAND(RemoveRecordingGroup, "Remove Group", "Removes the currently selected recording group", EUserInterfaceActionType::Button, FInputChord());
+	UI_COMMAND(DuplicateRecordingGroup, "Duplicate Group", "Duplicates the current recording group", EUserInterfaceActionType::Button, FInputChord());
 }
 
 #undef LOCTEXT_NAMESPACE
diff --git a/Engine/Source/Editor/SequenceRecorder/Private/SequenceRecorderCommands.h b/Engine/Source/Editor/SequenceRecorder/Private/SequenceRecorderCommands.h
--- a/Engine/Source/Editor/SequenceRecorder/Private/SequenceRecorderCommands.h
+++ b/Engine/Source/Editor/SequenceRecorder/Private/SequenceRecorderCommands.h
@@ -15,6 +15,9 @@ public:
 	TSharedPtr<FUICommandInfo> AddRecording;
 	TSharedPtr<FUICommandInfo> RemoveRecording;
 	TSharedPtr<FUICommandInfo> RemoveAllRecordings;
+	TSharedPtr<FUICommandInfo> AddRecordingGroup;
+	TSharedPtr<FUICommandInfo> RemoveRecordingGroup;
+	TSharedPtr<FUICommandInfo> DuplicateRecordingGroup;
 
 	/** Initialize commands */
 	virtual void RegisterCommands() override;
diff --git a/Engine/Source/Editor/SequenceRecorder/Private/SequenceRecorderModule.cpp b/Engine/Source/Editor/SequenceRecorder/Private/SequenceRecorderModule.cpp
--- a/Engine/Source/Editor/SequenceRecorder/Private/SequenceRecorderModule.cpp
+++ b/Engine/Source/Editor/SequenceRecorder/Private/SequenceRecorderModule.cpp
@@ -43,6 +43,7 @@
 #include "IStructureDetailsView.h"
 #include "WorkflowOrientedApp/WorkflowTabManager.h"
 #include "Framework/Docking/LayoutExtender.h"
+#include "MovieSceneTimeHelpers.h"
 
 #define LOCTEXT_NAMESPACE "SequenceRecorder"
 
@@ -551,11 +552,15 @@ class FSequenceRecorderModule : public ISequenceRecorder, private FSelfRegisteri
 		return FSequenceRecorder::Get().IsRecording();
 	}
 
-	virtual float GetCurrentRecordingLength() override
+	virtual FQualifiedFrameTime GetCurrentRecordingLength() override
 	{
-		TWeakObjectPtr<ULevelSequence> CurrentSequence = FSequenceRecorder::Get().GetCurrentSequence();
-
-		return CurrentSequence.IsValid() ? CurrentSequence->GetMovieScene()->GetPlaybackRange().Size<float>() : 0.0f;
+		ULevelSequence* CurrentSequence = FSequenceRecorder::Get().GetCurrentSequence().Get();
+		UMovieScene*    MovieScene      = CurrentSequence ? CurrentSequence->GetMovieScene() : nullptr;
+		if (MovieScene)
+		{
+			return FQualifiedFrameTime(FFrameTime(MovieScene::DiscreteSize(MovieScene->GetPlaybackRange())), MovieScene->GetFrameResolution());
+		}
+		return FQualifiedFrameTime();
 	}
 
 	virtual bool StartRecording(TArrayView<AActor* const> ActorsToRecord, const FString& PathToRecordTo, const FString& SequenceName) override
@@ -637,6 +642,16 @@ class FSequenceRecorderModule : public ISequenceRecorder, private FSelfRegisteri
 
 	virtual FOnRecordingFinished& OnRecordingFinished() override { return FSequenceRecorder::Get().OnRecordingFinishedDelegate; }
 
+	virtual FString GetSequenceRecordingName() const override
+	{
+		return FSequenceRecorder::Get().GetSequenceRecordingName();
+	}
+
+	virtual FString GetSequenceRecordingBasePath() const override
+	{
+		return FSequenceRecorder::Get().GetSequenceRecordingBasePath();
+	}
+
 	static void TickSequenceRecorder(float DeltaSeconds)
 	{
 		if (!IsRunningDedicatedServer() && !IsRunningCommandlet())
diff --git a/Engine/Source/Editor/SequenceRecorder/Private/SequenceRecorderSettings.cpp b/Engine/Source/Editor/SequenceRecorder/Private/SequenceRecorderSettings.cpp
--- a/Engine/Source/Editor/SequenceRecorder/Private/SequenceRecorderSettings.cpp
+++ b/Engine/Source/Editor/SequenceRecorder/Private/SequenceRecorderSettings.cpp
@@ -16,11 +16,9 @@ USequenceRecorderSettings::USequenceRecorderSettings(const FObjectInitializer& O
 	bImmersiveMode = false;
 	SequenceLength = FAnimationRecordingSettings::DefaultMaximumLength;
 	RecordingDelay = 4.0f;
-	SequenceName = TEXT("RecordedSequence");
 	AnimationSubDirectory = TEXT("Animations");
 	AudioSubDirectory = TEXT("Audio");
 	AudioGain = 0.0f;
-	SequenceRecordingBasePath.Path = TEXT("/Game/Cinematics/Sequences");
 	bRecordNearbySpawnedActors = true;
 	NearbyActorRecordingProximity = 5000.0f;
 	bRecordWorldSettingsActor = true;
@@ -39,13 +37,4 @@ void USequenceRecorderSettings::PostEditChangeChainProperty(struct FPropertyChan
 	Super::PostEditChangeChainProperty(PropertyChangedEvent);
 
 	SaveConfig();
-
-	if (PropertyChangedEvent.Property)
-	{
-		if (PropertyChangedEvent.Property->GetFName() == GET_MEMBER_NAME_CHECKED(USequenceRecorderSettings, SequenceName) ||
-			PropertyChangedEvent.Property->GetFName() == GET_MEMBER_NAME_CHECKED(USequenceRecorderSettings, SequenceRecordingBasePath))
-		{
-			FSequenceRecorder::Get().RefreshNextSequence();
-		}
-	}
 }
diff --git a/Engine/Source/Editor/SequenceRecorder/Public/ActorRecordingSettings.h b/Engine/Source/Editor/SequenceRecorder/Public/ActorRecordingSettings.h
--- a/Engine/Source/Editor/SequenceRecorder/Public/ActorRecordingSettings.h
+++ b/Engine/Source/Editor/SequenceRecorder/Public/ActorRecordingSettings.h
@@ -14,6 +14,10 @@ struct FActorRecordingSettings
 
 	FActorRecordingSettings();
 
+	FActorRecordingSettings(class UObject* InOuter);
+
+	void CreateSettingsObjectsFromFactory();
+
 	template <typename SettingsType>
 	SettingsType* GetSettingsObject() const
 	{
@@ -32,4 +36,7 @@ private:
 	/** External settings objects for recorders that supply them. Displayed via a details customization  */
 	UPROPERTY(EditAnywhere, Category = "Actor Recording")
 	TArray<UObject*> Settings;
+
+	/** An optional outer that settings objects should be created with. */
+	TWeakObjectPtr<UObject> Outer;
 };
diff --git a/Engine/Source/Editor/SequenceRecorder/Public/IMovieSceneSectionRecorderFactory.h b/Engine/Source/Editor/SequenceRecorder/Public/IMovieSceneSectionRecorderFactory.h
--- a/Engine/Source/Editor/SequenceRecorder/Public/IMovieSceneSectionRecorderFactory.h
+++ b/Engine/Source/Editor/SequenceRecorder/Public/IMovieSceneSectionRecorderFactory.h
@@ -6,7 +6,7 @@
 #include "IMovieSceneSectionRecorder.h"
 #include "Features/IModularFeature.h"
 
-/** 
+/**  
  * Factory class interface that allows the recorder to determine what recorders to apply to 
  * actors/components/objects it is presented with.
  */
@@ -32,5 +32,5 @@ public:
 	 * Create a per-recording settings object. You can access this object at record time using FActorRecordingSettings::Settings
 	 * @return an object used for user settings for this recorder
 	 */
-	virtual UObject* CreateSettingsObject() const { return nullptr; }
+	virtual UObject* CreateSettingsObject(class UObject* InOuter) const { return nullptr; }
 };
diff --git a/Engine/Source/Editor/SequenceRecorder/Public/ISequenceRecorder.h b/Engine/Source/Editor/SequenceRecorder/Public/ISequenceRecorder.h
--- a/Engine/Source/Editor/SequenceRecorder/Public/ISequenceRecorder.h
+++ b/Engine/Source/Editor/SequenceRecorder/Public/ISequenceRecorder.h
@@ -6,6 +6,7 @@
 #include "Misc/Guid.h"
 #include "Modules/ModuleInterface.h"
 #include "Containers/ArrayView.h"
+#include "QualifiedFrameTime.h"
 
 class AActor;
 class ISequenceAudioRecorder;
@@ -32,7 +33,7 @@ public:
 	virtual bool IsRecording() = 0;
 
 	/** How long is the currently recording sequence */
-	virtual float GetCurrentRecordingLength() = 0;
+	virtual FQualifiedFrameTime GetCurrentRecordingLength() = 0;
 
 	/**
 	 * Start a recording, possibly with some delay (specified by the sequence recording settings).
@@ -107,4 +108,10 @@ public:
 
 	/** Get the sequence recorder finished delegate */
 	virtual FOnRecordingFinished& OnRecordingFinished() = 0;
+	
+	/** Get the name the of the sequence recording. */
+	virtual FString GetSequenceRecordingName() const = 0;
+	
+	/** Get the directory that the sequence should record into. */
+	virtual FString GetSequenceRecordingBasePath() const = 0;
 };
diff --git a/Engine/Source/Editor/SequenceRecorder/Public/SequenceRecorderSettings.h b/Engine/Source/Editor/SequenceRecorder/Public/SequenceRecorderSettings.h
--- a/Engine/Source/Editor/SequenceRecorder/Public/SequenceRecorderSettings.h
+++ b/Engine/Source/Editor/SequenceRecorder/Public/SequenceRecorderSettings.h
@@ -82,14 +82,6 @@ public:
 	UPROPERTY(Config, EditAnywhere, Category = "Sequence Recording", meta = (ClampMin="0.0", UIMin = "0.0", ClampMax="9.0", UIMax = "9.0"))
 	float RecordingDelay;
 
-	/** The base name of the sequence to record to. This name will also be used to auto-generate any assets created by this recording. */
-	UPROPERTY(Config, EditAnywhere, Category = "Sequence Recording")
-	FString SequenceName;
-
-	/** Base path for this recording. Sub-assets will be created in subdirectories as specified */
-	UPROPERTY(Config, EditAnywhere, Category = "Sequence Recording", meta=(ContentDir))
-	FDirectoryPath SequenceRecordingBasePath;
-
 	/** The name of the subdirectory animations will be placed in. Leave this empty to place into the same directory as the sequence base path */
 	UPROPERTY(Config, EditAnywhere, AdvancedDisplay, Category = "Sequence Recording")
 	FString AnimationSubDirectory;
@@ -110,7 +102,7 @@ public:
 	UPROPERTY(Config, EditAnywhere, Category = "Sequence Recording")
 	bool bSplitAudioChannelsIntoSeparateTracks;
 
-	/** Whether to record nearby spawned actors. */
+	/** Whether to record nearby spawned actors. If an actor matches a class in the ActorFilter, this state will by bypassed. */
 	UPROPERTY(Config, EditAnywhere, Category = "Sequence Recording")
 	bool bRecordNearbySpawnedActors;
 
@@ -134,7 +126,7 @@ public:
 	UPROPERTY(Transient, EditAnywhere, Category = "Sequence Recording")
 	TArray<TLazyObjectPtr<class ALevelSequenceActor>> LevelSequenceActorsToTrigger;
 
-	/** Default settings applied to animation recording */
+	/** Default animation settings which are used to initialize all new actor recording's animation settings */
 	UPROPERTY(Config, EditAnywhere, Category = "Sequence Recording")
 	FAnimationRecordingSettings DefaultAnimationSettings;
 
diff --git a/Engine/Source/Editor/SequenceRecorder/SequenceRecorder.Build.cs b/Engine/Source/Editor/SequenceRecorder/SequenceRecorder.Build.cs
--- a/Engine/Source/Editor/SequenceRecorder/SequenceRecorder.Build.cs
+++ b/Engine/Source/Editor/SequenceRecorder/SequenceRecorder.Build.cs
@@ -21,8 +21,9 @@ namespace UnrealBuildTool.Rules
 			PublicDependencyModuleNames.AddRange(
 				new string[]
 				{
+					"TimeManagement"
 				}
-				);
+			);
 
 			PrivateDependencyModuleNames.AddRange(
 				new string[]
diff --git a/Engine/Source/Editor/SequenceRecorderSections/Private/MovieScene3DAttachSectionRecorder.cpp b/Engine/Source/Editor/SequenceRecorderSections/Private/MovieScene3DAttachSectionRecorder.cpp
--- a/Engine/Source/Editor/SequenceRecorderSections/Private/MovieScene3DAttachSectionRecorder.cpp
+++ b/Engine/Source/Editor/SequenceRecorderSections/Private/MovieScene3DAttachSectionRecorder.cpp
@@ -34,7 +34,10 @@ void FMovieScene3DAttachSectionRecorder::Record(float CurrentTime)
 	{
 		if(MovieSceneSection.IsValid())
 		{
-			MovieSceneSection->SetEndTime(CurrentTime);
+			FFrameRate FrameResolution = MovieSceneSection->GetTypedOuter<UMovieScene>()->GetFrameResolution();
+			FFrameNumber CurrentFrame = (CurrentTime * FrameResolution).FloorToFrame();
+
+			MovieSceneSection->ExpandToFrame(CurrentFrame);
 		}
 
 		// get attachment and check if the actor we are attached to is being recorded
@@ -56,7 +59,11 @@ void FMovieScene3DAttachSectionRecorder::Record(float CurrentTime)
 			if(!MovieSceneSection.IsValid() || AttachedToActor != ActorAttachedTo.Get())
 			{
 				MovieSceneSection = Cast<UMovieScene3DAttachSection>(AttachTrack->CreateNewSection());
-				MovieSceneSection->SetStartTime(CurrentTime);
+
+				FFrameRate FrameResolution = MovieSceneSection->GetTypedOuter<UMovieScene>()->GetFrameResolution();
+				FFrameNumber CurrentFrame = (CurrentTime * FrameResolution).FloorToFrame();
+
+				MovieSceneSection->SetRange(TRange<FFrameNumber>::Inclusive(CurrentFrame, CurrentFrame));
 				MovieSceneSection->SetConstraintId(Guid);
 				MovieSceneSection->AttachSocketName = SocketName;
 				MovieSceneSection->AttachComponentName = ComponentName;
diff --git a/Engine/Source/Editor/SequenceRecorderSections/Private/MovieSceneParticleTrackSectionRecorder.cpp b/Engine/Source/Editor/SequenceRecorderSections/Private/MovieSceneParticleTrackSectionRecorder.cpp
--- a/Engine/Source/Editor/SequenceRecorderSections/Private/MovieSceneParticleTrackSectionRecorder.cpp
+++ b/Engine/Source/Editor/SequenceRecorderSections/Private/MovieSceneParticleTrackSectionRecorder.cpp
@@ -3,6 +3,7 @@
 #include "MovieSceneParticleTrackSectionRecorder.h"
 #include "Tracks/MovieSceneParticleTrack.h"
 #include "MovieScene.h"
+#include "Channels/MovieSceneChannelProxy.h"
 
 TSharedPtr<IMovieSceneSectionRecorder> FMovieSceneParticleTrackSectionRecorderFactory::CreateSectionRecorder(const struct FActorRecordingSettings& InActorRecordingSettings) const
 {
@@ -35,7 +36,9 @@ void FMovieSceneParticleTrackSectionRecorder::CreateSection(UObject* InObjectToR
 
 		ParticleTrack->AddSection(*MovieSceneSection);
 
-		MovieSceneSection->SetStartTime(Time);
+		FFrameRate FrameResolution = MovieSceneSection->GetTypedOuter<UMovieScene>()->GetFrameResolution();
+		FFrameNumber CurrentFrame = (Time * FrameResolution).FloorToFrame();
+		MovieSceneSection->SetRange(TRange<FFrameNumber>::Inclusive(CurrentFrame, CurrentFrame));
 
 		bWasTriggered = false;
 
@@ -56,9 +59,12 @@ void FMovieSceneParticleTrackSectionRecorder::Record(float CurrentTime)
 {
 	if(SystemToRecord.IsValid())
 	{
-		MovieSceneSection->SetEndTime(CurrentTime);
+		FFrameRate   FrameResolution = MovieSceneSection->GetTypedOuter<UMovieScene>()->GetFrameResolution();
+		FFrameNumber KeyTime         = (CurrentTime * FrameResolution).FloorToFrame();
 
-		EParticleKey::Type NewState = EParticleKey::Deactivate;
+		MovieSceneSection->ExpandToFrame(KeyTime);
+
+		EParticleKey NewState = EParticleKey::Deactivate;
 		if(SystemToRecord->IsRegistered() && SystemToRecord->IsActive() && !SystemToRecord->bWasDeactivated)
 		{
 			if(bWasTriggered)
@@ -78,7 +84,11 @@ void FMovieSceneParticleTrackSectionRecorder::Record(float CurrentTime)
 
 		if(NewState != PreviousState)
 		{
-			MovieSceneSection->AddKey(CurrentTime, NewState);
+			FMovieSceneParticleChannel* Channel = MovieSceneSection->GetChannelProxy().GetChannel<FMovieSceneParticleChannel>(0);
+			if (ensure(Channel))
+			{
+				Channel->GetInterface().AddKey(KeyTime, (uint8)NewState);
+			}
 		}
 
 		if(NewState == EParticleKey::Trigger)
diff --git a/Engine/Source/Editor/SequenceRecorderSections/Private/MovieSceneParticleTrackSectionRecorder.h b/Engine/Source/Editor/SequenceRecorderSections/Private/MovieSceneParticleTrackSectionRecorder.h
--- a/Engine/Source/Editor/SequenceRecorderSections/Private/MovieSceneParticleTrackSectionRecorder.h
+++ b/Engine/Source/Editor/SequenceRecorderSections/Private/MovieSceneParticleTrackSectionRecorder.h
@@ -67,7 +67,7 @@ private:
 
 	bool bWasTriggered;
 
-	EParticleKey::Type PreviousState;
+	EParticleKey PreviousState;
 
 	TWeakObjectPtr<UMovieSceneParticleTrackSectionRecorder> DelegateProxy;
 };
diff --git a/Engine/Source/Editor/SequenceRecorderSections/Private/MovieSceneSpawnSectionRecorder.cpp b/Engine/Source/Editor/SequenceRecorderSections/Private/MovieSceneSpawnSectionRecorder.cpp
--- a/Engine/Source/Editor/SequenceRecorderSections/Private/MovieSceneSpawnSectionRecorder.cpp
+++ b/Engine/Source/Editor/SequenceRecorderSections/Private/MovieSceneSpawnSectionRecorder.cpp
@@ -7,6 +7,7 @@
 #include "Tracks/MovieSceneSpawnTrack.h"
 #include "MovieScene.h"
 #include "SequenceRecorderSettings.h"
+#include "Channels/MovieSceneChannelProxy.h"
 
 TSharedPtr<IMovieSceneSectionRecorder> FMovieSceneSpawnSectionRecorderFactory::CreateSectionRecorder(const struct FActorRecordingSettings& InActorRecordingSettings) const
 {
@@ -30,11 +31,14 @@ void FMovieSceneSpawnSectionRecorder::CreateSection(UObject* InObjectToRecord, U
 		SpawnTrack->AddSection(*MovieSceneSection);
 		SpawnTrack->SetObjectId(Guid);
 
-		MovieSceneSection->SetDefault(false);
-		MovieSceneSection->AddKey(0.0f, false, EMovieSceneKeyInterpolation::Break);
+		FMovieSceneBoolChannel* BoolChannel = MovieSceneSection->GetChannelProxy().GetChannel<FMovieSceneBoolChannel>(0);
+		check(BoolChannel);
+		BoolChannel->SetDefault(false);
+		BoolChannel->GetInterface().AddKey(0, false);
 
-		MovieSceneSection->SetStartTime(Time);
-		MovieSceneSection->SetIsInfinite(true);
+		FFrameRate FrameResolution = MovieSceneSection->GetTypedOuter<UMovieScene>()->GetFrameResolution();
+		FFrameNumber CurrentFrame = (Time * FrameResolution).FloorToFrame();
+		MovieSceneSection->SetRange(TRange<FFrameNumber>::Inclusive(CurrentFrame, CurrentFrame));
 	}
 
 	bWasSpawned = false;
@@ -45,15 +49,27 @@ void FMovieSceneSpawnSectionRecorder::FinalizeSection()
 	const bool bSpawned = ObjectToRecord.IsValid();
 	if(bSpawned != bWasSpawned)
 	{
-		MovieSceneSection->AddKey(MovieSceneSection->GetEndTime(), bSpawned, EMovieSceneKeyInterpolation::Break);
+		FMovieSceneBoolChannel* Channel = MovieSceneSection->GetChannelProxy().GetChannel<FMovieSceneBoolChannel>(0);
+		if (ensure(Channel) && MovieSceneSection->HasEndFrame())
+		{
+			Channel->GetInterface().AddKey(MovieSceneSection->GetExclusiveEndFrame()-1, bSpawned);
+		}
 	}
 
 	// If the section is degenerate, assume the actor was spawned and destroyed. Give it a 1 frame spawn section.
-	if (MovieSceneSection->GetRange().IsDegenerate())
+	if (MovieSceneSection->GetRange().IsDegenerate() && MovieSceneSection->HasEndFrame())
 	{
-		float OneFrameInterval = 1.0f/GetDefault<USequenceRecorderSettings>()->DefaultAnimationSettings.SampleRate;
-		MovieSceneSection->AddKey(MovieSceneSection->GetEndTime() - OneFrameInterval, true, EMovieSceneKeyInterpolation::Break);
-		MovieSceneSection->SetStartTime(MovieSceneSection->GetEndTime() - OneFrameInterval);
+		FMovieSceneBoolChannel* Channel = MovieSceneSection->GetChannelProxy().GetChannel<FMovieSceneBoolChannel>(0);
+		if (ensure(Channel))
+		{
+			double       OneFrameInterval = 1.0/GetDefault<USequenceRecorderSettings>()->DefaultAnimationSettings.SampleRate;
+
+			FFrameRate   FrameResolution  = MovieSceneSection->GetTypedOuter<UMovieScene>()->GetFrameResolution();
+			FFrameNumber StartTime        = MovieSceneSection->GetExclusiveEndFrame() - (OneFrameInterval * FrameResolution).CeilToFrame();
+
+			Channel->GetInterface().AddKey(StartTime, true);
+			MovieSceneSection->SetStartFrame(StartTime);
+		}
 	}
 }
 
@@ -61,13 +77,21 @@ void FMovieSceneSpawnSectionRecorder::Record(float CurrentTime)
 {
 	if(ObjectToRecord.IsValid())
 	{
-		MovieSceneSection->SetEndTime(CurrentTime);
+		FFrameRate FrameResolution = MovieSceneSection->GetTypedOuter<UMovieScene>()->GetFrameResolution();
+		MovieSceneSection->ExpandToFrame((CurrentTime * FrameResolution).FloorToFrame());
 	}
 
 	const bool bSpawned = ObjectToRecord.IsValid();
 	if(bSpawned != bWasSpawned)
 	{
-		MovieSceneSection->AddKey(CurrentTime, bSpawned, EMovieSceneKeyInterpolation::Break);
+		FMovieSceneBoolChannel* Channel = MovieSceneSection->GetChannelProxy().GetChannel<FMovieSceneBoolChannel>(0);
+		if (ensure(Channel))
+		{
+			FFrameRate   FrameResolution = MovieSceneSection->GetTypedOuter<UMovieScene>()->GetFrameResolution();
+			FFrameNumber KeyTime         = (CurrentTime * FrameResolution).FloorToFrame();
+
+			Channel->GetInterface().UpdateOrAddKey(KeyTime, bSpawned);
+		}
 	}
 	bWasSpawned = bSpawned;
 }
diff --git a/Engine/Source/Editor/SequenceRecorderSections/Private/MovieSceneVisibilitySectionRecorder.cpp b/Engine/Source/Editor/SequenceRecorderSections/Private/MovieSceneVisibilitySectionRecorder.cpp
--- a/Engine/Source/Editor/SequenceRecorderSections/Private/MovieSceneVisibilitySectionRecorder.cpp
+++ b/Engine/Source/Editor/SequenceRecorderSections/Private/MovieSceneVisibilitySectionRecorder.cpp
@@ -8,6 +8,8 @@
 #include "Tracks/MovieSceneVisibilityTrack.h"
 #include "MovieScene.h"
 #include "ActorRecordingSettings.h"
+#include "Channels/MovieSceneChannelProxy.h"
+#include "MovieSceneTimeHelpers.h"
 
 static const FName ActorVisibilityTrackName = TEXT("bHidden");
 static const FName ComponentVisibilityTrackName = TEXT("bHiddenInGame");
@@ -51,7 +53,11 @@ void FMovieSceneVisibilitySectionRecorder::CreateSection(UObject* InObjectToReco
 
 		VisibilityTrack->AddSection(*MovieSceneSection);
 
-		MovieSceneSection->SetDefault(false);
+		FMovieSceneBoolChannel* Channel = MovieSceneSection->GetChannelProxy().GetChannel<FMovieSceneBoolChannel>(0);
+		if (ensure(Channel))
+		{
+			Channel->SetDefault(false);
+		}
 
 		bWasVisible = false;
 		if(SceneComponent)
@@ -63,17 +69,24 @@ void FMovieSceneVisibilitySectionRecorder::CreateSection(UObject* InObjectToReco
 			bWasVisible = !Actor->bHidden;
 		}
 
-		// if time is not at the very start of the movie scene, make sure 
-		// we are hidden by default as the track will extrapolate backwards and show 
-		// objects that shouldnt be visible.
-		if (Time != MovieScene->GetPlaybackRange().GetLowerBoundValue())
+		FFrameRate   FrameResolution = MovieSceneSection->GetTypedOuter<UMovieScene>()->GetFrameResolution();
+		FFrameNumber CurrentFrame    = (Time * FrameResolution).FloorToFrame();
+
+		if (ensure(Channel))
 		{
-			MovieSceneSection->AddKey(MovieScene->GetPlaybackRange().GetLowerBoundValue(), false, EMovieSceneKeyInterpolation::Break);
+			// if time is not at the very start of the movie scene, make sure 
+			// we are hidden by default as the track will extrapolate backwards and show 
+			// objects that shouldnt be visible.
+			FFrameNumber LowerBoundValue = MovieScene->GetPlaybackRange().GetLowerBoundValue();
+			if (CurrentFrame != LowerBoundValue)
+			{
+				Channel->GetInterface().AddKey(LowerBoundValue, false);
+			}
+
+			Channel->GetInterface().AddKey(CurrentFrame, bWasVisible);
 		}
 
-		MovieSceneSection->AddKey(Time, bWasVisible, EMovieSceneKeyInterpolation::Break);
-
-		MovieSceneSection->SetStartTime(Time);
+		MovieSceneSection->SetRange(TRange<FFrameNumber>::Inclusive(CurrentFrame, CurrentFrame));
 	}
 }
 
@@ -85,7 +98,10 @@ void FMovieSceneVisibilitySectionRecorder::Record(float CurrentTime)
 {
 	if(ObjectToRecord.IsValid())
 	{
-		MovieSceneSection->SetEndTime(CurrentTime);
+		FFrameRate   FrameResolution = MovieSceneSection->GetTypedOuter<UMovieScene>()->GetFrameResolution();
+		FFrameNumber CurrentFrame    = (CurrentTime * FrameResolution).FloorToFrame();
+
+		MovieSceneSection->ExpandToFrame(CurrentFrame);
 
 		bool bVisible = false;
 		if(USceneComponent* SceneComponent = Cast<USceneComponent>(ObjectToRecord.Get()))
@@ -99,7 +115,11 @@ void FMovieSceneVisibilitySectionRecorder::Record(float CurrentTime)
 
 		if(bVisible != bWasVisible)
 		{
-			MovieSceneSection->AddKey(CurrentTime, bVisible, EMovieSceneKeyInterpolation::Break);
+			FMovieSceneBoolChannel* Channel = MovieSceneSection->GetChannelProxy().GetChannel<FMovieSceneBoolChannel>(0);
+			if (ensure(Channel))
+			{
+				Channel->GetInterface().AddKey(CurrentFrame, bVisible);
+			}
 		}
 		bWasVisible = bVisible;
 	}
diff --git a/Engine/Source/Editor/SequenceRecorderSections/Private/MovieSceneVisibilitySectionRecorder.h b/Engine/Source/Editor/SequenceRecorderSections/Private/MovieSceneVisibilitySectionRecorder.h
--- a/Engine/Source/Editor/SequenceRecorderSections/Private/MovieSceneVisibilitySectionRecorder.h
+++ b/Engine/Source/Editor/SequenceRecorderSections/Private/MovieSceneVisibilitySectionRecorder.h
@@ -18,7 +18,7 @@ public:
 
 	virtual TSharedPtr<IMovieSceneSectionRecorder> CreateSectionRecorder(const struct FActorRecordingSettings& InActorRecordingSettings) const override;
 	virtual bool CanRecordObject(class UObject* InObjectToRecord) const override;
-	virtual UObject* CreateSettingsObject() const override { return NewObject<UMovieSceneVisibilitySectionRecorderSettings>(); }
+	virtual UObject* CreateSettingsObject(class UObject* InOuter) const override { return  NewObject<UMovieSceneVisibilitySectionRecorderSettings>(InOuter, FName(TEXT("MovieSceneVisibilitySectionRecorder"))); }
 };
 
 class FMovieSceneVisibilitySectionRecorder : public IMovieSceneSectionRecorder
diff --git a/Engine/Source/Editor/Sequencer/Private/DisplayNodes/SequencerDisplayNode.cpp b/Engine/Source/Editor/Sequencer/Private/DisplayNodes/SequencerDisplayNode.cpp
--- a/Engine/Source/Editor/Sequencer/Private/DisplayNodes/SequencerDisplayNode.cpp
+++ b/Engine/Source/Editor/Sequencer/Private/DisplayNodes/SequencerDisplayNode.cpp
@@ -17,7 +17,6 @@
 #include "IKeyArea.h"
 #include "DisplayNodes/SequencerTrackNode.h"
 #include "Sequencer.h"
-#include "GroupedKeyArea.h"
 #include "SAnimationOutlinerTreeNode.h"
 #include "SequencerSettings.h"
 #include "SSequencerSectionAreaView.h"
@@ -25,13 +24,11 @@
 #include "Framework/Commands/GenericCommands.h"
 #include "ScopedTransaction.h"
 #include "SequencerKeyTimeCache.h"
+#include "SequencerKeyCollection.h"
 
 #define LOCTEXT_NAMESPACE "SequencerDisplayNode"
 
 
-/** When 0, regeneration of dynamic key groups is enabled, when non-zero, such behaviour is disabled */
-FThreadSafeCounter KeyGroupRegenerationLock;
-
 namespace SequencerNodeConstants
 {
 	extern const float CommonPadding;
@@ -68,7 +65,9 @@ public:
 
 	SLATE_BEGIN_ARGS(SSequencerObjectTrack) {}
 		/** The view range of the section area */
-		SLATE_ATTRIBUTE( TRange<float>, ViewRange )
+		SLATE_ATTRIBUTE( TRange<double>, ViewRange )
+		/** The frame resolution of the current sequence*/
+		SLATE_ATTRIBUTE( FFrameRate, FrameResolution )
 	SLATE_END_ARGS()
 
 	/** SLeafWidget Interface */
@@ -80,6 +79,7 @@ public:
 		RootNode = InRootNode;
 		
 		ViewRange = InArgs._ViewRange;
+		FrameResolution = InArgs._FrameResolution;
 
 		check(RootNode->GetType() == ESequencerNode::Object);
 	}
@@ -100,12 +100,16 @@ private:
 	TSharedPtr<FSequencerDisplayNode> RootNode;
 
 	/** The current view range */
-	TAttribute< TRange<float> > ViewRange;
+	TAttribute< TRange<double> > ViewRange;
+	/** The current frame resolution */
+	TAttribute< FFrameRate > FrameResolution;
 
 	FSequencerKeyCollectionSignature KeyCollectionSignature;
 
+	/** The cached frame resolution these positions were generated with */
+	FFrameRate CachedFrameResolution;
 	/** The time-range for which KeyDrawPositions was generated */
-	TRange<float> CachedViewRange;
+	TRange<double> CachedViewRange;
 	/** Cached pixel positions for all keys in the current view range */
 	TArray<float> KeyDrawPositions;
 	/** Cached key times per key area. Updated when section signature changes */
@@ -117,10 +121,16 @@ void SSequencerObjectTrack::Tick( const FGeometry& AllottedGeometry, const doubl
 {
 	SWidget::Tick(AllottedGeometry, InCurrentTime, InDeltaTime);
 
-	FSequencerKeyCollectionSignature NewCollectionSignature = FSequencerKeyCollectionSignature::FromNodesRecursive({ RootNode.Get() }, 0.f);
-	if (NewCollectionSignature != KeyCollectionSignature || CachedViewRange != ViewRange.Get())
+	FSequencerKeyCollectionSignature NewCollectionSignature = FSequencerKeyCollectionSignature::FromNodesRecursive({ RootNode.Get() }, 0);
+
+	TRange<double> OldCachedViewRange = CachedViewRange;
+	FFrameRate OldCachedFrameResolution = CachedFrameResolution;
+
+	CachedViewRange = ViewRange.Get();
+	CachedFrameResolution = FrameResolution.Get();
+
+	if (NewCollectionSignature != KeyCollectionSignature || CachedViewRange != OldCachedViewRange || CachedFrameResolution != OldCachedFrameResolution)
 	{
-		CachedViewRange = ViewRange.Get();
 		KeyCollectionSignature = MoveTemp(NewCollectionSignature);
 		GenerateCachedKeyPositions(AllottedGeometry);
 	}
@@ -128,7 +138,7 @@ void SSequencerObjectTrack::Tick( const FGeometry& AllottedGeometry, const doubl
 
 int32 SSequencerObjectTrack::OnPaint(const FPaintArgs& Args, const FGeometry& AllottedGeometry, const FSlateRect& MyCullingRect, FSlateWindowElementList& OutDrawElements, int32 LayerId, const FWidgetStyle& InWidgetStyle, bool bParentEnabled) const
 {
-	if (RootNode->GetSequencer().GetSettings()->GetShowCombinedKeyframes())
+	if (RootNode->GetSequencer().GetSequencerSettings()->GetShowCombinedKeyframes())
 	{
 		for (float KeyPosition : KeyDrawPositions)
 		{
@@ -190,7 +200,7 @@ void SSequencerObjectTrack::GenerateCachedKeyPositions(const FGeometry& Allotted
 
 		// Generate a new cache
 		FSequencerCachedKeys TempCache;
-		TempCache.Update(KeyArea);
+		TempCache.Update(KeyArea, CachedFrameResolution);
 
 		if (CacheKey.IsValid())
 		{
@@ -208,50 +218,54 @@ void SSequencerObjectTrack::GenerateCachedKeyPositions(const FGeometry& Allotted
 	// we use a collection of iterators that are only incremented when they've been added to the KeyDrawPositions array
 	struct FIter
 	{
-		FIter(TArrayView<const FSequencerCachedKey> In) : KeysInRange(In), CurrentIndex(0) {}
+		FIter(TArrayView<const double> InTimes) : KeysInRange(InTimes), CurrentIndex(0) {}
 
-		explicit                     operator bool() const   { return KeysInRange.IsValidIndex(CurrentIndex); }
-		FIter&                       operator++()            { ++CurrentIndex; return *this; }
+		explicit      operator bool() const { return KeysInRange.IsValidIndex(CurrentIndex); }
+		FIter&        operator++()          { ++CurrentIndex; return *this; }
 
-		const FSequencerCachedKey&   operator*() const       { return KeysInRange[CurrentIndex]; }
-		const FSequencerCachedKey*   operator->() const      { return &KeysInRange[CurrentIndex]; }
+		double        operator*() const     { return KeysInRange[CurrentIndex]; }
+		const double* operator->() const    { return &KeysInRange[CurrentIndex]; }
 	private:
-		TArrayView<const FSequencerCachedKey> KeysInRange;
+		TArrayView<const double> KeysInRange;
 		int32 CurrentIndex;
 	};
 
 	TArray<FIter> AllIterators;
 	for (auto& Pair : SectionToKeyTimeCache)
 	{
-		AllIterators.Add(Pair.Value.GetKeysInRange(CachedViewRange));
+		TArrayView<const double> Times;
+		Pair.Value.GetKeysInRange(CachedViewRange, &Times, nullptr);
+		AllIterators.Add(Times);
 	}
 	for (auto& Uncached: UncachableKeyTimes)
 	{
-		AllIterators.Add(Uncached.GetKeysInRange(CachedViewRange));
+		TArrayView<const double> Times;
+		Uncached.GetKeysInRange(CachedViewRange, &Times, nullptr);
+		AllIterators.Add(Times);
 	}
 
-	FTimeToPixel TimeToPixelConverter(AllottedGeometry, CachedViewRange);
+	FTimeToPixel TimeToPixelConverter(AllottedGeometry, CachedViewRange, CachedFrameResolution);
 
 	// While any iterator is still valid, find and add the earliest time
 	while (AllIterators.ContainsByPredicate([](FIter& It){ return It; }))
 	{
-		float EarliestTime = TNumericLimits<float>::Max();
+		double EarliestTime = TNumericLimits<double>::Max();
 		for (FIter& It : AllIterators)
 		{
-			if (It && It->Time < EarliestTime)
+			if (It && *It < EarliestTime)
 			{
-				EarliestTime = It->Time;
+				EarliestTime = *It;
 			}
 		}
 
 		// Add the position as a pixel position
-		const float KeyPosition = TimeToPixelConverter.TimeToPixel(EarliestTime);
+		const float KeyPosition = TimeToPixelConverter.SecondsToPixel(EarliestTime);
 		KeyDrawPositions.Add(KeyPosition);
 
 		// Increment any other iterators that are close enough to the time we just added
 		for (FIter& It : AllIterators)
 		{
-			while (It && FMath::IsNearlyEqual(KeyPosition, TimeToPixelConverter.TimeToPixel(It->Time), DuplicateThresholdPx))
+			while (It && FMath::IsNearlyEqual(KeyPosition, TimeToPixelConverter.SecondsToPixel(*It), DuplicateThresholdPx))
 			{
 				++It;
 			}
@@ -435,7 +449,65 @@ void FSequencerDisplayNode::AddKeyAreaNode(FName KeyAreaName, const FText& Displ
 
 FLinearColor FSequencerDisplayNode::GetDisplayNameColor() const
 {
-	return FLinearColor( 1.f, 1.f, 1.f, 1.f );
+	auto FindInActiveSection = [](FSequencerDisplayNode& InNode, bool EmptyNotActive = true)
+	{
+		if (InNode.GetType() == ESequencerNode::KeyArea)
+		{
+			const FSequencerSectionKeyAreaNode& KeyAreaNode = static_cast<FSequencerSectionKeyAreaNode&>(InNode);
+			auto KeyAreaNodes=  KeyAreaNode.GetAllKeyAreas();
+			if (KeyAreaNodes.Num() > 0)
+			{
+				for (const TSharedRef<IKeyArea>& KeyArea : KeyAreaNodes)
+				{
+					const UMovieSceneSection* Section = KeyArea->GetOwningSection();
+					if (Section && Section->IsActive())
+					{
+						// Stop traversing
+						return false;
+					}
+				}
+			}
+			else
+			{
+				return EmptyNotActive;
+			}
+		}
+		else if (InNode.GetType() == ESequencerNode::Track)
+		{
+			TArray<TSharedRef<FSequencerSectionKeyAreaNode>> KeyAreaNodes;
+			const FSequencerTrackNode& TrackNode = static_cast<FSequencerTrackNode&>(InNode);
+			TrackNode.GetChildKeyAreaNodesRecursively(KeyAreaNodes);
+			if (KeyAreaNodes.Num() > 0)
+			{
+				for (TSharedRef<FSequencerSectionKeyAreaNode> KeyAreaNode : KeyAreaNodes)
+				{
+					for (TSharedPtr<IKeyArea> KeyArea : KeyAreaNode->GetAllKeyAreas())
+					{
+						const UMovieSceneSection* Section = KeyArea->GetOwningSection();
+						if (Section && Section->IsActive())
+						{
+							// Stop traversing
+							return false;
+						}
+					}
+				}
+			}
+			else
+			{
+				return EmptyNotActive;
+			}
+		}
+		// Continue traversing
+		return true;
+	};
+
+	FSequencerDisplayNode *This = const_cast<FSequencerDisplayNode*>(this);
+	//if empty with no key areas or sections then it's active, otherwise
+	//find first child with active section, then it's active, else inactive.
+	const bool bFoundInActiveSection = ChildNodes.Num() > 0 ? This->Traverse_ParentFirst(FindInActiveSection) :
+		((this->GetType() == ESequencerNode::Track || this->GetType() == ESequencerNode::KeyArea) && FindInActiveSection(*(This), false))
+		||false;
+	return bFoundInActiveSection ? FLinearColor(0.6f, 0.6f, 0.6f, 0.6f) : FLinearColor(1.0f, 1.0f, 1.0f, 1.0f);
 }
 
 FText FSequencerDisplayNode::GetDisplayNameToolTipText() const
@@ -483,7 +555,7 @@ FText FSequencerDisplayNode::GetIconToolTipText() const
 	return FText();
 }
 
-TSharedRef<SWidget> FSequencerDisplayNode::GenerateWidgetForSectionArea(const TAttribute< TRange<float> >& ViewRange)
+TSharedRef<SWidget> FSequencerDisplayNode::GenerateWidgetForSectionArea(const TAttribute< TRange<double> >& ViewRange)
 {
 	if (GetType() == ESequencerNode::Track && static_cast<FSequencerTrackNode&>(*this).GetSubTrackMode() != FSequencerTrackNode::ESubTrackMode::ParentTrack)
 	{
@@ -494,13 +566,18 @@ TSharedRef<SWidget> FSequencerDisplayNode::GenerateWidgetForSectionArea(const TA
 	if (GetType() == ESequencerNode::Object)
 	{
 		return SNew(SSequencerObjectTrack, SharedThis(this))
-			.ViewRange(ViewRange);
+			.ViewRange(ViewRange)
+			.FrameResolution(this, &FSequencerDisplayNode::GetFrameResolution);
 	}
 
 	// currently only section areas display widgets
 	return SNullWidget::NullWidget;
 }
 
+FFrameRate FSequencerDisplayNode::GetFrameResolution() const
+{
+	return GetSequencer().GetFocusedFrameResolution();
+}
 
 TSharedPtr<FSequencerDisplayNode> FSequencerDisplayNode::GetSectionAreaAuthority() const
 {
@@ -529,6 +606,7 @@ FString FSequencerDisplayNode::GetPathName() const
 
 	if (ParentNode.IsValid())
 	{
+		ensure(ParentNode != SharedThis(this));
 		PathName = ParentNode.Pin()->GetPathName() + TEXT(".");
 	}
 
@@ -738,24 +816,6 @@ bool FSequencerDisplayNode::IsHovered() const
 }
 
 
-TSharedRef<FGroupedKeyArea> FSequencerDisplayNode::GetKeyGrouping(UMovieSceneSection* InSection)
-{
-	TSharedRef<FGroupedKeyArea>* KeyGroup = KeyGroupings.FindByPredicate([=](const TSharedRef<FGroupedKeyArea>& InKeyArea) { return InKeyArea->GetOwningSection() == InSection; });
-	if (KeyGroup)
-	{
-		if (KeyGroupRegenerationLock.GetValue() == 0)
-		{
-			(*KeyGroup)->Update();
-		}
-		return *KeyGroup;
-	}
-
-	// Just make a new one
-	KeyGroupings.Emplace(MakeShared<FGroupedKeyArea>(*this, InSection));
-	return KeyGroupings.Last();
-}
-
-
 void FSequencerDisplayNode::HandleContextMenuRenameNodeExecute()
 {
 	RenameRequestedEvent.Broadcast();
@@ -768,18 +828,6 @@ bool FSequencerDisplayNode::HandleContextMenuRenameNodeCanExecute() const
 }
 
 
-void FSequencerDisplayNode::DisableKeyGoupingRegeneration()
-{
-	KeyGroupRegenerationLock.Increment();
-}
-
-
-void FSequencerDisplayNode::EnableKeyGoupingRegeneration()
-{
-	KeyGroupRegenerationLock.Decrement();
-}
-
-
 void FSequencerDisplayNode::AddChildAndSetParent( TSharedRef<FSequencerDisplayNode> InChild )
 {
 	ChildNodes.Add( InChild );
diff --git a/Engine/Source/Editor/Sequencer/Private/DisplayNodes/SequencerDisplayNode.h b/Engine/Source/Editor/Sequencer/Private/DisplayNodes/SequencerDisplayNode.h
--- a/Engine/Source/Editor/Sequencer/Private/DisplayNodes/SequencerDisplayNode.h
+++ b/Engine/Source/Editor/Sequencer/Private/DisplayNodes/SequencerDisplayNode.h
@@ -7,8 +7,8 @@
 #include "Styling/SlateColor.h"
 #include "Widgets/SWidget.h"
 #include "SequencerNodeTree.h"
+#include "FrameRate.h"
 
-class FGroupedKeyArea;
 class FMenuBuilder;
 class FSequencer;
 class FSequencerDisplayNodeDragDropOp;
@@ -181,6 +181,31 @@ public:
 		
 	}
 
+	/**
+	 * @return What is the sorting order of this node relative to it's siblings in the node tree
+	 */
+	virtual int32 GetSortingOrder() const
+	{
+		return 0;
+	}
+
+	/**
+	 * Sets the node's sorting order relative to it's siblings in the node tree. Does not take effect until node tree is refreshed.
+	 * @param InSortingOrder The resulting sorting order the object should have relative to it's siblings in the Sequencer tree view.
+	 */
+	virtual void SetSortingOrder(const int32 InSortingOrder)
+	{
+
+	}
+
+	/**
+	 * Calls Modify on the underlying data structure before calling SetSortingOrder.
+	 */
+	virtual void ModifyAndSetSortingOrder(const int32 InSortingOrder)
+	{
+
+	}
+
 	/**
 	 * Generates a container widget for tree display in the animation outliner portion of the track area
 	 * 
@@ -201,7 +226,7 @@ public:
 	 * @param ViewRange	The range of time in the sequencer that we are displaying
 	 * @return Generated outliner widget
 	 */
-	virtual TSharedRef<SWidget> GenerateWidgetForSectionArea( const TAttribute<TRange<float>>& ViewRange );
+	virtual TSharedRef<SWidget> GenerateWidgetForSectionArea( const TAttribute<TRange<double>>& ViewRange );
 
 	/**
 	 * Gets an icon that represents this sequencer display node
@@ -237,6 +262,8 @@ public:
 	 */
 	TSharedPtr<FSequencerDisplayNode> GetSectionAreaAuthority() const;
 
+	FFrameRate GetFrameResolution() const;
+
 	/**
 	 * @return the path to this node starting with the outermost parent
 	 */
@@ -390,12 +417,6 @@ public:
 		return VirtualBottom;
 	}
 
-	/** Get the key grouping for the specified section index, ensuring it is fully up to date */
-	TSharedRef<FGroupedKeyArea> GetKeyGrouping(UMovieSceneSection* InSection);
-
-	/** Get key groupings array */
-	const TArray<TSharedRef<FGroupedKeyArea>>& GetKeyGroupings() const { return KeyGroupings; }
-
 	DECLARE_EVENT(FSequencerDisplayNode, FRequestRenameEvent);
 	FRequestRenameEvent& OnRenameRequested() { return RenameRequestedEvent; }
 
@@ -414,18 +435,13 @@ public:
 	 */
 	virtual void Drop( const TArray<TSharedRef<FSequencerDisplayNode>>& DraggedNodes, EItemDropZone DropZone ) { }
 
-public:
-
-	/** Temporarily disable dynamic regeneration of key groupings. This prevents overlapping key groups from being amalgamated during drags. Key times will continue to update correctly. */
-	static void DisableKeyGoupingRegeneration();
-
-	/** Re-enable dynamic regeneration of key groupings */
-	static void EnableKeyGoupingRegeneration();
-
 protected:
 
 	/** Adds a child to this node, and sets it's parent to this node. */
 	void AddChildAndSetParent( TSharedRef<FSequencerDisplayNode> InChild );
+public:
+	/** Clears the parent of this node. */
+	void ClearParent() { ParentNode = nullptr; }
 
 private:
 
@@ -460,9 +476,6 @@ protected:
 	/** Whether or not the node is expanded */
 	bool bExpanded;
 
-	/** Transient grouped keys for this node */
-	TArray<TSharedRef<FGroupedKeyArea>> KeyGroupings;
-
 	/** Event that is triggered when rename is requested */
 	FRequestRenameEvent RenameRequestedEvent;
 };
diff --git a/Engine/Source/Editor/Sequencer/Private/DisplayNodes/SequencerFolderNode.cpp b/Engine/Source/Editor/Sequencer/Private/DisplayNodes/SequencerFolderNode.cpp
--- a/Engine/Source/Editor/Sequencer/Private/DisplayNodes/SequencerFolderNode.cpp
+++ b/Engine/Source/Editor/Sequencer/Private/DisplayNodes/SequencerFolderNode.cpp
@@ -19,6 +19,8 @@
 #include "SequencerDisplayNodeDragDropOp.h"
 #include "ScopedTransaction.h"
 #include "Widgets/Colors/SColorPicker.h"
+#include "SequencerNodeTree.h"
+#include "SequencerNodeSortingMethods.h"
 
 #define LOCTEXT_NAMESPACE "SequencerFolderNode"
 
@@ -74,7 +76,17 @@ void FSequencerFolderNode::SetDisplayName( const FText& NewDisplayName )
 			{
 				if ( SiblingNode != AsShared() )
 				{
-					SiblingNames.Add( SiblingNode->GetNodeName() );
+					SiblingNames.Add( FName(*SiblingNode->GetDisplayName().ToString()) );
+				}
+			}
+		}
+		else
+		{
+			for (TSharedRef<FSequencerDisplayNode> RootNode : GetParentTree().GetRootNodes())
+			{
+				if ( RootNode != AsShared() )
+				{
+					SiblingNames.Add( FName(*RootNode->GetDisplayName().ToString()) );
 				}
 			}
 		}
@@ -111,168 +123,232 @@ TOptional<EItemDropZone> FSequencerFolderNode::CanDrop( FSequencerDisplayNodeDra
 {
 	DragDropOp.ResetToDefaultToolTip();
 
-	if ( ItemDropZone == EItemDropZone::AboveItem )
+	// Prevent taking any parent that's part of the dragged node hierarchy from being put inside a child of itself
+	// This is done first before the other checks so that the UI stays consistent as you move between them.
+	TSharedPtr<FSequencerDisplayNode> CurrentNode = SharedThis((FSequencerDisplayNode*)this);
+	while (CurrentNode.IsValid())
 	{
-		if ( GetParent().IsValid() )
+		if (DragDropOp.GetDraggedNodes().Contains(CurrentNode))
 		{
-			// When dropping above, only allow it for root level nodes.
+			DragDropOp.CurrentHoverText = NSLOCTEXT("SequencerFolderNode", "ParentIntoChildDragErrorFormat", "Can't drag a parent node into one of it's children.");
 			return TOptional<EItemDropZone>();
 		}
-		else
+		CurrentNode = CurrentNode->GetParent();
+	}
+
+	// Don't allow "Drop Below" as a target as this causes a confusing UI. Instead, users should try to use the Drop Above of the item after this, or, onto of this (to put it at the end).
+	if (ItemDropZone == EItemDropZone::BelowItem)
+	{
+		ItemDropZone = EItemDropZone::OntoItem;
+	}
+
+	// If they're trying to move an item above/below or into us. This item may or may not be a sibling; If it is not already a sibling then
+	// we need to check if it has a conflicting name with a sibling that we already have (for folders).
+	TArray<UMovieSceneFolder*> AdjacentFolders;
+	TSharedPtr<FSequencerDisplayNode> ChildParent;
+
+	if (ItemDropZone == EItemDropZone::OntoItem)
+	{
+		// If the item is being dropped onto us, we check our own children for name conflicts
+		ChildParent = SharedThis((FSequencerDisplayNode*)this);
+	}
+	else
+	{
+		// Otherwise, we try to check our siblings
+		ChildParent = GetParent();
+	}
+
+	if (ChildParent.IsValid())
+	{
+		for (TSharedRef <FSequencerDisplayNode> Child : ChildParent->GetChildNodes())
 		{
-			// Make sure there are no folder name collisions with the root folders
-			UMovieScene* FocusedMovieScene = GetParentTree().GetSequencer().GetFocusedMovieSceneSequence()->GetMovieScene();
-			TSet<FName> RootFolderNames;
-			for ( UMovieSceneFolder* RootFolder : FocusedMovieScene->GetRootFolders() )
+			if (Child->GetType() == ESequencerNode::Folder)
 			{
-				RootFolderNames.Add( RootFolder->GetFolderName() );
+				TSharedRef<FSequencerFolderNode> FolderNode = StaticCastSharedRef<FSequencerFolderNode>(Child);
+				AdjacentFolders.Add(&FolderNode->GetFolder());
 			}
+		}
+	}
+	else
+	{
+		// If this folder has no parent then this is a root level folder, so we need to check the Movie Scene's child list for conflicting children names.
+		UMovieScene* FocusedMovieScene = GetSequencer().GetFocusedMovieSceneSequence()->GetMovieScene();
+		AdjacentFolders.Append(FocusedMovieScene->GetRootFolders());
+	}
 
-			for ( TSharedRef<FSequencerDisplayNode> DraggedNode : DragDropOp.GetDraggedNodes() )
+	// Check each node we're dragging to see if any of them have a name conflict - if so, block the whole drag/drop operation.
+	for (TSharedRef<FSequencerDisplayNode> DraggedNode : DragDropOp.GetDraggedNodes())
+	{
+		if (DraggedNode->GetType() == ESequencerNode::Folder)
+		{
+			TSharedRef<FSequencerFolderNode> DraggedFolder = StaticCastSharedRef<FSequencerFolderNode>(DraggedNode);
+
+			// Name Conflicts are only an issue on folders.
+			bool bHasNameConflict = false;
+			for (UMovieSceneFolder* Folder : AdjacentFolders)
 			{
-				if ( DraggedNode->GetType() == ESequencerNode::Folder )
+
+				// We don't allow a folder with the same name to become a sibling, but we need to not check the dragged node if it is already at that
+				// hierarchy depth so that we can rearrange them by triggering EItemDropZone::AboveItem / EItemDropZone::BelowItem on the same hierarchy.
+				if (&DraggedFolder->GetFolder() != Folder && DraggedFolder->GetFolder().GetFolderName() == Folder->GetFolderName())
 				{
-					TSharedRef<FSequencerFolderNode> DraggedFolder = StaticCastSharedRef<FSequencerFolderNode>( DraggedNode );
-					if ( RootFolderNames.Contains( DraggedFolder->GetFolder().GetFolderName() ) )
-					{
-						DragDropOp.CurrentHoverText = FText::Format(
-							NSLOCTEXT( "SeqeuencerFolderNode", "DuplicateRootFolderDragErrorFormat", "Root folder with name '{0}' already exists." ),
-							FText::FromName( DraggedFolder->GetFolder().GetFolderName() ) );
-						return TOptional<EItemDropZone>();
-					}
+					bHasNameConflict = true;
+					break;
 				}
 			}
+
+			if (bHasNameConflict)
+			{
+				DragDropOp.CurrentHoverText = FText::Format(
+					NSLOCTEXT("SequencerFolderNode", "DuplicateFolderDragErrorFormat", "Folder with name '{0}' already exists."),
+					FText::FromName(DraggedFolder->GetFolder().GetFolderName()));
+
+				return TOptional<EItemDropZone>();
+			}
 		}
-		return TOptional<EItemDropZone>( EItemDropZone::AboveItem );
 	}
-	else
+
+	return ItemDropZone;
+}
+
+void FSequencerFolderNode::MoveDisplayNodeToFolder(TSharedRef<FSequencerDisplayNode>& Node)
+{
+	GetFolder().Modify();
+	TSharedPtr<FSequencerDisplayNode> ParentSeqNode = Node->GetParent();
+	switch (Node->GetType())
 	{
-		// When dropping onto, don't allow dropping into the same folder, don't allow dropping
-		// parents into children, and don't allow duplicate folder names.
-		TSet<FName> ChildFolderNames;
-		for ( UMovieSceneFolder* ChildFolder : GetFolder().GetChildFolders() )
+		case ESequencerNode::Folder:
 		{
-			ChildFolderNames.Add( ChildFolder->GetFolderName() );
-		}
+			TSharedRef<FSequencerFolderNode> DraggedFolderNode = StaticCastSharedRef<FSequencerFolderNode>(Node);
+			UMovieScene* FocusedMovieScene = GetParentTree().GetSequencer().GetFocusedMovieSceneSequence()->GetMovieScene();
 
-		for ( TSharedRef<FSequencerDisplayNode> DraggedNode : DragDropOp.GetDraggedNodes() )
-		{
-			TSharedPtr<FSequencerDisplayNode> ParentSeqNode = DraggedNode->GetParent();
-			if ( ParentSeqNode.IsValid() && ParentSeqNode.Get() == this )
+			// Remove the folder from where it currently resides, and then we'll add it to it's new location later.
+			// We remove it before adding so that when you move a folder within the same hierarchy it doesn't end up
+			// removing it after changing order.
+			if (ParentSeqNode.IsValid())
 			{
-				DragDropOp.CurrentHoverText = NSLOCTEXT( "SeqeuencerFolderNode", "SameParentDragError", "Can't drag a node onto the same parent." );
-				return TOptional<EItemDropZone>();
+				checkf(ParentSeqNode->GetType() == ESequencerNode::Folder, TEXT("Can not remove from unsupported parent node."));
+				TSharedPtr<FSequencerFolderNode> ParentFolder = StaticCastSharedPtr<FSequencerFolderNode>(ParentSeqNode);
+				ParentFolder->GetFolder().Modify();
+				ParentFolder->GetFolder().RemoveChildFolder(&DraggedFolderNode->GetFolder());
+			}
+			else
+			{
+				FocusedMovieScene->Modify();
+				FocusedMovieScene->GetRootFolders().Remove(&DraggedFolderNode->GetFolder());
 			}
 
-			if ( DraggedNode->GetType() == ESequencerNode::Folder )
+			// Add this as a child of ourself now
+			GetFolder().AddChildFolder(&DraggedFolderNode->GetFolder());
+			break;
+		}
+		case ESequencerNode::Track:
+		{
+			TSharedRef<FSequencerTrackNode> DraggedTrackNode = StaticCastSharedRef<FSequencerTrackNode>(Node);
+
+			if (ParentSeqNode.IsValid())
 			{
-				TSharedRef<FSequencerFolderNode> DraggedFolder = StaticCastSharedRef<FSequencerFolderNode>( DraggedNode );
-				if ( ChildFolderNames.Contains( DraggedFolder->GetFolder().GetFolderName() ) )
-				{
-					DragDropOp.CurrentHoverText = FText::Format(
-						NSLOCTEXT( "SeqeuencerFolderNode", "DuplicateChildFolderDragErrorFormat", "Folder with name '{0}' already exists." ),
-						FText::FromName( DraggedFolder->GetFolder().GetFolderName() ) );
-					return TOptional<EItemDropZone>();
-				}
+				checkf(ParentSeqNode->GetType() == ESequencerNode::Folder, TEXT("Can not remove from unsupported parent node."));
+				TSharedPtr<FSequencerFolderNode> ParentFolder = StaticCastSharedPtr<FSequencerFolderNode>(ParentSeqNode);
+				ParentFolder->GetFolder().Modify();
+				ParentFolder->GetFolder().RemoveChildMasterTrack(DraggedTrackNode->GetTrack());
 			}
+			
+			GetFolder().AddChildMasterTrack(DraggedTrackNode->GetTrack());
+			break;
 		}
-		TSharedPtr<FSequencerDisplayNode> CurrentNode = SharedThis( ( FSequencerDisplayNode* )this );
-		while ( CurrentNode.IsValid() )
+		case ESequencerNode::Object:
 		{
-			if ( DragDropOp.GetDraggedNodes().Contains( CurrentNode ) )
+			TSharedRef<FSequencerObjectBindingNode> DraggedObjectBindingNode = StaticCastSharedRef<FSequencerObjectBindingNode>(Node);
+
+			if (ParentSeqNode.IsValid())
 			{
-				DragDropOp.CurrentHoverText = NSLOCTEXT( "SeqeuencerFolderNode", "ParentIntoChildDragError", "Can't drag a parent node into one of it's children." );
-				return TOptional<EItemDropZone>();
+				checkf(ParentSeqNode->GetType() == ESequencerNode::Folder, TEXT("Can not remove from unsupported parent node."));
+				TSharedPtr<FSequencerFolderNode> ParentFolder = StaticCastSharedPtr<FSequencerFolderNode>(ParentSeqNode);
+				ParentFolder->GetFolder().Modify();
+				ParentFolder->GetFolder().RemoveChildObjectBinding(DraggedObjectBindingNode->GetObjectBinding());
 			}
-			CurrentNode = CurrentNode->GetParent();
+
+			GetFolder().AddChildObjectBinding(DraggedObjectBindingNode->GetObjectBinding());
+			break;
 		}
-		return TOptional<EItemDropZone>( EItemDropZone::OntoItem );
 	}
-}
 
+	// Update the node's parent so that request for the NodePath reflect the new path
+	// instead of waiting until all nodes are regenerated by the subsequent Refresh call.
+	AddChildAndSetParent(Node);
+
+	// Update the expansion state using our new path.
+	Node->Traverse_ParentFirst([](FSequencerDisplayNode& TraversalNode)
+	{
+		TraversalNode.GetParentTree().SaveExpansionState(TraversalNode, TraversalNode.IsExpanded());
+		return true;
+	}, true);
+}
 
 void FSequencerFolderNode::Drop( const TArray<TSharedRef<FSequencerDisplayNode>>& DraggedNodes, EItemDropZone ItemDropZone )
 {
-	const FScopedTransaction Transaction( NSLOCTEXT( "SequencerFolderNode", "MoveIntoFolder", "Move items into folder." ) );
+	const FScopedTransaction Transaction(NSLOCTEXT("SequencerFolderNode", "MoveItems", "Move items."));
 	GetFolder().SetFlags(RF_Transactional);
 	GetFolder().Modify();
-	for ( TSharedRef<FSequencerDisplayNode> DraggedNode : DraggedNodes )
+
+	// Open this folder if an item was dropped into the folder
+	if (ItemDropZone == EItemDropZone::OntoItem)
 	{
-		TSharedPtr<FSequencerDisplayNode> ParentSeqNode = DraggedNode->GetParent();
-		switch ( DraggedNode->GetType() )
-		{
-			case ESequencerNode::Folder:
-			{
-				TSharedRef<FSequencerFolderNode> DraggedFolderNode = StaticCastSharedRef<FSequencerFolderNode>( DraggedNode );
-				UMovieScene* FocusedMovieScene = GetParentTree().GetSequencer().GetFocusedMovieSceneSequence()->GetMovieScene();
+		SetExpansionState(true);
+	}
 
-				if ( ItemDropZone == EItemDropZone::OntoItem )
-				{
-					GetFolder().AddChildFolder( &DraggedFolderNode->GetFolder() );
-				}
-				else
-				{
-					FocusedMovieScene->Modify();
-					FocusedMovieScene->GetRootFolders().Add( &DraggedFolderNode->GetFolder() );
-				}
-				
-				if ( ParentSeqNode.IsValid() )
-				{
-					checkf( ParentSeqNode->GetType() == ESequencerNode::Folder, TEXT( "Can not remove from unsupported parent node." ) );
-					TSharedPtr<FSequencerFolderNode> ParentFolder = StaticCastSharedPtr<FSequencerFolderNode>( ParentSeqNode );
-					ParentFolder->GetFolder().Modify();
-					ParentFolder->GetFolder().RemoveChildFolder( &DraggedFolderNode->GetFolder() );
-				}
-				else
-				{
-					FocusedMovieScene->Modify();
-					FocusedMovieScene->GetRootFolders().Remove( &DraggedFolderNode->GetFolder() );
-				}
+	for (TSharedRef<FSequencerDisplayNode> DraggedNode : DraggedNodes)
+	{
+		TSharedPtr<FSequencerDisplayNode> DraggedSeqNodeParent = DraggedNode->GetParent();
 
-				break;
-			}
-			case ESequencerNode::Track:
+		if (ItemDropZone == EItemDropZone::OntoItem)
+		{
+			// Remove the node from it's old parent and put it as a child of ourself
+			MoveDisplayNodeToFolder(DraggedNode);
+		}
+		else
+		{
+			if (GetParent().IsValid())
 			{
-				TSharedRef<FSequencerTrackNode> DraggedTrackNode = StaticCastSharedRef<FSequencerTrackNode>( DraggedNode );
-				
-				if( ItemDropZone == EItemDropZone::OntoItem )
-				{
-					GetFolder().AddChildMasterTrack( DraggedTrackNode->GetTrack() );
-				}
-				
-				if ( ParentSeqNode.IsValid() )
-				{
-					checkf( ParentSeqNode->GetType() == ESequencerNode::Folder, TEXT( "Can not remove from unsupported parent node." ) );
-					TSharedPtr<FSequencerFolderNode> ParentFolder = StaticCastSharedPtr<FSequencerFolderNode>( ParentSeqNode );
-					ParentFolder->GetFolder().Modify();
-					ParentFolder->GetFolder().RemoveChildMasterTrack( DraggedTrackNode->GetTrack() );
-				}
+				checkf(GetParent()->GetType() == ESequencerNode::Folder, TEXT("Cannot reorder when parent is not a folder."));
+				TSharedPtr<FSequencerFolderNode> ParentFolder = StaticCastSharedPtr<FSequencerFolderNode>(GetParent());
 
-				break;
+				// Let the folder we're going into remove us from our old parent and put us as a child of it first.
+				ParentFolder->MoveDisplayNodeToFolder(DraggedNode);
 			}
-			case ESequencerNode::Object:
+			else
 			{
-				TSharedRef<FSequencerObjectBindingNode> DraggedObjectBindingNode = StaticCastSharedRef<FSequencerObjectBindingNode>( DraggedNode );
-				
-				if ( ItemDropZone == EItemDropZone::OntoItem )
-				{
-					GetFolder().AddChildObjectBinding( DraggedObjectBindingNode->GetObjectBinding() );
-				}
+				// We're at root and they're placing above or below us
+				ParentTree.MoveDisplayNodeToRoot(DraggedNode);
+			}
+		}
+	}
 
-				if ( ParentSeqNode.IsValid() )
-				{
-					checkf( ParentSeqNode->GetType() == ESequencerNode::Folder, TEXT( "Can not remove from unsupported parent node." ) );
-					TSharedPtr<FSequencerFolderNode> ParentFolder = StaticCastSharedPtr<FSequencerFolderNode>( ParentSeqNode );
-					ParentFolder->GetFolder().Modify();
-					ParentFolder->GetFolder().RemoveChildObjectBinding( DraggedObjectBindingNode->GetObjectBinding() );
-				}
+	if (ItemDropZone == EItemDropZone::OntoItem)
+	{
+		// Sort our dragged nodes relative to our children
+		SortAndSetSortingOrder(DraggedNodes, GetChildNodes(), ItemDropZone, FDisplayNodeTreePositionSorter(), SharedThis(this));
+	}
+	else
+	{
+		// If it's above or below us, then we need to check our parent and have them sort the nodes relative to our siblings.
+		if (GetParent().IsValid())
+		{
+			checkf(GetParent()->GetType() == ESequencerNode::Folder, TEXT("Cannot reorder when parent is not a folder."));
+			TSharedPtr<FSequencerFolderNode> ParentFolder = StaticCastSharedPtr<FSequencerFolderNode>(GetParent());
 
-				break;
-			}
+			// Sort our dragged nodes relative to our siblings.
+			SortAndSetSortingOrder(DraggedNodes, ParentFolder->GetChildNodes(), ItemDropZone, FDisplayNodeTreePositionSorter(), SharedThis(this));
+		}
+		else
+		{
+			// We're at root
+			SortAndSetSortingOrder(DraggedNodes, GetSequencer().GetNodeTree()->GetRootNodes(), ItemDropZone, FDisplayNodeTreePositionSorter(), SharedThis(this));
 		}
 	}
-	SetExpansionState( true );
-	ParentTree.GetSequencer().NotifyMovieSceneDataChanged( EMovieSceneDataChangeType::MovieSceneStructureItemsChanged );
+
+	ParentTree.GetSequencer().NotifyMovieSceneDataChanged(EMovieSceneDataChangeType::MovieSceneStructureItemsChanged);
 }
 
 void FSequencerFolderNode::BuildContextMenu(FMenuBuilder& MenuBuilder)
@@ -349,4 +425,19 @@ UMovieSceneFolder& FSequencerFolderNode::GetFolder() const
 	return MovieSceneFolder;
 }
 
+int32 FSequencerFolderNode::GetSortingOrder() const
+{
+	return MovieSceneFolder.GetSortingOrder();
+}
+
+void FSequencerFolderNode::SetSortingOrder(const int32 InSortingOrder)
+{
+	MovieSceneFolder.SetSortingOrder(InSortingOrder);
+}
+
+void FSequencerFolderNode::ModifyAndSetSortingOrder(const int32 InSortingOrder)
+{
+	MovieSceneFolder.Modify();
+	SetSortingOrder(InSortingOrder);
+}
 #undef LOCTEXT_NAMESPACE
diff --git a/Engine/Source/Editor/Sequencer/Private/DisplayNodes/SequencerFolderNode.h b/Engine/Source/Editor/Sequencer/Private/DisplayNodes/SequencerFolderNode.h
--- a/Engine/Source/Editor/Sequencer/Private/DisplayNodes/SequencerFolderNode.h
+++ b/Engine/Source/Editor/Sequencer/Private/DisplayNodes/SequencerFolderNode.h
@@ -30,7 +30,13 @@ public:
 	virtual bool CanDrag() const override;
 	virtual TOptional<EItemDropZone> CanDrop( FSequencerDisplayNodeDragDropOp& DragDropOp, EItemDropZone ItemDropZone ) const override;
 	virtual void Drop( const TArray<TSharedRef<FSequencerDisplayNode>>& DraggedNodes, EItemDropZone ItemDropZone ) override;
-	virtual void BuildContextMenu(FMenuBuilder& MenuBuilder);
+	virtual void BuildContextMenu(FMenuBuilder& MenuBuilder) override;
+	virtual int32 GetSortingOrder() const override;
+	virtual void SetSortingOrder(const int32 InSortingOrder) override;
+	virtual void ModifyAndSetSortingOrder(const int32 InSortingOrder) override;
+
+	/** Removes a node from it's old parent and makes it a child of this node. */
+	void MoveDisplayNodeToFolder(TSharedRef<FSequencerDisplayNode>& Node);
 
 	/** Adds a child node to this folder node. */
 	void AddChildNode( TSharedRef<FSequencerDisplayNode> ChildNode );
@@ -47,7 +53,7 @@ public:
 	/** Callback for the color picker being closed */
 	void OnColorPickerClosed(const TSharedRef<SWindow>& Window);
 
-	/** Callback for the color picker being cancelled */
+	/** Callback for the color picker being canceled */
 	void OnColorPickerCancelled(FLinearColor NewFolderColor);
 
 private:
diff --git a/Engine/Source/Editor/Sequencer/Private/DisplayNodes/SequencerKeyTimeCache.cpp b/Engine/Source/Editor/Sequencer/Private/DisplayNodes/SequencerKeyTimeCache.cpp
--- a/Engine/Source/Editor/Sequencer/Private/DisplayNodes/SequencerKeyTimeCache.cpp
+++ b/Engine/Source/Editor/Sequencer/Private/DisplayNodes/SequencerKeyTimeCache.cpp
@@ -5,42 +5,50 @@
 #include "Algo/Sort.h"
 #include "Algo/BinarySearch.h"
 
-void FSequencerCachedKeys::Update(TSharedRef<IKeyArea> KeyArea)
+void FSequencerCachedKeys::Update(TSharedRef<IKeyArea> InKeyArea, FFrameRate SourceResolution)
 {
-	if (!CachedKeys.IsSet())
-	{
-		CachedKeys.Emplace();
-	}
-
-	UMovieSceneSection* Section = KeyArea->GetOwningSection();
-	if (!Section || !CachedSignature.IsValid() || Section->GetSignature() != CachedSignature)
+	UMovieSceneSection* Section = InKeyArea->GetOwningSection();
+	if (!Section || !CachedSignature.IsValid() || Section->GetSignature() != CachedSignature || SourceResolution != CachedFrameResolution)
 	{
 		CachedSignature = Section ? Section->GetSignature() : FGuid();
+		CachedFrameResolution = SourceResolution;
 
-		CachedKeys->Reset();
+		TArray<FFrameNumber> Times;
+		TArray<FKeyHandle> Handles;
+		InKeyArea->GetKeyInfo(&Handles, &Times);
+
+		CachedKeyTimes.Reset(Times.Num());
+		CachedKeyHandles.Reset(Times.Num());
 
 		// Generate and cache
-		for (FKeyHandle KeyHandle : KeyArea->GetUnsortedKeyHandles())
+		for (int32 Index = 0; Index < Times.Num(); ++Index)
 		{
-			CachedKeys->Add(FSequencerCachedKey{ KeyHandle, KeyArea->GetKeyTime(KeyHandle) });
+			CachedKeyTimes.Add(Times[Index] / SourceResolution);
+			CachedKeyHandles.Add(Handles[Index]);
 		}
 
-		Algo::SortBy(CachedKeys.GetValue(), &FSequencerCachedKey::Time);
+		KeyArea = InKeyArea;
 	}
 }
 
-TArrayView<const FSequencerCachedKey> FSequencerCachedKeys::GetKeysInRange(TRange<float> ViewRange) const
+void FSequencerCachedKeys::GetKeysInRange(const TRange<double>& Range, TArrayView<const double>* OutTimes, TArrayView<const FKeyHandle>* OutHandles) const
 {
 	// Binary search the first time that's >= the lower bound
-	int32 FirstVisibleIndex = Algo::LowerBoundBy(CachedKeys.GetValue(), ViewRange.GetLowerBoundValue(), &FSequencerCachedKey::Time, TLess<float>());
+	int32 FirstVisibleIndex = Algo::LowerBound(CachedKeyTimes, Range.GetLowerBoundValue());
 	// Binary search the last time that's > the upper bound
-	int32 LastVisibleIndex = Algo::UpperBoundBy(CachedKeys.GetValue(), ViewRange.GetUpperBoundValue(), &FSequencerCachedKey::Time, TLess<float>());
+	int32 LastVisibleIndex = Algo::UpperBound(CachedKeyTimes, Range.GetUpperBoundValue());
 
 	int32 Num = LastVisibleIndex - FirstVisibleIndex;
-	if (CachedKeys->IsValidIndex(FirstVisibleIndex) && LastVisibleIndex <= CachedKeys->Num())
+	if (CachedKeyTimes.IsValidIndex(FirstVisibleIndex) && LastVisibleIndex <= CachedKeyTimes.Num())
 	{
-		return MakeArrayView(&CachedKeys.GetValue()[FirstVisibleIndex], Num);
-	}
+		if (OutTimes)
+		{
+			*OutTimes = MakeArrayView(&CachedKeyTimes[FirstVisibleIndex], Num);
+		}
 
-	return TArrayView<const FSequencerCachedKey>();
+		if (OutHandles)
+		{
+			*OutHandles = MakeArrayView(&CachedKeyHandles[FirstVisibleIndex], Num);
+		}
+	}
 }
\ No newline at end of file
diff --git a/Engine/Source/Editor/Sequencer/Private/DisplayNodes/SequencerKeyTimeCache.h b/Engine/Source/Editor/Sequencer/Private/DisplayNodes/SequencerKeyTimeCache.h
--- a/Engine/Source/Editor/Sequencer/Private/DisplayNodes/SequencerKeyTimeCache.h
+++ b/Engine/Source/Editor/Sequencer/Private/DisplayNodes/SequencerKeyTimeCache.h
@@ -10,29 +10,35 @@
 #include "Templates/SharedPointer.h"
 #include "Misc/Optional.h"
 #include "IKeyArea.h"
-
-/** Key information that has been cached to avoid expensive operations */
-struct FSequencerCachedKey
-{
-	/** The key handle */
-	FKeyHandle Handle;
-	/** The local time of the key */
-	float Time;
-};
+#include "FrameRate.h"
 
 /** Simple structure that caches the sorted key times for a given key area */
 struct FSequencerCachedKeys
 {
 	/** Update this cache to store key times and handles from the specified key area */
-	void Update(TSharedRef<IKeyArea> KeyArea);
+	void Update(TSharedRef<IKeyArea> KeyArea, FFrameRate SourceResolution);
 
 	/** Get an view of the cached array for keys that fall within the specified range */
-	TArrayView<const FSequencerCachedKey> GetKeysInRange(TRange<float> Range) const;
+	void GetKeysInRange(const TRange<double>& Range, TArrayView<const double>* OutTimes, TArrayView<const FKeyHandle>* OutHandles) const;
+
+	/** Get the key area this cache was generated for, or nullptr if the cache has never been updated */
+	TSharedPtr<IKeyArea> GetKeyArea() const
+	{
+		return KeyArea;
+	}
 
 private:
 	/** Cached key information */
-	TOptional<TArray<FSequencerCachedKey>> CachedKeys;
+	TArray<double> CachedKeyTimes;
+	TArray<FKeyHandle> CachedKeyHandles;
+
 	/** The guid with which the above array was generated */
 	FGuid CachedSignature;
+
+	/** The frame resolution of the sequence that this cache was generated with */
+	FFrameRate CachedFrameResolution;
+
+	/** The key area this cache is for */
+	TSharedPtr<IKeyArea> KeyArea;
 };
 
diff --git a/Engine/Source/Editor/Sequencer/Private/DisplayNodes/SequencerObjectBindingNode.cpp b/Engine/Source/Editor/Sequencer/Private/DisplayNodes/SequencerObjectBindingNode.cpp
--- a/Engine/Source/Editor/Sequencer/Private/DisplayNodes/SequencerObjectBindingNode.cpp
+++ b/Engine/Source/Editor/Sequencer/Private/DisplayNodes/SequencerObjectBindingNode.cpp
@@ -29,6 +29,13 @@
 #include "SequencerUtilities.h"
 #include "Styling/SlateIconFinder.h"
 #include "ScopedTransaction.h"
+#include "SequencerDisplayNodeDragDropOp.h"
+#include "SequencerFolderNode.h"
+#include "SequencerNodeSortingMethods.h"
+#include "DisplayNodes/SequencerFolderNode.h"
+#include "MovieSceneSequence.h"
+#include "MovieScene.h"
+#include "MovieSceneFolder.h"
 
 #include "Tracks/MovieSceneSpawnTrack.h"
 #include "Sections/MovieSceneSpawnSection.h"
@@ -216,6 +223,8 @@ void FSequencerObjectBindingNode::BuildContextMenu(FMenuBuilder& MenuBuilder)
 	}
 	MenuBuilder.EndSection();
 
+	GetSequencer().BuildCustomContextMenuForGuid(MenuBuilder, ObjectBinding);
+
 	FSequencerDisplayNode::BuildContextMenu(MenuBuilder);
 }
 
@@ -427,6 +436,161 @@ bool FSequencerObjectBindingNode::CanDrag() const
 	return ParentSeqNode.IsValid() == false || ParentSeqNode->GetType() != ESequencerNode::Object;
 }
 
+TOptional<EItemDropZone> FSequencerObjectBindingNode::CanDrop(FSequencerDisplayNodeDragDropOp& DragDropOp, EItemDropZone ItemDropZone) const
+{
+	DragDropOp.ResetToDefaultToolTip();
+
+	// Prevent taking any parent that's part of the dragged node hierarchy from being put inside a child of itself
+	// This is done first before the other checks so that the UI stays consistent as you move between them, otherwise
+	// when you are above/below a node it reports this error, but if you were on top of a node it would do the standard
+	// no-drag-drop due to OntoItem being blocked. 
+	TSharedPtr<FSequencerDisplayNode> CurrentNode = SharedThis((FSequencerDisplayNode*)this);
+	while (CurrentNode.IsValid())
+	{
+		if (DragDropOp.GetDraggedNodes().Contains(CurrentNode))
+		{
+			DragDropOp.CurrentHoverText = NSLOCTEXT("SequencerFolderNode", "ParentIntoChildDragErrorFormat", "Can't drag a parent node into one of it's children.");
+			return TOptional<EItemDropZone>();
+		}
+		CurrentNode = CurrentNode->GetParent();
+	}
+
+	// Override Onto and Below to be Above to smooth out the UI changes as you scroll over many items.
+	// This removes a confusing "above" -> "blocked" -> "above/below" transition.
+	if (ItemDropZone == EItemDropZone::OntoItem || ItemDropZone == EItemDropZone::BelowItem)
+	{
+		ItemDropZone = EItemDropZone::AboveItem;
+	}
+
+	if (GetParent().IsValid() && GetParent()->GetType() != ESequencerNode::Folder)
+	{
+		// Object Binding Nodes can have other binding nodes as their parents and we
+		// don't allow re-arranging tracks within a binding node.
+		return TOptional<EItemDropZone>();
+	}
+
+	for (TSharedRef<FSequencerDisplayNode> Node : DragDropOp.GetDraggedNodes())
+	{
+		bool bValidType = Node->GetType() == ESequencerNode::Folder || Node->GetType() == ESequencerNode::Object || Node->GetType() == ESequencerNode::Track;
+		if (!bValidType)
+		{
+			return TOptional<EItemDropZone>();
+		}
+
+		TSharedPtr<FSequencerDisplayNode> ParentSeqNode = Node->GetParent();
+
+		if (ParentSeqNode.IsValid())
+		{
+			if (ParentSeqNode->GetType() != ESequencerNode::Folder)
+			{
+				// If we have a parent who is not a folder (ie: The node is a component track on an actor) then it can't be rearranged.
+				return TOptional<EItemDropZone>();
+			}
+		}
+	}
+
+	TArray<UMovieSceneFolder*> AdjacentFolders;
+	if (GetParent().IsValid())
+	{
+		// We are either trying to drop adjacent to ourself (when nestled), or as a child of ourself, so we add either our siblings or our children
+		// to the list of possibly conflicting names.
+		for (TSharedRef <FSequencerDisplayNode> Child : GetParent()->GetChildNodes())
+		{
+			if (Child->GetType() == ESequencerNode::Folder)
+			{
+				TSharedRef<FSequencerFolderNode> FolderNode = StaticCastSharedRef<FSequencerFolderNode>(Child);
+				AdjacentFolders.Add(&FolderNode->GetFolder());
+			}
+		}
+	}
+	else
+	{
+		// If this folder has no parent then this is a root level folder, so we need to check the Movie Scene's child list for conflicting children names.
+		UMovieScene* FocusedMovieScene = GetSequencer().GetFocusedMovieSceneSequence()->GetMovieScene();
+		AdjacentFolders.Append(FocusedMovieScene->GetRootFolders());
+	}
+
+	// Check each node we're dragging to see if any of them have a name conflict - if so, block the whole drag/drop operation.
+	for (TSharedRef<FSequencerDisplayNode> DraggedNode : DragDropOp.GetDraggedNodes())
+	{
+		if (DraggedNode->GetType() == ESequencerNode::Folder)
+		{
+			TSharedRef<FSequencerFolderNode> DraggedFolder = StaticCastSharedRef<FSequencerFolderNode>(DraggedNode);
+
+			// Name Conflicts are only an issue on folders.
+			bool bHasNameConflict = false;
+			for (UMovieSceneFolder* Folder : AdjacentFolders)
+			{
+				// We don't allow a folder with the same name to become a sibling, but we need to not check the dragged node if it is already at that
+				// hierarchy depth so that we can rearrange them by triggering EItemDropZone::AboveItem / EItemDropZone::BelowItem on the same hierarchy.
+				if (&DraggedFolder->GetFolder() != Folder && DraggedFolder->GetFolder().GetFolderName() == Folder->GetFolderName())
+				{
+					bHasNameConflict = true;
+					break;
+				}
+			}
+
+			if (bHasNameConflict)
+			{
+				DragDropOp.CurrentHoverText = FText::Format(
+					NSLOCTEXT("SequencerFolderNode", "DuplicateFolderDragErrorFormat", "Folder with name '{0}' already exists."),
+					FText::FromName(DraggedFolder->GetFolder().GetFolderName()));
+
+				return TOptional<EItemDropZone>();
+			}
+		}
+	}
+
+	// The dragged nodes were either all in folders, or all at the sequencer root.
+	return ItemDropZone;
+}
+
+void FSequencerObjectBindingNode::Drop(const TArray<TSharedRef<FSequencerDisplayNode>>& DraggedNodes, EItemDropZone ItemDropZone)
+{
+	const FScopedTransaction Transaction(NSLOCTEXT("SequencerObjectBindingNode", "MoveItems", "Move items."));
+	for (TSharedRef<FSequencerDisplayNode> DraggedNode : DraggedNodes)
+	{
+		TSharedPtr<FSequencerDisplayNode> DraggedSeqNodeParent = DraggedNode->GetParent();
+
+		if (GetParent().IsValid())
+		{
+			// If the object is coming from the root or it's coming from another folder then we can allow it to move adjacent to us.
+			if (!DraggedSeqNodeParent.IsValid() || (DraggedSeqNodeParent.IsValid() && DraggedSeqNodeParent->GetType() == ESequencerNode::Folder))
+			{
+				checkf(GetParent()->GetType() == ESequencerNode::Folder, TEXT("Cannot reorder when parent is not a folder."));
+				TSharedPtr<FSequencerFolderNode> ParentFolder = StaticCastSharedPtr<FSequencerFolderNode>(GetParent());
+
+				// Let the folder we're going into remove us from our old parent and put us as a child of it first.
+				ParentFolder->MoveDisplayNodeToFolder(DraggedNode);
+			}
+		}
+		else
+		{
+			// We're at root and they're placing above or below us
+			ParentTree.MoveDisplayNodeToRoot(DraggedNode);
+		}
+	}
+
+	if (DraggedNodes.Num() > 0)
+	{
+		if (GetParent().IsValid())
+		{
+			checkf(GetParent()->GetType() == ESequencerNode::Folder, TEXT("Cannot reorder when parent is not a folder."));
+			TSharedPtr<FSequencerFolderNode> ParentFolder = StaticCastSharedPtr<FSequencerFolderNode>(GetParent());
+
+			// Sort our dragged nodes relative to our siblings.
+			SortAndSetSortingOrder(DraggedNodes, ParentFolder->GetChildNodes(), ItemDropZone, FDisplayNodeTreePositionSorter(), SharedThis(this));
+		}
+		else
+		{
+			// We're at root and they're placing above or below us
+			SortAndSetSortingOrder(DraggedNodes, GetSequencer().GetNodeTree()->GetRootNodes(), ItemDropZone, FDisplayNodeTreePositionSorter(), SharedThis(this));
+		}
+	}
+
+	ParentTree.GetSequencer().NotifyMovieSceneDataChanged(EMovieSceneDataChangeType::MovieSceneStructureItemsChanged);
+}
+
 
 /* FSequencerObjectBindingNode implementation
  *****************************************************************************/
@@ -719,5 +883,43 @@ void FSequencerObjectBindingNode::HandlePropertyMenuItemExecute(FPropertyPath Pr
 	Sequencer.KeyProperty(KeyPropertyParams);
 }
 
+int32 FSequencerObjectBindingNode::GetSortingOrder() const
+{
+	UMovieScene* MovieScene = GetSequencer().GetFocusedMovieSceneSequence()->GetMovieScene();
+	const FMovieSceneBinding* MovieSceneBinding = MovieScene->GetBindings().FindByPredicate([&](FMovieSceneBinding& Binding)
+	{
+		return Binding.GetObjectGuid() == ObjectBinding;
+	});
+
+	if (MovieSceneBinding)
+	{
+		return MovieSceneBinding->GetSortingOrder();
+	}
+
+	return 0;
+}
+
+void FSequencerObjectBindingNode::SetSortingOrder(const int32 InSortingOrder)
+{
+	UMovieScene* MovieScene = GetSequencer().GetFocusedMovieSceneSequence()->GetMovieScene();
+
+	FMovieSceneBinding* MovieSceneBinding = (FMovieSceneBinding*) MovieScene->GetBindings().FindByPredicate([&](FMovieSceneBinding& Binding)
+	{
+		return Binding.GetObjectGuid() == ObjectBinding;
+	});
+
+	if (MovieSceneBinding)
+	{
+		MovieSceneBinding->SetSortingOrder(InSortingOrder);
+	}
+}
+
+void FSequencerObjectBindingNode::ModifyAndSetSortingOrder(const int32 InSortingOrder)
+{
+	UMovieScene* MovieScene = GetSequencer().GetFocusedMovieSceneSequence()->GetMovieScene();
+
+	MovieScene->Modify();
+	SetSortingOrder(InSortingOrder);
+}
 
 #undef LOCTEXT_NAMESPACE
diff --git a/Engine/Source/Editor/Sequencer/Private/DisplayNodes/SequencerObjectBindingNode.h b/Engine/Source/Editor/Sequencer/Private/DisplayNodes/SequencerObjectBindingNode.h
--- a/Engine/Source/Editor/Sequencer/Private/DisplayNodes/SequencerObjectBindingNode.h
+++ b/Engine/Source/Editor/Sequencer/Private/DisplayNodes/SequencerObjectBindingNode.h
@@ -76,6 +76,11 @@ public:
 	virtual ESequencerNode::Type GetType() const override;
 	virtual void SetDisplayName(const FText& NewDisplayName) override;
 	virtual bool CanDrag() const override;
+	virtual TOptional<EItemDropZone> CanDrop(FSequencerDisplayNodeDragDropOp& DragDropOp, EItemDropZone ItemDropZone) const override;
+	virtual void Drop(const TArray<TSharedRef<FSequencerDisplayNode>>& DraggedNodes, EItemDropZone ItemDropZone) override;
+	virtual int32 GetSortingOrder() const override;
+	virtual void SetSortingOrder(const int32 InSortingOrder) override;
+	virtual void ModifyAndSetSortingOrder(const int32 InSortingOrder) override;
 
 protected:
 
