Copying //UE4/Dev-Sequencer to //UE4/Dev-Main (Source: //UE4/Dev-Sequencer @ 3946692)

diff --git a/Engine/Source/Editor/Sequencer/Private/DisplayNodes/SequencerSpacerNode.cpp b/Engine/Source/Editor/Sequencer/Private/DisplayNodes/SequencerSpacerNode.cpp
new file mode 100644
--- /dev/null
+++ b/Engine/Source/Editor/Sequencer/Private/DisplayNodes/SequencerSpacerNode.cpp
@@ -0,0 +1,135 @@
+// Copyright 1998-2018 Epic Games, Inc. All Rights Reserved.
+
+#include "SequencerSpacerNode.h"
+#include "Widgets/Views/STableRow.h"
+#include "MovieScene.h"
+#include "SequencerDisplayNodeDragDropOp.h"
+#include "SequencerFolderNode.h"
+#include "MovieSceneFolder.h"
+#include "SequencerDisplayNode.h"
+#include "Sequencer.h"
+#include "MovieSceneSequence.h"
+#include "SequencerNodeSortingMethods.h"
+#include "ScopedTransaction.h"
+
+TOptional<EItemDropZone> FSequencerSpacerNode::CanDrop(FSequencerDisplayNodeDragDropOp& DragDropOp, EItemDropZone ItemDropZone) const
+{
+	// Some spacers shouldn't allow dropping above or below at all as they're 1px tall between rows.
+	if (!bIsDragAndDropTarget)
+		return TOptional<EItemDropZone>();
+
+	// The spacer node is currently always used at the bottom of the list and we want it to be a drop target that puts things
+	// at the end of the root level. We still need to do a folder name collision check though as people could be moving out of
+	// a nestled folder up to the root.
+	// If they're trying to move an item above/below or into us. This item may or may not be a sibling; If it is not already a sibling then
+	// we need to check if it has a conflicting name with a sibling that we already have (for folders).
+	TArray<UMovieSceneFolder*> AdjacentFolders;
+	TSharedPtr<FSequencerDisplayNode> ChildParent;
+
+	if (ItemDropZone == EItemDropZone::OntoItem)
+	{
+		// If the item is being dropped onto us, we check our own children for name conflicts
+		ChildParent = SharedThis((FSequencerDisplayNode*)this);
+	}
+	else
+	{
+		// Otherwise, we try to check our siblings
+		ChildParent = GetParent();
+	}
+
+	if (ChildParent.IsValid())
+	{
+		for (TSharedRef <FSequencerDisplayNode> Child : ChildParent->GetChildNodes())
+		{
+			if (Child->GetType() == ESequencerNode::Folder)
+			{
+				TSharedRef<FSequencerFolderNode> FolderNode = StaticCastSharedRef<FSequencerFolderNode>(Child);
+				AdjacentFolders.Add(&FolderNode->GetFolder());
+			}
+		}
+	}
+	else
+	{
+		// If this folder has no parent then this is a root level folder, so we need to check the Movie Scene's child list for conflicting children names.
+		UMovieScene* FocusedMovieScene = GetSequencer().GetFocusedMovieSceneSequence()->GetMovieScene();
+		AdjacentFolders.Append(FocusedMovieScene->GetRootFolders());
+	}
+
+	// Check each node we're dragging to see if any of them have a name conflict - if so, block the whole drag/drop operation.
+	for (TSharedRef<FSequencerDisplayNode> DraggedNode : DragDropOp.GetDraggedNodes())
+	{
+		if (DraggedNode->GetType() == ESequencerNode::Folder)
+		{
+			TSharedRef<FSequencerFolderNode> DraggedFolder = StaticCastSharedRef<FSequencerFolderNode>(DraggedNode);
+
+			// Name Conflicts are only an issue on folders.
+			bool bHasNameConflict = false;
+			for (UMovieSceneFolder* Folder : AdjacentFolders)
+			{
+
+				// We don't allow a folder with the same name to become a sibling, but we need to not check the dragged node if it is already at that
+				// hierarchy depth so that we can rearrange them by triggering EItemDropZone::AboveItem / EItemDropZone::BelowItem on the same hierarchy.
+				if (&DraggedFolder->GetFolder() != Folder && DraggedFolder->GetFolder().GetFolderName() == Folder->GetFolderName())
+				{
+					bHasNameConflict = true;
+					break;
+				}
+			}
+
+			if (bHasNameConflict)
+			{
+				DragDropOp.CurrentHoverText = FText::Format(
+					NSLOCTEXT("SequencerFolderNode", "DuplicateFolderDragErrorFormat", "Folder with name '{0}' already exists."),
+					FText::FromName(DraggedFolder->GetFolder().GetFolderName()));
+
+				return TOptional<EItemDropZone>();
+			}
+		}
+	}
+
+	// If there's no collision, we'll force them to drop above us so the UI draws where we want it to.
+	return TOptional<EItemDropZone>(EItemDropZone::AboveItem);
+}
+
+void FSequencerSpacerNode::Drop(const TArray<TSharedRef<FSequencerDisplayNode>>& DraggedNodes, EItemDropZone ItemDropZone)
+{
+	const FScopedTransaction Transaction(NSLOCTEXT("SequencerTrackNode", "MoveItems", "Move items."));
+	for (TSharedRef<FSequencerDisplayNode> DraggedNode : DraggedNodes)
+	{
+		TSharedPtr<FSequencerDisplayNode> DraggedSeqNodeParent = DraggedNode->GetParent();
+
+		if (GetParent().IsValid())
+		{
+			// If the object is coming from the root or it's coming from another folder then we can allow it to move adjacent to us.
+			if (!DraggedSeqNodeParent.IsValid() || (DraggedSeqNodeParent.IsValid() && DraggedSeqNodeParent->GetType() == ESequencerNode::Folder))
+			{
+				checkf(GetParent()->GetType() == ESequencerNode::Folder, TEXT("Cannot reorder when parent is not a folder."));
+				TSharedPtr<FSequencerFolderNode> ParentFolder = StaticCastSharedPtr<FSequencerFolderNode>(GetParent());
+
+				// Let the folder we're going into remove us from our old parent and put us as a child of it first.
+				ParentFolder->MoveDisplayNodeToFolder(DraggedNode);
+			}
+		}
+		else
+		{
+			// We're at root and they're placing above or below us
+			ParentTree.MoveDisplayNodeToRoot(DraggedNode);
+		}
+	}
+
+	if (GetParent().IsValid())
+	{
+		checkf(GetParent()->GetType() == ESequencerNode::Folder, TEXT("Cannot reorder when parent is not a folder."));
+		TSharedPtr<FSequencerFolderNode> ParentFolder = StaticCastSharedPtr<FSequencerFolderNode>(GetParent());
+
+		// Sort our dragged nodes relative to our siblings.
+		SortAndSetSortingOrder(DraggedNodes, ParentFolder->GetChildNodes(), ItemDropZone, FDisplayNodeTreePositionSorter(), SharedThis(this));
+	}
+	else
+	{
+		// We're at root and they're placing above or below us
+		SortAndSetSortingOrder(DraggedNodes, GetSequencer().GetNodeTree()->GetRootNodes(), ItemDropZone, FDisplayNodeTreePositionSorter(), SharedThis(this));
+	}
+
+	ParentTree.GetSequencer().NotifyMovieSceneDataChanged(EMovieSceneDataChangeType::MovieSceneStructureItemsChanged);
+}
\ No newline at end of file
diff --git a/Engine/Source/Editor/Sequencer/Private/DisplayNodes/SequencerSpacerNode.h b/Engine/Source/Editor/Sequencer/Private/DisplayNodes/SequencerSpacerNode.h
--- a/Engine/Source/Editor/Sequencer/Private/DisplayNodes/SequencerSpacerNode.h
+++ b/Engine/Source/Editor/Sequencer/Private/DisplayNodes/SequencerSpacerNode.h
@@ -23,10 +23,12 @@ public:
 	 * 
 	 * @param InParentNode The parent of this node, or nullptr if this is a root node.
 	 * @param InParentTree The tree this node is in.
+	 * @param InIsDragAndDropTarget Determines whether or not this spacer will allow dropping items above it at all.
 	 */
-	FSequencerSpacerNode(float InSize, TSharedPtr<FSequencerDisplayNode> InParentNode, FSequencerNodeTree& InParentTree)
+	FSequencerSpacerNode(float InSize, TSharedPtr<FSequencerDisplayNode> InParentNode, FSequencerNodeTree& InParentTree, bool InIsDragAndDropTarget)
 		: FSequencerDisplayNode(NAME_None, InParentNode, InParentTree)
 		, Size(InSize)
+		, bIsDragAndDropTarget(InIsDragAndDropTarget)
 	{ }
 
 public:
@@ -40,9 +42,14 @@ public:
 	virtual ESequencerNode::Type GetType() const override { return ESequencerNode::Spacer; }
 	virtual void SetDisplayName(const FText& NewDisplayName) override { }
 	virtual TSharedRef<SWidget> GenerateContainerWidgetForOutliner(const TSharedRef<SSequencerTreeViewRow>& InRow) override { return SNew(SBox).HeightOverride(Size); }
-
+	virtual bool IsSelectable() const override { return false; }
+	virtual TOptional<EItemDropZone> CanDrop(FSequencerDisplayNodeDragDropOp& DragDropOp, EItemDropZone ItemDropZone) const override;
+	virtual void Drop(const TArray<TSharedRef<FSequencerDisplayNode>>& DraggedNodes, EItemDropZone ItemDropZone) override;
 private:
 
 	/** The size of the spacer */
 	float Size;
+
+	/** Does this spacer allow drag and drop operations. Spacers that go between individual rows shouldn't allow them. */
+	bool bIsDragAndDropTarget;
 };
diff --git a/Engine/Source/Editor/Sequencer/Private/DisplayNodes/SequencerTrackNode.cpp b/Engine/Source/Editor/Sequencer/Private/DisplayNodes/SequencerTrackNode.cpp
--- a/Engine/Source/Editor/Sequencer/Private/DisplayNodes/SequencerTrackNode.cpp
+++ b/Engine/Source/Editor/Sequencer/Private/DisplayNodes/SequencerTrackNode.cpp
@@ -5,6 +5,7 @@
 #include "Widgets/SBoxPanel.h"
 #include "Widgets/Layout/SSpacer.h"
 #include "Widgets/Layout/SBox.h"
+#include "Widgets/Views/STableRow.h"
 #include "DisplayNodes/SequencerObjectBindingNode.h"
 #include "Framework/Commands/UIAction.h"
 #include "Textures/SlateIcon.h"
@@ -17,6 +18,14 @@
 #include "SequencerUtilities.h"
 #include "SKeyNavigationButtons.h"
 #include "Compilation/MovieSceneSegmentCompiler.h"
+#include "SequencerDisplayNodeDragDropOp.h"
+#include "SequencerFolderNode.h"
+#include "SequencerNodeSortingMethods.h"
+#include "SequencerNodeTree.h"
+#include "DisplayNodes/SequencerFolderNode.h"
+#include "MovieSceneSequence.h"
+#include "MovieScene.h"
+#include "MovieSceneFolder.h"
 
 #define LOCTEXT_NAMESPACE "SequencerTrackNode"
 
@@ -30,7 +39,7 @@ namespace SequencerNodeConstants
  *****************************************************************************/
 
 FSequencerTrackNode::FSequencerTrackNode(UMovieSceneTrack& InAssociatedTrack, ISequencerTrackEditor& InAssociatedEditor, bool bInCanBeDragged, TSharedPtr<FSequencerDisplayNode> InParentNode, FSequencerNodeTree& InParentTree)
-	: FSequencerDisplayNode(InAssociatedTrack.GetTrackName(), InParentNode, InParentTree)
+	: FSequencerDisplayNode(InAssociatedTrack.GetFName(), InParentNode, InParentTree)
 	, AssociatedEditor(InAssociatedEditor)
 	, AssociatedTrack(&InAssociatedTrack)
 	, bCanBeDragged(bInCanBeDragged)
@@ -46,7 +55,7 @@ void FSequencerTrackNode::SetSectionAsKeyArea(TSharedRef<IKeyArea>& KeyArea)
 	if( !TopLevelKeyNode.IsValid() )
 	{
 		bool bTopLevel = true;
-		TopLevelKeyNode = MakeShareable( new FSequencerSectionKeyAreaNode( GetNodeName(), FText::GetEmpty(), nullptr, ParentTree, bTopLevel ) );
+		TopLevelKeyNode = MakeShareable( new FSequencerSectionKeyAreaNode( GetNodeName(), FText::GetEmpty(), SharedThis(this), ParentTree, bTopLevel ) );
 	}
 
 	TopLevelKeyNode->AddKeyArea( KeyArea );
@@ -125,7 +134,7 @@ FReply FSequencerTrackNode::CreateNewSection() const
 	}
 
 	const int32 InsertAtIndex = SubTrackMode == ESubTrackMode::SubTrack ? GetRowIndex() : Track->GetMaxRowIndex() + 1;
-	const float StartAtTime = GetSequencer().GetLocalTime();
+	const FQualifiedFrameTime CurrentTime = GetSequencer().GetLocalTime();
 
 	FScopedTransaction Transaction(LOCTEXT("AddSectionText", "Add Section"));
 	UMovieSceneSection* Section = Track->CreateNewSection();
@@ -133,9 +142,8 @@ FReply FSequencerTrackNode::CreateNewSection() const
 	{
 		Track->Modify();
 
-		Section->SetIsInfinite(false);
-		Section->SetStartTime(StartAtTime);
-		Section->SetEndTime(StartAtTime + 10.f);
+		FFrameNumber Duration = (10.f * CurrentTime.Rate).RoundToFrame();
+		Section->SetRange(TRange<FFrameNumber>(CurrentTime.Time.FrameNumber, CurrentTime.Time.FrameNumber + Duration));
 		Section->SetRowIndex(InsertAtIndex);
 
 		Track->AddSection(*Section);
@@ -239,7 +247,155 @@ const FSlateBrush* FSequencerTrackNode::GetIconBrush() const
 
 bool FSequencerTrackNode::CanDrag() const
 {
-	return bCanBeDragged;
+	return bCanBeDragged && SubTrackMode != ESubTrackMode::SubTrack;
+}
+
+TOptional<EItemDropZone> FSequencerTrackNode::CanDrop(FSequencerDisplayNodeDragDropOp& DragDropOp, EItemDropZone ItemDropZone) const
+{
+	DragDropOp.ResetToDefaultToolTip();
+
+	// Prevent taking any parent that's part of the dragged node hierarchy from being put inside a child of itself
+	// This is done first before the other checks so that the UI stays consistent as you move between them, otherwise
+	// when you are above/below a node it reports this error, but if you were on top of a node it would do the standard
+	// no-drag-drop due to OntoItem being blocked. 
+	TSharedPtr<FSequencerDisplayNode> CurrentNode = SharedThis((FSequencerDisplayNode*)this);
+	while (CurrentNode.IsValid())
+	{
+		if (DragDropOp.GetDraggedNodes().Contains(CurrentNode))
+		{
+			DragDropOp.CurrentHoverText = NSLOCTEXT("SequencerFolderNode", "ParentIntoChildDragErrorFormat", "Can't drag a parent node into one of it's children.");
+			return TOptional<EItemDropZone>();
+		}
+		CurrentNode = CurrentNode->GetParent();
+	}
+
+	// If we can't be dragged, then we don't allow reordering things above or below us.
+	if (!CanDrag() || SubTrackMode == ESubTrackMode::SubTrack)
+	{
+		return TOptional<EItemDropZone>();
+	}
+
+	// Since tracks can't have children (via the UI) any attempts to drop below or onto them get rerouted into
+	// attempts to drop them above to help with the insert marker drawing in a confusing manner for multi-row children.
+	if (ItemDropZone == EItemDropZone::BelowItem || ItemDropZone == EItemDropZone::OntoItem)
+	{
+		ItemDropZone = EItemDropZone::AboveItem;
+	}
+
+	TArray<UMovieSceneFolder*> AdjacentFolders;
+	if (GetParent().IsValid())
+	{
+		// We are either trying to drop adjacent to ourself (when nestled), or as a child of ourself, so we add either our siblings or our children
+		// to the list of possibly conflicting names.
+		for (TSharedRef <FSequencerDisplayNode> Child : GetParent()->GetChildNodes())
+		{
+			if (Child->GetType() == ESequencerNode::Folder)
+			{
+				TSharedRef<FSequencerFolderNode> FolderNode = StaticCastSharedRef<FSequencerFolderNode>(Child);
+				AdjacentFolders.Add(&FolderNode->GetFolder());
+			}
+		}
+	}
+	else
+	{
+		// If this folder has no parent then this is a root level folder, so we need to check the Movie Scene's child list for conflicting children names.
+		UMovieScene* FocusedMovieScene = GetSequencer().GetFocusedMovieSceneSequence()->GetMovieScene();
+		AdjacentFolders.Append(FocusedMovieScene->GetRootFolders());
+	}
+
+	// Check each node we're dragging to see if any of them have a name conflict - if so, block the whole drag/drop operation.
+	for (TSharedRef<FSequencerDisplayNode> DraggedNode : DragDropOp.GetDraggedNodes())
+	{
+		if (DraggedNode->GetType() == ESequencerNode::Folder)
+		{
+			TSharedRef<FSequencerFolderNode> DraggedFolder = StaticCastSharedRef<FSequencerFolderNode>(DraggedNode);
+
+			// Name Conflicts are only an issue on folders.
+			bool bHasNameConflict = false;
+			for (UMovieSceneFolder* Folder : AdjacentFolders)
+			{
+				// We don't allow a folder with the same name to become a sibling, but we need to not check the dragged node if it is already at that
+				// hierarchy depth so that we can rearrange them by triggering EItemDropZone::AboveItem / EItemDropZone::BelowItem on the same hierarchy.
+				if (&DraggedFolder->GetFolder() != Folder && DraggedFolder->GetFolder().GetFolderName() == Folder->GetFolderName())
+				{
+					bHasNameConflict = true;
+					break;
+				}
+			}
+
+			if (bHasNameConflict)
+			{
+				DragDropOp.CurrentHoverText = FText::Format(
+					NSLOCTEXT("SequencerFolderNode", "DuplicateFolderDragErrorFormat", "Folder with name '{0}' already exists."),
+					FText::FromName(DraggedFolder->GetFolder().GetFolderName()));
+
+				return TOptional<EItemDropZone>();
+			}
+		}
+	}
+
+	TSharedPtr<FSequencerDisplayNode> ParentSeqNode = GetParent();
+	if (ParentSeqNode.IsValid())
+	{
+		if (ParentSeqNode->GetType() == ESequencerNode::Folder)
+		{
+			return TOptional<EItemDropZone>(ItemDropZone);
+		}
+		else
+		{
+			// If we have a parent who is not a folder (ie: This is a component track on an actor)
+			// then it can't be rearranged.
+			return TOptional<EItemDropZone>();
+		}
+	}
+	else
+	{
+		// We're at the root level and thus a Master track, so they can re-arrange.
+		return TOptional<EItemDropZone>(ItemDropZone);
+	}
+}
+
+void FSequencerTrackNode::Drop(const TArray<TSharedRef<FSequencerDisplayNode>>& DraggedNodes, EItemDropZone ItemDropZone)
+{
+	const FScopedTransaction Transaction(NSLOCTEXT("SequencerTrackNode", "MoveItems", "Move items."));
+	for (TSharedRef<FSequencerDisplayNode> DraggedNode : DraggedNodes)
+	{
+		TSharedPtr<FSequencerDisplayNode> DraggedSeqNodeParent = DraggedNode->GetParent();
+
+		if (GetParent().IsValid())
+		{
+			// If the object is coming from the root or it's coming from another folder then we can allow it to move adjacent to us.
+			if (!DraggedSeqNodeParent.IsValid() || (DraggedSeqNodeParent.IsValid() && DraggedSeqNodeParent->GetType() == ESequencerNode::Folder))
+			{
+				checkf(GetParent()->GetType() == ESequencerNode::Folder, TEXT("Cannot reorder when parent is not a folder."));
+				TSharedPtr<FSequencerFolderNode> ParentFolder = StaticCastSharedPtr<FSequencerFolderNode>(GetParent());
+
+				// Let the folder we're going into remove us from our old parent and put us as a child of it first.
+				ParentFolder->MoveDisplayNodeToFolder(DraggedNode);
+			}
+		}
+		else
+		{
+			// We're at root and they're placing above or below us
+			ParentTree.MoveDisplayNodeToRoot(DraggedNode);
+		}
+	}
+
+	if (GetParent().IsValid())
+	{
+		checkf(GetParent()->GetType() == ESequencerNode::Folder, TEXT("Cannot reorder when parent is not a folder."));
+		TSharedPtr<FSequencerFolderNode> ParentFolder = StaticCastSharedPtr<FSequencerFolderNode>(GetParent());
+
+		// Sort our dragged nodes relative to our siblings.
+		SortAndSetSortingOrder(DraggedNodes, ParentFolder->GetChildNodes(), ItemDropZone, FDisplayNodeTreePositionSorter(), SharedThis(this));
+	}
+	else
+	{
+		// We're at root and they're placing above or below us
+		SortAndSetSortingOrder(DraggedNodes, GetSequencer().GetNodeTree()->GetRootNodes(), ItemDropZone, FDisplayNodeTreePositionSorter(), SharedThis(this));
+	}
+
+	ParentTree.GetSequencer().NotifyMovieSceneDataChanged(EMovieSceneDataChangeType::MovieSceneStructureItemsChanged);
 }
 
 bool FSequencerTrackNode::IsResizable() const
@@ -326,7 +482,7 @@ void FSequencerTrackNode::SetDisplayName(const FText& NewDisplayName)
 
 TArray<FSequencerOverlapRange> FSequencerTrackNode::GetUnderlappingSections(UMovieSceneSection* InSection)
 {
-	TRange<float> InSectionRange = InSection->IsInfinite() ? TRange<float>::All() : InSection->GetRange();
+	TRange<FFrameNumber> InSectionRange = InSection->GetRange();
 
 	TMovieSceneEvaluationTree<int32> SectionIndexTree;
 
@@ -339,8 +495,8 @@ TArray<FSequencerOverlapRange> FSequencerTrackNode::GetUnderlappingSections(UMov
 			continue;
 		}
 
-		TRange<float> OtherSectionRange = SectionObj->IsInfinite() ? TRange<float>::All() : SectionObj->GetRange();
-		TRange<float> Intersection = TRange<float>::Intersection(OtherSectionRange, InSectionRange);
+		TRange<FFrameNumber> OtherSectionRange = SectionObj->GetRange();
+		TRange<FFrameNumber> Intersection = TRange<FFrameNumber>::Intersection(OtherSectionRange, InSectionRange);
 		if (!Intersection.IsEmpty())
 		{
 			SectionIndexTree.Add(Intersection, SectionIndex);
@@ -379,7 +535,7 @@ TArray<FSequencerOverlapRange> FSequencerTrackNode::GetUnderlappingSections(UMov
 
 TArray<FSequencerOverlapRange> FSequencerTrackNode::GetEasingSegmentsForSection(UMovieSceneSection* InSection)
 {
-	TRange<float> InSectionRange = InSection->IsInfinite() ? TRange<float>::All() : InSection->GetRange();
+	TRange<FFrameNumber> InSectionRange = InSection->GetRange();
 
 	TArray<FMovieSceneSectionData> CompileData;
 
@@ -394,13 +550,13 @@ TArray<FSequencerOverlapRange> FSequencerTrackNode::GetEasingSegmentsForSection(
 			continue;
 		}
 
-		TRange<float> Intersection = TRange<float>::Intersection(SectionObj->GetEaseInRange(), InSectionRange);
+		TRange<FFrameNumber> Intersection = TRange<FFrameNumber>::Intersection(SectionObj->GetEaseInRange(), InSectionRange);
 		if (!Intersection.IsEmpty())
 		{
 			SectionIndexTree.Add(Intersection, SectionIndex);
 		}
 
-		Intersection = TRange<float>::Intersection(SectionObj->GetEaseOutRange(), InSectionRange);
+		Intersection = TRange<FFrameNumber>::Intersection(SectionObj->GetEaseOutRange(), InSectionRange);
 		if (!Intersection.IsEmpty())
 		{
 			SectionIndexTree.Add(Intersection, SectionIndex);
@@ -437,4 +593,33 @@ TArray<FSequencerOverlapRange> FSequencerTrackNode::GetEasingSegmentsForSection(
 	return Result;
 }
 
+int32 FSequencerTrackNode::GetSortingOrder() const
+{
+	const UMovieSceneTrack* MovieSceneTrack = AssociatedTrack.Get();
+	if (MovieSceneTrack)
+	{
+		return MovieSceneTrack->GetSortingOrder();
+	}
+
+	return 0;
+}
+
+void FSequencerTrackNode::SetSortingOrder(const int32 InSortingOrder)
+{
+	UMovieSceneTrack* MovieSceneTrack = AssociatedTrack.Get();
+	if (MovieSceneTrack)
+	{
+		MovieSceneTrack->SetSortingOrder(InSortingOrder);
+	}
+}
+
+void FSequencerTrackNode::ModifyAndSetSortingOrder(const int32 InSortingOrder)
+{
+	UMovieSceneTrack* MovieSceneTrack = AssociatedTrack.Get();
+	if (MovieSceneTrack)
+	{
+		MovieSceneTrack->Modify();
+		SetSortingOrder(InSortingOrder);
+	}
+}
 #undef LOCTEXT_NAMESPACE
\ No newline at end of file
diff --git a/Engine/Source/Editor/Sequencer/Private/DisplayNodes/SequencerTrackNode.h b/Engine/Source/Editor/Sequencer/Private/DisplayNodes/SequencerTrackNode.h
--- a/Engine/Source/Editor/Sequencer/Private/DisplayNodes/SequencerTrackNode.h
+++ b/Engine/Source/Editor/Sequencer/Private/DisplayNodes/SequencerTrackNode.h
@@ -19,7 +19,7 @@ struct FSlateBrush;
 struct FSequencerOverlapRange
 {
 	/** The range for the overlap */
-	TRange<float> Range;
+	TRange<FFrameNumber> Range;
 	/** The sections that occupy this range, sorted by overlap priority */
 	TArray<FSectionHandle> Sections;
 };
@@ -149,9 +149,14 @@ public:
 	virtual void SetDisplayName(const FText& NewDisplayName) override;
 	virtual const FSlateBrush* GetIconBrush() const override;
 	virtual bool CanDrag() const override;
+	virtual TOptional<EItemDropZone> CanDrop(FSequencerDisplayNodeDragDropOp& DragDropOp, EItemDropZone ItemDropZone) const override;
+	virtual void Drop(const TArray<TSharedRef<FSequencerDisplayNode>>& DraggedNodes, EItemDropZone ItemDropZone) override;
 	virtual bool IsResizable() const override;
 	virtual void Resize(float NewSize) override;
-	
+	virtual int32 GetSortingOrder() const override;
+	virtual void SetSortingOrder(const int32 InSortingOrder) override;
+	virtual void ModifyAndSetSortingOrder(const int32 InSortingOrder) override;
+
 private:
 
 	FReply CreateNewSection() const;
diff --git a/Engine/Source/Editor/Sequencer/Private/IKeyArea.cpp b/Engine/Source/Editor/Sequencer/Private/IKeyArea.cpp
new file mode 100644
--- /dev/null
+++ b/Engine/Source/Editor/Sequencer/Private/IKeyArea.cpp
@@ -0,0 +1,243 @@
+// Copyright 1998-2018 Epic Games, Inc. All Rights Reserved.
+
+#include "IKeyArea.h"
+#include "ISequencerModule.h"
+#include "Modules/ModuleManager.h"
+#include "Widgets/SNullWidget.h"
+#include "ISequencerChannelInterface.h"
+#include "SequencerClipboardReconciler.h"
+#include "Tracks/MovieScenePropertyTrack.h"
+#include "Channels/MovieSceneChannelProxy.h"
+#include "CurveModel.h"
+
+IKeyArea::IKeyArea(
+	uint32                               InChannelTypeID,
+	UMovieSceneSection&                  InSection,
+	TMovieSceneChannelHandle<void>       InChannel,
+	const FMovieSceneChannelEditorData&  InCommonEditorData,
+	TMovieSceneChannelHandle<const void> InSpecializedData
+)
+	: ChannelTypeID(InChannelTypeID)
+	, WeakOwningSection(&InSection)
+	, Channel(InChannel)
+	, SpecializedData(InSpecializedData)
+	, Color(InCommonEditorData.Color)
+	, ChannelName(InCommonEditorData.Name)
+	, DisplayText(InCommonEditorData.DisplayText)
+{
+	UMovieScenePropertyTrack* PropertyTrack = InSection.GetTypedOuter<UMovieScenePropertyTrack>();
+	if (PropertyTrack && !PropertyTrack->GetPropertyPath().IsEmpty())
+	{
+		PropertyBindings = FTrackInstancePropertyBindings(PropertyTrack->GetPropertyName(), PropertyTrack->GetPropertyPath());
+	}
+}
+
+void* IKeyArea::GetChannelPtr() const
+{
+	return Channel.Get();
+}
+
+const void* IKeyArea::GetSpecializedEditorData() const
+{
+	return SpecializedData.Get();
+}
+
+UMovieSceneSection* IKeyArea::GetOwningSection() const
+{
+	return WeakOwningSection.Get();
+}
+
+FName IKeyArea::GetName() const
+{
+	return ChannelName;
+}
+
+void IKeyArea::SetName(FName InName)
+{
+	ChannelName = InName;
+}
+
+ISequencerChannelInterface* IKeyArea::FindChannelInterface() const
+{
+	ISequencerModule& SequencerModule = FModuleManager::LoadModuleChecked<ISequencerModule>("Sequencer");
+	ISequencerChannelInterface* ChannelInterface = SequencerModule.FindChannelInterface(ChannelTypeID);
+	ensureMsgf(ChannelInterface, TEXT("No channel interface found for type ID 0x%08x. Did you forget to call ISequencerModule::RegisterChannelInterface<ChannelType>()?"), ChannelTypeID);
+	return ChannelInterface;
+}
+
+bool IKeyArea::HasAnyKeys() const
+{
+	ISequencerChannelInterface* ChannelInterface = FindChannelInterface();
+	void *const RawChannelPtr = GetChannelPtr();
+
+	return ChannelInterface && RawChannelPtr && ChannelInterface->HasAnyKeys_Raw(MakeArrayView(&RawChannelPtr, 1));
+}
+
+FKeyHandle IKeyArea::AddOrUpdateKey(FFrameNumber Time, const FGuid& ObjectBindingID, ISequencer& InSequencer)
+{
+	ISequencerChannelInterface* ChannelInterface = FindChannelInterface();
+	void* RawChannelPtr = GetChannelPtr();
+
+	// The specialized data may be null, but is passed to the interface regardless
+	const void* RawSpecializedData = GetSpecializedEditorData();
+
+	if (ChannelInterface && RawChannelPtr)
+	{
+		FTrackInstancePropertyBindings* BindingsPtr = PropertyBindings.IsSet() ? &PropertyBindings.GetValue() : nullptr;
+		return ChannelInterface->AddOrUpdateKey_Raw(RawChannelPtr, RawSpecializedData, Time, InSequencer, ObjectBindingID, BindingsPtr);
+	}
+
+	return FKeyHandle();
+}
+
+FKeyHandle IKeyArea::DuplicateKey(FKeyHandle InKeyHandle) const
+{
+	ISequencerChannelInterface* ChannelInterface = FindChannelInterface();
+
+	FKeyHandle NewHandle = FKeyHandle::Invalid();
+
+	void* RawChannelPtr = GetChannelPtr();
+	if (ChannelInterface && RawChannelPtr)
+	{
+		ChannelInterface->DuplicateKeys_Raw(RawChannelPtr, TArrayView<const FKeyHandle>(&InKeyHandle, 1), TArrayView<FKeyHandle>(&NewHandle, 1));
+	}
+	return NewHandle;
+}
+
+void IKeyArea::SetKeyTimes(TArrayView<const FKeyHandle> InKeyHandles, TArrayView<const FFrameNumber> InKeyTimes) const
+{
+	check(InKeyHandles.Num() == InKeyTimes.Num());
+
+	ISequencerChannelInterface* ChannelInterface = FindChannelInterface();
+	void* RawChannelPtr = GetChannelPtr();
+
+	if (ChannelInterface && RawChannelPtr)
+	{
+		ChannelInterface->SetKeyTimes_Raw(RawChannelPtr, InKeyHandles, InKeyTimes);
+	}
+}
+
+void IKeyArea::GetKeyTimes(TArrayView<const FKeyHandle> InKeyHandles, TArrayView<FFrameNumber> OutTimes) const
+{
+	ISequencerChannelInterface* ChannelInterface = FindChannelInterface();
+
+	void* RawChannelPtr = GetChannelPtr();
+	if (ChannelInterface && RawChannelPtr)
+	{
+		ChannelInterface->GetKeyTimes_Raw(RawChannelPtr, InKeyHandles, OutTimes);
+	}
+}
+
+void IKeyArea::GetKeyInfo(TArray<FKeyHandle>* OutHandles, TArray<FFrameNumber>* OutTimes, const TRange<FFrameNumber>& WithinRange) const
+{
+	ISequencerChannelInterface* ChannelInterface = FindChannelInterface();
+
+	void* RawChannelPtr = GetChannelPtr();
+	if (ChannelInterface && RawChannelPtr)
+	{
+		ChannelInterface->GetKeys_Raw(RawChannelPtr, WithinRange, OutTimes, OutHandles);
+	}
+}
+
+TSharedPtr<FStructOnScope> IKeyArea::GetKeyStruct(FKeyHandle KeyHandle) const
+{
+	ISequencerChannelInterface* ChannelInterface = FindChannelInterface();
+
+	if (ChannelInterface && Channel.Get())
+	{
+		return ChannelInterface->GetKeyStruct_Raw(Channel, KeyHandle);
+	}
+	return nullptr;
+}
+
+void IKeyArea::DrawKeys(TArrayView<const FKeyHandle> InKeyHandles, TArrayView<FKeyDrawParams> OutKeyDrawParams)
+{
+	check(InKeyHandles.Num() == OutKeyDrawParams.Num());
+
+	ISequencerChannelInterface* ChannelInterface = FindChannelInterface();
+
+	void* RawChannelPtr = GetChannelPtr();
+	if (ChannelInterface && RawChannelPtr)
+	{
+		return ChannelInterface->DrawKeys_Raw(RawChannelPtr, InKeyHandles, OutKeyDrawParams);
+	}
+}
+
+bool IKeyArea::CanCreateKeyEditor() const
+{
+	ISequencerChannelInterface* ChannelInterface = FindChannelInterface();
+	void* RawChannelPtr = GetChannelPtr();
+	return ChannelInterface && RawChannelPtr && ChannelInterface->CanCreateKeyEditor_Raw(RawChannelPtr);
+}
+
+TSharedRef<SWidget> IKeyArea::CreateKeyEditor(TWeakPtr<ISequencer> Sequencer, const FGuid& ObjectBindingID)
+{
+	ISequencerChannelInterface* ChannelInterface = FindChannelInterface();
+	void* RawChannelPtr = GetChannelPtr();
+	UMovieSceneSection* OwningSection = GetOwningSection();
+
+	// The specialized data may be null, but is passed to the interface regardless
+	const void* RawSpecializedData = GetSpecializedEditorData();
+
+	TWeakPtr<FTrackInstancePropertyBindings> PropertyBindingsPtr;
+	if (PropertyBindings.IsSet())
+	{
+		PropertyBindingsPtr = TSharedPtr<FTrackInstancePropertyBindings>(AsShared(), &PropertyBindings.GetValue());
+	}
+
+	if (ChannelInterface && RawChannelPtr && OwningSection)
+	{
+		return ChannelInterface->CreateKeyEditor_Raw(RawChannelPtr, RawSpecializedData, OwningSection, ObjectBindingID, PropertyBindingsPtr, Sequencer);
+	}
+	return SNullWidget::NullWidget;
+}
+
+void IKeyArea::CopyKeys(FMovieSceneClipboardBuilder& ClipboardBuilder, TArrayView<const FKeyHandle> KeyMask) const
+{
+	ISequencerChannelInterface* ChannelInterface = FindChannelInterface();
+	void* RawChannelPtr = GetChannelPtr();
+	UMovieSceneSection* OwningSection = GetOwningSection();
+	if (ChannelInterface && RawChannelPtr && OwningSection)
+	{
+		ChannelInterface->CopyKeys_Raw(RawChannelPtr, OwningSection, ChannelName, ClipboardBuilder, KeyMask);
+	}
+}
+
+void IKeyArea::PasteKeys(const FMovieSceneClipboardKeyTrack& KeyTrack, const FMovieSceneClipboardEnvironment& SrcEnvironment, const FSequencerPasteEnvironment& DstEnvironment)
+{
+	ISequencerChannelInterface* ChannelInterface = FindChannelInterface();
+	void* RawChannelPtr = GetChannelPtr();
+	UMovieSceneSection* OwningSection = GetOwningSection();
+	if (ChannelInterface && RawChannelPtr && OwningSection)
+	{
+		TArray<FKeyHandle> PastedKeys;
+		ChannelInterface->PasteKeys_Raw(RawChannelPtr, OwningSection, KeyTrack, SrcEnvironment, DstEnvironment, PastedKeys);
+
+		for (FKeyHandle KeyHandle : PastedKeys)
+		{
+			DstEnvironment.ReportPastedKey(KeyHandle, *this);
+		}
+	}
+}
+
+TUniquePtr<FCurveModel> IKeyArea::CreateCurveEditorModel(TSharedRef<ISequencer> InSequencer) const
+{
+	ISequencerChannelInterface* ChannelInterface = FindChannelInterface();
+	void* RawChannelPtr = GetChannelPtr();
+	UMovieSceneSection* OwningSection = GetOwningSection();
+	if (ChannelInterface && RawChannelPtr && OwningSection)
+	{
+		TUniquePtr<FCurveModel> CurveModel = ChannelInterface->CreateCurveEditorModel_Raw(OwningSection->GetChannelProxy().MakeHandle(RawChannelPtr), OwningSection, InSequencer);
+		if (CurveModel.IsValid())
+		{
+			CurveModel->SetDisplayName(DisplayText);
+			if (Color.IsSet())
+			{
+				CurveModel->SetColor(Color.GetValue());
+			}
+		}
+		return CurveModel;
+	}
+
+	return nullptr;
+}
diff --git a/Engine/Source/Editor/Sequencer/Private/ISequencer.cpp b/Engine/Source/Editor/Sequencer/Private/ISequencer.cpp
new file mode 100644
--- /dev/null
+++ b/Engine/Source/Editor/Sequencer/Private/ISequencer.cpp
@@ -0,0 +1,53 @@
+// Copyright 1998-2018 Epic Games, Inc. All Rights Reserved.
+
+#include "ISequencer.h"
+#include "MovieSceneSequence.h"
+#include "MovieScene.h"
+
+FFrameRate ISequencer::GetRootFrameResolution() const
+{
+	UMovieSceneSequence* RootSequence = GetRootMovieSceneSequence();
+	if (RootSequence)
+	{
+		return RootSequence->GetMovieScene()->GetFrameResolution();
+	}
+
+	ensureMsgf(false, TEXT("No valid sequence found."));
+	return FFrameRate();
+}
+
+FFrameRate ISequencer::GetRootPlayRate() const
+{
+	UMovieSceneSequence* RootSequence = GetRootMovieSceneSequence();
+	if (RootSequence)
+	{
+		return RootSequence->GetMovieScene()->GetPlaybackFrameRate();
+	}
+
+	ensureMsgf(false, TEXT("No valid sequence found."));
+	return FFrameRate();
+}
+
+FFrameRate ISequencer::GetFocusedFrameResolution() const
+{
+	UMovieSceneSequence* FocusedSequence = GetFocusedMovieSceneSequence();
+	if (FocusedSequence)
+	{
+		return FocusedSequence->GetMovieScene()->GetFrameResolution();
+	}
+
+	ensureMsgf(false, TEXT("No valid sequence found."));
+	return FFrameRate();
+}
+
+FFrameRate ISequencer::GetFocusedPlayRate() const
+{
+	UMovieSceneSequence* FocusedSequence = GetFocusedMovieSceneSequence();
+	if (FocusedSequence)
+	{
+		return FocusedSequence->GetMovieScene()->GetPlaybackFrameRate();
+	}
+
+	ensureMsgf(false, TEXT("No valid sequence found."));
+	return FFrameRate();
+}
\ No newline at end of file
diff --git a/Engine/Source/Editor/Sequencer/Private/ISequencerSection.cpp b/Engine/Source/Editor/Sequencer/Private/ISequencerSection.cpp
new file mode 100644
--- /dev/null
+++ b/Engine/Source/Editor/Sequencer/Private/ISequencerSection.cpp
@@ -0,0 +1,199 @@
+// Copyright 1998-2018 Epic Games, Inc. All Rights Reserved.
+
+#include "ISequencerSection.h"
+#include "MovieSceneSection.h"
+#include "ISectionLayoutBuilder.h"
+#include "SequencerSectionPainter.h"
+#include "IKeyArea.h"
+#include "Channels/MovieSceneChannelProxy.h"
+
+/** Structure used during key area creation to group channels by their group name */
+struct FChannelData
+{
+	/** The channel's type ID (FMovieSceneChannelEntry::RegisterNewID) */
+	uint32 ChannelID;
+
+	/** Handle to the channel */
+	TMovieSceneChannelHandle<void> Channel;
+
+	/** The channel's common editor data */
+	const FMovieSceneChannelEditorData& EditorData;
+
+	/** Handle to the channel's specialized editor data */
+	TMovieSceneChannelHandle<const void> SpecializedData;
+
+	/**
+	 * Make a key area out of this data
+	 */
+	TSharedRef<IKeyArea> MakeKeyArea(UMovieSceneSection& InSection) const
+	{
+		return MakeShared<IKeyArea>(ChannelID, InSection, Channel, EditorData, SpecializedData);
+	}
+};
+
+/** Data pertaining to a group of channels */
+struct FGroupData
+{
+	FGroupData(FText InGroupText)
+		: GroupText(InGroupText)
+		, SortOrder(-1)
+	{}
+
+	void AddChannel(FChannelData&& InChannel)
+	{
+		if (InChannel.EditorData.SortOrder < SortOrder)
+		{
+			SortOrder = InChannel.EditorData.SortOrder;
+		}
+
+		Channels.Add(MoveTemp(InChannel));
+	}
+
+	/** Text to display for the group */
+	FText GroupText;
+
+	/** Sort order of the group */
+	uint8 SortOrder;
+
+	/** Array of channels within this group */
+	TArray<FChannelData, TInlineAllocator<4>> Channels;
+};
+
+void ISequencerSection::GenerateSectionLayout( ISectionLayoutBuilder& LayoutBuilder )
+{
+	UMovieSceneSection* Section = GetSectionObject();
+	if (!Section)
+	{
+		return;
+	}
+
+	// Group channels by their group name
+	TMap<FName, FGroupData> GroupToChannelsMap;
+
+	FMovieSceneChannelProxy& ChannelProxy = Section->GetChannelProxy();
+	for (const FMovieSceneChannelEntry& Entry : Section->GetChannelProxy().GetAllEntries())
+	{
+		const uint32 ChannelID = Entry.GetChannelID();
+
+		// One editor data ptr per channel
+		TArrayView<void* const>                        Channels         = Entry.GetChannels();
+		TArrayView<const FMovieSceneChannelEditorData> CommonEditorData = Entry.GetCommonEditorData();
+
+		for (int32 Index = 0; Index < Channels.Num(); ++Index)
+		{
+			TMovieSceneChannelHandle<void>       Channel         = ChannelProxy.MakeHandle(Channels[Index]);
+			TMovieSceneChannelHandle<const void> SpecializedData = ChannelProxy.MakeHandle(Entry.GetSpecializedEditorData(Index));
+
+			const FMovieSceneChannelEditorData& EditorData = CommonEditorData[Index];
+			if (EditorData.bEnabled)
+			{
+				FName GroupName = *EditorData.Group.ToString();
+
+				FGroupData* ExistingGroup = GroupToChannelsMap.Find(GroupName);
+				if (!ExistingGroup)
+				{
+					ExistingGroup = &GroupToChannelsMap.Add(GroupName, FGroupData(EditorData.Group));
+				}
+
+				FChannelData NewChannel{ ChannelID, Channel, EditorData, SpecializedData };
+				ExistingGroup->AddChannel(MoveTemp(NewChannel));
+			}
+		}
+	}
+
+	if (GroupToChannelsMap.Num() == 0)
+	{
+		return;
+	}
+
+	// Collapse single channels to the top level track node if allowed
+	if (GroupToChannelsMap.Num() == 1)
+	{
+		const TTuple<FName, FGroupData>& Pair = *GroupToChannelsMap.CreateIterator();
+		if (Pair.Value.Channels.Num() == 1 && Pair.Value.Channels[0].EditorData.bCanCollapseToTrack)
+		{
+			LayoutBuilder.SetSectionAsKeyArea(Pair.Value.Channels[0].MakeKeyArea(*Section));
+			return;
+		}
+	}
+
+	// Sort the channels in each group by its sort order and name
+	TArray<FName, TInlineAllocator<6>> SortedGroupNames;
+	for (auto& Pair : GroupToChannelsMap)
+	{
+		SortedGroupNames.Add(Pair.Key);
+
+		// Sort by sort order then name
+		Pair.Value.Channels.Sort([](const FChannelData& A, const FChannelData& B){
+			if (A.EditorData.SortOrder == B.EditorData.SortOrder)
+			{
+				return A.EditorData.Name < B.EditorData.Name;
+			}
+			return A.EditorData.SortOrder < B.EditorData.SortOrder;
+		});
+	}
+
+	// Sort groups by the lowest sort order in each group
+	auto SortPredicate = [&GroupToChannelsMap](FName A, FName B)
+	{
+		if (A.IsNone())
+		{
+			return false;
+		}
+		else if (B.IsNone())
+		{
+			return true;
+		}
+
+		const int32 SortOrderA = GroupToChannelsMap.FindChecked(A).SortOrder;
+		const int32 SortOrderB = GroupToChannelsMap.FindChecked(B).SortOrder;
+		return SortOrderA < SortOrderB;
+	};
+	SortedGroupNames.Sort(SortPredicate);
+
+	// Create key areas for each group name
+	for (FName GroupName : SortedGroupNames)
+	{
+		auto& ChannelData = GroupToChannelsMap.FindChecked(GroupName);
+
+		if (!GroupName.IsNone())
+		{
+			LayoutBuilder.PushCategory(GroupName, ChannelData.GroupText);
+		}
+
+		for (const FChannelData& ChannelAndData : ChannelData.Channels)
+		{
+			TSharedRef<IKeyArea> KeyArea = ChannelAndData.MakeKeyArea(*Section);
+			LayoutBuilder.AddKeyArea(ChannelAndData.EditorData.Name, ChannelAndData.EditorData.DisplayText, KeyArea);
+		}
+
+		if (!GroupName.IsNone())
+		{
+			LayoutBuilder.PopCategory();
+		}
+	}
+}
+
+void ISequencerSection::ResizeSection(ESequencerSectionResizeMode ResizeMode, FFrameNumber ResizeFrameNumber)
+{
+	UMovieSceneSection* SectionObject = GetSectionObject();
+	if (ResizeMode == ESequencerSectionResizeMode::SSRM_LeadingEdge)
+	{
+		FFrameNumber MaxFrame = SectionObject->HasEndFrame() ? SectionObject->GetExclusiveEndFrame()-1 : TNumericLimits<int32>::Max();
+		ResizeFrameNumber = FMath::Min( ResizeFrameNumber, MaxFrame );
+
+		SectionObject->SetRange(TRange<FFrameNumber>(TRangeBound<FFrameNumber>::Inclusive(ResizeFrameNumber), SectionObject->GetRange().GetUpperBound()));
+	}
+	else
+	{
+		FFrameNumber MinFrame = SectionObject->HasStartFrame() ? SectionObject->GetInclusiveStartFrame() : TNumericLimits<int32>::Lowest();
+		ResizeFrameNumber = FMath::Max( ResizeFrameNumber, MinFrame );
+
+		SectionObject->SetRange(TRange<FFrameNumber>(SectionObject->GetRange().GetLowerBound(), TRangeBound<FFrameNumber>::Exclusive(ResizeFrameNumber)));
+	}
+}
+
+int32 FSequencerSection::OnPaintSection(FSequencerSectionPainter& Painter) const
+{
+	return Painter.PaintSectionBackground();
+}
\ No newline at end of file
diff --git a/Engine/Source/Editor/Sequencer/Private/KeyPropertyParams.cpp b/Engine/Source/Editor/Sequencer/Private/KeyPropertyParams.cpp
--- a/Engine/Source/Editor/Sequencer/Private/KeyPropertyParams.cpp
+++ b/Engine/Source/Editor/Sequencer/Private/KeyPropertyParams.cpp
@@ -89,10 +89,10 @@ FKeyPropertyParams::FKeyPropertyParams(TArray<UObject*> InObjectsToKey, const IP
 {
 }
 
-FPropertyChangedParams::FPropertyChangedParams(TArray<UObject*> InObjectsThatChanged, const FPropertyPath& InPropertyPath, FName InStructPropertyNameToKey, ESequencerKeyMode InKeyMode)
+FPropertyChangedParams::FPropertyChangedParams(TArray<UObject*> InObjectsThatChanged, const FPropertyPath& InPropertyPath, const FPropertyPath& InStructPathToKey, ESequencerKeyMode InKeyMode)
 	: ObjectsThatChanged(InObjectsThatChanged)
 	, PropertyPath(InPropertyPath)
-	, StructPropertyNameToKey(InStructPropertyNameToKey)
+	, StructPathToKey(InStructPathToKey)
 	, KeyMode(InKeyMode)
 {
 }
diff --git a/Engine/Source/Editor/Sequencer/Private/LevelEditorSequencerIntegration.cpp b/Engine/Source/Editor/Sequencer/Private/LevelEditorSequencerIntegration.cpp
--- a/Engine/Source/Editor/Sequencer/Private/LevelEditorSequencerIntegration.cpp
+++ b/Engine/Source/Editor/Sequencer/Private/LevelEditorSequencerIntegration.cpp
@@ -88,6 +88,31 @@ public:
 		return false;
 	}
 
+	virtual bool IsPropertyAnimated(const IPropertyHandle& PropertyHandle, UObject *ParentObject) const
+	{
+		
+		for (const TWeakPtr<ISequencer>& WeakSequencer : Sequencers)
+		{
+			TSharedPtr<ISequencer> Sequencer = WeakSequencer.Pin();
+			if (Sequencer.IsValid() && Sequencer->GetFocusedMovieSceneSequence())
+			{
+				FGuid ObjectHandle = Sequencer->GetHandleToObject(ParentObject);
+				if (ObjectHandle.IsValid()) 
+				{
+					UMovieScene* MovieScene = Sequencer->GetFocusedMovieSceneSequence()->GetMovieScene();
+					UProperty* Property = PropertyHandle.GetProperty();
+					TSharedRef<FPropertyPath> PropertyPath = FPropertyPath::CreateEmpty();
+					PropertyPath->AddProperty(FPropertyInfo(Property));
+					FName PropertyName(*PropertyPath->ToString(TEXT(".")));
+					TSubclassOf<UMovieSceneTrack> TrackClass; //use empty @todo find way to get the UMovieSceneTrack from the Property type.
+					return MovieScene->FindTrack(TrackClass, ObjectHandle, PropertyName) != nullptr;
+				}
+				
+				return false;
+			}
+		}
+		return false;
+	}
 	virtual void OnKeyPropertyClicked(const IPropertyHandle& KeyedPropertyHandle)
 	{
 		TArray<UObject*> Objects;
@@ -184,12 +209,6 @@ void FLevelEditorSequencerIntegration::Initialize()
 		AcquiredResources.Add([=]{ USelection::SelectionChangedEvent.Remove(Handle); });
 	}
 
-	{
-		// Hook into the editor's mechanism for checking whether we need live capture of PIE/SIE actor state
-		FDelegateHandle Handle = GEditor->GetActorRecordingState().AddRaw(this, &FLevelEditorSequencerIntegration::GetActorRecordingState);
-		AcquiredResources.Add([=]{ GEditor->GetActorRecordingState().Remove(Handle); });
-	}
-
 	{
 		FDelegateHandle Handle = FCoreDelegates::OnActorLabelChanged.AddRaw(this, &FLevelEditorSequencerIntegration::OnActorLabelChanged);
 		AcquiredResources.Add([=]{ FCoreDelegates::OnActorLabelChanged.Remove(Handle); });
@@ -233,19 +252,6 @@ void FLevelEditorSequencerIntegration::Initialize()
 	}
 }
 
-void FLevelEditorSequencerIntegration::GetActorRecordingState( bool& bIsRecording ) const
-{
-	IterateAllSequencers(
-		[&](FSequencer& In, const FLevelEditorSequencerIntegrationOptions& Options)
-		{
-			if (In.IsRecordingLive())
-			{
-				bIsRecording = true;
-			}
-		}
-	);
-}
-
 void RenameSpawnableRecursive(FSequencer* Sequencer, FMovieSceneSequenceIDRef SequenceID, AActor* ChangedActor)
 {
 	FMovieSceneRootEvaluationTemplateInstance& RootInstance = Sequencer->GetEvaluationTemplate();
@@ -998,7 +1004,7 @@ public:
 		TSharedPtr<FSequencer> Sequencer = WeakSequencer.Pin();
 		if (Sequencer.IsValid())
 		{
-			Sequencer->GetSettings()->SetShowViewportTransportControls(!Sequencer->GetSettings()->GetShowViewportTransportControls());
+			Sequencer->GetSequencerSettings()->SetShowViewportTransportControls(!Sequencer->GetSequencerSettings()->GetShowViewportTransportControls());
 		}
 		return FReply::Handled();
 	}
@@ -1012,7 +1018,7 @@ public:
 		return (
 				Sequencer.IsValid() &&
 				ViewportClient &&
-				Sequencer->GetSettings()->GetShowViewportTransportControls() &&
+				Sequencer->GetSequencerSettings()->GetShowViewportTransportControls() &&
 				ViewportClient->ViewportType == LVT_Perspective &&
 				ViewportClient->AllowsCinematicPreview()
 			) ? EVisibility::Visible : EVisibility::Collapsed;
@@ -1078,7 +1084,7 @@ void FLevelEditorSequencerIntegration::SetViewportTransportControlsVisibility(bo
 	IterateAllSequencers(
 		[=](FSequencer& In, const FLevelEditorSequencerIntegrationOptions& Options)
 		{
-			In.GetSettings()->SetShowViewportTransportControls(bVisible);
+			In.GetSequencerSettings()->SetShowViewportTransportControls(bVisible);
 		}
 	);
 }
@@ -1090,7 +1096,7 @@ bool FLevelEditorSequencerIntegration::GetViewportTransportControlsVisibility()
 	IterateAllSequencers(
 		[&](FSequencer& In, const FLevelEditorSequencerIntegrationOptions& Options)
 		{
-			if (In.GetSettings()->GetShowViewportTransportControls())
+			if (In.GetSequencerSettings()->GetShowViewportTransportControls())
 			{
 				bVisible = true;
 			}
@@ -1173,7 +1179,7 @@ void FLevelEditorSequencerIntegration::ActivateRealtimeViewports()
 		TSharedPtr<FSequencer> Pinned = SequencerAndOptions.Sequencer.Pin();
 		if (Pinned.IsValid())
 		{
-			if (!Pinned.Get()->GetSettings()->ShouldActivateRealtimeViewports())
+			if (!Pinned.Get()->GetSequencerSettings()->ShouldActivateRealtimeViewports())
 			{
 				return;
 			}
diff --git a/Engine/Source/Editor/Sequencer/Private/MovieSceneTrackEditor.cpp b/Engine/Source/Editor/Sequencer/Private/MovieSceneTrackEditor.cpp
--- a/Engine/Source/Editor/Sequencer/Private/MovieSceneTrackEditor.cpp
+++ b/Engine/Source/Editor/Sequencer/Private/MovieSceneTrackEditor.cpp
@@ -29,10 +29,10 @@ UMovieSceneSequence* FMovieSceneTrackEditor::GetMovieSceneSequence() const
 	return Sequencer.Pin()->GetFocusedMovieSceneSequence();
 }
 
-float FMovieSceneTrackEditor::GetTimeForKey()
+FFrameNumber FMovieSceneTrackEditor::GetTimeForKey()
 { 
 	TSharedPtr<ISequencer> SequencerPin = Sequencer.Pin();
-	return SequencerPin.IsValid() ? SequencerPin->GetLocalTime() : 0.f;
+	return SequencerPin.IsValid() ? SequencerPin->GetLocalTime().Time.FrameNumber : FFrameNumber(0);
 }
 
 void FMovieSceneTrackEditor::UpdatePlaybackRange()
@@ -44,6 +44,11 @@ void FMovieSceneTrackEditor::UpdatePlaybackRange()
 	}
 }
 
+TSharedRef<ISequencerSection> FMovieSceneTrackEditor::MakeSectionInterface(UMovieSceneSection& SectionObject, UMovieSceneTrack& Track, FGuid ObjectBinding)
+{
+	return MakeShared<FSequencerSection>(SectionObject);
+}
+
 void FMovieSceneTrackEditor::AnimatablePropertyChanged( FOnKeyProperty OnKeyProperty )
 {
 	check(OnKeyProperty.IsBound());
@@ -52,40 +57,37 @@ void FMovieSceneTrackEditor::AnimatablePropertyChanged( FOnKeyProperty OnKeyProp
 	UMovieSceneSequence* MovieSceneSequence = GetMovieSceneSequence();
 	if (MovieSceneSequence)
 	{
-		float KeyTime = GetTimeForKey();
+		FFrameNumber KeyTime = GetTimeForKey();
+
+		// @todo Sequencer - The sequencer probably should have taken care of this
+		MovieSceneSequence->SetFlags(RF_Transactional);
+	
+		// Create a transaction record because we are about to add keys
+		const bool bShouldActuallyTransact = !GIsTransacting;		// Don't transact if we're recording in a PIE world.  That type of keyframe capture cannot be undone.
+		FScopedTransaction AutoKeyTransaction( NSLOCTEXT("AnimatablePropertyTool", "PropertyChanged", "Animatable Property Changed"), bShouldActuallyTransact );
+
+		FKeyPropertyResult KeyPropertyResult = OnKeyProperty.Execute( KeyTime );
+
+		if (KeyPropertyResult.bTrackCreated)
+		{
+			// If a track is created evaluate immediately so that the pre-animated state can be stored
+			Sequencer.Pin()->NotifyMovieSceneDataChanged( EMovieSceneDataChangeType::RefreshAllImmediately );
+		}
+		else if (KeyPropertyResult.bTrackModified)
+		{
+			Sequencer.Pin()->NotifyMovieSceneDataChanged( EMovieSceneDataChangeType::MovieSceneStructureItemAdded );
+		}
+		else if (KeyPropertyResult.bKeyCreated)
+		{
+			Sequencer.Pin()->NotifyMovieSceneDataChanged( EMovieSceneDataChangeType::TrackValueChanged );
+		}
+
+		UpdatePlaybackRange();
 
-		if( !Sequencer.Pin()->IsRecordingLive() )
+		TSharedPtr<FSequencer> SequencerToUpdate = StaticCastSharedPtr<FSequencer>(Sequencer.Pin());
+		if (SequencerToUpdate.IsValid())
 		{
-			// @todo Sequencer - The sequencer probably should have taken care of this
-			MovieSceneSequence->SetFlags(RF_Transactional);
-		
-			// Create a transaction record because we are about to add keys
-			const bool bShouldActuallyTransact = !GIsTransacting && !Sequencer.Pin()->IsRecordingLive();		// Don't transact if we're recording in a PIE world.  That type of keyframe capture cannot be undone.
-			FScopedTransaction AutoKeyTransaction( NSLOCTEXT("AnimatablePropertyTool", "PropertyChanged", "Animatable Property Changed"), bShouldActuallyTransact );
-
-			FKeyPropertyResult KeyPropertyResult = OnKeyProperty.Execute( KeyTime );
-
-			if (KeyPropertyResult.bTrackCreated)
-			{
-				// If a track is created evaluate immediately so that the pre-animated state can be stored
-				Sequencer.Pin()->NotifyMovieSceneDataChanged( EMovieSceneDataChangeType::RefreshAllImmediately );
-			}
-			else if (KeyPropertyResult.bTrackModified)
-			{
-				Sequencer.Pin()->NotifyMovieSceneDataChanged( EMovieSceneDataChangeType::MovieSceneStructureItemAdded );
-			}
-			else if (KeyPropertyResult.bKeyCreated)
-			{
-				Sequencer.Pin()->NotifyMovieSceneDataChanged( EMovieSceneDataChangeType::TrackValueChanged );
-			}
-
-			UpdatePlaybackRange();
-
-			TSharedPtr<FSequencer> SequencerToUpdate = StaticCastSharedPtr<FSequencer>(Sequencer.Pin());
-			if (SequencerToUpdate.IsValid())
-			{
-				SequencerToUpdate->SynchronizeSequencerSelectionWithExternalSelection();
-			}
+			SequencerToUpdate->SynchronizeSequencerSelectionWithExternalSelection();
 		}
 	}
 }
@@ -190,12 +192,12 @@ bool FMovieSceneTrackEditor::HandleAssetAdded(UObject* Asset, const FGuid& Targe
 	return false; 
 }
 
-bool FMovieSceneTrackEditor::OnAllowDrop(const FDragDropEvent& DragDropEvent, UMovieSceneTrack* Track)
+bool FMovieSceneTrackEditor::OnAllowDrop(const FDragDropEvent& DragDropEvent, UMovieSceneTrack* Track, int32 RowIndex, const FGuid& TargetObjectGuid)
 {
 	return false;
 }
 
-FReply FMovieSceneTrackEditor::OnDrop(const FDragDropEvent& DragDropEvent, UMovieSceneTrack* Track)
+FReply FMovieSceneTrackEditor::OnDrop(const FDragDropEvent& DragDropEvent, UMovieSceneTrack* Track, int32 RowIndex, const FGuid& TargetObjectGuid)
 {
 	return FReply::Unhandled();
 }
diff --git a/Engine/Source/Editor/Sequencer/Private/QualifiedFrameTimeDetailsCustomization.cpp b/Engine/Source/Editor/Sequencer/Private/QualifiedFrameTimeDetailsCustomization.cpp
new file mode 100644
--- /dev/null
+++ b/Engine/Source/Editor/Sequencer/Private/QualifiedFrameTimeDetailsCustomization.cpp
@@ -0,0 +1,80 @@
+// Copyright 1998-2018 Epic Games, Inc. All Rights Reserved.
+
+#include "QualifiedFrameTimeDetailsCustomization.h"
+#include "IDetailPropertyRow.h"
+#include "FrameNumber.h"
+#include "IDetailChildrenBuilder.h"
+#include "DetailLayoutBuilder.h"
+#include "Widgets/Text/STextBlock.h"
+#include "Widgets/Input/SEditableTextBox.h"
+#include "PropertyCustomizationHelpers.h"
+#include "FrameRate.h"
+
+void FQualifiedFrameTimeDetailsCustomization::CustomizeChildren(TSharedRef<IPropertyHandle> PropertyHandle, IDetailChildrenBuilder& ChildBuilder, IPropertyTypeCustomizationUtils& CustomizationUtils)
+{
+	TMap<FName, TSharedRef<IPropertyHandle>> CustomizedProperties;
+
+	uint32 NumChildren = 0;
+	PropertyHandle->GetNumChildren(NumChildren);
+
+	// Add child properties to UI and pick out the properties which need customization
+	for (uint32 ChildIndex = 0; ChildIndex < NumChildren; ++ChildIndex)
+	{
+		TSharedRef<IPropertyHandle> ChildHandle = PropertyHandle->GetChildHandle(ChildIndex).ToSharedRef();
+		CustomizedProperties.Add(ChildHandle->GetProperty()->GetFName(), ChildHandle);
+	}
+
+	TSharedRef<IPropertyHandle> FrameNumberProperty = CustomizedProperties.FindChecked(GET_MEMBER_NAME_CHECKED(FFrameNumber, Value));
+	ChildBuilder.AddCustomRow(NSLOCTEXT("TimeManagement", "TimeLabel", "Time"))
+		.NameContent()
+		[
+			SNew(STextBlock)
+			.Text(NSLOCTEXT("TimeManagement", "TimeLabel", "Time"))
+			.ToolTipText(NSLOCTEXT("TimeManagement", "TimeLabelTooltip", "Time field which takes timecode, frames and time formats."))
+			.Font(CustomizationUtils.GetRegularFont())
+		]
+		.ValueContent()
+		[
+			SNew(SEditableTextBox)
+			.Text(this, &FQualifiedFrameTimeDetailsCustomization::OnGetTimeText)
+			.OnTextCommitted(this, &FQualifiedFrameTimeDetailsCustomization::OnTimeTextCommitted)
+			.SelectAllTextWhenFocused(true)
+			.RevertTextOnEscape(true)
+			.Font(IDetailLayoutBuilder::GetDetailFont())
+		];
+}
+
+FText FQualifiedFrameTimeDetailsCustomization::OnGetTimeText() const
+{
+	FString Foo = TEXT("Time");
+	return FText::FromString(Foo);
+}
+
+void FQualifiedFrameTimeDetailsCustomization::OnTimeTextCommitted(const FText& InText, ETextCommit::Type CommitInfo)
+{
+	// @todo sequencer-timecode unified time display
+
+	// FSequencerTimeEvaluator Eval;
+	// TValueOrError<FFrameTime, FExpressionError> TimecodeResult = Eval.EvaluateTimecode(*InText.ToString(), FFrameRate(30, 1));
+	// bool bWasFrameText;
+	// TValueOrError<FFrameTime, FExpressionError> FrameResult = Eval.EvaluateFrame(*InText.ToString(), FFrameRate(30, 1), /*Out*/ bWasFrameText);
+	// bool bWasTimeText;
+	// TValueOrError<FFrameTime, FExpressionError> TimeResult = Eval.EvaluateTime(*InText.ToString(), FFrameRate(30,1), /*Out*/ bWasTimeText);
+	// 
+	// if (TimecodeResult.IsValid())
+	// {
+	// 	UE_LOG(LogTemp, Log, TEXT("Timecode Result: %d"), TimecodeResult.GetValue().GetFrame().Value);
+	// }
+	// else if (TimeResult.IsValid() && !bWasFrameText)
+	// {
+	// 	UE_LOG(LogTemp, Log, TEXT("Time Result: %d"), TimeResult.GetValue().GetFrame().Value);
+	// }
+	// else if (FrameResult.IsValid() && !bWasTimeText)
+	// {
+	// 	UE_LOG(LogTemp, Log, TEXT("Frame Result: %d"), FrameResult.GetValue().GetFrame().Value);
+	// }
+	// else
+	// {
+	// 	UE_LOG(LogTemp, Log, TEXT("Unknown"));
+	// }
+}
diff --git a/Engine/Source/Editor/Sequencer/Private/QualifiedFrameTimeDetailsCustomization.h b/Engine/Source/Editor/Sequencer/Private/QualifiedFrameTimeDetailsCustomization.h
new file mode 100644
--- /dev/null
+++ b/Engine/Source/Editor/Sequencer/Private/QualifiedFrameTimeDetailsCustomization.h
@@ -0,0 +1,31 @@
+// Copyright 1998-2018 Epic Games, Inc. All Rights Reserved.
+
+#pragma once
+
+#include "CoreMinimal.h"
+#include "IPropertyTypeCustomization.h"
+
+class IDetailLayoutBuilder;
+class ISequencer;
+
+/**
+ *  Customize the qualified frame time to support conversion from seconds/frames/timecode formats.
+ */
+class FQualifiedFrameTimeDetailsCustomization : public IPropertyTypeCustomization
+{
+public:
+	FQualifiedFrameTimeDetailsCustomization(TWeakPtr<ISequencer> InSequencer)
+	{
+		Sequencer = InSequencer;
+	}
+
+	/** IPropertyTypeCustomization interface */
+	virtual void CustomizeHeader(TSharedRef<IPropertyHandle> PropertyHandle, FDetailWidgetRow& HeaderRow, IPropertyTypeCustomizationUtils& CustomizationUtils) override {};
+	virtual void CustomizeChildren(TSharedRef<IPropertyHandle> PropertyHandle, IDetailChildrenBuilder& ChildBuilder, IPropertyTypeCustomizationUtils& CustomizationUtils) override;
+
+private:
+	FText OnGetTimeText() const;
+	void OnTimeTextCommitted(const FText& InText, ETextCommit::Type CommitInfo);
+
+	TWeakPtr<ISequencer> Sequencer;
+};
diff --git a/Engine/Source/Editor/Sequencer/Private/SKeyAreaEditorSwitcher.cpp b/Engine/Source/Editor/Sequencer/Private/SKeyAreaEditorSwitcher.cpp
--- a/Engine/Source/Editor/Sequencer/Private/SKeyAreaEditorSwitcher.cpp
+++ b/Engine/Source/Editor/Sequencer/Private/SKeyAreaEditorSwitcher.cpp
@@ -51,6 +51,9 @@ void SKeyAreaEditorSwitcher::Rebuild()
 		return;
 	}
 
+	TSharedPtr<FSequencerObjectBindingNode> ParentObjectBinding = KeyAreaNode->FindParentObjectBindingNode();
+	FGuid ObjectBindingID = ParentObjectBinding.IsValid() ? ParentObjectBinding->GetObjectBinding() : FGuid();
+
 	for (int32 Index = 0; Index < KeyAreaNode->GetAllKeyAreas().Num(); ++Index)
 	{
 		TSharedRef<IKeyArea> KeyArea = KeyAreaNode->GetAllKeyAreas()[Index];
@@ -65,7 +68,7 @@ void SKeyAreaEditorSwitcher::Rebuild()
 				.HAlign(HAlign_Left)
 				.Visibility(this, &SKeyAreaEditorSwitcher::GetWidgetVisibility, Index)
 				[
-					KeyArea->CreateKeyEditor(&KeyAreaNode->GetSequencer())
+					KeyArea->CreateKeyEditor(KeyAreaNode->GetSequencer().AsShared(), ObjectBindingID)
 				]
 			];
 		}
@@ -93,5 +96,5 @@ void SKeyAreaEditorSwitcher::Tick( const FGeometry& AllottedGeometry, const doub
 		AllSections.Add(KeyArea->GetOwningSection());
 	}
 
-	VisibleIndex = SequencerHelpers::GetSectionFromTime(AllSections, KeyAreaNode->GetSequencer().GetLocalTime());
+	VisibleIndex = SequencerHelpers::GetSectionFromTime(AllSections, KeyAreaNode->GetSequencer().GetLocalTime().Time.FrameNumber);
 }
diff --git a/Engine/Source/Editor/Sequencer/Private/SKeyNavigationButtons.h b/Engine/Source/Editor/Sequencer/Private/SKeyNavigationButtons.h
--- a/Engine/Source/Editor/Sequencer/Private/SKeyNavigationButtons.h
+++ b/Engine/Source/Editor/Sequencer/Private/SKeyNavigationButtons.h
@@ -6,7 +6,8 @@
 #include "Misc/Attribute.h"
 #include "Styling/SlateColor.h"
 #include "Input/Reply.h"
-#include "DisplayNodes/SequencerDisplayNode.h"
+#include "DisplayNodes/SequencerObjectBindingNode.h"
+#include "DisplayNodes/SequencerSectionKeyAreaNode.h"
 #include "Widgets/DeclarativeSyntaxSupport.h"
 #include "Widgets/SBoxPanel.h"
 #include "Widgets/SCompoundWidget.h"
@@ -16,10 +17,12 @@
 #include "Widgets/Input/SButton.h"
 #include "EditorStyleSet.h"
 #include "SequencerCommonHelpers.h"
-#include "IKeyArea.h"
 #include "MovieSceneCommonHelpers.h"
 #include "ScopedTransaction.h"
 #include "SequencerCommands.h"
+#include "ISequencerModule.h"
+#include "Modules/ModuleManager.h"
+#include "IKeyArea.h"
 
 #define LOCTEXT_NAMESPACE "SKeyNavigationButtons"
 
@@ -135,48 +138,50 @@ public:
 	FReply OnPreviousKeyClicked()
 	{
 		FSequencer& Sequencer = DisplayNode->GetSequencer();
-		float ClosestPreviousKeyDistance = MAX_FLT;
-		float CurrentTime = Sequencer.GetLocalTime();
-		float PreviousTime = 0;
-		bool PreviousKeyFound = false;
 
-		TArray<float> AllTimes;
+		FFrameTime ClosestPreviousKeyDistance = FFrameTime(TNumericLimits<int32>::Max(), 0.99999f);
+		FFrameTime CurrentTime = Sequencer.GetLocalTime().Time;
+		TOptional<FFrameTime> PreviousTime;
+
+		TArray<FFrameNumber> AllTimes;
 
 		TSet<TSharedPtr<IKeyArea>> KeyAreas;
 		SequencerHelpers::GetAllKeyAreas( DisplayNode, KeyAreas );
-		for ( TSharedPtr<IKeyArea> KeyArea : KeyAreas )
+
+		for ( TSharedPtr<IKeyArea> Keyarea : KeyAreas )
 		{
-			for ( FKeyHandle& KeyHandle : KeyArea->GetUnsortedKeyHandles() )
-			{
-				float KeyTime = KeyArea->GetKeyTime( KeyHandle );
-				AllTimes.Add(KeyTime);
-			}
+			Keyarea->GetKeyTimes(AllTimes);
 		}
 
 		TSet<TWeakObjectPtr<UMovieSceneSection> > Sections;
 		SequencerHelpers::GetAllSections( DisplayNode.ToSharedRef(), Sections );
 		for ( TWeakObjectPtr<UMovieSceneSection> Section : Sections )
 		{
-			if (Section.IsValid() && !Section->IsInfinite())
+			if (Section.IsValid())
 			{
-				AllTimes.Add(Section->GetStartTime());
-				AllTimes.Add(Section->GetEndTime());
+				if (Section->HasStartFrame())
+				{
+					AllTimes.Add(Section->GetInclusiveStartFrame());
+				}
+				if (Section->HasEndFrame())
+				{
+					AllTimes.Add(Section->GetExclusiveEndFrame()-1);
+				}
 			}
 		}
 
-		for (float Time : AllTimes)
+		for (FFrameNumber Time : AllTimes)
 		{
 			if (Time < CurrentTime && CurrentTime - Time < ClosestPreviousKeyDistance)
 			{
 				PreviousTime = Time;
 				ClosestPreviousKeyDistance = CurrentTime - Time;
-				PreviousKeyFound = true;
 			}
 		}
 
-		if (PreviousKeyFound)
+		if (PreviousTime.IsSet())
 		{
-			Sequencer.SetLocalTime(PreviousTime);
+			Sequencer.SetLocalTime(PreviousTime.GetValue());
 		}
 		return FReply::Handled();
 	}
@@ -185,48 +190,49 @@ public:
 	FReply OnNextKeyClicked()
 	{
 		FSequencer& Sequencer = DisplayNode->GetSequencer();
-		float ClosestNextKeyDistance = MAX_FLT;
-		float CurrentTime = Sequencer.GetLocalTime();
-		float NextTime = 0;
-		bool NextKeyFound = false;
+		FFrameTime ClosestNextKeyDistance = FFrameTime(TNumericLimits<int32>::Max(), 0.99999f);
+		FFrameTime CurrentTime = Sequencer.GetLocalTime().Time;
+		TOptional<FFrameTime> NextTime;
 
-		TArray<float> AllTimes;
+		TArray<FFrameNumber> AllTimes;
 
 		TSet<TSharedPtr<IKeyArea>> KeyAreas;
 		SequencerHelpers::GetAllKeyAreas( DisplayNode, KeyAreas );
-		for ( TSharedPtr<IKeyArea> KeyArea : KeyAreas )
+
+		for ( TSharedPtr<IKeyArea> Keyarea : KeyAreas )
 		{
-			for ( FKeyHandle& KeyHandle : KeyArea->GetUnsortedKeyHandles() )
-			{
-				float KeyTime = KeyArea->GetKeyTime( KeyHandle );
-				AllTimes.Add(KeyTime);
-			}
+			Keyarea->GetKeyTimes(AllTimes);
 		}
 
 		TSet<TWeakObjectPtr<UMovieSceneSection> > Sections;
 		SequencerHelpers::GetAllSections( DisplayNode.ToSharedRef(), Sections );
 		for ( TWeakObjectPtr<UMovieSceneSection> Section : Sections )
 		{
-			if (Section.IsValid() && !Section->IsInfinite())
+			if (Section.IsValid())
 			{
-				AllTimes.Add(Section->GetStartTime());
-				AllTimes.Add(Section->GetEndTime());
+				if (Section->HasStartFrame())
+				{
+					AllTimes.Add(Section->GetInclusiveStartFrame());
+				}
+				if (Section->HasEndFrame())
+				{
+					AllTimes.Add(Section->GetExclusiveEndFrame()-1);
+				}
 			}
 		}
 
-		for (float Time : AllTimes)
+		for (FFrameNumber Time : AllTimes)
 		{
 			if (Time > CurrentTime && Time - CurrentTime < ClosestNextKeyDistance)
 			{
 				NextTime = Time;
 				ClosestNextKeyDistance = Time - CurrentTime;
-				NextKeyFound = true;
 			}
 		}
 
-		if (NextKeyFound)
+		if (NextTime.IsSet())
 		{
-			Sequencer.SetLocalTime(NextTime);
+			Sequencer.SetLocalTime(NextTime.GetValue());
 		}
 
 		return FReply::Handled();
@@ -236,42 +242,61 @@ public:
 	FReply OnAddKeyClicked()
 	{
 		FSequencer& Sequencer = DisplayNode->GetSequencer();
-		float CurrentTime = Sequencer.GetLocalTime();
+		FFrameTime CurrentTime = Sequencer.GetLocalTime().Time;
 
-		TSet<TSharedPtr<IKeyArea>> KeyAreas;
-		SequencerHelpers::GetAllKeyAreas(DisplayNode, KeyAreas);
+		// Gather all sections on this node so we can decide which one to key
+		TSet<TWeakObjectPtr<UMovieSceneSection> > WeakSections;
+		SequencerHelpers::GetAllSections( DisplayNode.ToSharedRef(), WeakSections );
 
-		// Prune out any key areas that do not exist at the current time, or are overlapped
-		TMap<FName, TArray<TSharedPtr<IKeyArea>>> NameToKeyAreas;
-		for (TSharedPtr<IKeyArea> KeyArea : KeyAreas)
+		TArray<UMovieSceneSection*> SectionArray;
+		SectionArray.Reserve(WeakSections.Num());
+		for (TWeakObjectPtr<UMovieSceneSection> WeakSection : WeakSections)
 		{
-			NameToKeyAreas.FindOrAdd(KeyArea->GetName()).Add(KeyArea);
+			if (UMovieSceneSection* Section = WeakSection.Get())
+			{
+				SectionArray.Add(Section);
+			}
+		}
+
+		// Add keys specifically only on the closest or overlapping section
+		const int32 SectionIndex = SequencerHelpers::GetSectionFromTime(SectionArray, CurrentTime.FrameNumber);
+		if (SectionIndex == INDEX_NONE)
+		{
+			return FReply::Handled();
 		}
 
 		FScopedTransaction Transaction(LOCTEXT("AddKeys", "Add Keys at Current Time"));
-		for (auto& Pair : NameToKeyAreas)
+
+		// Add the section to the transaction
+		UMovieSceneSection* SectionToKey = SectionArray[SectionIndex];
+		SectionToKey->SetFlags(RF_Transactional);
+		if (!SectionToKey->TryModify())
 		{
-			TArray<UMovieSceneSection*> AllSections;
-			for (TSharedPtr<IKeyArea>& KeyArea : Pair.Value)
-			{
-				AllSections.Add(KeyArea->GetOwningSection());
-			}
+			return FReply::Handled();
+		}
 
-			int32 KeyAreaIndex = SequencerHelpers::GetSectionFromTime(AllSections, CurrentTime);
+		TSharedPtr<FSequencerObjectBindingNode> ParentObjectBinding = DisplayNode->FindParentObjectBindingNode();
+		FGuid ObjectBinding = ParentObjectBinding.IsValid() ? ParentObjectBinding->GetObjectBinding() : FGuid();
 
-			if (KeyAreaIndex != INDEX_NONE)
-			{
-				UMovieSceneSection* OwningSection = AllSections[KeyAreaIndex];
+		TArray<TSharedRef<FSequencerSectionKeyAreaNode>> KeyAreaNodes;
+		if (DisplayNode->GetType() == ESequencerNode::KeyArea)
+		{
+			KeyAreaNodes.Add(StaticCastSharedPtr<FSequencerSectionKeyAreaNode>(DisplayNode).ToSharedRef());
+		}
+		DisplayNode->GetChildKeyAreaNodesRecursively(KeyAreaNodes);
 
-				OwningSection->SetFlags(RF_Transactional);
-				if (OwningSection->TryModify())
-				{
-					Pair.Value[KeyAreaIndex]->AddKeyUnique(CurrentTime, Sequencer.GetKeyInterpolation());
-					Sequencer.NotifyMovieSceneDataChanged(EMovieSceneDataChangeType::TrackValueChanged);
-				}
+		for (TSharedRef<FSequencerSectionKeyAreaNode> KeyAreaNode : KeyAreaNodes)
+		{
+			TSharedPtr<IKeyArea> KeyArea =  KeyAreaNode->GetKeyArea(SectionToKey);
+			if (KeyArea.IsValid())
+			{
+				KeyArea->AddOrUpdateKey(CurrentTime.FrameNumber, ObjectBinding, Sequencer);
 			}
 		}
 
+		SectionToKey->ExpandToFrame(CurrentTime.FrameNumber);
+
+		Sequencer.NotifyMovieSceneDataChanged(EMovieSceneDataChangeType::TrackValueChanged);
 		Sequencer.UpdatePlaybackRange();
 
 		return FReply::Handled();
diff --git a/Engine/Source/Editor/Sequencer/Private/SSequencer.cpp b/Engine/Source/Editor/Sequencer/Private/SSequencer.cpp
--- a/Engine/Source/Editor/Sequencer/Private/SSequencer.cpp
+++ b/Engine/Source/Editor/Sequencer/Private/SSequencer.cpp
@@ -33,7 +33,6 @@
 #include "DisplayNodes/SequencerTrackNode.h"
 #include "Widgets/Input/SNumericDropDown.h"
 #include "SequencerCommonHelpers.h"
-#include "SSequencerCurveEditor.h"
 #include "SSequencerCurveEditorToolBar.h"
 #include "SSequencerLabelBrowser.h"
 #include "ISequencerWidgetsModule.h"
@@ -50,6 +49,7 @@
 #include "MovieSceneSequence.h"
 #include "SSequencerSplitterOverlay.h"
 #include "SequencerHotspots.h"
+#include "SSequencerTimePanel.h"
 #include "VirtualTrackArea.h"
 #include "Framework/Commands/GenericCommands.h"
 #include "SequencerContextMenus.h"
@@ -63,65 +63,14 @@
 #include "IVREditorModule.h"
 #include "EditorFontGlyphs.h"
 #include "HAL/PlatformApplicationMisc.h"
+#include "SSequencerPlayRateCombo.h"
 #include "Camera/CameraActor.h"
+#include "SCurveEditorPanel.h"
+#include "MovieSceneTimeHelpers.h"
+#include "FrameNumberNumericInterface.h"
 
 #define LOCTEXT_NAMESPACE "Sequencer"
 
-DECLARE_DELEGATE_RetVal(bool, FOnGetShowFrames)
-DECLARE_DELEGATE_RetVal(uint8, FOnGetZeroPad)
-
-
-/* Numeric type interface for showing numbers as frames or times */
-struct FFramesOrTimeInterface : public TNumericUnitTypeInterface<float>
-{
-	FFramesOrTimeInterface(FOnGetShowFrames InShowFrameNumbers,TSharedPtr<FSequencerTimeSliderController> InController, FOnGetZeroPad InOnGetZeroPad)
-		: TNumericUnitTypeInterface(EUnit::Seconds)
-		, ShowFrameNumbers(MoveTemp(InShowFrameNumbers))
-		, TimeSliderController(MoveTemp(InController))
-		, OnGetZeroPad(MoveTemp(InOnGetZeroPad))
-	{}
-
-private:
-	FOnGetShowFrames ShowFrameNumbers;
-	TSharedPtr<FSequencerTimeSliderController> TimeSliderController;
-	FOnGetZeroPad OnGetZeroPad;
-
-	virtual FString ToString(const float& Value) const override
-	{
-		if (ShowFrameNumbers.Execute())
-		{
-			int32 Frame = TimeSliderController->TimeToFrame(Value);
-			if (OnGetZeroPad.IsBound())
-			{
-				return FString::Printf(*FString::Printf(TEXT("%%0%dd"), OnGetZeroPad.Execute()), Frame);
-			}
-			return FString::Printf(TEXT("%d"), Frame);
-		}
-
-		return FString::Printf(TEXT("%.2fs"), Value);
-	}
-
-	virtual TOptional<float> FromString(const FString& InString, const float& InExistingValue) override
-	{
-		bool bShowFrameNumbers = ShowFrameNumbers.IsBound() ? ShowFrameNumbers.Execute() : false;
-		if (bShowFrameNumbers)
-		{
-			// Convert existing value to frames
-			float ExistingValueInFrames = TimeSliderController->TimeToFrame(InExistingValue);
-			TOptional<float> Result = TNumericUnitTypeInterface<float>::FromString(InString, ExistingValueInFrames);
-
-			if (Result.IsSet())
-			{
-				int32 NewEndFrame = FMath::RoundToInt(Result.GetValue());
-				return float(TimeSliderController->FrameToTime(NewEndFrame));
-			}
-		}
-
-		return TNumericUnitTypeInterface::FromString(InString, InExistingValue);
-	}
-};
-
-
 /* SSequencer interface
  *****************************************************************************/
 PRAGMA_DISABLE_OPTIMIZATION
@@ -130,11 +79,10 @@ void SSequencer::Construct(const FArguments& InArgs, TSharedRef<FSequencer> InSe
 	SequencerPtr = InSequencer;
 	bIsActiveTimerRegistered = false;
 	bUserIsSelecting = false;
-	CachedClampRange = TRange<float>::Empty();
-	CachedViewRange = TRange<float>::Empty();
+	CachedClampRange = TRange<double>::Empty();
+	CachedViewRange = TRange<double>::Empty();
 
-	Settings = InSequencer->GetSettings();
-	InSequencer->OnActivateSequence().AddSP(this, &SSequencer::OnSequenceInstanceActivated);
+	Settings = InSequencer->GetSequencerSettings();
 
 	ISequencerWidgetsModule& SequencerWidgets = FModuleManager::Get().LoadModuleChecked<ISequencerWidgetsModule>( "SequencerWidgets" );
 
@@ -145,11 +93,39 @@ void SSequencer::Construct(const FArguments& InArgs, TSharedRef<FSequencer> InSe
 
 	OnReceivedFocus = InArgs._OnReceivedFocus;
 
+	USequencerSettings* SequencerSettings = Settings;
+
+	// Get the desired display format from the user's settings each time.
+	auto GetDisplayFormatDelegate = FOnGetDisplayFormat::CreateLambda([SequencerSettings]() ->EFrameNumberDisplayFormats {
+		if (SequencerSettings)
+		{
+			return SequencerSettings->GetTimeDisplayFormat();
+		}
+		return EFrameNumberDisplayFormats::Frames;
+	});
+
+	// Get the number of zero pad frames from the user's settings as well.
+	auto GetZeroPadFramesDelegate = FOnGetZeroPad::CreateLambda([SequencerSettings]() -> uint8 {
+		if (SequencerSettings)
+		{
+			return SequencerSettings->GetZeroPadFrames();
+		}
+		return 0;
+	});
+
+	auto GetFrameResolutionDelegate = FOnGetFrameRate::CreateSP(InSequencer, &FSequencer::GetFocusedFrameResolution);
+	auto GetPlayRateDelegate = FOnGetFrameRate::CreateSP(InSequencer, &FSequencer::GetFocusedPlayRate);
+
+	// Create our numeric type interface so we can pass it to the time slider below.
+	NumericTypeInterface = MakeShareable(new FFrameNumberInterface(GetDisplayFormatDelegate, GetZeroPadFramesDelegate, GetFrameResolutionDelegate, GetPlayRateDelegate));
+
 	FTimeSliderArgs TimeSliderArgs;
 	{
 		TimeSliderArgs.ViewRange = InArgs._ViewRange;
 		TimeSliderArgs.ClampRange = InArgs._ClampRange;
 		TimeSliderArgs.PlaybackRange = InArgs._PlaybackRange;
+		TimeSliderArgs.PlayRate = TAttribute<FFrameRate>(InSequencer, &FSequencer::GetFocusedPlayRate);
+		TimeSliderArgs.FrameResolution = TAttribute<FFrameRate>(InSequencer, &FSequencer::GetFocusedFrameResolution);
 		TimeSliderArgs.SelectionRange = InArgs._SelectionRange;
 		TimeSliderArgs.OnPlaybackRangeChanged = InArgs._OnPlaybackRangeChanged;
 		TimeSliderArgs.OnPlaybackRangeBeginDrag = OnPlaybackRangeBeginDrag;
@@ -162,7 +138,6 @@ void SSequencer::Construct(const FArguments& InArgs, TSharedRef<FSequencer> InSe
 		TimeSliderArgs.OnGetNearestKey = InArgs._OnGetNearestKey;
 		TimeSliderArgs.IsPlaybackRangeLocked = InArgs._IsPlaybackRangeLocked;
 		TimeSliderArgs.OnTogglePlaybackRangeLocked = InArgs._OnTogglePlaybackRangeLocked;
-		TimeSliderArgs.TimeSnapInterval = InArgs._TimeSnapInterval;
 		TimeSliderArgs.ScrubPosition = InArgs._ScrubPosition;
 		TimeSliderArgs.OnBeginScrubberMovement = InArgs._OnBeginScrubbing;
 		TimeSliderArgs.OnEndScrubberMovement = InArgs._OnEndScrubbing;
@@ -171,27 +146,13 @@ void SSequencer::Construct(const FArguments& InArgs, TSharedRef<FSequencer> InSe
 		TimeSliderArgs.SubSequenceRange = InArgs._SubSequenceRange;
 
 		TimeSliderArgs.Settings = Settings;
+		TimeSliderArgs.NumericTypeInterface = GetNumericTypeInterface();
 	}
 
-	TimeSliderController = MakeShareable( new FSequencerTimeSliderController( TimeSliderArgs ) );
+	TimeSliderController = MakeShareable( new FSequencerTimeSliderController( TimeSliderArgs, SequencerPtr ) );
 	
 	TSharedRef<FSequencerTimeSliderController> TimeSliderControllerRef = TimeSliderController.ToSharedRef();
 
-	{
-		auto ShowFrameNumbersDelegate = FOnGetShowFrames::CreateSP(this, &SSequencer::ShowFrameNumbers);
-		USequencerSettings* SequencerSettings = Settings;
-		auto GetZeroPad = [=]() -> uint8 {
-			if (SequencerSettings)
-			{
-				return SequencerSettings->GetZeroPadFrames();
-			}
-			return 0;
-		};
-
-		NumericTypeInterface = MakeShareable( new FFramesOrTimeInterface(ShowFrameNumbersDelegate, TimeSliderControllerRef, FOnGetZeroPad()) );
-		ZeroPadNumericTypeInterface = MakeShareable( new FFramesOrTimeInterface(ShowFrameNumbersDelegate, TimeSliderControllerRef, FOnGetZeroPad::CreateLambda(GetZeroPad)) );
-	}
-
 	bool bMirrorLabels = false;
 	
 	// Create the top and bottom sliders
@@ -205,13 +166,13 @@ void SSequencer::Construct(const FArguments& InArgs, TSharedRef<FSequencer> InSe
 			EShowRange(EShowRange::WorkingRange | EShowRange::ViewRange),
 			TimeSliderControllerRef,
 			TAttribute<EVisibility>(this, &SSequencer::GetTimeRangeVisibility),
-			TAttribute<bool>(this, &SSequencer::ShowFrameNumbers),
-			ZeroPadNumericTypeInterface.ToSharedRef()
+			NumericTypeInterface.ToSharedRef()
 		),
-		SequencerWidgets.CreateTimeRangeSlider(TimeSliderControllerRef, TAttribute<float>(this, &SSequencer::OnGetTimeSnapInterval))
+		SequencerWidgets.CreateTimeRangeSlider(TimeSliderControllerRef)
 	);
 
 	OnGetAddMenuContent = InArgs._OnGetAddMenuContent;
+	OnBuildCustomContextMenuForGuid = InArgs._OnBuildCustomContextMenuForGuid;
 	AddMenuExtender = InArgs._AddMenuExtender;
 	ToolbarExtender = InArgs._ToolbarExtender;
 
@@ -234,20 +195,25 @@ void SSequencer::Construct(const FArguments& InArgs, TSharedRef<FSequencer> InSe
 		.Clipping(EWidgetClipping::ClipToBounds)
 		.OnGetContextMenuContent(FOnGetContextMenuContent::CreateSP(this, &SSequencer::GetContextMenuContent));
 
-	SAssignNew(CurveEditor, SSequencerCurveEditor, InSequencer, TimeSliderControllerRef)
-		.Visibility(this, &SSequencer::GetCurveEditorVisibility)
-		.OnViewRangeChanged(InArgs._OnViewRangeChanged)
-		.ViewRange(InArgs._ViewRange);
-
-	CurveEditor->SetAllowAutoFrame(SequencerPtr.Pin()->GetShowCurveEditor());
 	TrackArea->SetTreeView(TreeView);
 
+	TAttribute<FAnimatedRange> ViewRangeAttribute = InArgs._ViewRange;
+	TSharedRef<SCurveEditorPanel> CurveEditorPanel = 
+		SNew(SCurveEditorPanel, InSequencer->GetCurveEditor().ToSharedRef())
+		.Visibility(this, &SSequencer::GetCurveEditorVisibility)
+		// Grid lines match the color specified in FSequencerTimeSliderController::OnPaintSectionView
+		.GridLineTint(FLinearColor(0.f, 0.f, 0.f, 0.3f));
+
 	const int32 Column0 = 0, Column1 = 1;
 	const int32 Row0 = 0, Row1 = 1, Row2 = 2, Row3 = 3, Row4 = 4;
 
 	const float CommonPadding = 3.f;
 	const FMargin ResizeBarPadding(4.f, 0, 0, 0);
 
+	TSharedRef<FUICommandList> CurveEditorAndSequencerCommands = MakeShared<FUICommandList>();
+	CurveEditorAndSequencerCommands->Append(CurveEditorPanel->GetCommands().ToSharedRef());
+	CurveEditorAndSequencerCommands->Append(InSequencer->GetCommandBindings().ToSharedRef());
+
 	ChildSlot
 	[
 		SNew(SVerticalBox)
@@ -302,7 +268,7 @@ void SSequencer::Construct(const FArguments& InArgs, TSharedRef<FSequencer> InSe
 							+SHorizontalBox::Slot()
 							.AutoWidth()
 							[
-								SNew(SSequencerCurveEditorToolBar, InSequencer, CurveEditor->GetCommands())
+								SNew(SSequencerCurveEditorToolBar, InSequencer, CurveEditorAndSequencerCommands)
 								.Visibility(this, &SSequencer::GetCurveEditorToolBarVisibility)
 							]
 
@@ -464,7 +430,7 @@ void SSequencer::Construct(const FArguments& InArgs, TSharedRef<FSequencer> InSe
 					+ SGridPanel::Slot( Column1, Row2, SGridPanel::Layer(20) )
 					.Padding(ResizeBarPadding)
 					[
-						CurveEditor.ToSharedRef()
+						CurveEditorPanel
 					]
 
 					// Overlay that draws the scrub position
@@ -479,23 +445,33 @@ void SSequencer::Construct(const FArguments& InArgs, TSharedRef<FSequencer> InSe
 						.Clipping(EWidgetClipping::ClipToBounds)
 					]
 
-					// Goto box
 					+ SGridPanel::Slot(Column1, Row2, SGridPanel::Layer(40))
 						.Padding(ResizeBarPadding)
 						.HAlign(HAlign_Left)
 						.VAlign(VAlign_Top)
+					[
+						SNew(SHorizontalBox)
+
+						// Goto box
+						+SHorizontalBox::Slot()
+						.AutoWidth()
 						[
 							SAssignNew(GotoBox, SSequencerGotoBox, SequencerPtr.Pin().ToSharedRef(), *Settings, NumericTypeInterface.ToSharedRef())
 						]
 
-					// Transform box
-					+ SGridPanel::Slot(Column1, Row2, SGridPanel::Layer(50))
-						.Padding(ResizeBarPadding)
-						.HAlign(HAlign_Left)
-						.VAlign(VAlign_Top)
+						// Transform box
+						+SHorizontalBox::Slot()
+						.AutoWidth()
 						[
 							SAssignNew(TransformBox, SSequencerTransformBox, SequencerPtr.Pin().ToSharedRef(), *Settings, NumericTypeInterface.ToSharedRef())
 						]
+					]
+
+					+ SGridPanel::Slot(Column1, Row2, SGridPanel::Layer(50))
+					.Padding(ResizeBarPadding)
+					[
+						SAssignNew(FrameResolutionOverlay, SSequencerTimePanel, SequencerPtr)
+					]
 
 					// debug vis
 					+ SGridPanel::Slot( Column1, Row3, SGridPanel::Layer(10) )
@@ -556,6 +532,8 @@ void SSequencer::Construct(const FArguments& InArgs, TSharedRef<FSequencer> InSe
 		]
 	];
 
+	HideFrameResolutionOverlay();
+
 	InSequencer->GetSelection().GetOnKeySelectionChanged().AddSP(this, &SSequencer::HandleKeySelectionChanged);
 	InSequencer->GetSelection().GetOnSectionSelectionChanged().AddSP(this, &SSequencer::HandleSectionSelectionChanged);
 	InSequencer->GetSelection().GetOnOutlinerNodeSelectionChanged().AddSP(this, &SSequencer::HandleOutlinerNodeSelectionChanged);
@@ -597,7 +575,17 @@ void SSequencer::BindCommands(TSharedRef<FUICommandList> SequencerCommandBinding
 		FExecuteAction::CreateLambda([this]{ TransformBox->ToggleVisibility(); })
 	);
 }
-	
+
+void SSequencer::ShowFrameResolutionOverlay()
+{
+	FrameResolutionOverlay->SetVisibility(EVisibility::Visible);
+}
+
+void SSequencer::HideFrameResolutionOverlay()
+{
+	FrameResolutionOverlay->SetVisibility(EVisibility::Collapsed);
+}
+
 const ISequencerEditTool* SSequencer::GetEditTool() const
 {
 	return TrackArea->GetEditTool();
@@ -613,16 +601,11 @@ void SSequencer::NotifyPostChange(const FPropertyChangedEvent& PropertyChangedEv
 /* SSequencer implementation
  *****************************************************************************/
 
-TSharedRef<INumericTypeInterface<float>> SSequencer::GetNumericTypeInterface()
+TSharedRef<INumericTypeInterface<double>> SSequencer::GetNumericTypeInterface() const
 {
 	return NumericTypeInterface.ToSharedRef();
 }
 
-TSharedRef<INumericTypeInterface<float>> SSequencer::GetZeroPadNumericTypeInterface()
-{
-	return ZeroPadNumericTypeInterface.ToSharedRef();
-}
-
 
 /* SSequencer callbacks
  *****************************************************************************/
@@ -669,7 +652,6 @@ void SSequencer::HandleSectionSelectionChanged()
 void SSequencer::HandleOutlinerNodeSelectionChanged()
 {
 	const TSet<TSharedRef<FSequencerDisplayNode>>& OutlinerSelection = SequencerPtr.Pin()->GetSelection().GetSelectedOutlinerNodes();
-
 	if ( OutlinerSelection.Num() == 1 )
 	{
 		for ( auto& Node : OutlinerSelection )
@@ -680,7 +662,6 @@ void SSequencer::HandleOutlinerNodeSelectionChanged()
 	}
 }
 
-
 TSharedRef<SWidget> SSequencer::MakeAddButton()
 {
 	if (SequencerPtr.Pin()->IsReadOnly())
@@ -918,21 +899,8 @@ TSharedRef<SWidget> SSequencer::MakeToolBar()
 			true );
 
 		ToolBarBuilder.AddSeparator();
-		ToolBarBuilder.AddWidget(
-			SNew( SImage )
-				.Image(FEditorStyle::GetBrush("Sequencer.Time.Small")) );
 
-		ToolBarBuilder.AddWidget(
-			SNew( SBox )
-				.VAlign( VAlign_Center )
-				[
-					SNew( SNumericDropDown<float> )
-						.DropDownValues( SequencerSnapValues::GetTimeSnapValues() )
-						.bShowNamedValue(true)
-						.ToolTipText( LOCTEXT( "TimeSnappingIntervalToolTip", "Time snapping interval" ) )
-						.Value( this, &SSequencer::OnGetTimeSnapInterval )
-						.OnValueChanged( this, &SSequencer::OnTimeSnapIntervalChanged )
-				]);
+		ToolBarBuilder.AddWidget(SNew(SSequencerPlayRateCombo, SequencerPtr.Pin(), SharedThis(this)));
 	}
 	ToolBarBuilder.EndSection();
 
@@ -1011,6 +979,7 @@ TSharedRef<SWidget> SSequencer::MakeGeneralMenu()
 		MenuBuilder.AddMenuEntry( FSequencerCommands::Get().ToggleExpandCollapseNodesAndDescendants );
 		MenuBuilder.AddMenuEntry( FSequencerCommands::Get().ExpandAllNodesAndDescendants );
 		MenuBuilder.AddMenuEntry( FSequencerCommands::Get().CollapseAllNodesAndDescendants );
+		MenuBuilder.AddMenuEntry( FSequencerCommands::Get().SortAllNodesAndDescendants );
 	}
 	MenuBuilder.EndSection();
 
@@ -1023,7 +992,6 @@ TSharedRef<SWidget> SSequencer::MakeGeneralMenu()
 		MenuBuilder.AddMenuEntry(FSequencerCommands::Get().FixActorReferences);
 		MenuBuilder.AddMenuEntry(FSequencerCommands::Get().RebindPossessableReferences);
 	}
-	MenuBuilder.AddMenuEntry(FSequencerCommands::Get().FixFrameTiming);
 
 	if ( SequencerPtr.Pin()->IsLevelEditorSequencer() )
 	{
@@ -1067,9 +1035,9 @@ TSharedRef<SWidget> SSequencer::MakePlaybackMenu()
 	MenuBuilder.BeginSection("PlaybackThisSequence", LOCTEXT("PlaybackThisSequenceHeader", "Playback - This Sequence"));
 	{
 		// Menu entry for the start position
-		auto OnStartChanged = [=](float NewValue){
-			float Upper = SequencerPtr.Pin()->GetPlaybackRange().GetUpperBoundValue();
-			SequencerPtr.Pin()->SetPlaybackRange(TRange<float>(FMath::Min(NewValue, Upper), Upper));
+		auto OnStartChanged = [=](FFrameNumber NewValue){
+			FFrameNumber Upper = MovieScene::DiscreteExclusiveUpper(SequencerPtr.Pin()->GetPlaybackRange());
+			SequencerPtr.Pin()->SetPlaybackRange(TRange<FFrameNumber>(FMath::Min(NewValue, Upper-1), Upper));
 		};
 
 		MenuBuilder.AddWidget(
@@ -1081,32 +1049,34 @@ TSharedRef<SWidget> SSequencer::MakePlaybackMenu()
 			+ SHorizontalBox::Slot()
 				.AutoWidth()
 				[
-					SNew(SSpinBox<float>)
+					SNew(SSpinBox<double>)
 						.TypeInterface(NumericTypeInterface)
 						.IsEnabled_Lambda([=]() {
 							return !SequencerPtr.Pin()->IsPlaybackRangeLocked();
 						})
 						.Style(&FEditorStyle::GetWidgetStyle<FSpinBoxStyle>("Sequencer.HyperlinkSpinBox"))
-						.OnValueCommitted_Lambda([=](float Value, ETextCommit::Type){ OnStartChanged(Value); })
-						.OnValueChanged_Lambda(OnStartChanged)
+						.OnValueCommitted_Lambda([=](double Value, ETextCommit::Type){ OnStartChanged(FFrameTime::FromDecimal(Value).GetFrame()); })
+						.OnValueChanged_Lambda([=](double Value) { OnStartChanged(FFrameTime::FromDecimal(Value).GetFrame()); })
 						.OnBeginSliderMovement(OnPlaybackRangeBeginDrag)
-						.OnEndSliderMovement_Lambda([=](float Value){ OnStartChanged(Value); OnPlaybackRangeEndDrag.ExecuteIfBound(); })
-						.MinValue_Lambda([=]() -> float {
-							return SequencerPtr.Pin()->GetClampRange().GetLowerBoundValue(); 
+						.OnEndSliderMovement_Lambda([=](double Value){ OnStartChanged(FFrameTime::FromDecimal(Value).GetFrame()); OnPlaybackRangeEndDrag.ExecuteIfBound(); })
+						.MinValue_Lambda([=]() -> double {
+							// Clamp range is in seconds so we need to convert it to frame resolution.
+							return (SequencerPtr.Pin()->GetClampRange().GetLowerBoundValue() * SequencerPtr.Pin()->GetFocusedFrameResolution()).GetFrame().Value;
 						})
-						.MaxValue_Lambda([=]() -> float {
-							return SequencerPtr.Pin()->GetPlaybackRange().GetUpperBoundValue(); 
+						.MaxValue_Lambda([=]() -> double {
+							return SequencerPtr.Pin()->GetPlaybackRange().GetUpperBoundValue().Value; 
 						})
-						.Value_Lambda([=]() -> float {
-							return SequencerPtr.Pin()->GetPlaybackRange().GetLowerBoundValue();
+						.Value_Lambda([=]() -> double {
+							return SequencerPtr.Pin()->GetPlaybackRange().GetLowerBoundValue().Value;
 						})
-				],
+						.Delta(this, &SSequencer::GetSpinboxDelta)
+			],
 			LOCTEXT("PlaybackStartLabel", "Start"));
 
 		// Menu entry for the end position
-		auto OnEndChanged = [=](float NewValue){
-			float Lower = SequencerPtr.Pin()->GetPlaybackRange().GetLowerBoundValue();
-			SequencerPtr.Pin()->SetPlaybackRange(TRange<float>(Lower, FMath::Max(NewValue, Lower)));
+		auto OnEndChanged = [=](FFrameNumber NewValue){
+			FFrameNumber Lower = MovieScene::DiscreteInclusiveLower(SequencerPtr.Pin()->GetPlaybackRange());
+			SequencerPtr.Pin()->SetPlaybackRange(TRange<FFrameNumber>(Lower, FMath::Max(NewValue, Lower)));
 		};
 
 		MenuBuilder.AddWidget(
@@ -1118,25 +1088,27 @@ TSharedRef<SWidget> SSequencer::MakePlaybackMenu()
 			+ SHorizontalBox::Slot()
 				.AutoWidth()
 				[
-					SNew(SSpinBox<float>)
+					SNew(SSpinBox<double>)
 						.TypeInterface(NumericTypeInterface)
 						.IsEnabled_Lambda([=]() {
-							return !SequencerPtr.Pin()->IsPlaybackRangeLocked();
+					 		return !SequencerPtr.Pin()->IsPlaybackRangeLocked();
 						})
 						.Style(&FEditorStyle::GetWidgetStyle<FSpinBoxStyle>("Sequencer.HyperlinkSpinBox"))
-						.OnValueCommitted_Lambda([=](float Value, ETextCommit::Type){ OnEndChanged(Value); })
-						.OnValueChanged_Lambda(OnEndChanged)
+						.OnValueCommitted_Lambda([=](double Value, ETextCommit::Type){ OnEndChanged(FFrameTime::FromDecimal(Value).GetFrame()); })
+						.OnValueChanged_Lambda([=](double Value) { OnEndChanged(FFrameTime::FromDecimal(Value).GetFrame()); })
 						.OnBeginSliderMovement(OnPlaybackRangeBeginDrag)
-						.OnEndSliderMovement_Lambda([=](float Value){ OnEndChanged(Value); OnPlaybackRangeEndDrag.ExecuteIfBound(); })
-						.MinValue_Lambda([=]() -> float {
-							return SequencerPtr.Pin()->GetPlaybackRange().GetLowerBoundValue(); 
+						.OnEndSliderMovement_Lambda([=](double Value){ OnEndChanged(FFrameTime::FromDecimal(Value).GetFrame()); OnPlaybackRangeEndDrag.ExecuteIfBound(); })
+						.MinValue_Lambda([=]() -> double {
+					 		return SequencerPtr.Pin()->GetPlaybackRange().GetLowerBoundValue().Value;
 						})
-						.MaxValue_Lambda([=]() -> float {
-							return SequencerPtr.Pin()->GetClampRange().GetUpperBoundValue(); 
+						.MaxValue_Lambda([=]() -> double {
+							// Clamp range is in seconds so we need to convert it to frame resolution.
+							return (SequencerPtr.Pin()->GetClampRange().GetUpperBoundValue() * SequencerPtr.Pin()->GetFocusedFrameResolution()).GetFrame().Value;
 						})
-						.Value_Lambda([=]() -> float {
-							return SequencerPtr.Pin()->GetPlaybackRange().GetUpperBoundValue();
+						.Value_Lambda([=]() -> double {
+					 		return SequencerPtr.Pin()->GetPlaybackRange().GetUpperBoundValue().Value;
 						})
+						.Delta(this, &SSequencer::GetSpinboxDelta)
 				],
 			LOCTEXT("PlaybackStartEnd", "End"));
 
@@ -1144,7 +1116,6 @@ TSharedRef<SWidget> SSequencer::MakePlaybackMenu()
 		MenuBuilder.AddSubMenu(LOCTEXT("PlaybackSpeedHeader", "Playback Speed"), FText::GetEmpty(), FNewMenuDelegate::CreateRaw(this, &SSequencer::FillPlaybackSpeedMenu));
 
 		MenuBuilder.AddMenuEntry( FSequencerCommands::Get().TogglePlaybackRangeLocked );
-		MenuBuilder.AddMenuEntry( FSequencerCommands::Get().ToggleForceFixedFrameIntervalPlayback );
 
 		if (SequencerPtr.Pin()->IsLevelEditorSequencer())
 		{
@@ -1231,7 +1202,6 @@ TSharedRef<SWidget> SSequencer::MakeSnapMenu()
 	MenuBuilder.BeginSection("FramesRanges", LOCTEXT("SnappingMenuFrameRangesHeader", "Frame Ranges") );
 	{
 		MenuBuilder.AddMenuEntry( FSequencerCommands::Get().ToggleAutoScroll );
-		MenuBuilder.AddMenuEntry( FSequencerCommands::Get().ToggleShowFrameNumbers );
 		MenuBuilder.AddMenuEntry( FSequencerCommands::Get().ToggleShowRangeSlider );
 	}
 	MenuBuilder.EndSection();
@@ -1325,8 +1295,7 @@ TSharedRef<SWidget> SSequencer::MakeTimeRange(const TSharedRef<SWidget>& InnerCo
 		ShowRange,
 		TimeSliderController.ToSharedRef(),
 		EVisibility::Visible,
-		TAttribute<bool>(this, &SSequencer::ShowFrameNumbers),
-		GetZeroPadNumericTypeInterface()
+		NumericTypeInterface.ToSharedRef()
 		);
 	return SequencerWidgets.CreateTimeRange(Args, InnerContent);
 }
@@ -1481,6 +1450,11 @@ void SSequencer::UpdateLayoutTree()
 			}
 		}
 
+		// Add any additional paths that have been added by the user for nodes that may not exist yet but we want them to be selected
+		// after the node tree is updated and we restore selections.
+		SelectedPathNames.Append(AdditionalSelectionsToAdd);
+		AdditionalSelectionsToAdd.Empty();
+
 		// Suspend broadcasting selection changes because we don't want unnecessary rebuilds.
 		Sequencer->GetSelection().SuspendBroadcast();
 		Sequencer->GetSelection().Empty();
@@ -1493,7 +1467,6 @@ void SSequencer::UpdateLayoutTree()
 
 		// This must come after the selection state has been restored so that the tree and curve editor are populated with the correctly selected nodes
 		TreeView->Refresh();
-		CurveEditor->SetSequencerNodeTree(Sequencer->GetNodeTree());
 
 		RestoreKeySelection(SelectedKeys, Sequencer->GetSelection(), *Sequencer->GetNodeTree());
 		RestoreSectionSelection(SelectedSections, Sequencer->GetSelection());
@@ -1503,7 +1476,6 @@ void SSequencer::UpdateLayoutTree()
 	}
 }
 
-
 void SSequencer::UpdateBreadcrumbs()
 {
 	TSharedPtr<FSequencer> Sequencer = SequencerPtr.Pin();
@@ -1513,6 +1485,7 @@ void SSequencer::UpdateBreadcrumbs()
 	}
 
 	FMovieSceneSequenceID FocusedID = Sequencer->GetFocusedTemplateID();
+
 	if (BreadcrumbTrail->PeekCrumb().BreadcrumbType == FSequencerBreadcrumb::ShotType)
 	{
 		BreadcrumbTrail->PopCrumb();
@@ -1520,7 +1493,18 @@ void SSequencer::UpdateBreadcrumbs()
 
 	if( BreadcrumbTrail->PeekCrumb().BreadcrumbType == FSequencerBreadcrumb::MovieSceneType && BreadcrumbTrail->PeekCrumb().SequenceID != FocusedID )
 	{
-		FText CrumbName = Sequencer->GetFocusedMovieSceneSequence()->GetDisplayName();
+		FText CrumbName;
+		if(Sequencer->GetFocusedSequenceIsActive())
+		{
+			CrumbName = Sequencer->GetFocusedMovieSceneSequence()->GetDisplayName();
+		}
+		else
+		{
+			CrumbName = FText::Format(LOCTEXT("InactiveSequenceBreadcrumbFormat", "{0} [{1}]"),
+				Sequencer->GetFocusedMovieSceneSequence()->GetDisplayName(),
+				LOCTEXT("InactiveSequenceBreadcrumb", "Inactive"));
+		}
+			
 		// The current breadcrumb is not a moviescene so we need to make a new breadcrumb in order return to the parent moviescene later
 		BreadcrumbTrail->PushCrumb( CrumbName, FSequencerBreadcrumb( FocusedID ) );
 	}
@@ -1560,18 +1544,6 @@ void SSequencer::OnOutlinerSearchChanged( const FText& Filter )
 }
 
 
-float SSequencer::OnGetTimeSnapInterval() const
-{
-	TSharedPtr<FSequencer> Sequencer = SequencerPtr.Pin();
-	if ( Sequencer.IsValid() )
-	{
-		return Sequencer->GetFixedFrameInterval();
-	}	
-
-	return 1.f;
-}
-
-
 void SSequencer::OnDragEnter( const FGeometry& MyGeometry, const FDragDropEvent& DragDropEvent )
 {
 	// @todo sequencer: Add drop validity cue
@@ -1977,41 +1949,45 @@ void SSequencer::StepToKey(bool bStepToNextKey, bool bCameraOnly)
 
 		if ( Nodes.Num() > 0 )
 		{
-			float ClosestKeyDistance = MAX_FLT;
-			float CurrentTime = Sequencer->GetLocalTime();
-			float StepToTime = 0;
-			bool StepToKeyFound = false;
+			FFrameTime ClosestKeyDistance = FFrameTime(TNumericLimits<int32>::Max(), 0.99999f);
+			FFrameTime CurrentTime = Sequencer->GetLocalTime().Time;
+			TOptional<FFrameTime> NextTime;
+
+			TOptional<FFrameNumber> StepToTime;
 
 			auto It = Nodes.CreateConstIterator();
 			bool bExpand = !( *It ).Get().IsExpanded();
 
 			for ( auto Node : Nodes )
 			{
-				TArray<float> AllTimes;
+				TArray<FFrameNumber> AllTimes;
 
 				TSet<TSharedPtr<IKeyArea>> KeyAreas;
 				SequencerHelpers::GetAllKeyAreas( Node, KeyAreas );
 				for ( TSharedPtr<IKeyArea> KeyArea : KeyAreas )
 				{
-					for ( FKeyHandle& KeyHandle : KeyArea->GetUnsortedKeyHandles() )
-					{
-						float KeyTime = KeyArea->GetKeyTime( KeyHandle );
-						AllTimes.Add(KeyTime);
-					}
+					KeyArea->GetKeyTimes(AllTimes);
 				}
 
 				TSet<TWeakObjectPtr<UMovieSceneSection> > Sections;
-				SequencerHelpers::GetAllSections( Node, Sections );				
+				SequencerHelpers::GetAllSections( Node, Sections );
+
 				for ( TWeakObjectPtr<UMovieSceneSection> Section : Sections )
 				{
-					if (Section.IsValid() && !Section->IsInfinite())
+					if (Section.IsValid())
 					{
-						AllTimes.Add(Section->GetStartTime());
-						AllTimes.Add(Section->GetEndTime());
+						if (Section->HasStartFrame())
+						{
+							AllTimes.Add(Section->GetInclusiveStartFrame());
+						}
+						if (Section->HasEndFrame())
+						{
+							AllTimes.Add(Section->GetExclusiveEndFrame()-1);
+						}
 					}
 				}
 
-				for (float Time : AllTimes)
+				for (FFrameNumber Time : AllTimes)
 				{
 					if ( bStepToNextKey )
 					{
@@ -2019,7 +1995,6 @@ void SSequencer::StepToKey(bool bStepToNextKey, bool bCameraOnly)
 						{
 							StepToTime = Time;
 							ClosestKeyDistance = Time - CurrentTime;
-							StepToKeyFound = true;
 						}
 					}
 					else
@@ -2028,15 +2003,14 @@ void SSequencer::StepToKey(bool bStepToNextKey, bool bCameraOnly)
 						{
 							StepToTime = Time;
 							ClosestKeyDistance = CurrentTime - Time;
-							StepToKeyFound = true;
 						}
 					}
 				}
 			}
 
-			if ( StepToKeyFound )
+			if ( StepToTime.IsSet() )
 			{
-				Sequencer->SetLocalTime( StepToTime );
+				Sequencer->SetLocalTime( StepToTime.GetValue() );
 			}
 		}
 	}
@@ -2066,13 +2040,11 @@ EVisibility SSequencer::GetTimeRangeVisibility() const
 	return Settings->GetShowRangeSlider() ? EVisibility::Visible : EVisibility::Hidden;
 }
 
-
-bool SSequencer::ShowFrameNumbers() const
+EFrameNumberDisplayFormats SSequencer::GetTimeDisplayFormat() const
 {
-	return SequencerPtr.Pin()->CanShowFrameNumbers() && Settings->GetShowFrameNumbers();
+	return Settings->GetTimeDisplayFormat();
 }
 
-
 float SSequencer::GetOutlinerSpacerFill() const
 {
 	const float Column1Coeff = GetColumnFillCoefficient(1);
@@ -2100,57 +2072,50 @@ EVisibility SSequencer::GetCurveEditorVisibility() const
 
 void SSequencer::OnCurveEditorVisibilityChanged()
 {
-	if (CurveEditor.IsValid())
+	if (!Settings->GetLinkCurveEditorTimeRange())
 	{
-		if (!Settings->GetLinkCurveEditorTimeRange())
+		TRange<double> ClampRange = SequencerPtr.Pin()->GetClampRange();
+		if (CachedClampRange.IsEmpty())
 		{
-			TRange<float> ClampRange = SequencerPtr.Pin()->GetClampRange();
-			if (CachedClampRange.IsEmpty())
-			{
-				CachedClampRange = ClampRange;
-			}
-			SequencerPtr.Pin()->SetClampRange(CachedClampRange);
 			CachedClampRange = ClampRange;
-
-			TRange<float> ViewRange = SequencerPtr.Pin()->GetViewRange();
-			if (CachedViewRange.IsEmpty())
-			{
-				CachedViewRange = ViewRange;
-			}
-			SequencerPtr.Pin()->SetViewRange(CachedViewRange);
-			CachedViewRange = ViewRange;
 		}
+		SequencerPtr.Pin()->SetClampRange(CachedClampRange);
+		CachedClampRange = ClampRange;
 
-		// Only zoom horizontally if the editor is visible
-		CurveEditor->SetAllowAutoFrame(SequencerPtr.Pin()->GetShowCurveEditor());
-
-		if (CurveEditor->GetAutoFrame())
+		TRange<double> ViewRange = SequencerPtr.Pin()->GetViewRange();
+		if (CachedViewRange.IsEmpty())
 		{
-			CurveEditor->ZoomToFit();
+			CachedViewRange = ViewRange;
 		}
+		SequencerPtr.Pin()->SetViewRange(CachedViewRange);
+		CachedViewRange = ViewRange;
 	}
 
+	SequencerPtr.Pin()->SyncCurveEditorToSelection(false);
+	SequencerPtr.Pin()->GetCurveEditor()->ZoomToFit();
+
 	TreeView->UpdateTrackArea();
 }
 
 
 void SSequencer::OnTimeSnapIntervalChanged(float InInterval)
 {
-	TSharedPtr<FSequencer> Sequencer = SequencerPtr.Pin();
-	if ( Sequencer.IsValid() )
-	{
-		UMovieScene* MovieScene = Sequencer->GetFocusedMovieSceneSequence()->GetMovieScene();
-		if (!FMath::IsNearlyEqual(MovieScene->GetFixedFrameInterval(), InInterval))
-		{
-			FScopedTransaction SetFixedFrameIntervalTransaction( NSLOCTEXT( "Sequencer", "SetFixedFrameInterval", "Set scene fixed frame interval" ) );
-			MovieScene->Modify();
-			MovieScene->SetFixedFrameInterval( InInterval );
+	// @todo: sequencer-timecode: Address dealing with different time intervals
+	// TSharedPtr<FSequencer> Sequencer = SequencerPtr.Pin();
+	// if ( Sequencer.IsValid() )
+	// {
+	// 	UMovieScene* MovieScene = Sequencer->GetFocusedMovieSceneSequence()->GetMovieScene();
+	// 	if (!FMath::IsNearlyEqual(MovieScene->GetFixedFrameInterval(), InInterval))
+	// 	{
+	// 		FScopedTransaction SetFixedFrameIntervalTransaction( NSLOCTEXT( "Sequencer", "SetFixedFrameInterval", "Set scene fixed frame interval" ) );
+	// 		MovieScene->Modify();
+	// 		MovieScene->SetFixedFrameInterval( InInterval );
 
-			// Update the current time to the new interval
-			float NewTime = SequencerHelpers::SnapTimeToInterval(Sequencer->GetLocalTime(), InInterval);
-			Sequencer->SetLocalTime(NewTime);
-		}
-	}
+	// 		// Update the current time to the new interval
+	// 		float NewTime = SequencerHelpers::SnapTimeToInterval(Sequencer->GetLocalTime(), InInterval);
+	// 		Sequencer->SetLocalTime(NewTime);
+	// 	}
+	// }
 }
 
 
@@ -2174,12 +2139,13 @@ FVirtualTrackArea SSequencer::GetVirtualTrackArea() const
 	return FVirtualTrackArea(*SequencerPtr.Pin(), *TreeView.Get(), TrackArea->GetCachedGeometry());
 }
 
-FPasteContextMenuArgs SSequencer::GeneratePasteArgs(float PasteAtTime, TSharedPtr<FMovieSceneClipboard> Clipboard)
+FPasteContextMenuArgs SSequencer::GeneratePasteArgs(FFrameNumber PasteAtTime, TSharedPtr<FMovieSceneClipboard> Clipboard)
 {
 	TSharedPtr<FSequencer> Sequencer = SequencerPtr.Pin();
 	if (Settings->GetIsSnapEnabled())
 	{
-		PasteAtTime = SequencerHelpers::SnapTimeToInterval(PasteAtTime, Sequencer->GetFixedFrameInterval());
+		// @todo: sequencer-timecode: play rate override
+		//PasteAtTime = FQualifiedFrameTime(PasteAtTime, GetFrameResolution()).ToFrameRate();
 	}
 
 	// Open a paste menu at the current mouse position
@@ -2276,7 +2242,7 @@ bool SSequencer::OpenPasteMenu()
 	TSharedPtr<FSequencer> Sequencer = SequencerPtr.Pin();
 	if (Sequencer->GetClipboardStack().Num() != 0)
 	{
-		FPasteContextMenuArgs Args = GeneratePasteArgs(Sequencer->GetLocalTime(), Sequencer->GetClipboardStack().Last());
+		FPasteContextMenuArgs Args = GeneratePasteArgs(Sequencer->GetLocalTime().Time.FrameNumber, Sequencer->GetClipboardStack().Last());
 		ContextMenu = FPasteContextMenu::CreateMenu(*Sequencer, Args);
 	}
 
@@ -2316,7 +2282,7 @@ void SSequencer::PasteFromHistory()
 		return;
 	}
 
-	FPasteContextMenuArgs Args = GeneratePasteArgs(Sequencer->GetLocalTime());
+	FPasteContextMenuArgs Args = GeneratePasteArgs(Sequencer->GetLocalTime().Time.FrameNumber);
 	TSharedPtr<FPasteFromHistoryContextMenu> ContextMenu = FPasteFromHistoryContextMenu::CreateMenu(*Sequencer, Args);
 
 	if (ContextMenu.IsValid())
@@ -2339,27 +2305,20 @@ void SSequencer::PasteFromHistory()
 	}
 }
 
-void SSequencer::OnSequenceInstanceActivated( FMovieSceneSequenceIDRef ActiveInstanceID )
+EVisibility SSequencer::GetDebugVisualizerVisibility() const
 {
-	TSharedPtr<FSequencer> Sequencer = SequencerPtr.Pin();
-	if ( Sequencer.IsValid() )
-	{
-		UMovieScene* MovieScene = Sequencer->GetFocusedMovieSceneSequence()->GetMovieScene();
-		if ( MovieScene->GetFixedFrameInterval() == 0 )
-		{
-			MovieScene->Modify();
-			MovieScene->SetFixedFrameInterval(Settings->GetTimeSnapInterval());
+	return Settings->ShouldShowDebugVisualization() ? EVisibility::Visible : EVisibility::Collapsed;
+}
 
-			// Update the current time to the new interval
-			float NewTime = SequencerHelpers::SnapTimeToInterval(Sequencer->GetLocalTime(), Settings->GetTimeSnapInterval());
-			Sequencer->SetLocalTime(NewTime);
-		}
-	}
+void SSequencer::BuildCustomContextMenuForGuid(FMenuBuilder& MenuBuilder, FGuid ObjectBinding)
+{
+	OnBuildCustomContextMenuForGuid.ExecuteIfBound(MenuBuilder, ObjectBinding);
 }
 
-EVisibility SSequencer::GetDebugVisualizerVisibility() const
+double SSequencer::GetSpinboxDelta() const
 {
-	return Settings->ShouldShowDebugVisualization() ? EVisibility::Visible : EVisibility::Collapsed;
+	TSharedPtr<FSequencer> Sequencer = SequencerPtr.Pin();
+	return Sequencer->GetPlayRateDeltaFrameCount();
 }
 
 #undef LOCTEXT_NAMESPACE
