Copying //UE4/Dev-Sequencer to //UE4/Dev-Main (Source: //UE4/Dev-Sequencer @ 3946692)

diff --git a/Engine/Source/Editor/Sequencer/Private/SSequencer.h b/Engine/Source/Editor/Sequencer/Private/SSequencer.h
--- a/Engine/Source/Editor/Sequencer/Private/SSequencer.h
+++ b/Engine/Source/Editor/Sequencer/Private/SSequencer.h
@@ -24,7 +24,6 @@ class FSequencerTimeSliderController;
 class FUnloadedClassDragDropOp;
 class FVirtualTrackArea;
 class ISequencerEditTool;
-class SSequencerCurveEditor;
 class SSequencerGotoBox;
 class SSequencerLabelBrowser;
 class SSequencerTrackArea;
@@ -97,7 +96,6 @@ public:
 
 	DECLARE_DELEGATE_OneParam( FOnToggleBoolOption, bool )
 	SLATE_BEGIN_ARGS( SSequencer )
-		: _ScrubPosition( 1.0f )
 	{ }
 		/** The current view range (seconds) */
 		SLATE_ATTRIBUTE( FAnimatedRange, ViewRange )
@@ -106,25 +104,19 @@ public:
 		SLATE_ATTRIBUTE( FAnimatedRange, ClampRange )
 
 		/** The playback range */
-		SLATE_ATTRIBUTE( TRange<float>, PlaybackRange )
+		SLATE_ATTRIBUTE( TRange<FFrameNumber>, PlaybackRange )
 
 		/** The selection range */
-		SLATE_ATTRIBUTE(TRange<float>, SelectionRange)
+		SLATE_ATTRIBUTE( TRange<FFrameNumber>, SelectionRange)
 
 		/** The current sub sequence range */
-		SLATE_ATTRIBUTE(TOptional<TRange<float>>, SubSequenceRange)
+		SLATE_ATTRIBUTE( TOptional<TRange<FFrameNumber>>, SubSequenceRange)
 
 		/** The playback status */
 		SLATE_ATTRIBUTE( EMovieScenePlayerStatus::Type, PlaybackStatus )
 
 		/** Called when the user changes the playback range */
-		SLATE_EVENT( FOnRangeChanged, OnInOutRangeChanged )
-
-		/** Called when the user has begun dragging the selection range */
-		SLATE_EVENT( FSimpleDelegate, OnBeginInOutRangeDrag )
-
-		/** Called when the user changes the playback range */
-		SLATE_EVENT( FOnRangeChanged, OnPlaybackRangeChanged )
+		SLATE_EVENT( FOnFrameRangeChanged, OnPlaybackRangeChanged )
 
 		/** Called when the user has begun dragging the playback range */
 		SLATE_EVENT( FSimpleDelegate, OnPlaybackRangeBeginDrag )
@@ -133,7 +125,7 @@ public:
 		SLATE_EVENT( FSimpleDelegate, OnPlaybackRangeEndDrag )
 
 		/** Called when the user changes the selection range */
-		SLATE_EVENT( FOnRangeChanged, OnSelectionRangeChanged )
+		SLATE_EVENT( FOnFrameRangeChanged, OnSelectionRangeChanged )
 
 		/** Called when the user has begun dragging the selection range */
 		SLATE_EVENT( FSimpleDelegate, OnSelectionRangeBeginDrag )
@@ -147,17 +139,14 @@ public:
 		/** Called when the user toggles the play back range lock */
 		SLATE_EVENT( FSimpleDelegate, OnTogglePlaybackRangeLocked )
 
-		/** The time snap interval */
-		SLATE_ATTRIBUTE( float, TimeSnapInterval )
-
 		/** The current scrub position in (seconds) */
-		SLATE_ATTRIBUTE( float, ScrubPosition )
+		SLATE_ATTRIBUTE( FFrameTime, ScrubPosition )
 
 		/** Called when the user changes the view range */
 		SLATE_EVENT( FOnViewRangeChanged, OnViewRangeChanged )
 
 		/** Called when the user changes the clamp range */
-		SLATE_EVENT( FOnRangeChanged, OnClampRangeChanged )
+		SLATE_EVENT( FOnTimeRangeChanged, OnClampRangeChanged )
 
 		/** Called to get the nearest key */
 		SLATE_EVENT( FOnGetNearestKey, OnGetNearestKey )
@@ -174,6 +163,9 @@ public:
 		/** Called to populate the add combo button in the toolbar. */
 		SLATE_EVENT( FOnGetAddMenuContent, OnGetAddMenuContent )
 
+		/** Called when object is clicked. */
+		SLATE_EVENT(FOnBuildCustomContextMenuForGuid, OnBuildCustomContextMenuForGuid)
+			
 		/** Called when any widget contained within sequencer has received focus */
 		SLATE_EVENT( FSimpleDelegate, OnReceivedFocus )
 
@@ -238,14 +230,14 @@ public:
 	TArray<FSectionHandle> GetSectionHandles(const TSet<TWeakObjectPtr<UMovieSceneSection>>& DesiredSections) const;
 
 	/** @return a numeric type interface that will parse and display numbers as frames and times correctly */
-	TSharedRef<INumericTypeInterface<float>> GetNumericTypeInterface();
-
-	/** @return a numeric type interface that will parse and display numbers as frames and times correctly, including any zero padding, if necessary */
-	TSharedRef<INumericTypeInterface<float>> GetZeroPadNumericTypeInterface();
+	TSharedRef<INumericTypeInterface<double>> GetNumericTypeInterface() const;
 	
 	/** Access the currently active track area edit tool */
 	const ISequencerEditTool* GetEditTool() const;
 
+	void ShowFrameResolutionOverlay();
+	void HideFrameResolutionOverlay();
+
 public:
 
 	// FNotifyHook overrides
@@ -281,6 +273,9 @@ private:
 	/** Handles changes to the selected outliner nodes. */
 	void HandleOutlinerNodeSelectionChanged();
 
+	/** Syncs the current node selection to the curve editor. */
+	void SyncCurveEditorToSelection();
+
 	/** Empty active timer to ensure Slate ticks during Sequencer playback */
 	EActiveTimerReturnType EnsureSlateTickDuringPlayback(double InCurrentTime, float InDeltaTime);	
 
@@ -327,11 +322,6 @@ public:
 
 private:
 
-	/**
-	* @return The value of the current time snap interval.
-	*/
-	float OnGetTimeSnapInterval() const;
-
 	/**
 	* Called when the time snap interval changes.
 	*
@@ -419,8 +409,8 @@ private:
 	/** Gets whether or not the time range should be visible. */
 	EVisibility GetTimeRangeVisibility() const;
 
-	/** Gets whether or not to show frame numbers. */
-	bool ShowFrameNumbers() const;
+	/** What is the preferred display format for time values. */
+	EFrameNumberDisplayFormats GetTimeDisplayFormat() const;
 
 	/** Called when a column fill percentage is changed by a splitter slot. */
 	void OnColumnFillCoefficientChanged(float FillCoefficient, int32 ColumnIndex);
@@ -428,14 +418,14 @@ private:
 	/** Gets paint options for painting the playback range on sequencer */
 	FPaintPlaybackRangeArgs GetSectionPlaybackRangeArgs() const;
 
-	/** Called whenever the active sequence instance changes on the FSequencer */
-	void OnSequenceInstanceActivated( FMovieSceneSequenceIDRef ActiveInstanceID );
-
 	EVisibility GetDebugVisualizerVisibility() const;
 	
 	void SetPlaybackSpeed(float InPlaybackSpeed);
 	float GetPlaybackSpeed() const;
 
+	/** Controls how fast Spinboxes change values. */
+	double GetSpinboxDelta() const;
+
 public:
 	/** On Paste Command */
 	void OnPaste();
@@ -451,8 +441,13 @@ public:
 	void PasteFromHistory();
 
 	/** Generate a paste menu args structure */
-	struct FPasteContextMenuArgs GeneratePasteArgs(float PasteAtTime, TSharedPtr<FMovieSceneClipboard> Clipboard = nullptr);
+	struct FPasteContextMenuArgs GeneratePasteArgs(FFrameNumber PasteAtTime, TSharedPtr<FMovieSceneClipboard> Clipboard = nullptr);
 
+	/** Execute custom context menu if passed in the FSequencerViewParams  */
+	void BuildCustomContextMenuForGuid(FMenuBuilder& MenuBuilder, FGuid ObjectBinding);
+
+	/** This adds the specified path to the selection set to be restored the next time the tree view is refreshed. */
+	void AddAdditionalPathToSelectionSet(const FString& Path) { AdditionalSelectionsToAdd.Add(Path); }
 private:
 
 	/** Goto box widget. */
@@ -467,9 +462,6 @@ private:
 	/** Outliner widget */
 	TSharedPtr<SSequencerTrackOutliner> TrackOutliner;
 
-	/** The curve editor. */
-	TSharedPtr<SSequencerCurveEditor> CurveEditor;
-
 	/** The breadcrumb trail widget for this sequencer */
 	TSharedPtr<SBreadcrumbTrail<FSequencerBreadcrumb>> BreadcrumbTrail;
 
@@ -507,14 +499,16 @@ private:
 	TSharedPtr<FExtender> ToolbarExtender;
 
 	/** Numeric type interface used for converting parsing and generating strings from numbers */
-	TSharedPtr<INumericTypeInterface<float>> NumericTypeInterface;
-	TSharedPtr<INumericTypeInterface<float>> ZeroPadNumericTypeInterface;
+	TSharedPtr<INumericTypeInterface<double>> NumericTypeInterface;
 
 	/** Time slider controller for this sequencer */
 	TSharedPtr<FSequencerTimeSliderController> TimeSliderController;
 
 	FOnGetAddMenuContent OnGetAddMenuContent;
 
+	/** Called when object is clicked in track list */
+	FOnBuildCustomContextMenuForGuid OnBuildCustomContextMenuForGuid;
+
 	/** Called when the user has begun dragging the selection selection range */
 	FSimpleDelegate OnSelectionRangeBeginDrag;
 
@@ -531,6 +525,15 @@ private:
 	FSimpleDelegate OnReceivedFocus;
 
 	/** Cached clamp and view range for unlinking the curve editor time range */
-	TRange<float> CachedClampRange;
-	TRange<float> CachedViewRange;
+	TRange<double> CachedClampRange;
+	TRange<double> CachedViewRange;
+
+	/**
+	 * A list of additional paths to add to the selection set when it is restored after rebuilding the tree.
+	 * This can be used to highlight nodes that may not exist until the rebuild. Cleared after the tree is rebuilt
+	 * and the selection list is restored.
+	*/
+	TSet<FString> AdditionalSelectionsToAdd;
+
+	TSharedPtr<SWidget> FrameResolutionOverlay;
 };
diff --git a/Engine/Source/Editor/Sequencer/Private/SSequencerCurveEditorToolBar.cpp b/Engine/Source/Editor/Sequencer/Private/SSequencerCurveEditorToolBar.cpp
--- a/Engine/Source/Editor/Sequencer/Private/SSequencerCurveEditorToolBar.cpp
+++ b/Engine/Source/Editor/Sequencer/Private/SSequencerCurveEditorToolBar.cpp
@@ -9,14 +9,24 @@
 #include "Widgets/Input/SNumericDropDown.h"
 #include "SequencerCommonHelpers.h"
 #include "SequencerSettings.h"
-#include "RichCurveEditorCommands.h"
+#include "CurveEditorCommands.h"
+#include "SequencerCommands.h"
 
 #define LOCTEXT_NAMESPACE "CurveEditorToolBar"
 
 void SSequencerCurveEditorToolBar::Construct( const FArguments& InArgs, TSharedRef<FSequencer> InSequencer, TSharedPtr<FUICommandList> CurveEditorCommandList )
 {
 	Sequencer = InSequencer;
-	SequencerSettings = InSequencer->GetSettings();
+	SequencerSettings = InSequencer->GetSequencerSettings();
+
+	TArray<SNumericDropDown<float>::FNamedValue> SnapValues = {
+		SNumericDropDown<float>::FNamedValue( 0.001f, LOCTEXT( "Snap_OneThousandth", "0.001" ), LOCTEXT( "SnapDescription_OneThousandth", "Set snap to 1/1000th" ) ),
+		SNumericDropDown<float>::FNamedValue( 0.01f,  LOCTEXT( "Snap_OneHundredth",  "0.01" ),  LOCTEXT( "SnapDescription_OneHundredth",  "Set snap to 1/100th" ) ),
+		SNumericDropDown<float>::FNamedValue( 0.1f,   LOCTEXT( "Snap_OneTenth",      "0.1" ),   LOCTEXT( "SnapDescription_OneTenth",      "Set snap to 1/10th" ) ),
+		SNumericDropDown<float>::FNamedValue( 1.0f,   LOCTEXT( "Snap_One",           "1" ),     LOCTEXT( "SnapDescription_One",           "Set snap to 1" ) ),
+		SNumericDropDown<float>::FNamedValue( 10.0f,  LOCTEXT( "Snap_Ten",           "10" ),    LOCTEXT( "SnapDescription_Ten",           "Set snap to 10" ) ),
+		SNumericDropDown<float>::FNamedValue( 100.0f, LOCTEXT( "Snap_OneHundred",    "100" ),   LOCTEXT( "SnapDescription_OneHundred",    "Set snap to 100" ) ),
+	};
 
 	FToolBarBuilder ToolBarBuilder( CurveEditorCommandList, FMultiBoxCustomization::None, TSharedPtr<FExtender>(), Orient_Horizontal, true );
 
@@ -28,7 +38,7 @@ void SSequencerCurveEditorToolBar::Construct( const FArguments& InArgs, TSharedR
 		TAttribute<FSlateIcon>(),
 		true );
 
-	ToolBarBuilder.AddToolBarButton(FRichCurveEditorCommands::Get().ToggleOutputSnapping);
+	ToolBarBuilder.AddToolBarButton(FCurveEditorCommands::Get().ToggleOutputSnapping);
 
 	ToolBarBuilder.AddWidget(
 		SNew( SImage )
@@ -39,7 +49,7 @@ void SSequencerCurveEditorToolBar::Construct( const FArguments& InArgs, TSharedR
 			.VAlign( VAlign_Center )
 			[
 				SNew( SNumericDropDown<float> )
-					.DropDownValues( SequencerSnapValues::GetSnapValues() )
+					.DropDownValues( SnapValues )
 					.ToolTipText( LOCTEXT( "ValueSnappingIntervalToolTip", "Curve value snapping interval" ) )
 					.Value( this, &SSequencerCurveEditorToolBar::OnGetValueSnapInterval )
 					.OnValueChanged( this, &SSequencerCurveEditorToolBar::OnValueSnapIntervalChanged )
@@ -47,26 +57,26 @@ void SSequencerCurveEditorToolBar::Construct( const FArguments& InArgs, TSharedR
 
 	ToolBarBuilder.BeginSection( "Curve" );
 	{
-		ToolBarBuilder.AddToolBarButton( FRichCurveEditorCommands::Get().ZoomToFitHorizontal );
-		ToolBarBuilder.AddToolBarButton( FRichCurveEditorCommands::Get().ZoomToFitVertical );
-		ToolBarBuilder.AddToolBarButton( FRichCurveEditorCommands::Get().ZoomToFit );
+		ToolBarBuilder.AddToolBarButton( FCurveEditorCommands::Get().ZoomToFitHorizontal );
+		ToolBarBuilder.AddToolBarButton( FCurveEditorCommands::Get().ZoomToFitVertical );
+		ToolBarBuilder.AddToolBarButton( FCurveEditorCommands::Get().ZoomToFit );
 	}
 	ToolBarBuilder.EndSection();
 
 	ToolBarBuilder.BeginSection( "Interpolation" );
 	{
-		ToolBarBuilder.AddToolBarButton( FRichCurveEditorCommands::Get().InterpolationCubicAuto );
-		ToolBarBuilder.AddToolBarButton( FRichCurveEditorCommands::Get().InterpolationCubicUser );
-		ToolBarBuilder.AddToolBarButton( FRichCurveEditorCommands::Get().InterpolationCubicBreak );
-		ToolBarBuilder.AddToolBarButton( FRichCurveEditorCommands::Get().InterpolationLinear );
-		ToolBarBuilder.AddToolBarButton( FRichCurveEditorCommands::Get().InterpolationConstant );
+		ToolBarBuilder.AddToolBarButton( FCurveEditorCommands::Get().InterpolationCubicAuto );
+		ToolBarBuilder.AddToolBarButton( FCurveEditorCommands::Get().InterpolationCubicUser );
+		ToolBarBuilder.AddToolBarButton( FCurveEditorCommands::Get().InterpolationCubicBreak );
+		ToolBarBuilder.AddToolBarButton( FCurveEditorCommands::Get().InterpolationLinear );
+		ToolBarBuilder.AddToolBarButton( FCurveEditorCommands::Get().InterpolationConstant );
 	}
 	ToolBarBuilder.EndSection();
 
 	ToolBarBuilder.BeginSection( "Tangents" );
 	{
-		ToolBarBuilder.AddToolBarButton( FRichCurveEditorCommands::Get().FlattenTangents );
-		ToolBarBuilder.AddToolBarButton( FRichCurveEditorCommands::Get().StraightenTangents );
+		ToolBarBuilder.AddToolBarButton( FCurveEditorCommands::Get().FlattenTangents );
+		ToolBarBuilder.AddToolBarButton( FCurveEditorCommands::Get().StraightenTangents );
 	}
 	ToolBarBuilder.EndSection();
 
@@ -90,24 +100,24 @@ TSharedRef<SWidget> SSequencerCurveEditorToolBar::MakeCurveEditorViewOptionsMenu
 
 	MenuBuilder.BeginSection( "CurveVisibility", LOCTEXT( "CurveEditorMenuCurveVisibilityHeader", "Curve Visibility" ) );
 	{
-		MenuBuilder.AddMenuEntry( FRichCurveEditorCommands::Get().SetAllCurveVisibility );
-		MenuBuilder.AddMenuEntry( FRichCurveEditorCommands::Get().SetSelectedCurveVisibility );
-		MenuBuilder.AddMenuEntry( FRichCurveEditorCommands::Get().SetAnimatedCurveVisibility );
+		MenuBuilder.AddMenuEntry( FSequencerCommands::Get().SetAllCurveVisibility );
+		MenuBuilder.AddMenuEntry( FSequencerCommands::Get().SetSelectedCurveVisibility );
+		MenuBuilder.AddMenuEntry( FSequencerCommands::Get().SetAnimatedCurveVisibility );
 	}
 	MenuBuilder.EndSection();
 
 	MenuBuilder.BeginSection( "TangentVisibility", LOCTEXT( "CurveEditorMenuTangentVisibilityHeader", "Tangent Visibility" ) );
 	{
-		MenuBuilder.AddMenuEntry( FRichCurveEditorCommands::Get().SetAllTangentsVisibility );
-		MenuBuilder.AddMenuEntry( FRichCurveEditorCommands::Get().SetSelectedKeysTangentVisibility );
-		MenuBuilder.AddMenuEntry( FRichCurveEditorCommands::Get().SetNoTangentsVisibility );
+		MenuBuilder.AddMenuEntry( FCurveEditorCommands::Get().SetAllTangentsVisibility );
+		MenuBuilder.AddMenuEntry( FCurveEditorCommands::Get().SetSelectedKeysTangentVisibility );
+		MenuBuilder.AddMenuEntry( FCurveEditorCommands::Get().SetNoTangentsVisibility );
 	}
 	MenuBuilder.EndSection();
 
 	MenuBuilder.AddMenuSeparator();
 
-	MenuBuilder.AddMenuEntry( FRichCurveEditorCommands::Get().ToggleAutoFrameCurveEditor);
-	MenuBuilder.AddMenuEntry( FRichCurveEditorCommands::Get().ToggleShowCurveEditorCurveToolTips);
+	MenuBuilder.AddMenuEntry( FCurveEditorCommands::Get().ToggleAutoFrameCurveEditor);
+	MenuBuilder.AddMenuEntry( FCurveEditorCommands::Get().ToggleShowCurveEditorCurveToolTips);
 
 	return MenuBuilder.MakeWidget();
 }
@@ -120,11 +130,11 @@ TSharedRef<SWidget> SSequencerCurveEditorToolBar::MakeCurveEditorCurveOptionsMen
 		{
 			MenuBuilder.BeginSection( "Pre-Infinity Extrapolation", LOCTEXT( "CurveEditorMenuPreInfinityExtrapHeader", "Extrapolation" ) );
 			{
-				MenuBuilder.AddMenuEntry( FRichCurveEditorCommands::Get().SetPreInfinityExtrapCycle);
-				MenuBuilder.AddMenuEntry( FRichCurveEditorCommands::Get().SetPreInfinityExtrapCycleWithOffset);
-				MenuBuilder.AddMenuEntry( FRichCurveEditorCommands::Get().SetPreInfinityExtrapOscillate);
-				MenuBuilder.AddMenuEntry( FRichCurveEditorCommands::Get().SetPreInfinityExtrapLinear);
-				MenuBuilder.AddMenuEntry( FRichCurveEditorCommands::Get().SetPreInfinityExtrapConstant);
+				MenuBuilder.AddMenuEntry( FCurveEditorCommands::Get().SetPreInfinityExtrapCycle);
+				MenuBuilder.AddMenuEntry( FCurveEditorCommands::Get().SetPreInfinityExtrapCycleWithOffset);
+				MenuBuilder.AddMenuEntry( FCurveEditorCommands::Get().SetPreInfinityExtrapOscillate);
+				MenuBuilder.AddMenuEntry( FCurveEditorCommands::Get().SetPreInfinityExtrapLinear);
+				MenuBuilder.AddMenuEntry( FCurveEditorCommands::Get().SetPreInfinityExtrapConstant);
 			}
 			MenuBuilder.EndSection();
 		}
@@ -133,11 +143,11 @@ TSharedRef<SWidget> SSequencerCurveEditorToolBar::MakeCurveEditorCurveOptionsMen
 		{
 			MenuBuilder.BeginSection( "Post-Infinity Extrapolation", LOCTEXT( "CurveEditorMenuPostInfinityExtrapHeader", "Extrapolation" ) );
 			{
-				MenuBuilder.AddMenuEntry( FRichCurveEditorCommands::Get().SetPostInfinityExtrapCycle);
-				MenuBuilder.AddMenuEntry( FRichCurveEditorCommands::Get().SetPostInfinityExtrapCycleWithOffset);
-				MenuBuilder.AddMenuEntry( FRichCurveEditorCommands::Get().SetPostInfinityExtrapOscillate);
-				MenuBuilder.AddMenuEntry( FRichCurveEditorCommands::Get().SetPostInfinityExtrapLinear);
-				MenuBuilder.AddMenuEntry( FRichCurveEditorCommands::Get().SetPostInfinityExtrapConstant);
+				MenuBuilder.AddMenuEntry( FCurveEditorCommands::Get().SetPostInfinityExtrapCycle);
+				MenuBuilder.AddMenuEntry( FCurveEditorCommands::Get().SetPostInfinityExtrapCycleWithOffset);
+				MenuBuilder.AddMenuEntry( FCurveEditorCommands::Get().SetPostInfinityExtrapOscillate);
+				MenuBuilder.AddMenuEntry( FCurveEditorCommands::Get().SetPostInfinityExtrapLinear);
+				MenuBuilder.AddMenuEntry( FCurveEditorCommands::Get().SetPostInfinityExtrapConstant);
 			}
 			MenuBuilder.EndSection();
 		}
@@ -145,8 +155,8 @@ TSharedRef<SWidget> SSequencerCurveEditorToolBar::MakeCurveEditorCurveOptionsMen
 
 	FMenuBuilder MenuBuilder( true, CurveEditorCommandList );
 
-	MenuBuilder.AddMenuEntry( FRichCurveEditorCommands::Get().BakeCurve);
-	MenuBuilder.AddMenuEntry( FRichCurveEditorCommands::Get().ReduceCurve);
+	MenuBuilder.AddMenuEntry( FCurveEditorCommands::Get().BakeCurve);
+	MenuBuilder.AddMenuEntry( FCurveEditorCommands::Get().ReduceCurve);
 
 	MenuBuilder.AddSubMenu(
 		LOCTEXT( "PreInfinitySubMenu", "Pre-Infinity" ),
diff --git a/Engine/Source/Editor/Sequencer/Private/SSequencerDebugVisualizer.cpp b/Engine/Source/Editor/Sequencer/Private/SSequencerDebugVisualizer.cpp
--- a/Engine/Source/Editor/Sequencer/Private/SSequencerDebugVisualizer.cpp
+++ b/Engine/Source/Editor/Sequencer/Private/SSequencerDebugVisualizer.cpp
@@ -7,6 +7,9 @@
 #include "Widgets/Layout/SBorder.h"
 #include "Widgets/Text/STextBlock.h"
 #include "Layout/ArrangedChildren.h"
+#include "MovieSceneSequence.h"
+#include "MovieScene.h"
+#include "MovieSceneTimeHelpers.h"
 
 #define LOCTEXT_NAMESPACE "SSequencerDebugVisualizer"
 
@@ -37,13 +40,15 @@ void SSequencerDebugVisualizer::Refresh()
 		return;
 	}
 
+	const UMovieSceneSequence*           ActiveSequence = Sequencer->GetFocusedMovieSceneSequence();
 	const FMovieSceneEvaluationTemplate* ActiveTemplate = GetTemplate();
-	if (!ActiveTemplate)
+	if (!ActiveTemplate || !ActiveSequence)
 	{
 		return;
 	}
 
-	const FMovieSceneEvaluationField& EvaluationField = ActiveTemplate->EvaluationField;
+	const FFrameRate                  SequenceResolution = ActiveSequence->GetMovieScene()->GetFrameResolution();
+	const FMovieSceneEvaluationField& EvaluationField    = ActiveTemplate->EvaluationField;
 
 	CachedSignature = EvaluationField.GetSignature();
 
@@ -79,7 +84,7 @@ void SSequencerDebugVisualizer::Refresh()
 
 	for (int32 Index = 0; Index < EvaluationField.Size(); ++Index)
 	{
-		TRange<float> Range = EvaluationField.GetRange(Index);
+		TRange<FFrameNumber> Range = EvaluationField.GetRange(Index);
 
 		const float Complexity = SegmentComplexity[Index];
 		
@@ -91,7 +96,7 @@ void SSequencerDebugVisualizer::Refresh()
 
 		Children.Add(
 			SNew(SSequencerDebugSlot, Index)
-			.Visibility(this, &SSequencerDebugVisualizer::GetSegmentVisibility, Range)
+			.Visibility(this, &SSequencerDebugVisualizer::GetSegmentVisibility, Range / SequenceResolution)
 			.ToolTip(
 				SNew(SToolTip)
 				[
@@ -131,10 +136,10 @@ FGeometry SSequencerDebugVisualizer::GetSegmentGeometry(const FGeometry& Allotte
 		return AllottedGeometry.MakeChild(FVector2D(0,0), FVector2D(0,0));
 	}
 
-	TRange<float> SegmentRange = ActiveTemplate->EvaluationField.GetRange(Slot.GetSegmentIndex());
+	TRange<FFrameNumber> SegmentRange = ActiveTemplate->EvaluationField.GetRange(Slot.GetSegmentIndex());
 
-	float PixelStartX = SegmentRange.GetLowerBound().IsOpen() ? 0.f : TimeToPixelConverter.TimeToPixel(SegmentRange.GetLowerBoundValue());
-	float PixelEndX = SegmentRange.GetUpperBound().IsOpen() ? AllottedGeometry.GetDrawSize().X : TimeToPixelConverter.TimeToPixel(SegmentRange.GetUpperBoundValue());
+	float PixelStartX = SegmentRange.GetLowerBound().IsOpen() ? 0.f : TimeToPixelConverter.FrameToPixel(MovieScene::DiscreteInclusiveLower(SegmentRange));
+	float PixelEndX   = SegmentRange.GetUpperBound().IsOpen() ? AllottedGeometry.GetDrawSize().X : TimeToPixelConverter.FrameToPixel(MovieScene::DiscreteExclusiveUpper(SegmentRange));
 
 	const float MinSectionWidth = 0.f;
 	float SectionLength = FMath::Max(MinSectionWidth, PixelEndX - PixelStartX);
@@ -145,7 +150,7 @@ FGeometry SSequencerDebugVisualizer::GetSegmentGeometry(const FGeometry& Allotte
 		);
 }
 
-EVisibility SSequencerDebugVisualizer::GetSegmentVisibility(TRange<float> Range) const
+EVisibility SSequencerDebugVisualizer::GetSegmentVisibility(TRange<double> Range) const
 {
 	return ViewRange.Get().Overlaps(Range) ? EVisibility::Visible : EVisibility::Collapsed;
 }
@@ -209,7 +214,15 @@ void SSequencerDebugVisualizer::Tick( const FGeometry& AllottedGeometry, const d
 
 void SSequencerDebugVisualizer::OnArrangeChildren( const FGeometry& AllottedGeometry, FArrangedChildren& ArrangedChildren ) const
 {
-	FTimeToPixel TimeToPixelConverter = FTimeToPixel(AllottedGeometry, ViewRange.Get());
+	TSharedPtr<FSequencer>     Sequencer      = WeakSequencer.Pin();
+	const UMovieSceneSequence* ActiveSequence = Sequencer.IsValid() ? Sequencer->GetFocusedMovieSceneSequence() : nullptr;
+
+	if (!ActiveSequence)
+	{
+		return;
+	}
+
+	FTimeToPixel TimeToPixelConverter = FTimeToPixel(AllottedGeometry, ViewRange.Get(), ActiveSequence->GetMovieScene()->GetFrameResolution());
 
 	for (int32 WidgetIndex = 0; WidgetIndex < Children.Num(); ++WidgetIndex)
 	{
diff --git a/Engine/Source/Editor/Sequencer/Private/SSequencerDebugVisualizer.h b/Engine/Source/Editor/Sequencer/Private/SSequencerDebugVisualizer.h
--- a/Engine/Source/Editor/Sequencer/Private/SSequencerDebugVisualizer.h
+++ b/Engine/Source/Editor/Sequencer/Private/SSequencerDebugVisualizer.h
@@ -43,7 +43,7 @@ class SSequencerDebugVisualizer : public SPanel
 {
 public:
 	SLATE_BEGIN_ARGS(SSequencerDebugVisualizer){}
-		SLATE_ATTRIBUTE( TRange<float>, ViewRange )
+		SLATE_ATTRIBUTE( TRange<double>, ViewRange )
 	SLATE_END_ARGS()
 
 	SSequencerDebugVisualizer()
@@ -67,7 +67,7 @@ protected:
 
 	FGeometry GetSegmentGeometry(const FGeometry& AllottedGeometry, const SSequencerDebugSlot& Slot, const FTimeToPixel& TimeToPixelConverter) const;
 
-	EVisibility GetSegmentVisibility(TRange<float> Range) const;
+	EVisibility GetSegmentVisibility(TRange<double> Range) const;
 
 	TSharedRef<SWidget> GetTooltipForSegment(int32 SegmentIndex) const;
 
@@ -78,7 +78,7 @@ protected:
 private:
 
 	/** The current view range */
-	TAttribute<TRange<float>> ViewRange;
+	TAttribute<TRange<double>> ViewRange;
 
 	/** All the widgets in the panel */
 	TSlotlessChildren<SSequencerDebugSlot> Children;
diff --git a/Engine/Source/Editor/Sequencer/Private/SSequencerGotoBox.cpp b/Engine/Source/Editor/Sequencer/Private/SSequencerGotoBox.cpp
--- a/Engine/Source/Editor/Sequencer/Private/SSequencerGotoBox.cpp
+++ b/Engine/Source/Editor/Sequencer/Private/SSequencerGotoBox.cpp
@@ -6,17 +6,21 @@
 #include "Framework/Application/SlateApplication.h"
 #include "EditorStyleSet.h"
 #include "Widgets/Input/SNumericEntryBox.h"
-
+#include "Widgets/Input/SButton.h"
+#include "Widgets/Layout/SSpacer.h"
 
 #define LOCTEXT_NAMESPACE "Sequencer"
 
 
-void SSequencerGotoBox::Construct(const FArguments& InArgs, const TSharedRef<FSequencer>& InSequencer, USequencerSettings& InSettings, const TSharedRef<INumericTypeInterface<float>>& InNumericTypeInterface)
+void SSequencerGotoBox::Construct(const FArguments& InArgs, const TSharedRef<FSequencer>& InSequencer, USequencerSettings& InSettings, const TSharedRef<INumericTypeInterface<double>>& InNumericTypeInterface)
 {
 	SequencerPtr = InSequencer;
 	Settings = &InSettings;
 	NumericTypeInterface = InNumericTypeInterface;
 
+	const FDockTabStyle* GenericTabStyle = &FCoreStyle::Get().GetWidgetStyle<FDockTabStyle>("Docking.Tab");
+	const FButtonStyle* const CloseButtonStyle = &GenericTabStyle->CloseButtonStyle;
+
 	ChildSlot
 	[
 		SAssignNew(Border, SBorder)
@@ -34,16 +38,32 @@ void SSequencerGotoBox::Construct(const FArguments& InArgs, const TSharedRef<FSe
 							.Text(LOCTEXT("GotoLabel", "Go to:"))
 					]
 
+				 + SHorizontalBox::Slot()
+				 	.Padding(6.0f, 0.0f, 0.0f, 0.0f)
+				 	.AutoWidth()
+				 	[
+				 		 SAssignNew(EntryBox, SNumericEntryBox<double>)
+				 		 	.MinDesiredValueWidth(64.0f)
+				 		 	.OnValueCommitted(this, &SSequencerGotoBox::HandleEntryBoxValueCommitted)
+				 		 	.TypeInterface(NumericTypeInterface)
+				 		 	.Value_Lambda([this](){ return SequencerPtr.Pin()->GetLocalTime().Time.GetFrame().Value; })
+				 	]
+
 				+ SHorizontalBox::Slot()
-					.Padding(6.0f, 0.0f, 0.0f, 0.0f)
-					.AutoWidth()
+				.Padding(3.0f, 0.0f, 0.0f, 0.0f)
+				.AutoWidth()
+				.HAlign(HAlign_Center)
+				.VAlign(VAlign_Center)
+				[
+					SNew(SButton)
+					.ButtonStyle( CloseButtonStyle )
+					.OnClicked( this, &SSequencerGotoBox::OnCloseButtonClicked )
+					.ContentPadding( 0 )
 					[
-						SAssignNew(EntryBox, SNumericEntryBox<float>)
-							.MinDesiredValueWidth(64.0f)
-							.OnValueCommitted(this, &SSequencerGotoBox::HandleEntryBoxValueCommitted)
-							.TypeInterface(NumericTypeInterface)
-							.Value_Lambda([this](){ return SequencerPtr.Pin()->GetLocalTime(); })
+						SNew(SSpacer)
+						.Size( CloseButtonStyle->Normal.ImageSize )
 					]
+				]
 			]
 	];
 }
@@ -67,7 +87,7 @@ void SSequencerGotoBox::ToggleVisibility()
 }
 
 
-void SSequencerGotoBox::HandleEntryBoxValueCommitted(float Value, ETextCommit::Type CommitType)
+void SSequencerGotoBox::HandleEntryBoxValueCommitted(double Value, ETextCommit::Type CommitType)
 {
 	if (CommitType != ETextCommit::OnEnter)
 	{
@@ -78,18 +98,27 @@ void SSequencerGotoBox::HandleEntryBoxValueCommitted(float Value, ETextCommit::T
 
 	TSharedPtr<FSequencer> Sequencer = SequencerPtr.Pin();
 
-	// scroll view range if new value is not visible
+	// scroll view range if new value is not visible. 
 	const FAnimatedRange ViewRange = Sequencer->GetViewRange();
 
-	if (!ViewRange.Contains(Value))
+	// View range is in seconds, so we need to convert from frame numbers back into seconds.
+	FFrameTime ValueAsFrameTime = FFrameTime::FromDecimal(Value);
+	double ValueAsSeconds = Sequencer->GetFocusedFrameResolution().AsSeconds(ValueAsFrameTime);
+
+	if (!ViewRange.Contains(ValueAsSeconds))
 	{
-		const float HalfViewWidth = 0.5f * ViewRange.Size<float>();
-		const TRange<float> NewRange = TRange<float>(Value - HalfViewWidth, Value + HalfViewWidth);
+		const double HalfViewWidth = 0.5 * ViewRange.Size<double>();
+		const TRange<double> NewRange = TRange<double>(ValueAsSeconds - HalfViewWidth, ValueAsSeconds + HalfViewWidth);
 		Sequencer->SetViewRange(NewRange);
 	}
-
-	Sequencer->SetLocalTimeDirectly(Value);
+	Sequencer->SetLocalTimeDirectly(ValueAsFrameTime);
 }
 
+FReply SSequencerGotoBox::OnCloseButtonClicked()
+{
+	ToggleVisibility();
+
+	return FReply::Handled();
+}
 
 #undef LOCTEXT_NAMESPACE
diff --git a/Engine/Source/Editor/Sequencer/Private/SSequencerGotoBox.h b/Engine/Source/Editor/Sequencer/Private/SSequencerGotoBox.h
--- a/Engine/Source/Editor/Sequencer/Private/SSequencerGotoBox.h
+++ b/Engine/Source/Editor/Sequencer/Private/SSequencerGotoBox.h
@@ -7,6 +7,7 @@
 #include "Widgets/SCompoundWidget.h"
 #include "Widgets/Input/NumericTypeInterface.h"
 #include "Sequencer.h"
+#include "QualifiedFrameTime.h"
 
 class USequencerSettings;
 
@@ -19,7 +20,7 @@ public:
 	SLATE_END_ARGS()
 
 	/** Construct the widget. */
-	void Construct(const FArguments& InArgs, const TSharedRef<FSequencer>& InSequencer, USequencerSettings& InSettings, const TSharedRef<INumericTypeInterface<float>>& InNumericTypeInterface);
+	void Construct(const FArguments& InArgs, const TSharedRef<FSequencer>& InSequencer, USequencerSettings& InSettings, const TSharedRef<INumericTypeInterface<double>>& InNumericTypeInterface);
 
 	/** Toggle the widget's visibility. */
 	void ToggleVisibility();
@@ -27,7 +28,10 @@ public:
 private:
 
 	/** Callback for committed values in the entry box. */
-	void HandleEntryBoxValueCommitted(float Value, ETextCommit::Type CommitType);
+	void HandleEntryBoxValueCommitted(double Value, ETextCommit::Type CommitType);
+
+	/** Callback for when the close button is clicked. */
+	FReply OnCloseButtonClicked();
 
 private:
 
@@ -35,13 +39,13 @@ private:
 	TSharedPtr<SWidget> Border;
 
 	/** The entry box widget. */
-	TSharedPtr<SNumericEntryBox<float>> EntryBox;
+	TSharedPtr<SNumericEntryBox<double>> EntryBox;
 
 	/** The widget that focused prior to this goto box. */
 	TWeakPtr<SWidget> LastFocusedWidget;
 
 	/** Numeric type interface used for converting parsing and generating strings from numbers. */
-	TSharedPtr<INumericTypeInterface<float>> NumericTypeInterface;
+	TSharedPtr<INumericTypeInterface<double>> NumericTypeInterface;
 
 	/** The main sequencer interface. */
 	TWeakPtr<FSequencer> SequencerPtr;
diff --git a/Engine/Source/Editor/Sequencer/Private/SSequencerLabelBrowser.cpp b/Engine/Source/Editor/Sequencer/Private/SSequencerLabelBrowser.cpp
--- a/Engine/Source/Editor/Sequencer/Private/SSequencerLabelBrowser.cpp
+++ b/Engine/Source/Editor/Sequencer/Private/SSequencerLabelBrowser.cpp
@@ -90,6 +90,11 @@ void SSequencerLabelBrowser::ReloadLabelList(bool FullyReload)
 			TArray<FString> Strings;
 			Label.ParseIntoArray(Strings, TEXT("."), true);
 
+			if (Strings.Num() == 0)
+			{
+				continue;
+			}
+
 			TSharedRef<FSequencerLabelTreeNode> NewNode = MakeShareable(
 				new FSequencerLabelTreeNode(Label, FText::FromString(Strings.Last())));
 
@@ -276,8 +281,16 @@ void SSequencerLabelBrowser::HandleRenameLabelMenuEntryExecute()
 bool SSequencerLabelBrowser::HandleRenameLabelMenuEntryCanExecute() const
 {
 	TArray<TSharedPtr<FSequencerLabelTreeNode>> SelectedItems;
-	return (LabelTreeView->GetSelectedItems(SelectedItems) > 0);
-}
+	if (LabelTreeView->GetSelectedItems(SelectedItems) > 0)
+	{
+		auto ListRow = StaticCastSharedPtr<SSequencerLabelListRow>(LabelTreeView->WidgetFromItem(SelectedItems[0]));
 
+		if (ListRow.IsValid())
+		{
+			return ListRow->CanEnterRenameMode();
+		}
+	}
+	return false;
+}
 
 #undef LOCTEXT_NAMESPACE
diff --git a/Engine/Source/Editor/Sequencer/Private/SSequencerLabelListRow.h b/Engine/Source/Editor/Sequencer/Private/SSequencerLabelListRow.h
--- a/Engine/Source/Editor/Sequencer/Private/SSequencerLabelListRow.h
+++ b/Engine/Source/Editor/Sequencer/Private/SSequencerLabelListRow.h
@@ -106,7 +106,7 @@ public:
 						.VAlign(VAlign_Center)
 						[
 							SAssignNew(EditableLabel, SEditableLabel)
-								.CanEdit(this, &SSequencerLabelListRow::HandleFolderNameCanEdit)
+								.CanEdit(this, &SSequencerLabelListRow::CanEnterRenameMode)
 								.OnTextChanged(this, &SSequencerLabelListRow::HandleFolderNameTextChanged)
 								.Text(
 									Node->Label.IsEmpty()
@@ -125,6 +125,12 @@ public:
 		EditableLabel->EnterTextMode();
 	}
 
+	/** Can this label be edited? */
+	bool CanEnterRenameMode() const
+	{
+		return !Node->Label.IsEmpty();
+	}
+
 private:
 
 	EVisibility HandleEditIconVisibility() const
@@ -150,11 +156,6 @@ private:
 		return FLinearColor::Gray;
 	}
 
-	bool HandleFolderNameCanEdit() const
-	{
-		return !Node->Label.IsEmpty();
-	}
-
 	void HandleFolderNameTextChanged(const FText& NewLabel)
 	{
 		FString NewLabelString = NewLabel.ToString();
diff --git a/Engine/Source/Editor/Sequencer/Private/SSequencerPlayRateCombo.cpp b/Engine/Source/Editor/Sequencer/Private/SSequencerPlayRateCombo.cpp
new file mode 100644
--- /dev/null
+++ b/Engine/Source/Editor/Sequencer/Private/SSequencerPlayRateCombo.cpp
@@ -0,0 +1,448 @@
+// Copyright 1998-2018 Epic Games, Inc. All Rights Reserved.
+
+#include "SSequencerPlayRateCombo.h"
+#include "Sequencer.h"
+#include "SSequencer.h"
+#include "Widgets/SOverlay.h"
+#include "Widgets/Layout/SBox.h"
+#include "Widgets/Text/STextBlock.h"
+#include "Widgets/Input/SComboButton.h"
+#include "Widgets/Images/SImage.h"
+#include "Framework/MultiBox/MultiBoxBuilder.h"
+#include "MovieSceneSequence.h"
+#include "MovieScene.h"
+#include "ScopedTransaction.h"
+#include "Widgets/SFrameRateEntryBox.h"
+#include "EditorFontGlyphs.h"
+#include <Timecode.h>
+
+#define LOCTEXT_NAMESPACE "SSequencerPlayRateCombo"
+
+void SSequencerPlayRateCombo::Construct(const FArguments& InArgs, TWeakPtr<FSequencer> InWeakSequencer, TWeakPtr<SSequencer> InWeakSequencerWidget)
+{
+	WeakSequencer = InWeakSequencer;
+	WeakSequencerWidget = InWeakSequencerWidget;
+
+	FName BlockStyle = EMultiBlockLocation::ToName(ISlateStyle::Join( InArgs._StyleName, ".Button" ), InArgs._BlockLocation);
+	FName ColorStyle = ISlateStyle::Join( InArgs._StyleName, ".SToolBarComboButtonBlock.ComboButton.Color" );
+
+	SetToolTipText(MakeAttributeSP(this, &SSequencerPlayRateCombo::GetToolTipText));
+
+	ChildSlot
+	.VAlign(VAlign_Fill)
+	[
+		SNew(SComboButton)
+		.ContentPadding(FMargin(2.f, 0.f))
+		.VAlign(VAlign_Fill)
+		.ButtonStyle( InArgs._StyleSet, BlockStyle )
+		.ForegroundColor( InArgs._StyleSet->GetSlateColor( ColorStyle ) )
+		.OnGetMenuContent(this, &SSequencerPlayRateCombo::OnCreateMenu)
+		.ButtonContent()
+		[
+			SNew(SHorizontalBox)
+
+			+ SHorizontalBox::Slot()
+			.HAlign(HAlign_Center)
+			.VAlign(VAlign_Fill)
+			.AutoWidth()
+			[
+				SNew(STextBlock)
+				.Text(this, &SSequencerPlayRateCombo::GetFrameRateText)
+				.TextStyle( InArgs._StyleSet, ISlateStyle::Join( InArgs._StyleName, ".Label" ) )
+			]
+
+			+ SHorizontalBox::Slot()
+			.Padding(3.f, 0.f)
+			.AutoWidth()
+			.VAlign(VAlign_Center)
+			[
+				SNew(STextBlock)
+				.Visibility(this, &SSequencerPlayRateCombo::GetFrameLockedVisibility)
+				.TextStyle( InArgs._StyleSet, ISlateStyle::Join( InArgs._StyleName, ".Label" ) )
+				.Font(FEditorStyle::Get().GetFontStyle("FontAwesome.11"))
+				.Text(FEditorFontGlyphs::Lock)
+			]
+
+			+ SHorizontalBox::Slot()
+			.Padding(0.f, 0.f, 3.f, 0.f)
+			.AutoWidth()
+			.VAlign(VAlign_Center)
+			[
+				SNew(STextBlock)
+				.ToolTipText(this, &SSequencerPlayRateCombo::GetFrameRateErrorDescription)
+				.Visibility(this, &SSequencerPlayRateCombo::GetFrameRateErrorVisibility)
+				.TextStyle( InArgs._StyleSet, ISlateStyle::Join( InArgs._StyleName, ".Label" ) )
+				.Font(FEditorStyle::Get().GetFontStyle("FontAwesome.11"))
+				.Text(FEditorFontGlyphs::Exclamation_Triangle)
+			]
+		]
+	];
+
+	ChildSlot.Padding(InArgs._StyleSet->GetMargin(ISlateStyle::Join( InArgs._StyleName, ".SToolBarComboButtonBlock.Padding" )));
+}
+
+EVisibility SSequencerPlayRateCombo::GetFrameLockedVisibility() const
+{
+	TSharedPtr<FSequencer> Sequencer  = WeakSequencer.Pin();
+	UMovieSceneSequence*   Sequence   = Sequencer.IsValid() ? Sequencer->GetFocusedMovieSceneSequence() : nullptr;
+	return Sequence && Sequence->GetMovieScene()->GetEvaluationType() == EMovieSceneEvaluationType::FrameLocked ? EVisibility::HitTestInvisible : EVisibility::Collapsed;
+}
+
+EVisibility SSequencerPlayRateCombo::GetFrameRateErrorVisibility() const
+{
+	TSharedPtr<FSequencer> Sequencer  = WeakSequencer.Pin();
+	if (!Sequencer.IsValid() || Sequencer->GetFocusedPlayRate().IsMultipleOf(Sequencer->GetFocusedFrameResolution()))
+	{
+		return EVisibility::Collapsed;
+	}
+	return EVisibility::HitTestInvisible;
+}
+
+FText SSequencerPlayRateCombo::GetFrameRateErrorDescription() const
+{
+	TSharedPtr<FSequencer> Sequencer  = WeakSequencer.Pin();
+	if (Sequencer.IsValid())
+	{
+		FFrameRate PlayRate = Sequencer->GetFocusedPlayRate();
+		FFrameRate FrameResolution = Sequencer->GetFocusedFrameResolution();
+
+		const FCommonFrameRateInfo* PlayRateInfo        = FCommonFrameRates::Find(PlayRate);
+		const FCommonFrameRateInfo* FrameResolutionInfo = FCommonFrameRates::Find(FrameResolution);
+
+		FText PlayRateText        = PlayRateInfo        ? PlayRateInfo->DisplayName        : FText::Format(LOCTEXT("PlayRateFormat", "{0} fps"), PlayRate.AsDecimal());
+		FText FrameResolutionText = FrameResolutionInfo ? FrameResolutionInfo->DisplayName : FText::Format(LOCTEXT("FrameResolutionFormat", "{0} ticks every second"), FrameResolution.AsDecimal());
+
+		return FText::Format(LOCTEXT("FrameRateErrorDescription", "The current display rate of {0} is incompatible with this sequence's tick resolution of {1} ticks per second."), PlayRateText, FrameResolutionText);
+	}
+
+	return FText();
+}
+
+TSharedRef<SWidget> SSequencerPlayRateCombo::OnCreateMenu()
+{
+	TSharedPtr<FSequencer> Sequencer       = WeakSequencer.Pin();
+	TSharedPtr<SSequencer> SequencerWidget = WeakSequencerWidget.Pin();
+	if (!Sequencer.IsValid() || !SequencerWidget.IsValid())
+	{
+		return SNullWidget::NullWidget;
+	}
+
+	FMenuBuilder MenuBuilder(true, nullptr);
+
+	FFrameRate FrameResolution = Sequencer->GetFocusedFrameResolution();
+
+	TArray<FCommonFrameRateInfo> CompatibleRates;
+	for (const FCommonFrameRateInfo& Info : FCommonFrameRates::GetAll())
+	{
+		if (Info.FrameRate.IsMultipleOf(FrameResolution))
+		{
+			CompatibleRates.Add(Info);
+		}
+	}
+
+	CompatibleRates.Sort(
+		[=](const FCommonFrameRateInfo& A, const FCommonFrameRateInfo& B)
+		{
+			return A.FrameRate.AsDecimal() < B.FrameRate.AsDecimal();
+		}
+	);
+
+	MenuBuilder.BeginSection(NAME_None, LOCTEXT("RecommendedRates", "Sequence Display Rate"));
+	{
+		for (const FCommonFrameRateInfo& Info : CompatibleRates)
+		{
+			AddMenuEntry(MenuBuilder, Info);
+		}
+
+		MenuBuilder.AddWidget(
+			SNew(SBox)
+			.HAlign(HAlign_Right)
+			.MaxDesiredWidth(100.f)
+			[
+				SNew(SFrameRateEntryBox)
+				.Value(this, &SSequencerPlayRateCombo::GetPlaybackFrameRate)
+				.OnValueChanged(this, &SSequencerPlayRateCombo::SetPlaybackFrameRate)
+			],
+			LOCTEXT("CustomFramerateDisplayLabel", "Custom")
+		);
+
+		if (CompatibleRates.Num() != FCommonFrameRates::GetAll().Num())
+		{
+			MenuBuilder.AddSubMenu(
+				LOCTEXT("IncompatibleRates", "Incompatible Rates"),
+				FText::Format(LOCTEXT("IncompatibleRates_Description", "Choose from a list of display rates that are incompatible with a resolution of {0} ticks per second"), FrameResolution.AsDecimal()),
+				FNewMenuDelegate::CreateSP(this, &SSequencerPlayRateCombo::PopulateIncompatibleRatesMenu)
+				);
+		}
+	}
+	MenuBuilder.EndSection();
+
+	MenuBuilder.AddMenuSeparator();
+
+	MenuBuilder.AddSubMenu(
+		LOCTEXT("ShowTimesAs", "Show Time As"),
+		LOCTEXT("ShowTimesAs_Description", "Change how to display times in Sequencer"),
+		FNewMenuDelegate::CreateSP(this, &SSequencerPlayRateCombo::PopulateTimeDisplayMenu)
+	);
+
+	if (Sequencer->GetRootMovieSceneSequence() == Sequencer->GetFocusedMovieSceneSequence())
+	{
+		MenuBuilder.AddSubMenu(
+			LOCTEXT("ClockSource", "Clock Source"),
+			LOCTEXT("ClockSource_Description", "Change which clock should be used when playing back this sequence"),
+			FNewMenuDelegate::CreateSP(this, &SSequencerPlayRateCombo::PopulateClockSourceMenu)
+		);
+	}
+
+	MenuBuilder.AddMenuEntry(
+		LOCTEXT("LockPlayback", "Lock to Display Rate at Runtime"),
+		LOCTEXT("LockPlayback_Description", "When enabled, causes all runtime evaluation and the engine FPS to be locked to the current display frame rate"),
+		FSlateIcon(),
+		FUIAction(
+			FExecuteAction::CreateSP(this, &SSequencerPlayRateCombo::OnToggleFrameLocked),
+			FCanExecuteAction(),
+			FGetActionCheckState::CreateSP(this, &SSequencerPlayRateCombo::OnGetFrameLockedCheckState)
+		),
+		NAME_None,
+		EUserInterfaceActionType::ToggleButton
+	);
+
+	MenuBuilder.AddMenuEntry(
+		LOCTEXT("AdvancedOptions", "Advanced Options"),
+		LOCTEXT("AdvancedOptions_Description", "Open advanced time-related properties for this sequence"),
+		FSlateIcon(),
+		FUIAction(
+			FExecuteAction::CreateSP(SequencerWidget.Get(), &SSequencer::ShowFrameResolutionOverlay)
+		)
+	);
+
+	return MenuBuilder.MakeWidget();
+}
+
+void SSequencerPlayRateCombo::PopulateIncompatibleRatesMenu(FMenuBuilder& MenuBuilder)
+{
+	TSharedPtr<FSequencer> Sequencer = WeakSequencer.Pin();
+	if (Sequencer.IsValid())
+	{
+		FFrameRate FrameResolution = Sequencer->GetFocusedFrameResolution();
+
+		TArray<FCommonFrameRateInfo> IncompatibleRates;
+		for (const FCommonFrameRateInfo& Info : FCommonFrameRates::GetAll())
+		{
+			if (!Info.FrameRate.IsMultipleOf(FrameResolution))
+			{
+				IncompatibleRates.Add(Info);
+			}
+		}
+
+		IncompatibleRates.Sort(
+			[=](const FCommonFrameRateInfo& A, const FCommonFrameRateInfo& B)
+			{
+				return A.FrameRate.AsDecimal() < B.FrameRate.AsDecimal();
+			}
+		);
+
+		for (const FCommonFrameRateInfo& Info : IncompatibleRates)
+		{
+			AddMenuEntry(MenuBuilder, Info);
+		}
+	}
+}
+
+void SSequencerPlayRateCombo::PopulateTimeDisplayMenu(FMenuBuilder& MenuBuilder)
+{
+	TSharedPtr<FSequencer> Sequencer = WeakSequencer.Pin();
+	USequencerSettings*    Settings = Sequencer.IsValid() ? Sequencer->GetSequencerSettings() : nullptr;
+	bool bSupportsDropFormatDisplay = FTimecode::IsDropFormatTimecodeSupported(Sequencer->GetFocusedPlayRate());
+
+	const UEnum* FrameNumberDisplayEnum = FindObject<UEnum>(ANY_PACKAGE, TEXT("EFrameNumberDisplayFormats"), true);
+	check(FrameNumberDisplayEnum);
+
+	if (Settings)
+	{
+		EFrameNumberDisplayFormats CurrentDisplay = Settings->GetTimeDisplayFormat();
+
+		for (int32 Index = 0; Index < FrameNumberDisplayEnum->NumEnums() - 1; Index++)
+		{
+			if (!FrameNumberDisplayEnum->HasMetaData(TEXT("Hidden"), Index))
+			{
+				EFrameNumberDisplayFormats Value = (EFrameNumberDisplayFormats)FrameNumberDisplayEnum->GetValueByIndex(Index);
+
+				// Don't show Drop Frame Timecode when they're in a format that doesn't support it.
+				if (Value == EFrameNumberDisplayFormats::DropFrameTimecode && !bSupportsDropFormatDisplay)
+					continue;
+
+				bool bIsSet = CurrentDisplay == Value;
+				MenuBuilder.AddMenuEntry(
+					FrameNumberDisplayEnum->GetDisplayNameTextByIndex(Index),
+					FrameNumberDisplayEnum->GetToolTipTextByIndex(Index),
+					FSlateIcon(),
+					FUIAction(
+						FExecuteAction::CreateUObject(Settings, &USequencerSettings::SetTimeDisplayFormat, Value),
+						FCanExecuteAction(),
+						FIsActionChecked::CreateLambda([=]{ return CurrentDisplay == Value; })
+					),
+					NAME_None,
+					EUserInterfaceActionType::RadioButton
+				);
+			}
+		}
+	}
+}
+
+void SSequencerPlayRateCombo::PopulateClockSourceMenu(FMenuBuilder& MenuBuilder)
+{
+	TSharedPtr<FSequencer> Sequencer    = WeakSequencer.Pin();
+	UMovieSceneSequence*   RootSequence = Sequencer.IsValid() ? Sequencer->GetRootMovieSceneSequence() : nullptr;
+
+	const UEnum* ClockSourceEnum = FindObject<UEnum>(ANY_PACKAGE, TEXT("EUpdateClockSource"), true);
+
+	check(ClockSourceEnum);
+
+	if (RootSequence)
+	{
+		EUpdateClockSource CurrentDisplay = RootSequence->GetMovieScene()->GetClockSource();
+
+		for (int32 Index = 0; Index < ClockSourceEnum->NumEnums() - 1; Index++)
+		{
+			if (!ClockSourceEnum->HasMetaData(TEXT("Hidden"), Index))
+			{
+				EUpdateClockSource Value = (EUpdateClockSource)ClockSourceEnum->GetValueByIndex(Index);
+
+				bool bIsSet = CurrentDisplay == Value;
+				MenuBuilder.AddMenuEntry(
+					ClockSourceEnum->GetDisplayNameTextByIndex(Index),
+					ClockSourceEnum->GetToolTipTextByIndex(Index),
+					FSlateIcon(),
+					FUIAction(
+						FExecuteAction::CreateSP(this, &SSequencerPlayRateCombo::SetClockSource, Value),
+						FCanExecuteAction(),
+						FIsActionChecked::CreateLambda([=]{ return CurrentDisplay == Value; })
+					),
+					NAME_None,
+					EUserInterfaceActionType::RadioButton
+				);
+			}
+		}
+	}
+}
+
+void SSequencerPlayRateCombo::AddMenuEntry(FMenuBuilder& MenuBuilder, const FCommonFrameRateInfo& Info)
+{
+	MenuBuilder.AddMenuEntry(
+		Info.DisplayName,
+		Info.Description,
+		FSlateIcon(),
+		FUIAction(
+			FExecuteAction::CreateSP(this, &SSequencerPlayRateCombo::SetPlaybackFrameRate, Info.FrameRate),
+			FCanExecuteAction(),
+			FIsActionChecked::CreateSP(this, &SSequencerPlayRateCombo::IsSamePlaybackFrameRate, Info.FrameRate)
+		),
+		NAME_None,
+		EUserInterfaceActionType::RadioButton
+	);
+}
+
+void SSequencerPlayRateCombo::SetClockSource(EUpdateClockSource NewClockSource)
+{
+	TSharedPtr<FSequencer> Sequencer    = WeakSequencer.Pin();
+	UMovieSceneSequence*   RootSequence = Sequencer.IsValid() ? Sequencer->GetRootMovieSceneSequence() : nullptr;
+	if (RootSequence)
+	{
+		FScopedTransaction ScopedTransaction(LOCTEXT("SetClockSource", "Set Clock Source"));
+
+		RootSequence->GetMovieScene()->Modify();
+		RootSequence->GetMovieScene()->SetClockSource(NewClockSource);
+
+		Sequencer->ResetTimeController();
+	}
+}
+
+void SSequencerPlayRateCombo::SetPlaybackFrameRate(FFrameRate InFrameRate)
+{
+	TSharedPtr<FSequencer> Sequencer = WeakSequencer.Pin();
+	UMovieSceneSequence* FocusedSequence = Sequencer.IsValid() ? Sequencer->GetFocusedMovieSceneSequence() : nullptr;
+	if (FocusedSequence)
+	{
+		FScopedTransaction ScopedTransaction(FText::Format(LOCTEXT("SetDisplayRate", "Set Display Rate to {0}"), InFrameRate.ToPrettyText()));
+
+		FocusedSequence->GetMovieScene()->Modify();
+		FocusedSequence->GetMovieScene()->SetPlaybackFrameRate(InFrameRate);
+	}
+}
+
+FFrameRate SSequencerPlayRateCombo::GetPlaybackFrameRate() const
+{
+	TSharedPtr<FSequencer> Sequencer = WeakSequencer.Pin();
+	return Sequencer.IsValid() ? Sequencer->GetFocusedPlayRate() : FFrameRate();
+}
+
+bool SSequencerPlayRateCombo::IsSamePlaybackFrameRate(FFrameRate InFrameRate) const
+{
+	return GetPlaybackFrameRate() == InFrameRate;
+}
+
+FText SSequencerPlayRateCombo::GetFrameRateText() const
+{
+	TSharedPtr<FSequencer> Sequencer = WeakSequencer.Pin();
+
+	return Sequencer.IsValid() ? Sequencer->GetFocusedPlayRate().ToPrettyText() : FText();
+}
+
+FText SSequencerPlayRateCombo::GetToolTipText() const
+{
+	TSharedPtr<FSequencer> Sequencer         = WeakSequencer.Pin();
+	UMovieSceneSequence*   FocusedSequence   = Sequencer.IsValid() ? Sequencer->GetFocusedMovieSceneSequence() : nullptr;
+	UMovieScene*           FocusedMovieScene = FocusedSequence ? FocusedSequence->GetMovieScene() : nullptr;
+
+	if (FocusedMovieScene)
+	{
+		FFrameRate PlayRate        = FocusedMovieScene->GetPlaybackFrameRate();
+		FFrameRate FrameResolution = FocusedMovieScene->GetFrameResolution();
+
+		const FCommonFrameRateInfo* PlayRateInfo        = FCommonFrameRates::Find(PlayRate);
+		const FCommonFrameRateInfo* FrameResolutionInfo = FCommonFrameRates::Find(FrameResolution);
+
+		FText PlayRateText        = PlayRateInfo        ? PlayRateInfo->DisplayName        : FText::Format(LOCTEXT("PlayRateFormat", "{0} fps"), PlayRate.AsDecimal());
+		FText FrameResolutionText = FrameResolutionInfo ? FrameResolutionInfo->DisplayName : FText::Format(LOCTEXT("FrameResolutionFormat", "{0} ticks every second"), FrameResolution.AsDecimal());
+
+		return FocusedMovieScene->GetEvaluationType() == EMovieSceneEvaluationType::FrameLocked
+			? FText::Format(LOCTEXT("ToolTip_Format", "This sequence is locked at runtime to {0} and uses an underlying tick resolution of {1}."), PlayRateText, FrameResolutionText)
+			: FText::Format(LOCTEXT("ToolTip_Format", "This sequence is being presented as {0} and uses an underlying tick resolution of {1}."), PlayRateText, FrameResolutionText);
+	}
+
+	return FText();
+}
+
+
+void SSequencerPlayRateCombo::OnToggleFrameLocked()
+{
+	TSharedPtr<FSequencer> Sequencer         = WeakSequencer.Pin();
+	UMovieSceneSequence*   FocusedSequence   = Sequencer.IsValid() ? Sequencer->GetFocusedMovieSceneSequence() : nullptr;
+	UMovieScene*           FocusedMovieScene = FocusedSequence ? FocusedSequence->GetMovieScene() : nullptr;
+
+	if (FocusedMovieScene)
+	{
+		EMovieSceneEvaluationType NewType = FocusedMovieScene->GetEvaluationType() == EMovieSceneEvaluationType::WithSubFrames
+			? EMovieSceneEvaluationType::FrameLocked
+			: EMovieSceneEvaluationType::WithSubFrames;
+
+		FScopedTransaction ScopedTransaction(NewType == EMovieSceneEvaluationType::FrameLocked
+			? LOCTEXT("FrameLockedTransaction",   "Lock to Display Rate at Runtime")
+			: LOCTEXT("WithSubFramesTransaction", "Unlock to runtime frame rate"));
+
+		FocusedMovieScene->Modify();
+		FocusedMovieScene->SetEvaluationType(NewType);
+	}
+}
+
+ECheckBoxState SSequencerPlayRateCombo::OnGetFrameLockedCheckState() const
+{
+	TSharedPtr<FSequencer> Sequencer         = WeakSequencer.Pin();
+	UMovieSceneSequence*   FocusedSequence   = Sequencer.IsValid() ? Sequencer->GetFocusedMovieSceneSequence() : nullptr;
+	UMovieScene*           FocusedMovieScene = FocusedSequence ? FocusedSequence->GetMovieScene() : nullptr;
+
+	return FocusedMovieScene && FocusedMovieScene->GetEvaluationType() == EMovieSceneEvaluationType::FrameLocked ? ECheckBoxState::Checked : ECheckBoxState::Unchecked;
+}
+
+#undef LOCTEXT_NAMESPACE
\ No newline at end of file
