Copying //UE4/Dev-Sequencer to //UE4/Dev-Main (Source: //UE4/Dev-Sequencer @ 3946692)

diff --git a/Engine/Source/Editor/Sequencer/Private/SSequencerSection.cpp b/Engine/Source/Editor/Sequencer/Private/SSequencerSection.cpp
--- a/Engine/Source/Editor/Sequencer/Private/SSequencerSection.cpp
+++ b/Engine/Source/Editor/Sequencer/Private/SSequencerSection.cpp
@@ -4,7 +4,6 @@
 #include "Rendering/DrawElements.h"
 #include "EditorStyleSet.h"
 #include "SequencerSelectionPreview.h"
-#include "GroupedKeyArea.h"
 #include "SequencerSettings.h"
 #include "Editor.h"
 #include "Sequencer.h"
@@ -16,10 +15,50 @@
 #include "SequencerHotspots.h"
 #include "Widgets/SOverlay.h"
 #include "SequencerObjectBindingNode.h"
+#include "MovieScene.h"
 #include "Fonts/FontCache.h"
 #include "Framework/Application/SlateApplication.h"
+#include "KeyDrawParams.h"
+#include "MovieSceneTimeHelpers.h"
 
-double SSequencerSection::SelectionThrobEndTime = 0;
+double SSequencerSection::SectionSelectionThrobEndTime = 0;
+double SSequencerSection::KeySelectionThrobEndTime = 0;
+
+const FSectionLayoutElement& SSequencerSection::FLayoutElementKeyFuncs::GetSetKey(const TPair<FSectionLayoutElement, TArray<FSequencerCachedKeys, TInlineAllocator<1>>>& Element)
+{
+	return Element.Key;
+}
+
+bool SSequencerSection::FLayoutElementKeyFuncs::Matches(const FSectionLayoutElement& A, const FSectionLayoutElement& B)
+{
+	if (A.GetDisplayNode() != B.GetDisplayNode())
+	{
+		return false;
+	}
+	TArrayView<const TSharedRef<IKeyArea>> KeyAreasA = A.GetKeyAreas(), KeyAreasB = B.GetKeyAreas();
+	if (KeyAreasA.Num() != KeyAreasB.Num())
+	{
+		return false;
+	}
+	for (int32 Index = 0; Index < KeyAreasA.Num(); ++Index)
+	{
+		if (KeyAreasA[Index] != KeyAreasB[Index])
+		{
+			return false;
+		}
+	}
+	return true;
+}
+
+uint32 SSequencerSection::FLayoutElementKeyFuncs::GetKeyHash(const FSectionLayoutElement& Key)
+{
+	uint32 Hash = GetTypeHash(Key.GetDisplayNode()) ;
+	for (TSharedRef<IKeyArea> KeyArea : Key.GetKeyAreas())
+	{
+		Hash = HashCombine(GetTypeHash(KeyArea), Hash);
+	}
+	return Hash;
+}
 
 /** A point on an easing curve used for rendering */
 struct FEasingCurvePoint
@@ -32,16 +71,25 @@ struct FEasingCurvePoint
 	FLinearColor Color;
 };
 
+FTimeToPixel ConstructTimeConverterForSection(const FGeometry& InSectionGeometry, const UMovieSceneSection& InSection, FSequencer& InSequencer)
+{
+	TRange<double> ViewRange       = InSequencer.GetViewRange();
+
+	FFrameRate     FrameResolution = InSection.GetTypedOuter<UMovieScene>()->GetFrameResolution();
+	double         LowerTime       = InSection.HasStartFrame() ? InSection.GetInclusiveStartFrame() / FrameResolution : ViewRange.GetLowerBoundValue();
+	double         UpperTime       = InSection.HasEndFrame()   ? InSection.GetExclusiveEndFrame()   / FrameResolution : ViewRange.GetUpperBoundValue();
+
+	return FTimeToPixel(InSectionGeometry, TRange<double>(LowerTime, UpperTime), FrameResolution);
+}
+
+
 struct FSequencerSectionPainterImpl : FSequencerSectionPainter
 {
 	FSequencerSectionPainterImpl(FSequencer& InSequencer, UMovieSceneSection& InSection, FSlateWindowElementList& _OutDrawElements, const FGeometry& InSectionGeometry, const SSequencerSection& InSectionWidget)
 		: FSequencerSectionPainter(_OutDrawElements, InSectionGeometry, InSection)
 		, Sequencer(InSequencer)
 		, SectionWidget(InSectionWidget)
-		, TimeToPixelConverter(InSection.IsInfinite() ?
-			FTimeToPixel(SectionGeometry, Sequencer.GetViewRange()) :
-			FTimeToPixel(SectionGeometry, TRange<float>(InSection.GetStartTime(), InSection.GetEndTime()))
-			)
+		, TimeToPixelConverter(ConstructTimeConverterForSection(SectionGeometry, InSection, Sequencer))
 	{
 		CalculateSelectionColor();
 
@@ -56,9 +104,9 @@ struct FSequencerSectionPainterImpl : FSequencerSectionPainter
 	FLinearColor GetFinalTintColor(const FLinearColor& Tint) const
 	{
 		FLinearColor FinalTint = FSequencerSectionPainter::BlendColor(Tint);
-		if (bIsHighlighted && !Section.IsInfinite())
+		if (bIsHighlighted && Section.GetRange() != TRange<FFrameNumber>::All())
 		{
-			float Lum = FinalTint.ComputeLuminance()* 0.2f;
+			float Lum = FinalTint.ComputeLuminance() * 0.2f;
 			FinalTint = FinalTint + FLinearColor(Lum, Lum, Lum, 0.f);
 		}
 		return FinalTint;
@@ -74,15 +122,18 @@ struct FSequencerSectionPainterImpl : FSequencerSectionPainter
 		static const FSlateBrush* SectionBackgroundTintBrush = FEditorStyle::GetBrush("Sequencer.Section.BackgroundTint");
 		static const FSlateBrush* SelectedSectionOverlay = FEditorStyle::GetBrush("Sequencer.Section.SelectedSectionOverlay");
 
-		FGeometry InfiniteGeometry = SectionGeometry.MakeChild(FVector2D(-100.f, 0.f), FVector2D(SectionGeometry.GetLocalSize().X + 200.f, SectionGeometry.GetLocalSize().Y));
-
 		FLinearColor FinalTint = GetFinalTintColor(Tint);
 
-		FPaintGeometry PaintGeometry = Section.IsInfinite() ?
-				InfiniteGeometry.ToPaintGeometry() :
-				SectionGeometry.ToPaintGeometry();
+		// Offset lower bounds and size for infinte sections so we don't draw the rounded border on the visible area
+		const float InfiniteLowerOffset = Section.HasStartFrame() ? 0.f : 100.f;
+		const float InfiniteSizeOffset  = InfiniteLowerOffset + (Section.HasEndFrame() ? 0.f : 100.f);
 
-		if (!Section.IsInfinite() && Sequencer.GetSettings()->ShouldShowPrePostRoll())
+		FPaintGeometry PaintGeometry = SectionGeometry.ToPaintGeometry(
+			SectionGeometry.GetLocalSize() + FVector2D(InfiniteSizeOffset, 0.f),
+			FSlateLayoutTransform(FVector2D(-InfiniteLowerOffset, 0.f))
+			);
+
+		if (Sequencer.GetSequencerSettings()->ShouldShowPrePostRoll())
 		{
 			TOptional<FSlateClippingState> PreviousClipState = DrawElements.GetClippingState();
 			DrawElements.PopClip();
@@ -90,45 +141,57 @@ struct FSequencerSectionPainterImpl : FSequencerSectionPainter
 			static const FSlateBrush* PreRollBrush = FEditorStyle::GetBrush("Sequencer.Section.PreRoll");
 			float BrushHeight = 16.f, BrushWidth = 10.f;
 
-			const float PreRollPx = TimeToPixelConverter.TimeToPixel(Section.GetStartTime() + Section.GetPreRollTime()) - TimeToPixelConverter.TimeToPixel(Section.GetStartTime());
-			if (PreRollPx > 0)
+			if (Section.HasStartFrame())
 			{
-				const float RoundedPreRollPx = (int(PreRollPx / BrushWidth)+1) * BrushWidth;
+				FFrameNumber SectionStartTime = Section.GetInclusiveStartFrame();
+				FFrameNumber PreRollStartTime = SectionStartTime - Section.GetPreRollFrames();
 
-				// Round up to the nearest BrushWidth size
-				FGeometry PreRollArea = SectionGeometry.MakeChild(
-					FVector2D(RoundedPreRollPx, BrushHeight),
-					FSlateLayoutTransform(FVector2D(-PreRollPx, (SectionGeometry.GetLocalSize().Y - BrushHeight)*.5f))
+				const float PreRollPx = TimeToPixelConverter.FrameToPixel(SectionStartTime) - TimeToPixelConverter.FrameToPixel(PreRollStartTime);
+				if (PreRollPx > 0)
+				{
+					const float RoundedPreRollPx = (int(PreRollPx / BrushWidth)+1) * BrushWidth;
+
+					// Round up to the nearest BrushWidth size
+					FGeometry PreRollArea = SectionGeometry.MakeChild(
+						FVector2D(RoundedPreRollPx, BrushHeight),
+						FSlateLayoutTransform(FVector2D(-PreRollPx, (SectionGeometry.GetLocalSize().Y - BrushHeight)*.5f))
+						);
+
+					FSlateDrawElement::MakeBox(
+						DrawElements,
+						LayerId,
+						PreRollArea.ToPaintGeometry(),
+						PreRollBrush,
+						DrawEffects
 					);
-
-				FSlateDrawElement::MakeBox(
-					DrawElements,
-					LayerId,
-					PreRollArea.ToPaintGeometry(),
-					PreRollBrush,
-					DrawEffects
-				);
+				}
 			}
 
-			const float PostRollPx = TimeToPixelConverter.TimeToPixel(Section.GetEndTime() + Section.GetPostRollTime()) - TimeToPixelConverter.TimeToPixel(Section.GetEndTime());
-			if (PostRollPx > 0)
+			if (Section.HasEndFrame())
 			{
-				const float RoundedPostRollPx = (int(PostRollPx / BrushWidth)+1) * BrushWidth;
-				const float Difference = RoundedPostRollPx - PostRollPx;
+				FFrameNumber SectionEndTime  = Section.GetExclusiveEndFrame();
+				FFrameNumber PostRollEndTime = SectionEndTime + Section.GetPostRollFrames();
 
-				// Slate border brushes tile UVs along +ve X, so we round the arrows to a multiple of the brush width, and offset, to ensure we don't have a partial tile visible at the end
-				FGeometry PostRollArea = SectionGeometry.MakeChild(
-					FVector2D(RoundedPostRollPx, BrushHeight),
-					FSlateLayoutTransform(FVector2D(SectionGeometry.GetLocalSize().X - Difference, (SectionGeometry.GetLocalSize().Y - BrushHeight)*.5f))
+				const float PostRollPx = TimeToPixelConverter.FrameToPixel(PostRollEndTime) - TimeToPixelConverter.FrameToPixel(SectionEndTime);
+				if (PostRollPx > 0)
+				{
+					const float RoundedPostRollPx = (int(PostRollPx / BrushWidth)+1) * BrushWidth;
+					const float Difference = RoundedPostRollPx - PostRollPx;
+
+					// Slate border brushes tile UVs along +ve X, so we round the arrows to a multiple of the brush width, and offset, to ensure we don't have a partial tile visible at the end
+					FGeometry PostRollArea = SectionGeometry.MakeChild(
+						FVector2D(RoundedPostRollPx, BrushHeight),
+						FSlateLayoutTransform(FVector2D(SectionGeometry.GetLocalSize().X - Difference, (SectionGeometry.GetLocalSize().Y - BrushHeight)*.5f))
+						);
+
+					FSlateDrawElement::MakeBox(
+						DrawElements,
+						LayerId,
+						PostRollArea.ToPaintGeometry(),
+						PreRollBrush,
+						DrawEffects
 					);
-
-				FSlateDrawElement::MakeBox(
-					DrawElements,
-					LayerId,
-					PostRollArea.ToPaintGeometry(),
-					PreRollBrush,
-					DrawEffects
-				);
+				}
 			}
 
 			if (PreviousClipState.IsSet())
@@ -150,6 +213,7 @@ struct FSequencerSectionPainterImpl : FSequencerSectionPainter
 				SectionBackgroundBrush,
 				DrawEffects
 			);
+			++LayerId;
 
 			if (PreviousClipState.IsSet())
 			{
@@ -166,6 +230,7 @@ struct FSequencerSectionPainterImpl : FSequencerSectionPainter
 			DrawEffects,
 			FinalTint
 		);
+		++LayerId;
 
 		// Draw underlapping sections
 		DrawOverlaps(FinalTint);
@@ -203,7 +268,7 @@ struct FSequencerSectionPainterImpl : FSequencerSectionPainter
 	void CalculateSelectionColor()
 	{
 		// Don't draw selected if infinite
-		if (Section.IsInfinite())
+		if (Section.GetRange() == TRange<FFrameNumber>::All())
 		{
 			return;
 		}
@@ -264,7 +329,7 @@ struct FSequencerSectionPainterImpl : FSequencerSectionPainter
 			FontInfo.Size = FMath::Max(FMath::FloorToInt(FontInfo.Size - 6.f), 11);
 		}
 
-		FVector2D TextOffset = Section.IsInfinite() ? FVector2D(0.f, -1.f) : FVector2D(1.f, -1.f);
+		FVector2D TextOffset = Section.GetRange() == TRange<FFrameNumber>::All() ? FVector2D(0.f, -1.f) : FVector2D(1.f, -1.f);
 		FVector2D BottomLeft = SectionGeometry.AbsoluteToLocal(SectionClippingRect.GetBottomLeft()) + TextOffset;
 
 		FSlateDrawElement::MakeText(
@@ -328,14 +393,14 @@ struct FSequencerSectionPainterImpl : FSequencerSectionPainter
 		return TotalScale > 0.f ? EaseInInterp * (EaseInScale/TotalScale) + ((1.f-EaseOutInterp) * (EaseOutScale/TotalScale)) : 0.f;
 	}
 
-	FEasingCurvePoint MakeCurvePoint(FSectionHandle SectionHandle, float Time, const FLinearColor& FinalTint, const FLinearColor& EaseSelectionColor) const
+	FEasingCurvePoint MakeCurvePoint(FSectionHandle SectionHandle, FFrameTime Time, const FLinearColor& FinalTint, const FLinearColor& EaseSelectionColor) const
 	{
 		TOptional<float> EaseInValue, EaseOutValue;
 		float EaseInInterp = 0.f, EaseOutInterp = 1.f;
 		SectionHandle.GetSectionObject()->EvaluateEasing(Time, EaseInValue, EaseOutValue, &EaseInInterp, &EaseOutInterp);
 
 		return FEasingCurvePoint(
-			FVector2D(Time, EaseInValue.Get(1.f) * EaseOutValue.Get(1.f)),
+			FVector2D(Time / TimeToPixelConverter.GetFrameResolution(), EaseInValue.Get(1.f) * EaseOutValue.Get(1.f)),
 			FMath::Lerp(FinalTint, EaseSelectionColor, GetEaseHighlightAmount(SectionHandle, EaseInInterp, EaseOutInterp))
 		);
 	}
@@ -346,7 +411,7 @@ struct FSequencerSectionPainterImpl : FSequencerSectionPainter
 		static float GradientThreshold = .05f;
 		static float ValueThreshold = .05f;
 
-		float MinTimeSize = FMath::Max(0.0001f, TimeToPixelConverter.PixelToTime(2.5) - TimeToPixelConverter.PixelToTime(0));
+		float MinTimeSize = FMath::Max(0.0001, TimeToPixelConverter.PixelToSeconds(2.5) - TimeToPixelConverter.PixelToSeconds(0));
 
 		UMovieSceneSection* SectionObject = SectionHandle.GetSectionObject();
 
@@ -357,18 +422,19 @@ struct FSequencerSectionPainterImpl : FSequencerSectionPainter
 
 			if ((Upper.Location.X - Lower.Location.X)*.5f > MinTimeSize)
 			{
-				float NewPointTime = (Upper.Location.X + Lower.Location.X)*.5f;
-				float NewPointValue = SectionObject->EvaluateEasing(NewPointTime);
+				float      NewPointTime  = (Upper.Location.X + Lower.Location.X)*.5f;
+				FFrameTime FrameTime     = NewPointTime * TimeToPixelConverter.GetFrameResolution();
+				float      NewPointValue = SectionObject->EvaluateEasing(FrameTime);
 
 				// Check that the gradient is changing significantly
 				float LinearValue = (Upper.Location.Y + Lower.Location.Y) * .5f;
-				float PointGradient = NewPointValue - SectionObject->EvaluateEasing(FMath::Lerp(Lower.Location.X, NewPointTime, 0.9f));
+				float PointGradient = NewPointValue - SectionObject->EvaluateEasing(FMath::Lerp(Lower.Location.X, NewPointTime, 0.9f) * TimeToPixelConverter.GetFrameResolution());
 				float OuterGradient = Upper.Location.Y - Lower.Location.Y;
 				if (!FMath::IsNearlyEqual(OuterGradient, PointGradient, GradientThreshold) ||
 					!FMath::IsNearlyEqual(LinearValue, NewPointValue, ValueThreshold))
 				{
 					// Add the point
-					InOutPoints.Insert(MakeCurvePoint(SectionHandle, NewPointTime, FinalTint, EaseSelectionColor), Index+1);
+					InOutPoints.Insert(MakeCurvePoint(SectionHandle, FrameTime, FinalTint, EaseSelectionColor), Index+1);
 					--Index;
 				}
 			}
@@ -377,12 +443,12 @@ struct FSequencerSectionPainterImpl : FSequencerSectionPainter
 
 	void DrawEasingForSegment(const FSequencerOverlapRange& Segment, const FGeometry& InnerSectionGeometry, const FLinearColor& FinalTint)
 	{
-		const float StartTimePixel = TimeToPixelConverter.TimeToPixel(Section.GetStartTime());
-		const float RangeStartPixel = TimeToPixelConverter.TimeToPixel(Segment.Range.GetLowerBoundValue());
-		const float RangeEndPixel = TimeToPixelConverter.TimeToPixel(Segment.Range.GetUpperBoundValue());
+		// @todo: sequencer-timecode: Test that start offset is not required here
+		const float RangeStartPixel = TimeToPixelConverter.FrameToPixel(MovieScene::DiscreteInclusiveLower(Segment.Range));
+		const float RangeEndPixel = TimeToPixelConverter.FrameToPixel(MovieScene::DiscreteExclusiveUpper(Segment.Range));
 		const float RangeSizePixel = RangeEndPixel - RangeStartPixel;
 
-		FGeometry RangeGeometry = InnerSectionGeometry.MakeChild(FVector2D(RangeSizePixel, InnerSectionGeometry.Size.Y), FSlateLayoutTransform(FVector2D(RangeStartPixel - StartTimePixel, 0.f)));
+		FGeometry RangeGeometry = InnerSectionGeometry.MakeChild(FVector2D(RangeSizePixel, InnerSectionGeometry.Size.Y), FSlateLayoutTransform(FVector2D(RangeStartPixel, 0.f)));
 		if (!FSlateRect::DoRectanglesIntersect(RangeGeometry.GetLayoutBoundingRect(), ParentClippingRect))
 		{
 			return;
@@ -442,7 +508,8 @@ struct FSequencerSectionPainterImpl : FSequencerSectionPainter
 
 			for (const FEasingCurvePoint& Point : CurvePoints)
 			{
-				float U = (Point.Location.X - Segment.Range.GetLowerBoundValue()) / Segment.Range.Size<float>();
+				float SegmentStartTime = MovieScene::DiscreteInclusiveLower(Segment.Range) / TimeToPixelConverter.GetFrameResolution();
+				float U = (Point.Location.X - SegmentStartTime) / ( FFrameNumber(MovieScene::DiscreteSize(Segment.Range)) / TimeToPixelConverter.GetFrameResolution() );
 
 				// Add verts top->bottom
 				FVector2D UV(U, 0.f);
@@ -507,6 +574,8 @@ struct FSequencerSectionPainterImpl : FSequencerSectionPainter
 		{
 			DrawEasingForSegment(Segment, InnerSectionGeometry, FinalTint);
 		}
+
+		++LayerId;
 	}
 
 	void DrawOverlaps(const FLinearColor& FinalTint)
@@ -524,14 +593,14 @@ struct FSequencerSectionPainterImpl : FSequencerSectionPainter
 
 		const ESlateDrawEffect DrawEffects = bParentEnabled ? ESlateDrawEffect::None : ESlateDrawEffect::DisabledEffect;
 
-		const float StartTimePixel = Section.IsInfinite() ? 0.f : TimeToPixelConverter.TimeToPixel(Section.GetStartTime());
+		const float StartTimePixel = Section.HasStartFrame() ? TimeToPixelConverter.FrameToPixel(Section.GetInclusiveStartFrame()) : 0.f;
 
 		for (int32 SegmentIndex = 0; SegmentIndex < SectionWidget.UnderlappingSegments.Num(); ++SegmentIndex)
 		{
 			const FSequencerOverlapRange& Segment = SectionWidget.UnderlappingSegments[SegmentIndex];
 
-			const float RangeStartPixel	= Segment.Range.GetLowerBound().IsOpen() ? 0.f							: TimeToPixelConverter.TimeToPixel(Segment.Range.GetLowerBoundValue());
-			const float RangeEndPixel	= Segment.Range.GetUpperBound().IsOpen() ? InnerSectionGeometry.Size.X	: TimeToPixelConverter.TimeToPixel(Segment.Range.GetUpperBoundValue());
+			const float RangeStartPixel	= Segment.Range.GetLowerBound().IsOpen() ? 0.f							: TimeToPixelConverter.FrameToPixel(MovieScene::DiscreteInclusiveLower(Segment.Range));
+			const float RangeEndPixel	= Segment.Range.GetUpperBound().IsOpen() ? InnerSectionGeometry.Size.X	: TimeToPixelConverter.FrameToPixel(MovieScene::DiscreteExclusiveUpper(Segment.Range));
 			const float RangeSizePixel	= RangeEndPixel - RangeStartPixel;
 
 			FGeometry RangeGeometry = InnerSectionGeometry.MakeChild(FVector2D(RangeSizePixel, InnerSectionGeometry.Size.Y), FSlateLayoutTransform(FVector2D(RangeStartPixel - StartTimePixel, 0.f)));
@@ -562,6 +631,8 @@ struct FSequencerSectionPainterImpl : FSequencerSectionPainter
 				FLinearColor(1.f,1.f,1.f,.3f)
 			);
 		}
+
+		++LayerId;
 	}
 
 	void DrawEmptySpace()
@@ -574,7 +645,7 @@ struct FSequencerSectionPainterImpl : FSequencerSectionPainter
 
 		for (const FSectionLayoutElement& Element : SectionWidget.Layout->GetElements())
 		{
-			const bool bIsEmptySpace = Element.GetDisplayNode()->GetType() == ESequencerNode::KeyArea && !Element.GetKeyArea().IsValid();
+			const bool bIsEmptySpace = Element.GetDisplayNode()->GetType() == ESequencerNode::KeyArea && Element.GetKeyAreas().Num() == 0;
 			const bool bExistingEmptySpace = CurrentArea.IsSet();
 
 			if (bIsEmptySpace && bExistingEmptySpace && FMath::IsNearlyEqual(CurrentArea->Bottom, Element.GetOffset()))
@@ -601,6 +672,8 @@ struct FSequencerSectionPainterImpl : FSequencerSectionPainter
 			FPaintGeometry PaintGeom = SectionGeometry.MakeChild(CurrentArea->GetSize(), FSlateLayoutTransform(CurrentArea->GetTopLeft())).ToPaintGeometry();
 			FSlateDrawElement::MakeBox(DrawElements, LayerId, PaintGeom, EmptySpaceBrush, DrawEffects);
 		}
+
+		++LayerId;
 	}
 
 	TOptional<FLinearColor> SelectionColor;
@@ -634,152 +707,124 @@ FVector2D SSequencerSection::ComputeDesiredSize(float) const
 }
 
 
-FGeometry SSequencerSection::GetKeyAreaGeometry( const FSectionLayoutElement& KeyArea, const FGeometry& SectionGeometry ) const
+FGeometry SSequencerSection::GetKeyAreaGeometry( const FSectionLayoutElement& LayoutElement, const FGeometry& SectionGeometry ) const
 {
 	// Compute the geometry for the key area
-	return SectionGeometry.MakeChild( FVector2D( 0, KeyArea.GetOffset() ), FVector2D( SectionGeometry.GetLocalSize().X, KeyArea.GetHeight()  ) );
+	return SectionGeometry.MakeChild( FVector2D( 0, LayoutElement.GetOffset() ), FVector2D( SectionGeometry.GetLocalSize().X, LayoutElement.GetHeight()  ) );
 }
 
 
-FSequencerSelectedKey SSequencerSection::GetKeyUnderMouse( const FVector2D& MousePosition, const FGeometry& AllottedGeometry ) const
+void SSequencerSection::GetKeysUnderMouse( const FVector2D& MousePosition, const FGeometry& AllottedGeometry, TArray<FSequencerSelectedKey>& OutKeys ) const
 {
 	FGeometry SectionGeometry = MakeSectionGeometryWithoutHandles( AllottedGeometry, SectionInterface );
 
 	UMovieSceneSection& Section = *SectionInterface->GetSectionObject();
 
+	FTimeToPixel TimeToPixelConverter = ConstructTimeConverterForSection(SectionGeometry, Section, GetSequencer());
+	const float  MousePixel           = SectionGeometry.AbsoluteToLocal( MousePosition ).X;
+
+	// HitTest 
+	const FFrameTime HalfKeySizeFrames = TimeToPixelConverter.PixelDeltaToFrame( SequencerSectionConstants::KeySize.X*.5f );
+	const FFrameTime MouseFrameTime    = TimeToPixelConverter.PixelToFrame( MousePixel ) - FFrameTime(0, 0.5f);
+
+	TRange<FFrameNumber> HitTestRange = TRange<FFrameNumber>( (MouseFrameTime-HalfKeySizeFrames).CeilToFrame(), (MouseFrameTime+HalfKeySizeFrames).CeilToFrame());
+	if (HitTestRange.IsEmpty())
+	{
+		return;
+	}
+
 	// Search every key area until we find the one under the mouse
 	for (const FSectionLayoutElement& Element : Layout->GetElements())
 	{
-		TSharedPtr<IKeyArea> KeyArea = Element.GetKeyArea();
-		if (!KeyArea.IsValid())
+		// Is the key area under the mouse
+		if( !GetKeyAreaGeometry( Element, AllottedGeometry ).IsUnderLocation( MousePosition ) )
 		{
 			continue;
 		}
 
-		// Compute the current key area geometry
-		FGeometry KeyAreaGeometryPadded = GetKeyAreaGeometry( Element, AllottedGeometry );
-
-		// Is the key area under the mouse
-		if( KeyAreaGeometryPadded.IsUnderLocation( MousePosition ) )
+		for (TSharedPtr<IKeyArea> KeyArea : Element.GetKeyAreas())
 		{
-			FGeometry KeyAreaGeometry = GetKeyAreaGeometry( Element, SectionGeometry );
-			FVector2D LocalSpaceMousePosition = KeyAreaGeometry.AbsoluteToLocal( MousePosition );
-
-			FTimeToPixel TimeToPixelConverter = Section.IsInfinite()
-				? FTimeToPixel( ParentGeometry, GetSequencer().GetViewRange())
-				: FTimeToPixel( KeyAreaGeometry, TRange<float>( Section.GetStartTime(), Section.GetEndTime() ) );
-
-			// Check each key until we find one under the mouse (if any)
-			TArray<FKeyHandle> KeyHandles = KeyArea->GetUnsortedKeyHandles();
-			for( int32 KeyIndex = 0; KeyIndex < KeyHandles.Num(); ++KeyIndex )
+			if (KeyArea.IsValid())
 			{
-				FKeyHandle KeyHandle = KeyHandles[KeyIndex];
-				float KeyPosition = TimeToPixelConverter.TimeToPixel( KeyArea->GetKeyTime(KeyHandle) );
-				FGeometry KeyGeometry = KeyAreaGeometry.MakeChild( 
-					FVector2D( KeyPosition - FMath::CeilToFloat(SequencerSectionConstants::KeySize.X/2.0f), ((KeyAreaGeometry.GetLocalSize().Y*.5f)-(SequencerSectionConstants::KeySize.Y*.5f)) ),
-					SequencerSectionConstants::KeySize
-				);
+				TArray<FKeyHandle> KeyHandles;
+				KeyArea->GetKeyHandles(KeyHandles, HitTestRange);
 
-				if( KeyGeometry.IsUnderLocation( MousePosition ) )
+				// Only ever select one key from any given key area
+				if (KeyHandles.Num())
 				{
-					// The current key is under the mouse
-					return FSequencerSelectedKey( Section, KeyArea, KeyHandle );
+					OutKeys.Add( FSequencerSelectedKey( Section, KeyArea, KeyHandles[0] ) );
 				}
 			}
-
-			// no key was selected in the current key area but the mouse is in the key area so it cannot possibly be in any other key area
-			return FSequencerSelectedKey();
 		}
-	}
 
-	// No key was selected in any key area
-	return FSequencerSelectedKey();
+		// The mouse is in this key area so it cannot possibly be in any other key area
+		return;
+	}
 }
 
 
-FSequencerSelectedKey SSequencerSection::CreateKeyUnderMouse( const FVector2D& MousePosition, const FGeometry& AllottedGeometry, FSequencerSelectedKey InPressedKey )
+void SSequencerSection::CreateKeysUnderMouse( const FVector2D& MousePosition, const FGeometry& AllottedGeometry, TArrayView<const FSequencerSelectedKey> InPressedKeys, TArray<FSequencerSelectedKey>& OutKeys )
 {
 	UMovieSceneSection& Section = *SectionInterface->GetSectionObject();
-	FGeometry SectionGeometry = MakeSectionGeometryWithoutHandles( AllottedGeometry, SectionInterface );
 
-	// Search every key area until we find the one under the mouse
-	for (const FSectionLayoutElement& Element : Layout->GetElements())
+	// If the pressed key exists, offset the new key and look for it in the newly laid out key areas
+	if (InPressedKeys.Num())
 	{
-		TSharedPtr<IKeyArea> KeyArea = Element.GetKeyArea();
-		if (!KeyArea.IsValid())
+		Section.Modify();
+
+		// Offset by 1 pixel worth of time if possible
+		const FFrameTime TimeFuzz = (GetSequencer().GetViewRange().Size<double>() / ParentGeometry.GetLocalSize().X) * Section.GetTypedOuter<UMovieScene>()->GetFrameResolution();
+
+		for (const FSequencerSelectedKey& PressedKey : InPressedKeys)
 		{
-			continue;
+			const FFrameNumber CurrentTime = PressedKey.KeyArea->GetKeyTime(PressedKey.KeyHandle.GetValue());
+			const FKeyHandle NewHandle = PressedKey.KeyArea->DuplicateKey(PressedKey.KeyHandle.GetValue());
+
+			PressedKey.KeyArea->SetKeyTime(NewHandle, CurrentTime + TimeFuzz.FrameNumber);
+			OutKeys.Add(FSequencerSelectedKey(Section, PressedKey.KeyArea, NewHandle));
 		}
+	}
+	else
+	{
+		TSharedPtr<FSequencerObjectBindingNode> ObjectBindingNode = ParentSectionArea.IsValid() ? ParentSectionArea->FindParentObjectBindingNode() : nullptr;
+		FGuid ObjectBinding = ObjectBindingNode.IsValid() ? ObjectBindingNode->GetObjectBinding() : FGuid();
 
-		// Compute the current key area geometry
-		FGeometry KeyAreaGeometryPadded = GetKeyAreaGeometry( Element, AllottedGeometry );
+		FGeometry SectionGeometry = MakeSectionGeometryWithoutHandles( AllottedGeometry, SectionInterface );
 
-		// Is the key area under the mouse
-		if( KeyAreaGeometryPadded.IsUnderLocation( MousePosition ) )
+		// Search every key area until we find the one under the mouse
+		for (const FSectionLayoutElement& Element : Layout->GetElements())
 		{
-			FTimeToPixel TimeToPixelConverter = Section.IsInfinite() ? 			
-				FTimeToPixel( ParentGeometry, GetSequencer().GetViewRange()) : 
-				FTimeToPixel( SectionGeometry, TRange<float>( Section.GetStartTime(), Section.GetEndTime() ) );
+			// Compute the current key area geometry
+			FGeometry KeyAreaGeometryPadded = GetKeyAreaGeometry( Element, AllottedGeometry );
 
-			// If a key was pressed on, get the pressed on key's time to duplicate that key
-			float KeyTime;
-
-			if (InPressedKey.IsValid())
+			// Is the key area under the mouse
+			if( !KeyAreaGeometryPadded.IsUnderLocation( MousePosition ) )
 			{
-				KeyTime = KeyArea->GetKeyTime(InPressedKey.KeyHandle.GetValue());
-			}
-			// Otherwise, use the time where the mouse is pressed
-			else
-			{
-				FVector2D LocalSpaceMousePosition = SectionGeometry.AbsoluteToLocal( MousePosition );
-				KeyTime = TimeToPixelConverter.PixelToTime(LocalSpaceMousePosition.X);
+				continue;
 			}
 
-			if (Section.TryModify())
-			{
-				// If the pressed key exists, offset the new key and look for it in the newly laid out key areas
-				if (InPressedKey.IsValid())
-				{
-					// Offset by 1 pixel worth of time
-					const float TimeFuzz = (GetSequencer().GetViewRange().GetUpperBoundValue() - GetSequencer().GetViewRange().GetLowerBoundValue()) / ParentGeometry.GetLocalSize().X;
+			FTimeToPixel TimeToPixelConverter = ConstructTimeConverterForSection(SectionGeometry, Section, GetSequencer());
 
-					TArray<FKeyHandle> KeyHandles = KeyArea->AddKeyUnique(KeyTime+TimeFuzz, GetSequencer().GetKeyInterpolation(), KeyTime);
+			FVector2D LocalSpaceMousePosition = SectionGeometry.AbsoluteToLocal( MousePosition );
+			const FFrameTime KeyTime = TimeToPixelConverter.PixelToFrame(LocalSpaceMousePosition.X);
 
-					Layout = FSectionLayout(*ParentSectionArea, SectionIndex);
+			Section.Modify();
 
-					// Look specifically for the key with the offset key time
-					for (const FSectionLayoutElement& NewElement : Layout->GetElements())
-					{
-						TSharedPtr<IKeyArea> NewKeyArea = NewElement.GetKeyArea();
-						if (!NewKeyArea.IsValid())
-						{
-							continue;
-						}
-
-						for (auto KeyHandle : KeyHandles)
-						{
-							for (auto UnsortedKeyHandle : NewKeyArea->GetUnsortedKeyHandles())
-							{
-								if (FMath::IsNearlyEqual(KeyTime+TimeFuzz, NewKeyArea->GetKeyTime(UnsortedKeyHandle), KINDA_SMALL_NUMBER))
-								{
-									return FSequencerSelectedKey(Section, NewKeyArea, UnsortedKeyHandle);
-								}
-							}
-						}
-					}
-				}
-				else
+			for (TSharedPtr<IKeyArea> KeyArea : Element.GetKeyAreas())
+			{
+				if (KeyArea.IsValid())
 				{
-					KeyArea->AddKeyUnique(KeyTime, GetSequencer().GetKeyInterpolation(), KeyTime);
-					Layout = FSectionLayout(*ParentSectionArea, SectionIndex);
-						
-					return GetKeyUnderMouse(MousePosition, AllottedGeometry);
+					FKeyHandle NewHandle = KeyArea->AddOrUpdateKey(KeyTime.FrameNumber, ObjectBinding, GetSequencer());
+					OutKeys.Add(FSequencerSelectedKey(Section, KeyArea, NewHandle));
 				}
 			}
 		}
 	}
 
-	return FSequencerSelectedKey();
+	if (OutKeys.Num())
+	{
+		Layout = FSectionLayout(*ParentSectionArea, SectionIndex);
+	}
 }
 
 
@@ -791,11 +836,11 @@ bool SSequencerSection::CheckForEasingHandleInteraction( const FPointerEvent& Mo
 		return false;
 	}
 
-	FTimeToPixel TimeToPixelConverter(MakeSectionGeometryWithoutHandles(SectionGeometry, SectionInterface), ThisSection->IsInfinite() ? GetSequencer().GetViewRange() : ThisSection->GetRange());
+	FTimeToPixel TimeToPixelConverter = ConstructTimeConverterForSection(MakeSectionGeometryWithoutHandles(SectionGeometry, SectionInterface), *ThisSection, GetSequencer());
 
-	const float MouseTime = TimeToPixelConverter.PixelToTime(SectionGeometry.AbsoluteToLocal(MouseEvent.GetScreenSpacePosition()).X);
+	const double MouseTime = TimeToPixelConverter.PixelToSeconds(SectionGeometry.AbsoluteToLocal(MouseEvent.GetScreenSpacePosition()).X);
 	// We intentionally give the handles a little more hit-test area than is visible as they are quite small
-	const float HalfHandleSizeX = TimeToPixelConverter.PixelToTime(8.f) - TimeToPixelConverter.PixelToTime(0.f);
+	const double HalfHandleSizeX = TimeToPixelConverter.PixelToSeconds(8.f) - TimeToPixelConverter.PixelToSeconds(0.f);
 
 	// Now test individual easing handles if we're at the correct vertical position
 	float LocalMouseY = SectionGeometry.AbsoluteToLocal(MouseEvent.GetScreenSpacePosition()).Y;
@@ -817,25 +862,31 @@ bool SSequencerSection::CheckForEasingHandleInteraction( const FPointerEvent& Mo
 
 	for (FSectionHandle Handle : AllUnderlappingSections)
 	{
-		TSharedRef<ISequencerSection> EasingSection =  Handle.TrackNode->GetSections()[Handle.SectionIndex];
-		UMovieSceneSection* EasingSectionObj = EasingSection->GetSectionObject();
-		if (EasingSectionObj->IsInfinite())
-		{
-			continue;
-		}
+		TSharedRef<ISequencerSection> EasingSection    =  Handle.TrackNode->GetSections()[Handle.SectionIndex];
+		UMovieSceneSection*           EasingSectionObj = EasingSection->GetSectionObject();
 
-		TRange<float> EaseInRange = EasingSectionObj->GetEaseInRange();
-		if (FMath::IsNearlyEqual(MouseTime, EaseInRange.IsEmpty() ? EasingSectionObj->GetStartTime() : EaseInRange.GetUpperBoundValue(), HalfHandleSizeX))
+		if (EasingSectionObj->HasStartFrame())
 		{
-			GetSequencer().SetHotspot(MakeShared<FSectionEasingHandleHotspot>(ESequencerEasingType::In, Handle));
-			return true;
+			TRange<FFrameNumber> EaseInRange      = EasingSectionObj->GetEaseInRange();
+			double               HandlePositionIn = ( EaseInRange.IsEmpty() ? EasingSectionObj->GetInclusiveStartFrame() : EaseInRange.GetUpperBoundValue() ) / TimeToPixelConverter.GetFrameResolution();
+
+			if (FMath::IsNearlyEqual(MouseTime, HandlePositionIn, HalfHandleSizeX))
+			{
+				GetSequencer().SetHotspot(MakeShared<FSectionEasingHandleHotspot>(ESequencerEasingType::In, Handle));
+				return true;
+			}
 		}
 
-		TRange<float> EaseOutRange = EasingSectionObj->GetEaseOutRange();
-		if (FMath::IsNearlyEqual(MouseTime, EaseOutRange.IsEmpty() ? EasingSectionObj->GetEndTime() :EaseOutRange.GetLowerBoundValue(), HalfHandleSizeX))
+		if (EasingSectionObj->HasEndFrame())
 		{
-			GetSequencer().SetHotspot(MakeShared<FSectionEasingHandleHotspot>(ESequencerEasingType::Out, Handle));
-			return true;
+			TRange<FFrameNumber> EaseOutRange      = EasingSectionObj->GetEaseOutRange();
+			double               HandlePositionOut = (EaseOutRange.IsEmpty() ? EasingSectionObj->GetExclusiveEndFrame() : EaseOutRange.GetLowerBoundValue() ) / TimeToPixelConverter.GetFrameResolution();
+
+			if (FMath::IsNearlyEqual(MouseTime, HandlePositionOut, HalfHandleSizeX))
+			{
+				GetSequencer().SetHotspot(MakeShared<FSectionEasingHandleHotspot>(ESequencerEasingType::Out, Handle));
+				return true;
+			}
 		}
 	}
 
@@ -861,13 +912,14 @@ bool SSequencerSection::CheckForEdgeInteraction( const FPointerEvent& MouseEvent
 		}
 	}
 
-	FGeometry SectionGeometryWithoutHandles = MakeSectionGeometryWithoutHandles(SectionGeometry, SectionInterface);
-	FTimeToPixel TimeToPixelConverter(SectionGeometryWithoutHandles, ThisSection->IsInfinite() ? GetSequencer().GetViewRange() : ThisSection->GetRange());
+	FGeometry    SectionGeometryWithoutHandles = MakeSectionGeometryWithoutHandles(SectionGeometry, SectionInterface);
+	FTimeToPixel TimeToPixelConverter          = ConstructTimeConverterForSection(SectionGeometryWithoutHandles, *ThisSection, GetSequencer());
+
 	for (FSectionHandle Handle : AllUnderlappingSections)
 	{
 		TSharedRef<ISequencerSection> UnderlappingSection =  Handle.TrackNode->GetSections()[Handle.SectionIndex];
 		UMovieSceneSection* UnderlappingSectionObj = UnderlappingSection->GetSectionObject();
-		if (!UnderlappingSection->SectionIsResizable() || UnderlappingSectionObj->IsInfinite())
+		if (!UnderlappingSection->SectionIsResizable())
 		{
 			continue;
 		}
@@ -875,26 +927,33 @@ bool SSequencerSection::CheckForEdgeInteraction( const FPointerEvent& MouseEvent
 		const float ThisHandleOffset = UnderlappingSectionObj == ThisSection ? HandleOffsetPx : 0.f;
 		FVector2D GripSize( UnderlappingSection->GetSectionGripSize(), SectionGeometry.Size.Y );
 
-		// Make areas to the left and right of the geometry.  We will use these areas to determine if someone dragged the left or right edge of a section
-		FGeometry SectionRectLeft = SectionGeometryWithoutHandles.MakeChild(
-			FVector2D( TimeToPixelConverter.TimeToPixel(UnderlappingSectionObj->GetStartTime()) - ThisHandleOffset, 0.f ),
-			GripSize
-		);
-
-		FGeometry SectionRectRight = SectionGeometryWithoutHandles.MakeChild(
-			FVector2D( TimeToPixelConverter.TimeToPixel(UnderlappingSectionObj->GetEndTime()) - UnderlappingSection->GetSectionGripSize() + ThisHandleOffset, 0 ), 
-			GripSize
-		);
-
-		if( SectionRectLeft.IsUnderLocation( MouseEvent.GetScreenSpacePosition() ) )
+		if (UnderlappingSectionObj->HasStartFrame())
 		{
-			GetSequencer().SetHotspot(MakeShareable( new FSectionResizeHotspot(FSectionResizeHotspot::Left, Handle)) );
-			return true;
+			// Make areas to the left and right of the geometry.  We will use these areas to determine if someone dragged the left or right edge of a section
+			FGeometry SectionRectLeft = SectionGeometryWithoutHandles.MakeChild(
+				FVector2D( TimeToPixelConverter.FrameToPixel(UnderlappingSectionObj->GetInclusiveStartFrame()) - ThisHandleOffset, 0.f ),
+				GripSize
+			);
+
+			if( SectionRectLeft.IsUnderLocation( MouseEvent.GetScreenSpacePosition() ) )
+			{
+				GetSequencer().SetHotspot(MakeShareable( new FSectionResizeHotspot(FSectionResizeHotspot::Left, Handle)) );
+				return true;
+			}
 		}
-		else if( SectionRectRight.IsUnderLocation( MouseEvent.GetScreenSpacePosition() ) )
+
+		if (UnderlappingSectionObj->HasEndFrame())
 		{
-			GetSequencer().SetHotspot(MakeShareable( new FSectionResizeHotspot(FSectionResizeHotspot::Right, Handle)) );
-			return true;
+			FGeometry SectionRectRight = SectionGeometryWithoutHandles.MakeChild(
+				FVector2D( TimeToPixelConverter.FrameToPixel(UnderlappingSectionObj->GetExclusiveEndFrame()) - UnderlappingSection->GetSectionGripSize() + ThisHandleOffset, 0 ), 
+				GripSize
+			);
+
+			if( SectionRectRight.IsUnderLocation( MouseEvent.GetScreenSpacePosition() ) )
+			{
+				GetSequencer().SetHotspot(MakeShareable( new FSectionResizeHotspot(FSectionResizeHotspot::Right, Handle)) );
+				return true;
+			}
 		}
 	}
 	return false;
@@ -908,9 +967,8 @@ bool SSequencerSection::CheckForEasingAreaInteraction( const FPointerEvent& Mous
 		return false;
 	}
 
-	FTimeToPixel TimeToPixelConverter(MakeSectionGeometryWithoutHandles(SectionGeometry, SectionInterface), ThisSection->IsInfinite() ? GetSequencer().GetViewRange() : ThisSection->GetRange());
-
-	const float MouseTime = TimeToPixelConverter.PixelToTime(SectionGeometry.AbsoluteToLocal(MouseEvent.GetScreenSpacePosition()).X);
+	FTimeToPixel TimeToPixelConverter = ConstructTimeConverterForSection(MakeSectionGeometryWithoutHandles(SectionGeometry, SectionInterface), *ThisSection, GetSequencer());
+	const FFrameNumber MouseTime = TimeToPixelConverter.PixelToFrame(SectionGeometry.AbsoluteToLocal(MouseEvent.GetScreenSpacePosition()).X).FrameNumber;
 
 	// First off, set the hotspot to an easing area if necessary
 	for (const FSequencerOverlapRange& Segment : UnderlappingEasingSegments)
@@ -1017,10 +1075,10 @@ int32 SSequencerSection::OnPaint( const FPaintArgs& Args, const FGeometry& Allot
 	FText SectionTitle = SectionInterface->GetSectionTitle();
 	FMargin ContentPadding = SectionInterface->GetContentPadding();
 
-	const float EaseInAmount = SectionObject.Easing.GetEaseInTime();
-	if (EaseInAmount > 0.f)
+	const int32 EaseInAmount = SectionObject.Easing.GetEaseInDuration();
+	if (EaseInAmount > 0)
 	{
-		ContentPadding.Left += Painter.GetTimeConverter().TimeToPixel(EaseInAmount) - Painter.GetTimeConverter().TimeToPixel(0.f);
+		ContentPadding.Left += Painter.GetTimeConverter().FrameToPixel(EaseInAmount) - Painter.GetTimeConverter().FrameToPixel(0);
 	}
 
 	if (!SectionTitle.IsEmpty())
@@ -1054,12 +1112,6 @@ int32 SSequencerSection::OnPaint( const FPaintArgs& Args, const FGeometry& Allot
 void SSequencerSection::PaintKeys( FSequencerSectionPainter& InPainter, const FWidgetStyle& InWidgetStyle ) const
 {
 	static const FName HighlightBrushName("Sequencer.AnimationOutliner.DefaultBorder");
-
-	static const FName BackgroundBrushName("Sequencer.Section.BackgroundTint");
-	static const FName CircleKeyBrushName("Sequencer.KeyCircle");
-	static const FName DiamondKeyBrushName("Sequencer.KeyDiamond");
-	static const FName SquareKeyBrushName("Sequencer.KeySquare");
-	static const FName TriangleKeyBrushName("Sequencer.KeyTriangle");
 	static const FName StripeOverlayBrushName("Sequencer.Section.StripeOverlay");
 
 	static const FName SelectionColorName("SelectionColor");
@@ -1074,30 +1126,26 @@ void SSequencerSection::PaintKeys( FSequencerSectionPainter& InPainter, const FW
 	FSequencer& Sequencer = ParentSectionArea->GetSequencer();
 	TSharedPtr<ISequencerHotspot> Hotspot = Sequencer.GetHotspot();
 
+	const FSlateBrush* StripeOverlayBrush = FEditorStyle::GetBrush(StripeOverlayBrushName);
+	const FSlateBrush* HighlightBrush = FEditorStyle::GetBrush(HighlightBrushName);
+
 	// get hovered key
-	FSequencerSelectedKey HoveredKey;
+	TArrayView<const FSequencerSelectedKey> HoveredKeys;
 
 	if (Hotspot.IsValid() && Hotspot->GetType() == ESequencerHotspot::Key)
 	{
-		HoveredKey = static_cast<FKeyHotspot*>(Hotspot.Get())->Key;
+		HoveredKeys = static_cast<FKeyHotspot*>(Hotspot.Get())->Keys;
 	}
 
 	auto& Selection = Sequencer.GetSelection();
 	auto& SelectionPreview = Sequencer.GetSelectionPreview();
 
-	const float ThrobScaleValue = GetSelectionThrobValue();
+	const float KeyThrobScaleValue = GetKeySelectionThrobValue();
+	const float SectionThrobScaleValue = GetSectionSelectionThrobValue();
 
 	// draw all keys in each key area
 	UMovieSceneSection& SectionObject = *SectionInterface->GetSectionObject();
 
-	const FSlateBrush* HighlightBrush = FEditorStyle::GetBrush(HighlightBrushName);
-	const FSlateBrush* BackgroundBrush = FEditorStyle::GetBrush(BackgroundBrushName);
-	const FSlateBrush* StripeOverlayBrush = FEditorStyle::GetBrush(StripeOverlayBrushName);
-	const FSlateBrush* CircleKeyBrush = FEditorStyle::GetBrush(CircleKeyBrushName);
-	const FSlateBrush* DiamondKeyBrush = FEditorStyle::GetBrush(DiamondKeyBrushName);
-	const FSlateBrush* SquareKeyBrush = FEditorStyle::GetBrush(SquareKeyBrushName);
-	const FSlateBrush* TriangleKeyBrush = FEditorStyle::GetBrush(TriangleKeyBrushName);
-
 	const ESlateDrawEffect DrawEffects = InPainter.bParentEnabled
 		? ESlateDrawEffect::None
 		: ESlateDrawEffect::DisabledEffect;
@@ -1107,14 +1155,14 @@ void SSequencerSection::PaintKeys( FSequencerSectionPainter& InPainter, const FW
 	for (const FSectionLayoutElement& LayoutElement : Layout->GetElements())
 	{
 		// get key handles
-		TSharedPtr<IKeyArea> KeyArea = LayoutElement.GetKeyArea();
+		TArrayView<const TSharedRef<IKeyArea>> KeyAreas = LayoutElement.GetKeyAreas();
 
 		FGeometry KeyAreaGeometry = GetKeyAreaGeometry(LayoutElement, InPainter.SectionGeometry);
 
-		TOptional<FLinearColor> KeyAreaColor = KeyArea.IsValid() ? KeyArea->GetColor() : TOptional<FLinearColor>();
+		TOptional<FLinearColor> KeyAreaColor = KeyAreas.Num() == 1 ? KeyAreas[0]->GetColor() : TOptional<FLinearColor>();
 
 		// draw a box for the key area 
-		if (KeyAreaColor.IsSet() && Sequencer.GetSettings()->GetShowChannelColors())
+		if (KeyAreaColor.IsSet() && Sequencer.GetSequencerSettings()->GetShowChannelColors())
 		{
 			static float BoxThickness = 5.f;
 			FVector2D KeyAreaSize = KeyAreaGeometry.GetLocalSize();
@@ -1172,24 +1220,44 @@ void SSequencerSection::PaintKeys( FSequencerSectionPainter& InPainter, const FW
 			);
 		}
 
-		// Can't do any of the rest if there are no keys
-		if (!KeyArea.IsValid())
+		bool bSectionSelected = Selection.IsSelected(&SectionObject);
+		if (bSectionSelected && SectionThrobScaleValue != 0.f)
 		{
-			continue;
+			static const FName SelectedTrackTint("Sequencer.Section.BackgroundTint");
+
+			FLinearColor KeyAreaOutlineColor = SelectionColor;
+			KeyAreaOutlineColor.A = SectionThrobScaleValue;
+
+			FSlateDrawElement::MakeBox(
+				InPainter.DrawElements,
+				InPainter.LayerId,
+				KeyAreaGeometry.ToPaintGeometry(),
+				FEditorStyle::GetBrush(SelectedTrackTint),
+				DrawEffects,
+				KeyAreaOutlineColor
+			);
 		}
 
-		// Gather keys for a region larger thean the view range to ensure we draw keys that are only just offscreen.
-		TRange<float> PaddedViewRange;
+		double HalfFrameInterval = Sequencer.GetFocusedFrameResolution().AsInterval() * 0.5f;
+
+		// Gather keys for a region larger than the view range to ensure we draw keys that are only just offscreen.
+		TRange<double> PaddedViewRange;
 		{
-			const float KeyWidthAsTime = TimeToPixelConverter.PixelToTime(SequencerSectionConstants::KeySize.X) - TimeToPixelConverter.PixelToTime(0);
-			const TRange<float> ViewRange = GetSequencer().GetViewRange();
+			// Compute visible range taking into account a half-frame offset for keys, plus half a key width for keys that are partially offscreen
+			TRange<FFrameNumber> SectionRange   = SectionObject.GetRange();
+			const double         HalfKeyWidth   = 0.5f * (TimeToPixelConverter.PixelToSeconds(SequencerSectionConstants::KeySize.X) - TimeToPixelConverter.PixelToSeconds(0));
+			TRange<double>       VisibleRange   = MovieScene::DilateRange(GetSequencer().GetViewRange(), -(HalfFrameInterval + HalfKeyWidth), -HalfFrameInterval + HalfKeyWidth);
 
-			PaddedViewRange = TRange<float>(ViewRange.GetLowerBoundValue() - KeyWidthAsTime, ViewRange.GetUpperBoundValue() + KeyWidthAsTime);
+			PaddedViewRange = TRange<double>::Intersection(SectionRange / Sequencer.GetFocusedFrameResolution(), VisibleRange);
 		}
 
-		const FSequencerCachedKeys& CachedKeys = CachedKeyAreaPositions.FindChecked(KeyArea);
-		TArrayView<const FSequencerCachedKey> KeysInRange = CachedKeys.GetKeysInRange(PaddedViewRange);
-		if (!KeysInRange.Num())
+		// For each key area, we gather draw params for the visible time range, then combine
+		
+		auto& Found = CachedKeyAreaPositions.FindChecked(LayoutElement);
+		TArrayView<const FSequencerCachedKeys> CachedKeys = Found;
+
+		// Can't do any of the rest if there are no keys
+		if (CachedKeys.Num() == 0)
 		{
 			continue;
 		}
@@ -1199,148 +1267,206 @@ void SSequencerSection::PaintKeys( FSequencerSectionPainter& InPainter, const FW
 		TOptional<FSlateClippingState> PreviousClipState = InPainter.DrawElements.GetClippingState();
 		InPainter.DrawElements.PopClip();
 
-		static float PixelOverlapThreshold = 3.f;
+		struct FKeyDrawInformation
+		{
+			FKeyDrawInformation()
+				: NextIndex(0)
+			{}
+
+			int32 NextIndex;
+			TArrayView<const double> TimesInRange;
+			TArrayView<const FKeyHandle> HandlesInRange;
+			TArray<FKeyDrawParams> DrawParams;
+		};
+		TArray<FKeyDrawInformation> DrawInfoPerCache;
+		DrawInfoPerCache.Reserve(CachedKeys.Num());
 
-		for (int32 KeyIndex = 0; KeyIndex < KeysInRange.Num(); ++KeyIndex)
+		for (const FSequencerCachedKeys& Cache : CachedKeys)
 		{
-			FKeyHandle KeyHandle = KeysInRange[KeyIndex].Handle;
-			const float KeyTime = KeysInRange[KeyIndex].Time;
-			const float KeyPosition = TimeToPixelConverter.TimeToPixel(KeyTime);
+			// Gather all the key handles in this view range
+			FKeyDrawInformation KeyDrawInfo;
 
-			// Count the number of overlapping keys
-			int32 NumOverlaps = 0;
-			while (KeyIndex < KeysInRange.Num() && 
-					( (KeyIndex + 1 < KeysInRange.Num() && FMath::IsNearlyEqual(TimeToPixelConverter.TimeToPixel(KeysInRange[KeyIndex+1].Time), KeyPosition, PixelOverlapThreshold)) ||
-					  (KeyIndex > 1 && FMath::IsNearlyEqual(TimeToPixelConverter.TimeToPixel(KeysInRange[KeyIndex-1].Time), KeyPosition, PixelOverlapThreshold)) ))
+			Cache.GetKeysInRange(PaddedViewRange, &KeyDrawInfo.TimesInRange, &KeyDrawInfo.HandlesInRange);
+
+			// Generate draw data for this key area
+			if (KeyDrawInfo.TimesInRange.Num())
 			{
-				++KeyIndex;
-				++NumOverlaps;
+				KeyDrawInfo.DrawParams.SetNum(KeyDrawInfo.TimesInRange.Num());
+				Cache.GetKeyArea()->DrawKeys(KeyDrawInfo.HandlesInRange, KeyDrawInfo.DrawParams);
 			}
 
-			// omit keys which would not be visible
-			if( !SectionObject.IsTimeWithinSection(KeyTime))
+			check(KeyDrawInfo.DrawParams.Num() == KeyDrawInfo.TimesInRange.Num() && KeyDrawInfo.TimesInRange.Num() == KeyDrawInfo.HandlesInRange.Num());
+
+			// Add it to the array
+			DrawInfoPerCache.Add(MoveTemp(KeyDrawInfo));
+		}
+
+		static float PixelOverlapThreshold = 3.f;
+		const double TimeOverlapThreshold = TimeToPixelConverter.PixelToSeconds(PixelOverlapThreshold) - TimeToPixelConverter.PixelToSeconds(0.f);
+
+		auto AnythingLeftToDraw = [](const FKeyDrawInformation& In)
+		{
+			return In.NextIndex < In.TimesInRange.Num();
+		};
+
+		TArray<int32> KeyParamUpperBounds;
+		KeyParamUpperBounds.SetNum(DrawInfoPerCache.Num());
+
+		while (DrawInfoPerCache.ContainsByPredicate(AnythingLeftToDraw))
+		{
+			// Determine the next key position to draw
+			double CardinalKeyTime = TNumericLimits<double>::Max();
+			for (const FKeyDrawInformation& Info : DrawInfoPerCache)
 			{
-				continue;
+				if (Info.NextIndex < Info.TimesInRange.Num())
+				{
+					CardinalKeyTime = FMath::Min(CardinalKeyTime, Info.TimesInRange[Info.NextIndex]);
+				}
 			}
 
-			// determine the key's brush & color
-			const FSlateBrush* KeyBrush;
-			FLinearColor KeyColor;
-			FVector2D FillOffset(0.0f, 0.0f);
+			// Start grouping keys at the current key time plus 99% of the threshold to ensure that we group at the center of keys
+			// and that we avoid floating point precision issues where there is only one key [(KeyTime + TimeOverlapThreshold) - KeyTime != TimeOverlapThreshold] for some floats
+			CardinalKeyTime += TimeOverlapThreshold*0.9994f;
+
+			float AverageKeyTime = 0.f;
+			int32 NumKeyTimes = 0;
+			int32 NumOverlaps = 0;
 
-			switch (KeyArea->GetKeyInterpMode(KeyHandle))
+			// Determine the ranges of keys considered to reside at this position
+			for (int32 DrawIndex = 0; DrawIndex < DrawInfoPerCache.Num(); ++DrawIndex)
 			{
-			case RCIM_Linear:
-				KeyBrush = TriangleKeyBrush;
-				KeyColor = FLinearColor(0.0f, 0.617f, 0.449f, 1.0f); // blueish green
-				FillOffset = FVector2D(0.0f, 1.0f);
-				break;
+				const FKeyDrawInformation& Info = DrawInfoPerCache[DrawIndex];
+				if (Info.NextIndex >= Info.TimesInRange.Num())
+				{
+					continue;
+				}
+
+				if (FMath::IsNearlyEqual(Info.TimesInRange[Info.NextIndex], CardinalKeyTime, TimeOverlapThreshold))
+				{
+					int32 FinalIndexInThreshold = Info.NextIndex + 1;
 
-			case RCIM_Constant:
-				KeyBrush = SquareKeyBrush;
-				KeyColor = FLinearColor(0.0f, 0.445f, 0.695f, 1.0f); // blue
-				break;
+					AverageKeyTime += Info.TimesInRange[Info.NextIndex];
+					++NumKeyTimes;
 
-			case RCIM_Cubic:
-				KeyBrush = CircleKeyBrush;
+					// Count the number of overlapping keys
+					while (FinalIndexInThreshold < Info.TimesInRange.Num() && FMath::IsNearlyEqual(Info.TimesInRange[FinalIndexInThreshold], CardinalKeyTime, TimeOverlapThreshold))
+					{
+						AverageKeyTime += Info.TimesInRange[FinalIndexInThreshold];
+						++NumKeyTimes;
+
+						++FinalIndexInThreshold;
+						++NumOverlaps;
+					}
 
-				switch (KeyArea->GetKeyTangentMode(KeyHandle))
+					KeyParamUpperBounds[DrawIndex] = FinalIndexInThreshold;
+				}
+				else
 				{
-				case RCTM_Auto:
-					KeyColor = FLinearColor(0.972f, 0.2f, 0.2f, 1.0f); // vermillion
-					break;
+					KeyParamUpperBounds[DrawIndex] = Info.NextIndex;
+				}
+			}
 
-				case RCTM_Break:
-					KeyColor = FLinearColor(0.336f, 0.703f, 0.5f, 0.91f); // sky blue
-					break;
+			const float FinalKeyPosition = TimeToPixelConverter.SecondsToPixel((AverageKeyTime / NumKeyTimes) + HalfFrameInterval);
 
-				case RCTM_User:
-					KeyColor = FLinearColor(0.797f, 0.473f, 0.5f, 0.652f); // reddish purple
-					break;
+			static const FSlateBrush* PartialKeyBrush = FEditorStyle::GetBrush("Sequencer.PartialKey");
 
-				default:
-					KeyColor = FLinearColor(0.75f, 0.75f, 0.75f, 1.0f); // light gray
-				}
-				break;
+			TOptional<FKeyDrawParams> KeyDrawParam;
 
-			default:
-				KeyBrush = DiamondKeyBrush;
-				KeyColor = FLinearColor(1.0f, 1.0f, 1.0f, 1.0f); // white
-				break;
-			}
+			bool bPartialKey = false;
 
-			// allow group & section overrides
-			const FSlateBrush* OverrideBrush = nullptr;
+			int32 NumPreviewSelected = 0;
+			int32 NumPreviewNotSelected = 0;
+			int32 NumSelected = 0;
+			int32 NumHovered = 0;
+			int32 TotalNumKeys = 0;
 
-			if (LayoutElement.GetType() == FSectionLayoutElement::Group)
+			for (int32 DrawIndex = 0; DrawIndex < DrawInfoPerCache.Num(); ++DrawIndex)
 			{
-				auto Group = StaticCastSharedPtr<FGroupedKeyArea>(KeyArea);
-				OverrideBrush = Group->GetBrush(KeyHandle);
-			}
-			else
-			{
-				OverrideBrush = SectionInterface->GetKeyBrush(KeyHandle);
+				FKeyDrawInformation& Info = DrawInfoPerCache[DrawIndex];
+				int32 UpperBound = KeyParamUpperBounds[DrawIndex];
+
+				if (Info.NextIndex >= Info.TimesInRange.Num() || UpperBound - Info.NextIndex == 0)
+				{
+					bPartialKey = true;
+					continue;
+				}
+
+				for (int32 KeyIndex = Info.NextIndex; KeyIndex < UpperBound; ++KeyIndex)
+				{
+					if (!KeyDrawParam.IsSet())
+					{
+						KeyDrawParam = Info.DrawParams[KeyIndex];
+					}
+					else if (Info.DrawParams[KeyIndex] != KeyDrawParam.GetValue())
+					{
+						bPartialKey = true;
+					}
+
+					FSequencerSelectedKey TestKey(SectionObject, CachedKeys[DrawIndex].GetKeyArea(), Info.HandlesInRange[KeyIndex]);
+
+					ESelectionPreviewState SelectionPreviewState = SelectionPreview.GetSelectionState(TestKey);
+					NumPreviewSelected += int32(SelectionPreviewState == ESelectionPreviewState::Selected);
+					NumPreviewNotSelected += int32(SelectionPreviewState == ESelectionPreviewState::NotSelected);
+					NumSelected += int32(Selection.IsSelected(TestKey));
+					NumHovered += int32(HoveredKeys.Contains(TestKey));
+					++TotalNumKeys;
+				}
+
+				Info.NextIndex = UpperBound;
 			}
 
-			if (OverrideBrush != nullptr)
+			if (bPartialKey)
 			{
-				KeyBrush = OverrideBrush;
-				FillOffset = SectionInterface->GetKeyBrushOrigin(KeyHandle);
+				KeyDrawParam->FillOffset = FVector2D(0.f, 0.f);
+				KeyDrawParam->FillTint   = KeyDrawParam->BorderTint = FLinearColor::White;
+				KeyDrawParam->FillBrush  = KeyDrawParam->BorderBrush = PartialKeyBrush;
 			}
 
-			// determine draw colors based on hover, selection, etc.
-			FLinearColor BorderColor;
-			FLinearColor FillColor;
+			const bool bSelected = NumSelected == TotalNumKeys;
+			ensure(KeyDrawParam.IsSet());
 
-			FSequencerSelectedKey TestKey(SectionObject, KeyArea, KeyHandle);
-			ESelectionPreviewState SelectionPreviewState = SelectionPreview.GetSelectionState(TestKey);
-			bool bSelected = Selection.IsSelected(TestKey);
-
-			if (SelectionPreviewState == ESelectionPreviewState::Selected)
+			// Determine the key color based on its selection/hover states
+			if (NumPreviewSelected == TotalNumKeys)
 			{
 				FLinearColor PreviewSelectionColor = SelectionColor.LinearRGBToHSV();
 				PreviewSelectionColor.R += 0.1f; // +10% hue
 				PreviewSelectionColor.G = 0.6f; // 60% saturation
-				BorderColor = PreviewSelectionColor.HSVToLinearRGB();
-				FillColor = BorderColor;
+				KeyDrawParam->BorderTint = KeyDrawParam->FillTint = PreviewSelectionColor.HSVToLinearRGB();
 			}
-			else if (SelectionPreviewState == ESelectionPreviewState::NotSelected)
+			else if (NumPreviewNotSelected == TotalNumKeys)
 			{
-				BorderColor = FLinearColor(0.05f, 0.05f, 0.05f, 1.0f);
-				FillColor = KeyColor;
+				KeyDrawParam->BorderTint = FLinearColor(0.05f, 0.05f, 0.05f, 1.0f);
 			}
 			else if (bSelected)
 			{
-				BorderColor = SelectedKeyColor;
-				FillColor = FLinearColor(0.05f, 0.05f, 0.05f, 1.0f);
+				KeyDrawParam->BorderTint = SelectedKeyColor;
+				KeyDrawParam->FillTint = FLinearColor(0.05f, 0.05f, 0.05f, 1.0f);
+			}
+			else if (NumSelected != 0)
+			{
+				// partially selected
+				KeyDrawParam->BorderTint = SelectedKeyColor.CopyWithNewOpacity(0.5f);
+				KeyDrawParam->FillTint = FLinearColor(0.05f, 0.05f, 0.05f, 0.5f);
 			}
-			else if (TestKey == HoveredKey)
+			else if (NumHovered == TotalNumKeys)
 			{
-				BorderColor = FLinearColor(1.0f, 1.0f, 1.0f, 1.0f);
-				FillColor = FLinearColor(1.0f, 1.0f, 1.0f, 1.0f);
+				KeyDrawParam->BorderTint = FLinearColor(1.0f, 1.0f, 1.0f, 1.0f);
+				KeyDrawParam->FillTint = FLinearColor(1.0f, 1.0f, 1.0f, 1.0f);
 			}
 			else
 			{
-				BorderColor = FLinearColor(0.05f, 0.05f, 0.05f, 1.0f);
-				FillColor = KeyColor;
+				KeyDrawParam->BorderTint = FLinearColor(0.05f, 0.05f, 0.05f, 1.0f);
 			}
 
 			// Color keys with overlaps with a red border
 			if (NumOverlaps > 0)
 			{
-				BorderColor = FLinearColor(0.83f, 0.12f, 0.12f, 1.0f); // Red
-			}
-
-			// allow group to tint the color
-			if (LayoutElement.GetType() == FSectionLayoutElement::Group)
-			{
-				auto Group = StaticCastSharedPtr<FGroupedKeyArea>(KeyArea);
-				KeyColor *= Group->GetKeyTint(KeyHandle);
+				KeyDrawParam->BorderTint = FLinearColor(0.83f, 0.12f, 0.12f, 1.0f); // Red
 			}
 
 			// draw border
 			static FVector2D ThrobAmount(12.f, 12.f);
-			const FVector2D KeySize = bSelected ? SequencerSectionConstants::KeySize + ThrobAmount * ThrobScaleValue : SequencerSectionConstants::KeySize;
+			const FVector2D KeySize = bSelected ? SequencerSectionConstants::KeySize + ThrobAmount * KeyThrobScaleValue : SequencerSectionConstants::KeySize;
 
 			FSlateDrawElement::MakeBox(
 				InPainter.DrawElements,
@@ -1349,14 +1475,14 @@ void SSequencerSection::PaintKeys( FSequencerSectionPainter& InPainter, const FW
 				// Center the key along Y.  Ensure the middle of the key is at the actual key time
 				KeyAreaGeometry.ToPaintGeometry(
 					FVector2D(
-						KeyPosition - FMath::CeilToFloat(KeySize.X / 2.0f),
+						FinalKeyPosition - FMath::CeilToFloat(KeySize.X / 2.0f),
 						((KeyAreaGeometry.GetLocalSize().Y / 2.0f) - (KeySize.Y / 2.0f))
 					),
 					KeySize
 				),
-				KeyBrush,
+				KeyDrawParam->BorderBrush,
 				DrawEffects,
-				BorderColor
+				KeyDrawParam->BorderTint
 			);
 
 			// draw fill
@@ -1366,16 +1492,16 @@ void SSequencerSection::PaintKeys( FSequencerSectionPainter& InPainter, const FW
 				bSelected ? KeyLayer + 2 : KeyLayer + 1,
 				// Center the key along Y.  Ensure the middle of the key is at the actual key time
 				KeyAreaGeometry.ToPaintGeometry(
-					FillOffset + 
+					KeyDrawParam->FillOffset + 
 					FVector2D(
-						(KeyPosition - FMath::CeilToFloat((KeySize.X / 2.0f) - BrushBorderWidth)),
+						(FinalKeyPosition - FMath::CeilToFloat((KeySize.X / 2.0f) - BrushBorderWidth)),
 						((KeyAreaGeometry.GetLocalSize().Y / 2.0f) - ((KeySize.Y / 2.0f) - BrushBorderWidth))
 					),
 					KeySize - 2.0f * BrushBorderWidth
 				),
-				KeyBrush,
+				KeyDrawParam->FillBrush,
 				DrawEffects,
-				FillColor
+				KeyDrawParam->FillTint
 			);
 		}
 
@@ -1417,7 +1543,7 @@ void SSequencerSection::PaintEasingHandles( FSequencerSectionPainter& InPainter,
 	for (FSectionHandle Handle : AllUnderlappingSections)
 	{
 		UMovieSceneSection* UnderlappingSectionObj = Handle.GetSectionInterface()->GetSectionObject();
-		if (UnderlappingSectionObj->IsInfinite())
+		if (UnderlappingSectionObj->GetRange() == TRange<FFrameNumber>::All())
 		{
 			continue;
 		}
@@ -1472,39 +1598,49 @@ void SSequencerSection::PaintEasingHandles( FSequencerSectionPainter& InPainter,
 		const FSlateBrush* EasingHandle = FEditorStyle::GetBrush("Sequencer.Section.EasingHandle");
 		FVector2D HandleSize(10.f, 10.f);
 
-		TRange<float> EaseInRange = UnderlappingSectionObj->GetEaseInRange();
-		// Always draw handles if the section is highlighted, even if there is no range (to allow manual adjustment)
-		FVector2D HandlePos(TimeToPixelConverter.TimeToPixel(EaseInRange.IsEmpty() ? UnderlappingSectionObj->GetStartTime() : EaseInRange.GetUpperBoundValue()), 0.f);
-		FSlateDrawElement::MakeBox(
-			InPainter.DrawElements,
-			// always draw selected keys on top of other keys
-			InPainter.LayerId,
-			// Center the key along X.  Ensure the middle of the key is at the actual key time
-			InPainter.SectionGeometry.ToPaintGeometry(
-				HandlePos - FVector2D(HandleSize.X*0.5f, 0.f),
-				HandleSize
-			),
-			EasingHandle,
-			DrawEffects,
-			(bLeftHandleActive ? SelectionColor : EasingHandle->GetTint(FWidgetStyle()))
-		);
+		if (UnderlappingSectionObj->HasStartFrame())
+		{
+			TRange<FFrameNumber> EaseInRange = UnderlappingSectionObj->GetEaseInRange();
+			// Always draw handles if the section is highlighted, even if there is no range (to allow manual adjustment)
+			FFrameNumber HandleFrame = EaseInRange.IsEmpty() ? UnderlappingSectionObj->GetInclusiveStartFrame() : MovieScene::DiscreteExclusiveUpper(EaseInRange);
+			FVector2D HandlePos(TimeToPixelConverter.FrameToPixel(HandleFrame), 0.f);
+			FSlateDrawElement::MakeBox(
+				InPainter.DrawElements,
+				// always draw selected keys on top of other keys
+				InPainter.LayerId,
+				// Center the key along X.  Ensure the middle of the key is at the actual key time
+				InPainter.SectionGeometry.ToPaintGeometry(
+					HandlePos - FVector2D(HandleSize.X*0.5f, 0.f),
+					HandleSize
+				),
+				EasingHandle,
+				DrawEffects,
+				(bLeftHandleActive ? SelectionColor : EasingHandle->GetTint(FWidgetStyle()))
+			);
+		}
 
-		TRange<float> EaseOutRange = UnderlappingSectionObj->GetEaseOutRange();
-		// Always draw handles if the section is highlighted, even if there is no range (to allow manual adjustment)
-		HandlePos = FVector2D(TimeToPixelConverter.TimeToPixel(EaseOutRange.IsEmpty() ? UnderlappingSectionObj->GetEndTime() : EaseOutRange.GetLowerBoundValue()), 0.f);
-		FSlateDrawElement::MakeBox(
-			InPainter.DrawElements,
-			// always draw selected keys on top of other keys
-			InPainter.LayerId,
-			// Center the key along X.  Ensure the middle of the key is at the actual key time
-			InPainter.SectionGeometry.ToPaintGeometry(
-				HandlePos - FVector2D(HandleSize.X*0.5f, 0.f),
-				HandleSize
-			),
-			EasingHandle,
-			DrawEffects,
-			(bRightHandleActive ? SelectionColor : EasingHandle->GetTint(FWidgetStyle()))
-		);
+		if (UnderlappingSectionObj->HasEndFrame())
+		{
+			TRange<FFrameNumber> EaseOutRange = UnderlappingSectionObj->GetEaseOutRange();
+
+			// Always draw handles if the section is highlighted, even if there is no range (to allow manual adjustment)
+			FFrameNumber HandleFrame = EaseOutRange.IsEmpty() ? UnderlappingSectionObj->GetExclusiveEndFrame() : MovieScene::DiscreteInclusiveLower(EaseOutRange);
+			FVector2D    HandlePos   = FVector2D(TimeToPixelConverter.FrameToPixel(HandleFrame), 0.f);
+
+			FSlateDrawElement::MakeBox(
+				InPainter.DrawElements,
+				// always draw selected keys on top of other keys
+				InPainter.LayerId,
+				// Center the key along X.  Ensure the middle of the key is at the actual key time
+				InPainter.SectionGeometry.ToPaintGeometry(
+					HandlePos - FVector2D(HandleSize.X*0.5f, 0.f),
+					HandleSize
+				),
+				EasingHandle,
+				DrawEffects,
+				(bRightHandleActive ? SelectionColor : EasingHandle->GetTint(FWidgetStyle()))
+			);
+		}
 	}
 }
 
@@ -1533,12 +1669,13 @@ void SSequencerSection::DrawSectionHandles( const FGeometry& AllottedGeometry, F
 	}
 
 	FGeometry SectionGeometryWithoutHandles = MakeSectionGeometryWithoutHandles(AllottedGeometry, SectionInterface);
-	FTimeToPixel TimeToPixelConverter(SectionGeometryWithoutHandles, ThisSection->IsInfinite() ? GetSequencer().GetViewRange() : ThisSection->GetRange());
+	FTimeToPixel TimeToPixelConverter = ConstructTimeConverterForSection(SectionGeometryWithoutHandles, *ThisSection, GetSequencer());
+
 	for (FSectionHandle Handle : AllUnderlappingSections)
 	{
 		TSharedRef<ISequencerSection> UnderlappingSection =  Handle.TrackNode->GetSections()[Handle.SectionIndex];
 		UMovieSceneSection* UnderlappingSectionObj = UnderlappingSection->GetSectionObject();
-		if (!UnderlappingSection->SectionIsResizable() || UnderlappingSectionObj->IsInfinite())
+		if (!UnderlappingSection->SectionIsResizable() || UnderlappingSectionObj->GetRange() == TRange<FFrameNumber>::All())
 		{
 			continue;
 		}
@@ -1573,47 +1710,50 @@ void SSequencerSection::DrawSectionHandles( const FGeometry& AllottedGeometry, F
 		const float ThisHandleOffset = UnderlappingSectionObj == ThisSection ? HandleOffsetPx : 0.f;
 		FVector2D GripSize( UnderlappingSection->GetSectionGripSize(), AllottedGeometry.Size.Y );
 
-		// Make areas to the left and right of the geometry.  We will use these areas to determine if someone dragged the left or right edge of a section
-		FGeometry SectionRectLeft = SectionGeometryWithoutHandles.MakeChild(
-			FVector2D( TimeToPixelConverter.TimeToPixel(UnderlappingSectionObj->GetStartTime()) - ThisHandleOffset, 0.f ),
-			GripSize
-		);
-
-		FGeometry SectionRectRight = SectionGeometryWithoutHandles.MakeChild(
-			FVector2D( TimeToPixelConverter.TimeToPixel(UnderlappingSectionObj->GetEndTime()) - UnderlappingSection->GetSectionGripSize() + ThisHandleOffset, 0 ), 
-			GripSize
-		);
+		float Opacity = 0.5f;
+		if (ThisHandleOffset != 0)
+		{
+			Opacity = FMath::Clamp(.5f + ThisHandleOffset / GripSize.X * .5f, .5f, 1.f);
+		}
 
 		const FSlateBrush* LeftGripBrush = FEditorStyle::GetBrush("Sequencer.Section.GripLeft");
 		const FSlateBrush* RightGripBrush = FEditorStyle::GetBrush("Sequencer.Section.GripRight");
 
-		float Opacity = 0.5f;
-		if (UnderlappingSectionObj == ThisSection && HandleOffsetPx != 0)
+		// Left Grip
+		if (UnderlappingSectionObj->HasStartFrame())
 		{
-			Opacity = FMath::Clamp(.5f + HandleOffsetPx / UnderlappingSection->GetSectionGripSize() * .5f, .5f, 1.f);
+			FGeometry SectionRectLeft = SectionGeometryWithoutHandles.MakeChild(
+				FVector2D( TimeToPixelConverter.FrameToPixel(UnderlappingSectionObj->GetInclusiveStartFrame()) - ThisHandleOffset, 0.f ),
+				GripSize
+			);
+			FSlateDrawElement::MakeBox
+			(
+				OutDrawElements,
+				LayerId,
+				SectionRectLeft.ToPaintGeometry(),
+				LeftGripBrush,
+				DrawEffects,
+				(bLeftHandleActive ? SelectionColor : LeftGripBrush->GetTint(FWidgetStyle())).CopyWithNewOpacity(Opacity)
+			);
 		}
-		
-		// Left Grip
-		FSlateDrawElement::MakeBox
-		(
-			OutDrawElements,
-			LayerId,
-			SectionRectLeft.ToPaintGeometry(),
-			LeftGripBrush,
-			DrawEffects,
-			(bLeftHandleActive ? SelectionColor : LeftGripBrush->GetTint(FWidgetStyle())).CopyWithNewOpacity(Opacity)
-		);
-		
+
 		// Right Grip
-		FSlateDrawElement::MakeBox
-		(
-			OutDrawElements,
-			LayerId,
-			SectionRectRight.ToPaintGeometry(),
-			RightGripBrush,
-			DrawEffects,
-			(bRightHandleActive ? SelectionColor : RightGripBrush->GetTint(FWidgetStyle())).CopyWithNewOpacity(Opacity)
-		);
+		if (UnderlappingSectionObj->HasEndFrame())
+		{
+			FGeometry SectionRectRight = SectionGeometryWithoutHandles.MakeChild(
+				FVector2D( TimeToPixelConverter.FrameToPixel(UnderlappingSectionObj->GetExclusiveEndFrame()) - UnderlappingSection->GetSectionGripSize() + ThisHandleOffset, 0 ), 
+				GripSize
+			);
+			FSlateDrawElement::MakeBox
+			(
+				OutDrawElements,
+				LayerId,
+				SectionRectRight.ToPaintGeometry(),
+				RightGripBrush,
+				DrawEffects,
+				(bRightHandleActive ? SelectionColor : RightGripBrush->GetTint(FWidgetStyle())).CopyWithNewOpacity(Opacity)
+			);
+		}
 	}
 
 	OutDrawElements.PopClip();
@@ -1629,24 +1769,46 @@ void SSequencerSection::Tick( const FGeometry& AllottedGeometry, const double In
 	{
 		Layout = FSectionLayout(*ParentSectionArea, SectionIndex);
 
-		// Update cached key area key positions
+		// Update cached key area key positions by retaining existing caches where possible
+		TMap<FSectionLayoutElement, TArray<FSequencerCachedKeys, TInlineAllocator<1>>, FDefaultSetAllocator, FLayoutElementKeyFuncs> OldCachedKeyAreaPositions;
+		Swap(OldCachedKeyAreaPositions, CachedKeyAreaPositions);
+
+		FFrameRate FrameResolution = GetSequencer().GetFocusedFrameResolution();
+
+		// Move over any existing still valid caches 
 		for (const FSectionLayoutElement& LayoutElement : Layout->GetElements())
 		{
-			TSharedPtr<IKeyArea> KeyArea = LayoutElement.GetKeyArea();
-			if (KeyArea.IsValid())
+			TArray<FSequencerCachedKeys, TInlineAllocator<1>>* CacheArray = CachedKeyAreaPositions.Find(LayoutElement);
+			if (!CacheArray)
 			{
-				CachedKeyAreaPositions.FindOrAdd(KeyArea).Update(KeyArea.ToSharedRef());
+				// A new cache needs to be created
+				TArray<FSequencerCachedKeys, TInlineAllocator<1>> NewCachedKeys;
+				for (TSharedRef<IKeyArea> KeyArea : LayoutElement.GetKeyAreas())
+				{
+					NewCachedKeys.Emplace();
+					NewCachedKeys.Last().Update(KeyArea, FrameResolution);
+				}
+				CachedKeyAreaPositions.Add(LayoutElement, MoveTemp(NewCachedKeys));
+			}
+			else
+			{
+				// We can reuse this one
+				for (FSequencerCachedKeys& CachedKeys : *CacheArray)
+				{
+					CachedKeys.Update(CachedKeys.GetKeyArea().ToSharedRef(), FrameResolution);
+				}
+				CachedKeyAreaPositions.Add(LayoutElement, MoveTemp(*CacheArray));
 			}
 		}
 
 		UMovieSceneSection* Section = SectionInterface->GetSectionObject();
-		if (Section && !Section->IsInfinite())
+		if (Section && Section->HasStartFrame() && Section->HasEndFrame())
 		{
-			FTimeToPixel TimeToPixelConverter(ParentGeometry, GetSequencer().GetViewRange());
+			FTimeToPixel TimeToPixelConverter(ParentGeometry, GetSequencer().GetViewRange(), FrameResolution);
 
 			const int32 SectionLengthPx = FMath::Max(0,
 				FMath::RoundToInt(
-					TimeToPixelConverter.TimeToPixel(Section->GetEndTime())) - FMath::RoundToInt(TimeToPixelConverter.TimeToPixel(Section->GetStartTime())
+					TimeToPixelConverter.FrameToPixel(Section->GetExclusiveEndFrame())) - FMath::RoundToInt(TimeToPixelConverter.FrameToPixel(Section->GetInclusiveStartFrame())
 				)
 			);
 
@@ -1669,13 +1831,13 @@ FReply SSequencerSection::OnMouseButtonDown( const FGeometry& MyGeometry, const
 {
 	FSequencer& Sequencer = GetSequencer();
 
-	FSequencerSelectedKey HoveredKey;
+	TArrayView<const FSequencerSelectedKey> HoveredKeys;
 	
 	// The hovered key is defined from the sequencer hotspot
 	TSharedPtr<ISequencerHotspot> Hotspot = GetSequencer().GetHotspot();
 	if (Hotspot.IsValid() && Hotspot->GetType() == ESequencerHotspot::Key)
 	{
-		HoveredKey = static_cast<FKeyHotspot*>(Hotspot.Get())->Key;
+		HoveredKeys = static_cast<FKeyHotspot*>(Hotspot.Get())->Keys;
 	}
 
 	if (MouseEvent.GetEffectingButton() == EKeys::MiddleMouseButton)
@@ -1683,17 +1845,19 @@ FReply SSequencerSection::OnMouseButtonDown( const FGeometry& MyGeometry, const
 		GEditor->BeginTransaction(NSLOCTEXT("Sequencer", "CreateKey_Transaction", "Create Key"));
 
 		// Generate a key and set it as the PressedKey
-		FSequencerSelectedKey NewKey = CreateKeyUnderMouse(MouseEvent.GetScreenSpacePosition(), MyGeometry, HoveredKey);
+		TArray<FSequencerSelectedKey> NewKeys;
+		CreateKeysUnderMouse(MouseEvent.GetScreenSpacePosition(), MyGeometry, HoveredKeys, NewKeys);
 
-		if (NewKey.IsValid())
+		if (NewKeys.Num())
 		{
-			HoveredKey = NewKey;
-
 			Sequencer.GetSelection().EmptySelectedKeys();
-			Sequencer.GetSelection().AddToSelection(NewKey);
+			for (const FSequencerSelectedKey& NewKey : NewKeys)
+			{
+				Sequencer.GetSelection().AddToSelection(NewKey);
+			}
 
 			// Pass the event to the tool to copy the hovered key and move it
-			GetSequencer().SetHotspot( MakeShareable( new FKeyHotspot(NewKey) ) );
+			GetSequencer().SetHotspot( MakeShared<FKeyHotspot>(MoveTemp(NewKeys)) );
 
 			// Return unhandled so that the EditTool can handle the mouse down based on the newly created keyframe and prepare to move it
 			return FReply::Unhandled();
@@ -1768,10 +1932,11 @@ FReply SSequencerSection::OnMouseButtonDoubleClick( const FGeometry& MyGeometry,
 FReply SSequencerSection::OnMouseMove( const FGeometry& MyGeometry, const FPointerEvent& MouseEvent )
 {
 	// Checked for hovered key
-	FSequencerSelectedKey KeyUnderMouse = GetKeyUnderMouse( MouseEvent.GetScreenSpacePosition(), MyGeometry );
-	if ( KeyUnderMouse.IsValid() )
+	TArray<FSequencerSelectedKey> KeysUnderMouse;
+	GetKeysUnderMouse( MouseEvent.GetScreenSpacePosition(), MyGeometry, KeysUnderMouse );
+	if ( KeysUnderMouse.Num() )
 	{
-		GetSequencer().SetHotspot( MakeShareable( new FKeyHotspot(KeyUnderMouse) ) );
+		GetSequencer().SetHotspot( MakeShared<FKeyHotspot>( MoveTemp(KeysUnderMouse) ) );
 	}
 	// Check other interaction points in order of importance
 	else if (
@@ -1803,10 +1968,16 @@ void SSequencerSection::OnMouseLeave( const FPointerEvent& MouseEvent )
 	GetSequencer().SetHotspot(nullptr);
 }
 
-static float ThrobDurationSeconds = .5f;
-void SSequencerSection::ThrobSelection(int32 ThrobCount)
+static float SectionThrobDurationSeconds = 1.f;
+void SSequencerSection::ThrobSectionSelection(int32 ThrobCount)
 {
-	SelectionThrobEndTime = FPlatformTime::Seconds() + ThrobCount*ThrobDurationSeconds;
+	SectionSelectionThrobEndTime = FPlatformTime::Seconds() + ThrobCount*SectionThrobDurationSeconds;
+}
+
+static float KeyThrobDurationSeconds = .5f;
+void SSequencerSection::ThrobKeySelection(int32 ThrobCount)
+{
+	KeySelectionThrobEndTime = FPlatformTime::Seconds() + ThrobCount*KeyThrobDurationSeconds;
 }
 
 float EvaluateThrob(float Alpha)
@@ -1814,14 +1985,27 @@ float EvaluateThrob(float Alpha)
 	return .5f - FMath::Cos(FMath::Pow(Alpha, 0.5f) * 2 * PI) * .5f;
 }
 
-float SSequencerSection::GetSelectionThrobValue()
+float SSequencerSection::GetSectionSelectionThrobValue()
+{
+	double CurrentTime = FPlatformTime::Seconds();
+
+	if (SectionSelectionThrobEndTime > CurrentTime)
+	{
+		float Difference = SectionSelectionThrobEndTime - CurrentTime;
+		return EvaluateThrob(1.f - FMath::Fmod(Difference, SectionThrobDurationSeconds));
+	}
+
+	return 0.f;
+}
+
+float SSequencerSection::GetKeySelectionThrobValue()
 {
 	double CurrentTime = FPlatformTime::Seconds();
 
-	if (SelectionThrobEndTime > CurrentTime)
+	if (KeySelectionThrobEndTime > CurrentTime)
 	{
-		float Difference = SelectionThrobEndTime - CurrentTime;
-		return EvaluateThrob(1.f - FMath::Fmod(Difference, ThrobDurationSeconds));
+		float Difference = KeySelectionThrobEndTime - CurrentTime;
+		return EvaluateThrob(1.f - FMath::Fmod(Difference, KeyThrobDurationSeconds));
 	}
 
 	return 0.f;
@@ -1837,17 +2021,22 @@ bool SSequencerSection::IsSectionHighlighted(FSectionHandle InSectionHandle, con
 	switch(Hotspot->GetType())
 	{
 	case ESequencerHotspot::Key:
-		return static_cast<const FKeyHotspot*>(Hotspot)->Key.Section == InSectionHandle.GetSectionObject();
+		return static_cast<const FKeyHotspot*>(Hotspot)->Keys.ContainsByPredicate([InSectionHandle](const FSequencerSelectedKey& Key){ return Key.Section == InSectionHandle.GetSectionObject(); });
+
 	case ESequencerHotspot::Section:
 		return static_cast<const FSectionHotspot*>(Hotspot)->Section == InSectionHandle;
+
 	case ESequencerHotspot::SectionResize_L:
 	case ESequencerHotspot::SectionResize_R:
 		return static_cast<const FSectionResizeHotspot*>(Hotspot)->Section == InSectionHandle;
+
 	case ESequencerHotspot::EaseInHandle:
 	case ESequencerHotspot::EaseOutHandle:
 		return static_cast<const FSectionEasingHandleHotspot*>(Hotspot)->Section == InSectionHandle;
+
 	case ESequencerHotspot::EasingArea:
 		return static_cast<const FSectionEasingAreaHotspot*>(Hotspot)->Contains(InSectionHandle);
+
 	default:
 		return false;
 	}
diff --git a/Engine/Source/Editor/Sequencer/Private/SSequencerSection.h b/Engine/Source/Editor/Sequencer/Private/SSequencerSection.h
--- a/Engine/Source/Editor/Sequencer/Private/SSequencerSection.h
+++ b/Engine/Source/Editor/Sequencer/Private/SSequencerSection.h
@@ -51,19 +51,20 @@ private:
 	 *
 	 * @param MousePosition		The current screen space position of the mouse
 	 * @param AllottedGeometry	The geometry of the mouse event
+	 * @param OutKeys			Array to populate with the keys that are under the mouse
 	 * @return The key that is under the mouse.  Invalid if there is no key under the mouse
 	 */
-	FSequencerSelectedKey GetKeyUnderMouse( const FVector2D& MousePosition, const FGeometry& AllottedGeometry ) const;
+	void GetKeysUnderMouse( const FVector2D& MousePosition, const FGeometry& AllottedGeometry, TArray<FSequencerSelectedKey>& OutKeys ) const;
 
 	/**
 	 * Creates a key at the mouse position
 	 *
 	 * @param MousePosition		The current screen space position of the mouse
 	 * @param AllottedGeometry	The geometry of the mouse event
-	 * @param InPressedKey      Key if pressed
-	 * @return The newly created key
+	 * @param InPressedKeys     Keys if pressed
+	 * @param OutKeys			Array to populate with the new keys
 	 */
-	FSequencerSelectedKey CreateKeyUnderMouse( const FVector2D& MousePosition, const FGeometry& AllottedGeometry, FSequencerSelectedKey InPressedKey );
+	void CreateKeysUnderMouse( const FVector2D& MousePosition, const FGeometry& AllottedGeometry, TArrayView<const FSequencerSelectedKey> InPressedKeys, TArray<FSequencerSelectedKey>& OutKeys );
 
 	/**
 	 * Checks for user interaction (via the mouse) with the left and right edge of a section
@@ -118,11 +119,17 @@ private:
 
 public:
 
+	/** Indicate that the current section selection should throb the specified number of times. A single throb takes 0.2s. */
+	static void ThrobSectionSelection(int32 ThrobCount = 1);
+
 	/** Indicate that the current key selection should throb the specified number of times. A single throb takes 0.2s. */
-	static void ThrobSelection(int32 ThrobCount = 1);
+	static void ThrobKeySelection(int32 ThrobCount = 1);
+
+	/** Get a value between 0 and 1 that indicates the amount of throb-scale to apply to the currently selected sections */
+	static float GetSectionSelectionThrobValue();
 
-	/** Get a value between 0 and 1 that indicicates the amount of throb-scale to apply to the currently selected keys */
-	static float GetSelectionThrobValue();
+	/** Get a value between 0 and 1 that indicates the amount of throb-scale to apply to the currently selected keys */
+	static float GetKeySelectionThrobValue();
 
 	/**
 	 * Check to see whether the specified section is highlighted
@@ -140,8 +147,10 @@ private:
 	int32 SectionIndex;
 	/** Cached parent geometry to pass down to any section interfaces that need it during tick */
 	FGeometry ParentGeometry;
+	/** The end time for a throbbing animation for selected sections */
+	static double SectionSelectionThrobEndTime;
 	/** The end time for a throbbing animation for selected keys */
-	static double SelectionThrobEndTime;
+	static double KeySelectionThrobEndTime;
 	/** Handle offset amount in pixels */
 	float HandleOffsetPx;
 	/** Array of segments that define other sections that reside below this one */
@@ -150,8 +159,18 @@ private:
 	TArray<FSequencerOverlapRange> UnderlappingEasingSegments;
 	/** The signature of the track last time the overlapping segments were updated */
 	FGuid CachedTrackSignature;
+
+	/** Key funcs for looking up a set of cached keys by its layout element */
+	struct FLayoutElementKeyFuncs : BaseKeyFuncs<TArray<FSequencerCachedKeys, TInlineAllocator<1>>, FSectionLayoutElement>
+	{
+		enum { bAllowDuplicateKeys = false };
+
+		static const FSectionLayoutElement& GetSetKey(const TPair<FSectionLayoutElement, TArray<FSequencerCachedKeys, TInlineAllocator<1>>>& Element);
+		static bool Matches(const FSectionLayoutElement& A, const FSectionLayoutElement& B);
+		static uint32 GetKeyHash(const FSectionLayoutElement& Key);
+	};
 	/** Cache of key area positions */
-	TMap<TSharedPtr<IKeyArea>, FSequencerCachedKeys> CachedKeyAreaPositions;
+	TMap<FSectionLayoutElement, TArray<FSequencerCachedKeys, TInlineAllocator<1>>, FDefaultSetAllocator, FLayoutElementKeyFuncs> CachedKeyAreaPositions;
 
 	friend struct FSequencerSectionPainterImpl;
 };
diff --git a/Engine/Source/Editor/Sequencer/Private/SSequencerSectionAreaView.cpp b/Engine/Source/Editor/Sequencer/Private/SSequencerSectionAreaView.cpp
--- a/Engine/Source/Editor/Sequencer/Private/SSequencerSectionAreaView.cpp
+++ b/Engine/Source/Editor/Sequencer/Private/SSequencerSectionAreaView.cpp
@@ -4,6 +4,8 @@
 #include "Types/PaintArgs.h"
 #include "Layout/ArrangedChildren.h"
 #include "CommonMovieSceneTools.h"
+#include "MovieSceneTrack.h"
+#include "MovieScene.h"
 
 namespace SequencerSectionAreaConstants
 {
@@ -24,24 +26,19 @@ namespace SequencerSectionUtils
 	 */
 	FGeometry GetSectionGeometry( const FGeometry& AllottedGeometry, int32 RowIndex, int32 MaxTracks, float NodeHeight, TSharedPtr<ISequencerSection> SectionInterface, const FTimeToPixel& TimeToPixelConverter )
 	{
-		const UMovieSceneSection* Section = SectionInterface->GetSectionObject();
-		float PixelStartX = TimeToPixelConverter.TimeToPixel( Section->GetStartTime() );
-		// Note the -1 pixel at the end is because the section does not actually end at the end time if there is a section starting at that same time.  It is more important that a section lines up correctly with it's true start time
-		float PixelEndX = TimeToPixelConverter.TimeToPixel( Section->GetEndTime() );
+		const UMovieSceneSection* Section    = SectionInterface->GetSectionObject();
+		const FFrameRate          Resolution = Section->GetTypedOuter<UMovieScene>()->GetFrameResolution();
 
-		// If the section is infinite, occupy the entire width of the geometry where the section is located.
-		const bool bIsInfinite = Section->IsInfinite();
-		if (bIsInfinite)
-		{
-			PixelStartX = AllottedGeometry.Position.X;
-			PixelEndX = AllottedGeometry.Position.X + AllottedGeometry.GetLocalSize().X;
-		}
+		// Sections with an infite (open) start bound use the currently visible geometry and time range
+		float PixelStartX = Section->HasStartFrame() ? TimeToPixelConverter.FrameToPixel( Section->GetInclusiveStartFrame() )     : AllottedGeometry.Position.X;
+		// Note the -1 pixel at the end is because the section does not actually end at the end time if there is a section starting at that same time.  It is more important that a section lines up correctly with it's true start time
+		float PixelEndX   = Section->HasEndFrame()   ? TimeToPixelConverter.FrameToPixel( Section->GetExclusiveEndFrame() ) : AllottedGeometry.Position.X + AllottedGeometry.GetLocalSize().X;
 
 		const float MinSectionWidth = 1.f;
 		float SectionLength = FMath::Max(MinSectionWidth, PixelEndX-PixelStartX);
 
 		float GripOffset = 0;
-		if (!bIsInfinite)
+		if (Section->HasStartFrame() && Section->HasEndFrame())
 		{
 			float NewSectionLength = FMath::Max(SectionLength, MinSectionWidth + SectionInterface->GetSectionGripSize() * 2.f);
 
@@ -69,7 +66,7 @@ void SSequencerSectionAreaView::Construct( const FArguments& InArgs, TSharedRef<
 
 	// Generate widgets for sections in this view
 	GenerateSectionWidgets();
-};
+}
 
 FVector2D SSequencerSectionAreaView::ComputeDesiredSize(float) const
 {
@@ -188,7 +185,10 @@ void SSequencerSectionAreaView::OnArrangeChildren( const FGeometry& AllottedGeom
 
 FTimeToPixel SSequencerSectionAreaView::GetTimeToPixel( const FGeometry& AllottedGeometry ) const
 {
-	return FTimeToPixel( AllottedGeometry, ViewRange.Get() );
+	const UMovieSceneTrack* Track      = SectionAreaNode->GetTrack();
+	const FFrameRate        Resolution = Track->GetTypedOuter<UMovieScene>()->GetFrameResolution();
+
+	return FTimeToPixel( AllottedGeometry, ViewRange.Get(), Resolution );
 }
 
 
diff --git a/Engine/Source/Editor/Sequencer/Private/SSequencerSectionAreaView.h b/Engine/Source/Editor/Sequencer/Private/SSequencerSectionAreaView.h
--- a/Engine/Source/Editor/Sequencer/Private/SSequencerSectionAreaView.h
+++ b/Engine/Source/Editor/Sequencer/Private/SSequencerSectionAreaView.h
@@ -23,9 +23,9 @@ class SSequencerSectionAreaView : public SPanel
 {
 public:
 	SLATE_BEGIN_ARGS( SSequencerSectionAreaView )
-	{}		
+	{}
 		/** The view range of the section area */
-		SLATE_ATTRIBUTE( TRange<float>, ViewRange )
+		SLATE_ATTRIBUTE( TRange<double>, ViewRange )
 	SLATE_END_ARGS()
 
 	SSequencerSectionAreaView()
@@ -71,7 +71,7 @@ private:
 	/** The node containing the sections we are viewing/manipulating */
 	TSharedPtr<FSequencerTrackNode> SectionAreaNode;
 	/** The current view range */
-	TAttribute< TRange<float> > ViewRange;
+	TAttribute< TRange<double> > ViewRange;
 	/** All the widgets in the panel */
 	TSlotlessChildren<SSequencerSection> Children;
 };
diff --git a/Engine/Source/Editor/Sequencer/Private/SSequencerShotFilterOverlay.cpp b/Engine/Source/Editor/Sequencer/Private/SSequencerShotFilterOverlay.cpp
--- a/Engine/Source/Editor/Sequencer/Private/SSequencerShotFilterOverlay.cpp
+++ b/Engine/Source/Editor/Sequencer/Private/SSequencerShotFilterOverlay.cpp
@@ -21,22 +21,22 @@ void SSequencerShotFilterOverlay::Construct(const FArguments& InArgs, TSharedRef
 int32 SSequencerShotFilterOverlay::OnPaint( const FPaintArgs& Args, const FGeometry& AllottedGeometry, const FSlateRect& MyCullingRect, FSlateWindowElementList& OutDrawElements, int32 LayerId, const FWidgetStyle& InWidgetStyle, bool bParentEnabled ) const
 {
 	float Alpha = Sequencer.Pin()->GetOverlayFadeCurve();
-		
+
 	if (Alpha > 0.f)
 	{
-		FTimeToPixel TimeToPixelConverter = FTimeToPixel(AllottedGeometry, ViewRange.Get());
+		FTimeToPixel TimeToPixelConverter(AllottedGeometry, ViewRange.Get(), Sequencer.Pin()->GetFocusedFrameResolution());
 		
 		TRange<float> TimeBounds = TRange<float>(
-			TimeToPixelConverter.PixelToTime(0),
-			TimeToPixelConverter.PixelToTime(AllottedGeometry.GetLocalSize().X)
+			TimeToPixelConverter.PixelToSeconds(0),
+			TimeToPixelConverter.PixelToSeconds(AllottedGeometry.GetLocalSize().X)
 		);
 
 		TArray< TRange<float> > OverlayRanges = ComputeOverlayRanges(TimeBounds, CachedFilteredRanges);
 
 		for (int32 i = 0; i < OverlayRanges.Num(); ++i)
 		{
-			float LowerBound = TimeToPixelConverter.TimeToPixel(OverlayRanges[i].GetLowerBoundValue());
-			float UpperBound = TimeToPixelConverter.TimeToPixel(OverlayRanges[i].GetUpperBoundValue());
+			float LowerBound = TimeToPixelConverter.SecondsToPixel(OverlayRanges[i].GetLowerBoundValue());
+			float UpperBound = TimeToPixelConverter.SecondsToPixel(OverlayRanges[i].GetUpperBoundValue());
 
 			FSlateDrawElement::MakeBox(
 				OutDrawElements,
diff --git a/Engine/Source/Editor/Sequencer/Private/SSequencerShotFilterOverlay.h b/Engine/Source/Editor/Sequencer/Private/SSequencerShotFilterOverlay.h
--- a/Engine/Source/Editor/Sequencer/Private/SSequencerShotFilterOverlay.h
+++ b/Engine/Source/Editor/Sequencer/Private/SSequencerShotFilterOverlay.h
@@ -22,7 +22,7 @@ class SSequencerShotFilterOverlay
 {
 public:
 	SLATE_BEGIN_ARGS(SSequencerShotFilterOverlay) { }
-		SLATE_ATTRIBUTE(TRange<float>, ViewRange)
+		SLATE_ATTRIBUTE(TRange<double>, ViewRange)
 	SLATE_END_ARGS()
 
 	void Construct(const FArguments& InArgs, TSharedRef<FSequencer> InSequencer);
@@ -45,7 +45,7 @@ protected:
 private:
 
 	/** The current minimum view range */
-	TAttribute<TRange<float>> ViewRange;
+	TAttribute<TRange<double>> ViewRange;
 
 	/** The main sequencer interface */
 	TWeakPtr<FSequencer> Sequencer;
