Copying //UE4/Dev-Sequencer to //UE4/Dev-Main (Source: //UE4/Dev-Sequencer @ 3946692)

diff --git a/Engine/Source/Editor/Sequencer/Private/SSequencerTimePanel.cpp b/Engine/Source/Editor/Sequencer/Private/SSequencerTimePanel.cpp
new file mode 100644
--- /dev/null
+++ b/Engine/Source/Editor/Sequencer/Private/SSequencerTimePanel.cpp
@@ -0,0 +1,486 @@
+// Copyright 1998-2018 Epic Games, Inc. All Rights Reserved.
+
+#include "SSequencerTimePanel.h"
+#include "Widgets/SFrameRatePicker.h"
+#include "Widgets/SBoxPanel.h"
+#include "Widgets/Layout/SGridPanel.h"
+#include "Widgets/Layout/SScrollBox.h"
+#include "Widgets/Text/STextBlock.h"
+#include "Widgets/Input/SCheckBox.h"
+#include "Widgets/Input/SComboBox.h"
+#include "Widgets/Input/SButton.h"
+#include "MovieSceneSequence.h"
+#include "MovieScene.h"
+#include "MovieSceneTrack.h"
+#include "Channels/MovieSceneChannelProxy.h"
+#include "MovieSceneSection.h"
+#include "ScopedTransaction.h"
+#include "Sequencer.h"
+#include "EditorStyleSet.h"
+#include "Misc/ScopedSlowTask.h"
+#include "MovieSceneToolHelpers.h"
+#include "Sections/MovieSceneSubSection.h"
+
+
+#define LOCTEXT_NAMESPACE "SSequencerTimePanel"
+
+namespace
+{
+
+TRange<FFrameNumber> MigrateFrameRange(const TRange<FFrameNumber>& SourceRange, FFrameRate SourceRate, FFrameRate DestinationRate)
+{
+	if (SourceRate == DestinationRate)
+	{
+		return SourceRange;
+	}
+
+	TRange<FFrameNumber> NewRange = TRange<FFrameNumber>::All();
+
+	if (!SourceRange.GetLowerBound().IsOpen())
+	{
+		const FFrameNumber FrameNumber = ConvertFrameTime(SourceRange.GetLowerBoundValue(), SourceRate, DestinationRate).RoundToFrame();
+
+		NewRange.SetLowerBound(
+			SourceRange.GetLowerBound().IsExclusive()
+			? TRangeBound<FFrameNumber>::Exclusive(FrameNumber)
+			: TRangeBound<FFrameNumber>::Inclusive(FrameNumber)
+		);
+	}
+
+	if (!SourceRange.GetUpperBound().IsOpen())
+	{
+		const FFrameNumber FrameNumber = ConvertFrameTime(SourceRange.GetUpperBoundValue(), SourceRate, DestinationRate).RoundToFrame();
+
+		NewRange.SetUpperBound(
+			SourceRange.GetUpperBound().IsExclusive()
+			? TRangeBound<FFrameNumber>::Exclusive(FrameNumber)
+			: TRangeBound<FFrameNumber>::Inclusive(FrameNumber)
+		);
+	}
+
+	return NewRange;
+}
+
+}    // anon namespace
+
+
+void SSequencerTimePanel::Construct(const FArguments& InArgs, TWeakPtr<FSequencer> InSequencer)
+{
+	WeakSequencer = InSequencer;
+
+	TSharedPtr<FSequencer> Sequencer = InSequencer.Pin();
+	if (Sequencer.IsValid())
+	{
+		CurrentFrameResolution = Sequencer->GetFocusedFrameResolution();
+	}
+
+	TArray<FCommonFrameRateInfo> FrameResolutionRates;
+	{
+		TArrayView<const FCommonFrameRateInfo> CommonRates = FCommonFrameRates::GetAll();
+		FrameResolutionRates.Append(CommonRates.GetData(), CommonRates.Num());
+
+		FrameResolutionRates.Add(FCommonFrameRateInfo{ FFrameRate(1000, 1),   LOCTEXT("1000_Name",   "1000 fps (ms precision)"),                       LOCTEXT("1000_Description",   "Allows placement of sequence keys and sections with millisecond precision") });
+		FrameResolutionRates.Add(FCommonFrameRateInfo{ FFrameRate(24000, 1),  LOCTEXT("24000_Name",  "24000 fps (all integer rates + 23.976)"),        LOCTEXT("24000_Description",  "A very high framerate that allows frame-accurate evaluation of all common integer frame rates as well as NTSC 24.") });
+		FrameResolutionRates.Add(FCommonFrameRateInfo{ FFrameRate(60000, 1),  LOCTEXT("60000_Name",  "60000 fps (all integer rates + 29.97 & 59.94)"), LOCTEXT("60000_Description",  "A very high framerate that allows frame-accurate evaluation of all common integer frame rates as well as NTSC 30 and 60.") });
+		FrameResolutionRates.Add(FCommonFrameRateInfo{ FFrameRate(120000, 1), LOCTEXT("120000_Name", "120000 fps (all common rates)"),                 LOCTEXT("120000_Description", "A very high framerate that allows frame-accurate evaluation of all common integer and NTSC frame rates.") });
+	}
+
+	FText Description = LOCTEXT("Description", "Sequences store section start times and keys at points in time called 'ticks'.\n\nThe size of a single tick is defined per-sequence; it is recommended that you choose a tick-interval that fits into your desired display rate or content frame rates. Increasing the resolution will reduce the total supported time range.");
+
+	static float VerticalGridPadding = 5.f;
+	static float HorizontalGridPadding = 10.f;
+
+	static FMargin Col1Padding(0.f, 0.f, HorizontalGridPadding, VerticalGridPadding);
+	static FMargin Col2Padding(HorizontalGridPadding, 0.f, 0.f, VerticalGridPadding);
+
+	static FLinearColor WarningColor(FColor(0xffbbbb44));
+	int32 CurrentRow = 0;
+	ChildSlot
+	[
+		SNew(SBorder)
+		.OnMouseButtonUp(this, &SSequencerTimePanel::OnBorderFadeClicked)
+		.BorderImage(FEditorStyle::GetBrush("BlackBrush"))
+		.BorderBackgroundColor(FLinearColor(1.f, 1.f, 1.f, 0.5f))
+		.HAlign(HAlign_Center)
+		.VAlign(VAlign_Center)
+		[
+			SNew(SBorder)
+			.BorderImage(FEditorStyle::GetBrush("Sequencer.OverlayPanel.Background"))
+			.Padding(FMargin(20.f))
+			[
+				SNew(SOverlay)
+
+				+ SOverlay::Slot()
+				[
+					SNew(SBorder)
+					.OnMouseButtonUp_Lambda([](const FGeometry&, const FPointerEvent&){ return FReply::Handled(); })
+					.BorderImage(FEditorStyle::GetBrush("None"))
+					.Padding(FMargin(10.f))
+					[
+						SNew(SVerticalBox)
+
+						+ SVerticalBox::Slot()
+						[
+							SNew(SBox)
+							.MaxDesiredWidth(600.f)
+							[
+								SNew(SScrollBox)
+
+								+ SScrollBox::Slot()
+								[
+									SNew(STextBlock)
+									.TextStyle(FEditorStyle::Get(), "LargeText")
+									.Text(LOCTEXT("Title", "Advanced Time Properties"))
+								]
+
+								+ SScrollBox::Slot()
+								.Padding(FMargin(0.f, 0.f, 0.f, 10.f))
+								[
+									SNew(STextBlock)
+									.AutoWrapText(true)
+									.TextStyle(FEditorStyle::Get(), "SmallText")
+									.Text(Description)
+								]
+
+								+ SScrollBox::Slot()
+								.HAlign(HAlign_Center)
+								[
+									SNew(SGridPanel)
+									.FillColumn(1, 1.f)
+
+									+ SGridPanel::Slot(0, 0)
+									.Padding(Col1Padding)
+									.HAlign(HAlign_Left)
+									[
+										SNew(STextBlock)
+										.TextStyle(FEditorStyle::Get(), "NormalText.Important")
+										.Text(LOCTEXT("NewTickInterval", "Desired Tick Interval"))
+									]
+									+ SGridPanel::Slot(1, 0)
+									.HAlign(HAlign_Left)
+									.Padding(Col2Padding)
+									[
+										SNew(SFrameRatePicker)
+										.RecommendedText(LOCTEXT("CompatibleWithDisplayRate", "Compatible with this sequence"))
+										.NotRecommendedText(LOCTEXT("NotCompatibleWithDisplayRate", "Other"))
+										.NotRecommendedToolTip(LOCTEXT("NotCompatibleWithDisplayRate_Tip", "All other preset frame rates that are not compatible with the current display and tick rate"))
+										.IsPresetRecommended(this, &SSequencerTimePanel::IsRecommendedResolution)
+										.PresetValues(MoveTemp(FrameResolutionRates))
+										.Value(this, &SSequencerTimePanel::OnGetFrameResolution)
+										.OnValueChanged(this, &SSequencerTimePanel::OnSetFrameResolution)
+									]
+
+									+ SGridPanel::Slot(0, 1)
+									.ColumnSpan(2)
+									.Padding(FMargin(0.f, VerticalGridPadding*2, 0.f, VerticalGridPadding))
+									.HAlign(HAlign_Left)
+									.VAlign(VAlign_Top)
+									[
+										SNew(STextBlock)
+										.Text(LOCTEXT("NewTickInterval_Tip", "Sequence will have the following properties if applied:"))
+										.TextStyle(FEditorStyle::Get(), "TinyText")
+									]
+
+									+ SGridPanel::Slot(0, 2)
+									.Padding(Col1Padding)
+									.HAlign(HAlign_Left)
+									[
+										SNew(STextBlock)
+										.TextStyle(FEditorStyle::Get(), "NormalText.Important")
+										.Text(LOCTEXT("ResultingRange", "Time Range"))
+									]
+									+ SGridPanel::Slot(1, 2)
+									.Padding(Col2Padding)
+									[
+										SNew(STextBlock)
+										.Text(this, &SSequencerTimePanel::GetSupportedTimeRange)
+									]
+
+									+ SGridPanel::Slot(0, 3)
+									.Padding(Col1Padding)
+									.HAlign(HAlign_Left)
+									[
+										SNew(STextBlock)
+										.TextStyle(FEditorStyle::Get(), "NormalText.Important")
+										.Text(LOCTEXT("SupportedFrameRates", "Supported Rates"))
+									]
+									+ SGridPanel::Slot(1, 3)
+									.Padding(Col2Padding)
+									[
+										SAssignNew(CommonFrameRates, SVerticalBox)
+									]
+								]
+							]
+						]
+
+						+ SVerticalBox::Slot()
+						.AutoHeight()
+						.Padding(5.f)
+						.HAlign(HAlign_Center)
+						.VAlign(VAlign_Center)
+						[
+							SNew(SHorizontalBox)
+							.Visibility(this, &SSequencerTimePanel::GetWarningVisibility)
+
+							+ SHorizontalBox::Slot()
+							.Padding(FMargin(0.f, 0.f, 7.f, 0.f))
+							.AutoWidth()
+							[
+								SNew(STextBlock)
+								.TextStyle(FEditorStyle::Get(), "TextBlock.ShadowedTextWarning")
+								.ColorAndOpacity(WarningColor)
+								.Font(FEditorStyle::Get().GetFontStyle("FontAwesome.11"))
+								.Text(FText::FromString(FString(TEXT("\xf071"))) /*fa-exclamation-triangle*/)
+							]
+
+							+ SHorizontalBox::Slot()
+							.AutoWidth()
+							[
+								SNew(STextBlock)
+								.TextStyle(FEditorStyle::Get(), "TextBlock.ShadowedTextWarning")
+								.ColorAndOpacity(WarningColor)
+								.Text(LOCTEXT("ApplyWarning", "Applying these settings may result in changes to key positions or section boundaries."))
+							]
+						]
+
+						+ SVerticalBox::Slot()
+						.AutoHeight()
+						.Padding(5.f)
+						.HAlign(HAlign_Center)
+						.VAlign(VAlign_Center)
+						[
+							SNew(SButton)
+							.OnClicked(this, &SSequencerTimePanel::Apply)
+							[
+								SNew(STextBlock)
+								.Text(LOCTEXT("ApplyButtonText", "Apply"))
+							]
+						]
+					]
+				]
+
+				+ SOverlay::Slot()
+				.HAlign(HAlign_Right)
+				.VAlign(VAlign_Top)
+				[
+					SNew(SButton)
+					.ButtonStyle(&FEditorStyle::Get().GetWidgetStyle< FButtonStyle >("Sequencer.Transport.CloseButton"))
+					.OnClicked(this, &SSequencerTimePanel::Close)
+				]
+			]
+		]
+	];
+
+	UpdateCommonFrameRates();
+}
+
+FReply SSequencerTimePanel::OnBorderFadeClicked(const FGeometry& MyGeometry, const FPointerEvent& MouseEvent)
+{
+	return Close();
+}
+
+FReply SSequencerTimePanel::Close()
+{
+	SetVisibility(EVisibility::Collapsed);
+	return FReply::Handled();
+}
+
+FReply SSequencerTimePanel::Apply()
+{
+	UMovieSceneSequence* Sequence = GetFocusedSequence();
+	UMovieScene* MovieScene       = Sequence ? Sequence->GetMovieScene() : nullptr;
+
+	if (MovieScene)
+	{
+		FFrameRate Src = MovieScene->GetFrameResolution();
+		FFrameRate Dst = CurrentFrameResolution;
+
+		FScopedTransaction ScopedTransaction(FText::Format(LOCTEXT("MigrateFrameTimes", "Convert sequence tick interval from {0} to {1}"), Src.ToPrettyText(), Dst.ToPrettyText()));
+
+		MigrateFrameTimes(Src, Dst, MovieScene);
+	}
+
+	return Close();
+}
+
+EVisibility SSequencerTimePanel::GetWarningVisibility() const
+{
+	TSharedPtr<FSequencer> Sequencer  = WeakSequencer.Pin();
+	return Sequencer.IsValid() && Sequencer->GetFocusedFrameResolution().IsMultipleOf(CurrentFrameResolution) ? EVisibility::Collapsed : EVisibility::Visible;
+}
+
+void SSequencerTimePanel::UpdateCommonFrameRates()
+{
+	TArray<FCommonFrameRateInfo> CompatibleRates;
+	for (const FCommonFrameRateInfo& Info : FCommonFrameRates::GetAll())
+	{
+		if (Info.FrameRate.IsMultipleOf(CurrentFrameResolution))
+		{
+			CompatibleRates.Add(Info);
+		}
+	}
+
+	CompatibleRates.Sort(
+		[=](const FCommonFrameRateInfo& A, const FCommonFrameRateInfo& B)
+		{
+			return A.FrameRate.AsDecimal() < B.FrameRate.AsDecimal();
+		}
+	);
+
+	CommonFrameRates->ClearChildren();
+	for (const FCommonFrameRateInfo& FrameRate : CompatibleRates)
+	{
+		CommonFrameRates->AddSlot()
+		[
+			SNew(STextBlock)
+			.Text(FrameRate.DisplayName)
+		];
+	}
+}
+
+FText SSequencerTimePanel::GetSupportedTimeRange() const
+{
+	double FrameRate = CurrentFrameResolution.AsDecimal();
+
+	int64 TotalMaxSeconds = static_cast<int64>(TNumericLimits<int32>::Max() / FrameRate);
+
+	int64 Days    =  (TotalMaxSeconds        ) / 86400;
+	int64 Hours   =  (TotalMaxSeconds % 86400) / 3600;
+	int64 Minutes =  (TotalMaxSeconds % 3600 ) / 60;
+	int64 Seconds =  (TotalMaxSeconds % 60   );
+
+	FString String = Days > 0
+		? FString::Printf(TEXT("+/- %02dd %02dh %02dm %02ds"), Days, Hours, Minutes, Seconds)
+		: FString::Printf(TEXT("+/- %02dh %02dm %02ds"), Hours, Minutes, Seconds);
+
+	return FText::FromString(String);
+}
+
+UMovieSceneSequence* SSequencerTimePanel::GetFocusedSequence() const
+{
+	TSharedPtr<FSequencer> Sequencer  = WeakSequencer.Pin();
+	return Sequencer.IsValid() ? Sequencer->GetFocusedMovieSceneSequence() : nullptr;
+}
+
+bool SSequencerTimePanel::IsRecommendedResolution(FFrameRate InFrameRate) const
+{
+	UMovieSceneSequence* Sequence = GetFocusedSequence();
+	return !Sequence || (InFrameRate.IsFactorOf(Sequence->GetMovieScene()->GetPlaybackFrameRate()) && InFrameRate.IsFactorOf(Sequence->GetMovieScene()->GetFrameResolution()));
+}
+
+FFrameRate SSequencerTimePanel::OnGetFrameResolution() const
+{
+	return CurrentFrameResolution;
+}
+
+void SSequencerTimePanel::OnSetFrameResolution(FFrameRate InFrameResolution)
+{
+	CurrentFrameResolution = InFrameResolution;
+	UpdateCommonFrameRates();
+}
+
+void SSequencerTimePanel::MigrateFrameTimes(FFrameRate SourceRate, FFrameRate DestinationRate, UMovieScene* MovieScene)
+{
+	int32 TotalNumTracks = MovieScene->GetMasterTracks().Num() + (MovieScene->GetCameraCutTrack() ? 1 : 0);
+	for (const FMovieSceneBinding& Binding : MovieScene->GetBindings())
+	{
+		TotalNumTracks += Binding.GetTracks().Num();
+	}
+
+	FScopedSlowTask SlowTask(TotalNumTracks, LOCTEXT("ChangingFrameResolution", "Migrating sequence frame timing"));
+	SlowTask.MakeDialogDelayed(0.25f, true);
+
+	MovieScene->Modify();
+
+	MovieScene->SetPlaybackRange(MigrateFrameRange(MovieScene->GetPlaybackRange(), SourceRate, DestinationRate));
+	MovieScene->SetSelectionRange(MigrateFrameRange(MovieScene->GetSelectionRange(), SourceRate, DestinationRate));
+
+	for (UMovieSceneTrack* Track : MovieScene->GetMasterTracks())
+	{
+		SlowTask.EnterProgressFrame();
+		MigrateFrameTimes(SourceRate, DestinationRate, Track);
+	}
+
+	if (UMovieSceneTrack* Track = MovieScene->GetCameraCutTrack())
+	{
+		SlowTask.EnterProgressFrame();
+		MigrateFrameTimes(SourceRate, DestinationRate, Track);
+	}
+
+	for (const FMovieSceneBinding& Binding : MovieScene->GetBindings())
+	{
+		for (UMovieSceneTrack* Track : Binding.GetTracks())
+		{
+			SlowTask.EnterProgressFrame();
+			MigrateFrameTimes(SourceRate, DestinationRate, Track);
+		}
+	}
+
+	MovieScene->SetFrameResolutionDirectly(DestinationRate);
+}
+
+void SSequencerTimePanel::MigrateFrameTimes(FFrameRate SourceRate, FFrameRate DestinationRate, UMovieSceneTrack* Track)
+{
+	FScopedSlowTask SlowTask(Track->GetAllSections().Num());
+
+	for (UMovieSceneSection* Section : Track->GetAllSections())
+	{
+		SlowTask.EnterProgressFrame();
+		MigrateFrameTimes(SourceRate, DestinationRate, Section);
+	}
+}
+
+void SSequencerTimePanel::MigrateFrameTimes(FFrameRate SourceRate, FFrameRate DestinationRate, UMovieSceneSection* Section)
+{
+	Section->Modify();
+
+	TRangeBound<FFrameNumber> NewLowerBound, NewUpperBound;
+
+	if (Section->HasStartFrame())
+	{
+		FFrameNumber NewLowerBoundFrame = ConvertFrameTime(Section->GetInclusiveStartFrame(), SourceRate, DestinationRate).FloorToFrame();
+		NewLowerBound = TRangeBound<FFrameNumber>::Inclusive(NewLowerBoundFrame);
+	}
+
+	if (Section->HasEndFrame())
+	{
+		FFrameNumber NewUpperBoundValue = ConvertFrameTime(Section->GetExclusiveEndFrame(), SourceRate, DestinationRate).FloorToFrame();
+		NewUpperBound = TRangeBound<FFrameNumber>::Exclusive(NewUpperBoundValue);
+	}
+
+	Section->SetRange(TRange<FFrameNumber>(NewLowerBound, NewUpperBound));
+
+	if (Section->GetPreRollFrames() > 0)
+	{
+		FFrameNumber NewPreRollFrameCount = ConvertFrameTime(FFrameTime(Section->GetPreRollFrames()), SourceRate, DestinationRate).FloorToFrame();
+		Section->SetPreRollFrames(NewPreRollFrameCount.Value);
+	}
+
+	if (Section->GetPostRollFrames() > 0)
+	{
+		FFrameNumber NewPostRollFrameCount = ConvertFrameTime(FFrameTime(Section->GetPostRollFrames()), SourceRate, DestinationRate).FloorToFrame();
+		Section->SetPostRollFrames(NewPostRollFrameCount.Value);
+	}
+
+	if (Section->IsA(UMovieSceneSubSection::StaticClass()))
+	{
+		UMovieSceneSubSection* SubSection = Cast<UMovieSceneSubSection>(Section);
+		if (SubSection->Parameters.GetStartFrameOffset() > 0)
+		{
+			FFrameNumber NewStartFrameOffset = ConvertFrameTime(FFrameTime(SubSection->Parameters.GetStartFrameOffset()), SourceRate, DestinationRate).FloorToFrame();
+			SubSection->Parameters.SetStartFrameOffset(NewStartFrameOffset.Value);
+		}
+	}
+
+	Section->Easing.AutoEaseInDuration    = ConvertFrameTime(Section->Easing.AutoEaseInDuration,    SourceRate, DestinationRate).FloorToFrame().Value;
+	Section->Easing.AutoEaseOutDuration   = ConvertFrameTime(Section->Easing.AutoEaseOutDuration,   SourceRate, DestinationRate).FloorToFrame().Value;
+	Section->Easing.ManualEaseInDuration  = ConvertFrameTime(Section->Easing.ManualEaseInDuration,  SourceRate, DestinationRate).FloorToFrame().Value;
+	Section->Easing.ManualEaseOutDuration = ConvertFrameTime(Section->Easing.ManualEaseOutDuration, SourceRate, DestinationRate).FloorToFrame().Value;
+
+	for (const FMovieSceneChannelEntry& Entry : Section->GetChannelProxy().GetAllEntries())
+	{
+		Entry.GetBatchChannelInterface().ChangeFrameResolution_Batch(Entry.GetChannels(), SourceRate, DestinationRate);
+	}
+}
+
+#undef LOCTEXT_NAMESPACE
\ No newline at end of file
diff --git a/Engine/Source/Editor/Sequencer/Private/SSequencerTrackArea.cpp b/Engine/Source/Editor/Sequencer/Private/SSequencerTrackArea.cpp
--- a/Engine/Source/Editor/Sequencer/Private/SSequencerTrackArea.cpp
+++ b/Engine/Source/Editor/Sequencer/Private/SSequencerTrackArea.cpp
@@ -15,6 +15,7 @@
 #include "Tools/SequencerEditTool_Selection.h"
 #include "ISequencerTrackEditor.h"
 #include "DisplayNodes/SequencerTrackNode.h"
+#include "DisplayNodes/SequencerObjectBindingNode.h"
 
 FTrackAreaSlot::FTrackAreaSlot(const TSharedPtr<SSequencerTrackLane>& InSlotContent)
 {
@@ -441,11 +442,11 @@ void SSequencerTrackArea::Tick( const FGeometry& AllottedGeometry, const double
 	{
 		// Zoom by the difference in horizontal size
 		const float Difference = Size.X - SizeLastFrame->X;
-		TRange<float> OldRange = TimeSliderController->GetViewRange().GetAnimationTarget();
+		TRange<double> OldRange = TimeSliderController->GetViewRange().GetAnimationTarget();
 
 		TimeSliderController->SetViewRange(
 			OldRange.GetLowerBoundValue(),
-			OldRange.GetUpperBoundValue() + (Difference * OldRange.Size<float>() / SizeLastFrame->X),
+			OldRange.GetUpperBoundValue() + (Difference * OldRange.Size<double>() / SizeLastFrame->X),
 			EViewRangeInterpolation::Immediate
 		);
 	}
@@ -486,16 +487,23 @@ FReply SSequencerTrackArea::OnDragOver(const FGeometry& MyGeometry, const FDragD
 
 	if (DroppedNode.IsValid() && DroppedNode.Pin()->GetType() == ESequencerNode::Track && Sequencer.IsValid())
 	{
-		TSharedPtr<FSequencerDisplayNode> DisplayNode = DroppedNode.Pin();
+		TSharedPtr<FSequencerTrackNode> TrackNode = StaticCastSharedPtr<FSequencerTrackNode>(DroppedNode.Pin());
+		UMovieSceneTrack* Track = TrackNode->GetTrack();
+		int32 RowIndex = TrackNode->GetSubTrackMode() == FSequencerTrackNode::ESubTrackMode::SubTrack ? TrackNode->GetRowIndex() : 0;
 
-		UMovieSceneTrack* Track = (StaticCastSharedPtr<FSequencerTrackNode>(DroppedNode.Pin()))->GetTrack();
+		FGuid ObjectBinding;
+		TSharedPtr<FSequencerObjectBindingNode> ObjectBindingNode = TrackNode->FindParentObjectBindingNode();
+		if (ObjectBindingNode.IsValid())
+		{
+			ObjectBinding = ObjectBindingNode->GetObjectBinding();
+		}
 
 		// give track editors a chance to accept the drag event
 		auto TrackEditors = Sequencer.Pin()->GetTrackEditors();
 
 		for (const auto& TrackEditor : TrackEditors)
 		{
-			if (TrackEditor->OnAllowDrop(DragDropEvent, Track))
+			if (TrackEditor->OnAllowDrop(DragDropEvent, Track, RowIndex, ObjectBinding))
 			{
 				bAllowDrop = true;
 				return FReply::Handled();
@@ -515,18 +523,27 @@ FReply SSequencerTrackArea::OnDrop(const FGeometry& MyGeometry, const FDragDropE
 
 	if (DroppedNode.IsValid() && DroppedNode.Pin()->GetType() == ESequencerNode::Track && Sequencer.IsValid())
 	{
-		UMovieSceneTrack* Track = (StaticCastSharedPtr<FSequencerTrackNode>(DroppedNode.Pin()))->GetTrack();
+		TSharedPtr<FSequencerTrackNode> TrackNode = StaticCastSharedPtr<FSequencerTrackNode>(DroppedNode.Pin());
+		UMovieSceneTrack* Track = TrackNode->GetTrack();
+		int32 RowIndex = TrackNode->GetSubTrackMode() == FSequencerTrackNode::ESubTrackMode::SubTrack ? TrackNode->GetRowIndex() : 0;
+
+		FGuid ObjectBinding;
+		TSharedPtr<FSequencerObjectBindingNode> ObjectBindingNode = TrackNode->FindParentObjectBindingNode();
+		if (ObjectBindingNode.IsValid())
+		{
+			ObjectBinding = ObjectBindingNode->GetObjectBinding();
+		}
 
 		// give track editors a chance to process the drag event
 		auto TrackEditors = Sequencer.Pin()->GetTrackEditors();
 
 		for (const auto& TrackEditor : TrackEditors)
 		{
-			if (TrackEditor->OnAllowDrop(DragDropEvent, Track))
+			if (TrackEditor->OnAllowDrop(DragDropEvent, Track, RowIndex, ObjectBinding))
 			{
 				DroppedNode.Reset();
 
-				return TrackEditor->OnDrop(DragDropEvent, Track);
+				return TrackEditor->OnDrop(DragDropEvent, Track, RowIndex, ObjectBinding);
 			}
 		}
 	}
diff --git a/Engine/Source/Editor/Sequencer/Private/SSequencerTransformBox.cpp b/Engine/Source/Editor/Sequencer/Private/SSequencerTransformBox.cpp
--- a/Engine/Source/Editor/Sequencer/Private/SSequencerTransformBox.cpp
+++ b/Engine/Source/Editor/Sequencer/Private/SSequencerTransformBox.cpp
@@ -10,23 +10,15 @@
 
 #define LOCTEXT_NAMESPACE "Sequencer"
 
-void SSequencerTransformBox::Construct(const FArguments& InArgs, const TSharedRef<FSequencer>& InSequencer, USequencerSettings& InSettings, const TSharedRef<INumericTypeInterface<float>>& InNumericTypeInterface)
+void SSequencerTransformBox::Construct(const FArguments& InArgs, const TSharedRef<FSequencer>& InSequencer, USequencerSettings& InSettings, const TSharedRef<INumericTypeInterface<double>>& InNumericTypeInterface)
 {
 	SequencerPtr = InSequencer;
 	Settings = &InSettings;
 	NumericTypeInterface = InNumericTypeInterface;
 	
-	DeltaTime = 1.f;
+	DeltaTime = FFrameNumber(0);
 	ScaleFactor = 1.f;
 
-	// Initialize to 10 frames if showing frame numbers
-	float TimeSnapInterval = SequencerPtr.Pin()->GetFixedFrameInterval();
-	if (SequencerSnapValues::IsTimeSnapIntervalFrameRate(TimeSnapInterval))
-	{	
-		float FrameRate = 1.0f / TimeSnapInterval;
-		DeltaTime = SequencerHelpers::FrameToTime(10, FrameRate);
-	}
-
 	const FDockTabStyle* GenericTabStyle = &FCoreStyle::Get().GetWidgetStyle<FDockTabStyle>("Docking.Tab");
 	const FButtonStyle* const CloseButtonStyle = &GenericTabStyle->CloseButtonStyle;
 
@@ -61,12 +53,13 @@ void SSequencerTransformBox::Construct(const FArguments& InArgs, const TSharedRe
 					.Padding(6.0f, 0.0f, 0.0f, 0.0f)
 					.AutoWidth()
 					[
-						SAssignNew(EntryBox, SNumericEntryBox<float>)
-							.MinDesiredValueWidth(32.0f)
-							.TypeInterface(NumericTypeInterface)
-							.ToolTipText(LOCTEXT("Delta_Tooltip", "The amount to offset the keys/sections by"))
-							.OnValueCommitted(this, &SSequencerTransformBox::OnDeltaChanged)
-							.Value_Lambda([this](){ return DeltaTime; })
+						SAssignNew(OffsetEntryBox, SNumericEntryBox<double>)
+						.MinDesiredValueWidth(32.0f)
+						.TypeInterface(NumericTypeInterface)
+						.ToolTipText(LOCTEXT("TransformDelta_Tooltip", "The amount to offset the selected keys/sections by"))
+						.OnValueCommitted(this, &SSequencerTransformBox::OnDeltaCommitted)
+						.OnValueChanged(this, &SSequencerTransformBox::OnDeltaChanged)
+						.Value_Lambda([this](){ return DeltaTime.Value; })
 					]
 
 				+ SHorizontalBox::Slot()
@@ -92,10 +85,11 @@ void SSequencerTransformBox::Construct(const FArguments& InArgs, const TSharedRe
 					.Padding(6.0f, 0.0f, 0.0f, 0.0f)
 					.AutoWidth()
 					[
-						SAssignNew(EntryBox, SNumericEntryBox<float>)
-							.MinDesiredValueWidth(32.0f)
-							.ToolTipText(LOCTEXT("Scale_Tooltip", "The amount to scale the keys/section by (about the local time)"))
-							.OnValueCommitted(this, &SSequencerTransformBox::OnScaleChanged)
+						 SAssignNew(ScaleEntryBox, SNumericEntryBox<float>)
+						 	.MinDesiredValueWidth(32.0f)
+						 	.ToolTipText(LOCTEXT("TransformScale_Tooltip", "The amount to scale the selected keys/section by (about the local time)"))
+						 	.OnValueCommitted(this, &SSequencerTransformBox::OnScaleCommitted)
+							.OnValueChanged(this, &SSequencerTransformBox::OnScaleChanged)
 							.Value_Lambda([this](){ return ScaleFactor; })
 					]
 
@@ -136,33 +130,33 @@ void SSequencerTransformBox::ToggleVisibility()
 	{
 		Border->SetVisibility(EVisibility::Visible);
 		LastFocusedWidget = SlateApplication.GetUserFocusedWidget(0);
-		SlateApplication.SetAllUserFocus(EntryBox, EFocusCause::Navigation);
+		SlateApplication.SetAllUserFocus(OffsetEntryBox, EFocusCause::Navigation);
 	}
 }
 
-void SSequencerTransformBox::OnDeltaChanged(float Value, ETextCommit::Type CommitType)
+void SSequencerTransformBox::OnDeltaCommitted(double Value, ETextCommit::Type CommitType)
 {
-	if (CommitType != ETextCommit::OnEnter)
-	{
-		return;
-	}
+	OnDeltaChanged(Value);
+}
 
-	DeltaTime = Value;
+void SSequencerTransformBox::OnDeltaChanged(double Value)
+{
+	DeltaTime = FFrameTime::FromDecimal(Value).GetFrame();
 }
 
-void SSequencerTransformBox::OnScaleChanged(float Value, ETextCommit::Type CommitType)
+void SSequencerTransformBox::OnScaleCommitted(float Value, ETextCommit::Type CommitType)
 {
-	if (CommitType != ETextCommit::OnEnter)
-	{
-		return;
-	}
+	OnScaleChanged(Value);
+}
 
+void SSequencerTransformBox::OnScaleChanged(float Value)
+{
 	ScaleFactor = Value;
 }
 
 FReply SSequencerTransformBox::OnPlusButtonClicked()
 {
-	if (DeltaTime != 0.f)
+	if (DeltaTime != 0)
 	{
 		TSharedPtr<FSequencer> Sequencer = SequencerPtr.Pin();
 
@@ -174,7 +168,7 @@ FReply SSequencerTransformBox::OnPlusButtonClicked()
 
 FReply SSequencerTransformBox::OnMinusButtonClicked()
 {
-	if (DeltaTime != 0.f)
+	if (DeltaTime != 0)
 	{
 		TSharedPtr<FSequencer> Sequencer = SequencerPtr.Pin();
 	
diff --git a/Engine/Source/Editor/Sequencer/Private/SSequencerTransformBox.h b/Engine/Source/Editor/Sequencer/Private/SSequencerTransformBox.h
--- a/Engine/Source/Editor/Sequencer/Private/SSequencerTransformBox.h
+++ b/Engine/Source/Editor/Sequencer/Private/SSequencerTransformBox.h
@@ -8,6 +8,7 @@
 #include "Input/Reply.h"
 #include "Widgets/SCompoundWidget.h"
 #include "Widgets/Input/SNumericEntryBox.h"
+#include "QualifiedFrameTime.h"
 
 class FSequencer;
 class USequencerSettings;
@@ -21,7 +22,7 @@ public:
 	SLATE_END_ARGS()
 
 	/** Construct the widget. */
-	void Construct(const FArguments& InArgs, const TSharedRef<FSequencer>& InSequencer, USequencerSettings& InSettings, const TSharedRef<INumericTypeInterface<float>>& InNumericTypeInterface);
+	void Construct(const FArguments& InArgs, const TSharedRef<FSequencer>& InSequencer, USequencerSettings& InSettings, const TSharedRef<INumericTypeInterface<double>>& InNumericTypeInterface);
 
 	/** Toggle the widget's visibility. */
 	void ToggleVisibility();
@@ -33,8 +34,10 @@ private:
 	FReply OnMinusButtonClicked();
 	FReply OnMultiplyButtonClicked();
 	FReply OnDivideButtonClicked();
-	void OnDeltaChanged(float Value, ETextCommit::Type CommitType);
-	void OnScaleChanged(float Value, ETextCommit::Type CommitType);
+	void OnDeltaCommitted(double Value, ETextCommit::Type CommitType);
+	void OnDeltaChanged(double Value);
+	void OnScaleCommitted(float Value, ETextCommit::Type CommitType);
+	void OnScaleChanged(float Value);
 
 	/** Callback for when the close button is clicked. */
 	FReply OnCloseButtonClicked();
@@ -45,13 +48,16 @@ private:
 	TSharedPtr<SWidget> Border;
 
 	/** The entry box widget. */
-	TSharedPtr<SNumericEntryBox<float>> EntryBox;
+	TSharedPtr<SNumericEntryBox<double>> OffsetEntryBox;
+
+	/** The scale entry box widget. */
+	TSharedPtr<SNumericEntryBox<float>> ScaleEntryBox;
 
 	/** The widget that focused prior to this transform box. */
 	TWeakPtr<SWidget> LastFocusedWidget;
 
 	/** Numeric type interface used for converting parsing and generating strings from numbers. */
-	TSharedPtr<INumericTypeInterface<float>> NumericTypeInterface;
+	TSharedPtr<INumericTypeInterface<double>> NumericTypeInterface;
 
 	/** The main sequencer interface. */
 	TWeakPtr<FSequencer> SequencerPtr;
@@ -60,7 +66,7 @@ private:
 	USequencerSettings* Settings;
 
 	/** Cached delta time. */
-	float DeltaTime;
+	FFrameNumber DeltaTime;
 
 	/** Cached scale factor. */
 	float ScaleFactor;
diff --git a/Engine/Source/Editor/Sequencer/Private/SSequencerTreeView.cpp b/Engine/Source/Editor/Sequencer/Private/SSequencerTreeView.cpp
--- a/Engine/Source/Editor/Sequencer/Private/SSequencerTreeView.cpp
+++ b/Engine/Source/Editor/Sequencer/Private/SSequencerTreeView.cpp
@@ -54,12 +54,14 @@ void SSequencerTreeViewRow::Construct(const FArguments& InArgs, const TSharedRef
 {
 	Node = InNode;
 	OnGenerateWidgetForColumn = InArgs._OnGenerateWidgetForColumn;
+	bool bIsSelectable = InNode->IsSelectable();
 
 	SMultiColumnTableRow::Construct(
 		SMultiColumnTableRow::FArguments()
 			.OnDragDetected(this, &SSequencerTreeViewRow::OnDragDetected)
 			.OnCanAcceptDrop(this, &SSequencerTreeViewRow::OnCanAcceptDrop)
-			.OnAcceptDrop(this, &SSequencerTreeViewRow::OnAcceptDrop),
+			.OnAcceptDrop(this, &SSequencerTreeViewRow::OnAcceptDrop)
+			.ShowSelection(bIsSelectable),
 		OwnerTableView);
 }
 
@@ -91,6 +93,12 @@ FReply SSequencerTreeViewRow::OnDragDetected( const FGeometry& InGeometry, const
 				}
 			}
 
+			// If there were no nodes selected, don't start a drag drop operation.
+			if (DraggableNodes.Num() == 0)
+			{
+				return FReply::Unhandled();
+			}
+
 			FText DefaultText = FText::Format( NSLOCTEXT( "SequencerTreeViewRow", "DefaultDragDropFormat", "Move {0} item(s)" ), FText::AsNumber( DraggableNodes.Num() ) );
 			TSharedRef<FSequencerDisplayNodeDragDropOp> DragDropOp = FSequencerDisplayNodeDragDropOp::New( DraggableNodes, DefaultText, nullptr );
 
@@ -467,7 +475,10 @@ void SSequencerTreeView::SynchronizeTreeSelectionWithSequencerSelection()
 			FSequencer& Sequencer = SequencerNodeTree->GetSequencer();
 			for ( auto& Node : Sequencer.GetSelection().GetSelectedOutlinerNodes() )
 			{
-				Private_SetItemSelection( Node, true, false );
+				if (Node->IsSelectable())
+				{
+					Private_SetItemSelection( Node, true, false );
+				}
 			}
 
 			Private_SignalSelectionChanged( ESelectInfo::Direct );
@@ -569,9 +580,10 @@ bool SSequencerTreeView::SynchronizeSequencerSelectionWithTreeSelection()
 TSharedPtr<SWidget> SSequencerTreeView::OnContextMenuOpening()
 {
 	const TSet<TSharedRef<FSequencerDisplayNode>> SelectedNodes = SequencerNodeTree->GetSequencer().GetSelection().GetSelectedOutlinerNodes();
-	if ( SelectedNodes.Num() > 0 )
+	auto SelectedNodesArray = SelectedNodes.Array();
+	if (SelectedNodes.Num() > 0 && SelectedNodesArray[0]->IsSelectable())
 	{
-		return SelectedNodes.Array()[0]->OnSummonContextMenu();
+		return SelectedNodesArray[0]->OnSummonContextMenu();
 	}
 
 	// Otherwise, add a general menu for options
@@ -742,7 +754,7 @@ TSharedRef<ITableRow> SSequencerTreeView::OnGenerateRow(FDisplayNodeRef InDispla
 		if (!TrackLane.IsValid())
 		{
 			// Add a track slot for the row
-			TAttribute<TRange<float>> ViewRange = FAnimatedRange::WrapAttribute( TAttribute<FAnimatedRange>::Create(TAttribute<FAnimatedRange>::FGetter::CreateSP(&SequencerNodeTree->GetSequencer(), &FSequencer::GetViewRange)) );
+			TAttribute<TRange<double>> ViewRange = FAnimatedRange::WrapAttribute( TAttribute<FAnimatedRange>::Create(TAttribute<FAnimatedRange>::FGetter::CreateSP(&SequencerNodeTree->GetSequencer(), &FSequencer::GetViewRange)) );
 
 			TrackLane = SNew(SSequencerTrackLane, SectionAuthority.ToSharedRef(), SharedThis(this))
 			.IsEnabled(!InDisplayNode->GetSequencer().IsReadOnly())
diff --git a/Engine/Source/Editor/Sequencer/Private/SectionLayout.cpp b/Engine/Source/Editor/Sequencer/Private/SectionLayout.cpp
--- a/Engine/Source/Editor/Sequencer/Private/SectionLayout.cpp
+++ b/Engine/Source/Editor/Sequencer/Private/SectionLayout.cpp
@@ -1,13 +1,23 @@
 // Copyright 1998-2018 Epic Games, Inc. All Rights Reserved.
 
 #include "SectionLayout.h"
-#include "GroupedKeyArea.h"
 
-FSectionLayoutElement FSectionLayoutElement::FromGroup(const TSharedRef<FSequencerDisplayNode>& InNode, const TSharedRef<FGroupedKeyArea>& InKeyAreaGroup, float InOffset)
+FSectionLayoutElement FSectionLayoutElement::FromGroup(const TSharedRef<FSequencerDisplayNode>& InNode, UMovieSceneSection* InSection, float InOffset)
 {
+	TArray< TSharedRef<FSequencerSectionKeyAreaNode> > ChildKeyAreaNodes;
+	InNode->GetChildKeyAreaNodesRecursively(ChildKeyAreaNodes);
+
 	FSectionLayoutElement Tmp;
 	Tmp.Type = Group;
-	Tmp.KeyArea = InKeyAreaGroup;
+
+	for (TSharedRef<FSequencerSectionKeyAreaNode> KeyAreaNode : ChildKeyAreaNodes)
+	{
+		TSharedPtr<IKeyArea> KeyArea = KeyAreaNode->GetKeyArea(InSection);
+		if (KeyArea.IsValid())
+		{
+			Tmp.KeyAreas.Add(KeyArea.ToSharedRef());
+		}
+	}
 	Tmp.LocalOffset = InOffset;
 	Tmp.Height = InNode->GetNodeHeight();
 	Tmp.DisplayNode = InNode;
@@ -18,7 +28,11 @@ FSectionLayoutElement FSectionLayoutElement::FromKeyAreaNode(const TSharedRef<FS
 {
 	FSectionLayoutElement Tmp;
 	Tmp.Type = Single;
-	Tmp.KeyArea = InKeyAreaNode->GetKeyArea(InSection);
+	TSharedPtr<IKeyArea> KeyArea = InKeyAreaNode->GetKeyArea(InSection);
+	if (KeyArea.IsValid())
+	{
+		Tmp.KeyAreas.Add(KeyArea.ToSharedRef());
+	}
 	Tmp.LocalOffset = InOffset;
 	Tmp.DisplayNode = InKeyAreaNode;
 	Tmp.Height = InKeyAreaNode->GetNodeHeight();
@@ -29,7 +43,11 @@ FSectionLayoutElement FSectionLayoutElement::FromTrack(const TSharedRef<FSequenc
 {
 	FSectionLayoutElement Tmp;
 	Tmp.Type = Single;
-	Tmp.KeyArea = InTrackNode->GetTopLevelKeyNode()->GetKeyArea(InSection);
+	TSharedPtr<IKeyArea> KeyArea = InTrackNode->GetTopLevelKeyNode()->GetKeyArea(InSection);
+	if (KeyArea.IsValid())
+	{
+		Tmp.KeyAreas.Add(KeyArea.ToSharedRef());
+	}
 	Tmp.LocalOffset = InOffset;
 	Tmp.DisplayNode = InTrackNode;
 	Tmp.Height = InTrackNode->GetNodeHeight();
@@ -61,9 +79,9 @@ float FSectionLayoutElement::GetHeight() const
 	return Height;
 }
 
-TSharedPtr<IKeyArea> FSectionLayoutElement::GetKeyArea() const
+TArrayView<const TSharedRef<IKeyArea>> FSectionLayoutElement::GetKeyAreas() const
 {
-	return KeyArea;
+	return KeyAreas;
 }
 
 TSharedPtr<FSequencerDisplayNode> FSectionLayoutElement::GetDisplayNode() const
@@ -97,7 +115,7 @@ FSectionLayout::FSectionLayout(FSequencerTrackNode& TrackNode, int32 InSectionIn
 		{
 			Elements.Add(FSectionLayoutElement::FromGroup(
 				Node.AsShared(),
-				Node.GetKeyGrouping(Section),
+				Section,
 				Offset
 			));
 		}
diff --git a/Engine/Source/Editor/Sequencer/Private/SectionLayout.h b/Engine/Source/Editor/Sequencer/Private/SectionLayout.h
--- a/Engine/Source/Editor/Sequencer/Private/SectionLayout.h
+++ b/Engine/Source/Editor/Sequencer/Private/SectionLayout.h
@@ -7,8 +7,6 @@
 #include "DisplayNodes/SequencerSectionKeyAreaNode.h"
 #include "DisplayNodes/SequencerTrackNode.h"
 
-class FGroupedKeyArea;
-
 /** A layout element specifying the geometry required to render a key area */
 struct FSectionLayoutElement
 {
@@ -16,7 +14,7 @@ struct FSectionLayoutElement
 	enum EType { Single, Group };
 
 	/** Construct this element from a grouped key area */
-	static FSectionLayoutElement FromGroup(const TSharedRef<FSequencerDisplayNode>& InNode, const TSharedRef<FGroupedKeyArea>& InKeyAreaGroup, float InOffset);
+	static FSectionLayoutElement FromGroup(const TSharedRef<FSequencerDisplayNode>& InNode, UMovieSceneSection* InSection, float InOffset);
 
 	/** Construct this element from a single Key area node */
 	static FSectionLayoutElement FromKeyAreaNode(const TSharedRef<FSequencerSectionKeyAreaNode>& InKeyAreaNode, UMovieSceneSection* InSection, float InOffset);
@@ -36,8 +34,8 @@ struct FSectionLayoutElement
 	/** Retrieve the desired height of this element based on the specified parent geometry */
 	float GetHeight() const;
 
-	/** Access the key area that this layout element was generated for */
-	TSharedPtr<IKeyArea> GetKeyArea() const;
+	/** Access all the key areas that this layout element represents */
+	TArrayView<const TSharedRef<IKeyArea>> GetKeyAreas() const;
 
 	/** Access the display node that this layout element was generated for */
 	TSharedPtr<FSequencerDisplayNode> GetDisplayNode() const;
@@ -45,7 +43,7 @@ struct FSectionLayoutElement
 private:
 
 	/** Pointer to the key area that we were generated from */
-	TSharedPtr<IKeyArea> KeyArea;
+	TArray<TSharedRef<IKeyArea>, TInlineAllocator<1>> KeyAreas;
 
 	/** The specific node that this key area relates to */
 	TSharedPtr<FSequencerDisplayNode> DisplayNode;
