Copying //UE4/Dev-Sequencer to //UE4/Dev-Main (Source: //UE4/Dev-Sequencer @ 3946692)

diff --git a/Engine/Source/Editor/Sequencer/Private/Sequencer.cpp b/Engine/Source/Editor/Sequencer/Private/Sequencer.cpp
--- a/Engine/Source/Editor/Sequencer/Private/Sequencer.cpp
+++ b/Engine/Source/Editor/Sequencer/Private/Sequencer.cpp
@@ -18,6 +18,7 @@
 #include "Engine/Engine.h"
 #include "Settings/LevelEditorViewportSettings.h"
 #include "Editor.h"
+#include "Channels/MovieSceneChannelProxy.h"
 #include "MovieScenePossessable.h"
 #include "MovieScene.h"
 #include "Widgets/Layout/SBorder.h"
@@ -47,8 +48,8 @@
 #include "MovieSceneClipboard.h"
 #include "SequencerCommonHelpers.h"
 #include "SSequencer.h"
-#include "ISequencerKeyCollection.h"
-#include "GroupedKeyArea.h"
+#include "SSequencerSection.h"
+#include "SequencerKeyCollection.h"
 #include "SequencerSettings.h"
 #include "SequencerLog.h"
 #include "SequencerEdMode.h"
@@ -111,6 +112,14 @@
 #include "Compilation/MovieSceneCompiler.h"
 #include "SequencerKeyActor.h"
 #include "MovieSceneCopyableBinding.h"
+#include "ISequencerChannelInterface.h"
+#include "SequencerKeyCollection.h"
+#include "CurveEditor.h"
+#include "CurveEditorScreenSpace.h"
+#include "CurveDataAbstraction.h"
+#include "Fonts/FontMeasure.h"
+#include "MovieSceneTimeHelpers.h"
+#include "FrameNumberNumericInterface.h"
 
 #define LOCTEXT_NAMESPACE "Sequencer"
 
@@ -152,6 +161,90 @@ struct FSequencerTemplateStore : IMovieSceneSequenceTemplateStore
 };
 
 
+struct FSequencerCurveEditorBounds : ICurveEditorBounds
+{
+	FSequencerCurveEditorBounds(TSharedRef<FSequencer> InSequencer)
+		: OutputMin(0), OutputMax(1), WeakSequencer(InSequencer)
+	{}
+
+	virtual void GetInputBounds(double& OutMin, double& OutMax) const override
+	{
+		TSharedPtr<FSequencer> Sequencer = WeakSequencer.Pin();
+		if (Sequencer.IsValid())
+		{
+			TRange<double> Bounds = Sequencer->GetViewRange();
+			OutMin = Bounds.GetLowerBoundValue();
+			OutMax = Bounds.GetUpperBoundValue();
+		}
+	}
+	virtual void GetOutputBounds(double& OutMin, double& OutMax) const override
+	{
+		OutMin = OutputMin;
+		OutMax = OutputMax;
+	}
+
+	virtual void SetInputBounds(double InMin, double InMax) override
+	{
+		TSharedPtr<FSequencer> Sequencer = WeakSequencer.Pin();
+		if (Sequencer.IsValid())
+		{
+			FFrameRate FrameResolution = Sequencer->GetFocusedFrameResolution();
+
+			if (InMin * FrameResolution > TNumericLimits<int32>::Lowest() && InMax * FrameResolution < TNumericLimits<int32>::Max())
+			{
+				Sequencer->SetViewRange(TRange<double>(InMin, InMax), EViewRangeInterpolation::Immediate);
+			}
+		}
+	}
+	virtual void SetOutputBounds(double InMin, double InMax) override
+	{
+		if (InMax - InMin < TNumericLimits<float>::Max())
+		{
+			OutputMin = InMin;
+			OutputMax = InMax;
+		}
+	}
+
+	double OutputMin, OutputMax;
+	TWeakPtr<FSequencer> WeakSequencer;
+};
+
+class FSequencerCurveEditor : public FCurveEditor
+{
+public:
+	TWeakPtr<FSequencer> WeakSequencer;
+
+	FSequencerCurveEditor(TWeakPtr<FSequencer> InSequencer)
+		: WeakSequencer(InSequencer)
+	{}
+
+	virtual void GetGridLinesX(TArray<float>& MajorGridLines, TArray<float>& MinorGridLines) const
+	{
+		TSharedPtr<FSequencer> Sequencer = WeakSequencer.Pin();
+
+		FCurveEditorScreenSpace ScreenSpace = GetScreenSpace();
+
+		double MajorGridStep  = 0.0;
+		int32  MinorDivisions = 0;
+
+		if (Sequencer.IsValid() && Sequencer->GetGridMetrics(ScreenSpace.GetPhysicalWidth(), MajorGridStep, MinorDivisions))
+		{
+			const double FirstMajorLine = FMath::FloorToDouble(ScreenSpace.GetInputMin() / MajorGridStep) * MajorGridStep;
+			const double LastMajorLine  = FMath::CeilToDouble(ScreenSpace.GetInputMax() / MajorGridStep) * MajorGridStep;
+
+			for (double CurrentMajorLine = FirstMajorLine; CurrentMajorLine < LastMajorLine; CurrentMajorLine += MajorGridStep)
+			{
+				MajorGridLines.Add( ScreenSpace.SecondsToScreen(CurrentMajorLine) );
+
+				for (int32 Step = 1; Step < MinorDivisions; ++Step)
+				{
+					MinorGridLines.Add( ScreenSpace.SecondsToScreen(CurrentMajorLine + Step*MajorGridStep/MinorDivisions) );
+				}
+			}
+		}
+	}
+};
+
 void FSequencer::InitSequencer(const FSequencerInitParams& InitParams, const TSharedRef<ISequencerObjectChangeListener>& InObjectChangeListener, const TArray<FOnCreateTrackEditor>& TrackEditorDelegates, const TArray<FOnCreateEditorObjectBinding>& EditorObjectBindingDelegates)
 {
 	bIsEditingWithinLevelEditor = InitParams.bEditWithinLevelEditor;
@@ -186,11 +279,25 @@ void FSequencer::InitSequencer(const FSequencerInitParams& InitParams, const TSh
 
 	Settings = USequencerSettingsContainer::GetOrCreate<USequencerSettings>(*InitParams.ViewParams.UniqueName);
 
-	Settings->GetOnLockPlaybackToAudioClockChanged().AddSP(this, &FSequencer::ResetTimingManager);
-	ResetTimingManager(Settings->ShouldLockPlaybackToAudioClock());
-
 	Settings->GetOnEvaluateSubSequencesInIsolationChanged().AddSP(this, &FSequencer::RestorePreAnimatedState);
 	Settings->GetOnShowSelectedNodesOnlyChanged().AddSP(this, &FSequencer::OnSelectedNodesOnlyChanged);
+	Settings->GetOnCurveEditorCurveVisibilityChanged().AddSP(this, &FSequencer::SyncCurveEditorToSelection, false);
+
+	{
+		CurveEditorModel = MakeShared<FSequencerCurveEditor>(SharedThis(this));
+		CurveEditorModel->SetBounds(MakeUnique<FSequencerCurveEditorBounds>(SharedThis(this)));
+
+		CurveEditorModel->InputSnapEnabledAttribute   = MakeAttributeLambda([this]{ return Settings->GetIsSnapEnabled(); });
+		CurveEditorModel->OnInputSnapEnabledChanged   = FOnSetBoolean::CreateLambda([this](bool NewValue){ Settings->SetIsSnapEnabled(NewValue); });
+
+		CurveEditorModel->OutputSnapEnabledAttribute  = MakeAttributeLambda([this]{ return Settings->GetSnapCurveValueToInterval(); });
+		CurveEditorModel->OnOutputSnapEnabledChanged  = FOnSetBoolean::CreateLambda([this](bool NewValue){ Settings->SetSnapCurveValueToInterval(NewValue); });
+
+		CurveEditorModel->OutputSnapIntervalAttribute = MakeAttributeLambda([this]{ return (double)Settings->GetCurveValueSnapInterval(); });
+		CurveEditorModel->InputSnapRateAttribute      = MakeAttributeSP(this, &FSequencer::GetFocusedPlayRate);
+
+		CurveEditorModel->DefaultKeyAttributes        = MakeAttributeSP(this, &FSequencer::GetDefaultKeyAttributes);
+	}
 
 	{
 		FDelegateHandle OnBlueprintCompiledHandle = GEditor->OnBlueprintCompiled().AddLambda([&]{ State.InvalidateExpiredObjects(); });
@@ -205,7 +312,6 @@ void FSequencer::InitSequencer(const FSequencerInitParams& InitParams, const TSh
 	}
 	ToolkitHost = InitParams.ToolkitHost;
 
-	ScrubPosition = InitParams.ViewParams.InitialScrubPosition;
 	PlaybackSpeed = 1.f;
 	ShuttleMultiplier = 0;
 	ObjectChangeListener = InObjectChangeListener;
@@ -214,11 +320,17 @@ void FSequencer::InitSequencer(const FSequencerInitParams& InitParams, const TSh
 	
 	RootSequence = InitParams.RootSequence;
 
+	UpdateTimeBases();
+	PlayPosition.Reset(FFrameTime(0));
+
 	TemplateStore = MakeShared<FSequencerTemplateStore>();
 
 	ActiveTemplateIDs.Add(MovieSceneSequenceID::Root);
+	ActiveTemplateStates.Add(true);
 	RootTemplateInstance.Initialize(*InitParams.RootSequence, *this, TemplateStore.ToSharedRef());
 
+	ResetTimeController();
+
 	// Make internal widgets
 	SequencerWidget = SNew( SSequencer, SharedThis( this ) )
 		.ViewRange( this, &FSequencer::GetViewRange )
@@ -235,8 +347,7 @@ void FSequencer::InitSequencer(const FSequencerInitParams& InitParams, const TSh
 		.OnSelectionRangeEndDrag( this, &FSequencer::OnSelectionRangeEndDrag )
 		.IsPlaybackRangeLocked( this, &FSequencer::IsPlaybackRangeLocked )
 		.OnTogglePlaybackRangeLocked( this, &FSequencer::TogglePlaybackRangeLocked )
-		.TimeSnapInterval( this, &FSequencer::GetFixedFrameInterval )
-		.ScrubPosition( this, &FSequencer::GetLocalTime )
+		.ScrubPosition( this, &FSequencer::GetLocalFrameTime )
 		.OnBeginScrubbing( this, &FSequencer::OnBeginScrubbing )
 		.OnEndScrubbing( this, &FSequencer::OnEndScrubbing )
 		.OnScrubPositionChanged( this, &FSequencer::OnScrubPositionChanged )
@@ -244,6 +355,7 @@ void FSequencer::InitSequencer(const FSequencerInitParams& InitParams, const TSh
 		.OnClampRangeChanged( this, &FSequencer::OnClampRangeChanged )
 		.OnGetNearestKey( this, &FSequencer::OnGetNearestKey )
 		.OnGetAddMenuContent(InitParams.ViewParams.OnGetAddMenuContent)
+		.OnBuildCustomContextMenuForGuid(InitParams.ViewParams.OnBuildCustomContextMenuForGuid)
 		.OnReceivedFocus(InitParams.ViewParams.OnReceivedFocus)
 		.AddMenuExtender(InitParams.ViewParams.AddMenuExtender)
 		.ToolbarExtender(InitParams.ViewParams.ToolbarExtender);
@@ -279,6 +391,13 @@ void FSequencer::InitSequencer(const FSequencerInitParams& InitParams, const TSh
 	NotifyMovieSceneDataChanged( EMovieSceneDataChangeType::ActiveMovieSceneChanged );
 	UpdateTimeBoundsToFocusedMovieScene();
 
+	FQualifiedFrameTime CurrentTime = GetLocalTime();
+	if (!TargetViewRange.Contains(CurrentTime.AsSeconds()))
+	{
+		SetLocalTimeDirectly(LastViewRange.GetLowerBoundValue() * CurrentTime.Rate);
+		OnGlobalTimeChangedDelegate.Broadcast();
+	}
+
 	// NOTE: Could fill in asset editor commands here!
 
 	BindCommands();
@@ -297,9 +416,8 @@ FSequencer::FSequencer()
 	, SequencerSharedBindings( new FUICommandList )
 	, TargetViewRange(0.f, 5.f)
 	, LastViewRange(0.f, 5.f)
-	, ViewRangeBeforeZoom(TRange<float>::Empty())
+	, ViewRangeBeforeZoom(TRange<double>::Empty())
 	, PlaybackState( EMovieScenePlayerStatus::Stopped )
-	, ScrubPosition( 0.0f )
 	, bPerspectiveViewportPossessionEnabled( true )
 	, bPerspectiveViewportCameraCutEnabled( false )
 	, bIsEditingWithinLevelEditor( false )
@@ -361,8 +479,11 @@ void FSequencer::Tick(float InDeltaTime)
 		ensureAlwaysMsgf(SequencerRefCount == 1, TEXT("Multiple persistent shared references detected for Sequencer. There should only be one persistent authoritative reference. Found %d additional references which will result in FSequencer not being released correctly."), SequencerRefCount - 1);
 	}
 
+	// Ensure the time bases for our playback position are kept up to date with the root sequence
+	UpdateTimeBases();
+
 	Selection.Tick();
-	
+
 	if (PlaybackContextAttribute.IsBound())
 	{
 		TWeakObjectPtr<UObject> NewPlaybackContext = PlaybackContextAttribute.Get();
@@ -391,10 +512,15 @@ void FSequencer::Tick(float InDeltaTime)
 		SetPlaybackStatus(StoredPlaybackState);
 	}
 
+
 	UObject* PlaybackContext = GetPlaybackContext();
 	UWorld* World = PlaybackContext ? PlaybackContext->GetWorld() : nullptr;
 	float Dilation = World ? World->GetWorldSettings()->MatineeTimeDilation : 1.f;
-	FTimeAndDelta TimeAndDelta = TimingManager->AdjustTime(GetGlobalTime(), InDeltaTime, PlaybackSpeed, Dilation);
+
+	TimeController->Tick(InDeltaTime, PlaybackSpeed * Dilation);
+
+	FQualifiedFrameTime GlobalTime = GetGlobalTime();
+	FFrameTime NewGlobalTime = TimeController->RequestCurrentTime(GlobalTime, PlaybackSpeed * Dilation);
 
 	static const float AutoScrollFactor = 0.1f;
 
@@ -402,30 +528,23 @@ void FSequencer::Tick(float InDeltaTime)
 	if (AutoscrollOffset.IsSet())
 	{
 		float Offset = AutoscrollOffset.GetValue() * AutoScrollFactor;
-		SetViewRange(TRange<float>(TargetViewRange.GetLowerBoundValue() + Offset, TargetViewRange.GetUpperBoundValue() + Offset), EViewRangeInterpolation::Immediate);
+		SetViewRange(TRange<double>(TargetViewRange.GetLowerBoundValue() + Offset, TargetViewRange.GetUpperBoundValue() + Offset), EViewRangeInterpolation::Immediate);
 	}
 
 	// Animate the autoscrub offset if it's set
 	if (AutoscrubOffset.IsSet())
 	{
-		float Offset = AutoscrubOffset.GetValue() * AutoScrollFactor;
-		SetLocalTimeDirectly(GetLocalTime() + Offset);
+		FQualifiedFrameTime CurrentTime = GetLocalTime();
+		FFrameTime Offset = (AutoscrubOffset.GetValue() * AutoScrollFactor) * CurrentTime.Rate;
+		SetLocalTimeDirectly(CurrentTime.Time + Offset);
 	}
 
 	// override max frame rate
 	if (PlaybackState == EMovieScenePlayerStatus::Playing)
 	{
-		bool bIsFixedFrameIntervalPlayback = false;
-		if (GetFocusedMovieSceneSequence() && GetFocusedMovieSceneSequence()->GetMovieScene())
+		if (PlayPosition.GetEvaluationType() == EMovieSceneEvaluationType::FrameLocked)
 		{
-			bIsFixedFrameIntervalPlayback = GetFocusedMovieSceneSequence()->GetMovieScene()->GetForceFixedFrameIntervalPlayback();
-		}
-
-		const float TimeSnapInterval = GetFixedFrameInterval();
-
-		if (SequencerSnapValues::IsTimeSnapIntervalFrameRate(TimeSnapInterval) && bIsFixedFrameIntervalPlayback)
-		{
-			GEngine->SetMaxFPS(1.f / TimeSnapInterval);
+			GEngine->SetMaxFPS(1.f / PlayPosition.GetInputRate().AsInterval());
 		}
 		else
 		{
@@ -438,15 +557,14 @@ void FSequencer::Tick(float InDeltaTime)
 		Settings->SetLoopMode(SLM_Loop);
 	}
 
-	if (PlaybackState == EMovieScenePlayerStatus::Playing ||
-		PlaybackState == EMovieScenePlayerStatus::Recording)
+	if (PlaybackState == EMovieScenePlayerStatus::Playing)
 	{
 		// Put the time into local space
-		SetLocalTimeLooped(TimeAndDelta.Time * RootToLocalTransform);
+		SetLocalTimeLooped(NewGlobalTime * RootToLocalTransform);
 	}
 	else
 	{
-		PlayPosition.Reset(GetGlobalTime());
+		PlayPosition.Reset(GlobalTime.ConvertTo(PlayPosition.GetInputRate()));
 	}
 
 	UpdateSubSequenceData();
@@ -454,15 +572,14 @@ void FSequencer::Tick(float InDeltaTime)
 	// Tick all the tools we own as well
 	for (int32 EditorIndex = 0; EditorIndex < TrackEditors.Num(); ++EditorIndex)
 	{
-		TrackEditors[EditorIndex]->Tick(TimeAndDelta.Delta * PlaybackSpeed);
+		TrackEditors[EditorIndex]->Tick(InDeltaTime);
 	}
 
 	if (!IsInSilentMode())
 	{
 		if (bNeedsEvaluate)
 		{
-			FMovieSceneEvaluationRange Range = PlayPosition.JumpTo(ScrubPosition, GetFocusedMovieSceneSequence()->GetMovieScene()->GetOptionalFixedFrameInterval());
-			EvaluateInternal(Range);
+			EvaluateInternal(PlayPosition.GetCurrentPositionAsRange());
 		}
 	}
 
@@ -470,9 +587,16 @@ void FSequencer::Tick(float InDeltaTime)
 	if(SequenceRecorder.IsRecording())
 	{
 		UMovieSceneSubSection* Section = UMovieSceneSubSection::GetRecordingSection();
-		if(Section != nullptr)
+		if(Section != nullptr && Section->HasStartFrame())
 		{
-			Section->SetEndTime(Section->GetStartTime() + SequenceRecorder.GetCurrentRecordingLength());
+			FFrameRate   SectionResolution = Section->GetTypedOuter<UMovieScene>()->GetFrameResolution();
+			FFrameNumber RecordingLength   = SequenceRecorder.GetCurrentRecordingLength().ConvertTo(SectionResolution).CeilToFrame();
+
+			if (RecordingLength > 0)
+			{
+				FFrameNumber EndFrame = Section->GetInclusiveStartFrame() + RecordingLength;
+				Section->SetRange(TRange<FFrameNumber>(Section->GetInclusiveStartFrame(), TRangeBound<FFrameNumber>::Exclusive(EndFrame)));
+			}
 		}
 	}
 
@@ -529,6 +653,8 @@ void FSequencer::ResetToNewRootSequence(UMovieSceneSequence& NewSequence)
 
 	ActiveTemplateIDs.Reset();
 	ActiveTemplateIDs.Add(MovieSceneSequenceID::Root);
+	ActiveTemplateStates.Reset();
+	ActiveTemplateStates.Add(true);
 
 	RootTemplateInstance.Initialize(NewSequence, *this);
 
@@ -559,6 +685,7 @@ void FSequencer::FocusSequenceInstance(UMovieSceneSubSection& InSubSection)
 	}
 
 	ActiveTemplateIDs.Push(SequenceID);
+	ActiveTemplateStates.Push(InSubSection.IsActive());
 
 	if (Settings->ShouldEvaluateSubSequencesInIsolation())
 	{
@@ -638,11 +765,11 @@ TArray<UObject*> FSequencer::GetEventContexts() const
 	return Temp;
 }
 
-void FSequencer::GetKeysFromSelection(TUniquePtr<ISequencerKeyCollection>& KeyCollection, float DuplicateThreshold)
+void FSequencer::GetKeysFromSelection(TUniquePtr<FSequencerKeyCollection>& KeyCollection, float DuplicateThresholdSeconds)
 {
 	if (!KeyCollection.IsValid())
 	{
-		KeyCollection.Reset(new FGroupedKeyCollection);
+		KeyCollection.Reset(new FSequencerKeyCollection);
 	}
 
 	TArray<FSequencerDisplayNode*> SelectedNodes;
@@ -651,7 +778,26 @@ void FSequencer::GetKeysFromSelection(TUniquePtr<ISequencerKeyCollection>& KeyCo
 		SelectedNodes.Add(&Node.Get());
 	}
 
-	KeyCollection->InitializeRecursive(SelectedNodes, DuplicateThreshold);
+	FFrameNumber ThresholdFrames = (DuplicateThresholdSeconds * GetFocusedFrameResolution()).FloorToFrame();
+	KeyCollection->Update(FSequencerKeyCollectionSignature::FromNodesRecursive(SelectedNodes, ThresholdFrames));
+}
+
+
+void FSequencer::GetAllKeys(TUniquePtr<FSequencerKeyCollection>& KeyCollection, float DuplicateThresholdSeconds)
+{
+	if (!KeyCollection.IsValid())
+	{
+		KeyCollection.Reset(new FSequencerKeyCollection);
+	}
+
+	TArray<FSequencerDisplayNode*> AllNodes;
+	for (TSharedRef<FSequencerDisplayNode> Node : NodeTree->GetAllNodes())
+	{
+		AllNodes.Add(&Node.Get());
+	}
+
+	FFrameNumber ThresholdFrames = (DuplicateThresholdSeconds * GetFocusedFrameResolution()).FloorToFrame();
+	KeyCollection->Update(FSequencerKeyCollectionSignature::FromNodesRecursive(AllNodes, ThresholdFrames));
 }
 
 
@@ -663,6 +809,7 @@ void FSequencer::PopToSequenceInstance(FMovieSceneSequenceIDRef SequenceID)
 		while( SequenceID != ActiveTemplateIDs.Last() )
 		{
 			ActiveTemplateIDs.Pop();
+			ActiveTemplateStates.Pop();
 		}
 
 		check( ActiveTemplateIDs.Num() > 0 );
@@ -681,9 +828,6 @@ void FSequencer::PopToSequenceInstance(FMovieSceneSequenceIDRef SequenceID)
 			}
 		}
 
-		UMovieSceneSequence* Sequence = GetFocusedMovieSceneSequence();
-		SetPerspectiveViewportCameraCutEnabled(Sequence && Sequence->GetMovieScene()->GetCameraCutTrack());
-
 		ResetPerMovieSceneData();
 		SequencerWidget->UpdateBreadcrumbs();
 
@@ -695,7 +839,7 @@ void FSequencer::PopToSequenceInstance(FMovieSceneSequenceIDRef SequenceID)
 
 void FSequencer::UpdateSubSequenceData()
 {
-	SubSequenceRange = TRange<float>(0.f);
+	SubSequenceRange = TRange<FFrameNumber>::Empty();
 	RootToLocalTransform = FMovieSceneSequenceTransform();
 
 	// Find the parent sub section and set up the sub sequence range, if necessary
@@ -708,7 +852,7 @@ void FSequencer::UpdateSubSequenceData()
 
 	if (SubSequenceData)
 	{
-		SubSequenceRange = SubSequenceData->PlayRange;
+		SubSequenceRange = SubSequenceData->PlayRange.Value;
 		RootToLocalTransform = SubSequenceData->RootToSequenceTransform;
 	}
 }
@@ -826,19 +970,27 @@ void FSequencer::DeleteSections(const TSet<TWeakObjectPtr<UMovieSceneSection>>&
 
 void FSequencer::DeleteSelectedKeys()
 {
+	ISequencerModule& SequencerModule = FModuleManager::LoadModuleChecked<ISequencerModule>("Sequencer");
+
 	FScopedTransaction DeleteKeysTransaction( NSLOCTEXT("Sequencer", "DeleteSelectedKeys_Transaction", "Delete Selected Keys") );
 	bool bAnythingRemoved = false;
-	TArray<FSequencerSelectedKey> SelectedKeysArray = Selection.GetSelectedKeys().Array();
 
-	for (const FSequencerSelectedKey& Key : SelectedKeysArray)
+	FSelectedKeysByChannelType KeysByChannel(Selection.GetSelectedKeys().Array());
+	TSet<UMovieSceneSection*> ModifiedSections;
+
+	for (TTuple<void*, FSelectedChannelInfo>& Pair : KeysByChannel.ChannelToKeyHandleMap)
 	{
-		if (Key.IsValid())
+		ISequencerChannelInterface* ChannelInterface = SequencerModule.FindChannelInterface(Pair.Value.ChannelTypeID);
+		if (ChannelInterface)
 		{
-			if (Key.Section->TryModify())
+			if (!ModifiedSections.Contains(Pair.Value.OwningSection))
 			{
-				Key.KeyArea->DeleteKey(Key.KeyHandle.GetValue());
-				bAnythingRemoved = true;
+				Pair.Value.OwningSection->Modify();
+				ModifiedSections.Add(Pair.Value.OwningSection);
 			}
+
+			ChannelInterface->DeleteKeys_Raw(Pair.Key, Pair.Value.KeyHandles);
+			bAnythingRemoved = true;
 		}
 	}
 
@@ -863,82 +1015,90 @@ void FSequencer::SetInterpTangentMode(ERichCurveInterpMode InterpMode, ERichCurv
 	FScopedTransaction SetInterpTangentModeTransaction(NSLOCTEXT("Sequencer", "SetInterpTangentMode_Transaction", "Set Interpolation and Tangent Mode"));
 	bool bAnythingChanged = false;
 
-	for (const FSequencerSelectedKey& Key : SelectedKeysArray)
+	FSelectedKeysByChannelType KeysByChannel(SelectedKeysArray);
+	TSet<UMovieSceneSection*> ModifiedSections;
+
+	const uint32 FloatChannelID = FMovieSceneFloatChannel::GetChannelID();
+
+	// @todo: sequencer-timecode: move this float-specific logic elsewhere to make it extensible for any channel type
+	for (TTuple<void*, FSelectedChannelInfo>& Pair : KeysByChannel.ChannelToKeyHandleMap)
 	{
-		if (Key.IsValid())
+		if (Pair.Value.ChannelTypeID == FloatChannelID)
 		{
-			if (Key.Section->TryModify())
+			if (!ModifiedSections.Contains(Pair.Value.OwningSection))
 			{
-				Key.KeyArea->SetKeyInterpMode(Key.KeyHandle.GetValue(), InterpMode);
-				Key.KeyArea->SetKeyTangentMode(Key.KeyHandle.GetValue(), TangentMode);
-				bAnythingChanged = true;
+				Pair.Value.OwningSection->Modify();
+				ModifiedSections.Add(Pair.Value.OwningSection);
 			}
-		}
-	}
-
-	if (bAnythingChanged)
-	{
-		NotifyMovieSceneDataChanged( EMovieSceneDataChangeType::TrackValueChanged );
-	}
-}
 
+			FMovieSceneFloatChannel* Channel = static_cast<FMovieSceneFloatChannel*>(Pair.Key);
+			TMovieSceneChannel<FMovieSceneFloatValue> ChannelInterface = static_cast<FMovieSceneFloatChannel*>(Pair.Key)->GetInterface();
 
-bool FSequencer::IsInterpTangentModeSelected(ERichCurveInterpMode InterpMode, ERichCurveTangentMode TangentMode) const
-{
-	TArray<FSequencerSelectedKey> SelectedKeysArray = Selection.GetSelectedKeys().Array();
+			TArrayView<FMovieSceneFloatValue> Values = ChannelInterface.GetValues();
 
-	bool bAllSelected = false;
-	for (const FSequencerSelectedKey& Key : SelectedKeysArray)
-	{
-		if (Key.IsValid())
-		{
-			bAllSelected = true;
-			if (Key.KeyArea->GetKeyInterpMode(Key.KeyHandle.GetValue()) != InterpMode || 
-				Key.KeyArea->GetKeyTangentMode(Key.KeyHandle.GetValue()) != TangentMode)
+			for (FKeyHandle Handle : Pair.Value.KeyHandles)
 			{
-				bAllSelected = false;
-				break;
+				const int32 KeyIndex = ChannelInterface.GetIndex(Handle);
+				if (KeyIndex != INDEX_NONE)
+				{
+					Values[KeyIndex].InterpMode = InterpMode;
+					Values[KeyIndex].TangentMode = TangentMode;
+					bAnythingChanged = true;
+				}
 			}
+
+			Channel->AutoSetTangents();
 		}
 	}
-	return bAllSelected;
-}
-
 
-float FSequencer::GetFixedFrameInterval() const
-{
-	if (GetFocusedMovieSceneSequence() && GetFocusedMovieSceneSequence()->GetMovieScene())
+	if (bAnythingChanged)
 	{
-		return GetFocusedMovieSceneSequence()->GetMovieScene()->GetFixedFrameInterval();
+		NotifyMovieSceneDataChanged( EMovieSceneDataChangeType::TrackValueChanged );
 	}
-
-	return 1.f;
 }
 
 void FSequencer::SnapToFrame()
 {
+	ISequencerModule& SequencerModule = FModuleManager::LoadModuleChecked<ISequencerModule>("Sequencer");
+
 	FScopedTransaction SnapToFrameTransaction(NSLOCTEXT("Sequencer", "SnapToFrame_Transaction", "Snap Selected Keys to Frame"));
 	bool bAnythingChanged = false;
-	TArray<FSequencerSelectedKey> SelectedKeysArray = Selection.GetSelectedKeys().Array();
 
-	for (const FSequencerSelectedKey& Key : SelectedKeysArray)
+	FSelectedKeysByChannelType KeysByChannel(Selection.GetSelectedKeys().Array());
+	TSet<UMovieSceneSection*> ModifiedSections;
+
+	TArray<FFrameNumber> KeyTimesScratch;
+	for (TTuple<void*, FSelectedChannelInfo>& Pair : KeysByChannel.ChannelToKeyHandleMap)
 	{
-		if (Key.IsValid())
+		ISequencerChannelInterface* ChannelInterface = SequencerModule.FindChannelInterface(Pair.Value.ChannelTypeID);
+		if (ChannelInterface)
 		{
-			if (Key.Section->TryModify())
+			if (!ModifiedSections.Contains(Pair.Value.OwningSection))
 			{
-				float NewKeyTime = Key.KeyArea->GetKeyTime(Key.KeyHandle.GetValue());
+				Pair.Value.OwningSection->Modify();
+				ModifiedSections.Add(Pair.Value.OwningSection);
+			}
 
-				// Convert to frame
-				float FrameRate = 1.0f / GetFixedFrameInterval();
-				int32 NewFrame = SequencerHelpers::TimeToFrame(NewKeyTime, FrameRate);
+			const int32 NumKeys = Pair.Value.KeyHandles.Num();
+			KeyTimesScratch.Reset(NumKeys);
+			KeyTimesScratch.SetNum(NumKeys);
 
-				// Convert back to time
-				NewKeyTime = SequencerHelpers::FrameToTime(NewFrame, FrameRate);
+			ChannelInterface->GetKeyTimes_Raw(Pair.Key, Pair.Value.KeyHandles, KeyTimesScratch);
 
-				Key.KeyArea->SetKeyTime(Key.KeyHandle.GetValue(), NewKeyTime);
-				bAnythingChanged = true;
+			FFrameRate FrameResolution = GetFocusedFrameResolution();
+			FFrameRate PlayRate        = GetFocusedPlayRate();
+
+			for (FFrameNumber& Time : KeyTimesScratch)
+			{
+				// Convert to frame
+				FFrameNumber PlayFrame    = FFrameRate::TransformTime(Time,      FrameResolution, PlayRate).RoundToFrame();
+				FFrameNumber SnappedFrame = FFrameRate::TransformTime(PlayFrame, PlayRate, FrameResolution).RoundToFrame();
+
+				Time = SnappedFrame;
 			}
+
+			ChannelInterface->SetKeyTimes_Raw(Pair.Key, Pair.Value.KeyHandles, KeyTimesScratch);
+			bAnythingChanged = true;
 		}
 	}
 
@@ -953,80 +1113,66 @@ bool FSequencer::CanSnapToFrame() const
 {
 	const bool bKeysSelected = Selection.GetSelectedKeys().Num() > 0;
 
-	return bKeysSelected && CanShowFrameNumbers();
+	return bKeysSelected;
 }
 
-void FSequencer::TransformSelectedKeysAndSections(float InDeltaTime, float InScale)
+void FSequencer::TransformSelectedKeysAndSections(FFrameTime InDeltaTime, float InScale)
 {
+	ISequencerModule& SequencerModule = FModuleManager::LoadModuleChecked<ISequencerModule>("Sequencer");
+
 	FScopedTransaction TransformKeysAndSectionsTransaction(NSLOCTEXT("Sequencer", "TransformKeysandSections_Transaction", "Transform Keys and Sections"));
 	bool bAnythingChanged = false;
+
 	TArray<FSequencerSelectedKey> SelectedKeysArray = Selection.GetSelectedKeys().Array();
 
-	TArray<UMovieSceneSection*> SectionsWithKeys;
-	
-	FSequencerDisplayNode::DisableKeyGoupingRegeneration();
-	
-	for (const FSequencerSelectedKey& Key : SelectedKeysArray)
+	const FFrameTime OriginTime = GetLocalTime().Time;
+
+	FSelectedKeysByChannelType KeysByChannel(SelectedKeysArray);
+	TMap<UMovieSceneSection*, TRange<FFrameNumber>> SectionToNewBounds;
+
+	TArray<FFrameNumber> KeyTimesScratch;
+	if (InScale != 0.f)
 	{
-		if (Key.IsValid())
+		// Dilate the keys
+		for (TTuple<void*, FSelectedChannelInfo>& Pair : KeysByChannel.ChannelToKeyHandleMap)
 		{
-			UMovieSceneSection* Section = Key.Section;
+			ISequencerChannelInterface* ChannelInterface = SequencerModule.FindChannelInterface(Pair.Value.ChannelTypeID);
+			if (ChannelInterface)
+			{
+				const int32 NumKeys = Pair.Value.KeyHandles.Num();
+				KeyTimesScratch.Reset(NumKeys);
+				KeyTimesScratch.SetNum(NumKeys);
 
-			bool SectionModified = SectionsWithKeys.Contains(Section);
+				// Populate the key times scratch buffer with the times for these handles
+				ChannelInterface->GetKeyTimes_Raw(Pair.Key, Pair.Value.KeyHandles, KeyTimesScratch);
 
-			if (!SectionModified)
-			{
-				if (Key.Section->TryModify())
-				{
-					SectionModified = true;
-					SectionsWithKeys.Add(Section);
-				}
-			}
-			
-			if (SectionModified)
-			{
-				if (InScale != 0.f)
-				{
-					Key.KeyArea->DilateKey(Key.KeyHandle.GetValue(), InScale, GetLocalTime());
-					bAnythingChanged = true;
-				}
-				if (InDeltaTime != 0.f)
+				// Perform the transformation
+				for (FFrameNumber& Time : KeyTimesScratch)
 				{
-					Key.KeyArea->MoveKey(Key.KeyHandle.GetValue(), InDeltaTime);
-					bAnythingChanged = true;
+					FFrameTime KeyTime = Time;
+					Time = (OriginTime + InDeltaTime + (KeyTime - OriginTime) * InScale).FloorToFrame();
 				}
 
-				float NewKeyTime = Key.KeyArea->GetKeyTime(Key.KeyHandle.GetValue()) + InDeltaTime;
-				if (NewKeyTime > Section->GetEndTime())
-				{
-					Section->SetEndTime(NewKeyTime);
-				}
-				else if (NewKeyTime < Section->GetStartTime())
+				TRange<FFrameNumber>* NewSectionBounds = SectionToNewBounds.Find(Pair.Value.OwningSection);
+				if (!NewSectionBounds)
 				{
-					Section->SetStartTime(NewKeyTime);
+					Pair.Value.OwningSection->Modify();
+					NewSectionBounds = &SectionToNewBounds.Add(Pair.Value.OwningSection, TRange<FFrameNumber>(KeyTimesScratch[0], KeyTimesScratch.Last()+1));
 				}
+
+				// Apply the new, transformed key times
+				ChannelInterface->SetKeyTimes_Raw(Pair.Key, Pair.Value.KeyHandles, KeyTimesScratch);
+				bAnythingChanged = true;
 			}
 		}
 	}
-
-	FSequencerDisplayNode::EnableKeyGoupingRegeneration();
-
-	for (TWeakObjectPtr<UMovieSceneSection> SelectedSection : Selection.GetSelectedSections())
+	
+	SectionToNewBounds.Remove(nullptr);
+	for (TTuple<UMovieSceneSection*, TRange<FFrameNumber>>& Pair : SectionToNewBounds)
 	{
-		if (SelectedSection.IsValid())
+		if (!Pair.Key->GetRange().Contains(Pair.Value))
 		{
-			TSet<FKeyHandle> EmptyKeyHandles;
-			SelectedSection.Get()->SetFlags( RF_Transactional );
-			if (InScale != 1.f)
-			{
-				SelectedSection.Get()->DilateSection(InScale, GetLocalTime(), EmptyKeyHandles);
-				bAnythingChanged = true;
-			}
-			if (InDeltaTime != 0.f)
-			{
-				SelectedSection.Get()->MoveSection(InDeltaTime, EmptyKeyHandles);
-				bAnythingChanged = true;
-			}
+			Pair.Key->SetRange(Pair.Value);
 		}
 	}
 
@@ -1038,7 +1184,9 @@ void FSequencer::TransformSelectedKeysAndSections(float InDeltaTime, float InSca
 
 void FSequencer::TranslateSelectedKeysAndSections(bool bTranslateLeft)
 {
-	TransformSelectedKeysAndSections(bTranslateLeft ? -GetFixedFrameInterval() : GetFixedFrameInterval(), 1.f);
+	int32 Shift = bTranslateLeft ? -1 : 1;
+	FFrameTime Delta = FQualifiedFrameTime(Shift, GetFocusedPlayRate()).ConvertTo(GetFocusedFrameResolution());
+	TransformSelectedKeysAndSections(Delta, 1.f);
 }
 
 void FSequencer::BakeTransform()
@@ -1054,6 +1202,8 @@ void FSequencer::BakeTransform()
 		{
 			continue;
 		}
+
+		FFrameTime ResetTime = PlayPosition.GetCurrentPosition();
 		auto ObjectBindingNode = StaticCastSharedRef<FSequencerObjectBindingNode>(Node);
 
 		FGuid Guid = ObjectBindingNode->GetObjectBinding();
@@ -1064,19 +1214,25 @@ void FSequencer::BakeTransform()
 			UCameraComponent* CameraComponent = MovieSceneHelpers::CameraComponentFromRuntimeObject(RuntimeObject.Get());
 
 			FVector Location = Actor->GetActorLocation();
-			FRotator Rotation = Actor->GetActorRotation();
+			FVector Rotation = Actor->GetActorRotation().Euler();
 			FVector Scale = Actor->GetActorScale();
 
 			// Cache transforms
 			TArray<FVector> Locations;
 			TArray<FRotator> Rotations;
-			TArray<float> KeyTimes;
-			float InRange = GetPlaybackRange().GetLowerBoundValue();
-			float OutRange = GetPlaybackRange().GetUpperBoundValue();
-			float CurrentTime = ScrubPosition;
-			for (float EvalTime = InRange; EvalTime <= OutRange; EvalTime += Settings->GetTimeSnapInterval())
+			TArray<FFrameNumber> KeyTimes;
+
+			FFrameRate   Resolution  = FocusedMovieScene->GetFrameResolution();
+			FFrameRate   SnapRate    = FocusedMovieScene->GetEvaluationType() == EMovieSceneEvaluationType::FrameLocked ? FocusedMovieScene->GetPlaybackFrameRate() : Resolution;
+
+			FFrameNumber InFrame     = MovieScene::DiscreteInclusiveLower(GetPlaybackRange());
+			FFrameNumber OutFrame    = MovieScene::DiscreteExclusiveUpper(GetPlaybackRange());
+
+			FFrameTime Interval = FFrameRate::TransformTime(1, SnapRate, Resolution);
+			for (FFrameTime EvalTime = InFrame; EvalTime < OutFrame; EvalTime += Interval)
 			{
-				FMovieSceneEvaluationRange Range = PlayPosition.JumpTo(EvalTime * RootToLocalTransform.Inverse(), FocusedMovieScene->GetOptionalFixedFrameInterval());
+				FFrameNumber KeyTime = FFrameRate::Snap(EvalTime, Resolution, SnapRate).FloorToFrame();
+				FMovieSceneEvaluationRange Range = PlayPosition.JumpTo(KeyTime * RootToLocalTransform.Inverse());
 				EvaluateInternal(Range);
 
 				if (CameraComponent)
@@ -1099,7 +1255,7 @@ void FSequencer::BakeTransform()
 					Rotations.Add(Actor->GetActorRotation());
 				}
 
-				KeyTimes.Add(EvalTime);
+				KeyTimes.Add(KeyTime);
 			}
 
 			// Delete any attach tracks
@@ -1143,8 +1299,7 @@ void FSequencer::BakeTransform()
 			}
 
 			// Reset position
-			FMovieSceneEvaluationRange Range = PlayPosition.JumpTo(ScrubPosition, FocusedMovieScene->GetOptionalFixedFrameInterval());
-			EvaluateInternal(Range);
+			EvaluateInternal(PlayPosition.JumpTo(ResetTime));
 
 			// If there was an attach track that was the parent, detach and attach to that actor's parent if it exists
 			FTransform ParentInverseTransform;
@@ -1169,39 +1324,37 @@ void FSequencer::BakeTransform()
 				UMovieScene3DTransformSection* TransformSection = CastChecked<UMovieScene3DTransformSection>(TransformTrack->CreateNewSection());
 				TransformTrack->AddSection(*TransformSection);
 			
-				TransformSection->SetIsInfinite(true);
-				TransformSection->SetStartTime(InRange);
-				TransformSection->SetEndTime(OutRange);
-
-				TransformSection->SetDefault(FTransformKey(EKey3DTransformChannel::Translation, EAxis::X, Location.X, false /*bUnwindRotation*/));
-				TransformSection->SetDefault(FTransformKey(EKey3DTransformChannel::Translation, EAxis::Y, Location.Y, false /*bUnwindRotation*/));
-				TransformSection->SetDefault(FTransformKey(EKey3DTransformChannel::Translation, EAxis::Z, Location.Z, false /*bUnwindRotation*/));
-
-				TransformSection->SetDefault(FTransformKey(EKey3DTransformChannel::Rotation, EAxis::X, Rotation.Euler().X, false /*bUnwindRotation*/));
-				TransformSection->SetDefault(FTransformKey(EKey3DTransformChannel::Rotation, EAxis::Y, Rotation.Euler().Y, false /*bUnwindRotation*/));
-				TransformSection->SetDefault(FTransformKey(EKey3DTransformChannel::Rotation, EAxis::Z, Rotation.Euler().Z, false /*bUnwindRotation*/));
-
-				TransformSection->SetDefault(FTransformKey(EKey3DTransformChannel::Scale, EAxis::X, Scale.X, false /*bUnwindRotation*/));
-				TransformSection->SetDefault(FTransformKey(EKey3DTransformChannel::Scale, EAxis::Y, Scale.Y, false /*bUnwindRotation*/));
-				TransformSection->SetDefault(FTransformKey(EKey3DTransformChannel::Scale, EAxis::Z, Scale.Z, false /*bUnwindRotation*/));
+				TransformSection->SetRange(TRange<FFrameNumber>::All());
+
+				TArrayView<FMovieSceneFloatChannel*> FloatChannels = TransformSection->GetChannelProxy().GetChannels<FMovieSceneFloatChannel>();
+				FloatChannels[0]->SetDefault(Location.X);
+				FloatChannels[1]->SetDefault(Location.Y);
+				FloatChannels[2]->SetDefault(Location.Z);
+				FloatChannels[3]->SetDefault(Rotation.X);
+				FloatChannels[4]->SetDefault(Rotation.Y);
+				FloatChannels[5]->SetDefault(Rotation.Z);
+				FloatChannels[6]->SetDefault(Scale.X);
+				FloatChannels[7]->SetDefault(Scale.Y);
+				FloatChannels[8]->SetDefault(Scale.Z);
 
 				for (int32 Counter = 0; Counter < KeyTimes.Num(); ++Counter)
 				{
-					float KeyTime = KeyTimes[Counter];
-					const bool bUnwindRotation = false;
+					FFrameNumber KeyTime = KeyTimes[Counter];
 
 					FTransform Transform(Rotations[Counter], Locations[Counter]);
 					FTransform LocalTransform = ParentInverseTransform * Transform;
 					FVector LocalTranslation = LocalTransform.GetTranslation();
 					FVector LocalRotation = LocalTransform.GetRotation().Euler();
 
-					TransformSection->AddKey(KeyTime, FTransformKey(EKey3DTransformChannel::Translation, EAxis::X, LocalTranslation.X, bUnwindRotation), EMovieSceneKeyInterpolation::Linear);
-					TransformSection->AddKey(KeyTime, FTransformKey(EKey3DTransformChannel::Translation, EAxis::Y, LocalTranslation.Y, bUnwindRotation), EMovieSceneKeyInterpolation::Linear);
-					TransformSection->AddKey(KeyTime, FTransformKey(EKey3DTransformChannel::Translation, EAxis::Z, LocalTranslation.Z, bUnwindRotation), EMovieSceneKeyInterpolation::Linear);
-
-					TransformSection->AddKey(KeyTime, FTransformKey(EKey3DTransformChannel::Rotation, EAxis::X, LocalRotation.X, bUnwindRotation), EMovieSceneKeyInterpolation::Linear);
-					TransformSection->AddKey(KeyTime, FTransformKey(EKey3DTransformChannel::Rotation, EAxis::Y, LocalRotation.Y, bUnwindRotation), EMovieSceneKeyInterpolation::Linear);
-					TransformSection->AddKey(KeyTime, FTransformKey(EKey3DTransformChannel::Rotation, EAxis::Z, LocalRotation.Z, bUnwindRotation), EMovieSceneKeyInterpolation::Linear);
+					FloatChannels[0]->AddLinearKey(KeyTime, LocalTranslation.X);
+					FloatChannels[1]->AddLinearKey(KeyTime, LocalTranslation.Y);
+					FloatChannels[2]->AddLinearKey(KeyTime, LocalTranslation.Z);
+					FloatChannels[3]->AddLinearKey(KeyTime, LocalRotation.X);
+					FloatChannels[4]->AddLinearKey(KeyTime, LocalRotation.Y);
+					FloatChannels[5]->AddLinearKey(KeyTime, LocalRotation.Z);
+					FloatChannels[6]->AddLinearKey(KeyTime, Scale.X);
+					FloatChannels[7]->AddLinearKey(KeyTime, Scale.Y);
+					FloatChannels[8]->AddLinearKey(KeyTime, Scale.Z);
 				}
 			}
 		}
@@ -1247,14 +1400,14 @@ void FSequencer::NotifyMovieSceneDataChanged( EMovieSceneDataChangeType DataChan
 	else if (DataChangeType == EMovieSceneDataChangeType::TrackValueChangedRefreshImmediately)
 	{
 		// Evaluate now
-		EvaluateInternal(FMovieSceneEvaluationRange(ScrubPosition));
+		EvaluateInternal(PlayPosition.GetCurrentPositionAsRange());
 	}
 	else if (DataChangeType == EMovieSceneDataChangeType::RefreshAllImmediately)
 	{
 		RefreshTree();
 
 		// Evaluate now
-		EvaluateInternal(FMovieSceneEvaluationRange(ScrubPosition));
+		EvaluateInternal(PlayPosition.GetCurrentPositionAsRange());
 	}
 	else
 	{
@@ -1307,35 +1460,35 @@ FAnimatedRange FSequencer::GetViewRange() const
 
 FAnimatedRange FSequencer::GetClampRange() const
 {
-	UMovieScene* FocusedMovieScene = GetFocusedMovieSceneSequence()->GetMovieScene();
-	return FocusedMovieScene->GetEditorData().WorkingRange;
+	return GetFocusedMovieSceneSequence()->GetMovieScene()->GetEditorData().GetWorkingRange();
 }
 
 
-void FSequencer::SetClampRange(TRange<float> InNewClampRange)
+void FSequencer::SetClampRange(TRange<double> InNewClampRange)
 {
-	UMovieScene* FocusedMovieScene = GetFocusedMovieSceneSequence()->GetMovieScene();
-	FocusedMovieScene->GetEditorData().WorkingRange = InNewClampRange;
+	FMovieSceneEditorData& EditorData = GetFocusedMovieSceneSequence()->GetMovieScene()->GetEditorData();
+	EditorData.WorkStart = InNewClampRange.GetLowerBoundValue();
+	EditorData.WorkEnd   = InNewClampRange.GetUpperBoundValue();
 }
 
 
-TOptional<TRange<float>> FSequencer::GetSubSequenceRange() const
+TOptional<TRange<FFrameNumber>> FSequencer::GetSubSequenceRange() const
 {
 	if (Settings->ShouldEvaluateSubSequencesInIsolation() || ActiveTemplateIDs.Num() == 1)
 	{
-		return TOptional<TRange<float>>();
+		return TOptional<TRange<FFrameNumber>>();
 	}
 	return SubSequenceRange;
 }
 
 
-TRange<float> FSequencer::GetSelectionRange() const
+TRange<FFrameNumber> FSequencer::GetSelectionRange() const
 {
 	return GetFocusedMovieSceneSequence()->GetMovieScene()->GetSelectionRange();
 }
 
 
-void FSequencer::SetSelectionRange(TRange<float> Range)
+void FSequencer::SetSelectionRange(TRange<FFrameNumber> Range)
 {
 	const FScopedTransaction Transaction(LOCTEXT("SetSelectionRange_Transaction", "Set Selection Range"));
 	UMovieScene* FocussedMovieScene = GetFocusedMovieSceneSequence()->GetMovieScene();
@@ -1346,65 +1499,57 @@ void FSequencer::SetSelectionRange(TRange<float> Range)
 
 void FSequencer::SetSelectionRangeEnd()
 {
-	const float LocalTime = GetLocalTime();
+	const FFrameNumber LocalTime = GetLocalTime().Time.FrameNumber + 1;
 
 	if (GetSelectionRange().GetLowerBoundValue() >= LocalTime)
 	{
-		SetSelectionRange(TRange<float>(LocalTime));
+		SetSelectionRange(TRange<FFrameNumber>(LocalTime));
 	}
 	else
 	{
-		SetSelectionRange(TRange<float>(GetSelectionRange().GetLowerBoundValue(), LocalTime));
+		SetSelectionRange(TRange<FFrameNumber>(GetSelectionRange().GetLowerBound(), LocalTime));
 	}
 }
 
 
 void FSequencer::SetSelectionRangeStart()
 {
-	const float LocalTime = GetLocalTime();
+	const FFrameNumber LocalTime = GetLocalTime().Time.FrameNumber;
 
 	if (GetSelectionRange().GetUpperBoundValue() <= LocalTime)
 	{
-		SetSelectionRange(TRange<float>(LocalTime));
+		SetSelectionRange(TRange<FFrameNumber>(LocalTime));
 	}
 	else
 	{
-		SetSelectionRange(TRange<float>(LocalTime, GetSelectionRange().GetUpperBoundValue()));
+		SetSelectionRange(TRange<FFrameNumber>(LocalTime, GetSelectionRange().GetUpperBound()));
 	}
 }
 
 
-void FSequencer::SelectInSelectionRange(const TSharedRef<FSequencerDisplayNode>& DisplayNode, const TRange<float>& SelectionRange, bool bSelectKeys, bool bSelectSections)
+void FSequencer::SelectInSelectionRange(const TSharedRef<FSequencerDisplayNode>& DisplayNode, const TRange<FFrameNumber>& SelectionRange, bool bSelectKeys, bool bSelectSections)
 {
 	if (DisplayNode->GetType() == ESequencerNode::Track)
 	{
 		if (bSelectKeys)
 		{
-			TSet<TSharedPtr<IKeyArea>> OutKeyAreas;
-			SequencerHelpers::GetAllKeyAreas(DisplayNode, OutKeyAreas);
+			TArray<FKeyHandle> HandlesScratch;
 
-			for (auto KeyArea : OutKeyAreas)
-			{
-				if (!KeyArea.IsValid())
-				{
-					continue;
-				}
+			TSet<TSharedPtr<IKeyArea>> KeyAreas;
+			SequencerHelpers::GetAllKeyAreas(DisplayNode, KeyAreas);
 
+			for (TSharedPtr<IKeyArea> KeyArea : KeyAreas)
+			{
 				UMovieSceneSection* Section = KeyArea->GetOwningSection();
 
-				if (Section == nullptr)
-				{
-					continue;
-				}
-
-				if (bSelectKeys)
+				if (Section)
 				{
-					TSet<FKeyHandle> OutKeyHandles;
-					Section->GetKeyHandles(OutKeyHandles, SelectionRange);
+					HandlesScratch.Reset();
+					KeyArea->GetKeyHandles(HandlesScratch, SelectionRange);
 
-					for (auto KeyHandle : OutKeyHandles)
+					for (int32 Index = 0; Index < HandlesScratch.Num(); ++Index)
 					{
-						Selection.AddToSelection(FSequencerSelectedKey(*Section, KeyArea, KeyHandle));
+						Selection.AddToSelection(FSequencerSelectedKey(*Section, KeyArea, HandlesScratch[Index]));
 					}
 				}
 			}
@@ -1412,14 +1557,12 @@ void FSequencer::SelectInSelectionRange(const TSharedRef<FSequencerDisplayNode>&
 
 		if (bSelectSections)
 		{
-			// Use an exclusive selection range to prevent the selection of a section that ends right at the selection range start
-			TRange<float> ExclusiveSectionRange = TRange<float>(TRangeBound<float>::Exclusive(SelectionRange.GetLowerBoundValue()), TRangeBound<float>::Exclusive(SelectionRange.GetUpperBoundValue()));
 			TSet<TWeakObjectPtr<UMovieSceneSection>> OutSections;
 			SequencerHelpers::GetAllSections(DisplayNode, OutSections);
 
 			for (auto Section : OutSections)
 			{
-				if (Section.IsValid() && Section->GetRange().Overlaps(ExclusiveSectionRange))
+				if (Section.IsValid() && Section->GetRange().Overlaps(SelectionRange))
 				{
 					Selection.AddToSelection(Section.Get());
 				}
@@ -1435,14 +1578,14 @@ void FSequencer::SelectInSelectionRange(const TSharedRef<FSequencerDisplayNode>&
 
 void FSequencer::ResetSelectionRange()
 {
-	SetSelectionRange(TRange<float>::Empty());
+	SetSelectionRange(TRange<FFrameNumber>::Empty());
 }
 
 void FSequencer::SelectInSelectionRange(bool bSelectKeys, bool bSelectSections)
 {
 	UMovieSceneSequence* Sequence = GetFocusedMovieSceneSequence();
 	UMovieScene* MovieScene = Sequence->GetMovieScene();
-	TRange<float> SelectionRange = MovieScene->GetSelectionRange();
+	TRange<FFrameNumber> SelectionRange = MovieScene->GetSelectionRange();
 
 	Selection.Empty();
 
@@ -1453,15 +1596,15 @@ void FSequencer::SelectInSelectionRange(bool bSelectKeys, bool bSelectSections)
 }
 
 
-TRange<float> FSequencer::GetPlaybackRange() const
+TRange<FFrameNumber> FSequencer::GetPlaybackRange() const
 {
 	return GetFocusedMovieSceneSequence()->GetMovieScene()->GetPlaybackRange();
 }
 
 
-void FSequencer::SetPlaybackRange(TRange<float> Range)
+void FSequencer::SetPlaybackRange(TRange<FFrameNumber> Range)
 {
-	if (ensure(Range.HasLowerBound() && Range.HasUpperBound() && !Range.IsDegenerate()))
+	if (ensure(Range.HasLowerBound() && Range.HasUpperBound()))
 	{
 		UMovieScene* FocusedMovieScene = GetFocusedMovieSceneSequence()->GetMovieScene();
 		
@@ -1469,14 +1612,14 @@ void FSequencer::SetPlaybackRange(TRange<float> Range)
 		{
 			const FScopedTransaction Transaction(LOCTEXT("SetPlaybackRange_Transaction", "Set Playback Range"));
 
-			FocusedMovieScene->SetPlaybackRange(Range.GetLowerBoundValue(), Range.GetUpperBoundValue());
+			FocusedMovieScene->SetPlaybackRange(Range);
 
 			bNeedsEvaluate = true;
 		}
 	}
 }
 
-UMovieSceneSection* FSequencer::FindNextOrPreviousShot(UMovieSceneSequence* Sequence, float CurrentTime, const bool bNextShot) const
+UMovieSceneSection* FSequencer::FindNextOrPreviousShot(UMovieSceneSequence* Sequence, FFrameNumber SearchFromTime, const bool bNextShot) const
 {
 	UMovieScene* OwnerMovieScene = Sequence->GetMovieScene();
 
@@ -1486,30 +1629,30 @@ UMovieSceneSection* FSequencer::FindNextOrPreviousShot(UMovieSceneSequence* Sequ
 		return nullptr;
 	}
 
-	float MinTime = TNumericLimits<float>::Max();
+	FFrameNumber MinTime = TNumericLimits<FFrameNumber>::Max();
 
-	TMap<float, int32> StartTimeMap;
+	TMap<FFrameNumber, int32> StartTimeMap;
 	for (int32 SectionIndex = 0; SectionIndex < CinematicShotTrack->GetAllSections().Num(); ++SectionIndex)
 	{
 		UMovieSceneSection* ShotSection = CinematicShotTrack->GetAllSections()[SectionIndex];
 
-		if (ShotSection)
+		if (ShotSection && ShotSection->HasStartFrame())
 		{
-			StartTimeMap.Add(ShotSection->GetStartTime(), SectionIndex);
+			StartTimeMap.Add(ShotSection->GetInclusiveStartFrame(), SectionIndex);
 		}
 	}
 
-	StartTimeMap.KeySort(TLess<float>());
+	StartTimeMap.KeySort(TLess<FFrameNumber>());
 
 	int32 MinShotIndex = -1;
 	for (auto StartTimeIt = StartTimeMap.CreateIterator(); StartTimeIt; ++StartTimeIt)
 	{
-		float StartTime = StartTimeIt->Key;
+		FFrameNumber StartTime = StartTimeIt->Key;
 		if (bNextShot)
 		{
-			if (StartTime > CurrentTime)
+			if (StartTime > SearchFromTime)
 			{
-				float DiffTime = FMath::Abs(StartTime - CurrentTime);
+				FFrameNumber DiffTime = FMath::Abs(StartTime - SearchFromTime);
 				if (DiffTime < MinTime)
 				{
 					MinTime = DiffTime;
@@ -1519,9 +1662,9 @@ UMovieSceneSection* FSequencer::FindNextOrPreviousShot(UMovieSceneSequence* Sequ
 		}
 		else
 		{
-			if (CurrentTime >= StartTime)
+			if (SearchFromTime >= StartTime)
 			{
-				float DiffTime = FMath::Abs(StartTime - CurrentTime);
+				FFrameNumber DiffTime = FMath::Abs(StartTime - SearchFromTime);
 				if (DiffTime < MinTime)
 				{
 					MinTime = DiffTime;
@@ -1564,11 +1707,12 @@ UMovieSceneSection* FSequencer::FindNextOrPreviousShot(UMovieSceneSequence* Sequ
 
 void FSequencer::SetSelectionRangeToShot(const bool bNextShot)
 {
-	UMovieSceneSection* TargetShotSection = FindNextOrPreviousShot(GetFocusedMovieSceneSequence(), GetGlobalTime(), bNextShot);
-		
-	if (TargetShotSection)
+	UMovieSceneSection* TargetShotSection = FindNextOrPreviousShot(GetFocusedMovieSceneSequence(), GetLocalTime().Time.FloorToFrame(), bNextShot);
+
+	TRange<FFrameNumber> NewSelectionRange = TargetShotSection ? TargetShotSection->GetRange() : TRange<FFrameNumber>::All();
+	if (NewSelectionRange.GetLowerBound().IsClosed() && NewSelectionRange.GetUpperBound().IsClosed())
 	{
-		SetSelectionRange(TRange<float>(TargetShotSection->GetStartTime(), TargetShotSection->GetEndTime()));
+		SetSelectionRange(NewSelectionRange);
 	}
 }
 
@@ -1583,27 +1727,17 @@ void FSequencer::SetPlaybackRangeToAllShots()
 		return;
 	}
 
-	float LowerBound = CinematicShotTrack->GetAllSections()[0]->GetStartTime();
-	float UpperBound = CinematicShotTrack->GetAllSections()[0]->GetEndTime();
+	TRange<FFrameNumber> NewRange = CinematicShotTrack->GetAllSections()[0]->GetRange();
 
-	for (int32 SectionIndex = 0; SectionIndex < CinematicShotTrack->GetAllSections().Num(); ++SectionIndex)
+	for (UMovieSceneSection* ShotSection : CinematicShotTrack->GetAllSections())
 	{
-		UMovieSceneSection* ShotSection = CinematicShotTrack->GetAllSections()[SectionIndex];
-		if (ShotSection)
+		if (ShotSection && ShotSection->HasStartFrame() && ShotSection->HasEndFrame())
 		{
-			if (ShotSection->GetStartTime() < LowerBound)
-			{
-				LowerBound = ShotSection->GetStartTime();
-			}
-
-			if (ShotSection->GetEndTime() > UpperBound)
-			{
-				UpperBound = ShotSection->GetEndTime();
-			}
+			NewRange = TRange<FFrameNumber>::Hull(ShotSection->GetRange(), NewRange);
 		}
 	}
 
-	SetPlaybackRange(TRange<float>(LowerBound, UpperBound));
+	SetPlaybackRange(NewRange);
 }
 
 bool FSequencer::IsPlaybackRangeLocked() const
@@ -1631,18 +1765,16 @@ void FSequencer::TogglePlaybackRangeLocked()
 
 void FSequencer::ResetViewRange()
 {
-	float InRange = GetPlaybackRange().GetLowerBoundValue();
-	float OutRange = GetPlaybackRange().GetUpperBoundValue();
-	const float OutputViewSize = OutRange - InRange;
-	const float OutputChange = OutputViewSize * 0.1f;
+	TRange<double> PlayRangeSeconds = GetPlaybackRange() / GetFocusedFrameResolution();
+	const double OutputViewSize = PlayRangeSeconds.Size<double>();
+	const double OutputChange = OutputViewSize * 0.1f;
 
 	if (OutputChange > 0)
 	{
-		InRange -= OutputChange;
-		OutRange += OutputChange;
+		PlayRangeSeconds = MovieScene::ExpandRange(PlayRangeSeconds, OutputChange);
 
-		SetClampRange(TRange<float>(InRange, OutRange));	
-		SetViewRange(TRange<float>(InRange, OutRange), EViewRangeInterpolation::Animated);
+		SetClampRange(PlayRangeSeconds);
+		SetViewRange(PlayRangeSeconds, EViewRangeInterpolation::Animated);
 	}
 }
 
@@ -1651,17 +1783,19 @@ void FSequencer::ZoomViewRange(float InZoomDelta)
 {
 	float LocalViewRangeMax = TargetViewRange.GetUpperBoundValue();
 	float LocalViewRangeMin = TargetViewRange.GetLowerBoundValue();
-	const float OutputViewSize = LocalViewRangeMax - LocalViewRangeMin;
-	const float OutputChange = OutputViewSize * InZoomDelta;
 
-	float CurrentPositionFraction = (ScrubPosition - LocalViewRangeMin) / OutputViewSize;
+	const double CurrentTime = GetLocalTime().AsSeconds();
+	const double OutputViewSize = LocalViewRangeMax - LocalViewRangeMin;
+	const double OutputChange = OutputViewSize * 0.1f;
 
-	float NewViewOutputMin = LocalViewRangeMin - (OutputChange * CurrentPositionFraction);
-	float NewViewOutputMax = LocalViewRangeMax + (OutputChange * (1.f - CurrentPositionFraction));
+	float CurrentPositionFraction = (CurrentTime - LocalViewRangeMin) / OutputViewSize;
+
+	double NewViewOutputMin = LocalViewRangeMin - (OutputChange * CurrentPositionFraction);
+	double NewViewOutputMax = LocalViewRangeMax + (OutputChange * (1.f - CurrentPositionFraction));
 
 	if (NewViewOutputMin < NewViewOutputMax)
 	{
-		SetViewRange(TRange<float>(NewViewOutputMin, NewViewOutputMax), EViewRangeInterpolation::Animated);
+		SetViewRange(TRange<double>(NewViewOutputMin, NewViewOutputMax), EViewRangeInterpolation::Animated);
 	}
 }
 
@@ -1677,41 +1811,6 @@ void FSequencer::ZoomOutViewRange()
 	ZoomViewRange(0.1f);
 }
 
-bool GetMovieSceneSectionPlayRange(UMovieSceneSection* InSection, TRange<float>& OutBounds)
-{
-	if (InSection->IsInfinite())
-	{
-		TRange<float> KeyBounds = TRange<float>::Empty();
-		 
-		TSet<FKeyHandle> KeyHandles;
-		InSection->GetKeyHandles(KeyHandles, TRange<float>());
-		for (auto KeyHandle : KeyHandles)
-		{
-			TOptional<float> KeyTime = InSection->GetKeyTime(KeyHandle);
-			if (KeyTime.IsSet())
-			{
-				if (KeyBounds.IsEmpty())
-				{
-					KeyBounds = TRange<float>(KeyTime.GetValue());
-				}
-				else
-				{
-					KeyBounds = TRange<float>( FMath::Min(KeyBounds.GetLowerBoundValue(), KeyTime.GetValue()),
-												FMath::Max(KeyBounds.GetUpperBoundValue(), KeyTime.GetValue()) );
-				}
-			}
-		}
-
-		OutBounds = KeyBounds;
-		return KeyHandles.Num() > 0;
-	}
-	else
-	{
-		OutBounds = InSection->GetRange();
-		return true;
-	}
-}
-
 void FSequencer::UpdatePlaybackRange()
 {
 	if (Settings->ShouldKeepPlayRangeInSectionBounds())
@@ -1719,59 +1818,22 @@ void FSequencer::UpdatePlaybackRange()
 		UMovieScene* FocusedMovieScene = GetFocusedMovieSceneSequence()->GetMovieScene();
 		TArray<UMovieSceneSection*> AllSections = FocusedMovieScene->GetAllSections();
 
-		if (AllSections.Num() > 0)
+		if (AllSections.Num() > 0 && !FocusedMovieScene->IsPlaybackRangeLocked())
 		{
-			TRange<float> NewBounds = TRange<float>::Empty();
-
-			TRange<float> OutBounds;
-			if (GetMovieSceneSectionPlayRange(AllSections[0], OutBounds))
-			{
-				NewBounds = OutBounds;
-			}
-	
-			for (auto MovieSceneSection : AllSections)
+			TRange<FFrameNumber> NewBounds = TRange<FFrameNumber>::Empty();
+			for (UMovieSceneSection* Section : AllSections)
 			{
-				if (GetMovieSceneSectionPlayRange(MovieSceneSection, OutBounds))
-				{
-					if (NewBounds.IsEmpty())
-					{
-						NewBounds = OutBounds;
-					}
-					else
-					{
-						NewBounds = TRange<float>( FMath::Min(NewBounds.GetLowerBoundValue(), OutBounds.GetLowerBoundValue()),
-												FMath::Max(NewBounds.GetUpperBoundValue(), OutBounds.GetUpperBoundValue()) );
-					}
-				}
+				NewBounds = TRange<FFrameNumber>::Hull(Section->ComputeEffectiveRange(), NewBounds);
 			}
 
 			// When the playback range is determined by the section bounds, don't mark the change in the playback range otherwise the scene will be marked dirty
-			if (!NewBounds.IsDegenerate() && !FocusedMovieScene->IsPlaybackRangeLocked())
+			if (!NewBounds.IsDegenerate())
 			{
 				const bool bAlwaysMarkDirty = false;
-				FocusedMovieScene->SetPlaybackRange(NewBounds.GetLowerBoundValue(), NewBounds.GetUpperBoundValue(), bAlwaysMarkDirty);
+				FocusedMovieScene->SetPlaybackRange(NewBounds, bAlwaysMarkDirty);
 			}
 		}
 	}
-	else
-	{
-		UMovieScene* FocusedMovieScene = GetFocusedMovieSceneSequence()->GetMovieScene();
-		TRange<float> NewBounds = TRange<float>::Empty();
-
-		for (auto MasterTrack : FocusedMovieScene->GetMasterTracks())
-		{
-			if (MasterTrack->AddsSectionBoundsToPlayRange())
-			{
-				NewBounds = TRange<float>( FMath::Min(NewBounds.GetLowerBoundValue(), MasterTrack->GetSectionBoundaries().GetLowerBoundValue()),
-											FMath::Max(NewBounds.GetUpperBoundValue(), MasterTrack->GetSectionBoundaries().GetUpperBoundValue()) );
-			}
-		}
-			
-		if (!NewBounds.IsEmpty() && !FocusedMovieScene->IsPlaybackRangeLocked())
-		{
-			FocusedMovieScene->SetPlaybackRange(NewBounds.GetLowerBoundValue(), NewBounds.GetUpperBoundValue());
-		}
-	}
 }
 
 
@@ -1853,32 +1915,35 @@ bool FSequencer::GetAutoSetTrackDefaults() const
 }
 
 
-bool FSequencer::IsRecordingLive() const 
+FQualifiedFrameTime FSequencer::GetLocalTime() const
 {
-	return PlaybackState == EMovieScenePlayerStatus::Recording && GIsPlayInEditorWorld;
-}
-
+	FFrameRate FocusedResolution = GetFocusedFrameResolution();
+	FFrameTime CurrentPosition   = PlayPosition.GetCurrentPosition();
 
-float FSequencer::GetLocalTime() const
-{
-	return ScrubPosition * RootToLocalTransform;
+	FFrameTime RootTime = ConvertFrameTime(CurrentPosition, PlayPosition.GetInputRate(), PlayPosition.GetOutputRate());
+	return FQualifiedFrameTime(RootTime * RootToLocalTransform, FocusedResolution);
 }
 
 
-float FSequencer::GetGlobalTime() const
+FQualifiedFrameTime FSequencer::GetGlobalTime() const
 {
-	return ScrubPosition;
+	FFrameTime RootTime = ConvertFrameTime(PlayPosition.GetCurrentPosition(), PlayPosition.GetInputRate(), PlayPosition.GetOutputRate());
+	return FQualifiedFrameTime(RootTime, PlayPosition.GetOutputRate());
 }
 
-void FSequencer::SetLocalTime( float NewTime, ESnapTimeMode SnapTimeMode )
+void FSequencer::SetLocalTime( FFrameTime NewTime, ESnapTimeMode SnapTimeMode )
 {
+	FFrameRate LocalResolution = GetFocusedFrameResolution();
+
 	// Ensure the time is in the current view
-	ScrollIntoView(NewTime);
+	ScrollIntoView(NewTime / LocalResolution);
 
 	// Perform snapping
 	if ((SnapTimeMode & ESnapTimeMode::STM_Interval) && Settings->GetIsSnapEnabled())
 	{
-		NewTime = SequencerHelpers::SnapTimeToInterval(NewTime, GetFixedFrameInterval());
+		FFrameRate LocalPlayRate = GetFocusedPlayRate();
+
+		NewTime = FFrameRate::TransformTime(FFrameRate::TransformTime(NewTime, LocalResolution, LocalPlayRate).RoundToFrame(), LocalPlayRate, LocalResolution);
 	}
 
 	if ((SnapTimeMode & ESnapTimeMode::STM_Keys) && (Settings->GetSnapPlayTimeToKeys() || FSlateApplication::Get().GetModifierKeys().IsShiftDown()))
@@ -1890,32 +1955,33 @@ void FSequencer::SetLocalTime( float NewTime, ESnapTimeMode SnapTimeMode )
 }
 
 
-void FSequencer::SetLocalTimeDirectly(float NewTime)
+void FSequencer::SetLocalTimeDirectly(FFrameTime NewTime)
 {
 	// Transform the time to the root time-space
 	SetGlobalTime(NewTime * RootToLocalTransform.Inverse());
 }
 
 
-void FSequencer::SetGlobalTime( float NewTime )
+void FSequencer::SetGlobalTime( FFrameTime NewTime )
 {
-	// Update the position
-	ScrubPosition = NewTime;
+	NewTime = ConvertFrameTime(NewTime, GetRootFrameResolution(), PlayPosition.GetInputRate());
+	if (PlayPosition.GetEvaluationType() == EMovieSceneEvaluationType::FrameLocked)
+	{
+		NewTime = NewTime.FloorToFrame();
+	}
 
 	// Don't update the sequence if the time hasn't changed as this will cause duplicate events and the like to fire.
 	// If we need to reevaluate the sequence at the same time for whetever reason, we should call ForceEvaluate()
-	TOptional<float> LastPosition = PlayPosition.GetPreviousPosition();
-	if (!LastPosition.IsSet() || LastPosition.GetValue() != ScrubPosition)
+	TOptional<FFrameTime> CurrentPosition = PlayPosition.GetCurrentPosition();
+	if (PlayPosition.GetCurrentPosition() != NewTime)
 	{
-		FMovieSceneEvaluationRange Range = PlayPosition.JumpTo(ScrubPosition, GetRootMovieSceneSequence()->GetMovieScene()->GetOptionalFixedFrameInterval());
-		EvaluateInternal(Range);
+		EvaluateInternal(PlayPosition.JumpTo(NewTime));
 	}
 }
 
 void FSequencer::ForceEvaluate()
 {
-	FMovieSceneEvaluationRange Range = PlayPosition.JumpTo(ScrubPosition, GetFocusedMovieSceneSequence()->GetMovieScene()->GetOptionalFixedFrameInterval());
-	EvaluateInternal(Range);
+	EvaluateInternal(PlayPosition.GetCurrentPositionAsRange());
 }
 
 void FSequencer::EvaluateInternal(FMovieSceneEvaluationRange InRange, bool bHasJumped)
@@ -1969,7 +2035,7 @@ void FSequencer::ScrollIntoView(float InLocalTime)
 		// When not scrubbing, we auto scroll the view range immediately
 		if (RangeOffset != 0.f)
 		{
-			TRange<float> WorkingRange = GetClampRange();
+			TRange<double> WorkingRange = GetClampRange();
 
 			// Adjust the offset so that the target range will be within the working range.
 			if (TargetViewRange.GetLowerBoundValue() + RangeOffset < WorkingRange.GetLowerBoundValue())
@@ -1981,12 +2047,12 @@ void FSequencer::ScrollIntoView(float InLocalTime)
 				RangeOffset = WorkingRange.GetUpperBoundValue() - TargetViewRange.GetUpperBoundValue();
 			}
 
-			SetViewRange(TRange<float>(TargetViewRange.GetLowerBoundValue() + RangeOffset, TargetViewRange.GetUpperBoundValue() + RangeOffset), EViewRangeInterpolation::Immediate);
+			SetViewRange(TRange<double>(TargetViewRange.GetLowerBoundValue() + RangeOffset, TargetViewRange.GetUpperBoundValue() + RangeOffset), EViewRangeInterpolation::Immediate);
 		}
 	}
 }
 
-void FSequencer::UpdateAutoScroll(float NewTime)
+void FSequencer::UpdateAutoScroll(double NewTime)
 {
 	float ThresholdPercentage = 0.025f;
 	AutoscrollOffset = CalculateAutoscrollEncroachment(NewTime, ThresholdPercentage);
@@ -1997,28 +2063,28 @@ void FSequencer::UpdateAutoScroll(float NewTime)
 		return;
 	}
 
-	TRange<float> ViewRange = GetViewRange();
-	const float Threshold = (ViewRange.GetUpperBoundValue() - ViewRange.GetLowerBoundValue()) * ThresholdPercentage;
+	TRange<double> ViewRange = GetViewRange();
+	const double Threshold = (ViewRange.GetUpperBoundValue() - ViewRange.GetLowerBoundValue()) * ThresholdPercentage;
 
-	const float LocalPosition = GetLocalTime();
+	const FQualifiedFrameTime LocalTime = GetLocalTime();
 
 	// If we have no autoscrub offset yet, we move the scrub position to the boundary of the autoscroll threasdhold, then autoscrub from there
 	if (!AutoscrubOffset.IsSet())
 	{
-		if (AutoscrollOffset.GetValue() < 0 && LocalPosition > ViewRange.GetLowerBoundValue() + Threshold)
+		if (AutoscrollOffset.GetValue() < 0 && LocalTime.AsSeconds() > ViewRange.GetLowerBoundValue() + Threshold)
 		{
-			SetLocalTimeDirectly( ViewRange.GetLowerBoundValue() + Threshold );
+			SetLocalTimeDirectly( (ViewRange.GetLowerBoundValue() + Threshold) * LocalTime.Rate );
 		}
-		else if (AutoscrollOffset.GetValue() > 0 && LocalPosition < ViewRange.GetUpperBoundValue() - Threshold)
+		else if (AutoscrollOffset.GetValue() > 0 && LocalTime.AsSeconds() < ViewRange.GetUpperBoundValue() - Threshold)
 		{
-			SetLocalTimeDirectly( ViewRange.GetUpperBoundValue() - Threshold );
+			SetLocalTimeDirectly( (ViewRange.GetUpperBoundValue() - Threshold) * LocalTime.Rate );
 		}
 	}
 
 	// Don't autoscrub if we're at the extremes of the movie scene range
-	const TRange<float>& WorkingRange = GetFocusedMovieSceneSequence()->GetMovieScene()->GetEditorData().WorkingRange;
-	if (NewTime < WorkingRange.GetLowerBoundValue() + Threshold ||
-		NewTime > WorkingRange.GetUpperBoundValue() - Threshold
+	const FMovieSceneEditorData& EditorData = GetFocusedMovieSceneSequence()->GetMovieScene()->GetEditorData();
+	if (NewTime < EditorData.WorkStart + Threshold ||
+		NewTime > EditorData.WorkEnd - Threshold
 		)
 	{
 		AutoscrubOffset.Reset();
@@ -2030,14 +2096,14 @@ void FSequencer::UpdateAutoScroll(float NewTime)
 }
 
 
-TOptional<float> FSequencer::CalculateAutoscrollEncroachment(float NewTime, float ThresholdPercentage) const
+TOptional<float> FSequencer::CalculateAutoscrollEncroachment(double NewTime, float ThresholdPercentage) const
 {
 	enum class EDirection { Positive, Negative };
-	const EDirection Movement = NewTime - GetLocalTime() >= 0 ? EDirection::Positive : EDirection::Negative;
+	const EDirection Movement = NewTime - GetLocalTime().AsSeconds() >= 0 ? EDirection::Positive : EDirection::Negative;
 
-	const TRange<float> CurrentRange = GetViewRange();
-	const float RangeMin = CurrentRange.GetLowerBoundValue(), RangeMax = CurrentRange.GetUpperBoundValue();
-	const float AutoScrollThreshold = (RangeMax - RangeMin) * ThresholdPercentage;
+	const TRange<double> CurrentRange = GetViewRange();
+	const double RangeMin = CurrentRange.GetLowerBoundValue(), RangeMax = CurrentRange.GetUpperBoundValue();
+	const double AutoScrollThreshold = (RangeMax - RangeMin) * ThresholdPercentage;
 
 	if (Movement == EDirection::Negative && NewTime < RangeMin + AutoScrollThreshold)
 	{
@@ -2068,11 +2134,16 @@ void FSequencer::SetPerspectiveViewportCameraCutEnabled(bool bEnabled)
 
 void FSequencer::RenderMovie(UMovieSceneSection* InSection) const
 {
-	RenderMovieInternal(InSection->GetStartTime(), InSection->GetEndTime(), true);
+	RenderMovieInternal(InSection->GetRange(), true);
 }
 
-void FSequencer::RenderMovieInternal(float InStartTime, float InEndTime, bool bSetFrameOverrides) const
+void FSequencer::RenderMovieInternal(TRange<FFrameNumber> Range, bool bSetFrameOverrides) const
 {
+	if (Range.GetLowerBound().IsOpen() || Range.GetUpperBound().IsOpen())
+	{
+		Range = TRange<FFrameNumber>::Hull(Range, GetPlaybackRange());
+	}
+
 	FLevelEditorModule& LevelEditorModule = FModuleManager::GetModuleChecked<FLevelEditorModule>(TEXT("LevelEditor"));
 
 	// Create a new movie scene capture object for an automated level sequence, and open the tab
@@ -2081,34 +2152,73 @@ void FSequencer::RenderMovieInternal(float InStartTime, float InEndTime, bool bS
 
 	MovieSceneCapture->SetLevelSequenceAsset(GetCurrentAsset()->GetPathName());
 
-	if (CanShowFrameNumbers())
+	FFrameRate PlayRate = GetFocusedPlayRate();
+	FFrameRate FrameResolution = GetFocusedFrameResolution();
+
+	MovieSceneCapture->Settings.FrameRate = PlayRate;
+	MovieSceneCapture->Settings.ZeroPadFrameNumbers = Settings->GetZeroPadFrames();
+	MovieSceneCapture->Settings.bUseRelativeFrameNumbers = false;
+
+	FFrameNumber StartFrame = MovieScene::DiscreteInclusiveLower(Range);
+	FFrameNumber EndFrame = MovieScene::DiscreteExclusiveUpper(Range);
+
+	FFrameNumber RoundedStartFrame = FFrameRate::TransformTime(StartFrame, FrameResolution, PlayRate).FloorToFrame();
+	FFrameNumber RoundedEndFrame = FFrameRate::TransformTime(EndFrame, FrameResolution, PlayRate).CeilToFrame();
+
+	if (bSetFrameOverrides)
+	{
+		MovieSceneCapture->SetFrameOverrides(RoundedStartFrame, RoundedEndFrame);
+	}
+	else
 	{
-		MovieSceneCapture->Settings.FrameRate = FMath::RoundToInt(1.f / GetFixedFrameInterval());
-		MovieSceneCapture->Settings.ZeroPadFrameNumbers = Settings->GetZeroPadFrames();
-		MovieSceneCapture->Settings.bUseRelativeFrameNumbers = false;
+		if (!MovieSceneCapture->bUseCustomStartFrame)
+		{
+			MovieSceneCapture->StartFrame = RoundedStartFrame;
+		}
 
-		const int32 StartFrame = FMath::RoundToInt( InStartTime * MovieSceneCapture->Settings.FrameRate );
-		const int32 EndFrame = FMath::Max( StartFrame, FMath::RoundToInt( InEndTime * MovieSceneCapture->Settings.FrameRate ) );
+		if (!MovieSceneCapture->bUseCustomEndFrame)
+		{
+			MovieSceneCapture->EndFrame = RoundedEndFrame;
+		}
+	}
 
-		if (bSetFrameOverrides)
+	// We create a new Numeric Type Interface that ties it's Capture/Resolution rates to the Capture Object so that it converts UI entries
+	// to the correct resolution for the capture, and not for the original sequence.
+	USequencerSettings* LocalSettings = Settings;
+	auto GetDisplayFormatDelegate = FOnGetDisplayFormat::CreateLambda([LocalSettings]() ->EFrameNumberDisplayFormats
+	{
+		if (LocalSettings)
 		{
-			MovieSceneCapture->SetFrameOverrides(StartFrame, EndFrame);
+			return LocalSettings->GetTimeDisplayFormat();
 		}
-		else
+		return EFrameNumberDisplayFormats::Frames;
+	});
+
+	auto GetZeroPadFramesDelegate = FOnGetZeroPad::CreateLambda([LocalSettings]()->uint8\
+	{
+		if (LocalSettings)
 		{
-			if (!MovieSceneCapture->bUseCustomStartFrame)
-			{
-				MovieSceneCapture->StartFrame = StartFrame;
-			}
+			return LocalSettings->GetZeroPadFrames();
+		}
+		return 0;
+	});
 
-			if (!MovieSceneCapture->bUseCustomEndFrame)
-			{
-				MovieSceneCapture->EndFrame = EndFrame;
-			}
+	// By using a FrameResolution/PlayRateResolution that match the numbers entered via the numeric interface don't change frames of reference.
+	// This is used here because the movie scene capture works entirely on play rate resolution and has no knowledge of the internal resolution
+	// so we don't need to convert the user's input into internal resolution.
+	auto GetFrameRateDelegate = FOnGetFrameRate::CreateLambda([MovieSceneCapture]() -> FFrameRate
+	{
+		if (MovieSceneCapture)
+		{
+			return MovieSceneCapture->GetSettings().FrameRate;
 		}
-	}
+		return FFrameRate(30, 1);
+	});
+
+	// Create our numeric type interface so we can pass it to the time slider below.
+	TSharedPtr<INumericTypeInterface<double>> MovieSceneCaptureNumericInterface = MakeShareable(new FFrameNumberInterface(GetDisplayFormatDelegate, GetZeroPadFramesDelegate, GetFrameRateDelegate, GetFrameRateDelegate));
 
-	IMovieSceneCaptureDialogModule::Get().OpenDialog(LevelEditorModule.GetLevelEditorTabManager().ToSharedRef(), MovieSceneCapture);
+	IMovieSceneCaptureDialogModule::Get().OpenDialog(LevelEditorModule.GetLevelEditorTabManager().ToSharedRef(), MovieSceneCapture, MovieSceneCaptureNumericInterface);
 }
 
 ISequencer::FOnActorAddedToSequencer& FSequencer::OnActorAddedToSequencer()
@@ -2136,16 +2246,11 @@ ISequencer::FOnCameraCut& FSequencer::OnCameraCut()
 	return OnCameraCutEvent;
 }
 
-TSharedRef<INumericTypeInterface<float>> FSequencer::GetNumericTypeInterface()
+TSharedRef<INumericTypeInterface<double>> FSequencer::GetNumericTypeInterface() const
 {
 	return SequencerWidget->GetNumericTypeInterface();
 }
 
-TSharedRef<INumericTypeInterface<float>> FSequencer::GetZeroPadNumericTypeInterface()
-{
-	return SequencerWidget->GetZeroPadNumericTypeInterface();
-}
-
 TSharedRef<SWidget> FSequencer::MakeTimeRange(const TSharedRef<SWidget>& InnerContent, bool bShowWorkingRange, bool bShowViewRange, bool bShowPlaybackRange)
 {
 	return SequencerWidget->MakeTimeRange(InnerContent, bShowWorkingRange, bShowViewRange, bShowPlaybackRange);
@@ -2285,6 +2390,10 @@ void FSequencer::PossessPIEViewports(UObject* CameraObject, UObject* UnlockIfCam
 		}
 
 		UCameraComponent* CameraComponent = MovieSceneHelpers::CameraComponentFromRuntimeObject(CameraObject);
+		if (CameraComponent && CameraComponent->GetOwner() != CameraObject)
+		{
+			CameraObject = CameraComponent->GetOwner();
+		}
 
 		if (CameraObject == ViewTarget)
 		{
@@ -2476,7 +2585,7 @@ void FSequencer::SetPlaybackStatus(EMovieScenePlayerStatus::Type InPlaybackStatu
 		ShuttleMultiplier = 0;
 	}
 
-	TimingManager->Update(PlaybackState, GetGlobalTime());
+	TimeController->PlayerStatusChanged(PlaybackState, GetGlobalTime());
 }
 
 
@@ -2626,7 +2735,6 @@ TSharedRef<SWidget> FSequencer::MakeTransportControls(bool bExtended)
 		TransportControlArgs.OnForwardStep.BindSP( this, &FSequencer::OnStepForward );
 		TransportControlArgs.OnForwardEnd.BindSP( this, &FSequencer::OnJumpToEnd );
 		TransportControlArgs.OnGetPlaybackMode.BindSP( this, &FSequencer::GetPlaybackMode );
-		TransportControlArgs.OnGetRecording.BindSP(this, &FSequencer::IsRecording);
 
 		if(bExtended)
 		{
@@ -2831,8 +2939,7 @@ UObject* FSequencer::FindSpawnedObjectOrTemplate(const FGuid& BindingId)
 
 FReply FSequencer::OnPlay(bool bTogglePlay)
 {
-	if( (PlaybackState == EMovieScenePlayerStatus::Playing ||
-		 PlaybackState == EMovieScenePlayerStatus::Recording) && bTogglePlay )
+	if( PlaybackState == EMovieScenePlayerStatus::Playing && bTogglePlay )
 	{
 		Pause();
 	}
@@ -2887,7 +2994,7 @@ FReply FSequencer::OnRecord()
 void FSequencer::HandleRecordingStarted(UMovieSceneSequence* Sequence)
 {
 	OnPlayForward(false);
-		
+
 	// Make sure Slate ticks during playback
 	SequencerWidget->RegisterActiveTimerForPlayback();
 
@@ -2895,10 +3002,11 @@ void FSequencer::HandleRecordingStarted(UMovieSceneSequence* Sequence)
 	UMovieSceneSubSection* Section = UMovieSceneSubSection::GetRecordingSection();
 	if(Section != nullptr)
 	{
-		float LocalTime = GetLocalTime();
+		FFrameRate   FrameResolution = GetFocusedFrameResolution();
+		FFrameNumber StartFrame      = GetLocalTime().ConvertTo(FrameResolution).FloorToFrame();
+		int32        Duration        = FFrameRate::TransformTime(1, GetFocusedPlayRate(), FrameResolution).CeilToFrame().Value;
 
-		Section->SetStartTime(LocalTime);
-		Section->SetEndTime(LocalTime + GetFixedFrameInterval());
+		Section->SetRange(TRange<FFrameNumber>(StartFrame, StartFrame + Duration));
 	}
 }
 
@@ -2917,8 +3025,10 @@ void FSequencer::HandleRecordingFinished(UMovieSceneSequence* Sequence)
 	{
 		Section->SetAsRecording(false);
 		Section->SetSequence(Sequence);
-		Section->SetEndTime(Section->GetStartTime() + Sequence->GetMovieScene()->GetPlaybackRange().Size<float>());
-	
+
+		FFrameNumber EndFrame = Section->GetInclusiveStartFrame() + MovieScene::DiscreteSize(Sequence->GetMovieScene()->GetPlaybackRange());
+		Section->SetRange(TRange<FFrameNumber>(Section->GetInclusiveStartFrame(), TRangeBound<FFrameNumber>::Exclusive(EndFrame)));
+
 		if (Section->IsA<UMovieSceneCinematicShotSection>())
 		{
 			const FMovieSceneSpawnable* SpawnedCamera = Sequence->GetMovieScene()->FindSpawnable(
@@ -2980,7 +3090,11 @@ FReply FSequencer::OnPlayBackward(bool bTogglePlay)
 FReply FSequencer::OnStepForward()
 {
 	SetPlaybackStatus(EMovieScenePlayerStatus::Stepping);
-	float NewPosition = GetLocalTime() + GetFixedFrameInterval();
+
+	FFrameRate          PlayRate    = GetFocusedPlayRate();
+	FQualifiedFrameTime CurrentTime = GetLocalTime();
+
+	FFrameTime NewPosition = FFrameRate::TransformTime(CurrentTime.ConvertTo(PlayRate).FloorToFrame() + 1, PlayRate, CurrentTime.Rate);
 	SetLocalTime(NewPosition, ESnapTimeMode::STM_Interval);
 	return FReply::Handled();
 }
@@ -2989,8 +3103,12 @@ FReply FSequencer::OnStepForward()
 FReply FSequencer::OnStepBackward()
 {
 	SetPlaybackStatus(EMovieScenePlayerStatus::Stepping);
-	float NewPosition = GetLocalTime() - GetFixedFrameInterval();
-	const bool bAllowSnappingToFrames = true;
+
+	FFrameRate          PlayRate    = GetFocusedPlayRate();
+	FQualifiedFrameTime CurrentTime = GetLocalTime();
+
+	FFrameTime NewPosition = FFrameRate::TransformTime(CurrentTime.ConvertTo(PlayRate).FloorToFrame() - 1, PlayRate, CurrentTime.Rate);
+
 	SetLocalTime(NewPosition, ESnapTimeMode::STM_Interval);
 	return FReply::Handled();
 }
@@ -3036,11 +3154,14 @@ FReply FSequencer::SetPlaybackEnd()
 	const UMovieSceneSequence* FocusedSequence = GetFocusedMovieSceneSequence();
 	if (FocusedSequence)
 	{
-		TRange<float> CurrentRange = FocusedSequence->GetMovieScene()->GetPlaybackRange();
-		float NewPos = FMath::Max(GetLocalTime(), CurrentRange.GetLowerBoundValue());
-		SetPlaybackRange(TRange<float>(CurrentRange.GetLowerBoundValue(), NewPos));
+		FFrameNumber         CurrentFrame = GetLocalTime().Time.FloorToFrame();
+		TRange<FFrameNumber> CurrentRange = FocusedSequence->GetMovieScene()->GetPlaybackRange();
+		if (CurrentFrame >= MovieScene::DiscreteInclusiveLower(CurrentRange))
+		{
+			CurrentRange.SetUpperBound(CurrentFrame);
+			SetPlaybackRange(CurrentRange);
+		}
 	}
-
 	return FReply::Handled();
 }
 
@@ -3049,40 +3170,56 @@ FReply FSequencer::SetPlaybackStart()
 	const UMovieSceneSequence* FocusedSequence = GetFocusedMovieSceneSequence();
 	if (FocusedSequence)
 	{
-		TRange<float> CurrentRange = FocusedSequence->GetMovieScene()->GetPlaybackRange();
-		float NewPos = FMath::Min(GetLocalTime(), CurrentRange.GetUpperBoundValue());
-		SetPlaybackRange(TRange<float>(NewPos, CurrentRange.GetUpperBoundValue()));	
+		FFrameNumber         CurrentFrame = GetLocalTime().Time.FloorToFrame();
+		TRange<FFrameNumber> CurrentRange = FocusedSequence->GetMovieScene()->GetPlaybackRange();
+		if (CurrentFrame < MovieScene::DiscreteExclusiveUpper(CurrentRange))
+		{
+			CurrentRange.SetLowerBound(CurrentFrame);
+			SetPlaybackRange(CurrentRange);
+		}
 	}
-
 	return FReply::Handled();
 }
 
 FReply FSequencer::JumpToPreviousKey()
 {
-	GetKeysFromSelection(SelectedKeyCollection, SMALL_NUMBER);
-	if (SelectedKeyCollection.IsValid())
+	if (Selection.GetSelectedOutlinerNodes().Num())
+	{
+		GetKeysFromSelection(SelectedKeyCollection, SMALL_NUMBER);
+	}
+	else
 	{
-		TRange<float> FindRange(TRange<float>::BoundsType(), GetLocalTime());
+		GetAllKeys(SelectedKeyCollection, SMALL_NUMBER);
+	}
 
-		TOptional<float> NewTime = SelectedKeyCollection->FindFirstKeyInRange(FindRange, EFindKeyDirection::Backwards);
+	if (SelectedKeyCollection.IsValid())
+	{
+		FFrameNumber FrameNumber = GetLocalTime().Time.FloorToFrame();
+		TOptional<FFrameNumber> NewTime = SelectedKeyCollection->GetNextKey(FrameNumber, EFindKeyDirection::Backwards);
 		if (NewTime.IsSet())
 		{
 			SetPlaybackStatus(EMovieScenePlayerStatus::Stepping);
 			SetLocalTimeDirectly(NewTime.GetValue());
 		}
 	}
-
 	return FReply::Handled();
 }
 
 FReply FSequencer::JumpToNextKey()
 {
-	GetKeysFromSelection(SelectedKeyCollection, SMALL_NUMBER);
-	if (SelectedKeyCollection.IsValid())
+	if (Selection.GetSelectedOutlinerNodes().Num())
+	{
+		GetKeysFromSelection(SelectedKeyCollection, SMALL_NUMBER);
+	}
+	else
 	{
-		TRange<float> FindRange(GetLocalTime(), TRange<float>::BoundsType());
+		GetAllKeys(SelectedKeyCollection, SMALL_NUMBER);
+	}
 
-		TOptional<float> NewTime = SelectedKeyCollection->FindFirstKeyInRange(FindRange, EFindKeyDirection::Forwards);
+	if (SelectedKeyCollection.IsValid())
+	{
+		FFrameNumber FrameNumber = GetLocalTime().Time.FloorToFrame();
+		TOptional<FFrameNumber> NewTime = SelectedKeyCollection->GetNextKey(FrameNumber, EFindKeyDirection::Forwards);
 		if (NewTime.IsSet())
 		{
 			SetPlaybackStatus(EMovieScenePlayerStatus::Stepping);
@@ -3099,13 +3236,19 @@ ESequencerLoopMode FSequencer::GetLoopMode() const
 }
 
 
-void FSequencer::SetLocalTimeLooped(float NewLocalTime)
+void FSequencer::SetLocalTimeLooped(FFrameTime NewLocalTime)
 {
 	TOptional<EMovieScenePlayerStatus::Type> NewPlaybackStatus;
 
-	float NewGlobalTime = NewLocalTime * RootToLocalTransform.Inverse();
+	FFrameTime NewGlobalTime = NewLocalTime * RootToLocalTransform.Inverse();
 
-	TRange<float> TimeBounds = GetTimeBounds();
+	TRange<FFrameNumber> TimeBounds = GetTimeBounds();
+
+	bool         bResetPosition       = false;
+	FFrameRate   LocalFrameResolution = GetFocusedFrameResolution();
+	FFrameRate   RootFrameResolution  = GetRootFrameResolution();
+	FFrameNumber MinInclusiveTime     = MovieScene::DiscreteInclusiveLower(TimeBounds);
+	FFrameNumber MaxInclusiveTime     = MovieScene::DiscreteExclusiveUpper(TimeBounds)-1;
 
 	bool bHasJumped = false;
 	bool bRestarted = false;
@@ -3114,75 +3257,74 @@ void FSequencer::SetLocalTimeLooped(float NewLocalTime)
 		const UMovieSceneSequence* FocusedSequence = GetFocusedMovieSceneSequence();
 		if (FocusedSequence)
 		{
-			if (NewLocalTime <= TimeBounds.GetLowerBoundValue() || NewLocalTime >= TimeBounds.GetUpperBoundValue())
+			if (NewLocalTime < MinInclusiveTime || NewLocalTime > MaxInclusiveTime)
 			{
-				NewGlobalTime = (PlaybackSpeed > 0 ? TimeBounds.GetLowerBoundValue() : TimeBounds.GetUpperBoundValue()) * RootToLocalTransform.Inverse();
-				TimingManager->OnStartPlaying(NewGlobalTime);
-
-				// Always evaluate from the start/end when looping
-				PlayPosition.Reset(NewGlobalTime);
+				NewGlobalTime = (PlaybackSpeed > 0 ? MinInclusiveTime : MaxInclusiveTime) * RootToLocalTransform.Inverse();
 
+				bResetPosition = true;
 				bHasJumped = true;
 			}
 		}
 	}
 	else
 	{
-		TRange<float> WorkingRange = GetClampRange();
+		TRange<double> WorkingRange = GetClampRange();
 
 		bool bReachedEnd = false;
 		if (PlaybackSpeed > 0)
 		{
-			bReachedEnd = GetLocalTime() < TimeBounds.GetUpperBoundValue() && NewLocalTime >= TimeBounds.GetUpperBoundValue();
+			bReachedEnd = GetLocalTime().Time < MaxInclusiveTime && NewLocalTime >= MaxInclusiveTime;
 		}
 		else
 		{
-			bReachedEnd = GetLocalTime() > TimeBounds.GetLowerBoundValue() && NewLocalTime <= TimeBounds.GetLowerBoundValue();
+			bReachedEnd = GetLocalTime().Time > MinInclusiveTime && NewLocalTime <= MinInclusiveTime;
 		}
 
 		// Stop if we hit the playback range end
 		if (bReachedEnd)
 		{
-			NewGlobalTime = (PlaybackSpeed > 0 ? TimeBounds.GetUpperBoundValue() : TimeBounds.GetLowerBoundValue()) * RootToLocalTransform.Inverse();
-			TimingManager->OnStartPlaying(NewGlobalTime);
+			NewGlobalTime = (PlaybackSpeed > 0 ? MaxInclusiveTime : MinInclusiveTime) * RootToLocalTransform.Inverse();
 			NewPlaybackStatus = EMovieScenePlayerStatus::Stopped;
 		}
 		// Constrain to the play range if necessary
 		else if (Settings->ShouldKeepCursorInPlayRange())
 		{
 			// Clamp to bound or jump back if necessary
-			if (NewLocalTime <= TimeBounds.GetLowerBoundValue() || 
-				NewLocalTime >= TimeBounds.GetUpperBoundValue())
+			if (NewLocalTime < MinInclusiveTime || NewLocalTime >= MaxInclusiveTime)
 			{
-				NewGlobalTime = (PlaybackSpeed > 0 ? TimeBounds.GetLowerBoundValue() : TimeBounds.GetUpperBoundValue()) * RootToLocalTransform.Inverse();
-				TimingManager->OnStartPlaying(NewGlobalTime);
+				NewGlobalTime = (PlaybackSpeed > 0 ? MinInclusiveTime : MaxInclusiveTime) * RootToLocalTransform.Inverse();
 
-				// Always evaluate from the start/end when looping
-				PlayPosition.Reset(NewGlobalTime);
+				bResetPosition = true;
 			}
 		}
 		// Ensure the time is within the working range
-		else if (!WorkingRange.Contains(NewLocalTime))
+		else if (!WorkingRange.Contains(NewLocalTime / LocalFrameResolution))
 		{
-			NewGlobalTime = FMath::Clamp(NewLocalTime, WorkingRange.GetLowerBoundValue(), WorkingRange.GetUpperBoundValue()) * RootToLocalTransform.Inverse();
-			TimingManager->OnStartPlaying(NewGlobalTime);
+			FFrameTime WorkingMin = (WorkingRange.GetLowerBoundValue() * LocalFrameResolution).CeilToFrame();
+			FFrameTime WorkingMax = (WorkingRange.GetUpperBoundValue() * LocalFrameResolution).FloorToFrame();
 
-			// Always evaluate from the start/end when looping
-			PlayPosition.Reset(NewGlobalTime);
+			NewGlobalTime = FMath::Clamp(NewLocalTime, WorkingMin, WorkingMax) * RootToLocalTransform.Inverse();
+
+			bResetPosition = true;
 			NewPlaybackStatus = EMovieScenePlayerStatus::Stopped;
 		}
 	}
 
-	// Ensure the time is in the current view
-	ScrollIntoView(NewGlobalTime * RootToLocalTransform);
+	// Ensure the time is in the current view - must occur before the time cursor changes
+	UMovieScene* MovieScene = GetFocusedMovieSceneSequence()->GetMovieScene();
+	ScrollIntoView( (NewGlobalTime * RootToLocalTransform) / RootFrameResolution );
 
-	// Update the position before fixing it to the time interval
-	ScrubPosition = NewGlobalTime;
+	FFrameTime NewPlayPosition = ConvertFrameTime(NewGlobalTime, RootFrameResolution, PlayPosition.GetInputRate());
 
-	// Evaluate the sequence
-	UMovieScene* MovieScene = GetFocusedMovieSceneSequence()->GetMovieScene();
-	FMovieSceneEvaluationRange EvalRange = PlayPosition.PlayTo(ScrubPosition, MovieScene->GetOptionalFixedFrameInterval());
+	// Reset the play cursor if we're looping or have otherwise jumpted to a new position in the sequence
+	if (bResetPosition)
+	{
+		PlayPosition.Reset(NewPlayPosition);
+		TimeController->Reset(FQualifiedFrameTime(NewGlobalTime, RootFrameResolution));
+	}
 
+	// Evaluate the sequence
+	FMovieSceneEvaluationRange EvalRange = PlayPosition.PlayTo(NewPlayPosition);
 	EvaluateInternal(EvalRange, bHasJumped);
 
 	// Set the playback status if we need to
@@ -3194,13 +3336,6 @@ void FSequencer::SetLocalTimeLooped(float NewLocalTime)
 	}
 }
 
-
-bool FSequencer::CanShowFrameNumbers() const
-{
-	return SequencerSnapValues::IsTimeSnapIntervalFrameRate(GetFixedFrameInterval());
-}
-
-
 EPlaybackMode::Type FSequencer::GetPlaybackMode() const
 {
 	if (PlaybackState == EMovieScenePlayerStatus::Playing)
@@ -3218,51 +3353,38 @@ EPlaybackMode::Type FSequencer::GetPlaybackMode() const
 	return EPlaybackMode::Stopped;
 }
 
-bool FSequencer::IsRecording() const
-{
-	return PlaybackState == EMovieScenePlayerStatus::Recording;
-}
-
 void FSequencer::UpdateTimeBoundsToFocusedMovieScene()
 {
 	UMovieScene* FocusedMovieScene = GetFocusedMovieSceneSequence()->GetMovieScene();
+	FQualifiedFrameTime CurrentTime = GetLocalTime();
 
 	// Set the view range to:
 	// 1. The moviescene view range
 	// 2. The moviescene playback range
 	// 3. Some sensible default
-	TRange<float> NewRange = FocusedMovieScene->GetEditorData().ViewRange;
+	TRange<double> NewRange = FocusedMovieScene->GetEditorData().GetViewRange();
 
 	if (NewRange.IsEmpty() || NewRange.IsDegenerate())
 	{
-		NewRange = FocusedMovieScene->GetPlaybackRange();
+		NewRange = FocusedMovieScene->GetPlaybackRange() / CurrentTime.Rate;
 	}
 	if (NewRange.IsEmpty() || NewRange.IsDegenerate())
 	{
-		NewRange = TRange<float>(0.f, 5.f);
+		NewRange = TRange<double>(0.0, 5.0);
 	}
 
 	// Set the view range to the new range
 	SetViewRange(NewRange, EViewRangeInterpolation::Immediate);
-
-	// Make sure the current time is within the bounds
-	if (!TargetViewRange.Contains(GetLocalTime()))
-	{
-		SetLocalTimeDirectly(LastViewRange.GetLowerBoundValue());
-		OnGlobalTimeChangedDelegate.Broadcast();
-	}
 }
 
 
-TRange<float> FSequencer::GetTimeBounds() const
+TRange<FFrameNumber> FSequencer::GetTimeBounds() const
 {
 	const UMovieSceneSequence* FocusedSequence = GetFocusedMovieSceneSequence();
 
-	// When recording, we never want to constrain the time bound range.  You might not even have any sections or keys yet
-	// but we need to be able to move the time cursor during playback so you can capture data in real-time
-	if( PlaybackState == EMovieScenePlayerStatus::Recording || !FocusedSequence)
+	if(!FocusedSequence)
 	{
-		return TRange<float>( -100000.0f, 100000.0f );
+		return TRange<FFrameNumber>( -100000, 100000 );
 	}
 	
 	if (GetLoopMode() == ESequencerLoopMode::SLM_LoopSelectionRange)
@@ -3282,7 +3404,7 @@ TRange<float> FSequencer::GetTimeBounds() const
 }
 
 
-void FSequencer::SetViewRange(TRange<float> NewViewRange, EViewRangeInterpolation Interpolation)
+void FSequencer::SetViewRange(TRange<double> NewViewRange, EViewRangeInterpolation Interpolation)
 {
 	if (!ensure(NewViewRange.HasUpperBound() && NewViewRange.HasLowerBound() && !NewViewRange.IsDegenerate()))
 	{
@@ -3318,40 +3440,41 @@ void FSequencer::SetViewRange(TRange<float> NewViewRange, EViewRangeInterpolatio
 		UMovieScene* FocusedMovieScene = FocusedMovieSequence->GetMovieScene();
 		if (FocusedMovieScene != nullptr)
 		{
-			FocusedMovieScene->GetEditorData().ViewRange = TargetViewRange;
+			FMovieSceneEditorData& EditorData = FocusedMovieScene->GetEditorData();
+			EditorData.ViewStart = TargetViewRange.GetLowerBoundValue();
+			EditorData.ViewEnd   = TargetViewRange.GetUpperBoundValue();
 
 			// Always ensure the working range is big enough to fit the view range
-			TRange<float>& WorkingRange = FocusedMovieScene->GetEditorData().WorkingRange;
-
-			WorkingRange = TRange<float>(
-				FMath::Min(TargetViewRange.GetLowerBoundValue(), WorkingRange.GetLowerBoundValue()),
-				FMath::Max(TargetViewRange.GetUpperBoundValue(), WorkingRange.GetUpperBoundValue())
-				);
+			EditorData.WorkStart = FMath::Min(TargetViewRange.GetLowerBoundValue(), EditorData.WorkStart);
+			EditorData.WorkEnd   = FMath::Max(TargetViewRange.GetUpperBoundValue(), EditorData.WorkEnd);
 		}
 	}
 }
 
 
-void FSequencer::OnClampRangeChanged( TRange<float> NewClampRange )
+void FSequencer::OnClampRangeChanged( TRange<double> NewClampRange )
 {
 	if (!NewClampRange.IsEmpty())
 	{
-		GetFocusedMovieSceneSequence()->GetMovieScene()->GetEditorData().WorkingRange = NewClampRange;
+		FMovieSceneEditorData& EditorData =  GetFocusedMovieSceneSequence()->GetMovieScene()->GetEditorData();
+
+		EditorData.WorkStart = NewClampRange.GetLowerBoundValue();
+		EditorData.WorkEnd   = NewClampRange.GetUpperBoundValue();
 	}
 }
 
-float FSequencer::OnGetNearestKey(float InTime)
+FFrameNumber FSequencer::OnGetNearestKey(FFrameTime InTime)
 {
-	float NearestKeyTime = InTime;
+	FFrameNumber NearestKeyTime = InTime.FloorToFrame();
 	GetKeysFromSelection(SelectedKeyCollection, SMALL_NUMBER);
 
 	if (SelectedKeyCollection.IsValid())
 	{
-		TRange<float> FindRangeBackwards(TRange<float>::BoundsType(), NearestKeyTime);
-		TOptional<float> NewTimeBackwards = SelectedKeyCollection->FindFirstKeyInRange(FindRangeBackwards, EFindKeyDirection::Backwards);
+		TRange<FFrameNumber> FindRangeBackwards(TRangeBound<FFrameNumber>::Open(), NearestKeyTime);
+		TOptional<FFrameNumber> NewTimeBackwards = SelectedKeyCollection->FindFirstKeyInRange(FindRangeBackwards, EFindKeyDirection::Backwards);
 
-		TRange<float> FindRangeForwards(NearestKeyTime, TRange<float>::BoundsType());
-		TOptional<float> NewTimeForwards = SelectedKeyCollection->FindFirstKeyInRange(FindRangeForwards, EFindKeyDirection::Forwards);
+		TRange<FFrameNumber> FindRangeForwards(NearestKeyTime, TRangeBound<FFrameNumber>::Open());
+		TOptional<FFrameNumber> NewTimeForwards = SelectedKeyCollection->FindFirstKeyInRange(FindRangeForwards, EFindKeyDirection::Forwards);
 		if (NewTimeForwards.IsSet())
 		{
 			if (NewTimeBackwards.IsSet())
@@ -3378,7 +3501,7 @@ float FSequencer::OnGetNearestKey(float InTime)
 	return NearestKeyTime;
 }
 
-void FSequencer::OnScrubPositionChanged( float NewScrubPosition, bool bScrubbing )
+void FSequencer::OnScrubPositionChanged( FFrameTime NewScrubPosition, bool bScrubbing )
 {
 	bool bClampToViewRange = true;
 
@@ -3393,7 +3516,7 @@ void FSequencer::OnScrubPositionChanged( float NewScrubPosition, bool bScrubbing
 			// Clamp to the view range when not auto-scrolling
 			bClampToViewRange = false;
 	
-			UpdateAutoScroll(NewScrubPosition);
+			UpdateAutoScroll(NewScrubPosition / GetFocusedFrameResolution());
 			
 			// When scrubbing, we animate auto-scrolled scrub position in Tick()
 			if (AutoscrubOffset.IsSet())
@@ -3405,16 +3528,19 @@ void FSequencer::OnScrubPositionChanged( float NewScrubPosition, bool bScrubbing
 
 	if (bClampToViewRange)
 	{
-		float LowerBound = TargetViewRange.GetLowerBoundValue();
-		float UpperBound = TargetViewRange.GetUpperBoundValue();
+		FFrameRate PlayRate        = GetFocusedPlayRate();
+		FFrameRate FrameResolution = GetFocusedFrameResolution();
+
+		FFrameTime LowerBound = (TargetViewRange.GetLowerBoundValue() * FrameResolution).CeilToFrame();
+		FFrameTime UpperBound = (TargetViewRange.GetUpperBoundValue() * FrameResolution).FloorToFrame();
 
 		if (Settings->GetIsSnapEnabled() && Settings->GetSnapPlayTimeToInterval())
 		{
-			LowerBound = SequencerHelpers::SnapTimeToInterval(LowerBound, GetFixedFrameInterval());
-			UpperBound = SequencerHelpers::SnapTimeToInterval(UpperBound, GetFixedFrameInterval());
+			LowerBound = FFrameRate::Snap(LowerBound, FrameResolution, PlayRate);
+			UpperBound = FFrameRate::Snap(UpperBound, FrameResolution, PlayRate);
 		}
 
-		NewScrubPosition = FMath::Clamp(NewScrubPosition, LowerBound, UpperBound);		
+		NewScrubPosition = FMath::Clamp(NewScrubPosition, LowerBound, UpperBound);
 	}
 
 	SetLocalTimeDirectly( NewScrubPosition );
@@ -3473,6 +3599,7 @@ void FSequencer::StartAutoscroll(float UnitsPerS)
 void FSequencer::StopAutoscroll()
 {
 	AutoscrollOffset.Reset();
+	AutoscrubOffset.Reset();
 }
 
 
@@ -3608,8 +3735,13 @@ void FSequencer::AddSubSequence(UMovieSceneSequence* Sequence)
 	OwnerMovieScene->Modify();
 
 	UMovieSceneSubTrack* SubTrack = OwnerMovieScene->AddMasterTrack<UMovieSceneSubTrack>();
-	float Duration = Sequence->GetMovieScene()->GetPlaybackRange().Size<float>();
-	SubTrack->AddSequence(Sequence, ScrubPosition, Duration);
+
+	FFrameNumber Duration = ConvertFrameTime(
+		Sequence->GetMovieScene()->GetPlaybackRange().Size<FFrameNumber>(),
+		Sequence->GetMovieScene()->GetFrameResolution(),
+		OwnerMovieScene->GetFrameResolution()).FloorToFrame();
+
+	SubTrack->AddSequence(Sequence, GetLocalTime().Time.FloorToFrame(), Duration.Value);
 }
 
 
@@ -3810,6 +3942,7 @@ void FSequencer::PostUndo(bool bSuccess)
 {
 	NotifyMovieSceneDataChanged( EMovieSceneDataChangeType::Unknown );
 	SynchronizeSequencerSelectionWithExternalSelection();
+	SyncCurveEditorToSelection(false);
 	OnActivateSequenceEvent.Broadcast(ActiveTemplateIDs.Top());
 }
 
@@ -3902,8 +4035,12 @@ void FSequencer::OnNewActorsDropped(const TArray<UObject*>& DroppedObjects, cons
 
 					// Create an attach track
 					UMovieScene3DAttachTrack* AttachTrack = Cast<UMovieScene3DAttachTrack>(OwnerMovieScene->AddTrack(UMovieScene3DAttachTrack::StaticClass(), NewCameraGuid));
+
 					FMovieSceneObjectBindingID AttachBindingID(NewGuid, MovieSceneSequenceID::Root);
-					AttachTrack->AddConstraint(GetPlaybackRange().GetLowerBoundValue(), GetPlaybackRange().GetUpperBoundValue(), NAME_None, NAME_None, AttachBindingID);
+					FFrameNumber StartTime = MovieScene::DiscreteInclusiveLower(GetPlaybackRange());
+					FFrameNumber Duration  = MovieScene::DiscreteSize(GetPlaybackRange());
+
+					AttachTrack->AddConstraint(StartTime, Duration.Value, NAME_None, NAME_None, AttachBindingID);
 				}
 				else
 				{
@@ -3933,13 +4070,13 @@ void FSequencer::OnNewActorsDropped(const TArray<UObject*>& DroppedObjects, cons
 
 					FKeyPropertyParams KeyPropertyParams(TArrayBuilder<UObject*>().Add(RailActor), PropertyPath, ESequencerKeyMode::ManualKeyForced);
 
-					float OriginalTime = GetLocalTime();
+					FFrameTime OriginalTime = GetLocalTime().Time;
 
-					SetLocalTimeDirectly(GetPlaybackRange().GetLowerBoundValue());
+					SetLocalTimeDirectly(MovieScene::DiscreteInclusiveLower(GetPlaybackRange()));
 					RailActor->CurrentPositionOnRail = 0.f;
 					KeyProperty(KeyPropertyParams);
 
-					SetLocalTimeDirectly(GetPlaybackRange().GetUpperBoundValue());
+					SetLocalTimeDirectly(MovieScene::DiscreteExclusiveUpper(GetPlaybackRange())-1);
 					RailActor->CurrentPositionOnRail = 1.f;
 					KeyProperty(KeyPropertyParams);
 
@@ -4133,8 +4270,7 @@ void FSequencer::SaveCurrentMovieScene()
 
 	FEditorFileUtils::PromptForCheckoutAndSave( PackagesToSave, bCheckDirty, bPromptToSave );
 
-	FMovieSceneEvaluationRange Range = PlayPosition.JumpTo(ScrubPosition, GetFocusedMovieSceneSequence()->GetMovieScene()->GetOptionalFixedFrameInterval());
-	EvaluateInternal(Range);
+	ForceEvaluate();
 
 	OnPostSaveEvent.Broadcast(*this);
 }
@@ -4171,7 +4307,6 @@ void FSequencer::SaveCurrentMovieSceneAs()
 
 TArray<FGuid> FSequencer::AddActors(const TArray<TWeakObjectPtr<AActor> >& InActors)
 {
-	// @todo sequencer: Undo doesn't seem to be working at all
 	const FScopedTransaction Transaction(LOCTEXT("UndoPossessingObject", "Possess Object in Sequencer"));
 	GetFocusedMovieSceneSequence()->Modify();
 
@@ -4200,6 +4335,71 @@ TArray<FGuid> FSequencer::AddActors(const TArray<TWeakObjectPtr<AActor> >& InAct
 
 	if (bPossessableAdded)
 	{
+		// Check if a folder is selected so we can add the actors to the selected folder.
+		TArray<UMovieSceneFolder*> SelectedParentFolders;
+		FString NewNodePath;
+		if (Selection.GetSelectedOutlinerNodes().Num() > 0)
+		{
+			for (TSharedRef<FSequencerDisplayNode> SelectedNode : Selection.GetSelectedOutlinerNodes())
+			{
+				TSharedPtr<FSequencerDisplayNode> CurrentNode = SelectedNode;
+				while (CurrentNode.IsValid() && CurrentNode->GetType() != ESequencerNode::Folder)
+				{
+					CurrentNode = CurrentNode->GetParent();
+				}
+				if (CurrentNode.IsValid())
+				{
+					SelectedParentFolders.Add(&StaticCastSharedPtr<FSequencerFolderNode>(CurrentNode)->GetFolder());
+
+					// The first valid folder we find will be used to put the new actors into, so it's the node that we
+					// want to know the path from.
+					if (NewNodePath.Len() == 0)
+					{
+						// Add an extra delimiter (".") as we know that the new objects will be appended onto the end of this.
+						NewNodePath = FString::Printf(TEXT("%s."), *CurrentNode->GetPathName());
+
+						// Make sure the folder is expanded too so that adding objects to hidden folders become visible.
+						CurrentNode->SetExpansionState(true);
+					}
+				}
+			}
+		}
+
+		// Clear our editor selection so we can make the selection our added actors.
+		// This has to be done after we know if the actor is going to be added to a
+		// folder, otherwise it causes the folder we wanted to pick to be deselected.
+		USelection* SelectedActors = GEditor->GetSelectedActors();
+		SelectedActors->BeginBatchSelectOperation();
+		SelectedActors->Modify();
+		GEditor->SelectNone(false, true);
+		for (TWeakObjectPtr<AActor> WeakActor : InActors)
+		{
+			if (AActor* Actor = WeakActor.Get())
+			{
+				GEditor->SelectActor(Actor, true, false);
+			}
+		}
+		SelectedActors->EndBatchSelectOperation();
+		GEditor->NoteSelectionChange();
+
+		// Add the possessables as children of the first selected folder
+		if (SelectedParentFolders.Num() > 0)
+		{
+			for (const FGuid& Possessable : PossessableGuids)
+			{
+				SelectedParentFolders[0]->AddChildObjectBinding(Possessable);
+			}
+		}
+
+		// Now add them all to the selection set to be selected after a tree rebuild.
+		for(const FGuid& Possessable : PossessableGuids)
+		{
+			FString PossessablePath = NewNodePath += Possessable.ToString();
+			
+			// Object Bindings use their FGuid as their unique key.
+			SequencerWidget->AddAdditionalPathToSelectionSet(PossessablePath);
+		}
+
 		RefreshTree();
 
 		SynchronizeSequencerSelectionWithExternalSelection();
@@ -4212,6 +4412,7 @@ TArray<FGuid> FSequencer::AddActors(const TArray<TWeakObjectPtr<AActor> >& InAct
 void FSequencer::OnSelectedOutlinerNodesChanged()
 {
 	SynchronizeExternalSelectionWithSequencerSelection();
+	SyncCurveEditorToSelection(true);
 
 	FSequencerEdMode* SequencerEdMode = (FSequencerEdMode*)(GLevelEditorModeTools().GetActiveMode(FSequencerEdMode::EM_SequencerMode));
 	if (SequencerEdMode != nullptr)
@@ -4487,14 +4688,172 @@ void FSequencer::OnSelectedNodesOnlyChanged()
 	SynchronizeSequencerSelectionWithExternalSelection();
 }
 
+void GatherKeyAreas(const TSet<TSharedRef<FSequencerDisplayNode>>& Selection, ECurveEditorCurveVisibility CurveVisibility, bool bAddChildren, TSharedRef<FSequencerDisplayNode> InNode, TSet<TSharedPtr<IKeyArea>>& OutKeyAreasToShow)
+{
+	// If we're only adding selected curves, and we've encountered a selected node, add all its child key areas
+	if (CurveVisibility == ECurveEditorCurveVisibility::SelectedCurves && Selection.Contains(InNode))
+	{
+		bAddChildren = true;
+	}
+
+	if (bAddChildren)
+	{
+		TSharedPtr<FSequencerSectionKeyAreaNode> KeyAreaNode;
+		if (InNode->GetType() == ESequencerNode::Track)
+		{
+			KeyAreaNode = StaticCastSharedRef<FSequencerTrackNode>(InNode)->GetTopLevelKeyNode();
+		}
+		else if (InNode->GetType() == ESequencerNode::KeyArea)
+		{
+			KeyAreaNode = StaticCastSharedRef<FSequencerSectionKeyAreaNode>(InNode);
+		}
+
+		if (KeyAreaNode.IsValid())
+		{
+			for (TSharedPtr<IKeyArea> KeyArea : KeyAreaNode->GetAllKeyAreas())
+			{
+				OutKeyAreasToShow.Add(KeyArea);
+			}
+		}
+	}
+
+	for (TSharedRef<FSequencerDisplayNode> Child : InNode->GetChildNodes())
+	{
+		GatherKeyAreas(Selection, CurveVisibility, bAddChildren, Child, OutKeyAreasToShow);
+	}
+}
+
+void FSequencer::SyncCurveEditorToSelection(bool bOutlinerSelectionChanged)
+{
+	if (!GetShowCurveEditor())
+	{
+		return;
+	}
+
+	ECurveEditorCurveVisibility CurveVisibility = Settings->GetCurveVisibility();
+	const TSet<TSharedRef<FSequencerDisplayNode>>& OutlinerSelection = Selection.GetSelectedOutlinerNodes();
+
+	// Traverse the tree to gather key areas
+	TSet<TSharedPtr<IKeyArea>> KeyAreasToShow;
+	for (const TSharedRef<FSequencerDisplayNode>& Node : NodeTree->GetRootNodes())
+	{
+		bool bAddChildren = CurveVisibility != ECurveEditorCurveVisibility::SelectedCurves;
+		GatherKeyAreas(OutlinerSelection, CurveVisibility, bAddChildren, Node, KeyAreasToShow);
+	}
+
+	// Remove unanimated curves if necessary
+	if (CurveVisibility == ECurveEditorCurveVisibility::AnimatedCurves)
+	{
+		TArray<TSharedPtr<IKeyArea>> UnanimatedCurves;
+		for (TSharedPtr<IKeyArea> KeyArea : KeyAreasToShow)
+		{
+			if (!KeyArea->HasAnyKeys())
+			{
+				UnanimatedCurves.Add(KeyArea);
+			}
+		}
+
+		for (TSharedPtr<IKeyArea> Unanimated : UnanimatedCurves)
+		{
+			KeyAreasToShow.Remove(Unanimated);
+		}
+	}
+
+	// Cache the curve editor's current selection
+	TMap<const void*, FCurveModelID> ExistingCurveIDs;
+	for (const TTuple<FCurveModelID, TUniquePtr<FCurveModel>>& Pair : CurveEditorModel->GetCurves())
+	{
+		if (const void* Ptr = Pair.Value->GetCurve())
+		{
+			ExistingCurveIDs.Add(Ptr, Pair.Key);
+		}
+	}
+
+	bool bAnythingChanged = false;
+
+	// Add newly selected curves to the curve editor
+	for (TSharedPtr<IKeyArea> KeyArea : KeyAreasToShow)
+	{
+		const void* Channel = KeyArea->GetChannelPtr();
+		if (!Channel)
+		{
+			continue;
+		}
+
+		if (ExistingCurveIDs.Contains(Channel))
+		{
+			ExistingCurveIDs.Remove(Channel);
+			continue;
+		}
+		else
+		{
+			TUniquePtr<FCurveModel> NewCurve = KeyArea->CreateCurveEditorModel(AsShared());
+			if (NewCurve.IsValid())
+			{
+				bAnythingChanged = true;
+				CurveEditorModel->AddCurve(MoveTemp(NewCurve));
+			}
+		}
+	}
+
+	// Remove anything that's no longer selected
+	for (const TTuple<const void*, FCurveModelID>& Pair : ExistingCurveIDs)
+	{
+		bAnythingChanged = true;
+		CurveEditorModel->RemoveCurve(Pair.Value);
+	}
+
+	if (bAnythingChanged && CurveEditorModel->ShouldAutoFrame())
+	{
+		CurveEditorModel->ZoomToFit();
+	}
+	else if (bOutlinerSelectionChanged && CurveEditorModel->ShouldAutoFrame())
+	{
+		// If outliner selection changes, zoom to fit only the selected curves
+		TSet<TSharedPtr<IKeyArea>> SelectedKeyAreasToShow;
+		for (const TSharedRef<FSequencerDisplayNode>& Node : NodeTree->GetRootNodes())
+		{
+			GatherKeyAreas(OutlinerSelection, ECurveEditorCurveVisibility::SelectedCurves, false, Node, SelectedKeyAreasToShow);
+		}
+	
+		TArray<FCurveModelID> CurveModelIDs;
+		for (const TTuple<FCurveModelID, TUniquePtr<FCurveModel>>& Pair : CurveEditorModel->GetCurves())
+		{
+			if (const void* Ptr = Pair.Value->GetCurve())
+			{
+				for (TSharedPtr<IKeyArea> KeyArea : SelectedKeyAreasToShow)
+				{
+					const void* Channel = KeyArea->GetChannelPtr();
+					if (!Channel)
+					{
+						continue;
+					}
+
+					if (Ptr == Channel)
+					{
+						CurveModelIDs.Add(Pair.Key);
+						break;
+					}
+				}
+			}
+		}
+
+		if (CurveModelIDs.Num())
+		{
+			CurveEditorModel->ZoomToFitCurves(MakeArrayView(CurveModelIDs.GetData(), CurveModelIDs.Num()));		
+		}
+	}
+}
+
 void FSequencer::ZoomToSelectedSections()
 {
-	TArray< TRange<float> > Bounds;
+	FFrameRate FrameResolution = GetFocusedFrameResolution();
+
+	TRange<FFrameNumber> BoundsHull = TRange<FFrameNumber>::Empty();
 	for (TWeakObjectPtr<UMovieSceneSection> SelectedSection : Selection.GetSelectedSections())
 	{
-		Bounds.Add(SelectedSection->GetRange());
+		BoundsHull = TRange<FFrameNumber>::Hull(SelectedSection->GetRange(), BoundsHull);
 	}
-	TRange<float> BoundsHull = TRange<float>::Hull(Bounds);
 
 	if (BoundsHull.IsEmpty())
 	{
@@ -4503,10 +4862,12 @@ void FSequencer::ZoomToSelectedSections()
 
 	if (!BoundsHull.IsEmpty() && !BoundsHull.IsDegenerate())
 	{
+		const double Tolerance = KINDA_SMALL_NUMBER;
+
 		// Zoom back to last view range if already expanded
 		if (!ViewRangeBeforeZoom.IsEmpty() &&
-			FMath::IsNearlyEqual(BoundsHull.GetLowerBoundValue(), GetViewRange().GetLowerBoundValue(), KINDA_SMALL_NUMBER) &&
-			FMath::IsNearlyEqual(BoundsHull.GetUpperBoundValue(), GetViewRange().GetUpperBoundValue(), KINDA_SMALL_NUMBER))
+			FMath::IsNearlyEqual(BoundsHull.GetLowerBoundValue() / FrameResolution, GetViewRange().GetLowerBoundValue(), Tolerance) &&
+			FMath::IsNearlyEqual(BoundsHull.GetUpperBoundValue() / FrameResolution, GetViewRange().GetUpperBoundValue(), Tolerance))
 		{
 			SetViewRange(ViewRangeBeforeZoom, EViewRangeInterpolation::Animated);
 		}
@@ -4514,7 +4875,7 @@ void FSequencer::ZoomToSelectedSections()
 		{
 			ViewRangeBeforeZoom = GetViewRange();
 
-			SetViewRange(BoundsHull, EViewRangeInterpolation::Animated);
+			SetViewRange(BoundsHull / FrameResolution, EViewRangeInterpolation::Animated);
 		}
 	}
 }
@@ -4591,6 +4952,20 @@ void FSequencer::SelectSection(UMovieSceneSection* Section)
 	Selection.AddToSelection(Section);
 }
 
+void FSequencer::SelectKey(UMovieSceneSection* InSection, TSharedPtr<IKeyArea> KeyArea, FKeyHandle KeyHandle, bool bToggle)
+{
+	FSequencerSelectedKey SelectedKey(*InSection, KeyArea, KeyHandle);
+
+	if (bToggle && Selection.IsSelected(SelectedKey))
+	{
+		Selection.RemoveFromSelection(SelectedKey);
+	}
+	else
+	{
+		Selection.AddToSelection(SelectedKey);
+	}
+}
+
 void FSequencer::SelectByPropertyPaths(const TArray<FString>& InPropertyPaths)
 {
 	TArray<TSharedRef<FSequencerDisplayNode>> NodesToSelect;
@@ -4627,6 +5002,16 @@ void FSequencer::EmptySelection()
 	Selection.Empty();
 }
 
+void FSequencer::ThrobKeySelection()
+{
+	SSequencerSection::ThrobKeySelection();
+}
+
+void FSequencer::ThrobSectionSelection()
+{
+	SSequencerSection::ThrobSectionSelection();
+}
+
 float FSequencer::GetOverlayFadeCurve() const
 {
 	return OverlayCurve.GetLerp();
@@ -4946,15 +5331,21 @@ void ExportObjectsToText(const TArray<UMovieSceneCopyableBinding*>& ObjectsToExp
 class FObjectBindingTextFactory : public FCustomizableTextObjectFactory
 {
 public:
-	FObjectBindingTextFactory()
+	FObjectBindingTextFactory(FSequencer& InSequencer)
 		: FCustomizableTextObjectFactory(GWarn)
+		, SequencerPtr(&InSequencer)
 	{
 	}
 
 	// FCustomizableTextObjectFactory implementation
 	virtual bool CanCreateClass(UClass* InObjectClass, bool& bOmitSubObjs) const override
 	{
-		return true;
+		if (InObjectClass->IsChildOf<UMovieSceneCopyableBinding>())
+		{
+			return true;
+		}
+
+		return SequencerPtr->GetSpawnRegister().CanSpawnObject(InObjectClass);
 	}
 	
 
@@ -4976,16 +5367,19 @@ public:
 public:
 	TArray<UMovieSceneCopyableBinding*> NewCopyableBindings;
 	TArray<UObject*> NewSpawnableObjectTemplates;
+
+private:
+	FSequencer* SequencerPtr;
 };
 
 
-void ImportObjectsFromText(const FString& TextToImport, /*out*/ TArray<UMovieSceneCopyableBinding*>& ImportedObjects)
+void ImportObjectsFromText(FSequencer& InSequencer, const FString& TextToImport, /*out*/ TArray<UMovieSceneCopyableBinding*>& ImportedObjects)
 {
 	UPackage* TempPackage = NewObject<UPackage>(nullptr, TEXT("/Engine/Sequencer/Editor/Transient"), RF_Transient);
 	TempPackage->AddToRoot();
 
 	// Turn the text buffer into objects
-	FObjectBindingTextFactory Factory;
+	FObjectBindingTextFactory Factory(InSequencer);
 	Factory.ProcessBuffer(TempPackage, RF_Transactional, TextToImport);
 	ImportedObjects = Factory.NewCopyableBindings;
 
@@ -5118,7 +5512,7 @@ void FSequencer::PasteCopiedTracks()
 	UObject* BindingContext = GetPlaybackContext();
 
 	TArray<UMovieSceneCopyableBinding*> ImportedBindings;
-	ImportObjectsFromText(TextToImport, ImportedBindings);
+	ImportObjectsFromText(*this, TextToImport, ImportedBindings);
 
 	if (ImportedBindings.Num() != 0)
 	{
@@ -5341,9 +5735,9 @@ public:
 	TArray<UMovieSceneTrack*> NewTracks;
 };
 
-bool FSequencer::CanPaste(const FString& TextToImport) const
+bool FSequencer::CanPaste(const FString& TextToImport)
 {
-	FObjectBindingTextFactory ObjectBindingFactory;
+	FObjectBindingTextFactory ObjectBindingFactory(*this);
 	if (ObjectBindingFactory.CanCreateObjectsFromText(TextToImport))
 	{
 		return true;
@@ -5395,18 +5789,21 @@ void FSequencer::ToggleNodeActive()
 
 bool FSequencer::IsNodeActive() const
 {
-	// Active only if all are active
+	// Active if ONE is active, changed in 4.20
 	for (auto OutlinerNode : Selection.GetSelectedOutlinerNodes())
 	{
 		TSet<TWeakObjectPtr<UMovieSceneSection> > Sections;
 		SequencerHelpers::GetAllSections(OutlinerNode, Sections);
-
-		for (auto Section : Sections)
+		if (Sections.Num() > 0)
 		{
-			if (!Section->IsActive())
+			for (auto Section : Sections)
 			{
-				return false;
+				if (Section->IsActive())
+				{
+					return true;
+				}
 			}
+			return false;
 		}
 	}
 	return true;
@@ -5646,11 +6043,11 @@ FMovieSceneSpawnable* FSequencer::ConvertToSpawnableInternal(FGuid PossessableGu
 		}
 	}
 
-	FTransformData TransformData;
+	TOptional<FTransformData> TransformData;
 	SpawnRegister->HandleConvertPossessableToSpawnable(FoundObject, *this, TransformData);
 	SpawnRegister->SetupDefaultsForSpawnable(nullptr, Spawnable->GetGuid(), TransformData, AsShared(), Settings);
 
-	SetLocalTimeDirectly(ScrubPosition);
+	ForceEvaluate();
 
 	return Spawnable;
 }
@@ -5775,9 +6172,11 @@ FMovieScenePossessable* FSequencer::ConvertToPossessableInternal(FGuid Spawnable
 	}
 
 	FTransform SpawnTransform = SpawnableActorTemplate->GetActorTransform();
+	FActorSpawnParameters SpawnInfo;
+	SpawnInfo.bDeferConstruction = true;
 
 	UWorld* PlaybackContext = Cast<UWorld>(GetPlaybackContext());
-	AActor* PossessedActor = PlaybackContext->SpawnActor(Spawnable->GetObjectTemplate()->GetClass(), &SpawnTransform);
+	AActor* PossessedActor = PlaybackContext->SpawnActor(Spawnable->GetObjectTemplate()->GetClass(), &SpawnTransform, SpawnInfo);
 
 	if (!PossessedActor)
 	{
@@ -5789,6 +6188,9 @@ FMovieScenePossessable* FSequencer::ConvertToPossessableInternal(FGuid Spawnable
 	UEngine::CopyPropertiesForUnrelatedObjects(SpawnableActorTemplate, PossessedActor, CopyParams);
 	PossessedActor->RegisterAllComponents();
 
+	const bool bIsDefaultTransform = true;
+	PossessedActor->FinishSpawning(SpawnTransform, bIsDefaultTransform);
+
 	const FGuid PossessableGuid = CreateBinding(*PossessedActor, PossessedActor->GetActorLabel());
 
 	FMovieScenePossessable* Possessable = MovieScene->FindPossessable(PossessableGuid);
@@ -5815,32 +6217,19 @@ FMovieScenePossessable* FSequencer::ConvertToPossessableInternal(FGuid Spawnable
 	
 	GEditor->SelectActor(PossessedActor, false, true);
 
-	SetLocalTimeDirectly(ScrubPosition);
+	ForceEvaluate();
 
 	return Possessable;
 }
 
 void FSequencer::OnAddFolder()
-{
-	FScopedTransaction AddFolderTransaction( NSLOCTEXT("Sequencer", "AddFolder_Transaction", "Add Folder") );
-
-	// Check if a folder, or child of a folder is currently selected.
-	TArray<UMovieSceneFolder*> SelectedParentFolders;
-	if ( Selection.GetSelectedOutlinerNodes().Num() > 0 )
-	{
-		for ( TSharedRef<FSequencerDisplayNode> SelectedNode : Selection.GetSelectedOutlinerNodes() )
-		{
-			TSharedPtr<FSequencerDisplayNode> CurrentNode = SelectedNode;
-			while ( CurrentNode.IsValid() && CurrentNode->GetType() != ESequencerNode::Folder )
-			{
-				CurrentNode = CurrentNode->GetParent();
-			}
-			if ( CurrentNode.IsValid() )
-			{
-				SelectedParentFolders.Add( &StaticCastSharedPtr<FSequencerFolderNode>( CurrentNode )->GetFolder() );
-			}
-		}
-	}
+{
+	FScopedTransaction AddFolderTransaction( NSLOCTEXT("Sequencer", "AddFolder_Transaction", "Add Folder") );
+
+	// Check if a folder, or child of a folder is currently selected.
+	TArray<UMovieSceneFolder*> SelectedParentFolders;
+	FString NewNodePath;
+	CalculateSelectedFolderAndPath(SelectedParentFolders, NewNodePath);
 
 	TArray<FName> ExistingFolderNames;
 	UMovieScene* FocusedMovieScene = GetFocusedMovieSceneSequence()->GetMovieScene();
@@ -5866,6 +6255,9 @@ void FSequencer::OnAddFolder()
 	UMovieSceneFolder* NewFolder = NewObject<UMovieSceneFolder>( FocusedMovieScene, NAME_None, RF_Transactional );
 	NewFolder->SetFolderName( UniqueName );
 
+	// The folder's name is used as it's key in the path system.
+	NewNodePath += UniqueName.ToString();
+
 	if ( SelectedParentFolders.Num() == 1 )
 	{
 		SelectedParentFolders[0]->Modify();
@@ -5877,9 +6269,73 @@ void FSequencer::OnAddFolder()
 		FocusedMovieScene->GetRootFolders().Add( NewFolder );
 	}
 
+	Selection.Empty();
+
+	// We can't add the newly created folder to the selection set as the nodes for it don't actually exist yet.
+	// However, we can calculate the resulting path that the node will end up at and add that to the selection
+	// set, which will cause the newly created node to be selected when the selection is restored post-refresh.
+	SequencerWidget->AddAdditionalPathToSelectionSet(NewNodePath);
+
 	NotifyMovieSceneDataChanged( EMovieSceneDataChangeType::MovieSceneStructureItemAdded );
 }
 
+void FSequencer::OnAddTrack(const TWeakObjectPtr<UMovieSceneTrack>& InTrack)
+{
+	FString NewNodePath;
+
+	// Cinematic Shot Tracks and Camera Cut Tracks are always in the root and ignore sorting/folders, so we don't give them a chance to be placed into a folder.
+	bool bIsValidTrack = !(InTrack->IsA(UMovieSceneCinematicShotTrack::StaticClass()) || InTrack->IsA(UMovieSceneCameraCutTrack::StaticClass()));
+	if(bIsValidTrack)
+	{
+		TArray<UMovieSceneFolder*> SelectedParentFolders;
+		CalculateSelectedFolderAndPath(SelectedParentFolders, NewNodePath);
+
+		if (SelectedParentFolders.Num() == 1)
+		{
+			SelectedParentFolders[0]->Modify();
+			SelectedParentFolders[0]->AddChildMasterTrack(InTrack.Get());
+		}
+	}
+	Selection.Empty();
+
+	// We can't add the newly created folder to the selection set as the nodes for it don't actually exist yet.
+	// However, we can calculate the resulting path that the node will end up at and add that to the selection
+	// set, which will cause the newly created node to be selected when the selection is restored post-refresh.
+	NewNodePath += InTrack->GetFName().ToString();
+	SequencerWidget->AddAdditionalPathToSelectionSet(NewNodePath);
+}
+
+
+void FSequencer::CalculateSelectedFolderAndPath(TArray<UMovieSceneFolder*>& OutSelectedParentFolders, FString& OutNewNodePath)
+{
+	// Check if a folder, or child of a folder is currently selected.
+	if (Selection.GetSelectedOutlinerNodes().Num() > 0)
+	{
+		for (TSharedRef<FSequencerDisplayNode> SelectedNode : Selection.GetSelectedOutlinerNodes())
+		{
+			TSharedPtr<FSequencerDisplayNode> CurrentNode = SelectedNode;
+			while (CurrentNode.IsValid() && CurrentNode->GetType() != ESequencerNode::Folder)
+			{
+				CurrentNode = CurrentNode->GetParent();
+			}
+			if (CurrentNode.IsValid())
+			{
+				OutSelectedParentFolders.Add(&StaticCastSharedPtr<FSequencerFolderNode>(CurrentNode)->GetFolder());
+
+				// The first valid folder we find will be used to put the new folder into, so it's the node that we
+				// want to know the path from.
+				if (OutNewNodePath.Len() == 0)
+				{
+					// Add an extra delimiter (".") as we know that the new folder will be appended onto the end of this.
+					OutNewNodePath = FString::Printf(TEXT("%s."), *CurrentNode->GetPathName());
+
+					// Make sure this folder is expanded too so that adding objects to hidden folders become visible.
+					CurrentNode->SetExpansionState(true);
+				}
+			}
+		}
+	}
+}
 
 void FSequencer::TogglePlay()
 {
@@ -5935,19 +6391,22 @@ void FSequencer::Pause()
 	SetPlaybackStatus(EMovieScenePlayerStatus::Stopped);
 
 	// When stopping a sequence, we always evaluate a non-empty range if possible. This ensures accurate paused motion blur effects.
-	const float TimeSnapInterval = GetFixedFrameInterval();
-	if (Settings->GetIsSnapEnabled() && TimeSnapInterval > 0)
+	if (Settings->GetIsSnapEnabled())
 	{
-		ScrubPosition = (SequencerHelpers::SnapTimeToInterval(GetLocalTime(), GetFixedFrameInterval()) + TimeSnapInterval) * RootToLocalTransform.Inverse();
+		FQualifiedFrameTime LocalTime       = GetLocalTime();
+		FFrameRate          FocusedPlayRate = GetFocusedPlayRate();
 
-		FMovieSceneEvaluationRange Range = PlayPosition.PlayTo(ScrubPosition, GetRootMovieSceneSequence()->GetMovieScene()->GetOptionalFixedFrameInterval());
-		EvaluateInternal(Range);
+		// Snap to the focused play rate
+		FFrameTime RootPosition  = FFrameRate::Snap(LocalTime.Time, LocalTime.Rate, FocusedPlayRate) * RootToLocalTransform.Inverse();
+
+		// Convert the root position from frame resolution time base (the output rate), to the play position input rate
+		FFrameTime InputPosition = ConvertFrameTime(RootPosition, PlayPosition.GetOutputRate(), PlayPosition.GetInputRate());
+		EvaluateInternal(PlayPosition.PlayTo(InputPosition));
 	}
 	else
 	{
 		// Update on stop (cleans up things like sounds that are playing)
-		TOptional<FMovieSceneEvaluationRange> LastRange = PlayPosition.GetLastRange();
-		FMovieSceneEvaluationRange Range = LastRange.IsSet() ? LastRange.GetValue() : PlayPosition.JumpTo(ScrubPosition, GetFocusedMovieSceneSequence()->GetMovieScene()->GetOptionalFixedFrameInterval());
+		FMovieSceneEvaluationRange Range = PlayPosition.GetLastRange().Get(PlayPosition.GetCurrentPositionAsRange());
 		EvaluateInternal(Range);
 	}
 
@@ -6000,22 +6459,21 @@ void FSequencer::StepToNextShot()
 	FMovieSceneSequenceID OuterSequenceID = ActiveTemplateIDs[ActiveTemplateIDs.Num()-2];
 	UMovieSceneSequence* Sequence = RootTemplateInstance.GetSequence(OuterSequenceID);
 
-	float StartTime = 0.f * RootToLocalTransform.Inverse();
-	float CurrentTime = StartTime * RootTemplateInstance.GetHierarchy().FindSubData(OuterSequenceID)->RootToSequenceTransform;
+	FFrameTime StartTime = FFrameTime(0) * RootToLocalTransform.Inverse();
+	FFrameTime CurrentTime = StartTime * RootTemplateInstance.GetHierarchy().FindSubData(OuterSequenceID)->RootToSequenceTransform;
 
-	UMovieSceneSubSection* NextShot = Cast<UMovieSceneSubSection>(FindNextOrPreviousShot(Sequence, CurrentTime, true));
+	UMovieSceneSubSection* NextShot = Cast<UMovieSceneSubSection>(FindNextOrPreviousShot(Sequence, CurrentTime.FloorToFrame(), true));
 	if (!NextShot)
 	{
 		return;
 	}
-		
+
 	SequencerWidget->PopBreadcrumb();
 
 	PopToSequenceInstance(ActiveTemplateIDs[ActiveTemplateIDs.Num()-2]);
-			
 	FocusSequenceInstance(*NextShot);
-	
-	SetLocalTime(0.f);
+
+	SetLocalTime(FFrameTime(0));
 }
 
 
@@ -6029,22 +6487,21 @@ void FSequencer::StepToPreviousShot()
 	FMovieSceneSequenceID OuterSequenceID = ActiveTemplateIDs[ActiveTemplateIDs.Num()-2];
 	UMovieSceneSequence* Sequence = RootTemplateInstance.GetSequence(OuterSequenceID);
 
-	float StartTime = 0.f * RootToLocalTransform.Inverse();
-	float CurrentTime = StartTime * RootTemplateInstance.GetHierarchy().FindSubData(OuterSequenceID)->RootToSequenceTransform;
+	FFrameTime StartTime = FFrameTime(0) * RootToLocalTransform.Inverse();
+	FFrameTime CurrentTime = StartTime * RootTemplateInstance.GetHierarchy().FindSubData(OuterSequenceID)->RootToSequenceTransform;
 
-	UMovieSceneSubSection* PreviousShot = Cast<UMovieSceneSubSection>(FindNextOrPreviousShot(Sequence, CurrentTime, false));
+	UMovieSceneSubSection* PreviousShot = Cast<UMovieSceneSubSection>(FindNextOrPreviousShot(Sequence, CurrentTime.FloorToFrame(), false));
 	if (!PreviousShot)
 	{
 		return;
 	}
 
 	SequencerWidget->PopBreadcrumb();
-		
+
 	PopToSequenceInstance(ActiveTemplateIDs[ActiveTemplateIDs.Num()-2]);
-					
 	FocusSequenceInstance(*PreviousShot);
-	
-	SetLocalTime(0.f);
+
+	SetLocalTime(0);
 }
 
 
@@ -6061,6 +6518,11 @@ void FSequencer::CollapseAllNodesAndDescendants()
 	SequencerWidget->GetTreeView()->CollapseNodes(ETreeRecursion::Recursive, bExpandAll);
 }
 
+void FSequencer::SortAllNodesAndDescendants()
+{
+	FScopedTransaction SortAllNodesTransaction(NSLOCTEXT("Sequencer", "SortAllNodes_Transaction", "Sort Tracks"));
+	SequencerWidget->GetTreeView()->GetNodeTree()->SortAllNodesAndDescendants();
+}
 
 void FSequencer::ToggleExpandCollapseNodes()
 {
@@ -6083,7 +6545,7 @@ void FSequencer::SetKey()
 		if (OutlinerNode->GetType() == ESequencerNode::Track)
 		{
 			TSharedRef<FSequencerTrackNode> TrackNode = StaticCastSharedRef<FSequencerTrackNode>(OutlinerNode);
-						
+
 			TSharedRef<FSequencerDisplayNode> ObjectBindingNode = OutlinerNode;
 			if (SequencerHelpers::FindObjectBindingNode(TrackNode, ObjectBindingNode))
 			{
@@ -6093,19 +6555,49 @@ void FSequencer::SetKey()
 		}
 	}
 
-	TSet<TSharedPtr<IKeyArea> > KeyAreas;
-	for (auto OutlinerNode : Selection.GetSelectedOutlinerNodes())
+	ISequencerModule& SequencerModule = FModuleManager::LoadModuleChecked<ISequencerModule>("Sequencer");
+
+	TSet<TSharedRef<FSequencerDisplayNode>> NodesToKey = Selection.GetSelectedOutlinerNodes();
 	{
-		SequencerHelpers::GetAllKeyAreas(OutlinerNode, KeyAreas);
+		TSet<TSharedRef<FSequencerDisplayNode>> ChildNodes;
+		for (TSharedRef<FSequencerDisplayNode> Node : NodesToKey.Array())
+		{
+			ChildNodes.Reset();
+			SequencerHelpers::GetDescendantNodes(Node, ChildNodes);
+
+			for (TSharedRef<FSequencerDisplayNode> ChildNode : ChildNodes)
+			{
+				NodesToKey.Remove(ChildNode);
+			}
+		}
 	}
 
-	if (KeyAreas.Num() > 0)
-	{	
-		for (auto KeyArea : KeyAreas)
+	const FFrameNumber AddKeyTime = GetLocalTime().Time.FloorToFrame();
+
+	TSet<TSharedPtr<IKeyArea>> KeyAreas;
+	TSet<UMovieSceneSection*>  ModifiedSections;
+
+	for (TSharedRef<FSequencerDisplayNode> Node : NodesToKey)
+	{
+		KeyAreas.Reset();
+		SequencerHelpers::GetAllKeyAreas(Node, KeyAreas);
+
+		TSharedPtr<FSequencerObjectBindingNode> ParentObjectBinding = Node->FindParentObjectBindingNode();
+		FGuid ObjectBinding = ParentObjectBinding.IsValid() ? ParentObjectBinding->GetObjectBinding() : FGuid();
+
+		for (TSharedPtr<IKeyArea> KeyArea : KeyAreas)
 		{
-			if (KeyArea->GetOwningSection()->TryModify())
+			UMovieSceneSection* Section = KeyArea->GetOwningSection();
+			if (Section)
 			{
-				KeyArea->AddKeyUnique(GetLocalTime(), GetKeyInterpolation());
+				if (!ModifiedSections.Contains(Section))
+				{
+					Section->Modify();
+					ModifiedSections.Add(Section);
+				}
+
+
+				KeyArea->AddOrUpdateKey(AddKeyTime, ObjectBinding, *this);
 			}
 		}
 	}
@@ -6120,11 +6612,11 @@ bool FSequencer::CanSetKeyTime() const
 }
 
 
-void FSequencer::SetKeyTime(const bool bUseFrames)
+void FSequencer::SetKeyTime()
 {
 	TArray<FSequencerSelectedKey> SelectedKeysArray = Selection.GetSelectedKeys().Array();
 
-	float KeyTime = 0.f;
+	FFrameNumber KeyTime = 0;
 	for ( const FSequencerSelectedKey& Key : SelectedKeysArray )
 	{
 		if (Key.IsValid())
@@ -6134,29 +6626,25 @@ void FSequencer::SetKeyTime(const bool bUseFrames)
 		}
 	}
 
-	float FrameRate = 1.0f / GetFixedFrameInterval();
-	
-	GenericTextEntryModeless(
-		bUseFrames ? NSLOCTEXT("Sequencer.Popups", "SetKeyFramePopup", "New Frame") : NSLOCTEXT("Sequencer.Popups", "SetKeyTimePopup", "New Time"),
-		bUseFrames ? FText::AsNumber( SequencerHelpers::TimeToFrame( KeyTime, FrameRate )) : FText::AsNumber( KeyTime ),
-		FOnTextCommitted::CreateSP(this, &FSequencer::OnSetKeyTimeTextCommitted, bUseFrames)
-	);
+	// Create a popup showing the existing time value and let the user set a new one.
+ 	GenericTextEntryModeless(NSLOCTEXT("Sequencer.Popups", "SetKeyTimePopup", "New Time"), FText::FromString(GetNumericTypeInterface()->ToString(KeyTime.Value)),
+ 		FOnTextCommitted::CreateSP(this, &FSequencer::OnSetKeyTimeTextCommitted)
+ 	);
 }
 
 
-void FSequencer::OnSetKeyTimeTextCommitted(const FText& InText, ETextCommit::Type CommitInfo, const bool bUseFrames)
+void FSequencer::OnSetKeyTimeTextCommitted(const FText& InText, ETextCommit::Type CommitInfo)
 {
 	bool bAnythingChanged = false;
+
 	CloseEntryPopupMenu();
 	if (CommitInfo == ETextCommit::OnEnter)
 	{
-		float FrameRate = 1.0f / GetFixedFrameInterval();
-		double dNewTime = bUseFrames ? SequencerHelpers::FrameToTime(FCString::Atod(*InText.ToString()), FrameRate) : FCString::Atod(*InText.ToString());
-		const bool bIsNumber = InText.IsNumeric(); 
-		if(!bIsNumber)
+		TOptional<double> NewFrameTime = GetNumericTypeInterface()->FromString(InText.ToString(), 0);
+		if (!NewFrameTime.IsSet())
 			return;
 
-		const float NewKeyTime = (float)dNewTime;
+		FFrameNumber NewFrame = FFrameNumber((int32)NewFrameTime.GetValue());
 
 		FScopedTransaction SetKeyTimeTransaction(NSLOCTEXT("Sequencer", "SetKeyTime_Transaction", "Set Key Time"));
 		TArray<FSequencerSelectedKey> SelectedKeysArray = Selection.GetSelectedKeys().Array();
@@ -6165,20 +6653,13 @@ void FSequencer::OnSetKeyTimeTextCommitted(const FText& InText, ETextCommit::Typ
 		{
 			if (Key.IsValid())
 			{
-				if (Key.Section->TryModify())
-				{
-					Key.KeyArea->SetKeyTime(Key.KeyHandle.GetValue(), NewKeyTime);
-					bAnythingChanged = true;
+	 			if (Key.Section->TryModify())
+	 			{
+	 				Key.KeyArea->SetKeyTime(Key.KeyHandle.GetValue(), NewFrame);
+	 				bAnythingChanged = true;
 
-					if( NewKeyTime > Key.Section->GetEndTime() )
-					{
-						Key.Section->SetEndTime( NewKeyTime );
-					}
-					else if( NewKeyTime < Key.Section->GetStartTime() )
-					{
-						Key.Section->SetStartTime( NewKeyTime );
-					}
-				}
+					Key.Section->ExpandToFrame(NewFrame);
+	 			}
 			}
 		}
 	}
@@ -6189,63 +6670,6 @@ void FSequencer::OnSetKeyTimeTextCommitted(const FText& InText, ETextCommit::Typ
 	}
 }
 
-
-void FSequencer::SelectTrackKeys(TWeakObjectPtr<UMovieSceneSection> Section, float KeyTime, bool bAddToSelection, bool bToggleSelection)
-{
-	if (!bAddToSelection && !bToggleSelection)
-	{
-		Selection.EmptySelectedKeys();
-	}
-
-	TSet<TWeakObjectPtr<UMovieSceneSection> > Sections;
-	Sections.Add(Section);
-	TArray<FSectionHandle> SectionHandles = SequencerWidget->GetSectionHandles(Sections);
-	for (auto SectionHandle : SectionHandles)
-	{
-		if (SectionHandle.TrackNode.IsValid())
-		{
-			TSet<TSharedPtr<IKeyArea> > KeyAreas;
-			SequencerHelpers::GetAllKeyAreas(SectionHandle.TrackNode, KeyAreas);
-
-			for (auto KeyArea : KeyAreas)
-			{
-				if (KeyArea.IsValid())
-				{
-					TArray<FKeyHandle> KeyHandles = KeyArea.Get()->GetUnsortedKeyHandles();
-					for (auto KeyHandle : KeyHandles)
-					{
-						float KeyHandleTime = KeyArea.Get()->GetKeyTime(KeyHandle);
-
-						if (FMath::IsNearlyEqual(KeyHandleTime, KeyTime, KINDA_SMALL_NUMBER))
-						{
-							FSequencerSelectedKey SelectedKey(*Section.Get(), KeyArea, KeyHandle);
-
-							if (bToggleSelection)
-							{
-								if (Selection.IsSelected(SelectedKey))
-								{
-									Selection.RemoveFromSelection(SelectedKey);
-								}
-								else
-								{
-									Selection.AddToSelection(SelectedKey);
-								}
-							}
-							else
-							{
-								Selection.AddToSelection(SelectedKey);
-							}
-						}
-					}
-				}
-			}
-		}
-	}
-
-	SequencerHelpers::ValidateNodesWithSelectedKeysOrSections(*this);
-}
-
-
 TArray<TSharedPtr<FMovieSceneClipboard>> GClipboardStack;
 
 void FSequencer::CopySelection()
@@ -6323,40 +6747,36 @@ void FSequencer::DuplicateSelection()
 
 void FSequencer::CopySelectedKeys()
 {
-	TOptional<float> CopyRelativeTo;
+	TOptional<FFrameNumber> CopyRelativeTo;
 	
 	// Copy relative to the current key hotspot, if applicable
 	if (Hotspot.IsValid() && Hotspot->GetType() == ESequencerHotspot::Key)
 	{
-		FKeyHotspot* KeyHotspot = StaticCastSharedPtr<FKeyHotspot>(Hotspot).Get();
-		if (KeyHotspot->Key.KeyArea.IsValid() && KeyHotspot->Key.KeyHandle.IsSet())
-		{
-			CopyRelativeTo = KeyHotspot->Key.KeyArea->GetKeyTime(KeyHotspot->Key.KeyHandle.GetValue());
-		}
+		CopyRelativeTo = StaticCastSharedPtr<FKeyHotspot>(Hotspot)->GetTime();
 	}
 
 	FMovieSceneClipboardBuilder Builder;
 
 	// Map selected keys to their key areas
-	TMap<const IKeyArea*, TArray<FKeyHandle>> KeyAreaMap;
+	TMap<TSharedPtr<IKeyArea>, TArray<FKeyHandle>> KeyAreaMap;
 	for (const FSequencerSelectedKey& Key : Selection.GetSelectedKeys())
 	{
 		if (Key.KeyHandle.IsSet())
 		{
-			KeyAreaMap.FindOrAdd(Key.KeyArea.Get()).Add(Key.KeyHandle.GetValue());
+			KeyAreaMap.FindOrAdd(Key.KeyArea).Add(Key.KeyHandle.GetValue());
 		}
 	}
 
 	// Serialize each key area to the clipboard
 	for (auto& Pair : KeyAreaMap)
 	{
-		Pair.Key->CopyKeys(Builder, [&](FKeyHandle Handle, const IKeyArea&){
-			return Pair.Value.Contains(Handle);
-		});
+		Pair.Key->CopyKeys(Builder, Pair.Value);
 	}
 
 	TSharedRef<FMovieSceneClipboard> Clipboard = MakeShareable( new FMovieSceneClipboard(Builder.Commit(CopyRelativeTo)) );
 	
+	Clipboard->GetEnvironment().FrameResolution = GetFocusedFrameResolution();
+
 	if (Clipboard->GetKeyTrackGroups().Num())
 	{
 		GClipboardStack.Push(Clipboard);
@@ -6588,7 +7008,7 @@ void FSequencer::CreateCamera()
 
 	NewCameraAdded(CameraGuid, NewCamera);
 
-	if (ensure(Spawnable))
+	if (bCreateAsSpawnable && ensure(Spawnable))
 	{
 		Spawnable->SetSpawnOwnership(SavedOwnership);
 	}
@@ -6635,7 +7055,7 @@ void FSequencer::NewCameraAdded(FGuid CameraGuid, ACameraActor* NewCamera)
 
 	if (CameraCutTrack)
 	{
-		UMovieSceneSection* Section = MovieSceneHelpers::FindSectionAtTime(CameraCutTrack->GetAllSections(), GetLocalTime());
+		UMovieSceneSection* Section = MovieSceneHelpers::FindSectionAtTime(CameraCutTrack->GetAllSections(), GetLocalTime().Time.FloorToFrame());
 		UMovieSceneCameraCutSection* CameraCutSection = Cast<UMovieSceneCameraCutSection>(Section);
 
 		if (CameraCutSection)
@@ -6648,8 +7068,7 @@ void FSequencer::NewCameraAdded(FGuid CameraGuid, ACameraActor* NewCamera)
 			CameraCutTrack->Modify();
 
 			UMovieSceneCameraCutSection* NewSection = Cast<UMovieSceneCameraCutSection>(CameraCutTrack->CreateNewSection());
-			NewSection->SetStartTime(GetPlaybackRange().GetLowerBoundValue());
-			NewSection->SetEndTime(GetPlaybackRange().GetUpperBoundValue());
+			NewSection->SetRange(GetPlaybackRange());
 			NewSection->SetCameraGuid(CameraGuid);
 			CameraCutTrack->AddSection(*NewSection);
 		}
@@ -6760,140 +7179,6 @@ void FSequencer::RebindPossessableReferences()
 	}
 }
 
-float SnapTime( float TimeValue, float TimeInterval )
-{
-	return FMath::RoundToInt( TimeValue / TimeInterval ) * TimeInterval;
-}
-
-
-void FixSceneRangeTiming( UMovieScene* MovieScene, float FrameInterval )
-{
-	TRange<float> SceneRange = MovieScene->GetPlaybackRange();
-
-	float LowerBoundValue = SceneRange.GetLowerBoundValue();
-	float SnappedLowerBoundValue = SnapTime( LowerBoundValue, FrameInterval );
-
-	float UpperBoundValue = SceneRange.GetUpperBoundValue();
-	float SnappedUpperBoundValue = SnapTime( UpperBoundValue, FrameInterval );
-
-	if ( SnappedLowerBoundValue != LowerBoundValue || SnappedUpperBoundValue != UpperBoundValue)
-	{
-		MovieScene->SetPlaybackRange( SnappedLowerBoundValue, SnappedUpperBoundValue );
-	}
-}
-
-
-void FixSectionFrameTiming( UMovieSceneSection* Section, float FrameInterval )
-{
-	bool bSectionModified = false;
-	float SnappedStartTime = SnapTime( Section->GetStartTime(), FrameInterval );
-	if ( SnappedStartTime != Section->GetStartTime() )
-	{
-		Section->Modify();
-		bSectionModified = true;
-		Section->SetStartTime( SnappedStartTime );
-	}
-
-	float SnappedEndTime = SnapTime( Section->GetEndTime(), FrameInterval );
-	if ( SnappedEndTime != Section->GetEndTime() )
-	{
-		if ( bSectionModified == false )
-		{
-			Section->Modify();
-			bSectionModified = true;
-		}
-		Section->SetEndTime( SnappedEndTime );
-	}
-
-	TSet<FKeyHandle> KeyHandles;
-	Section->GetKeyHandles( KeyHandles, Section->GetRange() );
-	for ( const FKeyHandle& KeyHandle : KeyHandles )
-	{
-		TOptional<float> KeyTime = Section->GetKeyTime( KeyHandle );
-		if ( KeyTime.IsSet() )
-		{
-			float SnappedKeyTime = SnapTime( KeyTime.GetValue(), FrameInterval );
-			if ( SnappedKeyTime != KeyTime.GetValue() )
-			{
-				if ( bSectionModified == false )
-				{
-					Section->Modify();
-					bSectionModified = true;
-				}
-				Section->SetKeyTime( KeyHandle, SnappedKeyTime );
-			}
-		}
-	}
-}
-
-
-void GetAllMovieScenesRecursively( UMovieScene* CurrentMovieScene, TArray<UMovieScene*>& AllMovieScenes )
-{
-	if( CurrentMovieScene != nullptr && AllMovieScenes.Contains( CurrentMovieScene ) == false)
-	{
-		AllMovieScenes.Add( CurrentMovieScene );
-		for ( UMovieSceneTrack* MasterTrack : CurrentMovieScene->GetMasterTracks() )
-		{
-			UMovieSceneSubTrack* SubTrack = Cast<UMovieSceneSubTrack>( MasterTrack );
-			if ( SubTrack != nullptr )
-			{
-				for ( UMovieSceneSection* Section : SubTrack->GetAllSections() )
-				{
-					UMovieSceneSubSection* SubSection = Cast<UMovieSceneSubSection>( Section );
-					if ( SubSection != nullptr )
-					{
-						GetAllMovieScenesRecursively( SubSection->GetSequence()->GetMovieScene(), AllMovieScenes );
-					}
-				}
-			}
-		}
-	}
-}
-
-
-void FSequencer::FixFrameTiming()
-{
-	TArray<UMovieScene*> ScenesToFix;
-	GetAllMovieScenesRecursively( GetRootMovieSceneSequence()->GetMovieScene(), ScenesToFix );
-
-	FScopedTransaction FixFrameTimingTransaction( NSLOCTEXT( "Sequencer", "FixFrameTiming", "Fix frame timing" ) );
-	for( UMovieScene* SceneToFix : ScenesToFix)
-	{
-		float FrameInterval = SceneToFix->GetFixedFrameInterval();
-		if ( FrameInterval > 0 )
-		{
-			FixSceneRangeTiming( SceneToFix, FrameInterval );
-
-			// Collect all tracks.
-			TArray<UMovieSceneTrack*> TracksToFix;
-
-			if ( SceneToFix->GetCameraCutTrack() != nullptr )
-			{
-				TracksToFix.Add( SceneToFix->GetCameraCutTrack() );
-			}
-
-			for ( UMovieSceneTrack* MasterTrack : SceneToFix->GetMasterTracks() )
-			{
-				TracksToFix.Add( MasterTrack );
-			}
-
-			for ( const FMovieSceneBinding& ObjectBinding : SceneToFix->GetBindings() )
-			{
-				TracksToFix.Append( ObjectBinding.GetTracks() );
-			}
-
-			// Fix section and keys for tracks in the current scene.
-			for ( UMovieSceneTrack* TrackToFix : TracksToFix )
-			{
-				for ( UMovieSceneSection* Section : TrackToFix->GetAllSections() )
-				{
-					FixSectionFrameTiming( Section, FrameInterval );
-				}
-			}
-		}
-	}
-}
-
 void FSequencer::ImportFBX()
 {
 	UMovieScene* MovieScene = GetFocusedMovieSceneSequence()->GetMovieScene();
@@ -7067,7 +7352,7 @@ void FSequencer::CloseEntryPopupMenu()
 void FSequencer::TrimSection(bool bTrimLeft)
 {
 	FScopedTransaction TrimSectionTransaction( NSLOCTEXT("Sequencer", "TrimSection_Transaction", "Trim Section") );
-	MovieSceneToolHelpers::TrimSection(Selection.GetSelectedSections(), GetLocalTime(), bTrimLeft);
+	MovieSceneToolHelpers::TrimSection(Selection.GetSelectedSections(), GetLocalTime().Time.FloorToFrame(), bTrimLeft);
 	NotifyMovieSceneDataChanged( EMovieSceneDataChangeType::TrackValueChanged );
 }
 
@@ -7075,7 +7360,7 @@ void FSequencer::TrimSection(bool bTrimLeft)
 void FSequencer::SplitSection()
 {
 	FScopedTransaction SplitSectionTransaction( NSLOCTEXT("Sequencer", "SplitSection_Transaction", "Split Section") );
-	MovieSceneToolHelpers::SplitSection(Selection.GetSelectedSections(), GetLocalTime());
+	MovieSceneToolHelpers::SplitSection(Selection.GetSelectedSections(), GetLocalTime().Time.FloorToFrame());
 	NotifyMovieSceneDataChanged( EMovieSceneDataChangeType::MovieSceneStructureItemAdded );
 }
 
@@ -7125,6 +7410,10 @@ void FSequencer::BindCommands()
 		Commands.CollapseAllNodesAndDescendants,
 		FExecuteAction::CreateSP(this, &FSequencer::CollapseAllNodesAndDescendants));
 
+	SequencerCommandBindings->MapAction(
+		Commands.SortAllNodesAndDescendants,
+		FExecuteAction::CreateSP(this, &FSequencer::SortAllNodesAndDescendants));
+
 	SequencerCommandBindings->MapAction(
 		Commands.ToggleExpandCollapseNodes,
 		FExecuteAction::CreateSP(this, &FSequencer::ToggleExpandCollapseNodes));
@@ -7328,10 +7617,24 @@ void FSequencer::BindCommands()
 		FIsActionChecked::CreateLambda( [this]{ return Settings->GetShowSelectedNodesOnly(); } ) );
 
 	SequencerCommandBindings->MapAction(
-		Commands.ToggleShowFrameNumbers,
-		FExecuteAction::CreateLambda( [this]{ Settings->SetShowFrameNumbers( !Settings->GetShowFrameNumbers() ); } ),
-		FCanExecuteAction::CreateSP( this, &FSequencer::CanShowFrameNumbers ),
-		FIsActionChecked::CreateLambda( [this]{ return Settings->GetShowFrameNumbers(); } ) );
+		Commands.ChangeTimeDisplayFormat,
+		FExecuteAction::CreateLambda( [this]{
+			EFrameNumberDisplayFormats NextFormat = (EFrameNumberDisplayFormats)((uint8)Settings->GetTimeDisplayFormat() + 1);
+			if (NextFormat == EFrameNumberDisplayFormats::MAX_Count)
+			{
+				NextFormat = EFrameNumberDisplayFormats::NonDropFrameTimecode;
+			}
+
+			// If the next framerate in the list is drop format timecode and we're not in a play rate that supports drop format timecode,
+			// then we will skip over it.
+			bool bCanShowDropFrameTimecode = FTimecode::IsDropFormatTimecodeSupported(GetFocusedPlayRate());
+			if (!bCanShowDropFrameTimecode && NextFormat == EFrameNumberDisplayFormats::DropFrameTimecode)
+			{
+				NextFormat = EFrameNumberDisplayFormats::Seconds;
+			}
+			Settings->SetTimeDisplayFormat( NextFormat );
+		} ),
+		FCanExecuteAction::CreateLambda([] { return true; }));
 
 	SequencerCommandBindings->MapAction(
 		Commands.ToggleShowRangeSlider,
@@ -7480,6 +7783,16 @@ void FSequencer::BindCommands()
 		return false;
 	};
 
+	auto IsSelectionRangeNonEmpty = [this]{
+		UMovieSceneSequence* EditedSequence = GetFocusedMovieSceneSequence();
+		if (!EditedSequence || !EditedSequence->GetMovieScene())
+		{
+			return false;
+		}
+
+		return !EditedSequence->GetMovieScene()->GetSelectionRange().IsEmpty();
+	};
+
 	SequencerCommandBindings->MapAction(
 		FGenericCommands::Get().Cut,
 		FExecuteAction::CreateSP(this, &FSequencer::CutSelection),
@@ -7509,32 +7822,6 @@ void FSequencer::BindCommands()
 		FCanExecuteAction::CreateLambda( [this] { return GetFocusedMovieSceneSequence() != nullptr;	} ),
 		FIsActionChecked::CreateSP( this, &FSequencer::IsPlaybackRangeLocked ));
 
-	SequencerCommandBindings->MapAction(
-		Commands.ToggleForceFixedFrameIntervalPlayback,
-		FExecuteAction::CreateLambda( [this] 
-		{
-			UMovieSceneSequence* FocusedMovieSceneSequence = GetFocusedMovieSceneSequence();
-			if ( FocusedMovieSceneSequence != nullptr )
-			{
-				FScopedTransaction ToggleForceFixedFrameIntervalPlaybackTransaction( NSLOCTEXT( "Sequencer", "ToggleForceFixedFrameIntervalPlaybackTransaction", "Toggle force fixed frame interval playback" ) );
-				UMovieScene* MovieScene = FocusedMovieSceneSequence->GetMovieScene();
-				MovieScene->Modify();
-				MovieScene->SetForceFixedFrameIntervalPlayback( !MovieScene->GetForceFixedFrameIntervalPlayback() );
-				if ( MovieScene->GetForceFixedFrameIntervalPlayback() && MovieScene->GetFixedFrameInterval() == 0 )
-				{
-					MovieScene->SetFixedFrameInterval( Settings->GetTimeSnapInterval() );
-				}
-			} 
-		} ),
-		FCanExecuteAction::CreateLambda( [this] { return GetFocusedMovieSceneSequence() != nullptr;	} ),
-		FIsActionChecked::CreateLambda( [this] 
-		{ 
-			UMovieSceneSequence* FocusedMovieSceneSequence = GetFocusedMovieSceneSequence();
-			return FocusedMovieSceneSequence != nullptr
-				? FocusedMovieSceneSequence->GetMovieScene()->GetForceFixedFrameIntervalPlayback()
-				: false;
-		} ) );
-
 	SequencerCommandBindings->MapAction(
 		Commands.ToggleRerunConstructionScripts,
 		FExecuteAction::CreateLambda( [this]{ Settings->SetRerunConstructionScripts( !Settings->ShouldRerunConstructionScripts() ); } ),
@@ -7563,15 +7850,14 @@ void FSequencer::BindCommands()
 		Commands.ToggleEvaluateSubSequencesInIsolation,
 		FExecuteAction::CreateLambda( [this]{
 			Settings->SetEvaluateSubSequencesInIsolation( !Settings->ShouldEvaluateSubSequencesInIsolation() );
-			FMovieSceneEvaluationRange Range = PlayPosition.JumpTo(ScrubPosition, GetFocusedMovieSceneSequence()->GetMovieScene()->GetOptionalFixedFrameInterval());
-			EvaluateInternal(Range);
+			ForceEvaluate();
 		} ),
 		FCanExecuteAction::CreateLambda( [this]{ return ActiveTemplateIDs.Num() > 1; } ),
 		FIsActionChecked::CreateLambda( [this]{ return Settings->ShouldEvaluateSubSequencesInIsolation(); } ) );
 
 	SequencerCommandBindings->MapAction(
 		Commands.RenderMovie,
-		FExecuteAction::CreateLambda([this]{ RenderMovieInternal(GetPlaybackRange().GetLowerBoundValue(), GetPlaybackRange().GetUpperBoundValue()); }),
+		FExecuteAction::CreateLambda([this]{ RenderMovieInternal(GetPlaybackRange()); }),
 		FCanExecuteAction(),
 		FIsActionChecked(),
 		FIsActionButtonVisible::CreateLambda([this]{ return ExactCast<ULevelSequence>(GetFocusedMovieSceneSequence()) != nullptr; })
@@ -7616,11 +7902,6 @@ void FSequencer::BindCommands()
 		FExecuteAction::CreateSP( this, &FSequencer::RebindPossessableReferences ),
 		FCanExecuteAction::CreateLambda( []{ return true; } ) );
 
-	SequencerCommandBindings->MapAction(
-		Commands.FixFrameTiming,
-		FExecuteAction::CreateSP( this, &FSequencer::FixFrameTiming ),
-		FCanExecuteAction::CreateLambda( [] { return true; } ) );
-
 	SequencerCommandBindings->MapAction(
 		Commands.ImportFBX,
 		FExecuteAction::CreateSP( this, &FSequencer::ImportFBX ),
@@ -7713,19 +7994,23 @@ void FSequencer::BindCommands()
 
 	SequencerCommandBindings->MapAction(
 		Commands.ResetSelectionRange,
-		FExecuteAction::CreateLambda([this]{ ResetSelectionRange(); }));
+		FExecuteAction::CreateLambda([this]{ ResetSelectionRange(); }),
+		FCanExecuteAction::CreateLambda(IsSelectionRangeNonEmpty));
 
 	SequencerCommandBindings->MapAction(
 		Commands.SelectKeysInSelectionRange,
-		FExecuteAction::CreateSP(this, &FSequencer::SelectInSelectionRange, true, false));
+		FExecuteAction::CreateSP(this, &FSequencer::SelectInSelectionRange, true, false),
+		FCanExecuteAction::CreateLambda(IsSelectionRangeNonEmpty));
 
 	SequencerCommandBindings->MapAction(
 		Commands.SelectSectionsInSelectionRange,
-		FExecuteAction::CreateSP(this, &FSequencer::SelectInSelectionRange, false, true));
+		FExecuteAction::CreateSP(this, &FSequencer::SelectInSelectionRange, false, true),
+		FCanExecuteAction::CreateLambda(IsSelectionRangeNonEmpty));
 
 	SequencerCommandBindings->MapAction(
 		Commands.SelectAllInSelectionRange,
-		FExecuteAction::CreateSP(this, &FSequencer::SelectInSelectionRange, true, true));
+		FExecuteAction::CreateSP(this, &FSequencer::SelectInSelectionRange, true, true),
+		FCanExecuteAction::CreateLambda(IsSelectionRangeNonEmpty));
 
 	SequencerCommandBindings->MapAction(
 		Commands.StepToNextShot,
@@ -7776,6 +8061,22 @@ void FSequencer::BindCommands()
 		FExecuteAction::CreateSP( this, &FSequencer::SetPlaybackRangeToAllShots ),
 		FCanExecuteAction::CreateSP( this, &FSequencer::IsViewingMasterSequence ) );
 
+	// Curve Visibility
+	SequencerCommandBindings->MapAction(Commands.SetAllCurveVisibility,
+		FExecuteAction::CreateLambda( [this]{ Settings->SetCurveVisibility( ECurveEditorCurveVisibility::AllCurves ); } ),
+		FCanExecuteAction::CreateLambda( []{ return true; } ),
+		FIsActionChecked::CreateLambda( [this]{ return Settings->GetCurveVisibility() == ECurveEditorCurveVisibility::AllCurves; } ) );
+
+	SequencerCommandBindings->MapAction(Commands.SetSelectedCurveVisibility,
+		FExecuteAction::CreateLambda( [this]{ Settings->SetCurveVisibility( ECurveEditorCurveVisibility::SelectedCurves ); } ),
+		FCanExecuteAction::CreateLambda( []{ return true; } ),
+		FIsActionChecked::CreateLambda( [this]{ return Settings->GetCurveVisibility() == ECurveEditorCurveVisibility::SelectedCurves; } ) );
+
+	SequencerCommandBindings->MapAction(Commands.SetAnimatedCurveVisibility,
+		FExecuteAction::CreateLambda( [this]{ Settings->SetCurveVisibility( ECurveEditorCurveVisibility::AnimatedCurves ); } ),
+		FCanExecuteAction::CreateLambda( []{ return true; } ),
+		FIsActionChecked::CreateLambda( [this]{ return Settings->GetCurveVisibility() == ECurveEditorCurveVisibility::AnimatedCurves; } ) );
+
 	// bind widget specific commands
 	SequencerWidget->BindCommands(SequencerCommandBindings);
 }
@@ -7826,17 +8127,89 @@ void FSequencer::BuildObjectBindingEditButtons(TSharedPtr<SHorizontalBox> EditBo
 	}
 }
 
-void FSequencer::ResetTimingManager(bool bInShouldLockToAudioClock)
+void FSequencer::UpdateTimeBases()
+{
+	UMovieSceneSequence* RootSequencePtr = GetRootMovieSceneSequence();
+	UMovieScene*         RootMovieScene  = RootSequencePtr ? RootSequencePtr->GetMovieScene() : nullptr;
+
+	if (RootMovieScene)
+	{
+		EMovieSceneEvaluationType EvaluationType    = RootMovieScene->GetEvaluationType();
+		FFrameRate                FrameResolution   = RootMovieScene->GetFrameResolution();
+		FFrameRate                PlaybackFrameRate = EvaluationType == EMovieSceneEvaluationType::FrameLocked ? RootMovieScene->GetPlaybackFrameRate() : FrameResolution;
+
+		if (PlaybackFrameRate != PlayPosition.GetInputRate())
+		{
+			bNeedsEvaluate = true;
+		}
+
+		// We set the play position in terms of the playback frame rate,
+		// but want evaluation ranges in the moviescene's frame resolution
+		PlayPosition.SetTimeBase(PlaybackFrameRate, FrameResolution, EvaluationType);
+	}
+}
+
+void FSequencer::ResetTimeController()
 {
-	if (bInShouldLockToAudioClock && GEngine->GetMainAudioDevice())
+	switch (GetRootMovieSceneSequence()->GetMovieScene()->GetClockSource())
 	{
-		TimingManager.Reset(new FSequencerAudioClockTimer);
+	case EUpdateClockSource::Audio:    TimeController = MakeShared<FMovieSceneTimeController_AudioClock>();    break;
+	case EUpdateClockSource::Platform: TimeController = MakeShared<FMovieSceneTimeController_PlatformClock>(); break;
+	default:                           TimeController = MakeShared<FMovieSceneTimeController_Tick>();          break;
 	}
-	else
+
+	TimeController->PlayerStatusChanged(PlaybackState, GetGlobalTime());
+}
+
+
+void FSequencer::BuildCustomContextMenuForGuid(FMenuBuilder& MenuBuilder, FGuid ObjectBinding)
+{
+	SequencerWidget->BuildCustomContextMenuForGuid(MenuBuilder, ObjectBinding);
+}
+
+FKeyAttributes FSequencer::GetDefaultKeyAttributes() const
+{
+	switch (Settings->GetKeyInterpolation())
+	{
+	case EMovieSceneKeyInterpolation::User:     return FKeyAttributes().SetInterpMode(RCIM_Cubic).SetTangentMode(RCTM_User);
+	case EMovieSceneKeyInterpolation::Break:    return FKeyAttributes().SetInterpMode(RCIM_Cubic).SetTangentMode(RCTM_Break);
+	case EMovieSceneKeyInterpolation::Linear:   return FKeyAttributes().SetInterpMode(RCIM_Linear).SetTangentMode(RCTM_Auto);
+	case EMovieSceneKeyInterpolation::Constant: return FKeyAttributes().SetInterpMode(RCIM_Constant).SetTangentMode(RCTM_Auto);
+	default:                                    return FKeyAttributes().SetInterpMode(RCIM_Cubic).SetTangentMode(RCTM_Auto);
+	}
+}
+
+bool FSequencer::GetGridMetrics(float PhysicalWidth, double& OutMajorInterval, int32& OutMinorDivisions) const
+{
+	FSlateFontInfo SmallLayoutFont = FCoreStyle::GetDefaultFontStyle("Regular", 8);
+	TSharedRef<FSlateFontMeasure> FontMeasureService = FSlateApplication::Get().GetRenderer()->GetFontMeasureService();
+
+	double BiggestTime = GetViewRange().GetUpperBoundValue();
+	FString TickString = GetNumericTypeInterface()->ToString((BiggestTime * GetFocusedPlayRate()).FrameNumber.Value);
+	FVector2D MaxTextSize = FontMeasureService->Measure(TickString, SmallLayoutFont);
+
+	static float MajorTickMultiplier = 2.f;
+
+	float MinTickPx = MaxTextSize.X + 5.f;
+	float DesiredMajorTickPx = MaxTextSize.X * MajorTickMultiplier;
+
+	if (PhysicalWidth > 0)
 	{
-		TimingManager.Reset(new FSequencerDefaultTimingManager);
+		return GetFocusedPlayRate().ComputeGridSpacing(
+			PhysicalWidth / GetViewRange().Size<double>(),
+			OutMajorInterval,
+			OutMinorDivisions,
+			MinTickPx,
+			DesiredMajorTickPx);
 	}
-	TimingManager->Update(PlaybackState, GetGlobalTime());
+
+	return false;
 }
 
+double FSequencer::GetPlayRateDeltaFrameCount() const
+{
+	return GetFocusedFrameResolution().AsDecimal() * GetFocusedPlayRate().AsInterval();
+}
+
+
 #undef LOCTEXT_NAMESPACE
