Copying //UE4/Dev-Sequencer to //UE4/Dev-Main (Source: //UE4/Dev-Sequencer @ 3946692)

diff --git a/Engine/Source/Editor/Sequencer/Private/Sequencer.h b/Engine/Source/Editor/Sequencer/Private/Sequencer.h
--- a/Engine/Source/Editor/Sequencer/Private/Sequencer.h
+++ b/Engine/Source/Editor/Sequencer/Private/Sequencer.h
@@ -34,9 +34,9 @@
 #include "Evaluation/MovieScenePlayback.h"
 #include "Evaluation/MovieSceneEvaluationTemplateInstance.h"
 #include "LevelEditor.h"
-#include "SequencerTimingManager.h"
 #include "AcquiredResources.h"
 #include "SequencerSettings.h"
+#include "Curves/RichCurve.h"
 
 class AActor;
 class ACameraActor;
@@ -50,8 +50,9 @@ class FViewportClient;
 class IDetailKeyframeHandler;
 class ILevelViewport;
 class IMenu;
+class FCurveEditor;
 class ISequencerEditTool;
-class ISequencerKeyCollection;
+class FSequencerKeyCollection;
 class ISequencerTrackEditor;
 class ISequencerEditorObjectBinding;
 class SSequencer;
@@ -59,16 +60,19 @@ class ULevel;
 class UMovieSceneSequence;
 class UMovieSceneSubSection;
 class USequencerSettings;
+struct FMovieSceneTimeController;
 struct FMovieScenePossessable;
 struct FSequencerTemplateStore;
 struct FTransformData;
 struct ISequencerHotspot;
+struct FKeyAttributes;
+
 enum class EMapChangeType : uint8;
 
 /**
  * Sequencer is the editing tool for MovieScene assets.
  */
-class FSequencer
+class FSequencer final
 	: public ISequencer
 	, public FGCObject
 	, public FEditorUndoClient
@@ -96,11 +100,11 @@ public:
 
 	/** @return The current view range */
 	virtual FAnimatedRange GetViewRange() const override;
-	virtual void SetViewRange(TRange<float> NewViewRange, EViewRangeInterpolation Interpolation = EViewRangeInterpolation::Animated) override;
+	virtual void SetViewRange(TRange<double> NewViewRange, EViewRangeInterpolation Interpolation = EViewRangeInterpolation::Animated) override;
 
 	/** @return The current clamp range */
 	FAnimatedRange GetClampRange() const;
-	void SetClampRange(TRange<float> InNewClampRange);
+	void SetClampRange(TRange<double> InNewClampRange);
 
 public:
 
@@ -110,7 +114,7 @@ public:
 	 * @return The selection range.
 	 * @see SetSelectionRange, SetSelectionRangeEnd, SetSelectionRangeStart
 	 */
-	TRange<float> GetSelectionRange() const;
+	TRange<FFrameNumber> GetSelectionRange() const;
 
 	/**
 	 * Set the selection selection range.
@@ -118,7 +122,7 @@ public:
 	 * @param Range The new range to set.
 	 * @see GetSelectionRange, SetSelectionRangeEnd, SetSelectionRangeStart
 	 */
-	void SetSelectionRange(TRange<float> Range);
+	void SetSelectionRange(TRange<FFrameNumber> Range);
 	
 	/**
 	 * Set the selection range's end position to the current global time.
@@ -145,7 +149,12 @@ public:
 	 *
 	 * @return The sub sequence range, or an empty optional if we're viewing the root.
 	 */
-	TOptional<TRange<float>> GetSubSequenceRange() const;
+	TOptional<TRange<FFrameNumber>> GetSubSequenceRange() const;
+
+	/**
+	 * Compute a major grid interval and number of minor divisions to display
+	 */
+	bool GetGridMetrics(float PhysicalWidth, double& OutMajorInterval, int32& OutMinorDivisions) const;
 
 public:
 
@@ -155,7 +164,7 @@ public:
 	 * @return Playback range.
 	 * @see SetPlaybackRange, SetPlaybackRangeEnd, SetPlaybackRangeStart
 	 */
-	TRange<float> GetPlaybackRange() const;
+	TRange<FFrameNumber> GetPlaybackRange() const;
 
 	/**
 	 * Set the playback range.
@@ -163,7 +172,7 @@ public:
 	 * @param Range The new range to set.
 	 * @see GetPlaybackRange, SetPlaybackRangeEnd, SetPlaybackRangeStart
 	 */
-	void SetPlaybackRange(TRange<float> Range);
+	void SetPlaybackRange(TRange<FFrameNumber> Range);
 	
 	/**
 	 * Set the playback range's end position to the current global time.
@@ -172,7 +181,8 @@ public:
 	 */
 	void SetPlaybackRangeEnd()
 	{
-		SetPlaybackRange(TRange<float>(GetPlaybackRange().GetLowerBoundValue(), GetLocalTime()));
+		TRange<FFrameNumber> PlayRange = GetPlaybackRange();
+		SetPlaybackRange(TRange<FFrameNumber>(PlayRange.GetLowerBound(), TRangeBound<FFrameNumber>::Inclusive(GetLocalTime().Time.FrameNumber)));
 	}
 
 	/**
@@ -182,7 +192,8 @@ public:
 	 */
 	void SetPlaybackRangeStart()
 	{
-		SetPlaybackRange(TRange<float>(GetLocalTime(), GetPlaybackRange().GetUpperBoundValue()));
+		TRange<FFrameNumber> PlayRange = GetPlaybackRange();
+		SetPlaybackRange(TRange<FFrameNumber>(TRangeBound<FFrameNumber>::Inclusive(GetLocalTime().Time.FrameNumber), PlayRange.GetUpperBound()));
 	}
 
 	/**
@@ -207,13 +218,6 @@ public:
 	void ZoomOutViewRange();
 
 public:
-
-	/** Access the user-supplied settings object */
-	USequencerSettings* GetSettings() const
-	{
-		return Settings;
-	}
-
 	/** Gets the tree of nodes which is used to populate the animation outliner. */
 	TSharedRef<FSequencerNodeTree> GetNodeTree()
 	{
@@ -244,12 +248,6 @@ public:
 	/** Set interpolation modes. */
 	void SetInterpTangentMode(ERichCurveInterpMode InterpMode, ERichCurveTangentMode TangentMode);
 
-	/** Is interpolation mode selected. */
-	bool IsInterpTangentModeSelected(ERichCurveInterpMode InterpMode, ERichCurveTangentMode TangentMode) const;
-
-	/** Get the fixed frame interval of the current movie scene. */
-	float GetFixedFrameInterval() const;
-
 	/** Snap the currently selected keys to frame. */
 	void SnapToFrame();
 
@@ -257,7 +255,7 @@ public:
 	bool CanSnapToFrame() const;
 
 	/** Transform the selected keys and sections */
-	void TransformSelectedKeysAndSections(float InDeltaTime, float InScale);
+	void TransformSelectedKeysAndSections(FFrameTime InDeltaTime, float InScale);
 
 	/** Translate the selected keys and section by the time snap interval */
 	void TranslateSelectedKeysAndSections(bool bTranslateLeft);
@@ -418,17 +416,12 @@ public:
 	void HandleRecordingFinished(UMovieSceneSequence* Sequence);
 
 	/** Set the new global time, accounting for looping options */
-	void SetLocalTimeLooped(float InTime);
+	void SetLocalTimeLooped(FFrameTime InTime);
 
 	ESequencerLoopMode GetLoopMode() const;
 
 	EPlaybackMode::Type GetPlaybackMode() const;
 
-	bool IsRecording() const;
-
-	/** Called to determine whether a frame number is set so that frame numbers can be shown */
-	bool CanShowFrameNumbers() const;
-
 	/** @return The toolkit that this sequencer is hosted in (if any) */
 	TSharedPtr<IToolkitHost> GetToolkitHost() const { return ToolkitHost.Pin(); }
 
@@ -459,7 +452,7 @@ public:
 	void ExportTracksToText(TArray<UMovieSceneTrack*> TrackToExport, /*out*/ FString& ExportedText);
 
 	/** Called when a user executes the paste track menu item */
-	bool CanPaste(const FString& TextToImport) const;
+	bool CanPaste(const FString& TextToImport);
 	void PasteCopiedTracks();
 	void ImportTracksFromText(const FString& TextToImport, /*out*/ TArray<UMovieSceneTrack*>& ImportedTrack);
 
@@ -473,17 +466,16 @@ public:
 
 	/** Called when a user executes the set key time for selected keys */
 	bool CanSetKeyTime() const;
-	void SetKeyTime(const bool bUseFrames);
-	void OnSetKeyTimeTextCommitted(const FText& InText, ETextCommit::Type CommitInfo, const bool bUseFrames);
+	void SetKeyTime();
+	void OnSetKeyTimeTextCommitted(const FText& InText, ETextCommit::Type CommitInfo);
+
+	void SelectKey(UMovieSceneSection* InSection, TSharedPtr<IKeyArea> InKeyArea, FKeyHandle KeyHandle, bool bToggle);
 
 	FSequencerLabelManager& GetLabelManager()
 	{
 		return LabelManager;
 	}
 
-	/** Select keys belonging to a section at the key time */
-	void SelectTrackKeys(TWeakObjectPtr<UMovieSceneSection> Section, float KeyTime, bool bAddToSelection, bool bToggleSelection);
-
 	/** Updates the external selection to match the current sequencer selection. */
 	void SynchronizeExternalSelectionWithSequencerSelection();
 
@@ -498,6 +490,16 @@ public:
 
 	void OnSelectedNodesOnlyChanged();
 
+	void SyncCurveEditorToSelection(bool bOutlinerSelectionChanged);
+
+	TSharedPtr<FCurveEditor> GetCurveEditor() const
+	{
+		return CurveEditorModel;
+	}
+
+	/** Will create a custom menu if FSequencerViewParams::OnBuildCustomContextMenuForGuid is specified. */
+	void BuildCustomContextMenuForGuid(FMenuBuilder& MenuBuilder, FGuid ObjectBinding);
+
 public:
 
 	/** Copy the selection, whether it's keys or tracks */
@@ -536,9 +538,6 @@ public:
 	/** Rebinds all possessable references in the current sequence to update them to the latest referencing mechanism. */
 	void RebindPossessableReferences();
 
-	/** Moves all time data for the current scene onto a valid frame. */
-	void FixFrameTiming();
-
 	/** Imports the animation from an fbx file. */
 	void ImportFBX();
 
@@ -578,7 +577,7 @@ public:
 	/**
 	 * Update auto-scroll mechanics as a result of a new time position
 	 */
-	void UpdateAutoScroll(float NewTime);
+	void UpdateAutoScroll(double NewTime);
 
 public:
 
@@ -592,7 +591,7 @@ public:
 
 	virtual void Tick(float DeltaTime) override;
 	virtual ETickableTickType GetTickableTickType() const override { return ETickableTickType::Always; }
-	virtual TStatId GetStatId() const override { RETURN_QUICK_DECLARE_CYCLE_STAT(FSequencer, STATGROUP_Tickables); };
+	virtual TStatId GetStatId() const override { RETURN_QUICK_DECLARE_CYCLE_STAT(FSequencer, STATGROUP_Tickables); }
 
 public:
 
@@ -602,6 +601,7 @@ public:
 	virtual TSharedRef<SWidget> GetSequencerWidget() const override;
 	virtual FMovieSceneSequenceIDRef GetRootTemplateID() const override { return ActiveTemplateIDs[0]; }
 	virtual FMovieSceneSequenceIDRef GetFocusedTemplateID() const override { return ActiveTemplateIDs.Top(); }
+	virtual bool GetFocusedSequenceIsActive() const override { return ActiveTemplateStates.Top(); }
 	virtual UMovieSceneSequence* GetRootMovieSceneSequence() const override;
 	virtual UMovieSceneSequence* GetFocusedMovieSceneSequence() const override;
 	virtual FMovieSceneRootEvaluationTemplateInstance& GetEvaluationTemplate() override { return RootTemplateInstance; }
@@ -620,12 +620,11 @@ public:
 	virtual bool GetInfiniteKeyAreas() const override;
 	virtual void SetInfiniteKeyAreas(bool bInfiniteKeyAreas) override;
 	virtual bool GetAutoSetTrackDefaults() const override;
-	virtual bool IsRecordingLive() const override;
-	virtual float GetLocalTime() const override;
-	virtual float GetGlobalTime() const override;
-	virtual void SetLocalTime(float Time, ESnapTimeMode SnapTimeMode = ESnapTimeMode::STM_None) override;
-	virtual void SetLocalTimeDirectly(float NewTime) override;
-	virtual void SetGlobalTime(float Time) override;
+	virtual FQualifiedFrameTime GetLocalTime() const override;
+	virtual FQualifiedFrameTime GetGlobalTime() const override;
+	virtual void SetLocalTime(FFrameTime Time, ESnapTimeMode SnapTimeMode = ESnapTimeMode::STM_None) override;
+	virtual void SetLocalTimeDirectly(FFrameTime NewTime) override;
+	virtual void SetGlobalTime(FFrameTime Time) override;
 	virtual void ForceEvaluate() override;
 	virtual void SetPerspectiveViewportPossessionEnabled(bool bEnabled) override;
 	virtual void SetPerspectiveViewportCameraCutEnabled(bool bEnabled) override;
@@ -656,6 +655,8 @@ public:
 	virtual void SelectSection(UMovieSceneSection* Section) override;
 	virtual void SelectByPropertyPaths(const TArray<FString>& InPropertyPaths) override;
 	virtual void EmptySelection() override;
+	virtual void ThrobKeySelection() override;
+	virtual void ThrobSectionSelection() override;
 	virtual FOnGlobalTimeChanged& OnGlobalTimeChanged() override { return OnGlobalTimeChangedDelegate; }
 	virtual FOnPlayEvent& OnPlayEvent() override { return OnPlayDelegate; }
 	virtual FOnStopEvent& OnStopEvent() override { return OnStopDelegate; }
@@ -675,8 +676,7 @@ public:
 	virtual FOnPostSave& OnPostSave() override;
 	virtual FOnActivateSequence& OnActivateSequence() override;
 	virtual FOnCameraCut& OnCameraCut() override;
-	virtual TSharedRef<INumericTypeInterface<float>> GetNumericTypeInterface() override;
-	virtual TSharedRef<INumericTypeInterface<float>> GetZeroPadNumericTypeInterface() override;
+	virtual TSharedRef<INumericTypeInterface<double>> GetNumericTypeInterface() const override;
 	virtual TSharedRef<SWidget> MakeTransportControls(bool bExtended) override;
 	virtual FReply OnPlay(bool bTogglePlay = true) override;
 	virtual void Pause() override;
@@ -685,6 +685,7 @@ public:
 	virtual FGuid MakeNewSpawnable(UObject& SourceObject, UActorFactory* ActorFactory = nullptr) override;
 	virtual bool IsReadOnly() const override;
 	virtual void ExternalSelectionHasChanged() override { SynchronizeSequencerSelectionWithExternalSelection(); }
+	/** Access the user-supplied settings object */
 	virtual USequencerSettings* GetSequencerSettings() override { return Settings; }
 	virtual TSharedPtr<class ITimeSlider> GetTopTimeSliderWidget() const override;
 
@@ -717,7 +718,7 @@ protected:
 	 * If the scene has shots, it only takes the shot section boundaries
 	 * Otherwise, it finds the furthest boundaries of all sections
 	 */
-	TRange<float> GetTimeBounds() const;
+	TRange<FFrameNumber> GetTimeBounds() const;
 	
 	/**
 	 * Gets the time boundaries of the currently filtering shot sections.
@@ -730,7 +731,7 @@ protected:
 	 *
 	 * @param	NewClampRange The new clamp range
 	 */
-	void OnClampRangeChanged( TRange<float> NewClampRange );
+	void OnClampRangeChanged( TRange<double> NewClampRange );
 
 	/*
 	 * Called to get the nearest key
@@ -738,7 +739,7 @@ protected:
 	 * @param InTime The time to get the nearest key to
 	 * @return NearestKey
 	 */
-	float OnGetNearestKey(float InTime);
+	FFrameNumber OnGetNearestKey(FFrameTime InTime);
 
 	/**
 	 * Called when the scrub position is changed by the user
@@ -746,7 +747,7 @@ protected:
 	 *
 	 * @param NewScrubPosition	The new scrub position
 	 */
-	void OnScrubPositionChanged( float NewScrubPosition, bool bScrubbing );
+	void OnScrubPositionChanged( FFrameTime NewScrubPosition, bool bScrubbing );
 
 	/** Called when the user has begun scrubbing */
 	void OnBeginScrubbing();
@@ -766,6 +767,9 @@ protected:
 	/** Called when the user has finished dragging the selection range */
 	void OnSelectionRangeEndDrag();
 
+	/** Get the unqualified local time */
+	FFrameTime GetLocalFrameTime() const { return GetLocalTime().Time; }
+
 protected:
 
 	/**
@@ -776,7 +780,7 @@ protected:
 	/**
 	 * Calculates the amount of encroachment the specified time has into the autoscroll region, if any
 	 */
-	TOptional<float> CalculateAutoscrollEncroachment(float NewTime, float ThresholdPercentage = 0.1f) const;
+	TOptional<float> CalculateAutoscrollEncroachment(double NewTime, float ThresholdPercentage = 0.1f) const;
 
 	/** Called to toggle auto-scroll on and off */
 	void OnToggleAutoScroll();
@@ -797,9 +801,11 @@ protected:
 protected:
 
 	/** Get all the keys for the current sequencer selection */
-	virtual void GetKeysFromSelection(TUniquePtr<ISequencerKeyCollection>& KeyCollection, float DuplicateThresoldTime) override;
+	virtual void GetKeysFromSelection(TUniquePtr<FSequencerKeyCollection>& KeyCollection, float DuplicateThresholdSeconds) override;
+
+	void GetAllKeys(TUniquePtr<FSequencerKeyCollection>& KeyCollection, float DuplicateThresoldSeconds);
 
-	UMovieSceneSection* FindNextOrPreviousShot(UMovieSceneSequence* Sequence, float Time, const bool bNext) const;
+	UMovieSceneSection* FindNextOrPreviousShot(UMovieSceneSequence* Sequence, FFrameNumber SearchFromTime, const bool bNext) const;
 
 protected:
 	
@@ -830,6 +836,9 @@ protected:
 	/** Expand or collapse selected nodes and descendants*/
 	void ToggleExpandCollapseNodesAndDescendants();
 
+	/** Sort all nodes and their descendants by category then alphabetically */
+	void SortAllNodesAndDescendants();
+
 	/** Manually sets a key for the selected objects at the current time */
 	void SetKey();
 
@@ -865,11 +874,17 @@ protected:
 	FMovieScenePossessable* ConvertToPossessableInternal(FGuid SpawnableGuid);
 
 	/** Internal function to render movie for a given start/end time */
-	void RenderMovieInternal(float InStartTime, float InEndTime, bool bSetFrameOverrides = false) const;
+	void RenderMovieInternal(TRange<FFrameNumber> Range, bool bSetFrameOverrides = false) const;
 
 	/** Handles adding a new folder to the outliner tree. */
 	void OnAddFolder();
 
+	/** Handles adding a newly created track to the outliner tree by assigning it into a folder and selecting it. */
+	void OnAddTrack(const TWeakObjectPtr<UMovieSceneTrack>& InTrack);
+
+	/** Determines the selected parent folders and returns the node path to the first folder. Also expands the first folder. */
+	void CalculateSelectedFolderAndPath(TArray<UMovieSceneFolder*>& OutSelectedParentFolders, FString& OutNewNodePath);
+
 	/** Create set playback start transport control */
 	TSharedRef<SWidget> OnCreateTransportSetPlaybackStart();
 
@@ -883,7 +898,7 @@ protected:
 	TSharedRef<SWidget> OnCreateTransportSetPlaybackEnd();
 
 	/** Select keys and/or sections in a display node that fall into the current selection range. */
-	void SelectInSelectionRange(const TSharedRef<FSequencerDisplayNode>& DisplayNode, const TRange<float>& SelectionRange, bool bSelectKeys, bool bSelectSections);
+	void SelectInSelectionRange(const TSharedRef<FSequencerDisplayNode>& DisplayNode, const TRange<FFrameNumber>& SelectionRange, bool bSelectKeys, bool bSelectSections);
 	
 	/** Create loop mode transport control */
 	TSharedRef<SWidget> OnCreateTransportLoopMode();
@@ -906,10 +921,20 @@ protected:
 	/** Check whether we're viewing the master sequence or not */
 	bool IsViewingMasterSequence() const { return ActiveTemplateIDs.Num() == 1; }
 
+	/** Get the default key attributes to apply to newly created keys on the curve editor */
+	FKeyAttributes GetDefaultKeyAttributes() const;
+
+public:
+
+	/** Reset the timing manager to the clock source specified by the root movie scene */
+	void ResetTimeController();
+
+	/** Helper function which returns how many frames (in frame resolution) one display rate frame represents. */
+	double GetPlayRateDeltaFrameCount() const;
 private:
 
-	/** Reset the timing manager to default, or audio clock locked */
-	void ResetTimingManager(bool bUseAudioClock);
+	/** Update the time bases for the current movie scene */
+	void UpdateTimeBases();
 
 	/** User-supplied settings object for this sequencer */
 	USequencerSettings* Settings;
@@ -941,18 +966,25 @@ private:
 	TWeakObjectPtr<UMovieSceneSequence> RootSequence;
 	FMovieSceneRootEvaluationTemplateInstance RootTemplateInstance;
 
+	/** A stack of the current sequence hierarchy for keeping track of nestled sequences. */
 	TArray<FMovieSceneSequenceID> ActiveTemplateIDs;
+	/**
+	* The active state of each sequence. A sequence can be in another sequence multiple times
+	* and the owning subsection contains the active state of the sequence, so this stack is kept 
+	* in sync with the ActiveTemplateIDs as you enter a sequence via the specific subsection node.
+	 */
+	TArray<bool> ActiveTemplateStates;
 
 	FMovieSceneSequenceTransform RootToLocalTransform;
 
 	/** The time range target to be viewed */
-	TRange<float> TargetViewRange;
+	TRange<double> TargetViewRange;
 
 	/** The last time range that was viewed */
-	TRange<float> LastViewRange;
+	TRange<double> LastViewRange;
 
 	/** The view range before zooming */
-	TRange<float> ViewRangeBeforeZoom;
+	TRange<double> ViewRangeBeforeZoom;
 
 	/** The amount of autoscroll pan offset that is currently being applied */
 	TOptional<float> AutoscrollOffset;
@@ -971,10 +1003,6 @@ private:
 	/** Whether we are playing, recording, etc. */
 	EMovieScenePlayerStatus::Type PlaybackState;
 
-	/** The current scrub position */
-	// @todo sequencer: Should use FTimespan or "double" for Time Cursor Position! (cascades)
-	float ScrubPosition;
-
 	/** Current play position */
 	FMovieScenePlaybackPosition PlayPosition;
 
@@ -1064,8 +1092,8 @@ private:
 	double OldMaxTickRate;
 
 	/** Timing manager that can adjust playback times */
-	TUniquePtr<FSequencerTimingManager> TimingManager;
-	
+	TSharedPtr<FMovieSceneTimeController> TimeController;
+
 	struct FCachedViewTarget
 	{
 		/** The player controller we're possessing */
@@ -1093,13 +1121,15 @@ private:
 
 	FAcquiredResources AcquiredResources;
 
-	/** The range of the currently displayed sub sequence in relation to its parent section */
-	TRange<float> SubSequenceRange;
+	/** The range of the currently displayed sub sequence in relation to its parent section, in the resolution of the current sub sequence */
+	TRange<FFrameNumber> SubSequenceRange;
 
 	TSharedPtr<struct FSequencerTemplateStore> TemplateStore;
 
 	TMap<FName, TFunction<void()>> CleanupFunctions;
 
 	/** Transient collection of keys that is used for jumping between keys contained within the current selection */
-	TUniquePtr<ISequencerKeyCollection> SelectedKeyCollection;
+	TUniquePtr<FSequencerKeyCollection> SelectedKeyCollection;
+
+	TSharedPtr<FCurveEditor> CurveEditorModel;
 };
diff --git a/Engine/Source/Editor/Sequencer/Private/SequencerChannelTraits.cpp b/Engine/Source/Editor/Sequencer/Private/SequencerChannelTraits.cpp
new file mode 100644
--- /dev/null
+++ b/Engine/Source/Editor/Sequencer/Private/SequencerChannelTraits.cpp
@@ -0,0 +1,24 @@
+// Copyright 1998-2018 Epic Games, Inc. All Rights Reserved.
+
+#include "SequencerChannelTraits.h"
+#include "EditorStyleSet.h"
+
+namespace Sequencer
+{
+
+
+void DrawKeys(void* Channel, TArrayView<const FKeyHandle> InHandles, TArrayView<FKeyDrawParams> OutKeyDrawParams)
+{
+	// By default just render diamonds for keys
+	FKeyDrawParams DefaultParams;
+	DefaultParams.BorderBrush = DefaultParams.FillBrush = FEditorStyle::Get().GetBrush("Sequencer.KeyDiamond");
+
+	for (FKeyDrawParams& Param : OutKeyDrawParams)
+	{
+		Param = DefaultParams;
+	}
+}
+
+
+
+}	// namespace Sequencer
\ No newline at end of file
diff --git a/Engine/Source/Editor/Sequencer/Private/SequencerCommands.cpp b/Engine/Source/Editor/Sequencer/Private/SequencerCommands.cpp
--- a/Engine/Source/Editor/Sequencer/Private/SequencerCommands.cpp
+++ b/Engine/Source/Editor/Sequencer/Private/SequencerCommands.cpp
@@ -31,8 +31,7 @@ void FSequencerCommands::RegisterCommands()
 	UI_COMMAND( SetSelectionRangeToPreviousShot, "Set Selection Range to Previous Shot", "Set the selection range to the previous shot", EUserInterfaceActionType::Button, FInputChord(EKeys::PageDown) );
 	UI_COMMAND( SetPlaybackRangeToAllShots, "Set Playback Range to All Shots", "Set the playback range to all the shots", EUserInterfaceActionType::Button, FInputChord(EKeys::End) );
 
-	UI_COMMAND( TogglePlaybackRangeLocked, "Locked", "Lock/unlock the playback range.", EUserInterfaceActionType::ToggleButton, FInputChord() );
-	UI_COMMAND( ToggleForceFixedFrameIntervalPlayback, "Force Fixed Frame Interval Playback", "Forces scene evaluation to a fixed frame interval in editor and at runtime, even if this would result in duplicated or dropped frames.", EUserInterfaceActionType::ToggleButton, FInputChord() );
+	UI_COMMAND( TogglePlaybackRangeLocked, "Playback Range Locked", "Prevent editing the start and end times for the sequence.", EUserInterfaceActionType::ToggleButton, FInputChord() );
 	UI_COMMAND( ToggleRerunConstructionScripts, "Rerun Construction Scripts", "Rerun construction scripts on bound actors every frame.", EUserInterfaceActionType::ToggleButton, FInputChord() );
 	
 	UI_COMMAND( ToggleKeepCursorInPlaybackRangeWhileScrubbing, "Keep Cursor in Playback Range While Scrubbing", "When checked, the cursor will be constrained to the current playback range while scrubbing", EUserInterfaceActionType::ToggleButton, FInputChord() );
@@ -43,6 +42,7 @@ void FSequencerCommands::RegisterCommands()
 	UI_COMMAND( CollapseAllNodesAndDescendants, "Collapse All Nodes", "Collapse all selected nodes", EUserInterfaceActionType::Button, FInputChord() );
 	UI_COMMAND( ToggleExpandCollapseNodes, "Expand/Collapse Nodes", "Toggle expand or collapse selected nodes", EUserInterfaceActionType::Button, FInputChord(EModifierKey::Control, EKeys::O) );
 	UI_COMMAND( ToggleExpandCollapseNodesAndDescendants, "Expand/Collapse Nodes and Descendants", "Toggle expand or collapse selected nodes and descendants", EUserInterfaceActionType::Button, FInputChord(EModifierKey::Shift, EKeys::O) );
+	UI_COMMAND( SortAllNodesAndDescendants, "Sort All Nodes", "Sorts all nodes by type and then alphabetically.", EUserInterfaceActionType::Button, FInputChord());
 
 	UI_COMMAND( SetSelectionRangeEnd, "Set Selection End", "Sets the end of the selection range", EUserInterfaceActionType::Button, FInputChord(EKeys::O) );
 	UI_COMMAND( SetSelectionRangeStart, "Set Selection Start", "Sets the start of the selection range", EUserInterfaceActionType::Button, FInputChord(EKeys::I) );
@@ -79,9 +79,9 @@ void FSequencerCommands::RegisterCommands()
 
 	UI_COMMAND( ToggleAutoScroll, "Auto Scroll", "Toggle auto-scroll: When enabled, automatically scrolls the sequencer view to keep the current time visible", EUserInterfaceActionType::ToggleButton, FInputChord(EModifierKey::Shift, EKeys::S) );
 
-	UI_COMMAND( ToggleShowFrameNumbers, "Show Frame Numbers", "Enables and disables showing frame numbers", EUserInterfaceActionType::ToggleButton, FInputChord(EModifierKey::Control, EKeys::T) );
+	UI_COMMAND( ChangeTimeDisplayFormat, "Change Time Display Format", "Rotates through supported display formats for time", EUserInterfaceActionType::Button, FInputChord(EModifierKey::Control, EKeys::T) );
 	UI_COMMAND( ToggleShowGotoBox, "Go to Time...", "Go to a particular point on the timeline", EUserInterfaceActionType::Button, FInputChord(EModifierKey::Control, EKeys::G) );
-	UI_COMMAND( ToggleShowTransformBox, "Transform...", "Transform the selected keys and sections by a given amount", EUserInterfaceActionType::Button, FInputChord(EModifierKey::Control, EKeys::M) );
+	UI_COMMAND( ToggleShowTransformBox, "Transform Selection...", "Transform the selected keys and sections by a given amount", EUserInterfaceActionType::Button, FInputChord(EModifierKey::Control, EKeys::M) );
 	UI_COMMAND( BakeTransform, "Bake Transform", "Bake transform in world space, removing any existing transform and attach tracks", EUserInterfaceActionType::Button, FInputChord() );
 	UI_COMMAND( ToggleShowRangeSlider, "Show Range Slider", "Enables and disables showing the time range slider", EUserInterfaceActionType::ToggleButton, FInputChord() );
 	UI_COMMAND( ToggleIsSnapEnabled, "Enable Snapping", "Enables and disables snapping", EUserInterfaceActionType::ToggleButton, FInputChord() );
@@ -108,6 +108,10 @@ void FSequencerCommands::RegisterCommands()
 	UI_COMMAND( ToggleShowCurveEditor, "Curve Editor", "Show the animation keys in a curve editor", EUserInterfaceActionType::ToggleButton, FInputChord() );
 	UI_COMMAND( ToggleLinkCurveEditorTimeRange, "Link Curve Editor Time Range", "Link the curve editor time range to the sequence", EUserInterfaceActionType::ToggleButton, FInputChord() );
 
+	UI_COMMAND( SetAllCurveVisibility, "All Curves", "Show all curves in the curve editor.", EUserInterfaceActionType::RadioButton, FInputChord() );
+	UI_COMMAND( SetSelectedCurveVisibility, "Selected Curves", "Show only selected curves in the curve editor.", EUserInterfaceActionType::RadioButton, FInputChord() );
+	UI_COMMAND( SetAnimatedCurveVisibility, "Animated Curves", "Show only curves with keyframes in the curve editor.", EUserInterfaceActionType::RadioButton, FInputChord() );
+
 	UI_COMMAND( ToggleShowPreAndPostRoll, "Show Pre/Post Roll", "Toggles visualization of pre and post roll", EUserInterfaceActionType::ToggleButton, FInputChord() );
 
 	UI_COMMAND( RenderMovie, "Render Movie", "Render this movie to a video, or image frame sequence", EUserInterfaceActionType::Button, FInputChord() );
@@ -126,8 +130,6 @@ void FSequencerCommands::RegisterCommands()
 	UI_COMMAND( FixActorReferences, "Fix Actor References", "Try to automatically fix up broken actor bindings.", EUserInterfaceActionType::Button, FInputChord() );
 	UI_COMMAND( RebindPossessableReferences, "Rebind Possesable References", "Rebinds all possessables in the current sequence to ensure they're using the most robust referencing mechanism.", EUserInterfaceActionType::Button, FInputChord() );
 
-	UI_COMMAND( FixFrameTiming, "Fix Frame Timing", "Moves all time data for this sequence to a valid frame time.", EUserInterfaceActionType::Button, FInputChord() );
-
 	UI_COMMAND( RecordSelectedActors, "Record Selected Actors", "Records the selected actors into a new sub sequence of the currently active sequence in Sequencer.", EUserInterfaceActionType::Button, FInputChord(EModifierKey::Alt, EKeys::R) );
 
 	UI_COMMAND( ImportFBX, "Import...", "Imports the animation from an FBX file.", EUserInterfaceActionType::Button, FInputChord() );
diff --git a/Engine/Source/Editor/Sequencer/Private/SequencerCommonHelpers.cpp b/Engine/Source/Editor/Sequencer/Private/SequencerCommonHelpers.cpp
--- a/Engine/Source/Editor/Sequencer/Private/SequencerCommonHelpers.cpp
+++ b/Engine/Source/Editor/Sequencer/Private/SequencerCommonHelpers.cpp
@@ -5,7 +5,6 @@
 #include "DisplayNodes/SequencerSectionKeyAreaNode.h"
 #include "DisplayNodes/SequencerTrackNode.h"
 #include "SSequencer.h"
-#include "GroupedKeyArea.h"
 #include "ISequencerHotspot.h"
 #include "SSequencerTreeView.h"
 #include "VirtualTrackArea.h"
@@ -51,9 +50,9 @@ void SequencerHelpers::GetAllKeyAreas(TSharedPtr<FSequencerDisplayNode> DisplayN
 	}
 }
 
-int32 SequencerHelpers::GetSectionFromTime(TArrayView<UMovieSceneSection* const> InSections, float Time)
+int32 SequencerHelpers::GetSectionFromTime(TArrayView<UMovieSceneSection* const> InSections, FFrameNumber Time)
 {
-	float ClosestLowerBound = TNumericLimits<float>::Max();
+	FFrameNumber ClosestLowerBound = TNumericLimits<int32>::Max();
 	TOptional<int32> MaxOverlapPriority, MaxProximalPriority;
 
 	int32 MostRelevantIndex = INDEX_NONE;
@@ -64,7 +63,7 @@ int32 SequencerHelpers::GetSectionFromTime(TArrayView<UMovieSceneSection* const>
 		if (Section)
 		{
 			const int32 ThisSectionPriority = Section->GetOverlapPriority();
-			TRange<float> SectionRange = Section->IsInfinite() ? TRange<float>::All() : Section->GetRange();
+			TRange<FFrameNumber> SectionRange = Section->GetRange();
 
 			// If the specified time is within the section bounds
 			if (SectionRange.Contains(Time))
@@ -78,7 +77,7 @@ int32 SequencerHelpers::GetSectionFromTime(TArrayView<UMovieSceneSection* const>
 			// Check for nearby sections if there is nothing overlapping
 			else if (!MaxOverlapPriority.IsSet())
 			{
-				const float LowerBoundValue = SectionRange.GetLowerBoundValue();
+				const FFrameNumber LowerBoundValue = SectionRange.GetLowerBoundValue();
 				// If this section exists beyond the current time, we can choose it if its closest to the time
 				if (LowerBoundValue >= Time)
 				{
@@ -176,24 +175,6 @@ bool SequencerHelpers::FindObjectBindingNode(TSharedRef<FSequencerDisplayNode> D
 	return false;
 }
 
-int32 SequencerHelpers::TimeToFrame(float Time, float FrameRate)
-{
-	float Frame = Time * FrameRate;
-	return FMath::RoundToInt(Frame);
-}
-
-float SequencerHelpers::FrameToTime(int32 Frame, float FrameRate)
-{
-	return Frame / FrameRate;
-}
-
-float SequencerHelpers::SnapTimeToInterval(float InTime, float InFrameRate)
-{
-	return InFrameRate > 0
-		? FMath::RoundToInt( InTime / InFrameRate ) * InFrameRate
-		: InTime;
-}
-
 bool IsSectionSelectedInNode(FSequencer& Sequencer, TSharedRef<FSequencerDisplayNode> InNode)
 {
 	if (InNode->GetType() == ESequencerNode::Track)
@@ -213,33 +194,17 @@ bool IsSectionSelectedInNode(FSequencer& Sequencer, TSharedRef<FSequencerDisplay
 
 bool AreKeysSelectedInNode(FSequencer& Sequencer, TSharedRef<FSequencerDisplayNode> InNode)
 {
-	TSharedRef<FSequencerTrackNode> TrackNode = StaticCastSharedRef<FSequencerTrackNode>(InNode);
-
 	TSet<TSharedPtr<IKeyArea>> KeyAreas;
 	SequencerHelpers::GetAllKeyAreas(InNode, KeyAreas);
-	for (auto KeyArea : KeyAreas)
-	{
-		for (auto KeyHandle : KeyArea->GetUnsortedKeyHandles())
-		{
-			FSequencerSelectedKey TestKey(*KeyArea->GetOwningSection(), KeyArea, KeyHandle);
-			if (Sequencer.GetSelection().IsSelected(TestKey))
-			{
-				return true;
-			}
-		}
-	}
 
-	for (TSharedRef<FGroupedKeyArea> KeyGrouping : InNode->GetKeyGroupings())
+	for (const FSequencerSelectedKey& Key : Sequencer.GetSelection().GetSelectedKeys())
 	{
-		for (auto KeyHandle : KeyGrouping->GetUnsortedKeyHandles())
+		if (KeyAreas.Contains(Key.KeyArea))
 		{
-			FSequencerSelectedKey TestKey(*KeyGrouping->GetOwningSection(), KeyGrouping, KeyHandle);
-			if (Sequencer.GetSelection().IsSelected(TestKey))
-			{
-				return true;
-			}
+			return true;
 		}
 	}
+
 	return false;
 }
 
@@ -334,11 +299,18 @@ void SequencerHelpers::PerformDefaultSelection(FSequencer& Sequencer, const FPoi
 	{
 		if (Hotspot->GetType() == ESequencerHotspot::Key)
 		{
-			FSequencerSelectedKey Key = static_cast<FKeyHotspot*>(Hotspot.Get())->Key;
-			if (!Selection.IsSelected(Key))
+			bool bHasClearedSelection = false;
+			for (const FSequencerSelectedKey& Key : static_cast<FKeyHotspot*>(Hotspot.Get())->Keys)
 			{
-				ConditionallyClearSelection();
-				Selection.AddToSelection(Key);
+				if (!Selection.IsSelected(Key))
+				{
+					if (!bHasClearedSelection)
+					{
+						ConditionallyClearSelection();
+						bHasClearedSelection = true;
+					}
+					Selection.AddToSelection(Key);
+				}
 			}
 		}
 		else if (Hotspot->GetType() == ESequencerHotspot::Section || Hotspot->GetType() == ESequencerHotspot::EasingArea)
@@ -381,14 +353,16 @@ void SequencerHelpers::PerformDefaultSelection(FSequencer& Sequencer, const FPoi
 		
 	if (Hotspot->GetType() == ESequencerHotspot::Key)
 	{
-		FSequencerSelectedKey Key = static_cast<FKeyHotspot*>(Hotspot.Get())->Key;
-		if (bForceSelect || !Selection.IsSelected(Key))
+		for (const FSequencerSelectedKey& Key : static_cast<FKeyHotspot*>(Hotspot.Get())->Keys)
 		{
-			Selection.AddToSelection(Key);
-		}
-		else
-		{
-			Selection.RemoveFromSelection(Key);
+			if (bForceSelect || !Selection.IsSelected(Key))
+			{
+				Selection.AddToSelection(Key);
+			}
+			else
+			{
+				Selection.RemoveFromSelection(Key);
+			}
 		}
 	}
 	else if (Hotspot->GetType() == ESequencerHotspot::Section || Hotspot->GetType() == ESequencerHotspot::EasingArea)
@@ -396,7 +370,7 @@ void SequencerHelpers::PerformDefaultSelection(FSequencer& Sequencer, const FPoi
 		UMovieSceneSection* Section = static_cast<FSectionHotspot*>(Hotspot.Get())->Section.GetSectionObject();
 
 		// Never allow infinite sections to be selected through normal click (they're only selectable through right click)
-		if (!Section->IsInfinite())
+		if (Section->GetRange() != TRange<FFrameNumber>::All())
 		{
 			if (bForceSelect || !Selection.IsSelected(Section))
 			{
@@ -428,7 +402,7 @@ TSharedPtr<SWidget> SequencerHelpers::SummonContextMenu(FSequencer& Sequencer, c
 
 	// Attempt to paste into either the current node selection, or the clicked on track
 	TSharedRef<SSequencer> SequencerWidget = StaticCastSharedRef<SSequencer>(Sequencer.GetSequencerWidget());
-	const float PasteAtTime = Sequencer.GetLocalTime();
+	const FFrameNumber PasteAtTime = Sequencer.GetLocalTime().Time.FrameNumber;
 
 	const bool bShouldCloseWindowAfterMenuSelection = true;
 	FMenuBuilder MenuBuilder(bShouldCloseWindowAfterMenuSelection, Sequencer.GetCommandBindings());
diff --git a/Engine/Source/Editor/Sequencer/Private/SequencerCommonHelpers.h b/Engine/Source/Editor/Sequencer/Private/SequencerCommonHelpers.h
--- a/Engine/Source/Editor/Sequencer/Private/SequencerCommonHelpers.h
+++ b/Engine/Source/Editor/Sequencer/Private/SequencerCommonHelpers.h
@@ -6,6 +6,7 @@
 #include "DisplayNodes/SequencerDisplayNode.h"
 #include "Widgets/Input/SNumericDropDown.h"
 #include "Containers/ArrayView.h"
+#include "FrameNumber.h"
 
 #define LOCTEXT_NAMESPACE "SequencerHelpers"
 
@@ -25,7 +26,7 @@ public:
 	 * Get the section index that relates to the specified time
 	 * @return the index corresponding to the highest overlapping section, or nearest section where no section overlaps the current time
 	 */
-	static int32 GetSectionFromTime(TArrayView<UMovieSceneSection* const> InSections, float Time);
+	static int32 GetSectionFromTime(TArrayView<UMovieSceneSection* const> InSections, FFrameNumber Time);
 
 	/**
 	 * Get descendant nodes
@@ -42,19 +43,6 @@ public:
 	 */
 	static bool FindObjectBindingNode(TSharedRef<FSequencerDisplayNode> DisplayNode, TSharedRef<FSequencerDisplayNode>& ObjectBindingNode);
 
-	/**
-	 * Convert time to frame
-	 */
-	static int32 TimeToFrame(float Time, float FrameRate);
-
-	/**
-	 * Convert frame to time
-	 */
-	static float FrameToTime(int32 Frame, float FrameRate);
-
-	/** Snaps a time value in seconds to the frame rate. */
-	static float SnapTimeToInterval(float InTime, float InFrameRate);
-
 	/**
 	 * Validate that the nodes with selected keys or sections actually are true
 	 */
@@ -82,102 +70,4 @@ public:
 	static TSharedPtr<SWidget> SummonContextMenu(FSequencer& Sequencer, const FGeometry& MyGeometry, const FPointerEvent& MouseEvent);
 };
 
-
-class SequencerSnapValues
-{
-public:
-	SequencerSnapValues() {}
-
-	static const TArray<SNumericDropDown<float>::FNamedValue>& GetSnapValues()
-	{
-		static TArray<SNumericDropDown<float>::FNamedValue> SnapValues;
-		static bool SnapValuesInitialized = false;
-
-		if (!SnapValuesInitialized)
-		{
-			SnapValues.Add( SNumericDropDown<float>::FNamedValue( 0.001f, LOCTEXT( "Snap_OneThousandth", "0.001" ), LOCTEXT( "SnapDescription_OneThousandth", "Set snap to 1/1000th" ) ) );
-			SnapValues.Add( SNumericDropDown<float>::FNamedValue( 0.01f, LOCTEXT( "Snap_OneHundredth", "0.01" ), LOCTEXT( "SnapDescription_OneHundredth", "Set snap to 1/100th" ) ) );
-			SnapValues.Add( SNumericDropDown<float>::FNamedValue( 0.1f, LOCTEXT( "Snap_OneTenth", "0.1" ), LOCTEXT( "SnapDescription_OneTenth", "Set snap to 1/10th" ) ) );
-			SnapValues.Add( SNumericDropDown<float>::FNamedValue( 1.0f, LOCTEXT( "Snap_One", "1" ), LOCTEXT( "SnapDescription_One", "Set snap to 1" ) ) );
-			SnapValues.Add( SNumericDropDown<float>::FNamedValue( 10.0f, LOCTEXT( "Snap_Ten", "10" ), LOCTEXT( "SnapDescription_Ten", "Set snap to 10" ) ) );
-			SnapValues.Add( SNumericDropDown<float>::FNamedValue( 100.0f, LOCTEXT( "Snap_OneHundred", "100" ), LOCTEXT( "SnapDescription_OneHundred", "Set snap to 100" ) ) );
-
-			SnapValuesInitialized = true;
-		}
-
-		return SnapValues;
-	}
-
-	static const TArray<SNumericDropDown<float>::FNamedValue>& GetSecondsSnapValues()
-	{
-		static TArray<SNumericDropDown<float>::FNamedValue> SecondsSnapValues;
-		static bool SecondsSnapValuesInitialized = false;
-
-		if (!SecondsSnapValuesInitialized)
-		{
-			SecondsSnapValues.Add( SNumericDropDown<float>::FNamedValue( 0.001f, LOCTEXT( "Snap_OneThousandthSeconds", "0.001s" ), LOCTEXT( "SnapDescription_OneThousandthSeconds", "Set snap to 1/1000th of a second" ) ) );
-			SecondsSnapValues.Add( SNumericDropDown<float>::FNamedValue( 0.01f, LOCTEXT( "Snap_OneHundredthSeconds", "0.01s" ), LOCTEXT( "SnapDescription_OneHundredthSeconds", "Set snap to 1/100th of a second" ) ) );
-			SecondsSnapValues.Add( SNumericDropDown<float>::FNamedValue( 0.1f, LOCTEXT( "Snap_OneTenthSeconds", "0.1s" ), LOCTEXT( "SnapDescription_OneTenthSeconds", "Set snap to 1/10th of a second" ) ) );
-			SecondsSnapValues.Add( SNumericDropDown<float>::FNamedValue( 1.0f, LOCTEXT( "Snap_OneSeconds", "1s" ), LOCTEXT( "SnapDescription_OneSeconds", "Set snap to 1 second" ) ) );
-			SecondsSnapValues.Add( SNumericDropDown<float>::FNamedValue( 10.0f, LOCTEXT( "Snap_TenSeconds", "10s" ), LOCTEXT( "SnapDescription_TenSeconds", "Set snap to 10 seconds" ) ) );
-			SecondsSnapValues.Add( SNumericDropDown<float>::FNamedValue( 100.0f, LOCTEXT( "Snap_OneHundredSeconds", "100s" ), LOCTEXT( "SnapDescription_OneHundredSeconds", "Set snap to 100 seconds" ) ) );
-
-			SecondsSnapValuesInitialized = true;
-		}
-
-		return SecondsSnapValues;
-	}
-
-	static const TArray<SNumericDropDown<float>::FNamedValue>& GetFrameRateSnapValues()
-	{
-		static TArray<SNumericDropDown<float>::FNamedValue> FrameRateSnapValues;
-		static bool FrameRateSnapValuesInitialized = false;
-
-		if (!FrameRateSnapValuesInitialized)
-		{
-			FrameRateSnapValues.Add( SNumericDropDown<float>::FNamedValue( 1 / 15.0f, LOCTEXT( "Snap_15Fps", "15 fps" ), LOCTEXT( "SnapDescription_15Fps", "Set snap to 15 fps" ) ) );
-			FrameRateSnapValues.Add( SNumericDropDown<float>::FNamedValue( 1 / 24.0f, LOCTEXT( "Snap_24Fps", "24 fps (film)" ), LOCTEXT( "SnapDescription_24Fps", "Set snap to 24 fps" ) ) );
-			FrameRateSnapValues.Add( SNumericDropDown<float>::FNamedValue( 1 / 25.0f, LOCTEXT( "Snap_25Fps", "25 fps (PAL/25)" ), LOCTEXT( "SnapDescription_25Fps", "Set snap to 25 fps" ) ) );
-			FrameRateSnapValues.Add( SNumericDropDown<float>::FNamedValue( 1 / 29.97f, LOCTEXT( "Snap_29.97Fps", "29.97 fps (NTSC/30)" ), LOCTEXT( "SnapDescription_29.97Fps", "Set snap to 29.97 fps" ) ) );
-			FrameRateSnapValues.Add( SNumericDropDown<float>::FNamedValue( 1 / 30.0f, LOCTEXT( "Snap_30Fps", "30 fps" ), LOCTEXT( "SnapDescription_30Fps", "Set snap to 30 fps" ) ) );
-			FrameRateSnapValues.Add( SNumericDropDown<float>::FNamedValue( 1 / 48.0f, LOCTEXT( "Snap_48Fps", "48 fps" ), LOCTEXT( "SnapDescription_48Fps", "Set snap to 48 fps" ) ) );
-			FrameRateSnapValues.Add( SNumericDropDown<float>::FNamedValue( 1 / 50.0f, LOCTEXT( "Snap_50Fps", "50 fps (PAL/50)" ), LOCTEXT( "SnapDescription_50Fps", "Set snap to 50 fps" ) ) );
-			FrameRateSnapValues.Add( SNumericDropDown<float>::FNamedValue( 1 / 59.94f, LOCTEXT( "Snap_59.94Fps", "59.94 fps (NTSC/60)" ), LOCTEXT( "SnapDescription_59.94Fps", "Set snap to 59.94 fps" ) ) );
-			FrameRateSnapValues.Add( SNumericDropDown<float>::FNamedValue( 1 / 60.0f, LOCTEXT( "Snap_60Fps", "60 fps" ), LOCTEXT( "SnapDescription_60Fps", "Set snap to 60 fps" ) ) );
-			FrameRateSnapValues.Add( SNumericDropDown<float>::FNamedValue( 1 / 120.0f, LOCTEXT( "Snap_120Fps", "120 fps" ), LOCTEXT( "SnapDescription_120Fps", "Set snap to 120 fps" ) ) );
-
-			FrameRateSnapValuesInitialized = true;
-		}
-
-		return FrameRateSnapValues;
-	}
-
-	static const TArray<SNumericDropDown<float>::FNamedValue>& GetTimeSnapValues()
-	{
-		static TArray<SNumericDropDown<float>::FNamedValue> TimeSnapValues;
-		static bool TimeSnapValuesInitialized = false;
-		if (!TimeSnapValuesInitialized)
-		{
-			TimeSnapValues.Append(GetSecondsSnapValues());
-			TimeSnapValues.Append(GetFrameRateSnapValues());
-
-			TimeSnapValuesInitialized = true;
-		}
-		return TimeSnapValues;
-	}
-
-	static bool IsTimeSnapIntervalFrameRate(float InFrameRate)
-	{
-		for (auto FrameRateSnapValue : GetFrameRateSnapValues())
-		{
-			if (FMath::IsNearlyEqual(InFrameRate, FrameRateSnapValue.GetValue()))
-			{
-				return true;
-			}
-		}
-		return false;
-	}
-};
-
-
 #undef LOCTEXT_NAMESPACE
diff --git a/Engine/Source/Editor/Sequencer/Private/SequencerContextMenus.cpp b/Engine/Source/Editor/Sequencer/Private/SequencerContextMenus.cpp
--- a/Engine/Source/Editor/Sequencer/Private/SequencerContextMenus.cpp
+++ b/Engine/Source/Editor/Sequencer/Private/SequencerContextMenus.cpp
@@ -20,14 +20,11 @@
 #include "IDetailsView.h"
 #include "IStructureDetailsView.h"
 #include "PropertyEditorModule.h"
-#include "IntegralKeyDetailsCustomization.h"
 #include "Sections/MovieSceneSubSection.h"
 #include "Sections/MovieSceneCinematicShotSection.h"
 #include "Curves/IntegralCurve.h"
 #include "Editor.h"
 #include "SequencerUtilities.h"
-#include "Misc/NotifyHook.h"
-#include "EditorUndoClient.h"
 #include "ClassViewerModule.h"
 #include "Generators/MovieSceneEasingFunction.h"
 #include "ClassViewerFilter.h"
@@ -35,9 +32,62 @@
 #include "Widgets/Input/SCheckBox.h"
 #include "Widgets/SBoxPanel.h"
 #include "Widgets/Layout/SBox.h"
+#include "ISequencerChannelInterface.h"
+#include "Channels/MovieSceneChannelProxy.h"
+#include "SKeyEditInterface.h"
+#include "MovieSceneTimeHelpers.h"
+#include "FrameNumberDetailsCustomization.h"
+#include "MovieSceneSectionDetailsCustomization.h"
+#include "MovieSceneSequence.h"
+#include "MovieScene.h"
 
 #define LOCTEXT_NAMESPACE "SequencerContextMenus"
 
+static void CreateKeyStructForSelection(TSharedPtr<ISequencer> InSequencer, TSharedPtr<FStructOnScope>& OutKeyStruct, TWeakObjectPtr<UMovieSceneSection>& OutKeyStructSection)
+{
+	const TSet<FSequencerSelectedKey>& SelectedKeys = InSequencer->GetSelection().GetSelectedKeys();
+
+	if (SelectedKeys.Num() == 1)
+	{
+		for (const FSequencerSelectedKey& Key : SelectedKeys)
+		{
+			if (Key.KeyArea.IsValid() && Key.KeyHandle.IsSet())
+			{
+				OutKeyStruct = Key.KeyArea->GetKeyStruct(Key.KeyHandle.GetValue());
+				OutKeyStructSection = Key.KeyArea->GetOwningSection();
+				return;
+			}
+		}
+	}
+	else
+	{
+		TArray<FKeyHandle> KeyHandles;
+		UMovieSceneSection* CommonSection = nullptr;
+		for (const FSequencerSelectedKey& Key : SelectedKeys)
+		{
+			if (Key.KeyArea.IsValid() && Key.KeyHandle.IsSet())
+			{
+				KeyHandles.Add(Key.KeyHandle.GetValue());
+
+				if (!CommonSection)
+				{
+					CommonSection = Key.KeyArea->GetOwningSection();
+				}
+				else if (CommonSection != Key.KeyArea->GetOwningSection())
+				{
+					CommonSection = nullptr;
+					return;
+				}
+			}
+		}
+
+		if (CommonSection)
+		{
+			OutKeyStruct = CommonSection->GetKeyStruct(KeyHandles);
+			OutKeyStructSection = CommonSection;
+		}
+	}
+}
 
 void FKeyContextMenu::BuildMenu(FMenuBuilder& MenuBuilder, FSequencer& InSequencer)
 {
@@ -45,13 +95,40 @@ void FKeyContextMenu::BuildMenu(FMenuBuilder& MenuBuilder, FSequencer& InSequenc
 	Menu->PopulateMenu(MenuBuilder);
 }
 
-
 void FKeyContextMenu::PopulateMenu(FMenuBuilder& MenuBuilder)
 {
 	FSequencer* SequencerPtr = &Sequencer.Get();
 	TSharedRef<FKeyContextMenu> Shared = AsShared();
 
-	if(CanAddPropertiesMenu())
+	CreateKeyStructForSelection(Sequencer, KeyStruct, KeyStructSection);
+
+	{
+		ISequencerModule& SequencerModule = FModuleManager::LoadModuleChecked<ISequencerModule>("Sequencer");
+
+		FSelectedKeysByChannelType SelectedKeysByChannel(SequencerPtr->GetSelection().GetSelectedKeys().Array());
+
+		TMap<uint32, TArray<TChannelAndHandles<void>>> ChannelAndHandlesByType;
+		for (auto& Pair : SelectedKeysByChannel.ChannelToKeyHandleMap)
+		{
+			TChannelAndHandles<void> ChannelAndHandles;
+			ChannelAndHandles.Section = Pair.Value.OwningSection;
+			ChannelAndHandles.Channel = Pair.Value.OwningSection->GetChannelProxy().MakeHandle(Pair.Key);
+			ChannelAndHandles.Handles = MoveTemp(Pair.Value.KeyHandles);
+
+			ChannelAndHandlesByType.FindOrAdd(Pair.Value.ChannelTypeID).Add(MoveTemp(ChannelAndHandles));
+		}
+
+		for (auto& Pair : ChannelAndHandlesByType)
+		{
+			ISequencerChannelInterface* ChannelInterface = SequencerModule.FindChannelInterface(Pair.Key);
+			if (ChannelInterface)
+			{
+				ChannelInterface->ExtendKeyMenu_Raw(MenuBuilder, MoveTemp(Pair.Value), Sequencer);
+			}
+		}
+	}
+
+	if(KeyStruct.IsValid())
 	{
 		MenuBuilder.AddSubMenu(
 			LOCTEXT("KeyProperties", "Properties"),
@@ -60,7 +137,7 @@ void FKeyContextMenu::PopulateMenu(FMenuBuilder& MenuBuilder)
 			FUIAction (
 				FExecuteAction(),
 				// @todo sequencer: only one struct per structure view supported right now :/
-				FCanExecuteAction::CreateLambda([=]{ return (Shared->Sequencer->GetSelection().GetSelectedKeys().Num() == 1); })
+				FCanExecuteAction::CreateLambda([=]{ return KeyStruct.IsValid(); })
 			),
 			NAME_None,
 			EUserInterfaceActionType::Button
@@ -79,80 +156,14 @@ void FKeyContextMenu::PopulateMenu(FMenuBuilder& MenuBuilder)
 	}
 	MenuBuilder.EndSection(); // SequencerKeyEdit
 
-	MenuBuilder.BeginSection("SequencerInterpolation", LOCTEXT("KeyInterpolationMenu", "Key Interpolation"));
-	{
-		MenuBuilder.AddMenuEntry(
-			LOCTEXT("SetKeyInterpolationAuto", "Cubic (Auto)"),
-			LOCTEXT("SetKeyInterpolationAutoTooltip", "Set key interpolation to auto"),
-			FSlateIcon(FEditorStyle::GetStyleSetName(), "Sequencer.IconKeyAuto"),
-			FUIAction(
-				FExecuteAction::CreateSP(SequencerPtr, &FSequencer::SetInterpTangentMode, RCIM_Cubic, RCTM_Auto),
-				FCanExecuteAction(),
-				FIsActionChecked::CreateSP(SequencerPtr, &FSequencer::IsInterpTangentModeSelected, RCIM_Cubic, RCTM_Auto) ),
-			NAME_None,
-			EUserInterfaceActionType::ToggleButton
-		);
 
-		MenuBuilder.AddMenuEntry(
-			LOCTEXT("SetKeyInterpolationUser", "Cubic (User)"),
-			LOCTEXT("SetKeyInterpolationUserTooltip", "Set key interpolation to user"),
-			FSlateIcon(FEditorStyle::GetStyleSetName(), "Sequencer.IconKeyUser"),
-			FUIAction(
-				FExecuteAction::CreateSP(SequencerPtr, &FSequencer::SetInterpTangentMode, RCIM_Cubic, RCTM_User),
-				FCanExecuteAction(),
-				FIsActionChecked::CreateSP(SequencerPtr, &FSequencer::IsInterpTangentModeSelected, RCIM_Cubic, RCTM_User) ),
-			NAME_None,
-			EUserInterfaceActionType::ToggleButton
-		);
-
-		MenuBuilder.AddMenuEntry(
-			LOCTEXT("SetKeyInterpolationBreak", "Cubic (Break)"),
-			LOCTEXT("SetKeyInterpolationBreakTooltip", "Set key interpolation to break"),
-			FSlateIcon(FEditorStyle::GetStyleSetName(), "Sequencer.IconKeyBreak"),
-			FUIAction(
-				FExecuteAction::CreateSP(SequencerPtr, &FSequencer::SetInterpTangentMode, RCIM_Cubic, RCTM_Break),
-				FCanExecuteAction(),
-				FIsActionChecked::CreateSP(SequencerPtr, &FSequencer::IsInterpTangentModeSelected, RCIM_Cubic, RCTM_Break) ),
-			NAME_None,
-			EUserInterfaceActionType::ToggleButton
-		);
-
-		MenuBuilder.AddMenuEntry(
-			LOCTEXT("SetKeyInterpolationLinear", "Linear"),
-			LOCTEXT("SetKeyInterpolationLinearTooltip", "Set key interpolation to linear"),
-			FSlateIcon(FEditorStyle::GetStyleSetName(), "Sequencer.IconKeyLinear"),
-			FUIAction(
-				FExecuteAction::CreateSP(SequencerPtr, &FSequencer::SetInterpTangentMode, RCIM_Linear, RCTM_Auto),
-				FCanExecuteAction(),
-				FIsActionChecked::CreateSP(SequencerPtr, &FSequencer::IsInterpTangentModeSelected, RCIM_Linear, RCTM_Auto) ),
-			NAME_None,
-			EUserInterfaceActionType::ToggleButton
-		);
-
-		MenuBuilder.AddMenuEntry(
-			LOCTEXT("SetKeyInterpolationConstant", "Constant"),
-			LOCTEXT("SetKeyInterpolationConstantTooltip", "Set key interpolation to constant"),
-			FSlateIcon(FEditorStyle::GetStyleSetName(), "Sequencer.IconKeyConstant"),
-			FUIAction(
-				FExecuteAction::CreateSP(SequencerPtr, &FSequencer::SetInterpTangentMode, RCIM_Constant, RCTM_Auto),
-				FCanExecuteAction(),
-				FIsActionChecked::CreateSP(SequencerPtr, &FSequencer::IsInterpTangentModeSelected, RCIM_Constant, RCTM_Auto) ),
-			NAME_None,
-			EUserInterfaceActionType::ToggleButton
-		);
-	}
-	MenuBuilder.EndSection(); // SequencerInterpolation
 
 	MenuBuilder.BeginSection("SequencerKeys", LOCTEXT("KeysMenu", "Keys"));
 	{
-		const bool bUseFrames = SequencerSnapValues::IsTimeSnapIntervalFrameRate(Sequencer->GetFixedFrameInterval());
-
-		MenuBuilder.AddMenuEntry(
-			bUseFrames ? LOCTEXT("SetKeyFrame", "Set Key Frame") : LOCTEXT("SetKeyTime", "Set Key Time"),
-			bUseFrames ? LOCTEXT("SetKeyFrameTooltip", "Set key frame") : LOCTEXT("SetKeyTimeTooltip", "Set key time"),
+		MenuBuilder.AddMenuEntry(LOCTEXT("SetKeyTime", "Set Key Time"), LOCTEXT("SetKeyTimeTooltip", "Set the key to a specified time"),
 			FSlateIcon(),
 			FUIAction(
-				FExecuteAction::CreateSP(SequencerPtr, &FSequencer::SetKeyTime, bUseFrames),
+				FExecuteAction::CreateSP(SequencerPtr, &FSequencer::SetKeyTime),
 				FCanExecuteAction::CreateSP(SequencerPtr, &FSequencer::CanSetKeyTime))
 		);
 
@@ -175,175 +186,47 @@ void FKeyContextMenu::PopulateMenu(FMenuBuilder& MenuBuilder)
 	MenuBuilder.EndSection(); // SequencerKeys
 }
 
-bool FKeyContextMenu::CanAddPropertiesMenu() const
+void FKeyContextMenu::AddPropertiesMenu(FMenuBuilder& MenuBuilder)
 {
-	for (const FSequencerSelectedKey& Key : Sequencer->GetSelection().GetSelectedKeys())
+	auto UpdateAndRetrieveEditData = [this]
 	{
-		if (Key.KeyArea.IsValid() && Key.KeyHandle.IsSet())
-		{
-			TSharedPtr<FStructOnScope> KeyStruct = Key.KeyArea->GetKeyStruct(Key.KeyHandle.GetValue());
-
-			if (KeyStruct.IsValid())
-			{
-				return true;
-			}
-		}
-	}
+		FKeyEditData EditData;
+		CreateKeyStructForSelection(Sequencer, EditData.KeyStruct, EditData.OwningSection);
+		return EditData;
+	};
 
-	return false;
+	MenuBuilder.AddWidget(
+		SNew(SKeyEditInterface, Sequencer)
+		.EditData_Lambda(UpdateAndRetrieveEditData)
+	, FText::GetEmpty(), true);
 }
 
-/** Widget that represents a details panel that refreshes on undo, and handles modification of the section on edit */
-class SInlineDetailsView : public SCompoundWidget, public FEditorUndoClient, public FNotifyHook
-{
-public:
-	SLATE_BEGIN_ARGS(SInlineDetailsView){}
-	SLATE_END_ARGS()
-
-	~SInlineDetailsView()
-	{
-		GEditor->UnregisterForUndo(this);
-	}
 
-	void Construct(const FArguments&, TSharedRef<FSequencer> InSequencer)
-	{
-		GEditor->RegisterForUndo(this);
-		
-		WeakSequencer = InSequencer;
-		Initialize();
-	}
-	
-	/** (Re)Initialize this widget's details panel */
-	void Initialize()
+FSectionContextMenu::FSectionContextMenu(FSequencer& InSeqeuncer, FFrameTime InMouseDownTime)
+	: Sequencer(StaticCastSharedRef<FSequencer>(InSeqeuncer.AsShared()))
+	, MouseDownTime(InMouseDownTime)
+{
+	for (TWeakObjectPtr<UMovieSceneSection> WeakSection : Sequencer->GetSelection().GetSelectedSections())
 	{
-		// Reset the section and widget content
-		WeakSection = nullptr;
-		ChildSlot
-		[
-			SNullWidget::NullWidget
-		];
-
-		TSharedPtr<FSequencer> Sequencer = WeakSequencer.Pin();
-		if (!Sequencer.IsValid())
-		{
-			return;
-		}
-
-		// Set up the details panel only if a single selected key with a valid key struct exists
-		TSharedPtr<FStructOnScope> SelectedKeyStruct;
-		FSequencerSelectedKey SelectedKey;
-
-		for (const FSequencerSelectedKey& Key : Sequencer->GetSelection().GetSelectedKeys())
+		if (UMovieSceneSection* Section = WeakSection.Get())
 		{
-			if (Key.KeyArea.IsValid() && Key.KeyHandle.IsSet())
+			FMovieSceneChannelProxy& ChannelProxy = Section->GetChannelProxy();
+			for (const FMovieSceneChannelEntry& Entry : ChannelProxy.GetAllEntries())
 			{
-				TSharedPtr<FStructOnScope> KeyStruct = Key.KeyArea->GetKeyStruct(Key.KeyHandle.GetValue());
-				if (KeyStruct.IsValid())
+				TArray<UMovieSceneSection*>& SectionArray = SectionsByType.FindOrAdd(Entry.GetChannelID());
+				SectionArray.Add(Section);
+
+				TArray<TMovieSceneChannelHandle<void>>& ChannelArray = ChannelsByType.FindOrAdd(Entry.GetChannelID());
+				for (void* Channel : Entry.GetChannels())
 				{
-					if (SelectedKey.IsValid())
-					{
-						// @todo sequencer: only one struct per structure view supported right now :/
-						return;
-					}
-
-					SelectedKey = Key;
-					SelectedKeyStruct = KeyStruct;
+					ChannelArray.Add(ChannelProxy.MakeHandle(Channel));
 				}
 			}
 		}
-
-		// If there're no selected keys, or too many, bail
-		if (!SelectedKey.IsValid())
-		{
-			return;
-		}
-
-		// Set up the details panel
-		WeakSection = SelectedKey.Section;
-
-		FDetailsViewArgs DetailsViewArgs;
-		{
-			DetailsViewArgs.bAllowSearch = false;
-			DetailsViewArgs.bCustomFilterAreaLocation = true;
-			DetailsViewArgs.bCustomNameAreaLocation = true;
-			DetailsViewArgs.bHideSelectionTip = true;
-			DetailsViewArgs.bLockable = false;
-			DetailsViewArgs.bSearchInitialKeyFocus = true;
-			DetailsViewArgs.bUpdatesFromSelection = false;
-			DetailsViewArgs.bShowOptions = false;
-			DetailsViewArgs.bShowModifiedPropertiesOption = false;
-			DetailsViewArgs.bShowScrollBar = false;
-			DetailsViewArgs.NotifyHook = this;
-		}
-
-		FStructureDetailsViewArgs StructureViewArgs;
-		{
-			StructureViewArgs.bShowObjects = false;
-			StructureViewArgs.bShowAssets = true;
-			StructureViewArgs.bShowClasses = true;
-			StructureViewArgs.bShowInterfaces = false;
-		}
-
-		TSharedRef<IStructureDetailsView> StructureDetailsView = FModuleManager::GetModuleChecked<FPropertyEditorModule>("PropertyEditor")
-			.CreateStructureDetailView(DetailsViewArgs, StructureViewArgs, nullptr, LOCTEXT("MessageData", "Message Data"));
-
-		// register details customizations for this instance
-		TWeakObjectPtr<const UMovieSceneSection> ConstWeakSection = WeakSection;
-		StructureDetailsView->GetDetailsView()->RegisterInstancedCustomPropertyLayout(FIntegralKey::StaticStruct(), FOnGetDetailCustomizationInstance::CreateStatic(&FIntegralKeyDetailsCustomization::MakeInstance, ConstWeakSection));
-
-		StructureDetailsView->SetStructureData(SelectedKeyStruct);
-		StructureDetailsView->GetOnFinishedChangingPropertiesDelegate().AddSP(this, &SInlineDetailsView::OnFinishedChangingProperties, SelectedKeyStruct);
-
-		ChildSlot
-		[
-			StructureDetailsView->GetWidget().ToSharedRef()
-		];
-	}
-
-	void OnFinishedChangingProperties(const FPropertyChangedEvent& ChangeEvent, TSharedPtr<FStructOnScope> KeyStruct)
-	{
-		if (KeyStruct->GetStruct()->IsChildOf(FMovieSceneKeyStruct::StaticStruct()))
-		{
-			((FMovieSceneKeyStruct*)KeyStruct->GetStructMemory())->PropagateChanges(ChangeEvent);
-		}
-
-		TSharedPtr<FSequencer> Sequencer = WeakSequencer.Pin();
-		if (Sequencer.IsValid())
-		{
-			Sequencer->NotifyMovieSceneDataChanged( EMovieSceneDataChangeType::TrackValueChanged );
-		}
-	}
-
-	virtual void NotifyPreChange( UProperty* PropertyAboutToChange )
-	{
-		if (UMovieSceneSection* Section = WeakSection.Get())
-		{
-			Section->Modify();
-		}
-	}
-
-	virtual void PostUndo(bool bSuccess) override
-	{
-		Initialize();
-	}
-
-	virtual void PostRedo(bool bSuccess) override
-	{
-		Initialize();
 	}
-
-private:
-	TWeakObjectPtr<UMovieSceneSection> WeakSection;
-	TWeakPtr<FSequencer> WeakSequencer;
-};
-
-void FKeyContextMenu::AddPropertiesMenu(FMenuBuilder& MenuBuilder)
-{
-	MenuBuilder.AddWidget(SNew(SInlineDetailsView, Sequencer), FText::GetEmpty(), true);
 }
 
-
-void FSectionContextMenu::BuildMenu(FMenuBuilder& MenuBuilder, FSequencer& InSequencer, float InMouseDownTime)
+void FSectionContextMenu::BuildMenu(FMenuBuilder& MenuBuilder, FSequencer& InSequencer, FFrameTime InMouseDownTime)
 {
 	TSharedRef<FSectionContextMenu> Menu = MakeShareable(new FSectionContextMenu(InSequencer, InMouseDownTime));
 	Menu->PopulateMenu(MenuBuilder);
@@ -355,6 +238,19 @@ void FSectionContextMenu::PopulateMenu(FMenuBuilder& MenuBuilder)
 	// Copy a reference to the context menu by value into each lambda handler to ensure the type stays alive until the menu is closed
 	TSharedRef<FSectionContextMenu> Shared = AsShared();
 
+	ISequencerModule& SequencerModule = FModuleManager::LoadModuleChecked<ISequencerModule>("Sequencer");
+
+	for (auto& Pair : ChannelsByType)
+	{
+		const TArray<UMovieSceneSection*>& Sections = SectionsByType.FindChecked(Pair.Key);
+
+		ISequencerChannelInterface* ChannelInterface = SequencerModule.FindChannelInterface(Pair.Key);
+		if (ChannelInterface)
+		{
+			ChannelInterface->ExtendSectionMenu_Raw(MenuBuilder, Pair.Value, Sections, Sequencer);
+		}
+	}
+
 	MenuBuilder.AddSubMenu(
 		LOCTEXT("SectionProperties", "Properties"),
 		LOCTEXT("SectionPropertiesTooltip", "Modify the section properties"),
@@ -368,7 +264,7 @@ void FSectionContextMenu::PopulateMenu(FMenuBuilder& MenuBuilder)
 
 		if (Sequencer->GetClipboardStack().Num() != 0)
 		{
-			FPasteContextMenuArgs PasteArgs = FPasteContextMenuArgs::PasteAt(MouseDownTime);
+			FPasteContextMenuArgs PasteArgs = FPasteContextMenuArgs::PasteAt(MouseDownTime.FrameNumber);
 			PasteMenu = FPasteContextMenu::CreateMenu(*Sequencer, PasteArgs);
 			PasteFromHistoryMenu = FPasteFromHistoryContextMenu::CreateMenu(*Sequencer, PasteArgs);
 		}
@@ -438,23 +334,10 @@ void FSectionContextMenu::PopulateMenu(FMenuBuilder& MenuBuilder)
 			LOCTEXT("EditSectionTooltip", "Edit section"),
 			FNewMenuDelegate::CreateLambda([=](FMenuBuilder& InMenuBuilder){ Shared->AddEditMenu(InMenuBuilder); }));
 
-		if (CanSetExtrapolationMode())
-		{
-			MenuBuilder.AddSubMenu(
-				LOCTEXT("SetPreInfinityExtrap", "Pre-Infinity"),
-				LOCTEXT("SetPreInfinityExtrapTooltip", "Set pre-infinity extrapolation"),
-				FNewMenuDelegate::CreateLambda([=](FMenuBuilder& SubMenuBuilder){ Shared->AddExtrapolationMenu(SubMenuBuilder, true); }));
-
-			MenuBuilder.AddSubMenu(
-				LOCTEXT("SetPostInfinityExtrap", "Post-Infinity"),
-				LOCTEXT("SetPostInfinityExtrapTooltip", "Set post-infinity extrapolation"),
-				FNewMenuDelegate::CreateLambda([=](FMenuBuilder& SubMenuBuilder){ Shared->AddExtrapolationMenu(SubMenuBuilder, false); }));
-		}
-
 		MenuBuilder.AddSubMenu(
 			LOCTEXT("OrderSection", "Order"),
 			LOCTEXT("OrderSectionTooltip", "Order section"),
-			FNewMenuDelegate::CreateLambda([=](FMenuBuilder& SubMenuBuilder){ Shared->AddOrderMenu(SubMenuBuilder); }));			
+			FNewMenuDelegate::CreateLambda([=](FMenuBuilder& SubMenuBuilder){ Shared->AddOrderMenu(SubMenuBuilder); }));
 
 		if (GetSupportedBlendTypes().Num() > 1)
 		{
@@ -543,78 +426,6 @@ void FSectionContextMenu::AddEditMenu(FMenuBuilder& MenuBuilder)
 	);
 }
 
-
-void FSectionContextMenu::AddExtrapolationMenu(FMenuBuilder& MenuBuilder, bool bPreInfinity)
-{
-	// Copy a reference to the context menu by value into each lambda handler to ensure the type stays alive until the menu is closed
-	TSharedRef<FSectionContextMenu> Shared = AsShared();
-
-	MenuBuilder.AddMenuEntry(
-		LOCTEXT("SetExtrapCycle", "Cycle"),
-		LOCTEXT("SetExtrapCycleTooltip", "Set extrapolation cycle"),
-		FSlateIcon(),
-		FUIAction(
-			FExecuteAction::CreateLambda([=]{ Shared->SetExtrapolationMode(RCCE_Cycle, bPreInfinity); }),
-			FCanExecuteAction(),
-			FIsActionChecked::CreateLambda([=]{ return Shared->IsExtrapolationModeSelected(RCCE_Cycle, bPreInfinity); })
-		),
-		NAME_None,
-		EUserInterfaceActionType::RadioButton
-	);
-
-	MenuBuilder.AddMenuEntry(
-		LOCTEXT("SetExtrapCycleWithOffset", "Cycle with Offset"),
-		LOCTEXT("SetExtrapCycleWithOffsetTooltip", "Set extrapolation cycle with offset"),
-		FSlateIcon(),
-		FUIAction(
-			FExecuteAction::CreateLambda([=]{ Shared->SetExtrapolationMode(RCCE_CycleWithOffset, bPreInfinity); }),
-			FCanExecuteAction(),
-			FIsActionChecked::CreateLambda([=]{ return Shared->IsExtrapolationModeSelected(RCCE_CycleWithOffset, bPreInfinity); })
-		),
-		NAME_None,
-		EUserInterfaceActionType::RadioButton
-	);
-
-	MenuBuilder.AddMenuEntry(
-		LOCTEXT("SetExtrapOscillate", "Oscillate"),
-		LOCTEXT("SetExtrapOscillateTooltip", "Set extrapolation oscillate"),
-		FSlateIcon(),
-		FUIAction(
-			FExecuteAction::CreateLambda([=]{ Shared->SetExtrapolationMode(RCCE_Oscillate, bPreInfinity); }),
-			FCanExecuteAction(),
-			FIsActionChecked::CreateLambda([=]{ return Shared->IsExtrapolationModeSelected(RCCE_Oscillate, bPreInfinity); })
-		),
-		NAME_None,
-		EUserInterfaceActionType::RadioButton
-	);
-
-	MenuBuilder.AddMenuEntry(
-		LOCTEXT("SetExtrapLinear", "Linear"),
-		LOCTEXT("SetExtrapLinearTooltip", "Set extrapolation linear"),
-		FSlateIcon(),
-		FUIAction(
-			FExecuteAction::CreateLambda([=]{ Shared->SetExtrapolationMode(RCCE_Linear, bPreInfinity); }),
-			FCanExecuteAction(),
-			FIsActionChecked::CreateLambda([=]{ return Shared->IsExtrapolationModeSelected(RCCE_Linear, bPreInfinity); })
-		),
-		NAME_None,
-		EUserInterfaceActionType::RadioButton
-	);
-
-	MenuBuilder.AddMenuEntry(
-		LOCTEXT("SetExtrapConstant", "Constant"),
-		LOCTEXT("SetExtrapConstantTooltip", "Set extrapolation constant"),
-		FSlateIcon(),
-		FUIAction(
-			FExecuteAction::CreateLambda([=]{ Shared->SetExtrapolationMode(RCCE_Constant, bPreInfinity); }),
-			FCanExecuteAction(),
-			FIsActionChecked::CreateLambda([=]{ return Shared->IsExtrapolationModeSelected(RCCE_Constant, bPreInfinity); })
-		),
-		NAME_None,
-		EUserInterfaceActionType::RadioButton
-	);
-}
-
 FMovieSceneBlendTypeField FSectionContextMenu::GetSupportedBlendTypes() const
 {
 	FMovieSceneBlendTypeField BlendTypes = FMovieSceneBlendTypeField::All();
@@ -692,7 +503,16 @@ void FSectionContextMenu::AddPropertiesMenu(FMenuBuilder& MenuBuilder)
 		}
 	}
 
+	// We pass the current scene to the UMovieSceneSection customization so we can get the overall bounds of the section when we change a section from infinite->bounded.
+	UMovieScene* CurrentScene = Sequencer->GetFocusedMovieSceneSequence()->GetMovieScene();
+
 	TSharedRef<IDetailsView> DetailsView = FModuleManager::GetModuleChecked<FPropertyEditorModule>("PropertyEditor").CreateDetailView(DetailsViewArgs);
+	DetailsView->RegisterInstancedCustomPropertyTypeLayout("FrameNumber", FOnGetPropertyTypeCustomizationInstance::CreateLambda([=]() {
+		return MakeShared<FFrameNumberDetailsCustomization>(Sequencer->GetNumericTypeInterface()); }));
+	DetailsView->RegisterInstancedCustomPropertyLayout(UMovieSceneSection::StaticClass(), FOnGetDetailCustomizationInstance::CreateLambda([=]() {
+		return MakeShared<FMovieSceneSectionDetailsCustomization>(Sequencer->GetNumericTypeInterface(), CurrentScene); }));
+
+
 	Sequencer->OnInitializeDetailsPanel().Broadcast(DetailsView, Sequencer);
 	DetailsView->SetObjects(Sections);
 
@@ -751,10 +571,12 @@ void FSectionContextMenu::SelectAllKeys()
 		FSectionLayout Layout(*Handle.TrackNode, Handle.SectionIndex);
 		for (const FSectionLayoutElement& Element : Layout.GetElements())
 		{
-			TSharedPtr<IKeyArea> KeyArea = Element.GetKeyArea();
-			if (KeyArea.IsValid())
+			for (TSharedRef<IKeyArea> KeyArea : Element.GetKeyAreas())
 			{
-				for (const FKeyHandle& KeyHandle : KeyArea->GetUnsortedKeyHandles())
+				TArray<FKeyHandle> Handles;
+				KeyArea->GetKeyHandles(Handles);
+
+				for (FKeyHandle KeyHandle : Handles)
 				{
 					FSequencerSelectedKey SelectKey(*Section, KeyArea, KeyHandle);
 					Sequencer->GetSelection().AddToSelection(SelectKey);
@@ -820,42 +642,24 @@ bool FSectionContextMenu::CanPrimeForRecording() const
 
 bool FSectionContextMenu::CanSelectAllKeys() const
 {
-	TArray<FSectionHandle> SelectedSections = StaticCastSharedRef<SSequencer>(Sequencer->GetSequencerWidget())->GetSectionHandles(Sequencer->GetSelection().GetSelectedSections());
-	for (const FSectionHandle& Handle : SelectedSections)
-	{
-		UMovieSceneSection* Section = Handle.GetSectionObject();
-		if (!Section)
-		{
-			continue;
-		}
+	ISequencerModule& SequencerModule = FModuleManager::LoadModuleChecked<ISequencerModule>("Sequencer");
 
-		FSectionLayout Layout(*Handle.TrackNode, Handle.SectionIndex);
-		for (const FSectionLayoutElement& Element : Layout.GetElements())
+	for (auto& Pair : ChannelsByType)
+	{
+		TArray<void*> RawPtrs;
+		RawPtrs.Reserve(Pair.Value.Num());
+		for (const TMovieSceneChannelHandle<void>& Handle : Pair.Value)
 		{
-			TSharedPtr<IKeyArea> KeyArea = Element.GetKeyArea();
-			if (KeyArea.IsValid() && Element.GetKeyArea()->GetUnsortedKeyHandles().Num() > 0)
+			if (void* ChannelPtr = Handle.Get())
 			{
-				return true;
+				RawPtrs.Add(ChannelPtr);
 			}
 		}
-	}
-
-	return false;
-}
 
-bool FSectionContextMenu::CanSetExtrapolationMode() const
-{
-	TArray<FSectionHandle> SelectedSections = StaticCastSharedRef<SSequencer>(Sequencer->GetSequencerWidget())->GetSectionHandles(Sequencer->GetSelection().GetSelectedSections());
-	for (const FSectionHandle& Handle : SelectedSections)
-	{
-		FSectionLayout Layout(*Handle.TrackNode, Handle.SectionIndex);
-		for (const FSectionLayoutElement& Element : Layout.GetElements())
+		ISequencerChannelInterface* ChannelInterface = SequencerModule.FindChannelInterface(Pair.Key);
+		if (ChannelInterface && ChannelInterface->HasAnyKeys_Raw(RawPtrs))
 		{
-			TSharedPtr<IKeyArea> KeyArea = Element.GetKeyArea();
-			if (KeyArea.IsValid() && KeyArea->CanSetExtrapolationMode())
-			{
-				return true;
-			}
+			return true;
 		}
 	}
 	return false;
@@ -865,7 +669,8 @@ void FSectionContextMenu::TrimSection(bool bTrimLeft)
 {
 	FScopedTransaction TrimSectionTransaction(LOCTEXT("TrimSection_Transaction", "Trim Section"));
 
-	MovieSceneToolHelpers::TrimSection(Sequencer->GetSelection().GetSelectedSections(), Sequencer->GetLocalTime(), bTrimLeft);
+	FFrameNumber CurrentFrame = Sequencer->GetLocalTime().Time.FrameNumber;
+	MovieSceneToolHelpers::TrimSection(Sequencer->GetSelection().GetSelectedSections(), CurrentFrame, bTrimLeft);
 	Sequencer->NotifyMovieSceneDataChanged( EMovieSceneDataChangeType::TrackValueChanged );
 }
 
@@ -874,7 +679,8 @@ void FSectionContextMenu::SplitSection()
 {
 	FScopedTransaction SplitSectionTransaction(LOCTEXT("SplitSection_Transaction", "Split Section"));
 
-	MovieSceneToolHelpers::SplitSection(Sequencer->GetSelection().GetSelectedSections(), Sequencer->GetLocalTime());
+	FFrameNumber CurrentFrame = Sequencer->GetLocalTime().Time.FrameNumber;
+	MovieSceneToolHelpers::SplitSection(Sequencer->GetSelection().GetSelectedSections(), CurrentFrame);
 	Sequencer->NotifyMovieSceneDataChanged( EMovieSceneDataChangeType::MovieSceneStructureItemAdded );
 }
 
@@ -898,14 +704,23 @@ void FSectionContextMenu::ReduceKeys()
 		}
 	}
 
+
+	FKeyDataOptimizationParams Params;
+	Params.bAutoSetInterpolation = true;
+
 	for (auto KeyArea : KeyAreas)
 	{
 		if (KeyArea.IsValid())
 		{
-			if (KeyArea->GetRichCurve() && KeyArea->GetOwningSection())
+			UMovieSceneSection* Section = KeyArea->GetOwningSection();
+			if (KeyArea->GetOwningSection())
 			{
-				KeyArea->GetOwningSection()->Modify();
-				KeyArea->GetRichCurve()->RemoveRedundantKeys(KINDA_SMALL_NUMBER);
+				Section->Modify();
+
+				for (const FMovieSceneChannelEntry& Entry : Section->GetChannelProxy().GetAllEntries())
+				{
+					Entry.GetBatchChannelInterface().Optimize_Batch(Entry.GetChannels(), Params);
+				}
 			}
 		}
 	}
@@ -917,7 +732,7 @@ bool FSectionContextMenu::IsTrimmable() const
 {
 	for (auto Section : Sequencer->GetSelection().GetSelectedSections())
 	{
-		if (Section.IsValid() && Section->IsTimeWithinSection(Sequencer->GetLocalTime()))
+		if (Section.IsValid() && Section->IsTimeWithinSection(Sequencer->GetLocalTime().Time.FrameNumber))
 		{
 			return true;
 		}
@@ -946,75 +761,6 @@ bool FSectionContextMenu::CanReduceKeys() const
 }
 
 
-void FSectionContextMenu::SetExtrapolationMode(ERichCurveExtrapolation ExtrapMode, bool bPreInfinity)
-{
-	FScopedTransaction SetExtrapolationModeTransaction(LOCTEXT("SetExtrapolationMode_Transaction", "Set Extrapolation Mode"));
-
-	bool bAnythingChanged = false;
-
-	TArray<FSectionHandle> SelectedSections = StaticCastSharedRef<SSequencer>(Sequencer->GetSequencerWidget())->GetSectionHandles(Sequencer->GetSelection().GetSelectedSections());
-	for (const FSectionHandle& Handle : SelectedSections)
-	{
-		UMovieSceneSection* Section = Handle.GetSectionObject();
-		if (!Section)
-		{
-			continue;
-		}
-
-		if (Section->TryModify())
-		{
-			FSectionLayout Layout(*Handle.TrackNode, Handle.SectionIndex);
-			for (const FSectionLayoutElement& Element : Layout.GetElements())
-			{
-				TSharedPtr<IKeyArea> KeyArea = Element.GetKeyArea();
-				if (KeyArea.IsValid())
-				{
-					bAnythingChanged = true;
-					KeyArea->SetExtrapolationMode(ExtrapMode, bPreInfinity);
-				}
-			}
-		}
-	}
-
-	if (bAnythingChanged)
-	{
-		Sequencer->NotifyMovieSceneDataChanged( EMovieSceneDataChangeType::TrackValueChanged );
-	}
-	else
-	{
-		SetExtrapolationModeTransaction.Cancel();
-	}
-}
-
-
-bool FSectionContextMenu::IsExtrapolationModeSelected(ERichCurveExtrapolation ExtrapMode, bool bPreInfinity) const
-{
-	// @todo Sequencer should operate on selected sections
-	bool bAllSelected = false;
-
-	TArray<FSectionHandle> SelectedSections = StaticCastSharedRef<SSequencer>(Sequencer->GetSequencerWidget())->GetSectionHandles(Sequencer->GetSelection().GetSelectedSections());
-	for (const FSectionHandle& Handle : SelectedSections)
-	{
-		FSectionLayout Layout(*Handle.TrackNode, Handle.SectionIndex);
-		for (const FSectionLayoutElement& Element : Layout.GetElements())
-		{
-			TSharedPtr<IKeyArea> KeyArea = Element.GetKeyArea();
-			if (KeyArea.IsValid())
-			{
-				bAllSelected = true;
-				if (KeyArea->GetExtrapolationMode(bPreInfinity) != ExtrapMode)
-				{
-					bAllSelected = false;
-					break;
-				}
-			}
-		}
-	}
-
-	return bAllSelected;
-}
-
-
 void FSectionContextMenu::ToggleSectionActive()
 {
 	FScopedTransaction ToggleSectionActiveTransaction( LOCTEXT("ToggleSectionActive_Transaction", "Toggle Section Active") );
@@ -1235,7 +981,7 @@ void FSectionContextMenu::BringToFront()
 		}
 	}
 
-	Sequencer->SetLocalTimeDirectly(Sequencer->GetLocalTime());
+	Sequencer->SetLocalTimeDirectly(Sequencer->GetLocalTime().Time);
 }
 
 
@@ -1283,7 +1029,7 @@ void FSectionContextMenu::SendToBack()
 		}
 	}
 
-	Sequencer->SetLocalTimeDirectly(Sequencer->GetLocalTime());
+	Sequencer->SetLocalTimeDirectly(Sequencer->GetLocalTime().Time);
 }
 
 
@@ -1328,7 +1074,7 @@ void FSectionContextMenu::BringForward()
 		}
 	}
 
-	Sequencer->SetLocalTimeDirectly(Sequencer->GetLocalTime());
+	Sequencer->SetLocalTimeDirectly(Sequencer->GetLocalTime().Time);
 }
 
 
@@ -1373,7 +1119,7 @@ void FSectionContextMenu::SendBackward()
 		}
 	}
 
-	Sequencer->SetLocalTimeDirectly(Sequencer->GetLocalTime());
+	Sequencer->SetLocalTimeDirectly(Sequencer->GetLocalTime().Time);
 }
 
 
@@ -1673,6 +1419,7 @@ void FPasteContextMenu::PasteInto(int32 DestinationIndex, FName KeyAreaName)
 	TSet<FSequencerSelectedKey> NewSelection;
 
 	FSequencerPasteEnvironment PasteEnvironment;
+	PasteEnvironment.FrameResolution = Sequencer->GetFocusedFrameResolution();
 	PasteEnvironment.CardinalTime = Args.PasteAtTime;
 	PasteEnvironment.OnKeyPasted = [&](FKeyHandle Handle, IKeyArea& KeyArea){
 		NewSelection.Add(FSequencerSelectedKey(*KeyArea.GetOwningSection(), KeyArea.AsShared(), Handle));
@@ -1685,7 +1432,7 @@ void FPasteContextMenu::PasteInto(int32 DestinationIndex, FName KeyAreaName)
 	}
 	else
 	{
-		SSequencerSection::ThrobSelection();
+		SSequencerSection::ThrobKeySelection();
 
 		// @todo sequencer: selection in transactions
 		FSequencerSelection& Selection = Sequencer->GetSelection();
@@ -1758,9 +1505,9 @@ void FPasteFromHistoryContextMenu::PopulateMenu(FMenuBuilder& MenuBuilder)
 	MenuBuilder.EndSection();
 }
 
-void FEasingContextMenu::BuildMenu(FMenuBuilder& MenuBuilder, const TArray<FEasingAreaHandle>& InEasings, FSequencer& Sequencer, float InMouseDownTime)
+void FEasingContextMenu::BuildMenu(FMenuBuilder& MenuBuilder, const TArray<FEasingAreaHandle>& InEasings, FSequencer& Sequencer, FFrameTime InMouseDownTime)
 {
-	TSharedRef<FEasingContextMenu> EasingMenu = MakeShareable(new FEasingContextMenu(InEasings));
+	TSharedRef<FEasingContextMenu> EasingMenu = MakeShareable(new FEasingContextMenu(InEasings, Sequencer));
 	EasingMenu->PopulateMenu(MenuBuilder);
 
 	MenuBuilder.AddMenuSeparator();
@@ -1780,19 +1527,19 @@ void FEasingContextMenu::PopulateMenu(FMenuBuilder& MenuBuilder)
 		{
 			GEditor->BeginTransaction(LOCTEXT("SetEasingTimeText", "Set Easing Length"));
 		};
-		auto OnEndSliderMovement = [=](float NewLength)
+		auto OnEndSliderMovement = [=](double NewLength)
 		{
 			if (GEditor->IsTransactionActive())
 			{
 				GEditor->EndTransaction();
 			}
 		};
-		auto OnValueCommitted = [=](float NewLength, ETextCommit::Type CommitInfo)
+		auto OnValueCommitted = [=](double NewLength, ETextCommit::Type CommitInfo)
 		{
 			if (CommitInfo == ETextCommit::OnEnter || CommitInfo == ETextCommit::OnUserMovedFocus)
 			{
 				FScopedTransaction Transaction(LOCTEXT("SetEasingTimeText", "Set Easing Length"));
-				Shared->OnUpdateLength(NewLength);
+				Shared->OnUpdateLength((int32)NewLength);
 			}
 		};
 
@@ -1804,23 +1551,31 @@ void FEasingContextMenu::PopulateMenu(FMenuBuilder& MenuBuilder)
 				SNew(SBox)
 				.HAlign(HAlign_Right)
 				[
-					SNew(SNumericEntryBox<float>)
+					SNew(SNumericEntryBox<double>)
 					.SpinBoxStyle(&FEditorStyle::GetWidgetStyle<FSpinBoxStyle>("Sequencer.HyperlinkSpinBox"))
 					.EditableTextBoxStyle(&FEditorStyle::GetWidgetStyle<FEditableTextBoxStyle>("Sequencer.HyperlinkTextBox"))
 					// Don't update the value when undetermined text changes
 					.OnUndeterminedValueChanged_Lambda([](FText){})
 					.AllowSpin(true)
 					.MinValue(0.f)
-					.MaxValue(TOptional<float>())
-					.MaxSliderValue(TOptional<float>())
+					.MaxValue(TOptional<double>())
+					.MaxSliderValue(TOptional<double>())
 					.MinSliderValue(0.f)
-					.Delta(0.001f)
-					.Value_Lambda([=]{ return Shared->GetCurrentLength(); })
-					.OnValueChanged_Lambda([=](float NewLength){ Shared->OnUpdateLength(NewLength); })
+					.Delta_Lambda([=]() -> double { return Sequencer->GetPlayRateDeltaFrameCount(); })
+					.Value_Lambda([=] {
+						TOptional<int32> Current = Shared->GetCurrentLength();
+						if (Current.IsSet())
+						{
+							return TOptional<double>(Current.GetValue());
+						}
+						return TOptional<double>();
+					})
+					.OnValueChanged_Lambda([=](double NewLength){ Shared->OnUpdateLength(NewLength); })
 					.OnValueCommitted_Lambda(OnValueCommitted)
 					.OnBeginSliderMovement_Lambda(OnBeginSliderMovement)
 					.OnEndSliderMovement_Lambda(OnEndSliderMovement)
 					.BorderForegroundColor(FEditorStyle::GetSlateColor("DefaultForeground"))
+					.TypeInterface(Sequencer->GetNumericTypeInterface())
 				]
 			]
 
@@ -1853,26 +1608,26 @@ void FEasingContextMenu::PopulateMenu(FMenuBuilder& MenuBuilder)
 	MenuBuilder.EndSection();
 }
 
-TOptional<float> FEasingContextMenu::GetCurrentLength() const
+TOptional<int32> FEasingContextMenu::GetCurrentLength() const
 {
-	TOptional<float> Value;
+	TOptional<int32> Value;
 
 	for (const FEasingAreaHandle& Handle : Easings)
 	{
 		UMovieSceneSection* Section = Handle.Section.GetSectionObject();
 		if (Section)
 		{
-			if (Handle.EasingType == ESequencerEasingType::In && Section->Easing.GetEaseInTime() == Value.Get(Section->Easing.GetEaseInTime()))
+			if (Handle.EasingType == ESequencerEasingType::In && Section->Easing.GetEaseInDuration() == Value.Get(Section->Easing.GetEaseInDuration()))
 			{
-				Value = Section->Easing.GetEaseInTime();
+				Value = Section->Easing.GetEaseInDuration();
 			}
-			else if (Handle.EasingType == ESequencerEasingType::Out && Section->Easing.GetEaseOutTime() == Value.Get(Section->Easing.GetEaseOutTime()))
+			else if (Handle.EasingType == ESequencerEasingType::Out && Section->Easing.GetEaseOutDuration() == Value.Get(Section->Easing.GetEaseOutDuration()))
 			{
-				Value = Section->Easing.GetEaseOutTime();
+				Value = Section->Easing.GetEaseOutDuration();
 			}
 			else
 			{
-				return TOptional<float>();
+				return TOptional<int32>();
 			}
 		}
 	}
@@ -1880,7 +1635,7 @@ TOptional<float> FEasingContextMenu::GetCurrentLength() const
 	return Value;
 }
 
-void FEasingContextMenu::OnUpdateLength(float NewLength)
+void FEasingContextMenu::OnUpdateLength(int32 NewLength)
 {
 	for (const FEasingAreaHandle& Handle : Easings)
 	{
@@ -1890,12 +1645,12 @@ void FEasingContextMenu::OnUpdateLength(float NewLength)
 			if (Handle.EasingType == ESequencerEasingType::In)
 			{
 				Section->Easing.bManualEaseIn = true;
-				Section->Easing.ManualEaseInTime = FMath::Min(Section->GetRange().Size<float>(), NewLength);
+				Section->Easing.ManualEaseInDuration = FMath::Min(MovieScene::DiscreteSize(Section->GetRange()), NewLength);
 			}
 			else
 			{
 				Section->Easing.bManualEaseOut = true;
-				Section->Easing.ManualEaseOutTime = FMath::Min(Section->GetRange().Size<float>(), NewLength);
+				Section->Easing.ManualEaseOutDuration = FMath::Min(MovieScene::DiscreteSize(Section->GetRange()), NewLength);
 			}
 		}
 	}
diff --git a/Engine/Source/Editor/Sequencer/Private/SequencerContextMenus.h b/Engine/Source/Editor/Sequencer/Private/SequencerContextMenus.h
--- a/Engine/Source/Editor/Sequencer/Private/SequencerContextMenus.h
+++ b/Engine/Source/Editor/Sequencer/Private/SequencerContextMenus.h
@@ -7,6 +7,7 @@
 #include "Sequencer.h"
 #include "SequencerClipboardReconciler.h"
 #include "ScopedTransaction.h"
+#include "Channels/MovieSceneChannelHandle.h"
 
 struct FEasingAreaHandle;
 class FMenuBuilder;
@@ -20,26 +21,20 @@ class UMovieSceneSection;
  */
 struct FSectionContextMenu : TSharedFromThis<FSectionContextMenu>
 {
-	static void BuildMenu(FMenuBuilder& MenuBuilder, FSequencer& Sequencer, float InMouseDownTime);
+	static void BuildMenu(FMenuBuilder& MenuBuilder, FSequencer& Sequencer, FFrameTime InMouseDownTime);
 
 private:
 
 	/** Hidden AsShared() methods to discourage CreateSP delegate use. */
 	using TSharedFromThis::AsShared;
 
-	FSectionContextMenu(FSequencer& InSeqeuncer, float InMouseDownTime)
-		: Sequencer(StaticCastSharedRef<FSequencer>(InSeqeuncer.AsShared()))
-		, MouseDownTime(InMouseDownTime)
-	{}
+	FSectionContextMenu(FSequencer& InSeqeuncer, FFrameTime InMouseDownTime);
 
 	void PopulateMenu(FMenuBuilder& MenuBuilder);
 
 	/** Add edit menu for trim and split */
 	void AddEditMenu(FMenuBuilder& MenuBuilder);
 
-	/** Add extrapolation menu for pre and post infinity */
-	void AddExtrapolationMenu(FMenuBuilder& MenuBuilder, bool bPreInfinity);
-
 	/** Add the Properties sub-menu. */
 	void AddPropertiesMenu(FMenuBuilder& MenuBuilder);
 
@@ -60,8 +55,6 @@ private:
 
 	bool CanPrimeForRecording() const;
 
-	bool CanSetExtrapolationMode() const;
-
 	void TrimSection(bool bTrimLeft);
 
 	void SplitSection();
@@ -72,10 +65,6 @@ private:
 
 	bool CanReduceKeys() const;
 
-	void SetExtrapolationMode(ERichCurveExtrapolation ExtrapMode, bool bPreInfinity);
-
-	bool IsExtrapolationModeSelected(ERichCurveExtrapolation ExtrapMode, bool bPreInfinity) const;
-
 	void ToggleSectionActive();
 
 	bool IsSectionActive() const;
@@ -100,14 +89,17 @@ private:
 	TSharedRef<FSequencer> Sequencer;
 
 	/** The time that we clicked on to summon this menu */
-	float MouseDownTime;
+	FFrameTime MouseDownTime;
+
+	TMap<uint32, TArray<TMovieSceneChannelHandle<void>>> ChannelsByType;
+	TMap<uint32, TArray<UMovieSceneSection*>> SectionsByType;
 };
 
 /** Arguments required for a paste operation */
 struct FPasteContextMenuArgs
 {
 	/** Paste the clipboard into the specified array of sequencer nodes, at the given time */
-	static FPasteContextMenuArgs PasteInto(TArray<TSharedRef<FSequencerDisplayNode>> InNodes, float InTime, TSharedPtr<FMovieSceneClipboard> InClipboard = nullptr)
+	static FPasteContextMenuArgs PasteInto(TArray<TSharedRef<FSequencerDisplayNode>> InNodes, FFrameNumber InTime, TSharedPtr<FMovieSceneClipboard> InClipboard = nullptr)
 	{
 		FPasteContextMenuArgs Args;
 		Args.Clipboard = InClipboard;
@@ -117,7 +109,7 @@ struct FPasteContextMenuArgs
 	}
 
 	/** Paste the clipboard at the given time, using the sequencer selection states to determine paste destinations */
-	static FPasteContextMenuArgs PasteAt(float InTime, TSharedPtr<FMovieSceneClipboard> InClipboard = nullptr)
+	static FPasteContextMenuArgs PasteAt(FFrameNumber InTime, TSharedPtr<FMovieSceneClipboard> InClipboard = nullptr)
 	{
 		FPasteContextMenuArgs Args;
 		Args.Clipboard = InClipboard;
@@ -129,7 +121,7 @@ struct FPasteContextMenuArgs
 	TSharedPtr<FMovieSceneClipboard> Clipboard;
 
 	/** The Time to paste at */
-	float PasteAtTime;
+	FFrameNumber PasteAtTime;
 
 	/** Optional user-supplied nodes to paste into */
 	TArray<TSharedRef<FSequencerDisplayNode>> DestinationNodes;
@@ -215,8 +207,8 @@ struct FKeyContextMenu : TSharedFromThis<FKeyContextMenu>
 
 private:
 
-	FKeyContextMenu(FSequencer& InSeqeuncer)
-		: Sequencer(StaticCastSharedRef<FSequencer>(InSeqeuncer.AsShared()))
+	FKeyContextMenu(FSequencer& InSequencer)
+		: Sequencer(StaticCastSharedRef<FSequencer>(InSequencer.AsShared()))
 	{}
 
 	/** Hidden AsShared() methods to discourage CreateSP delegate use. */
@@ -225,13 +217,13 @@ private:
 	/** Add the Properties sub-menu. */
 	void AddPropertiesMenu(FMenuBuilder& MenuBuilder);
 
-	/** Check if we can add the key properties menu */
-	bool CanAddPropertiesMenu() const;
-
 	void PopulateMenu(FMenuBuilder& MenuBuilder);
 
 	/** The sequencer */
 	TSharedRef<FSequencer> Sequencer;
+
+	TSharedPtr<FStructOnScope> KeyStruct;
+	TWeakObjectPtr<UMovieSceneSection> KeyStructSection;
 };
 
 
@@ -243,12 +235,14 @@ private:
  */
 struct FEasingContextMenu : TSharedFromThis<FEasingContextMenu>
 {
-	static void BuildMenu(FMenuBuilder& MenuBuilder, const TArray<FEasingAreaHandle>& InEasings, FSequencer& Sequencer, float InMouseDownTime);
+	static void BuildMenu(FMenuBuilder& MenuBuilder, const TArray<FEasingAreaHandle>& InEasings, FSequencer& Sequencer, FFrameTime InMouseDownTime);
 
 private:
 
-	FEasingContextMenu(const TArray<FEasingAreaHandle>& InEasings)
+	FEasingContextMenu(const TArray<FEasingAreaHandle>& InEasings, FSequencer& InSequencer)
 		: Easings(InEasings)
+		, Sequencer(StaticCastSharedRef<FSequencer>(InSequencer.AsShared()))
+
 	{}
 
 	/** Hidden AsShared() methods to discourage CreateSP delegate use. */
@@ -264,17 +258,19 @@ private:
 
 	void OnEasingTypeChanged(UClass* NewClass);
 
-	void OnUpdateLength(float NewLength);
+	void OnUpdateLength(int32 NewLength);
 
-	TOptional<float> GetCurrentLength() const;
+	TOptional<int32> GetCurrentLength() const;
 
 	ECheckBoxState GetAutoEasingCheckState() const;
 
 	void SetAutoEasing(bool bAutoEasing);
 
-	/** The sequencer */
 	TArray<FEasingAreaHandle> Easings;
 
+	/** The sequencer */
+	TSharedRef<FSequencer> Sequencer;
+
 	/** A scoped transaction for a current operation */
 	TUniquePtr<FScopedTransaction> ScopedTransaction;
 };
diff --git a/Engine/Source/Editor/Sequencer/Private/SequencerEdMode.cpp b/Engine/Source/Editor/Sequencer/Private/SequencerEdMode.cpp
--- a/Engine/Source/Editor/Sequencer/Private/SequencerEdMode.cpp
+++ b/Engine/Source/Editor/Sequencer/Private/SequencerEdMode.cpp
@@ -18,6 +18,7 @@
 #include "SequencerKeyActor.h"
 #include "EditorWorldExtension.h"
 #include "ViewportWorldInteraction.h"
+#include "SSequencer.h"
 #include "Evaluation/MovieScene3DTransformTemplate.h"
 
 const FEditorModeID FSequencerEdMode::EM_SequencerMode(TEXT("EM_SequencerMode"));
@@ -131,9 +132,8 @@ void FSequencerEdMode::AddReferencedObjects(FReferenceCollector& Collector)
 
 void FSequencerEdMode::OnKeySelected(FViewport* Viewport, HMovieSceneKeyProxy* KeyProxy)
 {
-	bool bCtrlDown = Viewport->KeyState(EKeys::LeftControl) || Viewport->KeyState(EKeys::RightControl);
-	bool bAltDown = Viewport->KeyState(EKeys::LeftAlt) || Viewport->KeyState(EKeys::RightAlt);
-	bool bShiftDown = Viewport->KeyState(EKeys::LeftShift) || Viewport->KeyState(EKeys::RightShift);
+	const bool bToggleSelection = Viewport->KeyState(EKeys::LeftControl) || Viewport->KeyState(EKeys::RightControl);
+	const bool bAddToSelection = Viewport->KeyState(EKeys::LeftShift) || Viewport->KeyState(EKeys::RightShift);
 
 	for (TWeakPtr<FSequencer> WeakSequencer : Sequencers)
 	{
@@ -142,12 +142,41 @@ void FSequencerEdMode::OnKeySelected(FViewport* Viewport, HMovieSceneKeyProxy* K
 		{
 			Sequencer->SetLocalTimeDirectly(KeyProxy->Key.Time);
 
+			FSequencerSelection& Selection = Sequencer->GetSelection();
+			if (!bAddToSelection && !bToggleSelection)
+			{
+				Sequencer->GetSelection().EmptySelectedKeys();
+			}
+
 			for (const FTrajectoryKey::FData KeyData : KeyProxy->Key.KeyData)
 			{
-				if (UMovieSceneSection* Section = KeyData.Section.Get())
+				UMovieSceneSection* Section = KeyData.Section.Get();
+				if (Section && KeyData.KeyHandle.IsSet())
 				{
-					Sequencer->SelectTrackKeys(Section, KeyProxy->Key.Time, bShiftDown, bCtrlDown);
+					TSet<TWeakObjectPtr<UMovieSceneSection>> SectionsToFind;
+					SectionsToFind.Add(Section);
+
+					// Find the key area with the specified channel name
+					TArray<FSectionHandle> SectionHandles = StaticCastSharedRef<SSequencer>(Sequencer->GetSequencerWidget())->GetSectionHandles(SectionsToFind);
+
+					for (FSectionHandle Handle : SectionHandles)
+					{
+						TArray<TSharedRef<FSequencerSectionKeyAreaNode>> KeyAreaNodes;
+						Handle.TrackNode->GetChildKeyAreaNodesRecursively(KeyAreaNodes);
+
+						for (TSharedRef<FSequencerSectionKeyAreaNode> KeyAreaNode : KeyAreaNodes)
+						{
+							TSharedPtr<IKeyArea> KeyArea = KeyAreaNode->GetKeyArea(Section);
+							if (KeyArea.IsValid() && KeyArea->GetName() == KeyData.ChannelName)
+							{
+								Sequencer->SelectKey(Section, KeyArea, KeyData.KeyHandle.GetValue(), bToggleSelection);
+								goto next_key;
+							}
+						}
+					}
 				}
+next_key:
+				continue;
 			}
 		}
 	}
@@ -200,11 +229,11 @@ void FSequencerEdMode::CleanUpMeshTrails()
 namespace SequencerEdMode_Draw3D
 {
 static const FColor	KeySelectedColor(255,128,0);
-static const float	DrawTrackTimeRes = 0.1f;
+static const int32	DrawTrackTimeRes = 10;
 static const float	CurveHandleScale = 0.5f;
 }
 
-FTransform FSequencerEdMode::GetRefFrame(const TSharedPtr<FSequencer>& Sequencer, const UObject* InObject, float KeyTime)
+FTransform FSequencerEdMode::GetRefFrame(const TSharedPtr<FSequencer>& Sequencer, const UObject* InObject, FFrameTime KeyTime)
 {
 	FTransform RefTM = FTransform::Identity;
 
@@ -226,7 +255,7 @@ FTransform FSequencerEdMode::GetRefFrame(const TSharedPtr<FSequencer>& Sequencer
 	return RefTM;
 }
 
-FTransform FSequencerEdMode::GetRefFrame(const TSharedPtr<FSequencer>& Sequencer, const AActor* Actor, float KeyTime)
+FTransform FSequencerEdMode::GetRefFrame(const TSharedPtr<FSequencer>& Sequencer, const AActor* Actor, FFrameTime KeyTime)
 {
 	FTransform RefTM = FTransform::Identity;
 
@@ -238,7 +267,7 @@ FTransform FSequencerEdMode::GetRefFrame(const TSharedPtr<FSequencer>& Sequencer
 	return RefTM;
 }
 
-FTransform FSequencerEdMode::GetRefFrame(const TSharedPtr<FSequencer>& Sequencer, const USceneComponent* SceneComponent, float KeyTime)
+FTransform FSequencerEdMode::GetRefFrame(const TSharedPtr<FSequencer>& Sequencer, const USceneComponent* SceneComponent, FFrameTime KeyTime)
 {
 	FTransform RefTM = FTransform::Identity;
 
@@ -279,7 +308,7 @@ FTransform FSequencerEdMode::GetRefFrame(const TSharedPtr<FSequencer>& Sequencer
 					{
 						for (const UMovieSceneSection* Section : TransformTrack->GetAllSections())
 						{
-							if (Section->IsTimeWithinSection(KeyTime))
+							if (Section->IsTimeWithinSection(KeyTime.FrameNumber))
 							{
 								const UMovieScene3DTransformSection* ParentSection = Cast<UMovieScene3DTransformSection>(Section);
 
@@ -313,12 +342,12 @@ FTransform FSequencerEdMode::GetRefFrame(const TSharedPtr<FSequencer>& Sequencer
 	return RefTM;
 }
 
-void FSequencerEdMode::GetLocationAtTime(FMovieSceneEvaluationTrack* Track, UObject* Object, float KeyTime, FVector& KeyPos, FRotator& KeyRot, const TSharedPtr<FSequencer>& Sequencer)
+void FSequencerEdMode::GetLocationAtTime(FMovieSceneEvaluationTrack* Track, UObject* Object, FFrameTime KeyTime, FVector& KeyPos, FRotator& KeyRot, const TSharedPtr<FSequencer>& Sequencer)
 {
 	FMovieSceneInterrogationData InterrogationData;
 	Sequencer->GetEvaluationTemplate().CopyActuators(InterrogationData.GetAccumulator());
 
-	FMovieSceneContext Context(KeyTime);
+	FMovieSceneContext Context(FMovieSceneEvaluationRange(KeyTime, Sequencer->GetFocusedFrameResolution()));
 	Track->Interrogate(Context, InterrogationData, Object);
 
 	for (const FTransform& Transform : InterrogationData.Iterate<FTransform>(UMovieScene3DTransformTrack::GetInterrogationKey()))
@@ -369,6 +398,7 @@ void FSequencerEdMode::DrawTransformTrack(const TSharedPtr<FSequencer>& Sequence
 		}
 	);
 	
+	FFrameRate FrameResolution = Sequencer->GetFocusedFrameResolution();
 	FMovieSceneEvaluationTemplate* Template = Sequencer->GetEvaluationTemplate().FindTemplate(Sequencer->GetFocusedTemplateID());
 	if (!bShowTrajectory || !Template || !TransformTrack->GetAllSections().ContainsByPredicate([](UMovieSceneSection* In){ return In->IsActive(); }))
 	{
@@ -380,7 +410,7 @@ void FSequencerEdMode::DrawTransformTrack(const TSharedPtr<FSequencer>& Sequence
 	// Draw one line per-track (should only really ever be one)
 	if (FMovieSceneEvaluationTrack* EvalTrack = Template->FindTrack(TransformTrack->GetSignature()))
 	{
-		TArray<FTrajectoryKey> TrajectoryKeys = TransformTrack->GetTrajectoryData(Sequencer->GetLocalTime(), Sequencer->GetSettings()->GetTrajectoryPathCap());
+		TArray<FTrajectoryKey> TrajectoryKeys = TransformTrack->GetTrajectoryData(Sequencer->GetLocalTime().Time.FrameNumber, Sequencer->GetSequencerSettings()->GetTrajectoryPathCap());
 		for (TWeakObjectPtr<> WeakBinding : BoundObjects)
 		{
 			UObject* BoundObject = WeakBinding.Get();
@@ -390,12 +420,12 @@ void FSequencerEdMode::DrawTransformTrack(const TSharedPtr<FSequencer>& Sequence
 			}
 
 			FVector OldKeyPos(0);
-			float OldKeyTime = 0.f;
+			FFrameTime OldKeyTime = 0;
 			int KeyTimeIndex = 0;
 
 			for (const FTrajectoryKey& NewTrajectoryKey : TrajectoryKeys)
 			{
-				float NewKeyTime = NewTrajectoryKey.Time;
+				FFrameTime NewKeyTime = NewTrajectoryKey.Time;
 
 				FVector NewKeyPos(0);
 				FRotator NewKeyRot(0,0,0);
@@ -405,13 +435,14 @@ void FSequencerEdMode::DrawTransformTrack(const TSharedPtr<FSequencer>& Sequence
 				// If not the first keypoint, draw a line to the last keypoint.
 				if(KeyTimeIndex > 0)
 				{
-					int32 NumSteps = FMath::CeilToInt( (NewKeyTime - OldKeyTime)/SequencerEdMode_Draw3D::DrawTrackTimeRes );
-									// Limit the number of steps to prevent a rendering performance hit
+					int32 NumSteps = ((NewKeyTime - OldKeyTime)*SequencerEdMode_Draw3D::DrawTrackTimeRes).CeilToFrame().Value;
+					// Limit the number of steps to prevent a rendering performance hit
 					NumSteps = FMath::Min( 100, NumSteps );
-					float DrawSubstep = (NewKeyTime - OldKeyTime)/NumSteps;
+
+					FFrameTime DrawSubstep = NumSteps == 0 ? 0 : (NewKeyTime - OldKeyTime)*(1.f/NumSteps);
 
 					// Find position on first keyframe.
-					float OldTime = OldKeyTime;
+					FFrameTime OldTime = OldKeyTime;
 
 					FVector OldPos(0);
 					FRotator OldRot(0,0,0);
@@ -438,7 +469,7 @@ void FSequencerEdMode::DrawTransformTrack(const TSharedPtr<FSequencer>& Sequence
 						// Then draw a line for each substep.
 						for (int32 j=1; j<NumSteps+1; j++)
 						{
-							float NewTime = OldKeyTime + j*DrawSubstep;
+							FFrameTime NewTime = OldKeyTime + DrawSubstep*j;
 
 							FVector NewPos(0);
 							FRotator NewRot(0,0,0);
@@ -460,7 +491,7 @@ void FSequencerEdMode::DrawTransformTrack(const TSharedPtr<FSequencer>& Sequence
 								}
 								else if (TrailActor != nullptr)
 								{
-									TrailActor->AddFrameMeshComponent(NewTime, FTransform(NewRot, NewPos, FVector(3.0f)));
+									TrailActor->AddFrameMeshComponent(NewTime / FrameResolution, FTransform(NewRot, NewPos, FVector(3.0f)));
 								}
 							}
 							OldTime = NewTime;
@@ -477,7 +508,7 @@ void FSequencerEdMode::DrawTransformTrack(const TSharedPtr<FSequencer>& Sequence
 			// Draw keypoints on top of curve
 			for (const FTrajectoryKey& TrajectoryKey : TrajectoryKeys)
 			{
-				float NewKeyTime = TrajectoryKey.Time;
+				FFrameTime NewKeyTime = TrajectoryKey.Time;
 
 				// Find the time, position and orientation of this Key.
 				FVector NewKeyPos(0);
@@ -512,7 +543,7 @@ void FSequencerEdMode::DrawTransformTrack(const TSharedPtr<FSequencer>& Sequence
 
 					for (UMovieScene3DTransformSection* Section : AllSections)
 					{
-						TrailActor->AddKeyMeshActor(NewKeyTime, FTransform(NewKeyRot, NewKeyPos, FVector(3.0f)), Section);
+						TrailActor->AddKeyMeshActor(NewKeyTime / FrameResolution, FTransform(NewKeyRot, NewKeyPos, FVector(3.0f)), Section);
 					}
 				}
 
diff --git a/Engine/Source/Editor/Sequencer/Private/SequencerEdMode.h b/Engine/Source/Editor/Sequencer/Private/SequencerEdMode.h
--- a/Engine/Source/Editor/Sequencer/Private/SequencerEdMode.h
+++ b/Engine/Source/Editor/Sequencer/Private/SequencerEdMode.h
@@ -6,6 +6,7 @@
 #include "InputCoreTypes.h"
 #include "EditorModeTools.h"
 #include "EdMode.h"
+#include "FrameTime.h"
 
 class FCanvas;
 class FEditorViewportClient;
@@ -73,10 +74,10 @@ protected:
 	void DrawTransformTrack(const TSharedPtr<FSequencer>& Sequencer, FPrimitiveDrawInterface* PDI, UMovieScene3DTransformTrack* TransformTrack, const TArray<TWeakObjectPtr<UObject>>& BoundObjects, const bool bIsSelected);
 
 protected:
-	static void GetLocationAtTime(FMovieSceneEvaluationTrack* Track, UObject* BoundObject, float KeyTime, FVector& KeyPos, FRotator& KeyRot, const TSharedPtr<FSequencer>& Sequencer);
-	static FTransform GetRefFrame(const TSharedPtr<FSequencer>& Sequencer, const UObject* InObject, float KeyTime);
-	static FTransform GetRefFrame(const TSharedPtr<FSequencer>& Sequencer, const AActor* Actor, float KeyTime);
-	static FTransform GetRefFrame(const TSharedPtr<FSequencer>& Sequencer, const USceneComponent* SceneComponent, float KeyTime);
+	static void GetLocationAtTime(FMovieSceneEvaluationTrack* Track, UObject* BoundObject, FFrameTime KeyTime, FVector& KeyPos, FRotator& KeyRot, const TSharedPtr<FSequencer>& Sequencer);
+	static FTransform GetRefFrame(const TSharedPtr<FSequencer>& Sequencer, const UObject* InObject, FFrameTime KeyTime);
+	static FTransform GetRefFrame(const TSharedPtr<FSequencer>& Sequencer, const AActor* Actor, FFrameTime KeyTime);
+	static FTransform GetRefFrame(const TSharedPtr<FSequencer>& Sequencer, const USceneComponent* SceneComponent, FFrameTime KeyTime);
 
 private:
 	TArray<TWeakPtr<FSequencer>> Sequencers;
