Copying //UE4/Dev-Sequencer to //UE4/Dev-Main (Source: //UE4/Dev-Sequencer @ 3946692)

diff --git a/Engine/Source/Editor/Sequencer/Private/SequencerHotspots.cpp b/Engine/Source/Editor/Sequencer/Private/SequencerHotspots.cpp
--- a/Engine/Source/Editor/Sequencer/Private/SequencerHotspots.cpp
+++ b/Engine/Source/Editor/Sequencer/Private/SequencerHotspots.cpp
@@ -11,6 +11,8 @@
 #include "Tools/SequencerEditTool_Selection.h"
 #include "SequencerTrackNode.h"
 #include "Widgets/Layout/SBox.h"
+#include "Channels/MovieSceneChannelProxy.h"
+#include "MovieSceneTimeHelpers.h"
 
 #define LOCTEXT_NAMESPACE "SequencerHotspots"
 
@@ -29,28 +31,37 @@ void FKeyHotspot::UpdateOnHover(SSequencerTrackArea& InTrackArea, ISequencer& In
 	InTrackArea.AttemptToActivateTool(FSequencerEditTool_Movement::Identifier);
 }
 
-TOptional<float> FKeyHotspot::GetTime() const
+TOptional<FFrameNumber> FKeyHotspot::GetTime() const
 {
-	return Key.KeyArea->GetKeyTime(Key.KeyHandle.GetValue());
+	FFrameNumber Time = 0;
+
+	if (Keys.Num())
+	{
+		TArrayView<const FSequencerSelectedKey> FirstKey(&Keys[0], 1);
+		TArrayView<FFrameNumber> FirstKeyTime(&Time, 1);
+		GetKeyTimes(FirstKey, FirstKeyTime);
+	}
+
+	return Time;
 }
 
-bool FKeyHotspot::PopulateContextMenu(FMenuBuilder& MenuBuilder, ISequencer& InSequencer, float MouseDownTime)
+bool FKeyHotspot::PopulateContextMenu(FMenuBuilder& MenuBuilder, ISequencer& InSequencer, FFrameTime MouseDownTime)
 {
 	FSequencer& Sequencer = static_cast<FSequencer&>(InSequencer);
 	FKeyContextMenu::BuildMenu(MenuBuilder, Sequencer);
 	return true;
 }
 
-TOptional<float> FSectionHotspot::GetTime() const
+TOptional<FFrameNumber> FSectionHotspot::GetTime() const
 {
 	UMovieSceneSection* ThisSection = Section.GetSectionObject();
-	return ThisSection ? ThisSection->GetStartTime() : TOptional<float>();
+	return ThisSection && ThisSection->HasStartFrame() ? ThisSection->GetInclusiveStartFrame() : TOptional<FFrameNumber>();
 }
 
-TOptional<float> FSectionHotspot::GetOffsetTime() const
+TOptional<FFrameTime> FSectionHotspot::GetOffsetTime() const
 {
 	UMovieSceneSection* ThisSection = Section.GetSectionObject();
-	return ThisSection ? ThisSection->GetOffsetTime() : TOptional<float>();
+	return ThisSection ? ThisSection->GetOffsetTime() : TOptional<FFrameTime>();
 }
 
 void FSectionHotspot::UpdateOnHover(SSequencerTrackArea& InTrackArea, ISequencer& InSequencer) const
@@ -64,11 +75,18 @@ void FSectionHotspot::UpdateOnHover(SSequencerTrackArea& InTrackArea, ISequencer
 	}
 	else
 	{
-		// Activate selection mode if the section has keys or is infinite, otherwise just move it
-		TSet<FKeyHandle> KeyHandles;
-		ThisSection->GetKeyHandles(KeyHandles, ThisSection->GetRange());
+		// Activate selection mode if the section has keys
+		for (const FMovieSceneChannelEntry& Entry : ThisSection->GetChannelProxy().GetAllEntries())
+		{
+			if (Entry.GetBatchChannelInterface().GetNumKeys_Batch(Entry.GetChannels()) != 0)
+			{
+				InTrackArea.AttemptToActivateTool(FSequencerEditTool_Selection::Identifier);
+				return;
+			}
+		}
 
-		if (KeyHandles.Num() || ThisSection->IsInfinite())
+		// Activate selection mode if the section is infinite, otherwise just move it
+		if (ThisSection->GetRange() == TRange<FFrameNumber>::All())
 		{
 			InTrackArea.AttemptToActivateTool(FSequencerEditTool_Selection::Identifier);
 		}
@@ -79,7 +97,7 @@ void FSectionHotspot::UpdateOnHover(SSequencerTrackArea& InTrackArea, ISequencer
 	}
 }
 
-bool FSectionHotspot::PopulateContextMenu(FMenuBuilder& MenuBuilder, ISequencer& InSequencer, float MouseDownTime)
+bool FSectionHotspot::PopulateContextMenu(FMenuBuilder& MenuBuilder, ISequencer& InSequencer, FFrameTime MouseDownTime)
 {
 	FSequencer& Sequencer = static_cast<FSequencer&>(InSequencer);
 
@@ -102,14 +120,14 @@ bool FSectionHotspot::PopulateContextMenu(FMenuBuilder& MenuBuilder, ISequencer&
 	return true;
 }
 
-TOptional<float> FSectionResizeHotspot::GetTime() const
+TOptional<FFrameNumber> FSectionResizeHotspot::GetTime() const
 {
 	UMovieSceneSection* ThisSection = Section.GetSectionObject();
 	if (!ThisSection)
 	{
-		return TOptional<float>();
+		return TOptional<FFrameNumber>();
 	}
-	return HandleType == Left ? ThisSection->GetStartTime() : ThisSection->GetEndTime();
+	return HandleType == Left ? ThisSection->GetInclusiveStartFrame() : ThisSection->GetExclusiveEndFrame();
 }
 
 void FSectionResizeHotspot::UpdateOnHover(SSequencerTrackArea& InTrackArea, ISequencer& InSequencer) const
@@ -135,21 +153,21 @@ TSharedPtr<ISequencerEditToolDragOperation> FSectionResizeHotspot::InitiateDrag(
 	return MakeShareable( new FResizeSection(static_cast<FSequencer&>(Sequencer), SectionHandles, HandleType == Right, bIsSlipping) );
 }
 
-TOptional<float> FSectionEasingHandleHotspot::GetTime() const
+TOptional<FFrameNumber> FSectionEasingHandleHotspot::GetTime() const
 {
 	UMovieSceneSection* ThisSection = Section.GetSectionObject();
 	if (ThisSection)
 	{
 		if (HandleType == ESequencerEasingType::In && !ThisSection->GetEaseInRange().IsEmpty())
 		{
-			return ThisSection->GetEaseInRange().GetUpperBoundValue();
+			return MovieScene::DiscreteExclusiveUpper(ThisSection->GetEaseInRange());
 		}
 		else if (HandleType == ESequencerEasingType::Out && !ThisSection->GetEaseOutRange().IsEmpty())
 		{
-			return ThisSection->GetEaseOutRange().GetLowerBoundValue();
+			return MovieScene::DiscreteInclusiveLower(ThisSection->GetEaseOutRange());
 		}
 	}
-	return TOptional<float>();
+	return TOptional<FFrameNumber>();
 }
 
 void FSectionEasingHandleHotspot::UpdateOnHover(SSequencerTrackArea& InTrackArea, ISequencer& InSequencer) const
@@ -157,7 +175,7 @@ void FSectionEasingHandleHotspot::UpdateOnHover(SSequencerTrackArea& InTrackArea
 	InTrackArea.AttemptToActivateTool(FSequencerEditTool_Movement::Identifier);
 }
 
-bool FSectionEasingHandleHotspot::PopulateContextMenu(FMenuBuilder& MenuBuilder, ISequencer& Sequencer, float MouseDownTime)
+bool FSectionEasingHandleHotspot::PopulateContextMenu(FMenuBuilder& MenuBuilder, ISequencer& Sequencer, FFrameTime MouseDownTime)
 {
 	FEasingContextMenu::BuildMenu(MenuBuilder, { FEasingAreaHandle{Section, HandleType} }, static_cast<FSequencer&>(Sequencer), MouseDownTime);
 	return true;
@@ -168,7 +186,7 @@ TSharedPtr<ISequencerEditToolDragOperation> FSectionEasingHandleHotspot::Initiat
 	return MakeShareable( new FManipulateSectionEasing(static_cast<FSequencer&>(Sequencer), Section, HandleType == ESequencerEasingType::In) );
 }
 
-bool FSectionEasingAreaHotspot::PopulateContextMenu(FMenuBuilder& MenuBuilder, ISequencer& Sequencer, float MouseDownTime)
+bool FSectionEasingAreaHotspot::PopulateContextMenu(FMenuBuilder& MenuBuilder, ISequencer& Sequencer, FFrameTime MouseDownTime)
 {
 	FEasingContextMenu::BuildMenu(MenuBuilder, Easings, static_cast<FSequencer&>(Sequencer), MouseDownTime);
 
diff --git a/Engine/Source/Editor/Sequencer/Private/SequencerHotspots.h b/Engine/Source/Editor/Sequencer/Private/SequencerHotspots.h
--- a/Engine/Source/Editor/Sequencer/Private/SequencerHotspots.h
+++ b/Engine/Source/Editor/Sequencer/Private/SequencerHotspots.h
@@ -22,17 +22,17 @@ enum class ESequencerEasingType
 struct FKeyHotspot
 	: ISequencerHotspot
 {
-	FKeyHotspot(FSequencerSelectedKey InKey)
-		: Key(InKey)
+	FKeyHotspot(TArray<FSequencerSelectedKey> InKeys)
+		: Keys(MoveTemp(InKeys))
 	{ }
 
 	virtual ESequencerHotspot GetType() const override { return ESequencerHotspot::Key; }
 	virtual void UpdateOnHover(SSequencerTrackArea& InTrackArea, ISequencer& InSequencer) const override;
-	virtual TOptional<float> GetTime() const override;
-	virtual bool PopulateContextMenu(FMenuBuilder& MenuBuilder, ISequencer& Sequencer, float MouseDownTime) override;
+	virtual TOptional<FFrameNumber> GetTime() const override;
+	virtual bool PopulateContextMenu(FMenuBuilder& MenuBuilder, ISequencer& Sequencer, FFrameTime MouseDownTime) override;
 
-	/** The key itself */
-	FSequencerSelectedKey Key;
+	/** The keys that are part of this hotspot */
+	TArray<FSequencerSelectedKey> Keys;
 };
 
 
@@ -67,10 +67,10 @@ struct FSectionHotspot
 
 	virtual ESequencerHotspot GetType() const override { return ESequencerHotspot::Section; }
 	virtual void UpdateOnHover(SSequencerTrackArea& InTrackArea, ISequencer& InSequencer) const override;
-	virtual TOptional<float> GetTime() const override;
-	virtual TOptional<float> GetOffsetTime() const override;
+	virtual TOptional<FFrameNumber> GetTime() const override;
+	virtual TOptional<FFrameTime> GetOffsetTime() const override;
 	virtual TSharedPtr<ISequencerEditToolDragOperation> InitiateDrag(ISequencer&) override { return nullptr; }
-	virtual bool PopulateContextMenu(FMenuBuilder& MenuBuilder, ISequencer& Sequencer, float MouseDownTime) override;
+	virtual bool PopulateContextMenu(FMenuBuilder& MenuBuilder, ISequencer& Sequencer, FFrameTime MouseDownTime) override;
 
 	/** Handle to the section */
 	FSectionHandle Section;
@@ -91,7 +91,7 @@ struct FSectionResizeHotspot
 
 	virtual ESequencerHotspot GetType() const override { return HandleType == Left ? ESequencerHotspot::SectionResize_L : ESequencerHotspot::SectionResize_R; }
 	virtual void UpdateOnHover(SSequencerTrackArea& InTrackArea, ISequencer& InSequencer) const override;
-	virtual TOptional<float> GetTime() const override;
+	virtual TOptional<FFrameNumber> GetTime() const override;
 	virtual TSharedPtr<ISequencerEditToolDragOperation> InitiateDrag(ISequencer& Sequencer) override;
 	virtual FCursorReply GetCursor() const { return FCursorReply::Cursor( EMouseCursor::ResizeLeftRight ); }
 
@@ -112,8 +112,8 @@ struct FSectionEasingHandleHotspot
 
 	virtual ESequencerHotspot GetType() const override { return HandleType == ESequencerEasingType::In ? ESequencerHotspot::EaseInHandle : ESequencerHotspot::EaseOutHandle; }
 	virtual void UpdateOnHover(SSequencerTrackArea& InTrackArea, ISequencer& InSequencer) const override;
-	virtual bool PopulateContextMenu(FMenuBuilder& MenuBuilder, ISequencer& Sequencer, float MouseDownTime) override;
-	virtual TOptional<float> GetTime() const override;
+	virtual bool PopulateContextMenu(FMenuBuilder& MenuBuilder, ISequencer& Sequencer, FFrameTime MouseDownTime) override;
+	virtual TOptional<FFrameNumber> GetTime() const override;
 	virtual TSharedPtr<ISequencerEditToolDragOperation> InitiateDrag(ISequencer& Sequencer) override;
 	virtual FCursorReply GetCursor() const { return FCursorReply::Cursor( EMouseCursor::ResizeLeftRight ); }
 
@@ -139,7 +139,7 @@ struct FSectionEasingAreaHotspot
 	FSectionEasingAreaHotspot(const TArray<FEasingAreaHandle>& InEasings, FSectionHandle InVisibleSection) : FSectionHotspot(InVisibleSection), Easings(InEasings) {}
 
 	virtual ESequencerHotspot GetType() const override { return ESequencerHotspot::EasingArea; }
-	virtual bool PopulateContextMenu(FMenuBuilder& MenuBuilder, ISequencer& Sequencer, float MouseDownTime) override;
+	virtual bool PopulateContextMenu(FMenuBuilder& MenuBuilder, ISequencer& Sequencer, FFrameTime MouseDownTime) override;
 
 	bool Contains(FSectionHandle InSection) const { return Easings.ContainsByPredicate([=](const FEasingAreaHandle& InHandle){ return InHandle.Section == InSection; }); }
 
diff --git a/Engine/Source/Editor/Sequencer/Private/SequencerKeyActor.cpp b/Engine/Source/Editor/Sequencer/Private/SequencerKeyActor.cpp
--- a/Engine/Source/Editor/Sequencer/Private/SequencerKeyActor.cpp
+++ b/Engine/Source/Editor/Sequencer/Private/SequencerKeyActor.cpp
@@ -12,6 +12,8 @@
 #include "Modules/ModuleManager.h"
 #include "EditorViewportClient.h"
 #include "UnrealClient.h"
+#include "MovieScene.h"
+#include "Channels/MovieSceneChannelProxy.h"
 
 ASequencerKeyActor::ASequencerKeyActor()
 	: Super()
@@ -75,13 +77,15 @@ void ASequencerKeyActor::PropagateKeyChange()
 		// Mark the track section as dirty
 		TrackSection->Modify();
 
-		// Update the translation keys
-		FRichCurve& TransXCurve = TrackSection->GetTranslationCurve(EAxis::X);
-		FRichCurve& TransYCurve = TrackSection->GetTranslationCurve(EAxis::Y);
-		FRichCurve& TransZCurve = TrackSection->GetTranslationCurve(EAxis::Z);
-		TransXCurve.UpdateOrAddKey(KeyTime, GetActorTransform().GetLocation().X);
-		TransYCurve.UpdateOrAddKey(KeyTime, GetActorTransform().GetLocation().Y);
-		TransZCurve.UpdateOrAddKey(KeyTime, GetActorTransform().GetLocation().Z);
+		FFrameRate   FrameResolution = TrackSection->GetTypedOuter<UMovieScene>()->GetFrameResolution();
+		FFrameNumber FrameNumber     = (KeyTime * FrameResolution).RoundToFrame();
+
+		TArrayView<FMovieSceneFloatChannel*> FloatChannels = TrackSection->GetChannelProxy().GetChannels<FMovieSceneFloatChannel>();
+
+		const FVector Translation = GetActorTransform().GetLocation();
+		FloatChannels[0]->GetInterface().UpdateOrAddKey(FrameNumber, FMovieSceneFloatValue(Translation.X));
+		FloatChannels[1]->GetInterface().UpdateOrAddKey(FrameNumber, FMovieSceneFloatValue(Translation.Y));
+		FloatChannels[2]->GetInterface().UpdateOrAddKey(FrameNumber, FMovieSceneFloatValue(Translation.Z));
 
 		// Draw a single transform track based on the data from this key
 		FEditorViewportClient* ViewportClient = StaticCast<FEditorViewportClient*>(GEditor->GetActiveViewport()->GetClient());
diff --git a/Engine/Source/Editor/Sequencer/Private/SequencerKeyCollection.cpp b/Engine/Source/Editor/Sequencer/Private/SequencerKeyCollection.cpp
new file mode 100644
--- /dev/null
+++ b/Engine/Source/Editor/Sequencer/Private/SequencerKeyCollection.cpp
@@ -0,0 +1,286 @@
+// Copyright 1998-2018 Epic Games, Inc. All Rights Reserved.
+
+#include "SequencerKeyCollection.h"
+#include "IKeyArea.h"
+#include "DisplayNodes/SequencerDisplayNode.h"
+#include "DisplayNodes/SequencerTrackNode.h"
+#include "DisplayNodes/SequencerSectionKeyAreaNode.h"
+
+FSequencerKeyCollectionSignature FSequencerKeyCollectionSignature::FromNodes(const TArray<FSequencerDisplayNode*>& InNodes, FFrameNumber InDuplicateThresholdTime)
+{
+	FSequencerKeyCollectionSignature Result;
+	Result.DuplicateThresholdTime = InDuplicateThresholdTime;
+
+	for (const FSequencerDisplayNode* Node : InNodes)
+	{
+		const FSequencerSectionKeyAreaNode* KeyAreaNode = nullptr;
+
+		check(Node);
+		if (Node->GetType() == ESequencerNode::KeyArea)
+		{
+			KeyAreaNode = static_cast<const FSequencerSectionKeyAreaNode*>(Node);
+		}
+		else if (Node->GetType() == ESequencerNode::Track)
+		{
+			KeyAreaNode = static_cast<const FSequencerTrackNode*>(Node)->GetTopLevelKeyNode().Get();
+		}
+
+		if (KeyAreaNode)
+		{
+			for (const TSharedRef<IKeyArea>& KeyArea : KeyAreaNode->GetAllKeyAreas())
+			{
+				const UMovieSceneSection* Section = KeyArea->GetOwningSection();
+				Result.KeyAreaToSignature.Add(KeyArea, Section ? Section->GetSignature() : FGuid());
+			}
+		}
+	}
+
+	return Result;
+}
+
+FSequencerKeyCollectionSignature FSequencerKeyCollectionSignature::FromNodesRecursive(const TArray<FSequencerDisplayNode*>& InNodes, FFrameNumber InDuplicateThresholdTime)
+{
+	FSequencerKeyCollectionSignature Result;
+	Result.DuplicateThresholdTime = InDuplicateThresholdTime;
+
+	TArray<TSharedRef<FSequencerSectionKeyAreaNode>> AllKeyAreaNodes;
+	AllKeyAreaNodes.Reserve(36);
+	for (FSequencerDisplayNode* Node : InNodes)
+	{
+		if (Node->GetType() == ESequencerNode::KeyArea)
+		{
+			AllKeyAreaNodes.Add(StaticCastSharedRef<FSequencerSectionKeyAreaNode>(Node->AsShared()));
+		}
+
+		Node->GetChildKeyAreaNodesRecursively(AllKeyAreaNodes);
+	}
+
+	for (const TSharedRef<FSequencerSectionKeyAreaNode>& Node : AllKeyAreaNodes)
+	{
+		for (const TSharedRef<IKeyArea>& KeyArea : Node->GetAllKeyAreas())
+		{
+			const UMovieSceneSection* Section = KeyArea->GetOwningSection();
+			Result.KeyAreaToSignature.Add(KeyArea, Section ? Section->GetSignature() : FGuid());
+		}
+	}
+
+	return Result;
+}
+
+FSequencerKeyCollectionSignature FSequencerKeyCollectionSignature::FromNodeRecursive(FSequencerDisplayNode& InNode, UMovieSceneSection* InSection, FFrameNumber InDuplicateThresholdTime)
+{
+	FSequencerKeyCollectionSignature Result;
+	Result.DuplicateThresholdTime = InDuplicateThresholdTime;
+
+	TArray<TSharedRef<FSequencerSectionKeyAreaNode>> AllKeyAreaNodes;
+	AllKeyAreaNodes.Reserve(36);
+	InNode.GetChildKeyAreaNodesRecursively(AllKeyAreaNodes);
+
+	for (const auto& Node : AllKeyAreaNodes)
+	{
+		TSharedPtr<IKeyArea> KeyArea = Node->GetKeyArea(InSection);
+		if (KeyArea.IsValid())
+		{
+			Result.KeyAreaToSignature.Add(KeyArea.ToSharedRef(), InSection ? InSection->GetSignature() : FGuid());
+		}
+	}
+
+	return Result;
+}
+
+bool FSequencerKeyCollectionSignature::HasUncachableContent() const
+{
+	for (auto& Pair : KeyAreaToSignature)
+	{
+		if (!Pair.Value.IsValid())
+		{
+			return true;
+		}
+	}
+	return false;
+}
+
+bool operator!=(const FSequencerKeyCollectionSignature& A, const FSequencerKeyCollectionSignature& B)
+{
+	if (A.HasUncachableContent() || B.HasUncachableContent())
+	{
+		return true;
+	}
+
+	if (A.KeyAreaToSignature.Num() != B.KeyAreaToSignature.Num() || A.DuplicateThresholdTime != B.DuplicateThresholdTime)
+	{
+		return true;
+	}
+
+	for (auto& Pair : A.KeyAreaToSignature)
+	{
+		const FGuid* BSig = B.KeyAreaToSignature.Find(Pair.Key);
+		if (!BSig || *BSig != Pair.Value)
+		{
+			return true;
+		}
+	}
+
+	return false;
+}
+
+bool operator==(const FSequencerKeyCollectionSignature& A, const FSequencerKeyCollectionSignature& B)
+{
+	if (A.HasUncachableContent() || B.HasUncachableContent())
+	{
+		return false;
+	}
+
+	if (A.KeyAreaToSignature.Num() != B.KeyAreaToSignature.Num() || A.DuplicateThresholdTime != B.DuplicateThresholdTime)
+	{
+		return false;
+	}
+
+	for (auto& Pair : A.KeyAreaToSignature)
+	{
+		const FGuid* BSig = B.KeyAreaToSignature.Find(Pair.Key);
+		if (!BSig || *BSig != Pair.Value)
+		{
+			return false;
+		}
+	}
+
+	return true;
+}
+
+bool FSequencerKeyCollection::Update(const FSequencerKeyCollectionSignature& InSignature)
+{
+	if (InSignature == Signature)
+	{
+		return false;
+	}
+
+	TArray<FFrameNumber> AllTimes;
+
+	// Get all the key times for the key areas
+	for (auto& Pair : InSignature.GetKeyAreas())
+	{
+		Pair.Key->GetKeyTimes(AllTimes);
+	}
+
+	AllTimes.Sort();
+
+	GroupedTimes.Reset(AllTimes.Num());
+	int32 Index = 0;
+	while ( Index < AllTimes.Num() )
+	{
+		FFrameNumber PredicateTime = AllTimes[Index];
+		GroupedTimes.Add(PredicateTime);
+		while (Index < AllTimes.Num() && FMath::Abs(AllTimes[Index] - PredicateTime) <= InSignature.GetDuplicateThreshold())
+		{
+			++Index;
+		}
+	}
+	GroupedTimes.Shrink();
+
+	Signature = InSignature;
+
+	return true;
+}
+
+TOptional<FFrameNumber> FSequencerKeyCollection::FindFirstKeyInRange(const TRange<FFrameNumber>& Range, EFindKeyDirection Direction) const
+{
+	TArrayView<const FFrameNumber> KeysInRange = GetKeysInRange(Range);
+	if (KeysInRange.Num())
+	{
+		return Direction == EFindKeyDirection::Forwards ? KeysInRange[0] : KeysInRange[KeysInRange.Num()-1];
+	}
+	return TOptional<FFrameNumber>();
+}
+
+TArrayView<const FFrameNumber> FSequencerKeyCollection::GetKeysInRange(const TRange<FFrameNumber>& Range) const
+{
+	// Binary search the first time that's >= the lower bound
+	int32 FirstVisibleIndex = Range.GetLowerBound().IsClosed() ? Algo::LowerBound(GroupedTimes, Range.GetLowerBoundValue()) : 0;
+	// Binary search the last time that's > the upper bound
+	int32 LastVisibleIndex  = Range.GetUpperBound().IsClosed() ? Algo::UpperBound(GroupedTimes, Range.GetUpperBoundValue()) : GroupedTimes.Num();
+
+	int32 Num = LastVisibleIndex - FirstVisibleIndex;
+	if (GroupedTimes.IsValidIndex(FirstVisibleIndex) && LastVisibleIndex <= GroupedTimes.Num())
+	{
+		return MakeArrayView(&GroupedTimes[FirstVisibleIndex], Num);
+	}
+
+	return TArrayView<const FFrameNumber>();
+}
+/*
+FFrameNumber FSequencerKeyCollection::GetNextKey(FFrameNumber FrameNumber, EFindKeyDirection Direction) const
+{
+	FFrameNumber NextFrame = FrameNumber;
+	int32 Num = GroupedTimes.Num();
+	if (Num > 0)
+	{
+		if (FrameNumber < GroupedTimes[0])
+		{
+			NextFrame = GroupedTimes[0];
+		}
+		else if (FrameNumber > GroupedTimes[Num - 1])
+		{
+			NextFrame = GroupedTimes[Num - 1];
+		}
+		else
+		{
+			int32 Index = Algo::LowerBound(GroupedTimes, FrameNumber);
+			if (GroupedTimes[Index] != FrameNumber)
+			{
+				if (Direction == EFindKeyDirection::Forwards)
+				{
+					NextFrame = GroupedTimes[Index];
+				}
+				else
+				{
+					NextFrame = GroupedTimes[Index - 1];
+				}
+			}
+			else {
+				if (Direction == EFindKeyDirection::Forwards)
+				{
+					if (++Index >= Num)
+					{
+						Index = Num - 1;
+					}
+					NextFrame = GroupedTimes[Index];
+				}
+				else
+				{
+					if (--Index < 0)
+					{
+						Index = 0;
+					}
+					NextFrame = GroupedTimes[Index];
+				}
+			}
+		}
+	}
+	return NextFrame;
+
+}
+
+*/
+
+TOptional<FFrameNumber> FSequencerKeyCollection::GetNextKey(FFrameNumber FrameNumber, EFindKeyDirection Direction) const
+{
+	int32 Index = INDEX_NONE;
+	if (Direction == EFindKeyDirection::Forwards)
+	{
+		Index = Algo::UpperBound(GroupedTimes, FrameNumber);
+	}
+	else
+	{
+		Index = Algo::LowerBound(GroupedTimes, FrameNumber) - 1;
+	}
+
+	if (GroupedTimes.IsValidIndex(Index))
+	{
+		return GroupedTimes[Index];
+	}
+	else
+	{
+		return TOptional<FFrameNumber>();
+	}
+}
diff --git a/Engine/Source/Editor/Sequencer/Private/SequencerModule.cpp b/Engine/Source/Editor/Sequencer/Private/SequencerModule.cpp
--- a/Engine/Source/Editor/Sequencer/Private/SequencerModule.cpp
+++ b/Engine/Source/Editor/Sequencer/Private/SequencerModule.cpp
@@ -15,9 +15,13 @@
 #include "IDetailKeyframeHandler.h"
 #include "AnimatedPropertyKey.h"
 
-
 #define LOCTEXT_NAMESPACE "SequencerEditor"
 
+// Destructor defined in CPP to avoid having to #include SequencerChannelInterface.h in the main module definition
+ISequencerModule::~ISequencerModule()
+{
+}
+
 /**
  * SequencerModule implementation (private)
  */
diff --git a/Engine/Source/Editor/Sequencer/Private/SequencerNodeSortingMethods.h b/Engine/Source/Editor/Sequencer/Private/SequencerNodeSortingMethods.h
new file mode 100644
--- /dev/null
+++ b/Engine/Source/Editor/Sequencer/Private/SequencerNodeSortingMethods.h
@@ -0,0 +1,195 @@
+// Copyright 1998-2018 Epic Games, Inc. All Rights Reserved.
+
+#pragma once
+#include "CoreMinimal.h"
+#include "DisplayNodes/SequencerDisplayNode.h"
+
+/**
+* This sorting delegate sorts based on the vertical position of the node represented by where it falls
+* in the TreeView, ignoring collapsed nodes. This allows you to get a sorting order that matches the visual
+* representation which is more logical for the user.
+*/
+struct FDisplayNodeTreePositionSorter
+{
+	bool operator()(const TSharedRef<FSequencerDisplayNode>& A, const TSharedRef <FSequencerDisplayNode>& B) const
+	{
+		return A->GetVirtualTop() < B->GetVirtualTop();
+	}
+};
+
+/**
+* This sorting delegate sorts based only on the sorting order specified by the node. This respects
+* a user-defined sorting order and ignores the node's type.
+*/
+struct FDisplayNodeSortingOrderSorter
+{
+	bool operator()(const TSharedRef<FSequencerDisplayNode>& A, const TSharedRef<FSequencerDisplayNode>& B) const
+	{
+		// If both nodes have been sorted before we just compare their sorting orders.
+		if (A->GetSortingOrder() >= 0 && B->GetSortingOrder() >= 0)
+		{
+			return A->GetSortingOrder() < B->GetSortingOrder();
+		}
+
+		// Otherwise if one of them has not placed before, we want the lower number higher so that
+		// the unsorted node shows up at the end of the list.
+		return A->GetSortingOrder() > B->GetSortingOrder();
+	}
+};
+
+
+
+/**
+* This sorting delegate sorts based on category and then alphabetically after that. This replicates
+* existing behavior where all folders come first, then all tracks, and finally all object bindings.
+* Within each category these are sorted by display name alphabetically.
+*
+* This does not respect the user-defined sorting order.
+*/
+struct FDisplayNodeCategoricalSorter
+{
+	int32 NodeTypeToFolderSortId(ESequencerNode::Type NodeType) const
+	{
+		switch (NodeType)
+		{
+		case ESequencerNode::Folder:
+			return 0;
+		case ESequencerNode::Track:
+			return 1;
+		case ESequencerNode::Object:
+			return 2;
+		default:
+			return 3;
+		}
+	};
+
+	int32 NodeTypeToObjectSortId(ESequencerNode::Type NodeType) const
+	{
+		switch (NodeType)
+		{
+		case ESequencerNode::Object:
+			return 0;
+		case ESequencerNode::Track:
+			return 1;
+		default:
+			return 2;
+		}
+	};
+
+	bool operator()(const TSharedRef<FSequencerDisplayNode>& A, const TSharedRef<FSequencerDisplayNode>& B) const
+	{
+		TSharedPtr<FSequencerDisplayNode> ParentNode = A->GetParent();
+
+		// If the nodes are root nodes, or in folders and they are the same type, sort by name.
+		if ((ParentNode.IsValid() == false || ParentNode->GetType() == ESequencerNode::Folder) && A->GetType() == B->GetType())
+		{
+			return A->GetDisplayName().ToString() < B->GetDisplayName().ToString();
+		}
+
+		int32 SortIdA;
+		int32 SortIdB;
+
+		// Otherwise if they are root nodes or in folders use the folder sort id.
+		if (ParentNode.IsValid() == false || ParentNode->GetType() == ESequencerNode::Folder)
+		{
+			SortIdA = NodeTypeToFolderSortId(A->GetType());
+			SortIdB = NodeTypeToFolderSortId(B->GetType());
+		}
+		// Otherwise if they are in an object node use the object node sort id.
+		else if (ParentNode->GetType() == ESequencerNode::Object)
+		{
+			SortIdA = NodeTypeToObjectSortId(A->GetType());
+			SortIdB = NodeTypeToObjectSortId(B->GetType());
+		}
+		// Otherwise they are equal, and in a stable sort shouldn't change position.
+		else
+		{
+			SortIdA = 0;
+			SortIdB = 0;
+		}
+
+		return SortIdA < SortIdB;
+	}
+};
+
+
+/**
+* Sorts the supplied unsorted nodes and inserts them into the existing sorted nodes before assigning
+* a sorting order to the newly combined array. Inserted nodes can optionally be relative to an existing
+* node that comes from the ExistingSortedNodes array; this allows you to sort them and then insert them
+* mid array.
+*
+* This function modifies the sorting order of all nodes passed in from both lists, but does not change
+* their actual hierarchy.
+*
+* @param UnsortedNodesToInsert A list of nodes to be sorted by Predicate
+* @param ExistingSortedNdoes A list of nodes that are already sorted that you want to have their indexes
+* adjusted for having the unsorted nodes inserted among them.
+* @param ItemDropZone Specifies the offset for the RelativeToNode (above, below or at the end). This is
+* useful as many of the places this gets called from have this information already.
+* @param Predicate Sorting predicate to sort UnsortedNodesToSort before insertion.
+* @param RelativeToNode An optional node from the ExistingSortedNodes list to insert before/after. If null
+* unsorted nodes will be inserted at the start and all nodes in ExistingSortedNodes will have their sorting
+* order offset by that many.
+*/
+template <class PREDICATE_CLASS>
+void SortAndSetSortingOrder(const TArray<TSharedRef<FSequencerDisplayNode>>& UnsortedNodesToInsert, const TArray<TSharedRef<FSequencerDisplayNode>>& ExistingSortedNodes, TOptional<EItemDropZone> ItemDropZone, const PREDICATE_CLASS& Predicate, TSharedPtr<FSequencerDisplayNode> RelativeToNode)
+{
+	TArray<TSharedRef<FSequencerDisplayNode>> NewSortingOrder = ExistingSortedNodes;
+
+	// ExistingSortedNodes may contain nodes that we do not consider for sorting order (spacers) or nodes we wish to re-sort.
+	TArray<TSharedRef<FSequencerDisplayNode>> ChildrenToRemove;
+	for (TSharedRef<FSequencerDisplayNode> Node : ExistingSortedNodes)
+	{
+		// Discard any of our children that aren't folders, objects or tracks (this ignores spacers, etc.)
+		if (Node->GetType() != ESequencerNode::Folder && Node->GetType() != ESequencerNode::Object && Node->GetType() != ESequencerNode::Track)
+		{
+			ChildrenToRemove.Add(Node);
+			continue;
+		}
+
+		// If we're trying to sort this child remove it from our searchable list as well.
+		if (UnsortedNodesToInsert.Contains(Node))
+		{
+			ChildrenToRemove.Add(Node);
+		}
+	}
+
+	for (TSharedRef<FSequencerDisplayNode> Node : ChildrenToRemove)
+	{
+		NewSortingOrder.Remove(Node);
+	}
+
+	// Now get our index and insert the dragged nodes either before/after/at the end of that index depending on the ItemDropZone.
+	int32 DropAdjustedIndex = NewSortingOrder.Num();
+	int32 RelativeToIndex = NewSortingOrder.IndexOfByKey(RelativeToNode);
+	if (RelativeToIndex >= 0 && ItemDropZone.IsSet())
+	{
+		switch (ItemDropZone.GetValue())
+		{
+		case EItemDropZone::AboveItem:
+			DropAdjustedIndex = RelativeToIndex;
+			break;
+		case EItemDropZone::BelowItem:
+			DropAdjustedIndex = RelativeToIndex + 1;
+			break;
+			// Folders will emit this drop zone which means we want to put it at the end of the folder.
+		case EItemDropZone::OntoItem:
+			DropAdjustedIndex = NewSortingOrder.Num();
+			break;
+		}
+	}
+
+	// Allow our caller to specify the sorting order of new nodes. This allows us to sort incoming nodes by category or vertical position on the treeview, etc.
+	TArray<TSharedRef<FSequencerDisplayNode>> SortedNodes = UnsortedNodesToInsert;
+	SortedNodes.Sort(Predicate);
+
+	// Insert our sorted nodes into our child list so that we have the new absolute order for all items.
+	NewSortingOrder.Insert(SortedNodes, DropAdjustedIndex);
+
+	// And then re-assign the sorting order index of everything in the folder according to the new list.
+	for (int i = 0; i < NewSortingOrder.Num(); i++)
+	{
+		NewSortingOrder[i]->ModifyAndSetSortingOrder(i);
+	}
+};
\ No newline at end of file
diff --git a/Engine/Source/Editor/Sequencer/Private/SequencerNodeTree.cpp b/Engine/Source/Editor/Sequencer/Private/SequencerNodeTree.cpp
--- a/Engine/Source/Editor/Sequencer/Private/SequencerNodeTree.cpp
+++ b/Engine/Source/Editor/Sequencer/Private/SequencerNodeTree.cpp
@@ -16,6 +16,8 @@
 #include "SequencerSectionLayoutBuilder.h"
 #include "ISequencerTrackEditor.h"
 #include "DisplayNodes/SequencerSpacerNode.h"
+#include "Widgets/Views/STableRow.h"
+#include "SequencerNodeSortingMethods.h"
 
 void FSequencerNodeTree::Empty()
 {
@@ -28,73 +30,25 @@ void FSequencerNodeTree::Empty()
 }
 
 
-int32 NodeTypeToFolderSortId(ESequencerNode::Type NodeType)
+void FSequencerNodeTree::AddObjectBindingAndTracks(const FMovieSceneBinding& Binding, TMap<FGuid, const FMovieSceneBinding*>& GuidToBindingMap, TArray< TSharedRef<FSequencerObjectBindingNode> >& OutNodeList)
 {
-	switch ( NodeType )
-	{
-	case ESequencerNode::Folder:
-		return 0;
-	case ESequencerNode::Track:
-		return 1;
-	case ESequencerNode::Object:
-		return 2;
-	default:
-		return 3;
-	}
-}
+	TSharedRef<FSequencerObjectBindingNode> ObjectBindingNode = AddObjectBinding( Binding.GetName(), Binding.GetObjectGuid(), GuidToBindingMap, OutNodeList );
 
-
-int32 NodeTypeToObjectSortId( ESequencerNode::Type NodeType )
-{
-	switch ( NodeType )
-	{
-	case ESequencerNode::Object:
-		return 0;
-	case ESequencerNode::Track:
-		return 1;
-	default:
-		return 2;
-	}
-}
-
-
-struct FDisplayNodeSorter
-{
-	bool operator()( const TSharedRef<FSequencerDisplayNode>& A, const TSharedRef<FSequencerDisplayNode>& B ) const
+	for( UMovieSceneTrack* Track : Binding.GetTracks() )
 	{
-		TSharedPtr<FSequencerDisplayNode> ParentNode = A->GetParent();
-		
-		// If the nodes are root nodes, or in folders and they are the same type, sort by name.
-		if ( (ParentNode.IsValid() == false || ParentNode->GetType() == ESequencerNode::Folder) && A->GetType() == B->GetType() )
+		if (!Sequencer.IsTrackVisible(Track))
 		{
-			return A->GetDisplayName().ToString() < B->GetDisplayName().ToString();
+			continue;
 		}
 
-		int32 SortIdA;
-		int32 SortIdB;
-
-		// Otherwise if they are root nodes or in folders use the folder sort id.
-		if ( ParentNode.IsValid() == false || ParentNode->GetType() == ESequencerNode::Folder )
-		{
-			SortIdA = NodeTypeToFolderSortId( A->GetType() );
-			SortIdB = NodeTypeToFolderSortId( B->GetType() );
-		}
-		// Otherwise if they are in an object node use the object node sort id.
-		else if ( ParentNode->GetType() == ESequencerNode::Object )
-		{
-			SortIdA = NodeTypeToObjectSortId( A->GetType() );
-			SortIdB = NodeTypeToObjectSortId( B->GetType() );
-		}
-		// Otherwise they are equal, and in a stable sort shouldn't change position.
-		else
-		{
-			SortIdA = 0;
-			SortIdB = 0;
-		}
+		// Create the new track node
+		TSharedRef<FSequencerTrackNode> TrackNode = MakeShared<FSequencerTrackNode>(*Track, *FindOrAddTypeEditor(*Track), false, nullptr, *this);
 
-		return SortIdA < SortIdB;
+		// Make the sub tracks and section interfaces for this node, and add it to the object binding node
+		// Note: MakeSubTracksAndSectionInterfaces may return a new parent node
+		ObjectBindingNode->AddTrackNode(MakeSubTracksAndSectionInterfaces(TrackNode, ObjectBindingNode->GetObjectBinding()));
 	}
-};
+}
 
 
 void FSequencerNodeTree::Update()
@@ -139,25 +93,18 @@ void FSequencerNodeTree::Update()
 			continue;
 		}
 
-		TSharedRef<FSequencerObjectBindingNode> ObjectBindingNode = AddObjectBinding( Binding.GetName(), Binding.GetObjectGuid(), GuidToBindingMap, ObjectNodes );
+		AddObjectBindingAndTracks(Binding, GuidToBindingMap, ObjectNodes);
+	}
 
-		for( UMovieSceneTrack* Track : Binding.GetTracks() )
+	// If no bindings were added (presumably because of visibility) but there are bindings, add all regardless of visibility
+	if (!ObjectNodes.Num())
+	{
+		for( const FMovieSceneBinding& Binding : Bindings )
 		{
-			if (!Sequencer.IsTrackVisible(Track))
-			{
-				continue;
-			}
-
-			// Create the new track node
-			TSharedRef<FSequencerTrackNode> TrackNode = MakeShared<FSequencerTrackNode>(*Track, *FindOrAddTypeEditor(*Track), false, nullptr, *this);
-
-			// Make the sub tracks and section interfaces for this node, and add it to the object binding node
-			// Note: MakeSubTracksAndSectionInterfaces may return a new parent node
-			ObjectBindingNode->AddTrackNode(MakeSubTracksAndSectionInterfaces(TrackNode, ObjectBindingNode->GetObjectBinding()));
+			AddObjectBindingAndTracks(Binding, GuidToBindingMap, ObjectNodes);
 		}
 	}
 
-
 	// Cinematic shot track always comes first
 	if (CinematicShotTrack)
 	{
@@ -175,36 +122,62 @@ void FSequencerNodeTree::Update()
 	}
 
 	// Add all other nodes after the camera cut track
-	TArray<TSharedRef<FSequencerDisplayNode>> FolderAndObjectNodes;
+	TArray<TSharedRef<FSequencerDisplayNode>> FolderAndObjectAndTrackNodes;
 	TArray<TSharedRef<FSequencerDisplayNode>> MasterTrackNodesNotInFolders;
-	CreateAndPopulateFolderNodes( MasterTrackNodes, ObjectNodes, MovieScene->GetRootFolders(), FolderAndObjectNodes, MasterTrackNodesNotInFolders );
-	
-	// Add all other master tracks after the camera cut track
-	MasterTrackNodesNotInFolders.Sort(FDisplayNodeSorter());
-	for ( TSharedRef<FSequencerDisplayNode> Node : MasterTrackNodesNotInFolders)
-	{
-		Node->SortChildNodes(FDisplayNodeSorter());
-	}
+	CreateAndPopulateFolderNodes( MasterTrackNodes, ObjectNodes, MovieScene->GetRootFolders(), FolderAndObjectAndTrackNodes, MasterTrackNodesNotInFolders );
 
-	RootNodes.Append( MasterTrackNodesNotInFolders );
+	// Merge the two lists together before sorting them together.
+	FolderAndObjectAndTrackNodes.Append(MasterTrackNodesNotInFolders);
 
-	// Sort the created nodes.
-	FolderAndObjectNodes.Sort(FDisplayNodeSorter());
-	for ( TSharedRef<FSequencerDisplayNode> Node : FolderAndObjectNodes )
+	// Now sort the folders, tracks and objects together based on sorting order.
+	FolderAndObjectAndTrackNodes.Sort(FDisplayNodeSortingOrderSorter());
+
+	for (TSharedRef<FSequencerDisplayNode> Node : FolderAndObjectAndTrackNodes)
 	{
-		Node->SortChildNodes(FDisplayNodeSorter());
+		// Recursively sort the children of these tracks
+		Node->SortChildNodes(FDisplayNodeSortingOrderSorter());
 	}
 
-	RootNodes.Append( FolderAndObjectNodes );
+	// Now that we've sorted the children we normalize their sorting index. This doesn't call Modify (as we're not part of a transaction) but modifies the
+	// in-memory sorting index of the backing data structures. This means the next time the tree is refreshed, the existing nodes will keep their sort
+	// and any new nodes will get pushed to the end again. When an asset is saved it'll write the sorting index to the asset and the next time it is loaded
+	// the Sort function will keep them in the same order even if they exist in a different order within the owning data structures.
+	for (int32 i = 0; i < FolderAndObjectAndTrackNodes.Num(); i++)
+	{
+		FolderAndObjectAndTrackNodes[i]->SetSortingOrder(i);
+		FolderAndObjectAndTrackNodes[i]->Traverse_ParentFirst([](FSequencerDisplayNode& TraversalNode)
+		{
+			int32 ChildIndex = 0;
+			for (uint32 k = 0; k < TraversalNode.GetNumChildren(); k++)
+			{
+				// Sometimes a node can have multiple display node children because the sections within a row
+				// have been re-arranged into an overlapping state. These rows are backed by the same data structure
+				// as their parents, so we skip them instead of incrementing the parent's sorting order.
+				if (TraversalNode.GetChildNodes()[k]->GetType() == ESequencerNode::Track)
+				{
+					TSharedRef<FSequencerTrackNode> FolderNode = StaticCastSharedRef<FSequencerTrackNode>(TraversalNode.GetChildNodes()[k]);
+					if (FolderNode->GetSubTrackMode() == FSequencerTrackNode::ESubTrackMode::SubTrack)
+					{
+						continue;
+					}
+				}
+
+				TraversalNode.GetChildNodes()[k]->SetSortingOrder(ChildIndex);
+				ChildIndex++;
+			}
+			return true;
+		}, true);
+	}
 
+	RootNodes.Append(FolderAndObjectAndTrackNodes);
 	RootNodes.Reserve((RootNodes.Num()-1)*2);
 	for (int32 Index = 1; Index < RootNodes.Num(); Index += 2)
 	{
-		RootNodes.Insert(MakeShareable(new FSequencerSpacerNode(1.f, nullptr, *this)), Index);
+		RootNodes.Insert(MakeShareable(new FSequencerSpacerNode(1.f, nullptr, *this, false)), Index);
 	}
 
 	// Always make space at the end of the tree
-	RootNodes.Add(MakeShared<FSequencerSpacerNode>(20.f, nullptr, *this));
+	RootNodes.Add(MakeShared<FSequencerSpacerNode>(20.f, nullptr, *this, true));
 
 	// Set up virtual offsets, expansion states, and tints
 	float VerticalOffset = 0.f;
@@ -455,6 +428,106 @@ void FSequencerNodeTree::CreateAndPopulateFolderNodes(
 	}
 }
 
+void FSequencerNodeTree::MoveDisplayNodeToRoot(TSharedRef<FSequencerDisplayNode>& Node)
+{
+	// Objects that exist at the root level in a sequence are just removed from the folder they reside in.
+	// When the treeview is refreshed this will cause the regenerated nodes to show up at the root level.
+	TSharedPtr<FSequencerDisplayNode> ParentSeqNode = Node->GetParent();
+	switch (Node->GetType())
+	{
+		case ESequencerNode::Folder:
+		{
+			TSharedRef<FSequencerFolderNode> FolderNode = StaticCastSharedRef<FSequencerFolderNode>(Node);
+			UMovieScene* FocusedMovieScene = GetSequencer().GetFocusedMovieSceneSequence()->GetMovieScene();
+
+			if (ParentSeqNode.IsValid())
+			{
+				checkf(ParentSeqNode->GetType() == ESequencerNode::Folder, TEXT("Can not remove from unsupported parent node."));
+				TSharedPtr<FSequencerFolderNode> ParentFolder = StaticCastSharedPtr<FSequencerFolderNode>(ParentSeqNode);
+				ParentFolder->GetFolder().Modify();
+				ParentFolder->GetFolder().RemoveChildFolder(&FolderNode->GetFolder());
+			}
+			else
+			{
+				FocusedMovieScene->GetRootFolders().Remove(&FolderNode->GetFolder());
+			}
+
+			FocusedMovieScene->GetRootFolders().Add(&FolderNode->GetFolder());
+			break;
+		}
+		case ESequencerNode::Track:
+		{
+			TSharedRef<FSequencerTrackNode> DraggedTrackNode = StaticCastSharedRef<FSequencerTrackNode>(Node);
+			UMovieScene* FocusedMovieScene = GetSequencer().GetFocusedMovieSceneSequence()->GetMovieScene();
+
+			if (ParentSeqNode.IsValid())
+			{
+				checkf(ParentSeqNode->GetType() == ESequencerNode::Folder, TEXT("Can not remove from unsupported parent node."));
+				TSharedPtr<FSequencerFolderNode> ParentFolder = StaticCastSharedPtr<FSequencerFolderNode>(ParentSeqNode);
+				ParentFolder->GetFolder().Modify();
+				ParentFolder->GetFolder().RemoveChildMasterTrack(DraggedTrackNode->GetTrack());
+			}
+			break;
+		}
+		case ESequencerNode::Object:
+		{
+			TSharedRef<FSequencerObjectBindingNode> DraggedObjectBindingNode = StaticCastSharedRef<FSequencerObjectBindingNode>(Node);
+			UMovieScene* FocusedMovieScene = GetSequencer().GetFocusedMovieSceneSequence()->GetMovieScene();
+
+			if (ParentSeqNode.IsValid())
+			{
+				checkf(ParentSeqNode->GetType() == ESequencerNode::Folder, TEXT("Can not remove from unsupported parent node."));
+				TSharedPtr<FSequencerFolderNode> ParentFolder = StaticCastSharedPtr<FSequencerFolderNode>(ParentSeqNode);
+				ParentFolder->GetFolder().Modify();
+				ParentFolder->GetFolder().RemoveChildObjectBinding(DraggedObjectBindingNode->GetObjectBinding());
+			}
+			break;
+		}
+	}
+
+	// Clear the node's parent so that subsequent calls for GetNodePath correctly indicate that they no longer have a parent.
+	Node->ClearParent();
+
+	// Our children have changed parents which means that on subsequent creation they will retrieve their expansion state
+	// from the map using their new path. If the new path already exists the object goes to the state stored at that path.
+	// If the new path does not exist, the object returns to default state and not what is currently displayed. Either way 
+	// causes unexpected user behavior as nodes appear to randomly change expansion state as they are moved around the sequencer.
+
+	// To solve this, we update a node's parent when the node is moved, and then we update their expansion state here
+	// while we still have the current expansion state and the new node path. When the UI is regenerated on the subsequent
+	// refresh call, it will now retrieve the state the node was just in, instead of the state the node was in the last time it
+	// was in that location. This is done recursively as children store absolute paths so they need to be updated too.
+	Node->Traverse_ParentFirst([](FSequencerDisplayNode& TraversalNode)
+	{
+		TraversalNode.GetParentTree().SaveExpansionState(TraversalNode, TraversalNode.IsExpanded());
+		return true;
+	}, true);
+}
+
+void FSequencerNodeTree::SortAllNodesAndDescendants()
+{
+	// Sort the root first
+	SortAndSetSortingOrder(RootNodes, RootNodes, TOptional<EItemDropZone>(), FDisplayNodeCategoricalSorter(), nullptr);
+	
+	// Recursively sort our children looking for folders.
+	TArray<TSharedRef<FSequencerDisplayNode>> ChildNodes = GetRootNodes();
+	for (TSharedRef<FSequencerDisplayNode> Child : ChildNodes)
+	{
+		Child->Traverse_ParentFirst([&](FSequencerDisplayNode& Node)
+		{
+			// Folders are the only type of node that can have children that we can sort, so there is no need
+			// to follow the traversal all the way down.
+			if (Node.GetType() != ESequencerNode::Folder)
+				return false;
+	
+			SortAndSetSortingOrder(Node.GetChildNodes(), Node.GetChildNodes(), TOptional<EItemDropZone>(), FDisplayNodeCategoricalSorter(), nullptr);
+			return true;
+		}, true);
+	}
+
+	// Refresh the tree so that our changes are visible.
+	GetSequencer().RefreshTree();
+}
 
 void FSequencerNodeTree::SaveExpansionState(const FSequencerDisplayNode& Node, bool bExpanded)
 {	
diff --git a/Engine/Source/Editor/Sequencer/Private/SequencerNodeTree.h b/Engine/Source/Editor/Sequencer/Private/SequencerNodeTree.h
--- a/Engine/Source/Editor/Sequencer/Private/SequencerNodeTree.h
+++ b/Engine/Source/Editor/Sequencer/Private/SequencerNodeTree.h
@@ -110,12 +110,18 @@ public:
 	 */
 	FOnUpdated& OnUpdated() { return OnUpdatedDelegate; }
 
+	/** Make the contents of the given node have the root as their parent again instead of their current parent. */
+	void MoveDisplayNodeToRoot(TSharedRef<FSequencerDisplayNode>& Node);
+
+	/** Sorts all nodes and their descendants by category then alphabetically.*/
+	void SortAllNodesAndDescendants();
+
 private:
 	/**
 	 * Finds or adds a type editor for the track
 	 *
 	 * @param Track	The type to find an editor for
-	 * @rerturn The editor for the type
+	 * @return The editor for the type
 	 */
 	TSharedRef<ISequencerTrackEditor> FindOrAddTypeEditor( UMovieSceneTrack& Track );
 
@@ -134,6 +140,11 @@ private:
 	 */
 	void MakeSectionInterfaces( TSharedRef<FSequencerTrackNode> TrackNode, const FGuid& ObjectBinding );
 
+	/*
+	 * Helper function that creates object binding nodes along with their tracks. 
+	 */
+	void AddObjectBindingAndTracks(const FMovieSceneBinding& ObjectBinding, TMap<FGuid, const FMovieSceneBinding*>& GuidToBindingMap, TArray< TSharedRef<FSequencerObjectBindingNode> >& OutNodeList);
+
 	/**
 	 * Creates a new object binding node and any parent binding nodes.
 	 */
diff --git a/Engine/Source/Editor/Sequencer/Private/SequencerObjectChangeListener.cpp b/Engine/Source/Editor/Sequencer/Private/SequencerObjectChangeListener.cpp
--- a/Engine/Source/Editor/Sequencer/Private/SequencerObjectChangeListener.cpp
+++ b/Engine/Source/Editor/Sequencer/Private/SequencerObjectChangeListener.cpp
@@ -92,9 +92,12 @@ void FSequencerObjectChangeListener::BroadcastPropertyChanged( FKeyPropertyParam
 	{
 		// If the property path is not truncated, then we are keying the leafmost property anyways, so set to NAME_None
 		// Otherwise always set to leafmost property of the non-truncated property path, so we correctly pick up struct members
-		const bool bTruncatedPropertyPath = PropertyPath.GetNumProperties() != KeyPropertyParams.PropertyPath.GetNumProperties();
-		FName StructPropertyNameToKey = !bTruncatedPropertyPath ? NAME_None : KeyPropertyParams.PropertyPath.GetLeafMostProperty().Property->GetFName();
-		FPropertyChangedParams Params(KeyableObjects, PropertyPath, StructPropertyNameToKey, KeyPropertyParams.KeyMode);
+		FPropertyPath StructPathToKey;
+		if (PropertyPath.GetNumProperties() < KeyPropertyParams.PropertyPath.GetNumProperties())
+		{
+			StructPathToKey = *KeyPropertyParams.PropertyPath.TrimRoot(PropertyPath.GetNumProperties());
+		}
+		FPropertyChangedParams Params(KeyableObjects, PropertyPath, StructPathToKey, KeyPropertyParams.KeyMode);
 		Delegate.Broadcast(Params);
 	}
 }
diff --git a/Engine/Source/Editor/Sequencer/Private/SequencerSelectedKey.cpp b/Engine/Source/Editor/Sequencer/Private/SequencerSelectedKey.cpp
new file mode 100644
--- /dev/null
+++ b/Engine/Source/Editor/Sequencer/Private/SequencerSelectedKey.cpp
@@ -0,0 +1,127 @@
+// Copyright 1998-2018 Epic Games, Inc. All Rights Reserved.
+
+#include "SequencerSelectedKey.h"
+#include "ISequencerModule.h"
+#include "Modules/ModuleManager.h"
+#include "IKeyArea.h"
+#include "ISequencerChannelInterface.h"
+
+FSelectedKeysByChannelType::FSelectedKeysByChannelType(TArrayView<const FSequencerSelectedKey> InSelectedKeys)
+{
+	for (int32 Index = 0; Index < InSelectedKeys.Num(); ++Index)
+	{
+		FSequencerSelectedKey Key = InSelectedKeys[Index];
+		void* RawChannelPtr = Key.KeyArea->GetChannelPtr();
+
+		if (Key.KeyHandle.IsSet() && RawChannelPtr)
+		{
+			FSelectedChannelInfo* ChannelInfo = ChannelToKeyHandleMap.Find(RawChannelPtr);
+			if (!ChannelInfo)
+			{
+				ChannelInfo = &ChannelToKeyHandleMap.Add(RawChannelPtr, FSelectedChannelInfo(Key.KeyArea->GetChannelTypeID(), Key.KeyArea->GetOwningSection()));
+			}
+
+			ChannelInfo->KeyHandles.Add(Key.KeyHandle.GetValue());
+			ChannelInfo->OriginalIndices.Add(Index);
+		}
+	}
+}
+
+void GetKeyTimes(TArrayView<const FSequencerSelectedKey> InSelectedKeys, TArrayView<FFrameNumber> OutTimes)
+{
+	check(InSelectedKeys.Num() == OutTimes.Num());
+
+	ISequencerModule& SequencerModule = FModuleManager::LoadModuleChecked<ISequencerModule>("Sequencer");
+
+	FSelectedKeysByChannelType KeysByChannel(InSelectedKeys);
+
+	TArray<FFrameNumber> KeyTimesScratch;
+
+	for (TTuple<void*, FSelectedChannelInfo>& Pair : KeysByChannel.ChannelToKeyHandleMap)
+	{
+		void* Channel = Pair.Key;
+		const FSelectedChannelInfo& ChannelInfo = Pair.Value;
+
+		ISequencerChannelInterface* ChannelInterface = SequencerModule.FindChannelInterface(Pair.Value.ChannelTypeID);
+		if (ChannelInterface)
+		{
+			// Resize the scratch buffer to the correct size
+			const int32 NumKeys = ChannelInfo.KeyHandles.Num();
+			KeyTimesScratch.Reset(NumKeys);
+			KeyTimesScratch.SetNum(NumKeys);
+
+			// Populating the key times scratch buffer with the times for these handles
+			ChannelInterface->GetKeyTimes_Raw(Channel, ChannelInfo.KeyHandles, KeyTimesScratch);
+
+			for(int32 Index = 0; Index < KeyTimesScratch.Num(); ++Index)
+			{
+				int32 OriginalIndex = ChannelInfo.OriginalIndices[Index];
+				OutTimes[OriginalIndex] = KeyTimesScratch[Index];
+			}
+		}
+	}
+}
+
+void SetKeyTimes(TArrayView<const FSequencerSelectedKey> InSelectedKeys, TArrayView<const FFrameNumber> InTimes)
+{
+	check(InSelectedKeys.Num() == InTimes.Num());
+
+	ISequencerModule& SequencerModule = FModuleManager::LoadModuleChecked<ISequencerModule>("Sequencer");
+
+	FSelectedKeysByChannelType KeysByChannel(InSelectedKeys);
+
+	TArray<FFrameNumber> KeyTimesScratch;
+
+	for (TTuple<void*, FSelectedChannelInfo>& Pair : KeysByChannel.ChannelToKeyHandleMap)
+	{
+		void* Channel = Pair.Key;
+		const FSelectedChannelInfo& ChannelInfo = Pair.Value;
+
+		ISequencerChannelInterface* ChannelInterface = SequencerModule.FindChannelInterface(Pair.Value.ChannelTypeID);
+		if (ChannelInterface)
+		{
+			KeyTimesScratch.Reset(ChannelInfo.OriginalIndices.Num());
+			for (int32 Index : ChannelInfo.OriginalIndices)
+			{
+				KeyTimesScratch.Add(InTimes[Index]);
+			}
+
+			ChannelInterface->SetKeyTimes_Raw(Channel, ChannelInfo.KeyHandles, KeyTimesScratch);
+		}
+	}
+}
+
+void DuplicateKeys(TArrayView<const FSequencerSelectedKey> InSelectedKeys, TArrayView<FKeyHandle> OutNewHandles)
+{
+	check(InSelectedKeys.Num() == OutNewHandles.Num());
+
+	ISequencerModule& SequencerModule = FModuleManager::LoadModuleChecked<ISequencerModule>("Sequencer");
+
+	FSelectedKeysByChannelType KeysByChannel(InSelectedKeys);
+
+	TArray<FKeyHandle> KeyHandlesScratch;
+	for (TTuple<void*, FSelectedChannelInfo>& Pair : KeysByChannel.ChannelToKeyHandleMap)
+	{
+		void* Channel = Pair.Key;
+		const FSelectedChannelInfo& ChannelInfo = Pair.Value;
+
+		ISequencerChannelInterface* ChannelInterface = SequencerModule.FindChannelInterface(Pair.Value.ChannelTypeID);
+		if (ChannelInterface)
+		{
+			// Resize the scratch buffer to the correct size
+			const int32 NumKeys = ChannelInfo.KeyHandles.Num();
+			KeyHandlesScratch.Reset(NumKeys);
+			KeyHandlesScratch.SetNum(NumKeys);
+
+			// Duplicate the keys, populating the handles scratch buffer
+			ChannelInterface->DuplicateKeys_Raw(Channel, ChannelInfo.KeyHandles, KeyHandlesScratch);
+
+			// Copy the duplicated key handles to the output array view
+			for(int32 Index = 0; Index < KeyHandlesScratch.Num(); ++Index)
+			{
+				int32 OriginalIndex = ChannelInfo.OriginalIndices[Index];
+				OutNewHandles[OriginalIndex] = KeyHandlesScratch[Index];
+			}
+		}
+	}
+}
diff --git a/Engine/Source/Editor/Sequencer/Private/SequencerSelectedKey.h b/Engine/Source/Editor/Sequencer/Private/SequencerSelectedKey.h
--- a/Engine/Source/Editor/Sequencer/Private/SequencerSelectedKey.h
+++ b/Engine/Source/Editor/Sequencer/Private/SequencerSelectedKey.h
@@ -5,9 +5,13 @@
 #include "CoreMinimal.h"
 #include "Curves/KeyHandle.h"
 
+struct FFrameNumber;
+
 class IKeyArea;
 class UMovieSceneSection;
 
+template<typename> class TArrayView;
+
 /**
  * Represents a selected key in the sequencer
  */
@@ -54,3 +58,26 @@ public:
 			KeyHandle.GetValue() == OtherKey.KeyHandle.GetValue();
 	}
 };
+
+struct FSelectedChannelInfo
+{
+	explicit FSelectedChannelInfo(uint32 InChannelTypeID, UMovieSceneSection* InOwningSection)
+		: ChannelTypeID(InChannelTypeID), OwningSection(InOwningSection)
+	{}
+
+	uint32 ChannelTypeID;
+	UMovieSceneSection* OwningSection;
+	TArray<FKeyHandle> KeyHandles;
+	TArray<int32> OriginalIndices;
+};
+
+struct FSelectedKeysByChannelType
+{
+	explicit FSelectedKeysByChannelType(TArrayView<const FSequencerSelectedKey> InSelectedKeys);
+
+	TMap<void*, FSelectedChannelInfo> ChannelToKeyHandleMap;
+};
+
+void GetKeyTimes(TArrayView<const FSequencerSelectedKey> InSelectedKeys, TArrayView<FFrameNumber> OutTimes);
+void SetKeyTimes(TArrayView<const FSequencerSelectedKey> InSelectedKeys, TArrayView<const FFrameNumber> InTimes);
+void DuplicateKeys(TArrayView<const FSequencerSelectedKey> InSelectedKeys, TArrayView<FKeyHandle> OutNewHandles);
\ No newline at end of file
diff --git a/Engine/Source/Editor/Sequencer/Private/SequencerSelection.cpp b/Engine/Source/Editor/Sequencer/Private/SequencerSelection.cpp
--- a/Engine/Source/Editor/Sequencer/Private/SequencerSelection.cpp
+++ b/Engine/Source/Editor/Sequencer/Private/SequencerSelection.cpp
@@ -5,7 +5,6 @@
 #include "DisplayNodes/SequencerObjectBindingNode.h"
 #include "DisplayNodes/SequencerTrackNode.h"
 #include "SequencerCommonHelpers.h"
-#include "IKeyArea.h"
 
 FSequencerSelection::FSequencerSelection()
 	: SuspendBroadcastCount(0)
@@ -123,10 +122,7 @@ void FSequencerSelection::AddToSelection(FSequencerSelectedKey Key)
 	}
 
 	// Deselect any outliner nodes that aren't within the trunk of this key
-	if (Key.KeyArea.IsValid())
-	{
-		EmptySelectedOutlinerNodesWithoutSection(Key.KeyArea->GetOwningSection());
-	}
+	EmptySelectedOutlinerNodesWithoutSection(Key.Section);
 }
 
 void FSequencerSelection::AddToSelection(UMovieSceneSection* Section)
diff --git a/Engine/Source/Editor/Sequencer/Private/SequencerSettings.cpp b/Engine/Source/Editor/Sequencer/Private/SequencerSettings.cpp
--- a/Engine/Source/Editor/Sequencer/Private/SequencerSettings.cpp
+++ b/Engine/Source/Editor/Sequencer/Private/SequencerSettings.cpp
@@ -15,11 +15,8 @@ USequencerSettings::USequencerSettings( const FObjectInitializer& ObjectInitiali
 	bAutoSetTrackDefaults = false;
 	SpawnPosition = SSP_Origin;
 	bCreateSpawnableCameras = true;
-	bShowFrameNumbers = true;
 	bShowRangeSlider = false;
 	bIsSnapEnabled = true;
-	TimeSnapIntervalMode = ESequencerTimeSnapInterval::STSI_Custom;
-	CustomTimeSnapInterval = .05f;
 	bSnapKeyTimesToInterval = true;
 	bSnapKeyTimesToKeys = true;
 	bSnapSectionTimesToInterval = true;
@@ -35,7 +32,6 @@ USequencerSettings::USequencerSettings( const FObjectInitializer& ObjectInitiali
 	bRewindOnRecord = true;
 	ZoomPosition = ESequencerZoomPosition::SZP_CurrentTime;
 	bAutoScrollEnabled = false;
-	bShowCurveEditorCurveToolTips = true;
 	bLinkCurveEditorTimeRange = false;
 	LoopMode = ESequencerLoopMode::SLM_NoLoop;
 	bKeepCursorInPlayRangeWhileScrubbing = false;
@@ -46,23 +42,18 @@ USequencerSettings::USequencerSettings( const FObjectInitializer& ObjectInitiali
 	bInfiniteKeyAreas = false;
 	bShowChannelColors = false;
 	bShowViewportTransportControls = true;
-	bLockPlaybackToAudioClock = false;
 	bAllowPossessionOfPIEViewports = false;
 	bActivateRealtimeViewports = true;
 	bEvaluateSubSequencesInIsolation = false;
-	bRerunConstructionScripts = false;
+	bRerunConstructionScripts = true;
 	bVisualizePreAndPostRoll = true;
 	TrajectoryPathCap = 250;
+	CurveVisibility = ECurveEditorCurveVisibility::SelectedCurves;
+	FrameNumberDisplayFormat = EFrameNumberDisplayFormats::Seconds;
 }
 
 void USequencerSettings::PostEditChangeProperty(FPropertyChangedEvent& PropertyChangedEvent)
 {
-	FName PropertyName = (PropertyChangedEvent.Property != nullptr) ? PropertyChangedEvent.Property->GetFName() : NAME_None;
-	if (PropertyName == GET_MEMBER_NAME_CHECKED(USequencerSettings, bLockPlaybackToAudioClock))
-	{
-		OnLockPlaybackToAudioClockChanged.Broadcast(bLockPlaybackToAudioClock);
-	}
-
 	Super::PostEditChangeProperty(PropertyChangedEvent);
 }
 
@@ -166,20 +157,6 @@ void USequencerSettings::SetCreateSpawnableCameras(bool bInCreateSpawnableCamera
 	}
 }
 
-bool USequencerSettings::GetShowFrameNumbers() const
-{
-	return bShowFrameNumbers;
-}
-
-void USequencerSettings::SetShowFrameNumbers(bool InbShowFrameNumbers)
-{
-	if ( bShowFrameNumbers != InbShowFrameNumbers )
-	{
-		bShowFrameNumbers = InbShowFrameNumbers;
-		SaveConfig();
-	}
-}
-
 bool USequencerSettings::GetShowRangeSlider() const
 {
 	return bShowRangeSlider;
@@ -208,62 +185,6 @@ void USequencerSettings::SetIsSnapEnabled(bool InbIsSnapEnabled)
 	}
 }
 
-float USequencerSettings::GetTimeSnapInterval() const
-{
-	switch (TimeSnapIntervalMode)
-	{
-	case STSI_0_001:
-		return 0.001f;
-	case STSI_0_01:
-		return 0.01f;
-	case STSI_0_1:
-		return 0.1f;
-	case STSI_1:
-		return 1.0f;
-	case STSI_10:
-		return 10.0f;
-	case STSI_100:
-		return 100.0f;
-	case STSI_15Fps:
-		return 1 / 15.0f;
-	case STSI_24Fps:
-		return 1 / 24.0f;
-	case STSI_25Fps:
-		return 1 / 25.0f;
-	case STSI_29_97Fps:
-		return 1 / 29.97f;
-	case STSI_30Fps:
-		return 1 / 30.0f;
-	case STSI_48Fps:
-		return 1 / 48.0f;
-	case STSI_50Fps:
-		return 1 / 50.0f;
-	case STSI_59_94Fps:
-		return 1 / 59.94f;
-	case STSI_60Fps:
-		return 1 / 60.0f;
-	case STSI_120Fps:
-		return 1 / 120.0f;
-	default:
-		return CustomTimeSnapInterval;
-	}
-	return CustomTimeSnapInterval;
-}
-
-float USequencerSettings::GetCustomTimeSnapInterval() const
-{
-	return CustomTimeSnapInterval;
-}
-
-void USequencerSettings::SetCustomTimeSnapInterval(float InCustomTimeSnapInterval)
-{
-	if ( CustomTimeSnapInterval != InCustomTimeSnapInterval )
-	{
-		CustomTimeSnapInterval = InCustomTimeSnapInterval;
-		SaveConfig();
-	}
-}
-
 bool USequencerSettings::GetSnapKeyTimesToInterval() const
 {
 	return bSnapKeyTimesToInterval;
@@ -532,20 +453,6 @@ void USequencerSettings::SetKeepPlayRangeInSectionBounds(bool bInKeepPlayRangeIn
 	}
 }
 
-bool USequencerSettings::GetShowCurveEditorCurveToolTips() const
-{
-	return bShowCurveEditorCurveToolTips;
-}
-
-void USequencerSettings::SetShowCurveEditorCurveToolTips(bool InbShowCurveEditorCurveToolTips)
-{
-	if (bShowCurveEditorCurveToolTips != InbShowCurveEditorCurveToolTips)
-	{
-		bShowCurveEditorCurveToolTips = InbShowCurveEditorCurveToolTips;
-		SaveConfig();
-	}
-}
-
 
 bool USequencerSettings::GetLinkCurveEditorTimeRange() const
 {
@@ -663,21 +570,6 @@ void USequencerSettings::SetActivateRealtimeViewports(bool bInActivateRealtimeVi
 	}
 }
 
-bool USequencerSettings::ShouldLockPlaybackToAudioClock() const
-{
-	return bLockPlaybackToAudioClock;
-}
-
-void USequencerSettings::SetLockPlaybackToAudioClock(bool bInLockPlaybackToAudioClock)
-{
-	if (bLockPlaybackToAudioClock != bInLockPlaybackToAudioClock)
-	{
-		bLockPlaybackToAudioClock = bInLockPlaybackToAudioClock;
-		OnLockPlaybackToAudioClockChanged.Broadcast(bLockPlaybackToAudioClock);
-		SaveConfig();
-	}
-}
-
 bool USequencerSettings::GetAutoSetTrackDefaults() const
 {
 	return bAutoSetTrackDefaults;
@@ -749,3 +641,32 @@ void USequencerSettings::SetShouldShowPrePostRoll(bool bInVisualizePreAndPostRol
 		SaveConfig();
 	}
 }
+
+ECurveEditorCurveVisibility USequencerSettings::GetCurveVisibility() const
+{
+	return CurveVisibility;
+}
+
+void USequencerSettings::SetCurveVisibility(ECurveEditorCurveVisibility InCurveVisibility)
+{
+	if (CurveVisibility != InCurveVisibility)
+	{
+		CurveVisibility = InCurveVisibility;
+		OnCurveEditorCurveVisibilityChanged.Broadcast();
+		SaveConfig();
+	}
+}
+
+USequencerSettings::FOnCurveEditorCurveVisibilityChanged& USequencerSettings::GetOnCurveEditorCurveVisibilityChanged()
+{
+	return OnCurveEditorCurveVisibilityChanged;
+}
+
+void USequencerSettings::SetTimeDisplayFormat(EFrameNumberDisplayFormats InFormat)
+{
+	if (InFormat != FrameNumberDisplayFormat)
+	{
+		FrameNumberDisplayFormat = InFormat;
+		SaveConfig();
+	}
+}
