Copying //UE4/Dev-Sequencer to //UE4/Dev-Main (Source: //UE4/Dev-Sequencer @ 3946692)

diff --git a/Engine/Source/Editor/Sequencer/Private/SequencerTimeSliderController.cpp b/Engine/Source/Editor/Sequencer/Private/SequencerTimeSliderController.cpp
--- a/Engine/Source/Editor/Sequencer/Private/SequencerTimeSliderController.cpp
+++ b/Engine/Source/Editor/Sequencer/Private/SequencerTimeSliderController.cpp
@@ -15,6 +15,10 @@
 #include "EditorStyleSet.h"
 #include "SequencerCommonHelpers.h"
 #include "SequencerSettings.h"
+#include "QualifiedFrameTime.h"
+#include "MovieSceneTimeHelpers.h"
+#include "CommonFrameRates.h"
+#include "Sequencer.h"
 
 #define LOCTEXT_NAMESPACE "TimeSlider"
 
@@ -33,32 +37,27 @@ namespace ScrubConstants
 /** Utility struct for converting between scrub range space and local/absolute screen space */
 struct FSequencerTimeSliderController::FScrubRangeToScreen
 {
-	FVector2D WidgetSize;
+	double ViewStart;
+	float  PixelsPerInput;
 
-	TRange<float> ViewInput;
-	float ViewInputRange;
-	float PixelsPerInput;
-
-	FScrubRangeToScreen(TRange<float> InViewInput, const FVector2D& InWidgetSize )
+	FScrubRangeToScreen(const TRange<double>& InViewInput, const FVector2D& InWidgetSize )
 	{
-		WidgetSize = InWidgetSize;
+		float ViewInputRange = InViewInput.Size<double>();
 
-		ViewInput = InViewInput;
-		ViewInputRange = ViewInput.Size<float>();
-		PixelsPerInput = ViewInputRange > 0 ? ( WidgetSize.X / ViewInputRange ) : 0;
+		ViewStart      = InViewInput.GetLowerBoundValue();
+		PixelsPerInput = ViewInputRange > 0 ? ( InWidgetSize.X / ViewInputRange ) : 0;
 	}
 
 	/** Local Widget Space -> Curve Input domain. */
-	float LocalXToInput(float ScreenX) const
+	double LocalXToInput(float ScreenX) const
 	{
-		float LocalX = ScreenX;
-		return (LocalX/PixelsPerInput) + ViewInput.GetLowerBoundValue();
+		return PixelsPerInput > 0 ? (ScreenX/PixelsPerInput) + ViewStart : ViewStart;
 	}
 
 	/** Curve Input domain -> local Widget Space */
-	float InputToLocalX(float Input) const
+	float InputToLocalX(double Input) const
 	{
-		return (Input - ViewInput.GetLowerBoundValue()) * PixelsPerInput;
+		return (Input - ViewStart) * PixelsPerInput;
 	}
 };
 
@@ -82,35 +81,107 @@ static float GetNextSpacing( uint32 CurrentStep )
 	}
 }
 
-FSequencerTimeSliderController::FSequencerTimeSliderController( const FTimeSliderArgs& InArgs )
-	: TimeSliderArgs( InArgs )
+FSequencerTimeSliderController::FSequencerTimeSliderController( const FTimeSliderArgs& InArgs, TWeakPtr<FSequencer> InWeakSequencer )
+	: WeakSequencer(InWeakSequencer)
+	, TimeSliderArgs( InArgs )
 	, DistanceDragged( 0.0f )
 	, MouseDragType( DRAG_NONE )
 	, bPanning( false )
 {
-	ScrubHandleUp = FEditorStyle::GetBrush( TEXT( "Sequencer.Timeline.ScrubHandleUp" ) ); 
-	ScrubHandleDown = FEditorStyle::GetBrush( TEXT( "Sequencer.Timeline.ScrubHandleDown" ) );
-	ScrubHandleSize = 13.f;
+	ScrubFillBrush              = FEditorStyle::GetBrush( TEXT( "Sequencer.Timeline.ScrubFill" ) );
+	ScrubHandleUpBrush          = FEditorStyle::GetBrush( TEXT( "Sequencer.Timeline.ScrubHandleUp" ) ); 
+	ScrubHandleDownBrush        = FEditorStyle::GetBrush( TEXT( "Sequencer.Timeline.ScrubHandleDown" ) );
+	ClampedScrubHandleUpBrush   = FEditorStyle::GetBrush( TEXT( "Sequencer.Timeline.ClampedScrubHandleUp" ) ); 
+	ClampedScrubHandleDownBrush = FEditorStyle::GetBrush( TEXT( "Sequencer.Timeline.ClampedScrubHandleDown" ) );
 	ContextMenuSuppression = 0;
 }
 
-float FSequencerTimeSliderController::DetermineOptimalSpacing(float InPixelsPerInput, uint32 MinTick, float MinTickSpacing) const
+FFrameTime FSequencerTimeSliderController::ComputeScrubTimeFromMouse(const FGeometry& Geometry, FVector2D ScreenSpacePosition, FScrubRangeToScreen RangeToScreen) const
 {
-	uint32 CurStep = 0;
-
-	// Start with the smallest spacing
-	float Spacing = MinTickSpacing;
+	FVector2D           CursorPos     = Geometry.AbsoluteToLocal( ScreenSpacePosition );
+	double              MouseSeconds  = RangeToScreen.LocalXToInput( CursorPos.X );
+	FFrameTime          ScrubTime     = MouseSeconds * GetFrameResolution();
 
-	if (InPixelsPerInput > 0)
+	if ( TimeSliderArgs.Settings->GetIsSnapEnabled() )
 	{
-		while( Spacing * InPixelsPerInput < MinTick )
+		if (TimeSliderArgs.Settings->GetSnapPlayTimeToKeys())
+		{
+			ScrubTime = SnapTimeToNearestKey(RangeToScreen, CursorPos.X, ScrubTime);
+		}
+		else if (TimeSliderArgs.Settings->GetSnapPlayTimeToInterval())
 		{
-			Spacing = MinTickSpacing * GetNextSpacing( CurStep );
-			CurStep++;
+			ScrubTime = FFrameRate::Snap(ScrubTime, GetFrameResolution(), GetPlayRate());
 		}
 	}
 
-	return Spacing;
+	if (TimeSliderArgs.Settings->ShouldKeepCursorInPlayRangeWhileScrubbing())
+	{
+		ScrubTime = MovieScene::ClampToDiscreteRange(ScrubTime, TimeSliderArgs.PlaybackRange.Get());
+	}
+
+	return ScrubTime;
+}
+
+FFrameTime FSequencerTimeSliderController::ComputeFrameTimeFromMouse(const FGeometry& Geometry, FVector2D ScreenSpacePosition, FScrubRangeToScreen RangeToScreen, bool CheckSnapping) const
+{
+	FVector2D CursorPos  = Geometry.AbsoluteToLocal( ScreenSpacePosition );
+	double    MouseValue = RangeToScreen.LocalXToInput( CursorPos.X );
+
+	if (CheckSnapping && TimeSliderArgs.Settings->GetIsSnapEnabled())
+	{
+		FFrameNumber        SnappedFrameNumber = (MouseValue * GetPlayRate()).FloorToFrame();
+		FQualifiedFrameTime RoundedPlayFrame   = FQualifiedFrameTime(SnappedFrameNumber, GetPlayRate());
+		return RoundedPlayFrame.ConvertTo(GetFrameResolution());
+	}
+	else
+	{
+		return MouseValue * GetFrameResolution();
+	}
+}
+
+FSequencerTimeSliderController::FScrubPixelRange FSequencerTimeSliderController::GetHitTestScrubberPixelRange(FFrameTime ScrubTime, const FScrubRangeToScreen& RangeToScreen) const
+{
+	static const float DragToleranceSlateUnits = 2.f, MouseTolerance = 2.f;
+	return GetScrubberPixelRange(ScrubTime, GetFrameResolution(), GetPlayRate(), RangeToScreen, DragToleranceSlateUnits + MouseTolerance);
+}
+
+FSequencerTimeSliderController::FScrubPixelRange FSequencerTimeSliderController::GetScrubberPixelRange(FFrameTime ScrubTime, const FScrubRangeToScreen& RangeToScreen) const
+{
+	return GetScrubberPixelRange(ScrubTime, GetFrameResolution(), GetPlayRate(), RangeToScreen);
+}
+
+FSequencerTimeSliderController::FScrubPixelRange FSequencerTimeSliderController::GetScrubberPixelRange(FFrameTime ScrubTime, FFrameRate Resolution, FFrameRate PlayRate, const FScrubRangeToScreen& RangeToScreen, float DilationPixels) const
+{
+	const FFrameNumber Frame = ScrubTime.FloorToFrame();// ConvertFrameTime(ScrubTime, Resolution, PlayRate).FloorToFrame();
+
+	float StartPixel = RangeToScreen.InputToLocalX(  Frame    / Resolution);
+	float EndPixel   = RangeToScreen.InputToLocalX( (Frame+1) / Resolution);
+
+	{
+		float RoundedStartPixel = FMath::RoundToInt(StartPixel);
+		EndPixel -= (StartPixel - RoundedStartPixel);
+
+		StartPixel = RoundedStartPixel;
+		EndPixel = FMath::Max(EndPixel, StartPixel + 1);
+	}
+
+	FScrubPixelRange Range;
+
+	float MinScrubSize = 14.f;
+	Range.bClamped    = EndPixel - StartPixel < MinScrubSize;
+	Range.Range       = TRange<float>(StartPixel, EndPixel);
+	if (Range.bClamped)
+	{
+		Range.HandleRange = TRange<float>(
+			(StartPixel + EndPixel - MinScrubSize) * .5f,
+			(StartPixel + EndPixel + MinScrubSize) * .5f);
+	}
+	else
+	{
+		Range.HandleRange = Range.Range;
+	}
+
+	return Range;
 }
 
 struct FSequencerTimeSliderController::FDrawTickArgs
@@ -136,116 +207,93 @@ struct FSequencerTimeSliderController::FDrawTickArgs
 	
 };
 
-void FSequencerTimeSliderController::DrawTicks( FSlateWindowElementList& OutDrawElements, const struct FScrubRangeToScreen& RangeToScreen, FDrawTickArgs& InArgs ) const
+void FSequencerTimeSliderController::DrawTicks( FSlateWindowElementList& OutDrawElements, const TRange<double>& ViewRange, const FScrubRangeToScreen& RangeToScreen, FDrawTickArgs& InArgs ) const
 {
-	// The math here breaks down when pixels per input is near zero or zero, so just skip drawing ticks to avoid an infinite loop.
-	if (FMath::IsNearlyZero(RangeToScreen.PixelsPerInput))
+	TSharedPtr<FSequencer> Sequencer = WeakSequencer.Pin();
+	if (!Sequencer.IsValid())
+	{
+		return;
+	}
+
+	FFrameRate     FrameResolution = GetFrameResolution();
+	FPaintGeometry PaintGeometry   = InArgs.AllottedGeometry.ToPaintGeometry();
+	FSlateFontInfo SmallLayoutFont = FCoreStyle::GetDefaultFontStyle("Regular", 8);
+
+	double MajorGridStep  = 0.0;
+	int32  MinorDivisions = 0;
+	if (!Sequencer->GetGridMetrics(InArgs.AllottedGeometry.Size.X, MajorGridStep, MinorDivisions))
 	{
 		return;
 	}
 
-	float MinDisplayTickSpacing = ScrubConstants::MinDisplayTickSpacing;
-	if (SequencerSnapValues::IsTimeSnapIntervalFrameRate(TimeSliderArgs.TimeSnapInterval.Get()) && TimeSliderArgs.TimeSnapInterval.Get())
+	if (InArgs.bOnlyDrawMajorTicks)
 	{
-		MinDisplayTickSpacing = TimeSliderArgs.TimeSnapInterval.Get();
+		MinorDivisions = 0;
 	}
 
-	const float Spacing = DetermineOptimalSpacing( RangeToScreen.PixelsPerInput, ScrubConstants::MinPixelsPerDisplayTick, MinDisplayTickSpacing );
+	TArray<FVector2D> LinePoints;
+	LinePoints.SetNumUninitialized(2);
 
-	// Sub divisions
-	// @todo Sequencer may need more robust calculation
-	const int32 Divider = 10;
-	// For slightly larger halfway tick mark
-	const int32 HalfDivider = Divider / 2;
-	// Find out where to start from
-	int32 OffsetNum = FMath::FloorToInt(RangeToScreen.ViewInput.GetLowerBoundValue() / Spacing);
-	
-	FSlateFontInfo SmallLayoutFont = FCoreStyle::GetDefaultFontStyle("Regular", 8);
+	const bool bAntiAliasLines = false;
 
-	TArray<FVector2D> LinePoints;
-	LinePoints.AddUninitialized(2);
+	const double FirstMajorLine = FMath::FloorToDouble(ViewRange.GetLowerBoundValue() / MajorGridStep) * MajorGridStep;
+	const double LastMajorLine  = FMath::CeilToDouble(ViewRange.GetUpperBoundValue() / MajorGridStep) * MajorGridStep;
 
-	float Seconds = 0;
-	while( (Seconds = OffsetNum*Spacing) < RangeToScreen.ViewInput.GetUpperBoundValue() )
+	for (double CurrentMajorLine = FirstMajorLine; CurrentMajorLine < LastMajorLine; CurrentMajorLine += MajorGridStep)
 	{
-		// X position local to start of the widget area
-		float XPos = RangeToScreen.InputToLocalX( Seconds );
-		uint32 AbsOffsetNum = FMath::Abs(OffsetNum);
+		float MajorLinePx = RangeToScreen.InputToLocalX( CurrentMajorLine );
 
-		if ( AbsOffsetNum % Divider == 0 )
-		{
-			FVector2D Offset( XPos, InArgs.TickOffset );
-			FVector2D TickSize( 0.0f, InArgs.MajorTickHeight );
+		LinePoints[0] = FVector2D( MajorLinePx, InArgs.TickOffset );
+		LinePoints[1] = FVector2D( MajorLinePx, InArgs.TickOffset + InArgs.MajorTickHeight );
 
-			LinePoints[0] = FVector2D( 0.0f,1.0f);
-			LinePoints[1] = TickSize;
+		// Draw each tick mark
+		FSlateDrawElement::MakeLines(
+			OutDrawElements,
+			InArgs.StartLayer,
+			PaintGeometry,
+			LinePoints,
+			InArgs.DrawEffects,
+			InArgs.TickColor,
+			bAntiAliasLines
+			);
 
-			// lines should not need anti-aliasing
-			const bool bAntiAliasLines = false;
+		if (!InArgs.bOnlyDrawMajorTicks)
+		{
+			FString FrameString = TimeSliderArgs.NumericTypeInterface->ToString((CurrentMajorLine * FrameResolution).RoundToFrame().Value);
 
-			// Draw each tick mark
-			FSlateDrawElement::MakeLines(
+			// Space the text between the tick mark but slightly above
+			FVector2D TextOffset( MajorLinePx + 5.f, InArgs.bMirrorLabels ? 3.f : FMath::Abs( InArgs.AllottedGeometry.Size.Y - (InArgs.MajorTickHeight+3.f) ) );
+			FSlateDrawElement::MakeText(
 				OutDrawElements,
-				InArgs.StartLayer,
-				InArgs.AllottedGeometry.ToPaintGeometry( Offset, TickSize ),
-				LinePoints,
+				InArgs.StartLayer+1, 
+				InArgs.AllottedGeometry.ToPaintGeometry( TextOffset, InArgs.AllottedGeometry.Size ), 
+				FrameString, 
+				SmallLayoutFont,
 				InArgs.DrawEffects,
-				InArgs.TickColor,
-				bAntiAliasLines
-				);
-
-			if( !InArgs.bOnlyDrawMajorTicks )
-			{
-				FString FrameString;
-				if (SequencerSnapValues::IsTimeSnapIntervalFrameRate(TimeSliderArgs.TimeSnapInterval.Get()) && TimeSliderArgs.Settings->GetShowFrameNumbers())
-				{
-					FrameString = FString::Printf( TEXT("%d"), TimeToFrame(Seconds));
-				}
-				else
-				{
-					FrameString = Spacing == ScrubConstants::MinDisplayTickSpacing ? FString::Printf( TEXT("%.3f"), Seconds ) : FString::Printf( TEXT("%.2f"), Seconds );
-				}
-
-				// Space the text between the tick mark but slightly above
-				const TSharedRef< FSlateFontMeasure > FontMeasureService = FSlateApplication::Get().GetRenderer()->GetFontMeasureService();
-				FVector2D TextSize = FontMeasureService->Measure(FrameString, SmallLayoutFont);
-				FVector2D TextOffset( XPos + 5.f, InArgs.bMirrorLabels ? 3.f : FMath::Abs( InArgs.AllottedGeometry.Size.Y - (InArgs.MajorTickHeight+3.f) ) );
-				FSlateDrawElement::MakeText(
-					OutDrawElements,
-					InArgs.StartLayer+1, 
-					InArgs.AllottedGeometry.ToPaintGeometry( TextOffset, TextSize ), 
-					FrameString, 
-					SmallLayoutFont,
-					InArgs.DrawEffects,
-					InArgs.TickColor*0.65f 
-				);
-			}
+				InArgs.TickColor*0.65f 
+			);
 		}
-		else if( !InArgs.bOnlyDrawMajorTicks )
+
+		for (int32 Step = 1; Step < MinorDivisions; ++Step)
 		{
 			// Compute the size of each tick mark.  If we are half way between to visible values display a slightly larger tick mark
-			const float MinorTickHeight = AbsOffsetNum % HalfDivider == 0 ? 6.0f : 2.0f;
+			const float MinorTickHeight = ( (MinorDivisions % 2 == 0) && (Step % (MinorDivisions/2)) == 0 ) ? 6.0f : 2.0f;
+			const float MinorLinePx = RangeToScreen.InputToLocalX( CurrentMajorLine + Step*MajorGridStep/MinorDivisions );
 
-			FVector2D Offset(XPos, InArgs.bMirrorLabels ? 0.0f : FMath::Abs( InArgs.AllottedGeometry.Size.Y - MinorTickHeight ) );
-			FVector2D TickSize(0.0f, MinorTickHeight);
+			LinePoints[0] = FVector2D(MinorLinePx, InArgs.bMirrorLabels ? 0.0f : FMath::Abs( InArgs.AllottedGeometry.Size.Y - MinorTickHeight ) );
+			LinePoints[1] = FVector2D(MinorLinePx, LinePoints[0].Y + MinorTickHeight);
 
-			LinePoints[0] = FVector2D(0.0f,1.0f);
-			LinePoints[1] = TickSize;
-
-			const bool bAntiAlias = false;
 			// Draw each sub mark
 			FSlateDrawElement::MakeLines(
 				OutDrawElements,
 				InArgs.StartLayer,
-				InArgs.AllottedGeometry.ToPaintGeometry( Offset, TickSize ),
+				PaintGeometry,
 				LinePoints,
 				InArgs.DrawEffects,
 				InArgs.TickColor,
-				bAntiAlias
+				bAntiAliasLines
 			);
 		}
-		// Advance to next tick mark
-		++OffsetNum;
 	}
 }
 
@@ -255,10 +303,10 @@ int32 FSequencerTimeSliderController::OnPaintTimeSlider( bool bMirrorLabels, con
 	const bool bEnabled = bParentEnabled;
 	const ESlateDrawEffect DrawEffects = bEnabled ? ESlateDrawEffect::None : ESlateDrawEffect::DisabledEffect;
 
-	TRange<float> LocalViewRange = TimeSliderArgs.ViewRange.Get();
-	const float LocalViewRangeMin = LocalViewRange.GetLowerBoundValue();
-	const float LocalViewRangeMax = LocalViewRange.GetUpperBoundValue();
-	const float LocalSequenceLength = LocalViewRangeMax-LocalViewRangeMin;
+	TRange<double> LocalViewRange      = TimeSliderArgs.ViewRange.Get();
+	const float    LocalViewRangeMin   = LocalViewRange.GetLowerBoundValue();
+	const float    LocalViewRangeMax   = LocalViewRange.GetUpperBoundValue();
+	const float    LocalSequenceLength = LocalViewRangeMax-LocalViewRangeMin;
 	
 	FVector2D Scale = FVector2D(1.0f,1.0f);
 	if ( LocalSequenceLength > 0)
@@ -281,7 +329,7 @@ int32 FSequencerTimeSliderController::OnPaintTimeSlider( bool bMirrorLabels, con
 			Args.MajorTickHeight = MajorTickHeight;
 		}
 
-		DrawTicks( OutDrawElements, RangeToScreen, Args );
+		DrawTicks( OutDrawElements, LocalViewRange, RangeToScreen, Args );
 
 		// draw playback & selection range
 		FPaintPlaybackRangeArgs PlaybackRangeArgs(
@@ -296,12 +344,14 @@ int32 FSequencerTimeSliderController::OnPaintTimeSlider( bool bMirrorLabels, con
 		PlaybackRangeArgs.SolidFillOpacity = 0.05f;
 		LayerId = DrawSelectionRange(AllottedGeometry, MyCullingRect, OutDrawElements, LayerId, RangeToScreen, PlaybackRangeArgs);
 
-		float HalfSize = FMath::CeilToFloat(ScrubHandleSize / 2.0f);
-
 		// Draw the scrub handle
-		float XPos = RangeToScreen.InputToLocalX( TimeSliderArgs.ScrubPosition.Get() );
+		FFrameTime       ScrubPosition      = TimeSliderArgs.ScrubPosition.Get();
+		FScrubPixelRange ScrubberPixelRange = GetScrubberPixelRange(ScrubPosition, RangeToScreen);
+		const float      HandleStart        = ScrubberPixelRange.HandleRange.GetLowerBoundValue();
+		const float      HandleEnd          = ScrubberPixelRange.HandleRange.GetUpperBoundValue();
+
 		const int32 ArrowLayer = LayerId + 2;
-		FPaintGeometry MyGeometry =	AllottedGeometry.ToPaintGeometry( FVector2D( XPos-HalfSize, 0 ), FVector2D( ScrubHandleSize, AllottedGeometry.Size.Y ) );
+		FPaintGeometry MyGeometry =	AllottedGeometry.ToPaintGeometry( FVector2D( HandleStart, 0 ), FVector2D( HandleEnd - HandleStart, AllottedGeometry.Size.Y ) );
 		FLinearColor ScrubColor = InWidgetStyle.GetColorAndOpacityTint();
 		{
 			// @todo Sequencer this color should be specified in the style
@@ -310,71 +360,59 @@ int32 FSequencerTimeSliderController::OnPaintTimeSlider( bool bMirrorLabels, con
 			ScrubColor.G *= 0.2f;
 		}
 
-		FSlateDrawElement::MakeBox( 
+		const FSlateBrush* Brush = ScrubberPixelRange.bClamped
+			? ( bMirrorLabels ? ClampedScrubHandleUpBrush : ClampedScrubHandleDownBrush )
+			: ( bMirrorLabels ? ScrubHandleUpBrush        : ScrubHandleDownBrush );
+
+		FSlateDrawElement::MakeBox(
 			OutDrawElements,
-			ArrowLayer, 
+			ArrowLayer,
 			MyGeometry,
-			bMirrorLabels ? ScrubHandleUp : ScrubHandleDown,
-			DrawEffects, 
+			Brush,
+			DrawEffects,
 			ScrubColor
 		);
 
-		// Draw the current time next to the scrub handle
-		float Time = TimeSliderArgs.ScrubPosition.Get();
-		FString FrameString;
-
-		if (SequencerSnapValues::IsTimeSnapIntervalFrameRate(TimeSliderArgs.TimeSnapInterval.Get()) && TimeSliderArgs.Settings->GetShowFrameNumbers())
 		{
-			float FrameRate = 1.0f/TimeSliderArgs.TimeSnapInterval.Get();
-			float FrameTime = Time * FrameRate;
-			int32 Frame = SequencerHelpers::TimeToFrame(Time, FrameRate);
-			const float FrameTolerance = 0.001f;
+			// Draw the current time next to the scrub handle
+			FString FrameString;
 
-			if (FMath::IsNearlyEqual(FrameTime, (float)Frame, FrameTolerance) || TimeSliderArgs.PlaybackStatus.Get() == EMovieScenePlayerStatus::Playing)
-			{
-				FrameString = FString::Printf( TEXT("%d"), TimeToFrame(Time));
-			}
-			else
-			{
-				FrameString = FString::Printf( TEXT("%.3f"), FrameTime);
-			}
-		}
-		else
-		{
-			FrameString = FString::Printf( TEXT("%.2f"), Time );
-		}
+			FrameString = TimeSliderArgs.NumericTypeInterface->ToString(TimeSliderArgs.ScrubPosition.Get().GetFrame().Value);
 
-		FSlateFontInfo SmallLayoutFont = FCoreStyle::GetDefaultFontStyle("Regular", 10);
+			FSlateFontInfo SmallLayoutFont = FCoreStyle::GetDefaultFontStyle("Regular", 10);
 
-		const TSharedRef< FSlateFontMeasure > FontMeasureService = FSlateApplication::Get().GetRenderer()->GetFontMeasureService();
-		FVector2D TextSize = FontMeasureService->Measure(FrameString, SmallLayoutFont);
+			const TSharedRef< FSlateFontMeasure > FontMeasureService = FSlateApplication::Get().GetRenderer()->GetFontMeasureService();
+			FVector2D TextSize = FontMeasureService->Measure(FrameString, SmallLayoutFont);
 
-		// Flip the text position if getting near the end of the view range
-		if ((AllottedGeometry.Size.X - XPos) < (TextSize.X + 14.f))
-		{
-			XPos = XPos - TextSize.X - 12.f;
-		}
-		else
-		{
-			XPos = XPos + 10.f;
-		}
+			// Flip the text position if getting near the end of the view range
+			static const float TextOffsetPx  = 2.f;
+			bool  bDrawLeft    = (AllottedGeometry.Size.X - HandleEnd) < (TextSize.X + 14.f) - TextOffsetPx;
+			float TextPosition = bDrawLeft ? HandleStart - TextSize.X - TextOffsetPx : HandleEnd + TextOffsetPx;
 
-		FVector2D TextOffset( XPos,  Args.bMirrorLabels ? TextSize.Y-6.f : Args.AllottedGeometry.Size.Y - (Args.MajorTickHeight+TextSize.Y) );
+			FVector2D TextOffset( TextPosition, Args.bMirrorLabels ? TextSize.Y-6.f : Args.AllottedGeometry.Size.Y - (Args.MajorTickHeight+TextSize.Y) );
 
-		FSlateDrawElement::MakeText(
-			OutDrawElements,
-			Args.StartLayer+1, 
-			Args.AllottedGeometry.ToPaintGeometry( TextOffset, TextSize ), 
-			FrameString, 
-			SmallLayoutFont,
-			Args.DrawEffects,
-			Args.TickColor 
-		);
+			FSlateDrawElement::MakeText(
+				OutDrawElements,
+				Args.StartLayer+1, 
+				Args.AllottedGeometry.ToPaintGeometry( TextOffset, TextSize ), 
+				FrameString, 
+				SmallLayoutFont,
+				Args.DrawEffects,
+				Args.TickColor 
+			);
+		}
 		
 		if (MouseDragType == DRAG_SETTING_RANGE)
 		{
-			float MouseStartPosX = RangeToScreen.InputToLocalX(MouseDownRange[0]);
-			float MouseEndPosX = RangeToScreen.InputToLocalX(MouseDownRange[1]);
+			FFrameRate Resolution     = GetFrameResolution();
+			FFrameTime MouseDownTime[2];
+
+			FScrubRangeToScreen MouseDownRange(TimeSliderArgs.ViewRange.Get(), MouseDownGeometry.Size);
+			MouseDownTime[0] = ComputeFrameTimeFromMouse(MouseDownGeometry, MouseDownPosition[0], MouseDownRange);
+			MouseDownTime[1] = ComputeFrameTimeFromMouse(MouseDownGeometry, MouseDownPosition[1], MouseDownRange);
+
+			float      MouseStartPosX = RangeToScreen.InputToLocalX(MouseDownTime[0] / Resolution);
+			float      MouseEndPosX   = RangeToScreen.InputToLocalX(MouseDownTime[1] / Resolution);
 
 			float RangePosX = MouseStartPosX < MouseEndPosX ? MouseStartPosX : MouseEndPosX;
 			float RangeSizeX = FMath::Abs(MouseStartPosX - MouseEndPosX);
@@ -383,7 +421,7 @@ int32 FSequencerTimeSliderController::OnPaintTimeSlider( bool bMirrorLabels, con
 				OutDrawElements,
 				LayerId+1,
 				AllottedGeometry.ToPaintGeometry( FVector2D(RangePosX, 0.f), FVector2D(RangeSizeX, AllottedGeometry.Size.Y) ),
-				bMirrorLabels ? ScrubHandleDown : ScrubHandleUp,
+				bMirrorLabels ? ClampedScrubHandleDownBrush : ClampedScrubHandleUpBrush,
 				DrawEffects,
 				MouseStartPosX < MouseEndPosX ? FLinearColor(0.5f, 0.5f, 0.5f) : FLinearColor(0.25f, 0.3f, 0.3f)
 			);
@@ -398,7 +436,7 @@ int32 FSequencerTimeSliderController::OnPaintTimeSlider( bool bMirrorLabels, con
 
 int32 FSequencerTimeSliderController::DrawSelectionRange(const FGeometry& AllottedGeometry, const FSlateRect& MyCullingRect, FSlateWindowElementList& OutDrawElements, int32 LayerId, const FScrubRangeToScreen& RangeToScreen, const FPaintPlaybackRangeArgs& Args) const
 {
-	const TRange<float> SelectionRange = TimeSliderArgs.SelectionRange.Get();
+	TRange<double> SelectionRange = TimeSliderArgs.SelectionRange.Get() / GetFrameResolution();
 
 	if (!SelectionRange.IsEmpty())
 	{
@@ -448,12 +486,9 @@ int32 FSequencerTimeSliderController::DrawPlaybackRange(const FGeometry& Allotte
 		return LayerId;
 	}
 
-	TOptional<TRange<float>> SubSequenceRangeValue;
-	SubSequenceRangeValue = TimeSliderArgs.SubSequenceRange.Get(SubSequenceRangeValue);
-
-	const uint8 OpacityBlend = SubSequenceRangeValue.IsSet() ? 128 : 255;
+	const uint8 OpacityBlend = TimeSliderArgs.SubSequenceRange.Get().IsSet() ? 128 : 255;
 
-	const TRange<float> PlaybackRange = TimeSliderArgs.PlaybackRange.Get();
+	const TRange<double> PlaybackRange = TimeSliderArgs.PlaybackRange.Get() / GetFrameResolution();
 	const float PlaybackRangeL = RangeToScreen.InputToLocalX(PlaybackRange.GetLowerBoundValue()) - 1;
 	const float PlaybackRangeR = RangeToScreen.InputToLocalX(PlaybackRange.GetUpperBoundValue()) + 1;
 
@@ -499,7 +534,7 @@ int32 FSequencerTimeSliderController::DrawPlaybackRange(const FGeometry& Allotte
 
 int32 FSequencerTimeSliderController::DrawSubSequenceRange(const FGeometry& AllottedGeometry, const FSlateRect& MyCullingRect, FSlateWindowElementList& OutDrawElements, int32 LayerId, const FScrubRangeToScreen& RangeToScreen, const FPaintPlaybackRangeArgs& Args) const
 {
-	TOptional<TRange<float>> RangeValue;
+	TOptional<TRange<FFrameNumber>> RangeValue;
 	RangeValue = TimeSliderArgs.SubSequenceRange.Get(RangeValue);
 
 	if (!RangeValue.IsSet() || RangeValue->IsEmpty())
@@ -507,8 +542,12 @@ int32 FSequencerTimeSliderController::DrawSubSequenceRange(const FGeometry& Allo
 		return LayerId;
 	}
 
-	const float SubSequenceRangeL = RangeToScreen.InputToLocalX(RangeValue->GetLowerBoundValue()) - 1;
-	const float SubSequenceRangeR = RangeToScreen.InputToLocalX(RangeValue->GetUpperBoundValue()) + 1;
+	const FFrameRate   Resolution = GetFrameResolution();
+	const FFrameNumber LowerFrame = MovieScene::DiscreteInclusiveLower(RangeValue.GetValue());
+	const FFrameNumber UpperFrame = MovieScene::DiscreteExclusiveUpper(RangeValue.GetValue());
+
+	const float SubSequenceRangeL = RangeToScreen.InputToLocalX(LowerFrame / Resolution) - 1;
+	const float SubSequenceRangeR = RangeToScreen.InputToLocalX(UpperFrame / Resolution) + 1;
 
 	static const FSlateBrush* LineBrushL(FEditorStyle::GetBrush("Sequencer.Timeline.PlayRange_L"));
 	static const FSlateBrush* LineBrushR(FEditorStyle::GetBrush("Sequencer.Timeline.PlayRange_R"));
@@ -577,19 +616,19 @@ int32 FSequencerTimeSliderController::DrawSubSequenceRange(const FGeometry& Allo
 FReply FSequencerTimeSliderController::OnMouseButtonDown( SWidget& WidgetOwner, const FGeometry& MyGeometry, const FPointerEvent& MouseEvent )
 {
 	DistanceDragged = 0;
-
-	FScrubRangeToScreen RangeToScreen( TimeSliderArgs.ViewRange.Get(), MyGeometry.Size );
-	MouseDownRange[0] = RangeToScreen.LocalXToInput(MyGeometry.AbsoluteToLocal(MouseEvent.GetScreenSpacePosition()).X);
-	MouseDownRange[1] = MouseDownRange[0];
-
+	MouseDownPosition[0] = MouseDownPosition[1] = MouseEvent.GetScreenSpacePosition();
+	MouseDownGeometry = MyGeometry;
 	return FReply::Unhandled();
 }
 
 FReply FSequencerTimeSliderController::OnMouseButtonUp( SWidget& WidgetOwner, const FGeometry& MyGeometry, const FPointerEvent& MouseEvent )
 {
-	bool bHandleLeftMouseButton = MouseEvent.GetEffectingButton() == EKeys::LeftMouseButton && WidgetOwner.HasMouseCapture();
+	bool bHandleLeftMouseButton  = MouseEvent.GetEffectingButton() == EKeys::LeftMouseButton  && WidgetOwner.HasMouseCapture();
 	bool bHandleRightMouseButton = MouseEvent.GetEffectingButton() == EKeys::RightMouseButton && WidgetOwner.HasMouseCapture() && TimeSliderArgs.AllowZoom ;
 
+	FScrubRangeToScreen RangeToScreen = FScrubRangeToScreen(TimeSliderArgs.ViewRange.Get(), MyGeometry.Size);
+	FFrameTime          MouseTime     = ComputeFrameTimeFromMouse(MyGeometry, MouseEvent.GetScreenSpacePosition(), RangeToScreen);
+
 	if ( bHandleRightMouseButton )
 	{
 		if (!bPanning)
@@ -597,16 +636,7 @@ FReply FSequencerTimeSliderController::OnMouseButtonUp( SWidget& WidgetOwner, co
 			// Open a context menu if allowed
 			if (ContextMenuSuppression == 0 && TimeSliderArgs.PlaybackRange.IsSet())
 			{
-				FScrubRangeToScreen RangeToScreen( TimeSliderArgs.ViewRange.Get(), MyGeometry.Size );
-				FVector2D CursorPos = MyGeometry.AbsoluteToLocal( MouseEvent.GetScreenSpacePosition() );
-
-				float MouseValue = RangeToScreen.LocalXToInput( CursorPos.X );
-				if (TimeSliderArgs.Settings->GetIsSnapEnabled())
-				{
-					MouseValue = SequencerHelpers::SnapTimeToInterval(MouseValue, TimeSliderArgs.TimeSnapInterval.Get());
-				}
-
-				TSharedRef<SWidget> MenuContent = OpenSetPlaybackRangeMenu(MouseValue);
+				TSharedRef<SWidget> MenuContent = OpenSetPlaybackRangeMenu(MouseTime.FrameNumber);
 				FSlateApplication::Get().PushMenu(
 					WidgetOwner.AsShared(),
 					MouseEvent.GetEventPath() != nullptr ? *MouseEvent.GetEventPath() : FWidgetPath(),
@@ -649,48 +679,32 @@ FReply FSequencerTimeSliderController::OnMouseButtonUp( SWidget& WidgetOwner, co
 		}
 		else if (MouseDragType == DRAG_SETTING_RANGE)
 		{
-			FScrubRangeToScreen RangeToScreen( TimeSliderArgs.ViewRange.Get(), MyGeometry.Size );
-			FVector2D CursorPos = MyGeometry.AbsoluteToLocal(MouseEvent.GetScreenSpacePosition());
-			float NewValue = RangeToScreen.LocalXToInput(CursorPos.X);
+			// Zooming
+			FFrameTime MouseDownStart = ComputeFrameTimeFromMouse(MyGeometry, MouseDownPosition[0], RangeToScreen);
 
-			if ( TimeSliderArgs.Settings->GetIsSnapEnabled() )
+			const bool bCanZoomIn  = MouseTime > MouseDownStart;
+			const bool bCanZoomOut = ViewRangeStack.Num() > 0;
+			if (bCanZoomIn || bCanZoomOut)
 			{
-				NewValue = SequencerHelpers::SnapTimeToInterval(NewValue, TimeSliderArgs.TimeSnapInterval.Get());
-			}
-
-			float DownValue = MouseDownRange[0];
-				
-			if ( TimeSliderArgs.Settings->GetIsSnapEnabled() )
-			{
-				DownValue = SequencerHelpers::SnapTimeToInterval(DownValue, TimeSliderArgs.TimeSnapInterval.Get());
-			}
+				TRange<double> ViewRange = TimeSliderArgs.ViewRange.Get();
+				if (!bCanZoomIn)
+				{
+					ViewRange = ViewRangeStack.Pop();
+				}
 
-			// Zoom in
-			bool bDoSetRange = false;
-			if (NewValue > DownValue)
-			{
-				// push the current value onto the stack
-				RangeStack.Add(FVector2D(TimeSliderArgs.ViewRange.Get().GetLowerBoundValue(), TimeSliderArgs.ViewRange.Get().GetUpperBoundValue()));
-				bDoSetRange = true;
-			}
-			// Zoom out
-			else if (RangeStack.Num())
-			{
-				// pop the stack
-				FVector2D LastRange = RangeStack.Pop();
-				DownValue = LastRange[0];
-				NewValue = LastRange[1];
-				bDoSetRange = true;
-			}
+				if (bCanZoomIn)
+				{
+					// push the current value onto the stack
+					ViewRangeStack.Add(ViewRange);
 
-			if (bDoSetRange)
-			{
-				TimeSliderArgs.OnViewRangeChanged.ExecuteIfBound(TRange<float>(DownValue, NewValue), EViewRangeInterpolation::Immediate);
-					
+					ViewRange = TRange<double>(MouseDownStart.FrameNumber / GetFrameResolution(), MouseTime.FrameNumber / GetFrameResolution());
+				}
+				
+				TimeSliderArgs.OnViewRangeChanged.ExecuteIfBound(ViewRange, EViewRangeInterpolation::Immediate);
 				if( !TimeSliderArgs.ViewRange.IsBound() )
-				{	
+				{
 					// The output is not bound to a delegate so we'll manage the value ourselves
-					TimeSliderArgs.ViewRange.Set( TRange<float>( DownValue, NewValue ) );
+					TimeSliderArgs.ViewRange.Set(ViewRange);
 				}
 			}
 		}
@@ -698,30 +712,19 @@ FReply FSequencerTimeSliderController::OnMouseButtonUp( SWidget& WidgetOwner, co
 		{
 			TimeSliderArgs.OnEndScrubberMovement.ExecuteIfBound();
 
-			FScrubRangeToScreen RangeToScreen( TimeSliderArgs.ViewRange.Get(), MyGeometry.Size );
-			FVector2D CursorPos = MyGeometry.AbsoluteToLocal(MouseEvent.GetScreenSpacePosition());
-			float NewValue = RangeToScreen.LocalXToInput(CursorPos.X);
+			FFrameTime ScrubTime = MouseTime;
+			FVector2D CursorPos  = MouseEvent.GetScreenSpacePosition();
 
-			if ( TimeSliderArgs.Settings->GetIsSnapEnabled() )
-			{				
-				if (TimeSliderArgs.Settings->GetSnapPlayTimeToInterval())
-				{
-					NewValue = SequencerHelpers::SnapTimeToInterval(NewValue, TimeSliderArgs.TimeSnapInterval.Get());
-				}
-			
-				if (MouseDragType == DRAG_SCRUBBING_TIME && TimeSliderArgs.Settings->ShouldKeepCursorInPlayRangeWhileScrubbing())
-				{
-					TRange<float> PlaybackRange = TimeSliderArgs.PlaybackRange.Get();
-					NewValue = FMath::Clamp(NewValue, PlaybackRange.GetLowerBoundValue(), PlaybackRange.GetUpperBoundValue());
-				}
-
-				if (TimeSliderArgs.Settings->GetSnapPlayTimeToKeys())
-				{
-					NewValue = SnapTimeToNearestKey(RangeToScreen, CursorPos.X, NewValue);
-				}
+			if (MouseDragType == DRAG_SCRUBBING_TIME)
+			{
+				ScrubTime = ComputeScrubTimeFromMouse(MyGeometry, CursorPos, RangeToScreen);
+			}
+			else if (TimeSliderArgs.Settings->GetSnapPlayTimeToKeys())
+			{
+				ScrubTime = SnapTimeToNearestKey(RangeToScreen, CursorPos.X, ScrubTime);
 			}
 
-			CommitScrubPosition( NewValue, /*bIsScrubbing=*/false );
+			CommitScrubPosition( ScrubTime, /*bIsScrubbing=*/false );
 		}
 
 		MouseDragType = DRAG_NONE;
@@ -736,7 +739,7 @@ FReply FSequencerTimeSliderController::OnMouseButtonUp( SWidget& WidgetOwner, co
 
 FReply FSequencerTimeSliderController::OnMouseMove( SWidget& WidgetOwner, const FGeometry& MyGeometry, const FPointerEvent& MouseEvent )
 {
-	bool bHandleLeftMouseButton = MouseEvent.IsMouseButtonDown( EKeys::LeftMouseButton );
+	bool bHandleLeftMouseButton  = MouseEvent.IsMouseButtonDown( EKeys::LeftMouseButton  );
 	bool bHandleRightMouseButton = MouseEvent.IsMouseButtonDown( EKeys::RightMouseButton ) && TimeSliderArgs.AllowZoom;
 
 	if (bHandleRightMouseButton)
@@ -751,17 +754,17 @@ FReply FSequencerTimeSliderController::OnMouseMove( SWidget& WidgetOwner, const
 		}
 		else
 		{
-			TRange<float> LocalViewRange = TimeSliderArgs.ViewRange.Get();
-			float LocalViewRangeMin = LocalViewRange.GetLowerBoundValue();
-			float LocalViewRangeMax = LocalViewRange.GetUpperBoundValue();
+			TRange<double> LocalViewRange = TimeSliderArgs.ViewRange.Get();
+			double LocalViewRangeMin = LocalViewRange.GetLowerBoundValue();
+			double LocalViewRangeMax = LocalViewRange.GetUpperBoundValue();
 
 			FScrubRangeToScreen ScaleInfo( LocalViewRange, MyGeometry.Size );
 			FVector2D ScreenDelta = MouseEvent.GetCursorDelta();
 			FVector2D InputDelta;
 			InputDelta.X = ScreenDelta.X/ScaleInfo.PixelsPerInput;
 
-			float NewViewOutputMin = LocalViewRangeMin - InputDelta.X;
-			float NewViewOutputMax = LocalViewRangeMax - InputDelta.X;
+			double NewViewOutputMin = LocalViewRangeMin - InputDelta.X;
+			double NewViewOutputMax = LocalViewRangeMax - InputDelta.X;
 
 			ClampViewRange(NewViewOutputMin, NewViewOutputMax);
 			SetViewRange(NewViewOutputMin, NewViewOutputMax, EViewRangeInterpolation::Immediate);
@@ -769,41 +772,44 @@ FReply FSequencerTimeSliderController::OnMouseMove( SWidget& WidgetOwner, const
 	}
 	else if (bHandleLeftMouseButton)
 	{
-		TRange<float> LocalViewRange = TimeSliderArgs.ViewRange.Get();
+		TRange<double> LocalViewRange = TimeSliderArgs.ViewRange.Get();
+		FScrubRangeToScreen RangeToScreen(LocalViewRange, MyGeometry.Size);
 		DistanceDragged += FMath::Abs( MouseEvent.GetCursorDelta().X );
 
 		if ( MouseDragType == DRAG_NONE )
 		{
 			if ( DistanceDragged > FSlateApplication::Get().GetDragTriggerDistance() )
 			{
-				FScrubRangeToScreen RangeToScreen(LocalViewRange, MyGeometry.Size);
-				const float ScrubPosition = TimeSliderArgs.ScrubPosition.Get();
+				FFrameTime MouseDownFree = ComputeFrameTimeFromMouse(MyGeometry, MouseDownPosition[0], RangeToScreen, false);
+
+				const FFrameRate FrameResolution    = GetFrameResolution();
+				const bool       bLockedPlayRange   = TimeSliderArgs.IsPlaybackRangeLocked.Get();
+				const float      MouseDownPixel     = RangeToScreen.InputToLocalX(MouseDownFree / FrameResolution);
+				const bool       bHitScrubber       = GetHitTestScrubberPixelRange(TimeSliderArgs.ScrubPosition.Get(), RangeToScreen).HandleRange.Contains(MouseDownPixel);
 
-				TRange<float> SelectionRange = TimeSliderArgs.SelectionRange.Get();
-				TRange<float> PlaybackRange = TimeSliderArgs.PlaybackRange.Get();
-				float LocalMouseDownPos = RangeToScreen.InputToLocalX(MouseDownRange[0]);
-				const bool bIsPlaybackRangeLocked = TimeSliderArgs.IsPlaybackRangeLocked.Get();
+				TRange<double>   SelectionRange   = TimeSliderArgs.SelectionRange.Get() / FrameResolution;
+				TRange<double>   PlaybackRange    = TimeSliderArgs.PlaybackRange.Get()  / FrameResolution;
 
 				// Disable selection range test if it's empty so that the playback range scrubbing gets priority
-				if (!SelectionRange.IsEmpty() && HitTestScrubberEnd(RangeToScreen, SelectionRange, LocalMouseDownPos, ScrubPosition))
+				if (!SelectionRange.IsEmpty() && !bHitScrubber && HitTestRangeEnd(RangeToScreen, SelectionRange, MouseDownPixel))
 				{
 					// selection range end scrubber
 					MouseDragType = DRAG_SELECTION_END;
 					TimeSliderArgs.OnSelectionRangeBeginDrag.ExecuteIfBound();
 				}
-				else if (!SelectionRange.IsEmpty() && HitTestScrubberStart(RangeToScreen, SelectionRange, LocalMouseDownPos, ScrubPosition))
+				else if (!SelectionRange.IsEmpty() && !bHitScrubber && HitTestRangeStart(RangeToScreen, SelectionRange, MouseDownPixel))
 				{
 					// selection range start scrubber
 					MouseDragType = DRAG_SELECTION_START;
 					TimeSliderArgs.OnSelectionRangeBeginDrag.ExecuteIfBound();
 				}
-				else if (!bIsPlaybackRangeLocked && HitTestScrubberEnd(RangeToScreen, PlaybackRange, LocalMouseDownPos, ScrubPosition))
+				else if (!bLockedPlayRange && !bHitScrubber && HitTestRangeEnd(RangeToScreen, PlaybackRange, MouseDownPixel))
 				{
 					// playback range end scrubber
 					MouseDragType = DRAG_PLAYBACK_END;
 					TimeSliderArgs.OnPlaybackRangeBeginDrag.ExecuteIfBound();
 				}
-				else if (!bIsPlaybackRangeLocked && HitTestScrubberStart(RangeToScreen, PlaybackRange, LocalMouseDownPos, ScrubPosition))
+				else if (!bLockedPlayRange && !bHitScrubber && HitTestRangeStart(RangeToScreen, PlaybackRange, MouseDownPixel))
 				{
 					// playback range start scrubber
 					MouseDragType = DRAG_PLAYBACK_START;
@@ -822,77 +828,36 @@ FReply FSequencerTimeSliderController::OnMouseMove( SWidget& WidgetOwner, const
 		}
 		else
 		{
-			FScrubRangeToScreen RangeToScreen( TimeSliderArgs.ViewRange.Get(), MyGeometry.Size );
-			FVector2D CursorPos = MyGeometry.AbsoluteToLocal( MouseEvent.GetScreenSpacePosition() );
-			float NewValue = RangeToScreen.LocalXToInput( CursorPos.X );
-
+			FFrameTime MouseTime = ComputeFrameTimeFromMouse(MyGeometry, MouseEvent.GetScreenSpacePosition(), RangeToScreen);
+			FFrameTime ScrubTime = ComputeScrubTimeFromMouse(MyGeometry, MouseEvent.GetScreenSpacePosition(), RangeToScreen);
 
 			// Set the start range time?
 			if (MouseDragType == DRAG_PLAYBACK_START)
 			{
-				if (TimeSliderArgs.Settings->GetIsSnapEnabled())
-				{
-					NewValue = SequencerHelpers::SnapTimeToInterval(NewValue, TimeSliderArgs.TimeSnapInterval.Get());
-				}
-
-				SetPlaybackRangeStart(NewValue);
+				SetPlaybackRangeStart(MouseTime.FrameNumber);
 			}
 			// Set the end range time?
 			else if(MouseDragType == DRAG_PLAYBACK_END)
 			{
-				if (TimeSliderArgs.Settings->GetIsSnapEnabled())
-				{
-					NewValue = SequencerHelpers::SnapTimeToInterval(NewValue, TimeSliderArgs.TimeSnapInterval.Get());
-				}
-					
-				SetPlaybackRangeEnd(NewValue);
+				SetPlaybackRangeEnd(MouseTime.FrameNumber-1);
 			}
 			else if (MouseDragType == DRAG_SELECTION_START)
 			{
-				if (TimeSliderArgs.Settings->GetIsSnapEnabled())
-				{
-					NewValue = SequencerHelpers::SnapTimeToInterval(NewValue, TimeSliderArgs.TimeSnapInterval.Get());
-				}
-
-				SetSelectionRangeStart(NewValue);
+				SetSelectionRangeStart(MouseTime.FrameNumber);
 			}
 			// Set the end range time?
 			else if(MouseDragType == DRAG_SELECTION_END)
 			{
-				if (TimeSliderArgs.Settings->GetIsSnapEnabled())
-				{
-					NewValue = SequencerHelpers::SnapTimeToInterval(NewValue, TimeSliderArgs.TimeSnapInterval.Get());
-				}
-
-				SetSelectionRangeEnd(NewValue);
+				SetSelectionRangeEnd(MouseTime.FrameNumber);
 			}
 			else if (MouseDragType == DRAG_SCRUBBING_TIME)
 			{
-				if ( TimeSliderArgs.Settings->GetIsSnapEnabled() )
-				{
-					if (TimeSliderArgs.Settings->GetSnapPlayTimeToInterval())
-					{
-						NewValue = SequencerHelpers::SnapTimeToInterval(NewValue, TimeSliderArgs.TimeSnapInterval.Get());
-					}
-
-					if (TimeSliderArgs.Settings->ShouldKeepCursorInPlayRangeWhileScrubbing())
-					{
-						TRange<float> PlaybackRange = TimeSliderArgs.PlaybackRange.Get();
-						NewValue = FMath::Clamp(NewValue, PlaybackRange.GetLowerBoundValue(), PlaybackRange.GetUpperBoundValue());
-					}
-
-					if (TimeSliderArgs.Settings->GetSnapPlayTimeToKeys())
-					{
-						NewValue = SnapTimeToNearestKey(RangeToScreen, CursorPos.X, NewValue);
-					}
-				}
-					
 				// Delegate responsibility for clamping to the current viewrange to the client
-				CommitScrubPosition( NewValue, /*bIsScrubbing=*/true );
+				CommitScrubPosition( ScrubTime, /*bIsScrubbing=*/true );
 			}
 			else if (MouseDragType == DRAG_SETTING_RANGE)
 			{
-				MouseDownRange[1] = NewValue;
+				MouseDownPosition[1] = MouseEvent.GetScreenSpacePosition();
 			}
 		}
 	}
@@ -907,7 +872,7 @@ FReply FSequencerTimeSliderController::OnMouseMove( SWidget& WidgetOwner, const
 }
 
 
-void FSequencerTimeSliderController::CommitScrubPosition( float NewValue, bool bIsScrubbing )
+void FSequencerTimeSliderController::CommitScrubPosition( FFrameTime NewValue, bool bIsScrubbing )
 {
 	// Manage the scrub position ourselves if its not bound to a delegate
 	if ( !TimeSliderArgs.ScrubPosition.IsBound() )
@@ -929,7 +894,7 @@ FReply FSequencerTimeSliderController::OnMouseWheel( SWidget& WidgetOwner, const
 		// If zooming on the current time, adjust mouse fractionX
 		if (TimeSliderArgs.Settings->GetZoomPosition() == ESequencerZoomPosition::SZP_CurrentTime)
 		{
-			const float ScrubPosition = TimeSliderArgs.ScrubPosition.Get();
+			const double ScrubPosition = TimeSliderArgs.ScrubPosition.Get() / GetFrameResolution();
 			if (TimeSliderArgs.ViewRange.Get().Contains(ScrubPosition))
 			{
 				FScrubRangeToScreen RangeToScreen(TimeSliderArgs.ViewRange.Get(), MyGeometry.Size);
@@ -956,27 +921,29 @@ FReply FSequencerTimeSliderController::OnMouseWheel( SWidget& WidgetOwner, const
 FCursorReply FSequencerTimeSliderController::OnCursorQuery( TSharedRef<const SWidget> WidgetOwner, const FGeometry& MyGeometry, const FPointerEvent& CursorEvent ) const
 {
 	FScrubRangeToScreen RangeToScreen(TimeSliderArgs.ViewRange.Get(), MyGeometry.Size);
-	TRange<float> PlaybackRange = TimeSliderArgs.PlaybackRange.Get();
-	TRange<float> SelectionRange = TimeSliderArgs.SelectionRange.Get();
-	const float ScrubPosition = TimeSliderArgs.ScrubPosition.Get();
-	const bool bIsPlaybackRangeLocked = TimeSliderArgs.IsPlaybackRangeLocked.Get();
 
-	// Use L/R resize cursor if we're dragging or hovering a playback range bound
-	float HitTestPosition = MyGeometry.AbsoluteToLocal(CursorEvent.GetScreenSpacePosition()).X;
+	const FFrameRate FrameResolution  = GetFrameResolution();
+	const bool       bLockedPlayRange = TimeSliderArgs.IsPlaybackRangeLocked.Get();
+	const float      HitTestPixel     = MyGeometry.AbsoluteToLocal(CursorEvent.GetScreenSpacePosition()).X;
+	const bool       bHitScrubber     = GetHitTestScrubberPixelRange(TimeSliderArgs.ScrubPosition.Get(), RangeToScreen).HandleRange.Contains(HitTestPixel);
+
+	TRange<double>   SelectionRange   = TimeSliderArgs.SelectionRange.Get() / FrameResolution;
+	TRange<double>   PlaybackRange    = TimeSliderArgs.PlaybackRange.Get()  / FrameResolution;
 
 	if (MouseDragType == DRAG_SCRUBBING_TIME)
 	{
 		return FCursorReply::Unhandled();
 	}
 
+	// Use L/R resize cursor if we're dragging or hovering a playback range bound
 	if ((MouseDragType == DRAG_PLAYBACK_END) ||
 		(MouseDragType == DRAG_PLAYBACK_START) ||
 		(MouseDragType == DRAG_SELECTION_START) ||
 		(MouseDragType == DRAG_SELECTION_END) ||
-		(!bIsPlaybackRangeLocked && HitTestScrubberStart(RangeToScreen, PlaybackRange, HitTestPosition, ScrubPosition)) ||
-		(!bIsPlaybackRangeLocked && HitTestScrubberEnd(RangeToScreen, PlaybackRange, HitTestPosition, ScrubPosition)) ||
-		(!SelectionRange.IsEmpty() && HitTestScrubberStart(RangeToScreen, SelectionRange, HitTestPosition, ScrubPosition)) ||
-		(!SelectionRange.IsEmpty() && HitTestScrubberEnd(RangeToScreen, SelectionRange, HitTestPosition, ScrubPosition)))
+		(!bLockedPlayRange         && !bHitScrubber && HitTestRangeStart(RangeToScreen, PlaybackRange,  HitTestPixel)) ||
+		(!bLockedPlayRange         && !bHitScrubber && HitTestRangeEnd(  RangeToScreen, PlaybackRange,  HitTestPixel)) ||
+		(!SelectionRange.IsEmpty() && !bHitScrubber && HitTestRangeStart(RangeToScreen, SelectionRange, HitTestPixel)) ||
+		(!SelectionRange.IsEmpty() && !bHitScrubber && HitTestRangeEnd(  RangeToScreen, SelectionRange, HitTestPixel)))
 	{
 		return FCursorReply::Cursor(EMouseCursor::ResizeLeftRight);
 	}
@@ -988,13 +955,7 @@ int32 FSequencerTimeSliderController::OnPaintSectionView( const FGeometry& Allot
 {
 	const ESlateDrawEffect DrawEffects = bEnabled ? ESlateDrawEffect::None : ESlateDrawEffect::DisabledEffect;
 
-	TRange<float> LocalViewRange = TimeSliderArgs.ViewRange.Get();
-	float LocalScrubPosition = TimeSliderArgs.ScrubPosition.Get();
-
-	float ViewRange = LocalViewRange.Size<float>();
-	float PixelsPerInput = ViewRange > 0 ? AllottedGeometry.Size.X / ViewRange : 0;
-	float LinePos =  (LocalScrubPosition - LocalViewRange.GetLowerBoundValue()) * PixelsPerInput;
-
+	TRange<double> LocalViewRange = TimeSliderArgs.ViewRange.Get();
 	FScrubRangeToScreen RangeToScreen( LocalViewRange, AllottedGeometry.Size );
 
 	if (Args.PlaybackRangeArgs.IsSet())
@@ -1009,7 +970,7 @@ int32 FSequencerTimeSliderController::OnPaintSectionView( const FGeometry& Allot
 	if( Args.bDisplayTickLines )
 	{
 		static FLinearColor TickColor(0.f, 0.f, 0.f, 0.3f);
-		
+
 		// Draw major tick lines in the section area
 		FDrawTickArgs DrawTickArgs;
 		{
@@ -1026,49 +987,62 @@ int32 FSequencerTimeSliderController::OnPaintSectionView( const FGeometry& Allot
 			DrawTickArgs.MajorTickHeight = AllottedGeometry.Size.Y;
 		}
 
-		DrawTicks( OutDrawElements, RangeToScreen, DrawTickArgs );
+		DrawTicks( OutDrawElements, LocalViewRange, RangeToScreen, DrawTickArgs );
 	}
 
 	if( Args.bDisplayScrubPosition )
 	{
-		// Draw a line for the scrub position
-		TArray<FVector2D> LinePoints;
+		FScrubPixelRange ScrubberPixelRange = GetScrubberPixelRange(TimeSliderArgs.ScrubPosition.Get(), RangeToScreen);
+
+		if (ScrubberPixelRange.Range.Size<float>() < 2.f)
 		{
-			LinePoints.AddUninitialized(2);
-			LinePoints[0] = FVector2D( 0.0f, 0.0f );
-			LinePoints[1] = FVector2D( 0.0f, FMath::RoundToFloat( AllottedGeometry.Size.Y ) );
-		}
+			// Draw a line for the scrub position
+			const float LinePos = ScrubberPixelRange.Range.GetLowerBoundValue() + ScrubberPixelRange.Range.Size<float>() * .5f;
 
-		FSlateDrawElement::MakeLines(
-			OutDrawElements,
-			LayerId+1,
-			AllottedGeometry.ToPaintGeometry( FVector2D(LinePos, 0.0f ), FVector2D(1.0f,1.0f) ),
-			LinePoints,
-			DrawEffects,
-			FLinearColor::White,
-			false
-		);
+			TArray<FVector2D> LinePoints;
+			{
+				LinePoints.AddUninitialized(2);
+				LinePoints[0] = FVector2D( 0.0f, 0.0f );
+				LinePoints[1] = FVector2D( 0.0f, FMath::FloorToFloat( AllottedGeometry.Size.Y ) );
+			}
+
+			FSlateDrawElement::MakeLines(
+				OutDrawElements,
+				LayerId+1,
+				AllottedGeometry.ToPaintGeometry( FVector2D(LinePos, 0.0f ), FVector2D(1.0f,1.0f) ),
+				LinePoints,
+				DrawEffects,
+				FLinearColor(1.f, 1.f, 1.f, .5f),
+				false
+			);
+		}
+		else
+		{
+			// Draw a box for the scrub position
+			FSlateDrawElement::MakeBox(
+				OutDrawElements,
+				LayerId+1,
+				AllottedGeometry.ToPaintGeometry( FVector2D(ScrubberPixelRange.Range.GetLowerBoundValue(), 0.0f ), FVector2D(ScrubberPixelRange.Range.Size<float>(), AllottedGeometry.Size.Y) ),
+				ScrubFillBrush,
+				DrawEffects, 
+				FLinearColor::White
+			);
+		}
 	}
 
 	return LayerId;
 }
 
-TSharedRef<SWidget> FSequencerTimeSliderController::OpenSetPlaybackRangeMenu(float MouseTime)
+TSharedRef<SWidget> FSequencerTimeSliderController::OpenSetPlaybackRangeMenu(FFrameNumber FrameNumber)
 {
 	const bool bShouldCloseWindowAfterMenuSelection = true;
 	FMenuBuilder MenuBuilder(bShouldCloseWindowAfterMenuSelection, nullptr);
 
 	FText CurrentTimeText;
-	if (SequencerSnapValues::IsTimeSnapIntervalFrameRate(TimeSliderArgs.TimeSnapInterval.Get()) && TimeSliderArgs.Settings->GetShowFrameNumbers())
-	{
-		CurrentTimeText = FText::Format(LOCTEXT("FrameTextFormat", "at frame {0}"), FText::AsNumber(TimeToFrame(MouseTime)));
-	}
-	else
-	{
-		CurrentTimeText = FText::Format(LOCTEXT("TimeTextFormat", "at {0}s"), FText::AsNumber(MouseTime));
-	}
+	CurrentTimeText = FText::FromString(TimeSliderArgs.NumericTypeInterface->ToString(FrameNumber.Value));
+	
 
-	TRange<float> PlaybackRange = TimeSliderArgs.PlaybackRange.Get();
+	TRange<FFrameNumber> PlaybackRange = TimeSliderArgs.PlaybackRange.Get();
 
 	MenuBuilder.BeginSection("SequencerPlaybackRangeMenu", FText::Format(LOCTEXT("PlaybackRangeTextFormat", "Playback Range ({0}):"), CurrentTimeText));
 	{
@@ -1077,8 +1051,8 @@ TSharedRef<SWidget> FSequencerTimeSliderController::OpenSetPlaybackRangeMenu(flo
 			FText(),
 			FSlateIcon(),
 			FUIAction(
-				FExecuteAction::CreateLambda([=]{ SetPlaybackRangeStart(MouseTime); }),
-				FCanExecuteAction::CreateLambda([=]{ return !TimeSliderArgs.IsPlaybackRangeLocked.Get() && MouseTime <= PlaybackRange.GetUpperBoundValue(); })
+				FExecuteAction::CreateLambda([=]{ SetPlaybackRangeStart(FrameNumber); }),
+				FCanExecuteAction::CreateLambda([=]{ return !TimeSliderArgs.IsPlaybackRangeLocked.Get() && FrameNumber < MovieScene::DiscreteExclusiveUpper(PlaybackRange); })
 			)
 		);
 
@@ -1087,8 +1061,8 @@ TSharedRef<SWidget> FSequencerTimeSliderController::OpenSetPlaybackRangeMenu(flo
 			FText(),
 			FSlateIcon(),
 			FUIAction(
-				FExecuteAction::CreateLambda([=]{ SetPlaybackRangeEnd(MouseTime); }),
-				FCanExecuteAction::CreateLambda([=]{ return !TimeSliderArgs.IsPlaybackRangeLocked.Get() && MouseTime >= PlaybackRange.GetLowerBoundValue(); })
+				FExecuteAction::CreateLambda([=]{ SetPlaybackRangeEnd(FrameNumber); }),
+				FCanExecuteAction::CreateLambda([=]{ return !TimeSliderArgs.IsPlaybackRangeLocked.Get() && FrameNumber >= MovieScene::DiscreteInclusiveLower(PlaybackRange); })
 			)
 		);
 
@@ -1107,7 +1081,7 @@ TSharedRef<SWidget> FSequencerTimeSliderController::OpenSetPlaybackRangeMenu(flo
 	}
 	MenuBuilder.EndSection(); // SequencerPlaybackRangeMenu
 
-	TRange<float> SelectionRange = TimeSliderArgs.SelectionRange.Get();
+	TRange<FFrameNumber> SelectionRange = TimeSliderArgs.SelectionRange.Get();
 	MenuBuilder.BeginSection("SequencerSelectionRangeMenu", FText::Format(LOCTEXT("SelectionRangeTextFormat", "Selection Range ({0}):"), CurrentTimeText));
 	{
 		MenuBuilder.AddMenuEntry(
@@ -1115,8 +1089,8 @@ TSharedRef<SWidget> FSequencerTimeSliderController::OpenSetPlaybackRangeMenu(flo
 			FText(),
 			FSlateIcon(),
 			FUIAction(
-				FExecuteAction::CreateLambda([=]{ SetSelectionRangeStart(MouseTime); }),
-				FCanExecuteAction::CreateLambda([=]{ return SelectionRange.IsEmpty() || MouseTime <= SelectionRange.GetUpperBoundValue(); })
+				FExecuteAction::CreateLambda([=]{ SetSelectionRangeStart(FrameNumber); }),
+				FCanExecuteAction::CreateLambda([=]{ return SelectionRange.IsEmpty() || FrameNumber < MovieScene::DiscreteExclusiveUpper(SelectionRange); })
 			)
 		);
 
@@ -1125,8 +1099,8 @@ TSharedRef<SWidget> FSequencerTimeSliderController::OpenSetPlaybackRangeMenu(flo
 			FText(),
 			FSlateIcon(),
 			FUIAction(
-				FExecuteAction::CreateLambda([=]{ SetSelectionRangeEnd(MouseTime); }),
-				FCanExecuteAction::CreateLambda([=]{ return SelectionRange.IsEmpty() || MouseTime >= SelectionRange.GetLowerBoundValue(); })
+				FExecuteAction::CreateLambda([=]{ SetSelectionRangeEnd(FrameNumber); }),
+				FCanExecuteAction::CreateLambda([=]{ return SelectionRange.IsEmpty() || FrameNumber >= MovieScene::DiscreteInclusiveLower(SelectionRange); })
 			)
 		);
 
@@ -1135,7 +1109,7 @@ TSharedRef<SWidget> FSequencerTimeSliderController::OpenSetPlaybackRangeMenu(flo
 			FText(),
 			FSlateIcon(),
 			FUIAction(
-				FExecuteAction::CreateLambda([=]{ TimeSliderArgs.OnSelectionRangeChanged.ExecuteIfBound(TRange<float>::Empty()); }),
+				FExecuteAction::CreateLambda([=]{ TimeSliderArgs.OnSelectionRangeChanged.ExecuteIfBound(TRange<FFrameNumber>::Empty()); }),
 				FCanExecuteAction::CreateLambda([=]{ return !SelectionRange.IsEmpty(); })
 			)
 		);
@@ -1145,29 +1119,17 @@ TSharedRef<SWidget> FSequencerTimeSliderController::OpenSetPlaybackRangeMenu(flo
 	return MenuBuilder.MakeWidget();
 }
 
-int32 FSequencerTimeSliderController::TimeToFrame(float Time) const
-{
-	float FrameRate = 1.0f/TimeSliderArgs.TimeSnapInterval.Get();
-	return SequencerHelpers::TimeToFrame(Time, FrameRate);
-}
-
-float FSequencerTimeSliderController::FrameToTime(int32 Frame) const
-{
-	float FrameRate = 1.0f/TimeSliderArgs.TimeSnapInterval.Get();
-	return SequencerHelpers::FrameToTime(Frame, FrameRate);
-}
-
-void FSequencerTimeSliderController::ClampViewRange(float& NewRangeMin, float& NewRangeMax)
+void FSequencerTimeSliderController::ClampViewRange(double& NewRangeMin, double& NewRangeMax)
 {
 	bool bNeedsClampSet = false;
-	float NewClampRangeMin = TimeSliderArgs.ClampRange.Get().GetLowerBoundValue();
+	double NewClampRangeMin = TimeSliderArgs.ClampRange.Get().GetLowerBoundValue();
 	if ( NewRangeMin < TimeSliderArgs.ClampRange.Get().GetLowerBoundValue() )
 	{
 		NewClampRangeMin = NewRangeMin;
 		bNeedsClampSet = true;
 	}
 
-	float NewClampRangeMax = TimeSliderArgs.ClampRange.Get().GetUpperBoundValue();
+	double NewClampRangeMax = TimeSliderArgs.ClampRange.Get().GetUpperBoundValue();
 	if ( NewRangeMax > TimeSliderArgs.ClampRange.Get().GetUpperBoundValue() )
 	{
 		NewClampRangeMax = NewRangeMax;
@@ -1180,22 +1142,39 @@ void FSequencerTimeSliderController::ClampViewRange(float& NewRangeMin, float& N
 	}
 }
 
-void FSequencerTimeSliderController::SetViewRange( float NewRangeMin, float NewRangeMax, EViewRangeInterpolation Interpolation )
+void FSequencerTimeSliderController::SetViewRange( double NewRangeMin, double NewRangeMax, EViewRangeInterpolation Interpolation )
 {
-	const TRange<float> NewRange(NewRangeMin, NewRangeMax);
+	// Clamp to a minimum size to avoid zero-sized or negative visible ranges
+	double MinVisibleTimeRange = FFrameNumber(1) / GetFrameResolution();
+	TRange<double> ExistingViewRange  = TimeSliderArgs.ViewRange.Get();
+	TRange<double> ExistingClampRange = TimeSliderArgs.ClampRange.Get();
+
+	if (NewRangeMax == ExistingViewRange.GetUpperBoundValue())
+	{
+		if (NewRangeMin > NewRangeMax - MinVisibleTimeRange)
+		{
+			NewRangeMin = NewRangeMax - MinVisibleTimeRange;
+		}
+	}
+	else if (NewRangeMax < NewRangeMin + MinVisibleTimeRange)
+	{
+		NewRangeMax = NewRangeMin + MinVisibleTimeRange;
+	}
 
+	// Clamp to the clamp range
+	const TRange<double> NewRange = TRange<double>::Intersection(TRange<double>(NewRangeMin, NewRangeMax), ExistingClampRange);
 	TimeSliderArgs.OnViewRangeChanged.ExecuteIfBound( NewRange, Interpolation );
 
 	if( !TimeSliderArgs.ViewRange.IsBound() )
-	{	
+	{
 		// The  output is not bound to a delegate so we'll manage the value ourselves (no animation)
 		TimeSliderArgs.ViewRange.Set( NewRange );
 	}
 }
 
-void FSequencerTimeSliderController::SetClampRange( float NewRangeMin, float NewRangeMax )
+void FSequencerTimeSliderController::SetClampRange( double NewRangeMin, double NewRangeMax )
 {
-	const TRange<float> NewRange(NewRangeMin, NewRangeMax);
+	const TRange<double> NewRange(NewRangeMin, NewRangeMax);
 
 	TimeSliderArgs.OnClampRangeChanged.ExecuteIfBound(NewRange);
 
@@ -1206,14 +1185,16 @@ void FSequencerTimeSliderController::SetClampRange( float NewRangeMin, float New
 	}
 }
 
-void FSequencerTimeSliderController::SetPlayRange( float NewRangeMin, float NewRangeMax )
+void FSequencerTimeSliderController::SetPlayRange( FFrameNumber RangeStart, int32 RangeDuration )
 {
-	const TRange<float> NewRange(NewRangeMin, NewRangeMax);
+	check(RangeDuration >= 0);
+
+	const TRange<FFrameNumber> NewRange(RangeStart, RangeStart + RangeDuration);
 
 	TimeSliderArgs.OnPlaybackRangeChanged.ExecuteIfBound(NewRange);
 
 	if( !TimeSliderArgs.PlaybackRange.IsBound() )
-	{	
+	{
 		// The  output is not bound to a delegate so we'll manage the value ourselves (no animation)
 		TimeSliderArgs.PlaybackRange.Set(NewRange);
 	}
@@ -1221,14 +1202,14 @@ void FSequencerTimeSliderController::SetPlayRange( float NewRangeMin, float NewR
 
 bool FSequencerTimeSliderController::ZoomByDelta( float InDelta, float MousePositionFraction )
 {
-	TRange<float> LocalViewRange = TimeSliderArgs.ViewRange.Get().GetAnimationTarget();
-	float LocalViewRangeMax = LocalViewRange.GetUpperBoundValue();
-	float LocalViewRangeMin = LocalViewRange.GetLowerBoundValue();
-	const float OutputViewSize = LocalViewRangeMax - LocalViewRangeMin;
-	const float OutputChange = OutputViewSize * InDelta;
+	TRange<double> LocalViewRange = TimeSliderArgs.ViewRange.Get().GetAnimationTarget();
+	double LocalViewRangeMax = LocalViewRange.GetUpperBoundValue();
+	double LocalViewRangeMin = LocalViewRange.GetLowerBoundValue();
+	const double OutputViewSize = LocalViewRangeMax - LocalViewRangeMin;
+	const double OutputChange = OutputViewSize * InDelta;
 
-	float NewViewOutputMin = LocalViewRangeMin - (OutputChange * MousePositionFraction);
-	float NewViewOutputMax = LocalViewRangeMax + (OutputChange * (1.f - MousePositionFraction));
+	double NewViewOutputMin = LocalViewRangeMin - (OutputChange * MousePositionFraction);
+	double NewViewOutputMax = LocalViewRangeMax + (OutputChange * (1.f - MousePositionFraction));
 
 	if( NewViewOutputMin < NewViewOutputMax )
 	{
@@ -1242,121 +1223,104 @@ bool FSequencerTimeSliderController::ZoomByDelta( float InDelta, float MousePosi
 
 void FSequencerTimeSliderController::PanByDelta( float InDelta )
 {
-	TRange<float> LocalViewRange = TimeSliderArgs.ViewRange.Get().GetAnimationTarget();
+	TRange<double> LocalViewRange = TimeSliderArgs.ViewRange.Get().GetAnimationTarget();
 
-	float CurrentMin = LocalViewRange.GetLowerBoundValue();
-	float CurrentMax = LocalViewRange.GetUpperBoundValue();
+	double CurrentMin = LocalViewRange.GetLowerBoundValue();
+	double CurrentMax = LocalViewRange.GetUpperBoundValue();
 
 	// Adjust the delta to be a percentage of the current range
 	InDelta *= ScrubConstants::ScrollPanFraction * (CurrentMax - CurrentMin);
 
-	float NewViewOutputMin = CurrentMin + InDelta;
-	float NewViewOutputMax = CurrentMax + InDelta;
+	double NewViewOutputMin = CurrentMin + InDelta;
+	double NewViewOutputMax = CurrentMax + InDelta;
 
 	ClampViewRange(NewViewOutputMin, NewViewOutputMax);
 	SetViewRange(NewViewOutputMin, NewViewOutputMax, EViewRangeInterpolation::Animated);
 }
 
 
-bool FSequencerTimeSliderController::HitTestScrubberStart(const FScrubRangeToScreen& RangeToScreen, const TRange<float>& PlaybackRange, float LocalHitPositionX, float ScrubPosition) const
+bool FSequencerTimeSliderController::HitTestRangeStart(const FScrubRangeToScreen& RangeToScreen, const TRange<double>& Range, float HitPixel) const
 {
 	static float BrushSizeInStateUnits = 6.f, DragToleranceSlateUnits = 2.f, MouseTolerance = 2.f;
-	float LocalPlaybackStartPos = RangeToScreen.InputToLocalX(PlaybackRange.GetLowerBoundValue());
-	float LocalScrubPos = RangeToScreen.InputToLocalX(ScrubPosition);
-
-	// We favor hit testing the scrub bar over hit testing the playback range bounds
-	if ((LocalScrubPos - ScrubHandleSize/2.f - MouseTolerance - DragToleranceSlateUnits) < LocalHitPositionX &&
-		(LocalScrubPos + ScrubHandleSize/2.f + MouseTolerance + DragToleranceSlateUnits) > LocalHitPositionX)
-	{
-		return false;
-	}
+	const float  RangeStartPixel = RangeToScreen.InputToLocalX(Range.GetLowerBoundValue());
 
 	// Hit test against the brush region to the right of the playback start position, +/- DragToleranceSlateUnits
-	return LocalHitPositionX >= LocalPlaybackStartPos - MouseTolerance - DragToleranceSlateUnits &&
-		LocalHitPositionX <= LocalPlaybackStartPos + MouseTolerance + BrushSizeInStateUnits + DragToleranceSlateUnits;
+	return HitPixel >= RangeStartPixel - MouseTolerance - DragToleranceSlateUnits &&
+		HitPixel <= RangeStartPixel + MouseTolerance + BrushSizeInStateUnits + DragToleranceSlateUnits;
 }
 
-bool FSequencerTimeSliderController::HitTestScrubberEnd(const FScrubRangeToScreen& RangeToScreen, const TRange<float>& PlaybackRange, float LocalHitPositionX, float ScrubPosition) const
+bool FSequencerTimeSliderController::HitTestRangeEnd(const FScrubRangeToScreen& RangeToScreen, const TRange<double>& Range, float HitPixel) const
 {
 	static float BrushSizeInStateUnits = 6.f, DragToleranceSlateUnits = 2.f, MouseTolerance = 2.f;
-	float LocalPlaybackEndPos = RangeToScreen.InputToLocalX(PlaybackRange.GetUpperBoundValue());
-	float LocalScrubPos = RangeToScreen.InputToLocalX(ScrubPosition);
-	
-	// We favor hit testing the scrub bar over hit testing the playback range bounds
-	if ((LocalScrubPos - ScrubHandleSize/2.f - MouseTolerance - DragToleranceSlateUnits) < LocalHitPositionX &&
-		(LocalScrubPos + ScrubHandleSize/2.f + MouseTolerance + DragToleranceSlateUnits) > LocalHitPositionX)
-	{
-		return false;
-	}
+	const float  RangeEndPixel = RangeToScreen.InputToLocalX(Range.GetUpperBoundValue());
 
 	// Hit test against the brush region to the left of the playback end position, +/- DragToleranceSlateUnits
-	return LocalHitPositionX >= LocalPlaybackEndPos - MouseTolerance - BrushSizeInStateUnits - DragToleranceSlateUnits &&
-		LocalHitPositionX <= LocalPlaybackEndPos + MouseTolerance + DragToleranceSlateUnits;
+	return HitPixel >= RangeEndPixel - MouseTolerance - BrushSizeInStateUnits - DragToleranceSlateUnits &&
+		HitPixel <= RangeEndPixel + MouseTolerance + DragToleranceSlateUnits;
 }
 
-float FSequencerTimeSliderController::SnapTimeToNearestKey(const FScrubRangeToScreen& RangeToScreen, float CursorPos, float InTime)
+FFrameTime FSequencerTimeSliderController::SnapTimeToNearestKey(const FScrubRangeToScreen& RangeToScreen, float CursorPos, FFrameTime InTime) const
 {
-	float OutTime = InTime;
 	if (TimeSliderArgs.OnGetNearestKey.IsBound())
 	{
-		float NearestKey = TimeSliderArgs.OnGetNearestKey.Execute(InTime);
-		float LocalKeyPos = RangeToScreen.InputToLocalX( NearestKey );
+		FFrameNumber NearestKey = TimeSliderArgs.OnGetNearestKey.Execute(InTime);
+		float LocalKeyPos = RangeToScreen.InputToLocalX( NearestKey / GetFrameResolution() );
 		static float MouseTolerance = 20.f;
 
-		if (FMath::Abs(LocalKeyPos - CursorPos) < MouseTolerance)
+		if (FMath::IsNearlyEqual(LocalKeyPos, CursorPos, MouseTolerance))
 		{
-			OutTime = NearestKey;
+			return NearestKey;
 		}
 	}
 
-	return OutTime;
+	return InTime;
 }
-					
-void FSequencerTimeSliderController::SetPlaybackRangeStart(float NewStart)
+
+void FSequencerTimeSliderController::SetPlaybackRangeStart(FFrameNumber NewStart)
 {
-	TRange<float> PlaybackRange = TimeSliderArgs.PlaybackRange.Get();
+	TRange<FFrameNumber> PlaybackRange = TimeSliderArgs.PlaybackRange.Get();
 
-	if (NewStart <= PlaybackRange.GetUpperBoundValue())
+	if (NewStart <= MovieScene::DiscreteExclusiveUpper(PlaybackRange))
 	{
-		TimeSliderArgs.OnPlaybackRangeChanged.ExecuteIfBound(TRange<float>(NewStart, PlaybackRange.GetUpperBoundValue()));
+		TimeSliderArgs.OnPlaybackRangeChanged.ExecuteIfBound(TRange<FFrameNumber>(NewStart, PlaybackRange.GetUpperBound()));
 	}
 }
 
-void FSequencerTimeSliderController::SetPlaybackRangeEnd(float NewEnd)
+void FSequencerTimeSliderController::SetPlaybackRangeEnd(FFrameNumber NewEnd)
 {
-	TRange<float> PlaybackRange = TimeSliderArgs.PlaybackRange.Get();
+	TRange<FFrameNumber> PlaybackRange = TimeSliderArgs.PlaybackRange.Get();
 
-	if (NewEnd >= PlaybackRange.GetLowerBoundValue())
+	if (NewEnd >= MovieScene::DiscreteInclusiveLower(PlaybackRange))
 	{
-		TimeSliderArgs.OnPlaybackRangeChanged.ExecuteIfBound(TRange<float>(PlaybackRange.GetLowerBoundValue(), NewEnd));
+		TimeSliderArgs.OnPlaybackRangeChanged.ExecuteIfBound(TRange<FFrameNumber>(PlaybackRange.GetLowerBound(), NewEnd));
 	}
 }
 
-void FSequencerTimeSliderController::SetSelectionRangeStart(float NewStart)
+void FSequencerTimeSliderController::SetSelectionRangeStart(FFrameNumber NewStart)
 {
-	TRange<float> SelectionRange = TimeSliderArgs.SelectionRange.Get();
+	TRange<FFrameNumber> SelectionRange = TimeSliderArgs.SelectionRange.Get();
 
 	if (SelectionRange.IsEmpty())
 	{
-		TimeSliderArgs.OnSelectionRangeChanged.ExecuteIfBound(TRange<float>(NewStart, NewStart + 1));
+		TimeSliderArgs.OnSelectionRangeChanged.ExecuteIfBound(TRange<FFrameNumber>(NewStart, NewStart + 1));
 	}
-	else if (NewStart <= SelectionRange.GetUpperBoundValue())
+	else if (NewStart <= MovieScene::DiscreteExclusiveUpper(SelectionRange))
 	{
-		TimeSliderArgs.OnSelectionRangeChanged.ExecuteIfBound(TRange<float>(NewStart, SelectionRange.GetUpperBoundValue()));
+		TimeSliderArgs.OnSelectionRangeChanged.ExecuteIfBound(TRange<FFrameNumber>(NewStart, SelectionRange.GetUpperBound()));
 	}
 }
 
-void FSequencerTimeSliderController::SetSelectionRangeEnd(float NewEnd)
+void FSequencerTimeSliderController::SetSelectionRangeEnd(FFrameNumber NewEnd)
 {
-	TRange<float> SelectionRange = TimeSliderArgs.SelectionRange.Get();
+	TRange<FFrameNumber> SelectionRange = TimeSliderArgs.SelectionRange.Get();
 
 	if (SelectionRange.IsEmpty())
 	{
-		TimeSliderArgs.OnSelectionRangeChanged.ExecuteIfBound(TRange<float>(NewEnd - 1, NewEnd));
+		TimeSliderArgs.OnSelectionRangeChanged.ExecuteIfBound(TRange<FFrameNumber>(NewEnd - 1, NewEnd));
 	}
-	else if (NewEnd >= SelectionRange.GetLowerBoundValue())
+	else if (NewEnd >= MovieScene::DiscreteInclusiveLower(SelectionRange))
 	{
-		TimeSliderArgs.OnSelectionRangeChanged.ExecuteIfBound(TRange<float>(SelectionRange.GetLowerBoundValue(), NewEnd));
+		TimeSliderArgs.OnSelectionRangeChanged.ExecuteIfBound(TRange<FFrameNumber>(SelectionRange.GetLowerBound(), NewEnd));
 	}
 }
 
diff --git a/Engine/Source/Editor/Sequencer/Private/SequencerTimeSliderController.h b/Engine/Source/Editor/Sequencer/Private/SequencerTimeSliderController.h
--- a/Engine/Source/Editor/Sequencer/Private/SequencerTimeSliderController.h
+++ b/Engine/Source/Editor/Sequencer/Private/SequencerTimeSliderController.h
@@ -12,6 +12,7 @@ class FSlateWindowElementList;
 struct FContextMenuSuppressor;
 struct FScrubRangeToScreen;
 struct FSlateBrush;
+class FSequencer;
 
 struct FPaintPlaybackRangeArgs
 {
@@ -53,7 +54,7 @@ struct FPaintSectionAreaViewArgs
 class FSequencerTimeSliderController : public ITimeSliderController
 {
 public:
-	FSequencerTimeSliderController( const FTimeSliderArgs& InArgs );
+	FSequencerTimeSliderController( const FTimeSliderArgs& InArgs, TWeakPtr<FSequencer> InWeakSequencer );
 
 	/**
 	* Determines the optimal spacing between tick marks in the slider for a given pixel density
@@ -75,20 +76,20 @@ public:
 	virtual FCursorReply OnCursorQuery( TSharedRef<const SWidget> WidgetOwner, const FGeometry& MyGeometry, const FPointerEvent& CursorEvent ) const override;
 	/** End ITimeSliderController Interface */
 
+	/** Get the current play rate for this controller */
+	virtual FFrameRate GetPlayRate() const override { return TimeSliderArgs.PlayRate.Get(); }
+
+	/** Get the current frame resolution for this controller */
+	virtual FFrameRate GetFrameResolution() const override { return TimeSliderArgs.FrameResolution.Get(); }
+
 	/** Get the current view range for this controller */
 	virtual FAnimatedRange GetViewRange() const override { return TimeSliderArgs.ViewRange.Get(); }
 
-	/** Get the current clamp range for this controller */
+	/** Get the current clamp range for this controller in seconds. */
 	virtual FAnimatedRange GetClampRange() const override { return TimeSliderArgs.ClampRange.Get(); }
 
 	/** Get the current play range for this controller */
-	virtual TRange<float> GetPlayRange() const override { return TimeSliderArgs.PlaybackRange.Get(TRange<float>()); }
-
-	/** Convert time to frame */
-	virtual int32 TimeToFrame(float Time) const override;
-
-	/** Convert time to frame */
-	virtual float FrameToTime(int32 Frame) const override;
+	virtual TRange<FFrameNumber> GetPlayRange() const override { return TimeSliderArgs.PlaybackRange.Get(TRange<FFrameNumber>()); }
 
 	/**
 	 * Clamp the given range to the clamp range 
@@ -96,7 +97,7 @@ public:
 	 * @param NewRangeMin		The new lower bound of the range
 	 * @param NewRangeMax		The new upper bound of the range
 	 */	
-	void ClampViewRange(float& NewRangeMin, float& NewRangeMax);
+	void ClampViewRange(double& NewRangeMin, double& NewRangeMax);
 
 	/**
 	 * Set a new range based on a min, max and an interpolation mode
@@ -105,7 +106,7 @@ public:
 	 * @param NewRangeMax		The new upper bound of the range
 	 * @param Interpolation		How to set the new range (either immediately, or animated)
 	 */
-	virtual void SetViewRange( float NewRangeMin, float NewRangeMax, EViewRangeInterpolation Interpolation ) override;
+	virtual void SetViewRange( double NewRangeMin, double NewRangeMax, EViewRangeInterpolation Interpolation ) override;
 
 	/**
 	 * Set a new clamp range based on a min, max
@@ -113,15 +114,15 @@ public:
 	 * @param NewRangeMin		The new lower bound of the clamp range
 	 * @param NewRangeMax		The new upper bound of the clamp range
 	 */
-	virtual void SetClampRange( float NewRangeMin, float NewRangeMax ) override;
+	virtual void SetClampRange( double NewRangeMin, double NewRangeMax ) override;
 
 	/**
 	 * Set a new playback range based on a min, max
 	 * 
-	 * @param NewRangeMin		The new lower bound of the playback range
-	 * @param NewRangeMax		The new upper bound of the playback range
+	 * @param RangeStart		The new lower bound of the playback range
+	 * @param RangeDuration		The total number of frames that we play for
 	 */
-	virtual void SetPlayRange( float NewRangeMin, float NewRangeMax ) override;
+	virtual void SetPlayRange( FFrameNumber RangeStart, int32 RangeDuration ) override;
 
 	/**
 	 * Zoom the range by a given delta.
@@ -154,16 +155,17 @@ private:
 	 * @param NewValue				Value resulting from the user's interaction
 	 * @param bIsScrubbing			True if done via scrubbing, false if just releasing scrubbing
 	 */
-	void CommitScrubPosition( float NewValue, bool bIsScrubbing );
+	void CommitScrubPosition( FFrameTime NewValue, bool bIsScrubbing );
 
 	/**
 	 * Draw time tick marks
 	 *
 	 * @param OutDrawElements	List to add draw elements to
+	 * @param ViewRange			The currently visible time range in seconds
 	 * @param RangeToScreen		Time range to screen space converter
 	 * @param InArgs			Parameters for drawing the tick lines
 	 */
-	void DrawTicks( FSlateWindowElementList& OutDrawElements, const FScrubRangeToScreen& RangeToScreen, FDrawTickArgs& InArgs ) const;
+	void DrawTicks( FSlateWindowElementList& OutDrawElements, const TRange<double>& ViewRange, const FScrubRangeToScreen& RangeToScreen, FDrawTickArgs& InArgs ) const;
 
 	/**
 	 * Draw the selection range.
@@ -189,36 +191,55 @@ private:
 private:
 
 	/**
-	 * Hit test the lower bound of the playback range
+	 * Hit test the lower bound of a range
 	 */
-	bool HitTestScrubberStart(const FScrubRangeToScreen& RangeToScreen, const TRange<float>& PlaybackRange, float LocalHitPositionX, float ScrubPosition) const;
+	bool HitTestRangeStart(const FScrubRangeToScreen& RangeToScreen, const TRange<double>& Range, float HitPixel) const;
 
 	/**
-	 * Hit test the upper bound of the playback range
+	 * Hit test the upper bound of a range
 	 */
-	bool HitTestScrubberEnd(const FScrubRangeToScreen& RangeToScreen, const TRange<float>& PlaybackRange, float LocalHitPositionX, float ScrubPosition) const;
+	bool HitTestRangeEnd(const FScrubRangeToScreen& RangeToScreen, const TRange<double>& Range, float HitPixel) const;
+
+	FFrameTime SnapTimeToNearestKey(const FScrubRangeToScreen& RangeToScreen, float CursorPos, FFrameTime InTime) const;
 
-	float SnapTimeToNearestKey(const FScrubRangeToScreen& RangeToScreen, float CursorPos, float InTime);
+	void SetPlaybackRangeStart(FFrameNumber NewStart);
+	void SetPlaybackRangeEnd(FFrameNumber NewEnd);
 
-	void SetPlaybackRangeStart(float NewStart);
-	void SetPlaybackRangeEnd(float NewEnd);
+	void SetSelectionRangeStart(FFrameNumber NewStart);
+	void SetSelectionRangeEnd(FFrameNumber NewEnd);
+
+	TSharedRef<SWidget> OpenSetPlaybackRangeMenu(FFrameNumber FrameNumber);
+	FFrameTime ComputeScrubTimeFromMouse(const FGeometry& Geometry, FVector2D ScreenSpacePosition, FScrubRangeToScreen RangeToScreen) const;
+	FFrameTime ComputeFrameTimeFromMouse(const FGeometry& Geometry, FVector2D ScreenSpacePosition, FScrubRangeToScreen RangeToScreen, bool CheckSnapping = true) const;
+private:
+
+	struct FScrubPixelRange
+	{
+		TRange<float> Range;
+		TRange<float> HandleRange;
+		bool bClamped;
+	};
 
-	void SetSelectionRangeStart(float NewStart);
-	void SetSelectionRangeEnd(float NewEnd);
+	FScrubPixelRange GetHitTestScrubberPixelRange(FFrameTime ScrubTime, const FScrubRangeToScreen& RangeToScreen) const;
 
-	TSharedRef<SWidget> OpenSetPlaybackRangeMenu(float MouseTime);
+	FScrubPixelRange GetScrubberPixelRange(FFrameTime ScrubTime, const FScrubRangeToScreen& RangeToScreen) const;
+	FScrubPixelRange GetScrubberPixelRange(FFrameTime ScrubTime, FFrameRate Resolution, FFrameRate PlayRate, const FScrubRangeToScreen& RangeToScreen, float DilationPixels = 0.f) const;
 
 private:
+
+	/** Pointer back to the sequencer object */
+	TWeakPtr<FSequencer> WeakSequencer;
+
 	FTimeSliderArgs TimeSliderArgs;
 
-	/** The size of the scrub handle */
-	float ScrubHandleSize;
-	
-	/** Brush for drawing an upwards facing scrub handle */
-	const FSlateBrush* ScrubHandleUp;
+	/** Brush for drawingthe fill area on the scrubber */
+	const FSlateBrush* ScrubFillBrush;
+
+	/** Brush for drawing an upwards facing scrub handles */
+	const FSlateBrush* ScrubHandleUpBrush, *ClampedScrubHandleUpBrush;
 	
 	/** Brush for drawing a downwards facing scrub handle */
-	const FSlateBrush* ScrubHandleDown;
+	const FSlateBrush* ScrubHandleDownBrush, *ClampedScrubHandleDownBrush;
 	
 	/** Total mouse delta during dragging **/
 	float DistanceDragged;
@@ -239,12 +260,15 @@ private:
 	
 	/** If we are currently panning the panel */
 	bool bPanning;
-	
-	/** Mouse down time range */
-	FVector2D MouseDownRange;
-	
+
+	/** Mouse down position range */
+	FVector2D MouseDownPosition[2];
+
+	/** Geometry on mouse down */
+	FGeometry MouseDownGeometry;
+
 	/** Range stack */
-	TArray<FVector2D> RangeStack;
+	TArray<TRange<double>> ViewRangeStack;
 
 	/** When > 0, we should not show context menus */
 	int32 ContextMenuSuppression;
diff --git a/Engine/Source/Editor/Sequencer/Private/SequencerUtilities.cpp b/Engine/Source/Editor/Sequencer/Private/SequencerUtilities.cpp
--- a/Engine/Source/Editor/Sequencer/Private/SequencerUtilities.cpp
+++ b/Engine/Source/Editor/Sequencer/Private/SequencerUtilities.cpp
@@ -98,8 +98,8 @@ void FSequencerUtilities::PopulateMenu_CreateNewSection(FMenuBuilder& MenuBuilde
 			return;
 		}
 
-		const float StartAtTime = Sequencer->GetLocalTime();
-		TRange<float> VisibleRange = Sequencer->GetViewRange();
+		FQualifiedFrameTime CurrentTime = Sequencer->GetLocalTime();
+		TRange<double> VisibleRange = Sequencer->GetViewRange();
 
 		FScopedTransaction Transaction(LOCTEXT("AddSectionTransactionText", "Add Section"));
 		if (UMovieSceneSection* NewSection = Track->CreateNewSection())
@@ -112,10 +112,9 @@ void FSequencerUtilities::PopulateMenu_CreateNewSection(FMenuBuilder& MenuBuilde
 
 			Track->Modify();
 
-			NewSection->SetIsInfinite(false);
-			NewSection->SetStartTime(StartAtTime);
+			int32 DurationFrames = ( (VisibleRange.Size<double>() * 0.75) * CurrentTime.Rate ).FloorToFrame().Value;
+			NewSection->SetRange(TRange<FFrameNumber>(CurrentTime.Time.FrameNumber, CurrentTime.Time.FrameNumber + DurationFrames));
 			NewSection->SetOverlapPriority(OverlapPriority);
-			NewSection->SetEndTime(VisibleRange.GetUpperBoundValue() - (VisibleRange.GetUpperBoundValue() - StartAtTime)*0.25f);
 			NewSection->SetRowIndex(RowIndex);
 			NewSection->SetBlendType(BlendType);
 
