Copying //UE4/Dev-Sequencer to //UE4/Dev-Main (Source: //UE4/Dev-Sequencer @ 3946692)

diff --git a/Engine/Source/Editor/Sequencer/Private/Tools/EditToolDragOperations.cpp b/Engine/Source/Editor/Sequencer/Private/Tools/EditToolDragOperations.cpp
--- a/Engine/Source/Editor/Sequencer/Private/Tools/EditToolDragOperations.cpp
+++ b/Engine/Source/Editor/Sequencer/Private/Tools/EditToolDragOperations.cpp
@@ -8,6 +8,11 @@
 #include "SequencerCommonHelpers.h"
 #include "VirtualTrackArea.h"
 #include "SequencerTrackNode.h"
+#include "Algo/AllOf.h"
+#include "MovieSceneTimeHelpers.h"
+#include "Modules/ModuleManager.h"
+#include "Channels/MovieSceneChannelProxy.h"
+#include "ISequencerModule.h"
 
 struct FDefaultKeySnappingCandidates : ISequencerSnapCandidate
 {
@@ -46,16 +51,19 @@ struct FDefaultSectionSnappingCandidates : ISequencerSnapCandidate
 	TSet<UMovieSceneSection*> SectionsToIgnore;
 };
 
-TOptional<FSequencerSnapField::FSnapResult> SnapToInterval(const TArray<float>& InTimes, float Threshold, FSequencer* Sequencer)
+TOptional<FSequencerSnapField::FSnapResult> SnapToInterval(const TArray<FFrameNumber>& InTimes, int32 FrameThreshold, FFrameRate Resolution, FFrameRate PlayRate)
 {
 	TOptional<FSequencerSnapField::FSnapResult> Result;
 
-	float SnapAmount = 0.f;
-	for (float Time : InTimes)
+	FFrameNumber SnapAmount(0);
+	for (FFrameNumber Time : InTimes)
 	{
-		float IntervalSnap = SequencerHelpers::SnapTimeToInterval(Time, Sequencer->GetFixedFrameInterval());
-		float ThisSnapAmount = IntervalSnap - Time;
-		if (FMath::Abs(ThisSnapAmount) <= Threshold)
+		// Convert from resolution to playrate, round to frame, then back again
+		FFrameNumber PlayIntervalTime = FFrameRate::TransformTime(Time,             Resolution, PlayRate  ).RoundToFrame();
+		FFrameNumber IntervalSnap     = FFrameRate::TransformTime(PlayIntervalTime, PlayRate,   Resolution).FloorToFrame();
+
+		FFrameNumber ThisSnapAmount   = IntervalSnap - Time;
+		if (FMath::Abs(ThisSnapAmount) <= FrameThreshold)
 		{
 			if (!Result.IsSet() || FMath::Abs(ThisSnapAmount) < SnapAmount)
 			{
@@ -71,7 +79,7 @@ TOptional<FSequencerSnapField::FSnapResult> SnapToInterval(const TArray<float>&
 /** How many pixels near the mouse has to be before snapping occurs */
 const float PixelSnapWidth = 10.f;
 
-TRange<float> GetSectionBoundaries(const UMovieSceneSection* Section, TArray<FSectionHandle>& SectionHandles, TSharedPtr<FSequencerTrackNode> SequencerNode)
+TRange<FFrameNumber> GetSectionBoundaries(const UMovieSceneSection* Section, TArray<FSectionHandle>& SectionHandles, TSharedPtr<FSequencerTrackNode> SequencerNode)
 {
 	// Only get boundaries for the sections that aren't being moved
 	TArray<const UMovieSceneSection*> SectionsBeingMoved;
@@ -81,7 +89,7 @@ TRange<float> GetSectionBoundaries(const UMovieSceneSection* Section, TArray<FSe
 	}
 
 	// Find the borders of where you can drag to
-	float LowerBound = -FLT_MAX, UpperBound = FLT_MAX;
+	FFrameNumber LowerBound = TNumericLimits<int32>::Lowest(), UpperBound = TNumericLimits<int32>::Max();
 
 	// Also get the closest borders on either side
 	const TArray< TSharedRef<ISequencerSection> >& AllSections = SequencerNode->GetSections();
@@ -91,24 +99,24 @@ TRange<float> GetSectionBoundaries(const UMovieSceneSection* Section, TArray<FSe
 
 		if (!SectionsBeingMoved.Contains(TestSection) && Section->GetRowIndex() == TestSection->GetRowIndex())
 		{
-			if (TestSection->GetEndTime() <= Section->GetStartTime() && TestSection->GetEndTime() > LowerBound)
+			if (TestSection->HasEndFrame() && Section->HasStartFrame() && TestSection->GetExclusiveEndFrame() <= Section->GetInclusiveStartFrame() && TestSection->GetExclusiveEndFrame() > LowerBound)
 			{
-				LowerBound = TestSection->GetEndTime();
+				LowerBound = TestSection->GetExclusiveEndFrame();
 			}
-			if (TestSection->GetStartTime() >= Section->GetEndTime() && TestSection->GetStartTime() < UpperBound)
+			if (TestSection->HasStartFrame() && Section->HasEndFrame() && TestSection->GetInclusiveStartFrame() >= Section->GetExclusiveEndFrame() && TestSection->GetInclusiveStartFrame() < UpperBound)
 			{
-				UpperBound = TestSection->GetStartTime();
+				UpperBound = TestSection->GetInclusiveStartFrame();
 			}
 		}
 	}
 
-	return TRange<float>(LowerBound, UpperBound);
+	return TRange<FFrameNumber>(LowerBound, UpperBound);
 }
 
 FEditToolDragOperation::FEditToolDragOperation( FSequencer& InSequencer )
 	: Sequencer(InSequencer)
 {
-	Settings = Sequencer.GetSettings();
+	Settings = Sequencer.GetSequencerSettings();
 }
 
 FCursorReply FEditToolDragOperation::GetCursor() const
@@ -154,7 +162,7 @@ FResizeSection::FResizeSection( FSequencer& InSequencer, TArray<FSectionHandle>
 	, Sections( MoveTemp(InSections) )
 	, bDraggingByEnd(bInDraggingByEnd)
 	, bIsSlipping(bInIsSlipping)
-	, MouseDownTime(0.f)
+	, MouseDownTime(0)
 {
 }
 
@@ -162,78 +170,97 @@ void FResizeSection::OnBeginDrag(const FPointerEvent& MouseEvent, FVector2D Loca
 {
 	BeginTransaction( Sections, NSLOCTEXT("Sequencer", "DragSectionEdgeTransaction", "Resize section") );
 
-	MouseDownTime = VirtualTrackArea.PixelToTime(LocalMousePos.X);
+	MouseDownTime = VirtualTrackArea.PixelToFrame(LocalMousePos.X);
 
 	// Construct a snap field of unselected sections
 	FDefaultSectionSnappingCandidates SnapCandidates(Sections);
 	SnapField = FSequencerSnapField(Sequencer, SnapCandidates, ESequencerEntity::Section);
 
-	DraggedKeyHandles.Empty();
 	SectionInitTimes.Empty();
 
 	bool bIsDilating = MouseEvent.IsControlDown();
+	ISequencerModule& SequencerModule = FModuleManager::Get().LoadModuleChecked<ISequencerModule>("Sequencer");
 
 	for (auto& Handle : Sections)
 	{
 		UMovieSceneSection* Section = Handle.GetSectionObject();
 
-		for (auto SequencerSection : Handle.TrackNode->GetSections())
+		TSharedRef<ISequencerSection> SectionInterface = Handle.GetSectionInterface();
+		if (bIsDilating)
 		{
-			if (SequencerSection->GetSectionObject() == Section)
+			// Populate the resize data for this section
+			PreDragSectionData.Empty();
+			FPreDragSectionData ResizeData; 
+			ResizeData.MovieSection = Section;
+			ResizeData.InitialRange = Section->GetRange();
+
+			// Add the key times for all keys of all channels on this section
+			FMovieSceneChannelProxy& Proxy = Section->GetChannelProxy();
+			for (const FMovieSceneChannelEntry& Entry : Proxy.GetAllEntries())
 			{
-				if (bIsDilating)
-				{
-					SequencerSection->BeginDilateSection();
-				}
-				else
+				ISequencerChannelInterface* ChannelInterface = SequencerModule.FindChannelInterface(Entry.GetChannelID());
+
+				for (void* Channel : Entry.GetChannels())
 				{
-					SequencerSection->BeginResizeSection();
+					// Populate the cached state of this channel
+					FPreDragChannelData& ChannelData = ResizeData.Channels[ResizeData.Channels.Emplace()];
+					ChannelData.ChannelType = Entry.GetChannelID();
+					ChannelData.Channel = Proxy.MakeHandle(Channel);
+					ChannelInterface->GetKeys_Raw(Channel, TRange<FFrameNumber>::All(), &ChannelData.FrameNumbers, &ChannelData.Handles);
 				}
-				break;
 			}
+			PreDragSectionData.Emplace(ResizeData);
+		}
+		else
+		{
+			SectionInterface->BeginResizeSection();
 		}
 
-		Section->GetKeyHandles(DraggedKeyHandles, Section->GetRange());
-		SectionInitTimes.Add(Section, bDraggingByEnd ? Section->GetEndTime() : Section->GetStartTime());
+		SectionInitTimes.Add(Section, bDraggingByEnd ? Section->GetExclusiveEndFrame() : Section->GetInclusiveStartFrame());
 	}
 }
 
 void FResizeSection::OnEndDrag(const FPointerEvent& MouseEvent, FVector2D LocalMousePos, const FVirtualTrackArea& VirtualTrackArea)
 {
 	EndTransaction();
-
-	DraggedKeyHandles.Empty();
 }
 
 void FResizeSection::OnDrag(const FPointerEvent& MouseEvent, FVector2D LocalMousePos, const FVirtualTrackArea& VirtualTrackArea)
 {
+	ISequencerModule& SequencerModule = FModuleManager::Get().LoadModuleChecked<ISequencerModule>("Sequencer");
+
 	bool bIsDilating = MouseEvent.IsControlDown();
 
+	FFrameRate   FrameResolution = Sequencer.GetFocusedFrameResolution();
+	FFrameRate   PlayRate        = Sequencer.GetFocusedPlayRate();
+
 	// Convert the current mouse position to a time
-	float DeltaTime = VirtualTrackArea.PixelToTime(LocalMousePos.X) - MouseDownTime;
+	FFrameNumber DeltaTime = (VirtualTrackArea.PixelToFrame(LocalMousePos.X) - MouseDownTime).RoundToFrame();
 
 	// Snapping
 	if ( Settings->GetIsSnapEnabled() )
 	{
-		TArray<float> SectionTimes;
-		for (auto& Handle : Sections)
+		TArray<FFrameNumber> SectionTimes;
+		for (const FSectionHandle& Handle : Sections)
 		{
 			UMovieSceneSection* Section = Handle.GetSectionObject();
 			SectionTimes.Add(SectionInitTimes[Section] + DeltaTime);
 		}
 
-		float SnapThresold = VirtualTrackArea.PixelToTime(PixelSnapWidth) - VirtualTrackArea.PixelToTime(0.f);
+		float SnapThresholdPx = VirtualTrackArea.PixelToSeconds(PixelSnapWidth) - VirtualTrackArea.PixelToSeconds(0.f);
+		int32 SnapThreshold   = ( SnapThresholdPx * FrameResolution ).FloorToFrame().Value;
 
 		TOptional<FSequencerSnapField::FSnapResult> SnappedTime;
 
 		if (Settings->GetSnapSectionTimesToSections())
 		{
-			SnappedTime = SnapField->Snap(SectionTimes, SnapThresold);
+			SnappedTime = SnapField->Snap(SectionTimes, SnapThreshold);
 		}
 
 		if (!SnappedTime.IsSet() && Settings->GetSnapSectionTimesToInterval())
 		{
-			SnappedTime = SnapToInterval(SectionTimes, Sequencer.GetFixedFrameInterval()/2.f, &Sequencer);
+			int32 IntervalSnapThreshold = FMath::RoundToInt( ( FrameResolution / PlayRate ).AsDecimal() );
+			SnappedTime = SnapToInterval(SectionTimes, IntervalSnapThreshold, FrameResolution, PlayRate);
 		}
 
 		if (SnappedTime.IsSet())
@@ -242,33 +269,71 @@ void FResizeSection::OnDrag(const FPointerEvent& MouseEvent, FVector2D LocalMous
 			DeltaTime += SnappedTime->Snapped - SnappedTime->Original;
 		}
 	}
+	
+	/********************************************************************/
+	if (bIsDilating)
+	{
+		for(FPreDragSectionData Data: PreDragSectionData)
+		{
+			FFrameNumber StartPosition  = bDraggingByEnd ? MovieScene::DiscreteExclusiveUpper(Data.InitialRange) : MovieScene::DiscreteInclusiveLower(Data.InitialRange);
+			FFrameNumber DilationOrigin = bDraggingByEnd ? MovieScene::DiscreteInclusiveLower(Data.InitialRange) : MovieScene::DiscreteExclusiveUpper(Data.InitialRange);
+			FFrameNumber NewPosition    = bDraggingByEnd ? FMath::Max(StartPosition + DeltaTime, DilationOrigin) : FMath::Min(StartPosition + DeltaTime, DilationOrigin);
 
-	for (auto& Handle : Sections)
+			float DilationFactor = FMath::Abs(NewPosition.Value - DilationOrigin.Value) / float(MovieScene::DiscreteSize(Data.InitialRange));
+
+			if (bDraggingByEnd)
+			{
+				Data.MovieSection->SetRange(TRange<FFrameNumber>(Data.MovieSection->GetRange().GetLowerBound(), TRangeBound<FFrameNumber>::Exclusive(NewPosition)));
+			}
+			else
+			{
+				Data.MovieSection->SetRange(TRange<FFrameNumber>(TRangeBound<FFrameNumber>::Inclusive(NewPosition), Data.MovieSection->GetRange().GetUpperBound()));
+			}
+
+			TArray<FFrameNumber> NewFrameNumbers;
+			for (const FPreDragChannelData& ChannelData : Data.Channels)
+			{
+				ISequencerChannelInterface* ChannelInterface = SequencerModule.FindChannelInterface(ChannelData.ChannelType);
+
+				// Compute new frame times for each key
+				NewFrameNumbers.Reset(ChannelData.FrameNumbers.Num());
+				for (FFrameNumber StartFrame : ChannelData.FrameNumbers)
+				{
+					FFrameNumber NewTime = DilationOrigin + FFrameNumber(FMath::FloorToInt((StartFrame - DilationOrigin).Value * DilationFactor));
+					NewFrameNumbers.Add(NewTime);
+				}
+
+				// Apply the key times to the channel
+				void* RawChannelPtr = ChannelData.Channel.Get();
+				if (RawChannelPtr)
+				{
+					ChannelInterface->SetKeyTimes_Raw(RawChannelPtr, ChannelData.Handles, NewFrameNumbers);
+				}
+			}
+		}
+	}
+	/********************************************************************/
+	else for (const FSectionHandle& Handle : Sections)
 	{
 		UMovieSceneSection* Section = Handle.GetSectionObject();
 
 		// Find the corresponding sequencer section to this movie scene section
-		for (auto SequencerSection : Handle.TrackNode->GetSections())
+		for (const TSharedRef<ISequencerSection>& SequencerSection : Handle.TrackNode->GetSections())
 		{
 			if (SequencerSection->GetSectionObject() == Section)
 			{
-				float NewTime = SectionInitTimes[Section] + DeltaTime;
+				FFrameNumber NewTime = SectionInitTimes[Section] + DeltaTime;
 
 				if( bDraggingByEnd )
 				{
+					FFrameNumber MinFrame = Section->HasStartFrame() ? Section->GetInclusiveStartFrame() : TNumericLimits<int32>::Lowest();
+
 					// Dragging the end of a section
 					// Ensure we aren't shrinking past the start time
-					NewTime = FMath::Max( NewTime, Section->GetStartTime() );
-
-					if (bIsDilating)
-					{
-						float NewSize = NewTime - Section->GetStartTime();
-						float DilationFactor = NewSize / Section->GetTimeSize();
-						SequencerSection->DilateSection(DilationFactor, Section->GetStartTime(), DraggedKeyHandles);
-					}
-					else if (bIsSlipping)
+					NewTime = FMath::Max( NewTime, MinFrame );
+				    if (bIsSlipping)
 					{
-						SequencerSection->SlipSection( NewTime );
+						SequencerSection->SlipSection( NewTime/FrameResolution );
 					}
 					else
 					{
@@ -277,19 +342,15 @@ void FResizeSection::OnDrag(const FPointerEvent& MouseEvent, FVector2D LocalMous
 				}
 				else
 				{
+					FFrameNumber MaxFrame = Section->HasEndFrame() ? Section->GetExclusiveEndFrame()-1 : TNumericLimits<int32>::Max();
+
 					// Dragging the start of a section
 					// Ensure we arent expanding past the end time
-					NewTime = FMath::Min( NewTime, Section->GetEndTime() );
+					NewTime = FMath::Min( NewTime, MaxFrame );
 
-					if (bIsDilating)
-					{
-						float NewSize = Section->GetEndTime() - NewTime;
-						float DilationFactor = NewSize / Section->GetTimeSize();
-						SequencerSection->DilateSection(DilationFactor, Section->GetEndTime(), DraggedKeyHandles);
-					}
-					else if (bIsSlipping)
+					if (bIsSlipping)
 					{
-						SequencerSection->SlipSection( NewTime );
+						SequencerSection->SlipSection( NewTime/FrameResolution );
 					}
 					else
 					{
@@ -331,9 +392,10 @@ FMoveSection::FMoveSection( FSequencer& InSequencer, TArray<FSectionHandle> InSe
 	: FEditToolDragOperation( InSequencer )
 {
 	// Only allow sections that are not infinite to be movable.
-	for (auto InSection : InSections)
+	for (const FSectionHandle& InSection : InSections)
 	{
-		if (!InSection.GetSectionObject()->IsInfinite())
+		const UMovieSceneSection* Section = InSection.GetSectionObject();
+		if (Section->HasStartFrame() && Section->HasEndFrame())
 		{
 			Sections.Add(InSection);
 		}
@@ -360,17 +422,14 @@ void FMoveSection::OnBeginDrag(const FPointerEvent& MouseEvent, FVector2D LocalM
 	FDefaultSectionSnappingCandidates SnapCandidates(Sections);
 	SnapField = FSequencerSnapField(Sequencer, SnapCandidates, ESequencerEntity::Section);
 
-	DraggedKeyHandles.Empty();
-
-	const FVector2D InitialPosition = VirtualTrackArea.PhysicalToVirtual(LocalMousePos);
+	const FFrameTime InitialPosition = VirtualTrackArea.PixelToFrame(LocalMousePos.X);
 
 	RelativeOffsets.Reserve(Sections.Num());
 	for (int32 Index = 0; Index < Sections.Num(); ++Index)
 	{
-		auto* Section = Sections[Index].GetSectionObject();
-
-		Section->GetKeyHandles(DraggedKeyHandles, Section->GetRange());
-		RelativeOffsets.Add(FRelativeOffset{ Section->GetStartTime() - InitialPosition.X, Section->GetEndTime() - InitialPosition.X });
+		UMovieSceneSection* Section = Sections[Index].GetSectionObject();
+		// Must have a start frame and end frame to be in the Sections array
+		RelativeOffsets.Add(FRelativeOffset{ Section->GetInclusiveStartFrame() - InitialPosition, Section->GetExclusiveEndFrame() - InitialPosition });
 	}
 
 	TSet<UMovieSceneTrack*> Tracks;
@@ -394,7 +453,6 @@ void FMoveSection::OnEndDrag(const FPointerEvent& MouseEvent, FVector2D LocalMou
 		return;
 	}
 
-	DraggedKeyHandles.Empty();
 	InitialRowIndices.Empty();
 
 	TSet<UMovieSceneTrack*> Tracks;
@@ -436,38 +494,45 @@ void FMoveSection::OnDrag(const FPointerEvent& MouseEvent, FVector2D LocalMouseP
 
 	LocalMousePos.Y = FMath::Clamp(LocalMousePos.Y, 0.f, VirtualTrackArea.GetPhysicalSize().Y);
 
+	FFrameRate FrameResolution = Sequencer.GetFocusedFrameResolution();
+	FFrameRate PlayRate        = Sequencer.GetFocusedPlayRate();
+
 	// Convert the current mouse position to a time
-	FVector2D VirtualMousePos = VirtualTrackArea.PhysicalToVirtual(LocalMousePos);
+	FVector2D  VirtualMousePos = VirtualTrackArea.PhysicalToVirtual(LocalMousePos);
+	FFrameTime MouseTime       = VirtualTrackArea.PixelToFrame(LocalMousePos.X);
 
 	// Snapping
 	if ( Settings->GetIsSnapEnabled() )
 	{
-		TArray<float> SectionTimes;
+
+		float SnapThresholdPx = VirtualTrackArea.PixelToSeconds(PixelSnapWidth) - VirtualTrackArea.PixelToSeconds(0.f);
+		int32 SnapThreshold   = ( SnapThresholdPx * FrameResolution ).FloorToFrame().Value;
+
+		TArray<FFrameNumber> SectionTimes;
 		SectionTimes.Reserve(RelativeOffsets.Num());
 		for (const auto& Offset : RelativeOffsets)
 		{
-			SectionTimes.Add(VirtualMousePos.X + Offset.StartTime);
-			SectionTimes.Add(VirtualMousePos.X + Offset.EndTime);
+			SectionTimes.Add((Offset.StartTime + MouseTime).FloorToFrame());
+			SectionTimes.Add((Offset.EndTime   + MouseTime).FloorToFrame());
 		}
 
-		float SnapThresold = VirtualTrackArea.PixelToTime(PixelSnapWidth) - VirtualTrackArea.PixelToTime(0.f);
-
 		TOptional<FSequencerSnapField::FSnapResult> SnappedTime;
 
 		if (Settings->GetSnapSectionTimesToSections())
 		{
-			SnappedTime = SnapField->Snap(SectionTimes, SnapThresold);
+			SnappedTime = SnapField->Snap(SectionTimes, SnapThreshold);
 		}
 
 		if (!SnappedTime.IsSet() && Settings->GetSnapSectionTimesToInterval())
 		{
-			SnappedTime = SnapToInterval(SectionTimes, Sequencer.GetFixedFrameInterval()/2.f, &Sequencer);
+			int32 IntervalSnapThreshold = FMath::RoundToInt( ( FrameResolution / PlayRate ).AsDecimal() );
+			SnappedTime = SnapToInterval(SectionTimes, IntervalSnapThreshold, FrameResolution, PlayRate);
 		}
 
 		if (SnappedTime.IsSet())
 		{
 			// Add the snapped amount onto the delta
-			VirtualMousePos.X += SnappedTime->Snapped - SnappedTime->Original;
+			MouseTime += (SnappedTime->Snapped - SnappedTime->Original);
 		}
 	}
 
@@ -484,7 +549,7 @@ void FMoveSection::OnDrag(const FPointerEvent& MouseEvent, FVector2D LocalMouseP
 		SectionsBeingMoved.Add(SectionHandle.GetSectionObject());
 	}
 
-	TOptional<float> MinDeltaXTime;
+	TOptional<FFrameNumber> MinDeltaXTime;
 
 	// Disallow movement if any of the sections can't move
 	for (int32 Index = 0; Index < Sections.Num(); ++Index)
@@ -496,26 +561,21 @@ void FMoveSection::OnDrag(const FPointerEvent& MouseEvent, FVector2D LocalMouseP
 			continue;
 		}
 
-		// *2 because we store start/end times
-		const float DeltaTime = VirtualMousePos.X + RelativeOffsets[Index].StartTime - Section->GetStartTime();
+		const FFrameNumber DeltaTime = (MouseTime + RelativeOffsets[Index].StartTime - Section->GetInclusiveStartFrame()).FloorToFrame();
 
 		// Find the borders of where you can move to
-		TRange<float> SectionBoundaries = GetSectionBoundaries(Section, Sections, Handle.TrackNode);
+		TRange<FFrameNumber> SectionBoundaries = GetSectionBoundaries(Section, Sections, Handle.TrackNode);
 
-		float LeftMovementMaximum = SectionBoundaries.GetLowerBoundValue();
-		float RightMovementMaximum = SectionBoundaries.GetUpperBoundValue();
-		float NewStartTime = Section->GetStartTime() + DeltaTime;
-		float NewEndTime = Section->GetEndTime() + DeltaTime;
+		FFrameNumber LeftMovementMaximum  = MovieScene::DiscreteInclusiveLower(SectionBoundaries);
+		FFrameNumber RightMovementMaximum = MovieScene::DiscreteExclusiveUpper(SectionBoundaries);
+		FFrameNumber NewStartTime         = Section->GetInclusiveStartFrame() + DeltaTime;
+		FFrameNumber NewEndTime           = Section->GetExclusiveEndFrame()   + DeltaTime;
 
 		if (NewStartTime < LeftMovementMaximum || NewEndTime > RightMovementMaximum)
 		{
-			float ClampedDeltaTime = NewStartTime < LeftMovementMaximum ? LeftMovementMaximum - Section->GetStartTime() : RightMovementMaximum - Section->GetEndTime();
+			FFrameNumber ClampedDeltaTime = NewStartTime < LeftMovementMaximum ? LeftMovementMaximum - Section->GetInclusiveStartFrame() : RightMovementMaximum - Section->GetExclusiveEndFrame();
 
-			if (!MinDeltaXTime.IsSet())
-			{
-				MinDeltaXTime = ClampedDeltaTime;
-			}
-			else if (MinDeltaXTime.GetValue() > ClampedDeltaTime)
+			if (!MinDeltaXTime.IsSet() || MinDeltaXTime.GetValue() > ClampedDeltaTime)
 			{
 				MinDeltaXTime = ClampedDeltaTime;
 			}
@@ -528,17 +588,16 @@ void FMoveSection::OnDrag(const FPointerEvent& MouseEvent, FVector2D LocalMouseP
 		auto& Handle = Sections[Index];
 		UMovieSceneSection* Section = Handle.GetSectionObject();
 
-		// *2 because we store start/end times
-		const float DeltaTime = VirtualMousePos.X + RelativeOffsets[Index].StartTime - Section->GetStartTime();
+		const FFrameNumber DeltaTime = (MouseTime + RelativeOffsets[Index].StartTime - Section->GetInclusiveStartFrame()).FloorToFrame();
 
-		auto AllSections = Handle.TrackNode->GetTrack()->GetAllSections();
+		const TArray<UMovieSceneSection*>& AllSections = Handle.TrackNode->GetTrack()->GetAllSections();
 
-		TArray<UMovieSceneSection*> MovieSceneSections;
-		for (int32 i = 0; i < AllSections.Num(); ++i)
+		TArray<UMovieSceneSection*> NonDraggedSections;
+		for (UMovieSceneSection* TrackSection : AllSections)
 		{
-			if (!SectionsBeingMoved.Contains(AllSections[i]))
+			if (!SectionsBeingMoved.Contains(TrackSection))
 			{
-				MovieSceneSections.Add(AllSections[i]);
+				NonDraggedSections.Add(TrackSection);
 			}
 		}
 
@@ -549,11 +608,11 @@ void FMoveSection::OnDrag(const FPointerEvent& MouseEvent, FVector2D LocalMouseP
 		{
 			// Compute the max row index whilst disregarding the one we're dragging
 			int32 MaxRowIndex = 0;
-			for (auto* Sec : MovieSceneSections)
+			for (UMovieSceneSection* NonDraggedSection : NonDraggedSections)
 			{
-				if (Sec != Section)
+				if (NonDraggedSection != Section)
 				{
-					MaxRowIndex = FMath::Max(Sec->GetRowIndex() + 1, MaxRowIndex);
+					MaxRowIndex = FMath::Max(NonDraggedSection->GetRowIndex() + 1, MaxRowIndex);
 				}
 			}
 
@@ -611,20 +670,20 @@ void FMoveSection::OnDrag(const FPointerEvent& MouseEvent, FVector2D LocalMouseP
 			}
 		}
 
-		bool bDeltaX = !FMath::IsNearlyZero(DeltaTime);
+		bool bDeltaX = DeltaTime != 0;
 		bool bDeltaY = TargetRowIndex != Section->GetRowIndex();
 
 		// Horizontal movement
 		if (bDeltaX)
 		{
-			Section->MoveSection(MinDeltaXTime.Get(DeltaTime), DraggedKeyHandles);
+			Section->MoveSection(MinDeltaXTime.Get(DeltaTime));
 		}
 
 		// Vertical movement
 		if (bDeltaY && !bSectionsAreOnDifferentRows &&
 			(
 				Section->GetBlendType().IsValid() ||
-				!Section->OverlapsWithSections(MovieSceneSections, TargetRowIndex - Section->GetRowIndex(), DeltaTime)
+				!Section->OverlapsWithSections(NonDraggedSections, TargetRowIndex - Section->GetRowIndex(), DeltaTime.Value)
 			)
 		)
 		{
@@ -751,22 +810,32 @@ void FMoveKeys::OnBeginDrag(const FPointerEvent& MouseEvent, FVector2D LocalMous
 {
 	check( SelectedKeys.Num() > 0 )
 
-	FSequencerDisplayNode::DisableKeyGoupingRegeneration();
 
 	FDefaultKeySnappingCandidates SnapCandidates(SelectedKeys);
 	SnapField = FSequencerSnapField(Sequencer, SnapCandidates);
 
+	SelectedKeyArray = SelectedKeys.Array();
+
 	// Begin an editor transaction and mark the section as transactional so it's state will be saved
 	TArray<FSectionHandle> DummySections;
 	BeginTransaction( DummySections, NSLOCTEXT("Sequencer", "MoveKeysTransaction", "Move Keys") );
 
-	const float MouseTime = VirtualTrackArea.PixelToTime(LocalMousePos.X);
+	// Populate the relative offset for each key
+	TArray<FFrameNumber> KeyTimes;
+	KeyTimes.SetNum(SelectedKeyArray.Num());
+	GetKeyTimes(SelectedKeyArray, KeyTimes);
+
+	const FFrameTime MouseTime = VirtualTrackArea.PixelToFrame(LocalMousePos.X);
 
-	for( FSequencerSelectedKey SelectedKey : SelectedKeys )
+	RelativeOffsets.Reserve(KeyTimes.Num());
+	for (FFrameNumber Time : KeyTimes)
 	{
-		UMovieSceneSection* OwningSection = SelectedKey.Section;
+		RelativeOffsets.Add(Time - MouseTime);
+	}
 
-		RelativeOffsets.Add(SelectedKey, SelectedKey.KeyArea->GetKeyTime(SelectedKey.KeyHandle.GetValue()) - MouseTime);
+	for( const FSequencerSelectedKey& SelectedKey : SelectedKeyArray )
+	{
+		UMovieSceneSection* OwningSection = SelectedKey.Section;
 
 		// Only modify sections once
 		if( !ModifiedSections.Contains( OwningSection ) )
@@ -785,89 +854,95 @@ void FMoveKeys::OnBeginDrag(const FPointerEvent& MouseEvent, FVector2D LocalMous
 
 void FMoveKeys::OnDrag(const FPointerEvent& MouseEvent, FVector2D LocalMousePos, const FVirtualTrackArea& VirtualTrackArea)
 {
-	float MouseTime = VirtualTrackArea.PixelToTime(LocalMousePos.X);
-	float DistanceMoved = MouseTime - VirtualTrackArea.PixelToTime(LocalMousePos.X - MouseEvent.GetCursorDelta().X);
-
-	if( DistanceMoved == 0.0f )
+	if( MouseEvent.GetCursorDelta().X == 0.0f )
 	{
 		return;
 	}
 
+	FFrameRate FrameResolution = Sequencer.GetFocusedFrameResolution();
+	FFrameRate PlayRate        = Sequencer.GetFocusedPlayRate();
+
+	FFrameTime MouseTime       = VirtualTrackArea.PixelToFrame(LocalMousePos.X);
+
+	TArray<FFrameNumber> KeyTimes;
+	KeyTimes.Reserve(RelativeOffsets.Num());
+
+	for( FFrameTime Time : RelativeOffsets )
+	{
+		KeyTimes.Add( (MouseTime + Time).FloorToFrame() );
+	}
+
 	// Snapping
 	if ( Settings->GetIsSnapEnabled() )
 	{
-		TArray<float> KeyTimes;
-		for( FSequencerSelectedKey SelectedKey : SelectedKeys )
-		{
-			KeyTimes.Add(MouseTime + RelativeOffsets.FindRef(SelectedKey));
-		}
 
-		float SnapThresold = VirtualTrackArea.PixelToTime(PixelSnapWidth) - VirtualTrackArea.PixelToTime(0.f);
+		float SnapThresholdPx = VirtualTrackArea.PixelToSeconds(PixelSnapWidth) - VirtualTrackArea.PixelToSeconds(0.f);
+		int32 SnapThreshold   = ( SnapThresholdPx * FrameResolution ).FloorToFrame().Value;
 
 		TOptional<FSequencerSnapField::FSnapResult> SnappedTime;
 
 		if (Settings->GetSnapKeyTimesToKeys())
 		{
-			SnappedTime = SnapField->Snap(KeyTimes, SnapThresold);
+			SnappedTime = SnapField->Snap(KeyTimes, SnapThreshold);
 		}
 
 		if (!SnappedTime.IsSet() && Settings->GetSnapKeyTimesToInterval())
 		{
-			SnappedTime = SnapToInterval(KeyTimes, Sequencer.GetFixedFrameInterval()/2.f, &Sequencer);
+			int32 IntervalSnapThreshold = FMath::RoundToInt( ( FrameResolution / PlayRate ).AsDecimal() );
+			SnappedTime = SnapToInterval(KeyTimes, IntervalSnapThreshold, FrameResolution, PlayRate);
 		}
 
 		if (SnappedTime.IsSet())
 		{
 			MouseTime += SnappedTime->Snapped - SnappedTime->Original;
+
+			// Reset the new key times to account for the snap position
+			for( int32 Index = 0; Index < KeyTimes.Num(); ++Index )
+			{
+				KeyTimes[Index] = (MouseTime + RelativeOffsets[Index]).FloorToFrame();
+			}
 		}
 	}
 
-	float PrevNewKeyTime = FLT_MAX;
+	// Apply new key times to the selection
+	SetKeyTimes(SelectedKeyArray, KeyTimes);
 
-	for( FSequencerSelectedKey SelectedKey : SelectedKeys )
+	for (int32 Index = 0; Index < KeyTimes.Num(); ++Index)
 	{
-		UMovieSceneSection* Section = SelectedKey.Section;
+		FSequencerSelectedKey SelectedKey = SelectedKeyArray[Index];
 
-		if (!ModifiedSections.Contains(Section))
+		UMovieSceneSection* Section = SelectedKey.Section;
+		if (ModifiedSections.Contains(Section))
 		{
-			continue;
-		}
-
-		TSharedPtr<IKeyArea>& KeyArea = SelectedKey.KeyArea;
+			// If the key moves outside of the section resize the section to fit the key
+			// @todo Sequencer - Doesn't account for hitting other sections 
+			const FFrameNumber   NewKeyTime   = KeyTimes[Index];
+			TRange<FFrameNumber> SectionRange = Section->GetRange();
 
-		float NewKeyTime = MouseTime + RelativeOffsets.FindRef(SelectedKey);
-		float CurrentTime = KeyArea->GetKeyTime( SelectedKey.KeyHandle.GetValue() );
+			if (!SectionRange.Contains(NewKeyTime))
+			{
+				TRange<FFrameNumber> NewRange = TRange<FFrameNumber>::Hull( SectionRange, TRange<FFrameNumber>(NewKeyTime) );
+				Section->SetRange(NewRange);
+			}
+		}
+	}
 
-		// Tell the key area to move the key.  We reset the key index as a result of the move because moving a key can change it's internal index 
-		KeyArea->MoveKey( SelectedKey.KeyHandle.GetValue(), NewKeyTime - CurrentTime );
 
-		// If the key moves outside of the section resize the section to fit the key
-		// @todo Sequencer - Doesn't account for hitting other sections 
-		if( NewKeyTime > Section->GetEndTime() )
-		{
-			Section->SetEndTime( NewKeyTime );
-		}
-		else if( NewKeyTime < Section->GetStartTime() )
+	// Snap the play time to the new dragged key time if all the keyframes were dragged to the same time
+	if (Settings->GetSnapPlayTimeToDraggedKey() && KeyTimes.Num())
+	{
+		FFrameNumber FirstFrame        = KeyTimes[0];
+		auto         EqualsFirstFrame  = [=](FFrameNumber In)
 		{
-			Section->SetStartTime( NewKeyTime );
-		}
+			return In == FirstFrame;
+		};
 
-		if (PrevNewKeyTime == FLT_MAX)
+		if (Algo::AllOf(KeyTimes, EqualsFirstFrame))
 		{
-			PrevNewKeyTime = NewKeyTime;
-		}
-		else if (!FMath::IsNearlyEqual(NewKeyTime, PrevNewKeyTime))
-		{
-			PrevNewKeyTime = -FLT_MAX;
+			Sequencer.SetLocalTime(FirstFrame);
 		}
 	}
 
-	// Snap the play time to the new dragged key time if all the keyframes were dragged to the same time
-	if (Settings->GetSnapPlayTimeToDraggedKey() && PrevNewKeyTime != FLT_MAX && PrevNewKeyTime != -FLT_MAX)
-	{
-		Sequencer.SetLocalTime(PrevNewKeyTime);
-	}
-
 	for (UMovieSceneSection* Section : ModifiedSections)
 	{
 		Section->MarkAsChanged();
@@ -879,14 +954,10 @@ void FMoveKeys::OnEndDrag(const FPointerEvent& MouseEvent, FVector2D LocalMouseP
 {
 	ModifiedSections.Empty();
 	EndTransaction();
-	FSequencerDisplayNode::EnableKeyGoupingRegeneration();
 }
 
 void FDuplicateKeys::OnBeginDrag(const FPointerEvent& MouseEvent, FVector2D LocalMousePos, const FVirtualTrackArea& VirtualTrackArea)
 {
-	// Duplicate and select all the keys
-	TSet<FSequencerSelectedKey> OldSelection = SelectedKeys;
-
 	// Begin an editor transaction and mark the section as transactional so it's state will be saved
 	TArray<FSectionHandle> DummySections;
 	BeginTransaction( DummySections, NSLOCTEXT("Sequencer", "DuplicateKeysTransaction", "Duplicate Keys") );
@@ -912,12 +983,25 @@ void FDuplicateKeys::OnBeginDrag(const FPointerEvent& MouseEvent, FVector2D Loca
 	// Then duplicate the keys
 
 	// @todo sequencer: selection in transactions
-	Sequencer.GetSelection().EmptySelectedKeys();
-	for (const FSequencerSelectedKey& SelectedKey : OldSelection)
+	FSequencerSelection& Selection = Sequencer.GetSelection();
+
+	TArray<FSequencerSelectedKey> DuplicatedKeyArray = SelectedKeys.Array();
+	TArray<FKeyHandle> NewKeyHandles;
+	// Ideally we'd memset here, but there's no existing method to copy n bytes to an array n times
+	NewKeyHandles.SetNumUninitialized(DuplicatedKeyArray.Num());
+	for (FKeyHandle& Handle : NewKeyHandles)
 	{
-		FSequencerSelectedKey NewKey = SelectedKey;
-		NewKey.KeyHandle = SelectedKey.KeyArea->DuplicateKey(SelectedKey.KeyHandle.GetValue());
-		Sequencer.GetSelection().AddToSelection(NewKey);
+		Handle = FKeyHandle::Invalid();
+	}
+
+	Selection.EmptySelectedKeys();
+	DuplicateKeys(DuplicatedKeyArray, NewKeyHandles);
+	
+	for (int32 Index = 0; Index < NewKeyHandles.Num(); ++Index)
+	{
+		FSequencerSelectedKey NewKey = DuplicatedKeyArray[Index];
+		NewKey.KeyHandle = NewKeyHandles[Index];
+		Selection.AddToSelection(NewKey);
 	}
 
 	// Now start the move drag
@@ -935,7 +1019,7 @@ FManipulateSectionEasing::FManipulateSectionEasing( FSequencer& InSequencer, FSe
 	: FEditToolDragOperation(InSequencer)
 	, Handle(InSection)
 	, bEaseIn(_bEaseIn)
-	, MouseDownTime(0.f)
+	, MouseDownTime(0)
 {
 }
 
@@ -947,7 +1031,7 @@ void FManipulateSectionEasing::OnBeginDrag(const FPointerEvent& MouseEvent, FVec
 	Section->SetFlags( RF_Transactional );
 	Section->Modify();
 
-	MouseDownTime = VirtualTrackArea.PixelToTime(LocalMousePos.X);
+	MouseDownTime = VirtualTrackArea.PixelToFrame(LocalMousePos.X);
 
 	if (Settings->GetSnapSectionTimesToSections())
 	{
@@ -956,40 +1040,48 @@ void FManipulateSectionEasing::OnBeginDrag(const FPointerEvent& MouseEvent, FVec
 		SnapField = FSequencerSnapField(Sequencer, SnapCandidates, ESequencerEntity::Section);
 	}
 
-	InitValue = bEaseIn ? Section->Easing.GetEaseInTime() : Section->Easing.GetEaseOutTime();
+	InitValue = bEaseIn ? Section->Easing.GetEaseInDuration() : Section->Easing.GetEaseOutDuration();
 }
 
 void FManipulateSectionEasing::OnDrag(const FPointerEvent& MouseEvent, FVector2D LocalMousePos, const FVirtualTrackArea& VirtualTrackArea)
 {
+	FFrameRate FrameResolution = Sequencer.GetFocusedFrameResolution();
+	FFrameRate PlayRate        = Sequencer.GetFocusedPlayRate();
+
 	// Convert the current mouse position to a time
-	float DeltaTime = VirtualTrackArea.PixelToTime(LocalMousePos.X) - MouseDownTime;
+	FFrameTime  DeltaTime = VirtualTrackArea.PixelToFrame(LocalMousePos.X) - MouseDownTime;
 
 	// Snapping
 	if (Settings->GetIsSnapEnabled())
 	{
-		TArray<float> SnapTimes;
+		TArray<FFrameNumber> SnapTimes;
+
 		UMovieSceneSection* Section = Handle.GetSectionObject();
 		if (bEaseIn)
 		{
-			SnapTimes.Add(FMath::Clamp(Section->GetStartTime() + InitValue.Get(0.f) + DeltaTime, Section->GetStartTime(), Section->GetEndTime()));
+			FFrameNumber DesiredTime = (DeltaTime + Section->GetInclusiveStartFrame() + InitValue.Get(0)).RoundToFrame();
+			SnapTimes.Add(DesiredTime);
 		}
 		else
 		{
-			SnapTimes.Add(FMath::Clamp(Section->GetEndTime() - InitValue.Get(0.f) + DeltaTime, Section->GetStartTime(), Section->GetEndTime()));
+			FFrameNumber DesiredTime = (Section->GetExclusiveEndFrame() - InitValue.Get(0) + DeltaTime).RoundToFrame();
+			SnapTimes.Add(DesiredTime);
 		}
 
-		float SnapThresold = VirtualTrackArea.PixelToTime(PixelSnapWidth) - VirtualTrackArea.PixelToTime(0.f);
+		float SnapThresholdPx = VirtualTrackArea.PixelToSeconds(PixelSnapWidth) - VirtualTrackArea.PixelToSeconds(0.f);
+		int32 SnapThreshold   = ( SnapThresholdPx * FrameResolution ).FloorToFrame().Value;
 
 		TOptional<FSequencerSnapField::FSnapResult> SnappedTime;
 
 		if (Settings->GetSnapSectionTimesToSections())
 		{
-			SnappedTime = SnapField->Snap(SnapTimes, SnapThresold);
+			SnappedTime = SnapField->Snap(SnapTimes, SnapThreshold);
 		}
 
 		if (!SnappedTime.IsSet() && Settings->GetSnapSectionTimesToInterval())
 		{
-			SnappedTime = SnapToInterval(SnapTimes, Sequencer.GetFixedFrameInterval()/2.f, &Sequencer);
+			int32 IntervalSnapThreshold = FMath::RoundToInt( ( FrameResolution / PlayRate ).AsDecimal() );
+			SnappedTime = SnapToInterval(SnapTimes, IntervalSnapThreshold, FrameResolution, PlayRate);
 		}
 
 		if (SnappedTime.IsSet())
@@ -1001,15 +1093,17 @@ void FManipulateSectionEasing::OnDrag(const FPointerEvent& MouseEvent, FVector2D
 
 	UMovieSceneSection* Section = Handle.GetSectionObject();
 
+	const int32 MaxEasingDuration = Section->HasStartFrame() && Section->HasEndFrame() ? MovieScene::DiscreteSize(Section->GetRange()) : TNumericLimits<int32>::Max() / 2;
+
 	if (bEaseIn)
 	{
 		Section->Easing.bManualEaseIn = true;
-		Section->Easing.ManualEaseInTime = FMath::Clamp(InitValue.Get(0.f) + DeltaTime, 0.f, Section->GetEndTime() - Section->GetStartTime());
+		Section->Easing.ManualEaseInDuration  = FMath::Clamp(InitValue.Get(0) + DeltaTime.RoundToFrame().Value, 0, MaxEasingDuration);
 	}
 	else
 	{
 		Section->Easing.bManualEaseOut = true;
-		Section->Easing.ManualEaseOutTime = FMath::Clamp(InitValue.Get(0.f) - DeltaTime, 0.f, Section->GetEndTime() - Section->GetStartTime());
+		Section->Easing.ManualEaseOutDuration = FMath::Clamp(InitValue.Get(0) - DeltaTime.RoundToFrame().Value, 0, MaxEasingDuration);
 	}
 
 	UMovieSceneTrack* OuterTrack = Section->GetTypedOuter<UMovieSceneTrack>();
diff --git a/Engine/Source/Editor/Sequencer/Private/Tools/EditToolDragOperations.h b/Engine/Source/Editor/Sequencer/Private/Tools/EditToolDragOperations.h
--- a/Engine/Source/Editor/Sequencer/Private/Tools/EditToolDragOperations.h
+++ b/Engine/Source/Editor/Sequencer/Private/Tools/EditToolDragOperations.h
@@ -10,6 +10,8 @@
 #include "SequencerHotspots.h"
 #include "ScopedTransaction.h"
 #include "Tools/SequencerSnapField.h"
+#include "Channels/MovieSceneChannelHandle.h"
+
 
 class FSequencer;
 class FSlateWindowElementList;
@@ -80,6 +82,31 @@ private:
 	/** The sections we are interacting with */
 	TArray<FSectionHandle> Sections;
 
+	/********************************************************/
+	struct FPreDragChannelData
+	{
+		/** The ID of the type of the channel, for use with ISequencerModule::FindChannelInterface */
+		uint32 ChannelType;
+		/** Weak handle to the type-erased channel ptr */
+		TMovieSceneChannelHandle<void> Channel;
+
+		/** Array of all the handles in the section at the start of the drag */
+		TArray<FKeyHandle> Handles;
+		/** Array of all the above handle's times, one per index of Handles */
+		TArray<FFrameNumber> FrameNumbers;
+	};
+
+	struct FPreDragSectionData
+	{
+		/** Pointer to the movie section, this section is only valid during a drag operation*/
+		UMovieSceneSection * MovieSection;
+		/** The initial range of the section before it was resized */
+		TRange<FFrameNumber> InitialRange;
+		/** Array of all the channels in the section before it was resized */
+		TArray<FPreDragChannelData> Channels;
+	};
+	TArray<FPreDragSectionData> PreDragSectionData;
+
 	/** true if dragging  the end of the section, false if dragging the start */
 	bool bDraggingByEnd;
 
@@ -87,13 +114,10 @@ private:
 	bool bIsSlipping;
 
 	/** Time where the mouse is pressed */
-	float MouseDownTime;
+	FFrameTime MouseDownTime;
 
 	/** The section start or end times when the mouse is pressed */
-	TMap<TWeakObjectPtr<UMovieSceneSection>, float> SectionInitTimes;
-
-	/** The exact key handles that we're dragging */
-	TSet<FKeyHandle> DraggedKeyHandles;
+	TMap<TWeakObjectPtr<UMovieSceneSection>, FFrameNumber> SectionInitTimes;
 
 	/** Optional snap field to use when dragging */
 	TOptional<FSequencerSnapField> SnapField;
@@ -126,11 +150,8 @@ private:
 	/** The sections we are interacting with */
 	TArray<FSectionHandle> Sections;
 
-	/** The exact key handles that we're dragging */
-	TSet<FKeyHandle> DraggedKeyHandles;
-
 	/** Array of desired offsets relative to the mouse position. Representes a contiguous array of start/end times directly corresponding to Sections array */
-	struct FRelativeOffset { float StartTime, EndTime; };
+	struct FRelativeOffset { FFrameTime StartTime, EndTime; };
 
 	TArray<FRelativeOffset> RelativeOffsets;
 
@@ -170,8 +191,11 @@ protected:
 	/** The selected keys being moved. */
 	const TSet<FSequencerSelectedKey>& SelectedKeys;
 
+	/** Sequential copy of the above set used during drag */
+	TArray<FSequencerSelectedKey> SelectedKeyArray;
+
 	/** Map of relative offsets from the original mouse position */
-	TMap<FSequencerSelectedKey, float> RelativeOffsets;
+	TArray<FFrameTime> RelativeOffsets;
 
 	/** Snap field used to assist in snapping calculations */
 	TOptional<FSequencerSnapField> SnapField;
@@ -227,10 +251,10 @@ private:
 	bool bEaseIn;
 
 	/** Time where the mouse is pressed */
-	float MouseDownTime;
+	FFrameTime MouseDownTime;
 
 	/** The section ease in/out when the mouse was pressed */
-	TOptional<float> InitValue;
+	TOptional<int32> InitValue;
 
 	/** Optional snap field to use when dragging */
 	TOptional<FSequencerSnapField> SnapField;
diff --git a/Engine/Source/Editor/Sequencer/Private/Tools/SequencerEditTool_Movement.cpp b/Engine/Source/Editor/Sequencer/Private/Tools/SequencerEditTool_Movement.cpp
--- a/Engine/Source/Editor/Sequencer/Private/Tools/SequencerEditTool_Movement.cpp
+++ b/Engine/Source/Editor/Sequencer/Private/Tools/SequencerEditTool_Movement.cpp
@@ -15,6 +15,7 @@
 #include "IKeyArea.h"
 #include "Widgets/Layout/SBox.h"
 #include "Framework/Application/SlateApplication.h"
+#include "Timecode.h"
 
 const FName FSequencerEditTool_Movement::Identifier = "Movement";
 
@@ -38,14 +39,17 @@ FReply FSequencerEditTool_Movement::OnMouseButtonDown(SWidget& OwnerWidget, cons
 
 		DelayedDrag = FDelayedDrag_Hotspot(VirtualTrackArea.CachedTrackAreaGeometry().AbsoluteToLocal(MouseEvent.GetScreenSpacePosition()), MouseEvent.GetEffectingButton(), Hotspot);
 
- 		if (Sequencer.GetSettings()->GetSnapPlayTimeToPressedKey() || (MouseEvent.IsShiftDown() && MouseEvent.GetEffectingButton() == EKeys::LeftMouseButton) )
+ 		if (Sequencer.GetSequencerSettings()->GetSnapPlayTimeToPressedKey() || (MouseEvent.IsShiftDown() && MouseEvent.GetEffectingButton() == EKeys::LeftMouseButton) )
 		{
 			if (DelayedDrag->Hotspot.IsValid())
 			{
 				if (DelayedDrag->Hotspot->GetType() == ESequencerHotspot::Key)
 				{
-					FSequencerSelectedKey& ThisKey = StaticCastSharedPtr<FKeyHotspot>(DelayedDrag->Hotspot)->Key;
-					Sequencer.SetLocalTime(ThisKey.KeyArea->GetKeyTime(ThisKey.KeyHandle.GetValue()));
+					TOptional<FFrameNumber> Time = StaticCastSharedPtr<FKeyHotspot>(DelayedDrag->Hotspot)->GetTime();
+					if (Time.IsSet())
+					{
+						Sequencer.SetLocalTime(Time.GetValue());
+					}
 				}
 			}
 		}
@@ -71,8 +75,8 @@ FReply FSequencerEditTool_Movement::OnMouseMove(SWidget& OwnerWidget, const FGeo
 			if (DragOperation.IsValid())
 			{
 				DragPosition = MyGeometry.AbsoluteToLocal(MouseEvent.GetScreenSpacePosition());
-				
-				float CurrentTime = VirtualTrackArea.PixelToTime(DragPosition.X);
+
+				double CurrentTime = VirtualTrackArea.PixelToSeconds(DragPosition.X);
 				Sequencer.UpdateAutoScroll(CurrentTime);
 				
 				DragOperation->OnDrag(MouseEvent, DragPosition, VirtualTrackArea);
@@ -97,11 +101,11 @@ FReply FSequencerEditTool_Movement::OnMouseMove(SWidget& OwnerWidget, const FGeo
 	return FReply::Unhandled();
 }
 
-bool FSequencerEditTool_Movement::GetHotspotTime(float& HotspotTime) const
+bool FSequencerEditTool_Movement::GetHotspotTime(FFrameTime& HotspotTime) const
 {
 	if (DelayedDrag->Hotspot.IsValid())
 	{
-		TOptional<float> OptionalHotspotTime = DelayedDrag->Hotspot->GetTime();
+		TOptional<FFrameNumber> OptionalHotspotTime = DelayedDrag->Hotspot->GetTime();
 		if (OptionalHotspotTime.IsSet())
 		{
 			HotspotTime = OptionalHotspotTime.GetValue();
@@ -111,12 +115,12 @@ bool FSequencerEditTool_Movement::GetHotspotTime(float& HotspotTime) const
 	return false;
 }
 
-float FSequencerEditTool_Movement::GetHotspotOffsetTime(float CurrentTime) const
+FFrameTime FSequencerEditTool_Movement::GetHotspotOffsetTime(FFrameTime CurrentTime) const
 {
 	//@todo abstract dragging offset from shift
 	if (DelayedDrag->Hotspot.IsValid() && FSlateApplication::Get().GetModifierKeys().IsShiftDown())
 	{
-		TOptional<float> OptionalOffsetTime = DelayedDrag->Hotspot->GetOffsetTime();
+		TOptional<FFrameTime> OptionalOffsetTime = DelayedDrag->Hotspot->GetOffsetTime();
 		if (OptionalOffsetTime.IsSet())
 		{
 			return OptionalOffsetTime.GetValue();
@@ -183,15 +187,23 @@ TSharedPtr<ISequencerEditToolDragOperation> FSequencerEditTool_Movement::CreateD
 		// Moving key(s)?
 		else if (HotspotType == ESequencerHotspot::Key)
 		{
-			FSequencerSelectedKey& ThisKey = StaticCastSharedPtr<FKeyHotspot>(DelayedDrag->Hotspot)->Key;
+			TArrayView<const FSequencerSelectedKey> HoveredKeys = StaticCastSharedPtr<FKeyHotspot>(DelayedDrag->Hotspot)->Keys;
+
+			auto AnyUnselectedKey = [&Selection](const FSequencerSelectedKey& InKey)
+			{
+				return !Selection.IsSelected(InKey);
+			};
 
-			// If it's not selected, we'll treat this as a unique drag
-			if (!Selection.IsSelected(ThisKey))
+			if (HoveredKeys.ContainsByPredicate(AnyUnselectedKey))
 			{
+				// If any are not selected, we'll treat this as a unique drag
 				Selection.EmptySelectedKeys();
 				Selection.EmptySelectedSections();
 				Selection.EmptyNodesWithSelectedKeysOrSections();
-				Selection.AddToSelection(ThisKey);
+				for (const FSequencerSelectedKey& Key : HoveredKeys)
+				{
+					Selection.AddToSelection(Key);
+				}
 				SequencerHelpers::UpdateHoveredNodeFromSelectedKeys(Sequencer);
 			}
 
@@ -234,6 +246,8 @@ FReply FSequencerEditTool_Movement::OnMouseButtonUp(SWidget& OwnerWidget, const
 			GEditor->EndTransaction();
 		}
 
+		Sequencer.StopAutoscroll();
+
 		// Only return handled if we actually started a drag
 		return FReply::Handled().ReleaseMouseCapture();
 	}
@@ -303,7 +317,7 @@ int32 FSequencerEditTool_Movement::OnPaint(const FGeometry& AllottedGeometry, co
 
 		if (Hotspot.IsValid())
 		{
-			float CurrentTime = 0.0f;
+			FFrameTime CurrentTime;
 
 			if (GetHotspotTime(CurrentTime))
 			{
@@ -320,8 +334,8 @@ int32 FSequencerEditTool_Movement::OnPaint(const FGeometry& AllottedGeometry, co
 				const float HorizontalDelta = DragPosition.X - DelayedDrag->GetInitialPosition().X;
 				const float InitialY = DelayedDrag->GetInitialPosition().Y;
 
-				const FVector2D OldPos = FVector2D(VirtualTrackArea.TimeToPixel(OriginalHotspotTime), InitialY);
-				const FVector2D NewPos = FVector2D(VirtualTrackArea.TimeToPixel(CurrentTime), InitialY);
+				const FVector2D OldPos = FVector2D(VirtualTrackArea.FrameToPixel(OriginalHotspotTime), InitialY);
+				const FVector2D NewPos = FVector2D(VirtualTrackArea.FrameToPixel(CurrentTime), InitialY);
 
 				TArray<FVector2D> LinePoints;
 				{
@@ -377,7 +391,7 @@ int32 FSequencerEditTool_Movement::OnPaint(const FGeometry& AllottedGeometry, co
 				);
 
 				// draw offset string
-				float OffsetTime = GetHotspotOffsetTime(CurrentTime);
+				FFrameTime OffsetTime = GetHotspotOffsetTime(CurrentTime);
 				const FString OffsetString = TimeToString(OffsetTime, true);
 				const FVector2D OffsetStringSize = FontMeasureService->Measure(OffsetString, SmallLayoutFont);
 				const FVector2D OffsetPos = FVector2D(NewPos.X + MousePadding, NewPos.Y - 0.5f * OffsetStringSize.Y);
@@ -439,22 +453,59 @@ bool FSequencerEditTool_Movement::CanDeactivate() const
 }
 
 
-FString FSequencerEditTool_Movement::TimeToString(float Time, bool IsDelta) const
+FString FSequencerEditTool_Movement::TimeToString(FFrameTime Time, bool IsDelta) const
 {
-	USequencerSettings* Settings = Sequencer.GetSettings();
+	USequencerSettings* Settings = Sequencer.GetSequencerSettings();
+	check(Settings);
 
-	if ((Settings != nullptr) && Settings->GetShowFrameNumbers())
+	// We don't use the Sequencer's Numeric Type interface as we want to show a "+" only for delta movement and not the absolute time.
+	EFrameNumberDisplayFormats DisplayFormat = Settings->GetTimeDisplayFormat();
+	switch (DisplayFormat)
 	{
-		if (SequencerSnapValues::IsTimeSnapIntervalFrameRate(Sequencer.GetFixedFrameInterval()))
+		case EFrameNumberDisplayFormats::Seconds:
+		{
+			FFrameRate FrameResolution = Sequencer.GetFocusedFrameResolution();
+			double TimeInSeconds = FrameResolution.AsSeconds(Time);
+			return IsDelta ? FString::Printf(TEXT("[%+.2fs]"), TimeInSeconds) : FString::Printf(TEXT("%.2fs"), TimeInSeconds);
+		}
+		case EFrameNumberDisplayFormats::Frames:
+		{
+			FFrameRate FrameResolution = Sequencer.GetFocusedFrameResolution();
+			FFrameRate PlayRate = Sequencer.GetFocusedPlayRate();
+
+			// Convert from sequence resolution into display rate frames.
+			FFrameTime DisplayTime = FFrameRate::TransformTime(Time, FrameResolution, PlayRate);
+			FString SubframeIndicator = FMath::IsNearlyZero(DisplayTime.GetSubFrame()) ? TEXT("") : TEXT("*");
+			int32 ZeroPadFrames = Sequencer.GetSequencerSettings()->GetZeroPadFrames();
+			return IsDelta ? FString::Printf(TEXT("[%+0*d%s]"), ZeroPadFrames, DisplayTime.GetFrame().Value, *SubframeIndicator) : FString::Printf(TEXT("%0*d%s"), ZeroPadFrames, DisplayTime.GetFrame().Value, *SubframeIndicator);
+		}
+		case EFrameNumberDisplayFormats::NonDropFrameTimecode:
+		{
+			FFrameRate SourceFrameRate = Sequencer.GetFocusedFrameResolution();
+			FFrameRate DestinationFrameRate = Sequencer.GetFocusedPlayRate();
+
+			FFrameNumber PlayRateFrameNumber = FFrameRate::TransformTime(Time, SourceFrameRate, DestinationFrameRate).FloorToFrame();
+
+			FTimecode AsNonDropTimecode = FTimecode::FromFrameNumber(PlayRateFrameNumber, DestinationFrameRate, false);
+
+			const bool bForceSignDisplay = IsDelta;
+			return IsDelta ? FString::Printf(TEXT("[%s]"), *AsNonDropTimecode.ToString(bForceSignDisplay)) : FString::Printf(TEXT("%s"), *AsNonDropTimecode.ToString(bForceSignDisplay));
+		}
+		case EFrameNumberDisplayFormats::DropFrameTimecode:
 		{
-			const float FrameRate = 1.0f / Sequencer.GetFixedFrameInterval();
-			const int32 Frame = SequencerHelpers::TimeToFrame(Time, FrameRate);
+			FFrameRate SourceFrameRate = Sequencer.GetFocusedFrameResolution();
+			FFrameRate DestinationFrameRate = Sequencer.GetFocusedPlayRate();
+
+			FFrameNumber PlayRateFrameNumber = FFrameRate::TransformTime(Time, SourceFrameRate, DestinationFrameRate).FloorToFrame();
+
+			FTimecode AsDropTimecode = FTimecode::FromFrameNumber(PlayRateFrameNumber, DestinationFrameRate, true);
 
-			return IsDelta ? FString::Printf(TEXT("[%+d]"), Frame) : FString::Printf(TEXT("%d"), Frame);
+			const bool bForceSignDisplay = IsDelta;
+			return IsDelta ? FString::Printf(TEXT("[%s]"), *AsDropTimecode.ToString(bForceSignDisplay)) : FString::Printf(TEXT("%s"), *AsDropTimecode.ToString(bForceSignDisplay));
 		}
 	}
 
-	return IsDelta ? FString::Printf(TEXT("[%+.3f]"), Time) : FString::Printf(TEXT("%.3f"), Time);
+	return FString();
 }
 
 const ISequencerHotspot* FSequencerEditTool_Movement::GetDragHotspot() const
diff --git a/Engine/Source/Editor/Sequencer/Private/Tools/SequencerEditTool_Movement.h b/Engine/Source/Editor/Sequencer/Private/Tools/SequencerEditTool_Movement.h
--- a/Engine/Source/Editor/Sequencer/Private/Tools/SequencerEditTool_Movement.h
+++ b/Engine/Source/Editor/Sequencer/Private/Tools/SequencerEditTool_Movement.h
@@ -7,7 +7,7 @@
 #include "Input/CursorReply.h"
 #include "Input/Reply.h"
 #include "Tools/SequencerEditTool.h"
-#include "Tools/DelayedDrag.h"
+#include "Framework/DelayedDrag.h"
 
 class SSequencer;
 struct ISequencerHotspot;
@@ -39,14 +39,14 @@ public:
 
 protected:
 
-	FString TimeToString(float Time, bool IsDelta) const;
+	FString TimeToString(FFrameTime Time, bool IsDelta) const;
 
 private:
 
 	TSharedPtr<ISequencerEditToolDragOperation> CreateDrag(const FPointerEvent& MouseEvent);
 
-	bool GetHotspotTime(float& HotspotTime) const;
-	float GetHotspotOffsetTime(float CurrentTime) const;
+	bool GetHotspotTime(FFrameTime& HotspotTime) const;
+	FFrameTime GetHotspotOffsetTime(FFrameTime CurrentTime) const;
 
 	struct FDelayedDrag_Hotspot : FDelayedDrag
 	{
@@ -68,5 +68,5 @@ private:
 	FVector2D DragPosition;
 
 	/** The hotspot's time before dragging started. */
-	float OriginalHotspotTime;
+	FFrameTime OriginalHotspotTime;
 };
diff --git a/Engine/Source/Editor/Sequencer/Private/Tools/SequencerEditTool_Selection.cpp b/Engine/Source/Editor/Sequencer/Private/Tools/SequencerEditTool_Selection.cpp
--- a/Engine/Source/Editor/Sequencer/Private/Tools/SequencerEditTool_Selection.cpp
+++ b/Engine/Source/Editor/Sequencer/Private/Tools/SequencerEditTool_Selection.cpp
@@ -19,7 +19,7 @@ struct FSelectionPreviewVisitor final
 		, SetStateTo(InSetStateTo)
 	{}
 
-	virtual void VisitKey(FKeyHandle KeyHandle, float KeyTime, const TSharedPtr<IKeyArea>& KeyArea, UMovieSceneSection* Section, TSharedRef<FSequencerDisplayNode> Node) const override
+	virtual void VisitKey(FKeyHandle KeyHandle, FFrameNumber KeyTime, const TSharedPtr<IKeyArea>& KeyArea, UMovieSceneSection* Section, TSharedRef<FSequencerDisplayNode> Node) const override
 	{
 		FSequencerSelectedKey Key(*Section, KeyArea, KeyHandle);
 
@@ -50,7 +50,7 @@ struct FSelectionPreviewVisitor final
 	{
 		// Never select a combination of sections and keys
 		// Never allow infinite sections to be selected (they're only selectable through right click)
-		if (SelectionPreview.GetDefinedKeyStates().Num() == 0 && !Section->IsInfinite())
+		if (SelectionPreview.GetDefinedKeyStates().Num() == 0 && Section->GetRange() != TRange<FFrameNumber>::All())
 		{
 			SelectionPreview.SetSelectionState(Section, SetStateTo);
 			SelectionPreview.SetSelectionState(Node, SetStateTo);
@@ -141,13 +141,13 @@ public:
 		CurrentMousePos = LocalMousePos;
 		CurrentMousePos.X = FMath::Clamp(CurrentMousePos.X, 0.f, VirtualTrackArea.GetPhysicalSize().X);
 
-		TRange<float> ViewRange = Sequencer.GetViewRange();
+		TRange<double> ViewRange = Sequencer.GetViewRange();
 
 		// Handle virtual scrolling when at the horizontal extremes of the widget
 		{
-			const float ScrollThresholdH = ViewRange.Size<float>() * 0.025f;
+			const double ScrollThresholdH = ViewRange.Size<double>() * 0.025f;
 
-			float Difference = CurrentPosition.X - (ViewRange.GetLowerBoundValue() + ScrollThresholdH);
+			double Difference = CurrentPosition.X - (ViewRange.GetLowerBoundValue() + ScrollThresholdH);
 			if (Difference < 0 && MouseDelta.X < 0)
 			{
 				Sequencer.StartAutoscroll(Difference);
@@ -181,7 +181,7 @@ public:
 		const auto& RootNodes = SequencerWidget->GetTreeView()->GetNodeTree()->GetRootNodes();
 
 		// Now walk everything within the current marquee range, setting preview selection states as we go
-		FSequencerEntityWalker Walker(FSequencerEntityRange(TopLeft(), BottomRight()), VirtualKeySize);
+		FSequencerEntityWalker Walker(FSequencerEntityRange(TopLeft(), BottomRight(), VirtualTrackArea.GetFrameResolution()), VirtualKeySize);
 		Walker.Traverse(FSelectionPreviewVisitor(SelectionPreview, Sequencer.GetSelection(), PreviewState), RootNodes);
 	}
 
@@ -227,6 +227,17 @@ public:
 			if (Pair.Value == ESelectionPreviewState::Selected)
 			{
 				Selection.AddToNodesWithSelectedKeysOrSections(Pair.Key);
+
+				// If this node exists inside a collapsed parent, we add such parents to the selection preview as well so that the node highlight is correct for accumulated keys
+				TSharedPtr<FSequencerDisplayNode> ParentNode = Pair.Key->GetParent();
+				while (ParentNode.IsValid())
+				{
+					if (!ParentNode->IsExpanded())
+					{
+						Selection.AddToNodesWithSelectedKeysOrSections(ParentNode.ToSharedRef());
+					}
+					ParentNode = ParentNode->GetParent();
+				}
 			}
 			else
 			{
diff --git a/Engine/Source/Editor/Sequencer/Private/Tools/SequencerEditTool_Selection.h b/Engine/Source/Editor/Sequencer/Private/Tools/SequencerEditTool_Selection.h
--- a/Engine/Source/Editor/Sequencer/Private/Tools/SequencerEditTool_Selection.h
+++ b/Engine/Source/Editor/Sequencer/Private/Tools/SequencerEditTool_Selection.h
@@ -7,7 +7,7 @@
 #include "Input/CursorReply.h"
 #include "Input/Reply.h"
 #include "Tools/SequencerEditTool.h"
-#include "Tools/DelayedDrag.h"
+#include "Framework/DelayedDrag.h"
 
 class SSequencer;
 struct ISequencerHotspot;
diff --git a/Engine/Source/Editor/Sequencer/Private/Tools/SequencerEntityVisitor.cpp b/Engine/Source/Editor/Sequencer/Private/Tools/SequencerEntityVisitor.cpp
--- a/Engine/Source/Editor/Sequencer/Private/Tools/SequencerEntityVisitor.cpp
+++ b/Engine/Source/Editor/Sequencer/Private/Tools/SequencerEntityVisitor.cpp
@@ -3,16 +3,18 @@
 #include "Tools/SequencerEntityVisitor.h"
 #include "IKeyArea.h"
 #include "DisplayNodes/SequencerTrackNode.h"
-#include "GroupedKeyArea.h"
 #include "MovieSceneTrack.h"
+#include "MovieSceneTimeHelpers.h"
 
-FSequencerEntityRange::FSequencerEntityRange(const TRange<float>& InRange)
-	: StartTime(InRange.GetLowerBoundValue()), EndTime(InRange.GetUpperBoundValue())
+FSequencerEntityRange::FSequencerEntityRange(const TRange<double>& InRange, FFrameRate InFrameResolution)
+	: FrameResolution(InFrameResolution)
+	, Range(InRange)
 {
 }
 
-FSequencerEntityRange::FSequencerEntityRange(FVector2D TopLeft, FVector2D BottomRight)
-	: StartTime(TopLeft.X), EndTime(BottomRight.X)
+FSequencerEntityRange::FSequencerEntityRange(FVector2D TopLeft, FVector2D BottomRight, FFrameRate InFrameResolution)
+	: FrameResolution(InFrameResolution)
+	, Range(TopLeft.X, BottomRight.X)
 	, VerticalTop(TopLeft.Y), VerticalBottom(BottomRight.Y)
 {
 }
@@ -20,7 +22,7 @@ FSequencerEntityRange::FSequencerEntityRange(FVector2D TopLeft, FVector2D Bottom
 bool FSequencerEntityRange::IntersectSection(const UMovieSceneSection* InSection, const TSharedRef<FSequencerTrackNode>& InTrackNode, int32 MaxRowIndex) const
 {
 	// Test horizontal bounds
-	if (!InSection->IsInfinite() && (InSection->GetStartTime() > EndTime || InSection->GetEndTime() < StartTime))
+	if (!(InSection->GetRange() / FrameResolution).Overlaps(Range))
 	{
 		return false;
 	}
@@ -195,18 +197,11 @@ void FSequencerEntityWalker::HandleSingleNode(const ISequencerEntityVisitor& Vis
 
 	if (bIterateKeyGroupings)
 	{
-		for (TSharedRef<ISequencerSection> SectionInterface : InSections)
+		if (InNode->GetChildNodes().Num() != 0 && Range.IntersectKeyArea(InNode, VirtualKeySize.X))
 		{
-			UMovieSceneSection* Section = SectionInterface->GetSectionObject();
-
-			if (Visitor.CheckEntityMask(ESequencerEntity::Key))
+			for (TSharedRef<FSequencerDisplayNode> ChildNode : InNode->GetChildNodes())
 			{
-				// Only handle grouped keys if we actually have children
-				if (InNode->GetChildNodes().Num() != 0 && Range.IntersectKeyArea(InNode, VirtualKeySize.X))
-				{
-					TSharedRef<IKeyArea> KeyArea = InNode->GetKeyGrouping(Section);
-					HandleKeyArea(Visitor, KeyArea, Section, InNode);
-				}
+				HandleSingleNode(Visitor, ChildNode, InSections);
 			}
 		}
 	}
@@ -233,13 +228,18 @@ void FSequencerEntityWalker::HandleKeyAreaNode(const ISequencerEntityVisitor& Vi
 
 void FSequencerEntityWalker::HandleKeyArea(const ISequencerEntityVisitor& Visitor, const TSharedRef<IKeyArea>& KeyArea, UMovieSceneSection* Section, const TSharedRef<FSequencerDisplayNode>& InNode)
 {
-	for (FKeyHandle KeyHandle : KeyArea->GetUnsortedKeyHandles())
+	TArray<FKeyHandle> Handles;
+	TArray<FFrameNumber> Times;
+
+	const FFrameTime HalfKeySizeFrames   = (VirtualKeySize.X*.5f) * Range.FrameResolution;
+	const FFrameTime RangeStartFrame     = Range.Range.GetLowerBoundValue() * Range.FrameResolution - FFrameTime(0, 0.5f);
+	const FFrameTime RangeEndFrame       = Range.Range.GetUpperBoundValue() * Range.FrameResolution - FFrameTime(0, 0.5f);
+
+	TRange<FFrameNumber> VisitRangeFrames( (RangeStartFrame-HalfKeySizeFrames).CeilToFrame(), (RangeEndFrame+HalfKeySizeFrames).FloorToFrame() );
+	KeyArea->GetKeyInfo(&Handles, &Times, VisitRangeFrames);
+
+	for (int32 Index = 0; Index < Times.Num(); ++Index)
 	{
-		float KeyPosition = KeyArea->GetKeyTime(KeyHandle);
-		if (KeyPosition + VirtualKeySize.X/2 > Range.StartTime &&
-			KeyPosition - VirtualKeySize.X/2 < Range.EndTime)
-		{
-			Visitor.VisitKey(KeyHandle, KeyPosition, KeyArea, Section, InNode);
-		}
+		Visitor.VisitKey(Handles[Index], Times[Index], KeyArea, Section, InNode);
 	}
 }
diff --git a/Engine/Source/Editor/Sequencer/Private/Tools/SequencerEntityVisitor.h b/Engine/Source/Editor/Sequencer/Private/Tools/SequencerEntityVisitor.h
--- a/Engine/Source/Editor/Sequencer/Private/Tools/SequencerEntityVisitor.h
+++ b/Engine/Source/Editor/Sequencer/Private/Tools/SequencerEntityVisitor.h
@@ -21,14 +21,14 @@ namespace ESequencerEntity
 	};
 
 	static const uint32 Everything = (uint32)-1;
-};
+}
 
 /** Visitor class used to handle specific sequencer entities */
 struct ISequencerEntityVisitor
 {
 	ISequencerEntityVisitor(uint32 InEntityMask = ESequencerEntity::Everything) : EntityMask(InEntityMask) {}
 
-	virtual void VisitKey(FKeyHandle KeyHandle, float KeyTime, const TSharedPtr<IKeyArea>& KeyArea, UMovieSceneSection* Section, TSharedRef<FSequencerDisplayNode>) const { }
+	virtual void VisitKey(FKeyHandle KeyHandle, FFrameNumber KeyTime, const TSharedPtr<IKeyArea>& KeyArea, UMovieSceneSection* Section, TSharedRef<FSequencerDisplayNode>) const { }
 	virtual void VisitSection(UMovieSceneSection* Section, TSharedRef<FSequencerDisplayNode>) const { }
 	
 	/** Check if the specified type of entity is applicable to this visitor */
@@ -44,8 +44,8 @@ protected:
 /** A range specifying time (and possibly vertical) bounds in the sequencer */
 struct FSequencerEntityRange
 {
-	FSequencerEntityRange(const TRange<float>& InRange);
-	FSequencerEntityRange(FVector2D TopLeft, FVector2D BottomRight);
+	FSequencerEntityRange(const TRange<double>& InRange, FFrameRate InFrameResolution);
+	FSequencerEntityRange(FVector2D TopLeft, FVector2D BottomRight, FFrameRate InFrameResolution);
 
 	/** Check whether the specified section intersects this range */
 	bool IntersectSection(const UMovieSceneSection* InSection, const TSharedRef<FSequencerTrackNode>& InTrackNode, int32 MaxRowIndex) const;
@@ -56,8 +56,11 @@ struct FSequencerEntityRange
 	/** Check whether the specified node's key area intersects this range */
 	bool IntersectKeyArea(TSharedRef<FSequencerDisplayNode> InNode, float VirtualKeyHeight) const;
 
+	/** Frame resolution of the current time-base */
+	FFrameRate FrameResolution;
+
 	/** Start/end times */
-	float StartTime, EndTime;
+	TRange<double> Range;
 
 	/** Optional vertical bounds */
 	TOptional<float> VerticalTop, VerticalBottom;
diff --git a/Engine/Source/Editor/Sequencer/Private/Tools/SequencerSnapField.cpp b/Engine/Source/Editor/Sequencer/Private/Tools/SequencerSnapField.cpp
--- a/Engine/Source/Editor/Sequencer/Private/Tools/SequencerSnapField.cpp
+++ b/Engine/Source/Editor/Sequencer/Private/Tools/SequencerSnapField.cpp
@@ -4,6 +4,7 @@
 #include "MovieScene.h"
 #include "SSequencer.h"
 #include "SSequencerTreeView.h"
+#include "MovieSceneTimeHelpers.h"
 #include "MovieSceneSequence.h"
 
 struct FSnapGridVisitor : ISequencerEntityVisitor
@@ -13,7 +14,7 @@ struct FSnapGridVisitor : ISequencerEntityVisitor
 		, Candidate(InCandidate)
 	{}
 
-	virtual void VisitKey(FKeyHandle KeyHandle, float KeyTime, const TSharedPtr<IKeyArea>& KeyArea, UMovieSceneSection* Section, TSharedRef<FSequencerDisplayNode>) const
+	virtual void VisitKey(FKeyHandle KeyHandle, FFrameNumber KeyTime, const TSharedPtr<IKeyArea>& KeyArea, UMovieSceneSection* Section, TSharedRef<FSequencerDisplayNode>) const
 	{
 		if (Candidate.IsKeyApplicable(KeyHandle, KeyArea, Section))
 		{
@@ -24,15 +25,22 @@ struct FSnapGridVisitor : ISequencerEntityVisitor
 	{
 		if (Candidate.AreSectionBoundsApplicable(Section))
 		{
-			Snaps.Add(FSequencerSnapPoint{ FSequencerSnapPoint::SectionBounds, Section->GetStartTime() });
-			Snaps.Add(FSequencerSnapPoint{ FSequencerSnapPoint::SectionBounds, Section->GetEndTime() });
+			if (Section->HasStartFrame())
+			{
+				Snaps.Add(FSequencerSnapPoint{ FSequencerSnapPoint::SectionBounds, Section->GetInclusiveStartFrame() });
+			}
+
+			if (Section->HasEndFrame())
+			{
+				Snaps.Add(FSequencerSnapPoint{ FSequencerSnapPoint::SectionBounds, Section->GetExclusiveEndFrame() });
+			}
 		}
 
 		if (Candidate.AreSectionCustomSnapsApplicable(Section))
 		{
-			TArray<float> CustomSnaps;
+			TArray<FFrameNumber> CustomSnaps;
 			Section->GetSnapTimes(CustomSnaps, false);
-			for (float Time : CustomSnaps)
+			for (FFrameNumber Time : CustomSnaps)
 			{
 				Snaps.Add(FSequencerSnapPoint{ FSequencerSnapPoint::CustomSection, Time });
 			}
@@ -53,25 +61,25 @@ FSequencerSnapField::FSequencerSnapField(const ISequencer& InSequencer, ISequenc
 		VisibleNodes.Add(Geometry.Node);
 	}
 
-	auto ViewRange = InSequencer.GetViewRange();
-	FSequencerEntityWalker Walker(ViewRange);
+	TRange<double> ViewRange = InSequencer.GetViewRange();
+	FSequencerEntityWalker Walker(FSequencerEntityRange(ViewRange, InSequencer.GetFocusedFrameResolution()));
 
 	// Traverse the visible space, collecting snapping times as we go
 	FSnapGridVisitor Visitor(Candidate, EntityMask);
 	Walker.Traverse(Visitor, VisibleNodes);
 
 	// Add the playback range start/end bounds as potential snap candidates
-	TRange<float> PlaybackRange = InSequencer.GetFocusedMovieSceneSequence()->GetMovieScene()->GetPlaybackRange();
+	TRange<FFrameNumber> PlaybackRange = InSequencer.GetFocusedMovieSceneSequence()->GetMovieScene()->GetPlaybackRange();
 	Visitor.Snaps.Add(FSequencerSnapPoint{ FSequencerSnapPoint::PlaybackRange, PlaybackRange.GetLowerBoundValue() });
 	Visitor.Snaps.Add(FSequencerSnapPoint{ FSequencerSnapPoint::PlaybackRange, PlaybackRange.GetUpperBoundValue() });
 
 	// Add the current time as a potential snap candidate
-	Visitor.Snaps.Add(FSequencerSnapPoint{ FSequencerSnapPoint::CurrentTime, InSequencer.GetLocalTime() });
+	Visitor.Snaps.Add(FSequencerSnapPoint{ FSequencerSnapPoint::CurrentTime, InSequencer.GetLocalTime().Time.FrameNumber });
 
 	// Add the selection range bounds as a potential snap candidate
-	TRange<float> SelectionRange = InSequencer.GetFocusedMovieSceneSequence()->GetMovieScene()->GetSelectionRange();
-	Visitor.Snaps.Add(FSequencerSnapPoint{ FSequencerSnapPoint::InOutRange, SelectionRange.GetLowerBoundValue() });
-	Visitor.Snaps.Add(FSequencerSnapPoint{ FSequencerSnapPoint::InOutRange, SelectionRange.GetUpperBoundValue() });
+	TRange<FFrameNumber> SelectionRange = InSequencer.GetFocusedMovieSceneSequence()->GetMovieScene()->GetSelectionRange();
+	Visitor.Snaps.Add(FSequencerSnapPoint{ FSequencerSnapPoint::InOutRange, MovieScene::DiscreteInclusiveLower(SelectionRange) });
+	Visitor.Snaps.Add(FSequencerSnapPoint{ FSequencerSnapPoint::InOutRange, MovieScene::DiscreteExclusiveUpper(SelectionRange) });
 
 	// Sort
 	Visitor.Snaps.Sort([](const FSequencerSnapPoint& A, const FSequencerSnapPoint& B){
@@ -81,12 +89,12 @@ FSequencerSnapField::FSequencerSnapField(const ISequencer& InSequencer, ISequenc
 	// Remove duplicates
 	for (int32 Index = 0; Index < Visitor.Snaps.Num(); ++Index)
 	{
-		const float CurrentTime = Visitor.Snaps[Index].Time;
+		const FFrameNumber CurrentTime = Visitor.Snaps[Index].Time;
 
 		int32 NumToMerge = 0;
 		for (int32 DuplIndex = Index + 1; DuplIndex < Visitor.Snaps.Num(); ++DuplIndex)
 		{
-			if (!FMath::IsNearlyEqual(CurrentTime, Visitor.Snaps[DuplIndex].Time))
+			if (CurrentTime != Visitor.Snaps[DuplIndex].Time)
 			{
 				break;
 			}
@@ -102,7 +110,7 @@ FSequencerSnapField::FSequencerSnapField(const ISequencer& InSequencer, ISequenc
 	SortedSnaps = MoveTemp(Visitor.Snaps);
 }
 
-TOptional<float> FSequencerSnapField::Snap(float InTime, float Threshold) const
+TOptional<FFrameNumber> FSequencerSnapField::Snap(FFrameNumber InTime, int32 Threshold) const
 {
 	int32 Min = 0;
 	int32 Max = SortedSnaps.Num();
@@ -112,7 +120,7 @@ TOptional<float> FSequencerSnapField::Snap(float InTime, float Threshold) const
 	{
 		int32 SearchIndex = Min + (Max - Min) / 2;
 
-		float ProspectiveSnapPos = SortedSnaps[SearchIndex].Time;
+		FFrameNumber ProspectiveSnapPos = SortedSnaps[SearchIndex].Time;
 		if (ProspectiveSnapPos > InTime + Threshold)
 		{
 			Max = SearchIndex;
@@ -125,12 +133,12 @@ TOptional<float> FSequencerSnapField::Snap(float InTime, float Threshold) const
 		{
 			// Linearly search forwards and backwards to find the closest snap
 
-			float SnapDelta = ProspectiveSnapPos - InTime;
+			FFrameNumber SnapDelta = ProspectiveSnapPos - InTime;
 
 			// Search forwards while we're in the threshold
 			for (int32 FwdIndex = SearchIndex+1; FwdIndex < Max-1 && SortedSnaps[FwdIndex].Time < InTime + Threshold; ++FwdIndex)
 			{
-				float ThisSnapDelta = InTime - SortedSnaps[FwdIndex].Time;
+				FFrameNumber ThisSnapDelta = InTime - SortedSnaps[FwdIndex].Time;
 				if (FMath::Abs(ThisSnapDelta) < FMath::Abs(SnapDelta))
 				{
 					SnapDelta = ThisSnapDelta;
@@ -141,7 +149,7 @@ TOptional<float> FSequencerSnapField::Snap(float InTime, float Threshold) const
 			// Search backwards while we're in the threshold
 			for (int32 BckIndex = SearchIndex-1; BckIndex >= Min && SortedSnaps[BckIndex].Time > InTime + Threshold; --BckIndex)
 			{
-				float ThisSnapDelta = InTime - SortedSnaps[BckIndex].Time;
+				FFrameNumber ThisSnapDelta = InTime - SortedSnaps[BckIndex].Time;
 				if (FMath::Abs(ThisSnapDelta) < FMath::Abs(SnapDelta))
 				{
 					SnapDelta = ThisSnapDelta;
@@ -153,23 +161,23 @@ TOptional<float> FSequencerSnapField::Snap(float InTime, float Threshold) const
 		}
 	}
 
-	return TOptional<float>();
+	return TOptional<FFrameNumber>();
 }
 
-TOptional<FSequencerSnapField::FSnapResult> FSequencerSnapField::Snap(const TArray<float>& InTimes, float Threshold) const
+TOptional<FSequencerSnapField::FSnapResult> FSequencerSnapField::Snap(const TArray<FFrameNumber>& InTimes, int32 Threshold) const
 {
 	TOptional<FSnapResult> ProspectiveSnap;
-	float SnapDelta = 0.f;
+	FFrameNumber SnapDelta(0);
 
-	for (float Time : InTimes)
+	for (FFrameNumber Time : InTimes)
 	{
-		TOptional<float> ThisSnap = Snap(Time, Threshold);
+		TOptional<FFrameNumber> ThisSnap = Snap(Time, Threshold);
 		if (!ThisSnap.IsSet())
 		{
 			continue;
 		}
 
-		float ThisSnapDelta = ThisSnap.GetValue() - Time;
+		FFrameNumber ThisSnapDelta = ThisSnap.GetValue() - Time;
 		if (!ProspectiveSnap.IsSet() || FMath::Abs(ThisSnapDelta) < FMath::Abs(SnapDelta))
 		{
 			ProspectiveSnap = FSnapResult{ Time, ThisSnap.GetValue() };
diff --git a/Engine/Source/Editor/Sequencer/Private/Tools/SequencerSnapField.h b/Engine/Source/Editor/Sequencer/Private/Tools/SequencerSnapField.h
--- a/Engine/Source/Editor/Sequencer/Private/Tools/SequencerSnapField.h
+++ b/Engine/Source/Editor/Sequencer/Private/Tools/SequencerSnapField.h
@@ -18,7 +18,7 @@ struct FSequencerSnapPoint
 	ESnapType Type;
 
 	/** The time of the snap */
-	float Time;
+	FFrameNumber Time;
 };
 
 /** Interface that defines how to construct an FSequencerSnapField */
@@ -45,9 +45,9 @@ public:
 	struct FSnapResult
 	{
 		/** The time before it was snapped */
-		float Original;
+		FFrameNumber Original;
 		/** The time after it was snapped */
-		float Snapped;
+		FFrameNumber Snapped;
 	};
 
 	/** Construction from a sequencer and a snap canidate implementation. Optionally provide an entity mask to completely ignore some entity types */
@@ -58,10 +58,10 @@ public:
 	FSequencerSnapField& operator=(FSequencerSnapField&& In) { SortedSnaps = MoveTemp(In.SortedSnaps); return *this; }
 	
 	/** Snap the specified time to this field with the given threshold */
-	TOptional<float> Snap(float InTime, float Threshold) const;
+	TOptional<FFrameNumber> Snap(FFrameNumber InTime, int32 Threshold) const;
 
 	/** Snap the specified times to this field with the given threshold. Will return the closest snap value of the entire intersection. */
-	TOptional<FSnapResult> Snap(const TArray<float>& InTimes, float Threshold) const;
+	TOptional<FSnapResult> Snap(const TArray<FFrameNumber>& InTimes, int32 Threshold) const;
 
 private:
 	/** Array of snap points, approximately grouped, and sorted in ascending order by time */
diff --git a/Engine/Source/Editor/Sequencer/Private/VirtualTrackArea.cpp b/Engine/Source/Editor/Sequencer/Private/VirtualTrackArea.cpp
--- a/Engine/Source/Editor/Sequencer/Private/VirtualTrackArea.cpp
+++ b/Engine/Source/Editor/Sequencer/Private/VirtualTrackArea.cpp
@@ -2,11 +2,10 @@
 
 #include "VirtualTrackArea.h"
 #include "SSequencerTreeView.h"
-#include "GroupedKeyArea.h"
 #include "SequencerTrackNode.h"
 
 FVirtualTrackArea::FVirtualTrackArea(const FSequencer& InSequencer, SSequencerTreeView& InTreeView, const FGeometry& InTrackAreaGeometry)
-	: FTimeToPixel(InTrackAreaGeometry, InSequencer.GetViewRange())
+	: FTimeToPixel(InTrackAreaGeometry, InSequencer.GetViewRange(), InSequencer.GetFocusedFrameResolution())
 	, TreeView(InTreeView)
 	, TrackAreaGeometry(InTrackAreaGeometry)
 {
@@ -25,7 +24,7 @@ float FVirtualTrackArea::VerticalOffsetToPixel(float InOffset) const
 FVector2D FVirtualTrackArea::PhysicalToVirtual(FVector2D InPosition) const
 {
 	InPosition.Y = PixelToVerticalOffset(InPosition.Y);
-	InPosition.X = PixelToTime(InPosition.X);
+	InPosition.X = PixelToSeconds(InPosition.X);
 
 	return InPosition;
 }
@@ -33,7 +32,7 @@ FVector2D FVirtualTrackArea::PhysicalToVirtual(FVector2D InPosition) const
 FVector2D FVirtualTrackArea::VirtualToPhysical(FVector2D InPosition) const
 {
 	InPosition.Y = VerticalOffsetToPixel(InPosition.Y);
-	InPosition.X = TimeToPixel(InPosition.X);
+	InPosition.X = SecondsToPixel(InPosition.X);
 
 	return InPosition;
 }
@@ -73,7 +72,7 @@ TOptional<FSectionHandle> FVirtualTrackArea::HitTestSection(FVector2D InPhysical
 
 		if (TrackNode.IsValid())
 		{
-			float Time = PixelToTime(InPhysicalPosition.X);
+			FFrameNumber Time = PixelToFrame(InPhysicalPosition.X).FloorToFrame();
 
 			const auto& Sections = TrackNode->GetSections();
 
@@ -120,8 +119,9 @@ FSequencerSelectedKey FVirtualTrackArea::HitTestKey(FVector2D InPhysicalPosition
 		return FSequencerSelectedKey();
 	}
 
-	const float KeyLeft  = PixelToTime(InPhysicalPosition.X - SequencerSectionConstants::KeySize.X/2);
-	const float KeyRight = PixelToTime(InPhysicalPosition.X + SequencerSectionConstants::KeySize.X/2);
+	const double   KeyLeft  = PixelToSeconds(InPhysicalPosition.X - SequencerSectionConstants::KeySize.X/2);
+	const double   KeyRight = PixelToSeconds(InPhysicalPosition.X + SequencerSectionConstants::KeySize.X/2);
+	TRange<FFrameNumber> KeyRange((KeyLeft * GetFrameResolution()).FloorToFrame(), (KeyRight * GetFrameResolution()).CeilToFrame());
 
 	TArray<TSharedRef<IKeyArea>> KeyAreas;
 
@@ -138,40 +138,49 @@ FSequencerSelectedKey FVirtualTrackArea::HitTestKey(FVector2D InPhysicalPosition
 		for (auto& KeyArea : KeyAreaNode->GetAllKeyAreas())
 		{
 			UMovieSceneSection* Section = KeyArea->GetOwningSection();
-			if (Section->GetStartTime() <= KeyRight && Section->GetEndTime() >= KeyLeft)
+			if (Section->GetRange().Overlaps(KeyRange))
 			{
 				KeyAreas.Add(KeyArea);
 			}
 		}
 	}
-	// Failing that, and the node is collapsed, we check for key groupings
+	// Failing that, and the node is collapsed, we check for collapsed key areas that are underneath this node
 	else if (!Node->IsExpanded())
 	{
 		TSharedPtr<FSequencerTrackNode> TrackNode = GetParentTrackNode(*Node);
 		if (TrackNode.IsValid())
 		{
+			TArray<TSharedRef<FSequencerSectionKeyAreaNode>> KeyAreaNodes;
+			TrackNode->GetChildKeyAreaNodesRecursively(KeyAreaNodes);
+
 			for (TSharedRef<ISequencerSection> SectionInterface : TrackNode->GetSections())
 			{
 				UMovieSceneSection* Section = SectionInterface->GetSectionObject();
-				if (Section->GetStartTime() <= KeyRight && Section->GetEndTime() >= KeyLeft)
+				if (Section->GetRange().Overlaps(KeyRange))
 				{
-					KeyAreas.Add(Node->GetKeyGrouping(Section));
+					for (TSharedRef<FSequencerSectionKeyAreaNode> ChildKeyAreaNode : KeyAreaNodes)
+					{
+						TSharedPtr<IKeyArea> KeyArea = ChildKeyAreaNode->IsHidden() ? nullptr : ChildKeyAreaNode->GetKeyArea(Section);
+						if (KeyArea.IsValid())
+						{
+							KeyAreas.Add(KeyArea.ToSharedRef());
+						}
+					}
 				}
 			}
 		}
 	}
 
 	// Search for any key that matches the position
-	// todo: investigate if this would be faster as a sort + binary search, rather than linear search
-	for (auto& KeyArea : KeyAreas)
+	TArray<FKeyHandle> Handles;
+	for (TSharedRef<IKeyArea> KeyArea : KeyAreas)
 	{
-		for (FKeyHandle Key : KeyArea->GetUnsortedKeyHandles())
+		Handles.Reset();
+		KeyArea->GetKeyHandles(Handles, KeyRange);
+
+		if (Handles.Num())
 		{
-			const float Time = KeyArea->GetKeyTime(Key);
-			if (Time >= KeyLeft && Time <= KeyRight)
-			{
-				return FSequencerSelectedKey(*KeyArea->GetOwningSection(), KeyArea, Key);
-			}
+			return FSequencerSelectedKey(*KeyArea->GetOwningSection(), KeyArea, Handles[0]);
 		}
 	}
 
diff --git a/Engine/Source/Editor/Sequencer/Public/IKeyArea.h b/Engine/Source/Editor/Sequencer/Public/IKeyArea.h
--- a/Engine/Source/Editor/Sequencer/Public/IKeyArea.h
+++ b/Engine/Source/Editor/Sequencer/Public/IKeyArea.h
@@ -2,218 +2,300 @@
 
 #pragma once
 
-#include "CoreMinimal.h"
-#include "Widgets/SWidget.h"
+#include "Math/Color.h"
+#include "UObject/ObjectKey.h"
+#include "UObject/NameTypes.h"
+#include "Templates/SharedPointer.h"
+#include "Containers/ArrayView.h"
 #include "Curves/KeyHandle.h"
-#include "Curves/RichCurve.h"
+#include "MovieSceneCommonHelpers.h"
+#include "Channels/MovieSceneChannelHandle.h"
 
-class FMovieSceneClipboardBuilder;
-class FMovieSceneClipboardKeyTrack;
-class FStructOnScope;
-class ISequencer;
-class UMovieSceneSection;
-struct FMovieSceneClipboardEnvironment;
+struct FRichCurve;
+struct FKeyDrawParams;
+struct FMovieSceneChannelProxy;
 struct FSequencerPasteEnvironment;
-struct FSlateBrush;
-enum class EMovieSceneKeyInterpolation : uint8;
-
-/**
- * Information for how to draw each key.
- */
-struct FKeyDrawingInfo
-{
-	FKeyDrawingInfo()
-		: Brush(nullptr)
-		, Size(FVector2D::ZeroVector)
-	{ }
-
-	/** The brush to use for each key */
-	const FSlateBrush* Brush;
-
-	/** The size of each key */
-	const FVector2D Size;
-};
+struct ISequencerChannelInterface;
+struct FMovieSceneChannelEditorData;
+struct FMovieSceneClipboardEnvironment;
 
+class SWidget;
+class ISequencer;
+class FCurveModel;
+class FStructOnScope;
+class UMovieSceneSection;
+class FMovieSceneClipboardBuilder;
+class FMovieSceneClipboardKeyTrack;
 
 /**
  * Interface that should be implemented for the UI portion of a key area within a section
  */
-class IKeyArea : public TSharedFromThis<IKeyArea>
+class SEQUENCER_API IKeyArea : public TSharedFromThis<IKeyArea>
 {
 public:
 
-	/** Virtual destructor. */
-	virtual ~IKeyArea() { }
+	/**
+	 * Constructor
+	 *
+	 * @param InChannelTypeID         The type ID of the channel (see FMovieSceneChannelEntry::GetChannelID)
+	 * @param InSection               The section that owns the channel that this key area represents
+	 * @param InChannel               Handle to the channel this key area represents
+	 * @param InCommonEditorData      The common editor data associated with this channel
+	 * @param InSpecializedData       Optional handle to the specialized data for this channel
+	 */
+	IKeyArea(
+		uint32                               InChannelTypeID,
+		UMovieSceneSection&                  InSection,
+		TMovieSceneChannelHandle<void>       InChannel,
+		const FMovieSceneChannelEditorData&  InCommonEditorData,
+		TMovieSceneChannelHandle<const void> InSpecializedData
+		);
 
-	/** @return The array of unsorted key handles in the key area */
-	virtual TArray<FKeyHandle> GetUnsortedKeyHandles() const = 0;
+public:
 
 	/**
-	 * Sets the time of a key given a handle
+	 * Locate the sequencer channel interface for this key area's channel
+	 * @note: Channel interfaces are registered via ISequencerModule::RegisterChannelInterface
 	 *
-	 * @param KeyHandle Handle of the key
-	 * @param The new time of the key
+	 * @return The channel interface, or nullptr if one was not found
 	 */
-	virtual void SetKeyTime(FKeyHandle KeyHandle, float NewKeyTime) = 0;
+	ISequencerChannelInterface* FindChannelInterface() const;
 
 	/**
-	 * Gets the time of a key given a handle
+	 * Access the channel type identifier for the channel that this key area wraps
 	 *
-	 * @param KeyHandle Handle of the key
-	 * @return The time of the key
+	 * @return The unique runtime type identifier for the type of this key area's channel
 	 */
-	virtual float GetKeyTime(FKeyHandle KeyHandle) const = 0;
+	uint32 GetChannelTypeID() const
+	{
+		return ChannelTypeID;
+	}
 
 	/**
-	 * Dilates the key by a specific factor
+	 * Get this key area's name
 	 *
-	 * @param Scale The multiplier which scales this key
-	 * @param Origin Time to scale from
-	 * @return The handle of the key
+	 * @return This key area's name
 	 */
-	virtual FKeyHandle DilateKey(FKeyHandle KeyHandle, float Scale, float Origin) = 0;
+	FName GetName() const;
 
 	/**
-	 * Moves a key
-	 * 
-	 * @param KeyHandle		Handle of the key to move
-	 * @param DeltaPosition	The delta position of the key
+	 * Set this key area's name
+	 *
+	 * @param InName This key area's name
 	 */
-	virtual FKeyHandle MoveKey(FKeyHandle KeyHandle, float DeltaPosition) = 0;
-		
+	void SetName(FName InName);
+
 	/**
-	 * Deletes a key
-	 * 
-	 * @param KeyHandle The key to delete
+	 * Get the color of this channel that should be drawn underneath its keys
+	 *
+	 * @return (Optional) This channel's color
 	 */
-	virtual void DeleteKey(FKeyHandle KeyHandle) = 0;
+	TOptional<FLinearColor> GetColor() const
+	{
+		return Color;
+	}
 
 	/**
-	 * Set key interpolation
+	 * Access the raw channel pointer that this key area represents
 	 *
-	 * @param KeyHandle The key handle
-	 * @param InterpMode The interpolation mode
+	 * @return The raw channel pointer, or nullptr if the channel proxy is no longer valid
 	 */
-	virtual void SetKeyInterpMode(FKeyHandle KeyHandle, ERichCurveInterpMode InterpMode) = 0;
+	void* GetChannelPtr() const;
 
 	/**
-	 * Get key interpolation
+	 * Get the raw specialized editor data for the channel this represents
 	 *
-	 * @param KeyHandle The key handle
-	 * @return InterpMode The interpolation mode
+	 * @return A pointer to the specialized data, or nullptr if it's invalid or expired
 	 */
-	virtual ERichCurveInterpMode GetKeyInterpMode(FKeyHandle Keyhandle) const = 0;
+	const void* GetSpecializedEditorData() const;
 
 	/**
-	 * Set key tangent
+	 * Access section that owns the channel this key area represents
 	 *
-	 * @param KeyHandle The key handle
-	 * @param TangentMode The tangent mode
+	 * @return The owning section, or nullptr if it has been destroyed
 	 */
-	virtual void SetKeyTangentMode(FKeyHandle KeyHandle, ERichCurveTangentMode TangentMode) = 0;
+	UMovieSceneSection* GetOwningSection() const;
+
+public:
 
 	/**
-	 * Get key tangent
-	 *
-	 * @param KeyHandle The key handle
-	 * @return TangentMode The tangent mode
+	 * Check whether this key area has any keys or not
 	 */
-	virtual ERichCurveTangentMode GetKeyTangentMode(FKeyHandle KeyHandle) const = 0;
+	bool HasAnyKeys() const;
 
 	/**
-	 * Set extrapolation
+	 * Add a key at the specified time with the current value of the channel, updating an existing key if possible
 	 *
-	 * @param ExtrapMode The extrapolation mode
-	 * @param bPreInfinity Set pre-infinity (or post-infinity)
+	 * @param Time            The time at which a key should be added (or updated)
+	 * @param ObjectBindingID The object binding ID this key area's track is bound to 
+	 * @param InSequencer     The currently active sequencer
+	 * @return A handle to the key that was added or updated
 	 */
-	virtual void SetExtrapolationMode(ERichCurveExtrapolation ExtrapMode, bool bPreInfinity) = 0;
+	FKeyHandle AddOrUpdateKey(FFrameNumber Time, const FGuid& ObjectBindingID, ISequencer& InSequencer);
 
 	/**
-	 * Get extrapolation
+	 * Duplicate the key represented by the specified handle
 	 *
-	 * @param bPreInfinity Get pre-infinity (or post-infinity)
-	 * @return ExtrapMode The extrapolation mode
+	 * @param InKeyHandle     A handle to the key to duplicate
+	 * @return A handle to the new duplicated key
 	 */
-	virtual ERichCurveExtrapolation GetExtrapolationMode(bool bPreInfinity) const = 0;
+	FKeyHandle DuplicateKey(FKeyHandle InKeyHandle) const;
 
 	/**
-	 * @return Whether this key area can set extrapolation mode
+	 * Get the time of the key represented by the specified handle
+	 *
+	 * @param InKeyHandle     A handle to the key to query for
+	 * @return The key's time, or TNumericLimits<int32>::Lowest() if the handle was not valid
 	 */
-	virtual bool CanSetExtrapolationMode() const { return false; }
+	FORCEINLINE FFrameNumber GetKeyTime(FKeyHandle InKeyHandle) const
+	{
+		FFrameNumber Time = TNumericLimits<int32>::Lowest();
+		GetKeyTimes(TArrayView<const FKeyHandle>(&InKeyHandle, 1), TArrayView<FFrameNumber>(&Time, 1));
+		return Time;
+	}
 
 	/**
-	 * Adds a key at the specified time if there isn't already a key present.  The value of the added key should
-	 * be the value which would be returned if the animation containing this key area was evaluated at the specified time.
+	 * Get the times of every key represented by the specified handles
 	 *
-	 * @param InKeyInterpolation KeyInterpolation
-	 * @param TimeToCopyFrom Optional time to copy key parameters from
-	 * @param return      The new keys that were added
+	 * @param InKeyHandles     A handle to the key to query for
+	 * @param OutTimes         A pre-sized array view to populate with key times
 	 */
-	virtual TArray<FKeyHandle> AddKeyUnique(float Time, EMovieSceneKeyInterpolation InKeyInterpolation, float TimeToCopyFrom = FLT_MAX) = 0;
+	void GetKeyTimes(TArrayView<const FKeyHandle> InKeyHandles, TArrayView<FFrameNumber> OutTimes) const;
 
 	/**
-	 * Duplicate the specified key
+	 * Get all key handles that exist within the given time range
+	 *
+	 * @param OutHandles       An array to populate with key handles, any key handles will be appended to this array
+	 * @param WithinRange      (Optional) A predicate range to search within
 	 */
-	virtual TOptional<FKeyHandle> DuplicateKey(FKeyHandle KeyToDuplicate) = 0;
+	FORCEINLINE void GetKeyHandles(TArray<FKeyHandle>& OutHandles, const TRange<FFrameNumber>& WithinRange = TRange<FFrameNumber>::All()) const
+	{
+		GetKeyInfo(&OutHandles, nullptr, WithinRange);
+	}
 
-	/** Gets the rich curve associated with this key area.  This can be null but it must be valid in order to be
-	  * edited by the curve editor. */
-	virtual FRichCurve* GetRichCurve() = 0;
+	/**
+	 * Get all key times that exist within the given time range
+	 *
+	 * @param OutTimes         An array to populate with key times, any key times will be appended to this array
+	 * @param WithinRange      (Optional) A predicate range to search within
+	 */
+	FORCEINLINE void GetKeyTimes(TArray<FFrameNumber>& OutTimes, const TRange<FFrameNumber>& WithinRange = TRange<FFrameNumber>::All()) const
+	{
+		GetKeyInfo(nullptr, &OutTimes, WithinRange);
+	}
 
 	/**
-	 * Gets the section which owns this key area.
+	 * Populate the specified handle and/or time arrays with information pertaining to keys that exist within the given range
+	 *
+	 * @param OutHandles       (Optional) An array to populate with key handles, any key handles will be appended to this array
+	 * @param OutTimes         (Optional) An array to populate with key times, any key times will be appended to this array
+	 * @param WithinRange      (Optional) A predicate range to search within
 	 */
-	virtual UMovieSceneSection* GetOwningSection() = 0;
+	void GetKeyInfo(TArray<FKeyHandle>* OutHandles, TArray<FFrameNumber>* OutTimes, const TRange<FFrameNumber>& WithinRange = TRange<FFrameNumber>::All()) const;
 
 	/**
-	 * Returns true if this key area can create a key editor widget for the animation outliner.
+	 * Set the time of the key with the specified handle
+	 *
+	 * @param InKeyHandle      The key handle to set
+	 * @param InKeyTime        The time to set the key to
 	 */
-	virtual bool CanCreateKeyEditor() = 0;
+	void SetKeyTime(FKeyHandle InKeyHandle, FFrameNumber InKeyTime) const
+	{
+		SetKeyTimes(TArrayView<const FKeyHandle>(&InKeyHandle, 1), TArrayView<const FFrameNumber>(&InKeyTime, 1));
+	}
 
 	/**
-	 * Creates a key editor for this key area for use in the animation outliner.
+	 * Set the times of the each key with the specified handles
+	 *
+	 * @param InKeyHandles     An array of handles that should have their time set to times in the corresponding InKeyTimes array
+	 * @param InKeyTimes       Array of times to set to, one per key handle. Must match the size of InKeyHandles
 	 */
-	virtual TSharedRef<SWidget> CreateKeyEditor(ISequencer* Sequencer) = 0;
+	void SetKeyTimes(TArrayView<const FKeyHandle> InKeyHandles, TArrayView<const FFrameNumber> InKeyTimes) const;
+
+public:
 
 	/**
-	 * Copy keys from this key area
+	 * Gather key drawing information for the specified key handles
 	 *
-	 * @param ClipboardBuilder	Clipboard builder to add keys to
-	 * @param KeyMask			Predicate to be called to check whether a key should be copied
+	 * @param InKeyHandles     An array of handles to draw
+	 * @param OutKeyDrawParams Pre-sized array view of parameters to populate with each key
 	 */
-	virtual void CopyKeys(FMovieSceneClipboardBuilder& ClipboardBuilder, const TFunctionRef<bool(FKeyHandle, const IKeyArea&)>& KeyMask) const = 0;
+	void DrawKeys(TArrayView<const FKeyHandle> InKeyHandles, TArrayView<FKeyDrawParams> OutKeyDrawParams);
 
 	/**
-	 * Paste keys into this key area
+	 * Copy all the keys specified in KeyMask to the specified clipboard
 	 *
-	 * @param KeyTrack			Container of keys to paste into this area
-	 * @param SrcEnvironment		The source environment the cliboard originated from
-	 * @param DstEnvironment		The new environment to paste in
+	 * @param ClipboardBuilder The structure responsible for building clipboard information for each key
+	 * @param KeyMask          A specific set of keys to copy
 	 */
-	virtual void PasteKeys(const FMovieSceneClipboardKeyTrack& KeyTrack, const FMovieSceneClipboardEnvironment& SrcEnvironment, const FSequencerPasteEnvironment& DstEnvironment) = 0;
+	void CopyKeys(FMovieSceneClipboardBuilder& ClipboardBuilder, TArrayView<const FKeyHandle> KeyMask) const;
+
 
 	/**
-	 * Give this key area a specific name
+	 * Paste the specified key track into this key area
+	 *
+	 * @param KeyTrack         The source clipboard data to paste
+	 * @param SrcEnvironment   The environment the source data was copied from
+	 * @param DstEnvironment   The environment we're pasting into
 	 */
-	virtual void SetName(FName Name) = 0;
+	void PasteKeys(const FMovieSceneClipboardKeyTrack& KeyTrack, const FMovieSceneClipboardEnvironment& SrcEnvironment, const FSequencerPasteEnvironment& DstEnvironment);
+
 
 	/**
-	 * Get this area's name, or None if not set
+	 * Create a new model for this key area that can be used on the curve editor interface
+	 *
+	 * @return (Optional) A new model to be added to a curve editor
 	 */
-	virtual FName GetName() const = 0;
+	TUniquePtr<FCurveModel> CreateCurveEditorModel(TSharedRef<ISequencer> InSequencer) const;
+
+public:
+
+	FRichCurve* GetRichCurve() const { return nullptr; }
 
 	/**
-	 * Get the key area's color.
+	 * Get a key structure for editing a value on this channel
 	 *
-	 * @return Key area color.
+	 * @return The key structure for this channel, or nullptr
 	 */
-	virtual TOptional<FLinearColor> GetColor() = 0;
+	TSharedPtr<FStructOnScope> GetKeyStruct(FKeyHandle KeyHandle) const;
 
 	/**
-	 * Get the data structure representing the specified key.
+	 * Check whether this key area can create an editor on the sequencer node tree
+	 */
+	bool CanCreateKeyEditor() const;
+
+	/**
+	 * Create an editor on the sequencer node tree
 	 *
-	 * @param KeyHandle The handle of the key.
-	 * @return The key's data structure representation, or nullptr if key not found or no structure available.
+	 * @param Sequencer        The currently active sequencer
+	 * @param ObjectBindingID  The ID of the object this key area's track is bound to
+	 * @return The editor widget to display on the key area's node
 	 */
-	virtual TSharedPtr<FStructOnScope> GetKeyStruct(FKeyHandle KeyHandle) = 0;
-};
+	TSharedRef<SWidget> CreateKeyEditor(TWeakPtr<ISequencer> Sequencer, const FGuid& ObjectBindingID);
+
+private:
+
+	/** The type of the channel this key area represents (see FMovieSceneChannelEntry::GetChannelID) */
+	uint32 ChannelTypeID;
+
+	/** A weak pointer back to the originating UMovieSceneSection that owns this channel */
+	TWeakObjectPtr<UMovieSceneSection> WeakOwningSection;
+
+	/** Handle to the channel itself */
+	TMovieSceneChannelHandle<void> Channel;
+
+	/** Handle to the channel's specialized data as defined by TMovieSceneChannelTraits::EditorDataType */
+	TMovieSceneChannelHandle<const void> SpecializedData;
+
+	/** Optional property bindings class where the section resides inside a UMovieScenePropertyTrack */
+	TOptional<FTrackInstancePropertyBindings> PropertyBindings;
+
+	/** The color of this channel that should be drawn underneath its keys */
+	TOptional<FLinearColor> Color;
+
+	/** The name of this channel. */
+	FName ChannelName;
+
+	/** The display text of this channel. */
+	FText DisplayText;
+};
\ No newline at end of file
diff --git a/Engine/Source/Editor/Sequencer/Public/ISectionLayoutBuilder.h b/Engine/Source/Editor/Sequencer/Public/ISectionLayoutBuilder.h
--- a/Engine/Source/Editor/Sequencer/Public/ISectionLayoutBuilder.h
+++ b/Engine/Source/Editor/Sequencer/Public/ISectionLayoutBuilder.h
@@ -4,6 +4,8 @@
 
 #include "CoreMinimal.h"
 
+class IKeyArea;
+
 /**
  * Builds an inner layout for a section
  */
@@ -24,7 +26,7 @@ public:
 	 * Sets the section as a key area itself
 	 * @param KeyArea		Interface for accessing and drawing keys
 	 */
-	virtual void SetSectionAsKeyArea( TSharedRef<class IKeyArea> KeyArea ) = 0;
+	virtual void SetSectionAsKeyArea( TSharedRef<IKeyArea> KeyArea ) = 0;
 
 	/**
 	 * Adds a key area onto the layout.  If a category is pushed, the key area will appear as a child of the current category
@@ -33,7 +35,7 @@ public:
 	 * @param DisplayLabel	The localized display label for the key area
 	 * @param KeyArea		Interface for accessing and drawing keys
 	 */
-	virtual void AddKeyArea( FName KeyAreaName, const FText& DisplayLabel, TSharedRef<class IKeyArea> KeyArea ) = 0;
+	virtual void AddKeyArea( FName KeyAreaName, const FText& DisplayLabel, TSharedRef<IKeyArea> KeyArea ) = 0;
 
 	/**
 	 * Pops a category off the stack
