Copying //UE4/Dev-Sequencer to //UE4/Dev-Main (Source: //UE4/Dev-Sequencer @ 3946692)

diff --git a/Engine/Source/Editor/Sequencer/Public/ISequencer.h b/Engine/Source/Editor/Sequencer/Public/ISequencer.h
--- a/Engine/Source/Editor/Sequencer/Public/ISequencer.h
+++ b/Engine/Source/Editor/Sequencer/Public/ISequencer.h
@@ -12,15 +12,19 @@
 #include "IMovieScenePlayer.h"
 #include "KeyPropertyParams.h"
 #include "MovieSceneBinding.h"
+#include "QualifiedFrameTime.h"
 #include "Widgets/Input/NumericTypeInterface.h"
 #include "Editor/SequencerWidgets/Public/ITimeSlider.h"
 
+struct FFrameTime;
+struct FQualifiedFrameTime;
+
 class UMovieSceneTrack;
 class AActor;
 class FSequencerSelection;
 class FSequencerSelectionPreview;
 class FUICommandList;
-class ISequencerKeyCollection;
+class FSequencerKeyCollection;
 class UMovieSceneSequence;
 class UMovieSceneSubSection;
 class IDetailsView;
@@ -163,6 +167,9 @@ public:
 	virtual FMovieSceneSequenceIDRef GetRootTemplateID() const = 0;
 	virtual FMovieSceneSequenceIDRef GetFocusedTemplateID() const = 0;
 
+	/** @return If the currently focused sequence is active in the sequence hierarchy. */
+	virtual bool GetFocusedSequenceIsActive() const = 0;
+
 	TArrayView<TWeakObjectPtr<>> FindObjectsInCurrentSequence(const FGuid& InObjectBinding)
 	{
 		return FindBoundObjects(InObjectBinding, GetFocusedTemplateID());
@@ -200,6 +207,12 @@ public:
 	 */
 	virtual TArray<FGuid> AddActors(const TArray<TWeakObjectPtr<AActor> >& InActors) = 0;
 
+	/**
+	 * Calling this function will add the specified track to the currently selected folder
+	 * if there is one, and will set the newly created track as the current selection.
+	 */
+	virtual void OnAddTrack(const TWeakObjectPtr<UMovieSceneTrack>& InTrack) = 0;
+
 	/**
 	 * Adds a movie scene as a section inside the current movie scene
 	 * 
@@ -246,16 +259,13 @@ public:
 	/** Gets whether or not property track defaults will be automatically set when adding tracks. */
 	virtual bool GetAutoSetTrackDefaults() const = 0;
 
-	/** @return Returns whether sequencer is currently recording live data from simulated actors */
-	virtual bool IsRecordingLive() const = 0;
-
 	/** @return Returns whether sequencer will respond to changes and possibly create a key or track */
-	virtual bool IsAllowedToChange() const { return IsRecordingLive() || GetAllowEditsMode() != EAllowEditsMode::AllowLevelEditsOnly || GetAutoChangeMode() != EAutoChangeMode::None; }
+	virtual bool IsAllowedToChange() const { return GetAllowEditsMode() != EAllowEditsMode::AllowLevelEditsOnly || GetAutoChangeMode() != EAutoChangeMode::None; }
 
 	/**
 	 * Gets the current time of the time slider relative to the currently focused movie scene
 	 */
-	virtual float GetLocalTime() const = 0;
+	virtual FQualifiedFrameTime GetLocalTime() const = 0;
 
 	/**
 	 * Gets the global time.
@@ -263,7 +273,7 @@ public:
 	 * @return Global time in seconds.
 	 * @see SetGlobalTime
 	 */
-	virtual float GetGlobalTime() const = 0;
+	virtual FQualifiedFrameTime GetGlobalTime() const = 0;
 
 	/**
 	 * Sets the cursor position relative to the currently focused sequence
@@ -271,13 +281,13 @@ public:
 	 * @param Time The local time to set.
 	 * @param SnapTimeMode The type of time snapping allowed.
 	 */
-	virtual void SetLocalTime(float Time, ESnapTimeMode SnapTimeMode = ESnapTimeMode::STM_None) = 0;
+	virtual void SetLocalTime(FFrameTime Time, ESnapTimeMode SnapTimeMode = ESnapTimeMode::STM_None) = 0;
 
 	/** Set the current local time directly, with no other snapping, scrolling or manipulation */
-	virtual void SetLocalTimeDirectly(float NewTime) = 0;
+	virtual void SetLocalTimeDirectly(FFrameTime NewTime) = 0;
 
 	/** Set the global time directly, without performing any auto-scroll, snapping or other adjustments to the supplied time  */
-	virtual void SetGlobalTime(float Time) = 0;
+	virtual void SetGlobalTime(FFrameTime Time) = 0;
 
 	/** Forcefully reevaluate the sequence */
 	virtual void ForceEvaluate() = 0;
@@ -293,7 +303,7 @@ public:
 	 * @param NewViewRange The new view range. Must be a finite range
 	 * @param Interpolation How to interpolate to the new view range
 	 */
-	virtual void SetViewRange(TRange<float> NewViewRange, EViewRangeInterpolation Interpolation = EViewRangeInterpolation::Animated) = 0;
+	virtual void SetViewRange(TRange<double> NewViewRange, EViewRangeInterpolation Interpolation = EViewRangeInterpolation::Animated) = 0;
 
 	/**
 	 * Sets whether perspective viewport hijacking is enabled.
@@ -384,9 +394,6 @@ public:
 
 	virtual void KeyProperty(FKeyPropertyParams KeyPropertyParams) = 0;
 
-	DEPRECATED( 4.13, "NotifyMovieSceneDataChanged() is deprecated, use the version that takes EMovieSceneDataChangeType" )
-	void NotifyMovieSceneDataChanged() { NotifyMovieSceneDataChangedInternal(); };
-
 	/** Refresh the sequencer tree view */
 	virtual void RefreshTree() = 0;
 
@@ -402,7 +409,7 @@ public:
 	virtual float GetPlaybackSpeed() const = 0;
 
 	/** Get all the keys for the current sequencer selection */
-	virtual void GetKeysFromSelection(TUniquePtr<ISequencerKeyCollection>& KeyCollection, float DuplicateThresoldTime) = 0;
+	virtual void GetKeysFromSelection(TUniquePtr<FSequencerKeyCollection>& KeyCollection, float DuplicateThresoldTime) = 0;
 
 	virtual FSequencerSelection& GetSelection() = 0;
 	virtual FSequencerSelectionPreview& GetSelectionPreview() = 0;
@@ -428,6 +435,10 @@ public:
 	/** Empties the current selection. */
 	virtual void EmptySelection() = 0;
 
+	/** Throb key or section selection */
+	virtual void ThrobKeySelection() = 0;
+	virtual void ThrobSectionSelection() = 0;
+
 	/** Gets a multicast delegate which is executed whenever the global time changes. */
 	virtual FOnGlobalTimeChanged& OnGlobalTimeChanged() = 0;
 
@@ -462,10 +473,7 @@ public:
 	virtual FOnSelectionChangedSections& GetSelectionChangedSections() = 0;
 
 	/** @return a numeric type interface that will parse and display numbers as frames and times correctly */
-	virtual TSharedRef<INumericTypeInterface<float>> GetNumericTypeInterface() = 0;
-
-	/** @return a numeric type interface that will parse and display numbers as frames and times correctly (including zero-pad, if applicable) */
-	virtual TSharedRef<INumericTypeInterface<float>> GetZeroPadNumericTypeInterface() = 0;
+	virtual TSharedRef<INumericTypeInterface<double>> GetNumericTypeInterface() const = 0;
 
 	/** @return the command bindings for this sequencer */
 	virtual TSharedPtr<FUICommandList> GetCommandBindings(ESequencerCommandBindings Type = ESequencerCommandBindings::Sequencer) const = 0;
@@ -506,6 +514,28 @@ public:
 	 */
 	virtual TSharedPtr<class ITimeSlider> GetTopTimeSliderWidget() const = 0;
 
+public:
+
+	/**
+	 * Get the frame resolution of the currently root sequence
+	 */
+	SEQUENCER_API FFrameRate GetRootFrameResolution() const;
+
+	/**
+	 * Get the playback frame rate of the currently root sequence
+	 */
+	SEQUENCER_API FFrameRate GetRootPlayRate() const;
+
+	/**
+	 * Get the frame resolution of the currently focused sequence
+	 */
+	SEQUENCER_API FFrameRate GetFocusedFrameResolution() const;
+
+	/**
+	 * Get the playback frame rate of the currently focused sequence
+	 */
+	SEQUENCER_API FFrameRate GetFocusedPlayRate() const;
+
 protected:
 	FOnInitializeDetailsPanel InitializeDetailsPanelEvent;
 	FOnGetIsBindingVisible GetIsBindingVisible;
diff --git a/Engine/Source/Editor/Sequencer/Public/ISequencerChannelInterface.h b/Engine/Source/Editor/Sequencer/Public/ISequencerChannelInterface.h
new file mode 100644
--- /dev/null
+++ b/Engine/Source/Editor/Sequencer/Public/ISequencerChannelInterface.h
@@ -0,0 +1,199 @@
+// Copyright 1998-2018 Epic Games, Inc. All Rights Reserved.
+
+#pragma once
+
+#include "Containers/ArrayView.h"
+#include "Containers/Array.h"
+#include "Templates/SharedPointer.h"
+#include "UObject/WeakObjectPtr.h"
+#include "Channels/MovieSceneChannelHandle.h"
+#include "FrameNumber.h"
+
+struct FKeyHandle;
+struct FKeyDrawParams;
+struct FMovieSceneChannelEntry;
+struct FSequencerPasteEnvironment;
+struct FMovieSceneClipboardEnvironment;
+
+class SWidget;
+class ISequencer;
+class FCurveModel;
+class FMenuBuilder;
+class FStructOnScope;
+class UMovieSceneSection;
+class ISectionLayoutBuilder;
+class FMovieSceneClipboardBuilder;
+class FMovieSceneClipboardKeyTrack;
+class FTrackInstancePropertyBindings;
+
+/** Utility struct representing a number of selected keys on a single channel */
+template<typename ChannelType>
+struct TChannelAndHandles
+{
+	TWeakObjectPtr<UMovieSceneSection> Section;
+	TMovieSceneChannelHandle<ChannelType> Channel;
+	TArray<FKeyHandle> Handles;
+};
+
+/**
+ * Abstract interface that defines all sequencer interactions with any channel type
+ * Channels are stored internally as void*, so a single interface should be registered for each channel type.
+ * Implementations are found in TSequencerChanelInterface which calls overloaded free functions for each channel.
+ */
+struct ISequencerChannelInterface
+{
+	virtual ~ISequencerChannelInterface() {}
+
+	/**
+	 * Add (or update) a key to the specified channel using it's current value at that time, or some external value specified by the specialized editor data
+	 *
+	 * @param Channel               The channel to add a key to
+	 * @param SpecializedEditorData A pointer to the specialized editor data for this channel of type TMovieSceneChannelTraits<>::EditorDataType
+	 * @param InTime                The time at which to add a key
+	 * @param InSequencer           The currently active sequencer
+	 * @param ObjectBindingID       The object binding ID for the track that this channel resides within
+	 * @param PropertyBindings      (Optional) Property bindings where this channel exists on a property track
+	 * @return A handle to the new or updated key
+	 */
+	virtual FKeyHandle AddOrUpdateKey_Raw(void* Channel, const void* SpecializedEditorData, FFrameNumber InTime, ISequencer& InSequencer, const FGuid& ObjectBindingID, FTrackInstancePropertyBindings* PropertyBindings) const = 0;
+
+	/**
+	 * Check whether any of the specified channels have any keys
+	 *
+	 * @return true if so, false otherwise
+	 */
+	virtual bool HasAnyKeys_Raw(TArrayView<void* const> Ptrs) const = 0;
+
+	/**
+	 * Get key information pertaining to all keys that exist within the specified range
+	 *
+	 * @param Channel               The channel to query
+	 * @param WithinRange           The range within which to return key information
+	 * @param OutKeyTimes           (Optional) Array to receive key times
+	 * @param OutKeyHandles         (Optional) Array to receive key handles
+	 */
+	virtual void GetKeys_Raw(void* InChannel, const TRange<FFrameNumber>& WithinRange, TArray<FFrameNumber>* OutKeyTimes, TArray<FKeyHandle>* OutKeyHandles) const = 0;
+
+	/**
+	 * Get all key times for the specified key handles
+	 *
+	 * @param Channel               The channel to query
+	 * @param InHandles             Array of handles to get times for
+	 * @param OutKeyTimes           Pre-sized array of key times to set. Invalid key handles will not assign to this array. Must match size of InHandles
+	 */
+	virtual void GetKeyTimes_Raw(void* InChannel, TArrayView<const FKeyHandle> InHandles, TArrayView<FFrameNumber> OutKeyTimes) const = 0;
+
+	/**
+	 * Set key times for the specified key handles
+	 *
+	 * @param Channel               The channel to query
+	 * @param InHandles             Array of handles to get times for
+	 * @param InKeyTimes            Array of times to apply - one per handle
+	 */
+	virtual void SetKeyTimes_Raw(void* InChannel, TArrayView<const FKeyHandle> InHandles, TArrayView<const FFrameNumber> InKeyTimes) const = 0;
+
+	/**
+	 * Duplicate the keys for the specified key handles
+	 *
+	 * @param Channel               The channel to query
+	 * @param InHandles             Array of handles to duplicate
+	 * @param OutKeyTimes           Pre-sized array to receive duplicated key handles. Invalid key handles will not be assigned to this array. Must match size of InHandles
+	 */
+	virtual void DuplicateKeys_Raw(void* InChannel, TArrayView<const FKeyHandle> InHandles, TArrayView<FKeyHandle> OutNewHandles) const = 0;
+
+	/**
+	 * Delete the keys for the specified key handles
+	 *
+	 * @param Channel               The channel to query
+	 * @param InHandles             Array of handles to delete
+	 */
+	virtual void DeleteKeys_Raw(void* InChannel, TArrayView<const FKeyHandle> InHandles) const = 0;
+
+	/**
+	 * Copy all the keys specified in KeyMask to the specified clipboard
+	 *
+	 * @param Channel               The channel to copy from
+	 * @param Section               The section that owns the channel
+	 * @param KeyAreaName           The name of the key area
+	 * @param ClipboardBuilder      The structure responsible for building clipboard information for each key
+	 * @param KeyMask               A specific set of keys to copy
+	 */
+	virtual void CopyKeys_Raw(void* Channel, const UMovieSceneSection* Section, FName KeyAreaName, FMovieSceneClipboardBuilder& ClipboardBuilder, TArrayView<const FKeyHandle> KeyMask) const = 0;
+
+	/**
+	 * Paste the specified key track into the specified channel
+	 *
+	 * @param Channel               The channel to copy from
+	 * @param Section               The section that owns the channel
+	 * @param KeyTrack              The source clipboard data to paste
+	 * @param SrcEnvironment        The environment the source data was copied from
+	 * @param DstEnvironment        The environment we're pasting into
+	 * @param OutPastedKeys         Array to receive key handles for any pasted keys
+	 */
+	virtual void PasteKeys_Raw(void* Channel, UMovieSceneSection* Section, const FMovieSceneClipboardKeyTrack& KeyTrack, const FMovieSceneClipboardEnvironment& SrcEnvironment, const FSequencerPasteEnvironment& DstEnvironment, TArray<FKeyHandle>& OutPastedKeys) const = 0;
+
+	/**
+	 * Get an editable key struct for the specified key
+	 *
+	 * @param Channel               The channel on which the key resides
+	 * @param KeyHandle             Handle of the key to get
+	 * @return A shared editable key struct
+	 */
+	virtual TSharedPtr<FStructOnScope> GetKeyStruct_Raw(TMovieSceneChannelHandle<void> Channel, FKeyHandle KeyHandle) const = 0;
+
+	/**
+	 * Check whether an editor on the sequencer node tree can be created for the specified channel
+	 *
+	 * @param Channel               The channel to check
+	 * @return true if a key editor should be constructed, false otherwise
+	 */
+	virtual bool CanCreateKeyEditor_Raw(void* Channel) const = 0;
+
+	/**
+	 * Create an editor on the sequencer node tree
+	 *
+	 * @param Channel               The channel to check
+	 * @param SpecializedEditorData A pointer to the specialized editor data for this channel of type TMovieSceneChannelTraits<>::EditorDataType
+	 * @param Section               The section that owns this channel
+	 * @param InObjectBindingID     The ID of the object this key area's track is bound to
+	 * @param PropertyBindings      (Optional) Property bindings where this channel exists on a property track
+	 * @param Sequencer             The currently active sequencer
+	 * @return The editor widget to display on the node tree
+	 */
+	virtual TSharedRef<SWidget> CreateKeyEditor_Raw(void* Channel, const void* SpecializedEditorData, UMovieSceneSection* Section, const FGuid& InObjectBindingID, TWeakPtr<FTrackInstancePropertyBindings> PropertyBindings, TWeakPtr<ISequencer> Sequencer) const = 0;
+
+	/**
+	 * Extend the key context menu
+	 *
+	 * @param MenuBuilder           The menu builder used to create this context menu
+	 * @param Channels              Array of channels and handles that are being shown in the context menu
+	 * @param InSequencer           The currently active sequencer
+	 */
+	virtual void ExtendKeyMenu_Raw(FMenuBuilder& MenuBuilder, TArrayView<const TChannelAndHandles<void>> Channels, TWeakPtr<ISequencer> InSequencer) const = 0;
+
+	/**
+	 * Extend the section context menu
+	 *
+	 * @param MenuBuilder           The menu builder used to create this context menu
+	 * @param Channels              Array of type specific channels that exist in the selected sections
+	 * @param Sections              Array of sections being shown on the context menu
+	 * @param InSequencer           The currently active sequencer
+	 */
+	virtual void ExtendSectionMenu_Raw(FMenuBuilder& MenuBuilder, TArrayView<TMovieSceneChannelHandle<void> const> Channels, TArrayView<UMovieSceneSection* const> Sections, TWeakPtr<ISequencer> InSequencer) const = 0;
+
+	/**
+	 * Gather information on how to draw the specified keys
+	 *
+	 * @param Channel               The channel to query
+	 * @param InKeyHandles          Array of handles to duplicate
+	 * @param OutKeyDrawParams      Pre-sized array to receive key draw parameters. Invalid key handles will not be assigned to this array. Must match size of InKeyHandles.
+	 */
+	virtual void DrawKeys_Raw(void* Channel, TArrayView<const FKeyHandle> InKeyHandles, TArrayView<FKeyDrawParams> OutKeyDrawParams) const = 0;
+
+	/**
+	 * Create a new model for this channel that can be used on the curve editor interface
+	 *
+	 * @return (Optional) A new model to be added to a curve editor
+	 */
+	virtual TUniquePtr<FCurveModel> CreateCurveEditorModel_Raw(TMovieSceneChannelHandle<void> Channel, UMovieSceneSection* OwningSection, TSharedRef<ISequencer> InSequencer) const = 0;
+};
\ No newline at end of file
diff --git a/Engine/Source/Editor/Sequencer/Public/ISequencerHotspot.h b/Engine/Source/Editor/Sequencer/Public/ISequencerHotspot.h
--- a/Engine/Source/Editor/Sequencer/Public/ISequencerHotspot.h
+++ b/Engine/Source/Editor/Sequencer/Public/ISequencerHotspot.h
@@ -4,6 +4,7 @@
 
 #include "CoreMinimal.h"
 #include "Input/CursorReply.h"
+#include "FrameTime.h"
 
 class FMenuBuilder;
 class ISequencer;
@@ -30,10 +31,10 @@ struct ISequencerHotspot
 	virtual ~ISequencerHotspot() { }
 	virtual ESequencerHotspot GetType() const = 0;
 	virtual void UpdateOnHover(SSequencerTrackArea& InTrackArea, ISequencer& InSequencer) const = 0;
-	virtual TOptional<float> GetTime() const { return TOptional<float>(); }
-	virtual TOptional<float> GetOffsetTime() const { return TOptional<float>(); }
+	virtual TOptional<FFrameNumber> GetTime() const { return TOptional<FFrameNumber>(); }
+	virtual TOptional<FFrameTime> GetOffsetTime() const { return TOptional<FFrameTime>(); }
 	virtual TSharedPtr<ISequencerEditToolDragOperation> InitiateDrag(ISequencer&) { return nullptr; }
-	virtual bool PopulateContextMenu(FMenuBuilder& MenuBuilder, ISequencer& Sequencer, float MouseDownTime){ return false; }
+	virtual bool PopulateContextMenu(FMenuBuilder& MenuBuilder, ISequencer& Sequencer, FFrameTime MouseDownTime){ return false; }
 	virtual FCursorReply GetCursor() const { return FCursorReply::Unhandled(); }
 
 public:
diff --git a/Engine/Source/Editor/Sequencer/Public/ISequencerModule.h b/Engine/Source/Editor/Sequencer/Public/ISequencerModule.h
--- a/Engine/Source/Editor/Sequencer/Public/ISequencerModule.h
+++ b/Engine/Source/Editor/Sequencer/Public/ISequencerModule.h
@@ -7,8 +7,11 @@
 #include "ISequencer.h"
 #include "Modules/ModuleInterface.h"
 #include "AnimatedPropertyKey.h"
+#include "ISequencerChannelInterface.h"
 
+class IKeyArea;
 class FExtender;
+class FStructOnScope;
 class FExtensibilityManager;
 class FMenuBuilder;
 class ISequencerTrackEditor;
@@ -31,6 +34,9 @@ DECLARE_DELEGATE_RetVal_OneParam(TSharedRef<ISequencerEditorObjectBinding>, FOnC
 /** A delegate that is executed when adding menu content. */
 DECLARE_DELEGATE_TwoParams(FOnGetAddMenuContent, FMenuBuilder& /*MenuBuilder*/, TSharedRef<ISequencer>);
 
+/** A delegate that is executed when menu object is clicked. Unlike FExtender delegates we pass in the FGuid which exists even for deleted objects. */
+DECLARE_DELEGATE_TwoParams(FOnBuildCustomContextMenuForGuid, FMenuBuilder&, FGuid);
+
 /** A delegate that gets executed then a sequencer is created */
 DECLARE_MULTICAST_DELEGATE_OneParam(FOnSequencerCreated, TSharedRef<ISequencer>);
 
@@ -39,11 +45,10 @@ DECLARE_MULTICAST_DELEGATE_OneParam(FOnSequencerCreated, TSharedRef<ISequencer>)
  */
 struct FSequencerViewParams
 {
-	/** Initial Scrub Position. */
-	float InitialScrubPosition;
-
 	FOnGetAddMenuContent OnGetAddMenuContent;
 
+	FOnBuildCustomContextMenuForGuid OnBuildCustomContextMenuForGuid;
+
 	/** Called when this sequencer has received user focus */
 	FSimpleDelegate OnReceivedFocus;
 
@@ -60,8 +65,7 @@ struct FSequencerViewParams
 	bool bReadOnly;
 
 	FSequencerViewParams(FString InName = FString())
-		: InitialScrubPosition(0.0f)
-		, UniqueName(MoveTemp(InName))
+		: UniqueName(MoveTemp(InName))
 		, bReadOnly(false)
 	{ }
 };
@@ -101,7 +105,6 @@ struct FSequencerInitParams
 	{}
 };
 
-
 /**
  * Interface for the Sequencer module.
  */
@@ -110,6 +113,8 @@ class ISequencerModule
 {
 public:
 
+	virtual ~ISequencerModule();
+
 	/**
 	 * Create a new instance of a standalone sequencer that can be added to other UIs.
 	 *
@@ -199,6 +204,22 @@ public:
 	 */
 	virtual TSharedPtr<FExtensibilityManager> GetToolBarExtensibilityManager() const = 0;
 
+	/**
+	 * Register a sequencer channel type. Include SequencerChannelInterface.h for definition.
+	 */
+	template<typename ChannelType>
+	void RegisterChannelInterface();
+
+	/**
+	 * Find a sequencer channel for the specified channel type ID
+	 */
+	ISequencerChannelInterface* FindChannelInterface(uint32 ChannelID) const
+	{
+		const TUniquePtr<ISequencerChannelInterface>* Found = ChannelToEditorInterfaceMap.Find(ChannelID);
+		ensureMsgf(Found, TEXT("No channel interface found for type ID. Did you call RegisterChannelInterface<> for that type?"));
+		return Found ? Found->Get() : nullptr;
+	}
+
 public:
 
 	/** 
@@ -227,4 +248,9 @@ public:
 	{
 		UnRegisterTrackEditor(InHandle);
 	}
+
+private:
+
+	/** */
+	TMap<uint32, TUniquePtr<ISequencerChannelInterface>> ChannelToEditorInterfaceMap;
 };
diff --git a/Engine/Source/Editor/Sequencer/Public/ISequencerSection.h b/Engine/Source/Editor/Sequencer/Public/ISequencerSection.h
--- a/Engine/Source/Editor/Sequencer/Public/ISequencerSection.h
+++ b/Engine/Source/Editor/Sequencer/Public/ISequencerSection.h
@@ -52,14 +52,6 @@ public:
 	 */
 	virtual UMovieSceneSection* GetSectionObject() = 0;
 
-	/**
-	 * Allows each section to have it's own unique widget for advanced editing functionality
-	 * OnPaintSection will still be called if a widget is provided.  OnPaintSection is still used for the background section display
-	 * 
-	 * @return The generated widget 
-	 */
-	virtual TSharedRef<SWidget> GenerateSectionWidget() { return SNullWidget::NullWidget; }
-
 	/**
 	 * Called when the section should be painted
 	 *
@@ -68,19 +60,18 @@ public:
 	 */
 	virtual int32 OnPaintSection( FSequencerSectionPainter& InPainter ) const = 0;
 
-	/** Allows a section to override the brush to use for a key by handle.
-	 *
-	 * @param KeyHandle the handle of the key to get a brush for.
-	 * @return A const pointer to a slate brush if the brush should be overridden, otherwise null.
+	/**
+	 * Allows each section to have it's own unique widget for advanced editing functionality
+	 * OnPaintSection will still be called if a widget is provided.  OnPaintSection is still used for the background section display
+	 * 
+	 * @return The generated widget 
 	 */
+	virtual TSharedRef<SWidget> GenerateSectionWidget() { return SNullWidget::NullWidget; }
+
+	DEPRECATED(4.20, "Please override Sequencer::DrawKeys instead")
 	virtual const FSlateBrush* GetKeyBrush(FKeyHandle KeyHandle) const { return nullptr; }
 
-	/** When a section overrides the brush to use, this allows it to set the scale origin tfor use
-	 * when the brush is scaled to create border and selection effects
-	 *
-	 * @param KeyHandle the handle of the key to get a brush origin for.
-	 * @return A FVector2D describing the custom origin, in slate units
-	 */
+	DEPRECATED(4.20, "Please override Sequencer::DrawKeys instead")
 	virtual FVector2D GetKeyBrushOrigin( FKeyHandle KeyHandle ) const { return FVector2D(0.0f, 0.0f); }
 
 	/**
@@ -102,9 +93,6 @@ public:
 	 */
 	virtual FReply OnSectionDoubleClicked( const FGeometry& SectionGeometry, const FPointerEvent& MouseEvent, const FGuid& ObjectBinding) { return FReply::Unhandled(); }
 
-	DEPRECATED(4.17, "This function is no longer used")
-	virtual FText GetDisplayName() const { return FText(); }
-	
 	/**
 	 * @return The display name of the section in the section view
 	 */
@@ -120,7 +108,7 @@ public:
 	 *
 	 * @param LayoutBuilder	The builder utility for creating section layouts
 	 */	
-	virtual void GenerateSectionLayout( class ISectionLayoutBuilder& LayoutBuilder ) const = 0;
+	SEQUENCER_API virtual void GenerateSectionLayout( class ISectionLayoutBuilder& LayoutBuilder );
 
 	/**
 	 * @return The height of the section
@@ -134,9 +122,6 @@ public:
 	 */
 	virtual bool SectionIsResizable() const {return true;}
 
-	DEPRECATED(4.17, "This function is no longer used")
-	virtual bool AreSectionsConnected() const { return false; }
-
 	/**
 	 * Ticks the section during the Slate tick
 	 *
@@ -178,17 +163,7 @@ public:
 	 * @param ResizeTime The time to resize to
 	 */
 	virtual void BeginResizeSection() {}
-	virtual void ResizeSection(ESequencerSectionResizeMode ResizeMode, float ResizeTime) { ResizeMode == ESequencerSectionResizeMode::SSRM_LeadingEdge ? GetSectionObject()->SetStartTime(ResizeTime) : GetSectionObject()->SetEndTime(ResizeTime); }
-
-	/**
-	 * Dilates the section by a specific factor
-	 *
-	 * @param DilationFactor The multiplier which scales this section
-	 * @param bFromStart Whether to dilate from the beginning or end (whichever stays put)
-	 * @param KeyHandles The key handles to operate on
-	 */
-	virtual void BeginDilateSection() {}
-	virtual void DilateSection(float DilationFactor, float Origin, TSet<FKeyHandle>& KeyHandles) { GetSectionObject()->DilateSection(DilationFactor, Origin, KeyHandles); }
+	SEQUENCER_API virtual void ResizeSection(ESequencerSectionResizeMode ResizeMode, FFrameNumber ResizeFrameNumber);
 
 	/**
 	 * Slips the section by a specific factor
@@ -196,5 +171,23 @@ public:
 	 * @param SlipTime The amount to slip this section by
 	 */
 	virtual void BeginSlipSection() {}
-	virtual void SlipSection(float SlipTime) {}
+	virtual void SlipSection(double SlipTime) {}
 };
+
+class FSequencerSection : public ISequencerSection
+{
+public:
+	FSequencerSection(UMovieSceneSection& InSection)
+		: WeakSection(&InSection)
+	{}
+
+	SEQUENCER_API virtual int32 OnPaintSection(FSequencerSectionPainter& Painter) const override;
+
+	virtual UMovieSceneSection* GetSectionObject() override final
+	{
+		return WeakSection.Get();
+	}
+
+protected:
+	TWeakObjectPtr<UMovieSceneSection> WeakSection;
+};
\ No newline at end of file
diff --git a/Engine/Source/Editor/Sequencer/Public/ISequencerTrackEditor.h b/Engine/Source/Editor/Sequencer/Public/ISequencerTrackEditor.h
--- a/Engine/Source/Editor/Sequencer/Public/ISequencerTrackEditor.h
+++ b/Engine/Source/Editor/Sequencer/Public/ISequencerTrackEditor.h
@@ -112,18 +112,28 @@ public:
 	 *
 	 * @param DragDropEvent The drag drop event.
 	 * @param Track The track that is receiving this drop event.
+	 * @param RowIndex The row index to drop onto.
+	 * @param TargetObjectGuid The object guid this asset is dropped onto, if applicable.
 	 * @return Whether the drop event can be handled.
 	 */
-	virtual bool OnAllowDrop(const FDragDropEvent& DragDropEvent, UMovieSceneTrack* Track) = 0;
+	virtual bool OnAllowDrop(const FDragDropEvent& DragDropEvent, UMovieSceneTrack* Track, int32 RowIndex, const FGuid& TargetObjectGuid) = 0;
+
+	DEPRECATED(4.19, "Use OnAllowDrop with a given RowIndex and Guid")
+	virtual bool OnAllowDrop(const FDragDropEvent& DragDropEvent, UMovieSceneTrack* Track) { return OnAllowDrop(DragDropEvent, Track, 0, FGuid()); }
 
 	/**
 	 * Called when an asset is dropped directly onto a track.
 	 *
 	 * @param DragDropEvent The drag drop event.
 	 * @param Track The track that is receiving this drop event.
+	 * @param RowIndex The row index to drop onto.
+	 * @param TargetObjectGuid The object guid this asset is dropped onto, if applicable.
 	 * @return Whether the drop event was handled.
 	 */	
-	virtual FReply OnDrop(const FDragDropEvent& DragDropEvent, UMovieSceneTrack* Track) = 0;
+	virtual FReply OnDrop(const FDragDropEvent& DragDropEvent, UMovieSceneTrack* Track, int32 RowIndex, const FGuid& TargetObjectGuid) = 0;
+
+	DEPRECATED(4.19, "Use OnDrop with a given RowIndex")
+	virtual FReply OnDrop(const FDragDropEvent& DragDropEvent, UMovieSceneTrack* Track) { return OnDrop(DragDropEvent, Track, 0, FGuid()); }
 
 	/**
 	 * Called to generate a section layout for a particular section.
diff --git a/Engine/Source/Editor/Sequencer/Public/KeyDrawParams.h b/Engine/Source/Editor/Sequencer/Public/KeyDrawParams.h
new file mode 100644
--- /dev/null
+++ b/Engine/Source/Editor/Sequencer/Public/KeyDrawParams.h
@@ -0,0 +1,43 @@
+// Copyright 1998-2018 Epic Games, Inc. All Rights Reserved.
+
+#pragma once
+
+#include "Math/Color.h"
+
+struct FSlateBrush;
+
+/**
+ * Structure defining how a key should be drawn
+ */
+struct FKeyDrawParams
+{
+	FKeyDrawParams()
+		: BorderBrush(nullptr), FillBrush(nullptr)
+		, BorderTint(FLinearColor::White), FillTint(FLinearColor::White)
+		, FillOffset(0.f, 0.f)
+	{}
+
+	friend bool operator==(const FKeyDrawParams& A, const FKeyDrawParams& B)
+	{
+		return A.BorderBrush == B.BorderBrush && A.FillBrush == B.FillBrush && A.BorderTint == B.BorderTint && A.FillTint == B.FillTint && A.FillOffset == B.FillOffset;
+	}
+	friend bool operator!=(const FKeyDrawParams& A, const FKeyDrawParams& B)
+	{
+		return A.BorderBrush != B.BorderBrush || A.FillBrush != B.FillBrush || A.BorderTint != B.BorderTint || A.FillTint != B.FillTint || A.FillOffset != B.FillOffset;
+	}
+
+	/** Brush to use for drawing the key's border */
+	const FSlateBrush* BorderBrush;
+
+	/** Brush to use for drawing the key's filled area */
+	const FSlateBrush* FillBrush;
+
+	/** Tint to be used for the key's border */
+	FLinearColor BorderTint;
+
+	/** Tint to be used for the key's filled area */
+	FLinearColor FillTint;
+
+	/** The amount to offset the fill brush from the keys center */
+	FVector2D FillOffset;
+};
\ No newline at end of file
diff --git a/Engine/Source/Editor/Sequencer/Public/KeyPropertyParams.h b/Engine/Source/Editor/Sequencer/Public/KeyPropertyParams.h
--- a/Engine/Source/Editor/Sequencer/Public/KeyPropertyParams.h
+++ b/Engine/Source/Editor/Sequencer/Public/KeyPropertyParams.h
@@ -72,7 +72,7 @@ struct SEQUENCER_API FKeyPropertyParams
 class SEQUENCER_API FPropertyChangedParams
 {
 public:
-	FPropertyChangedParams(TArray<UObject*> InObjectsThatChanged, const FPropertyPath& InPropertyPath, FName InStructPropertyNameToKey, ESequencerKeyMode InKeyMode);
+	FPropertyChangedParams(TArray<UObject*> InObjectsThatChanged, const FPropertyPath& InPropertyPath, const FPropertyPath& InStructPathToKey, ESequencerKeyMode InKeyMode);
 
 	/**
 	 * Gets the value of the property that changed.
@@ -115,9 +115,9 @@ public:
 	/** A path of properties to get from the root object to the property to be keyed. */
 	FPropertyPath PropertyPath;
 
-	/** Represents the FName of an inner property which should be keyed for a struct property.  If all inner 
-	properties should be keyed, this will be FName::None. */
-	const FName StructPropertyNameToKey;
+	/** Represents the path of an inner property which should be keyed for a struct property.  If all inner 
+	properties should be keyed, this will be empty. */
+	FPropertyPath StructPathToKey;
 
 	/** Keyframing params */
 	const ESequencerKeyMode KeyMode;
diff --git a/Engine/Source/Editor/Sequencer/Public/MovieSceneTrackEditor.h b/Engine/Source/Editor/Sequencer/Public/MovieSceneTrackEditor.h
--- a/Engine/Source/Editor/Sequencer/Public/MovieSceneTrackEditor.h
+++ b/Engine/Source/Editor/Sequencer/Public/MovieSceneTrackEditor.h
@@ -51,17 +51,17 @@ struct FKeyPropertyResult
 };
 
 /** Delegate for adding keys for a property
- * float - The time at which to add the key.
+ * FFrameNumber - The time at which to add the key.
  * return - KeyPropertyResult - 
  */
-DECLARE_DELEGATE_RetVal_OneParam(FKeyPropertyResult, FOnKeyProperty, float)
+DECLARE_DELEGATE_RetVal_OneParam(FKeyPropertyResult, FOnKeyProperty, FFrameNumber)
 
 
 /** Delegate for whether a property can be keyed
- * float - The time at which to add the key.
+ * FFrameNumber - The time at which to add the key.
  * return - True if the property can be keyed, otherwise false.
  */
-DECLARE_DELEGATE_RetVal_OneParam(bool, FCanKeyProperty, float)
+DECLARE_DELEGATE_RetVal_OneParam(bool, FCanKeyProperty, FFrameNumber)
 
 
 /**
@@ -90,7 +90,7 @@ public:
 	/**
 	 * @return The current local time at which we should add a key
 	 */
-	float GetTimeForKey();
+	FFrameNumber GetTimeForKey();
 
 	void UpdatePlaybackRange();
 
@@ -156,6 +156,8 @@ public:
 
 	virtual UMovieSceneTrack* AddTrack(UMovieScene* FocusedMovieScene, const FGuid& ObjectHandle, TSubclassOf<class UMovieSceneTrack> TrackClass, FName UniqueTypeName) override;
 
+	virtual TSharedRef<ISequencerSection> MakeSectionInterface(UMovieSceneSection& SectionObject, UMovieSceneTrack& Track, FGuid ObjectBinding) override;
+
 	virtual void BindCommands(TSharedRef<FUICommandList> SequencerCommandBindings) override;
 	virtual void BuildAddTrackMenu(FMenuBuilder& MenuBuilder) override;
 	virtual void BuildObjectBindingEditButtons(TSharedPtr<SHorizontalBox> EditBox, const FGuid& ObjectBinding, const UClass* ObjectClass) override;
@@ -163,8 +165,8 @@ public:
 	virtual TSharedPtr<SWidget> BuildOutlinerEditWidget(const FGuid& ObjectBinding, UMovieSceneTrack* Track, const FBuildEditWidgetParams& Params) override;
 	virtual void BuildTrackContextMenu( FMenuBuilder& MenuBuilder, UMovieSceneTrack* Track ) override;
 	virtual bool HandleAssetAdded(UObject* Asset, const FGuid& TargetObjectGuid) override;
-	virtual bool OnAllowDrop(const FDragDropEvent& DragDropEvent, UMovieSceneTrack* Track) override;
-	virtual FReply OnDrop(const FDragDropEvent& DragDropEvent, UMovieSceneTrack* Track) override;
+	virtual bool OnAllowDrop(const FDragDropEvent& DragDropEvent, UMovieSceneTrack* Track, int32 RowIndex, const FGuid& TargetObjectGuid) override;
+	virtual FReply OnDrop(const FDragDropEvent& DragDropEvent, UMovieSceneTrack* Track, int32 RowIndex, const FGuid& TargetObjectGuid) override;
 
 	virtual void OnInitialize() override;
 	virtual void OnRelease() override;
diff --git a/Engine/Source/Editor/Sequencer/Public/SequencerChannelInterface.h b/Engine/Source/Editor/Sequencer/Public/SequencerChannelInterface.h
new file mode 100644
--- /dev/null
+++ b/Engine/Source/Editor/Sequencer/Public/SequencerChannelInterface.h
@@ -0,0 +1,357 @@
+// Copyright 1998-2018 Epic Games, Inc. All Rights Reserved.
+
+#pragma once
+
+#include "ISequencerChannelInterface.h"
+#include "SequencerChannelTraits.h"
+#include "MovieSceneSection.h"
+#include "IKeyArea.h"
+#include "CurveModel.h"
+
+/**
+ * Templated channel interface that calls overloaded functions matching the necessary channel types.
+ * Designed this way to allow for specific customization of key-channel behavior without having to reimplement swathes of boilerplate.
+ * This base interface implements common functions that do not require specialized editor data.
+ *
+ * Behavior can be overridden for any channel type by declaring an overloaded function for the relevant channel type in the same namespace as the channel.
+ * For instance, to implement how to retrieve key times from a channel, implement the following function:
+ *
+ * void GetKeyTimes(FMyChannelType* InChannel, TArrayView<const FKeyHandle> InHandles, TArrayView<FFrameNumber> OutKeyTimes);
+ */
+template<typename ChannelType>
+struct TSequencerChannelInterfaceCommon : ISequencerChannelInterface
+{
+	/**
+	 * Check whether any of the specified channels have any keys
+	 *
+	 * @return true if so, false otherwise
+	 */
+	virtual bool HasAnyKeys_Raw(TArrayView<void* const> Ptrs) const override
+	{
+		using namespace Sequencer;
+
+		for (void* Ptr : Ptrs)
+		{
+			if (Ptr && HasAnyKeys(static_cast<ChannelType*>(Ptr)))
+			{
+				return true;
+			}
+		}
+		return false;
+	}
+
+	/**
+	 * Get key information pertaining to all keys that exist within the specified range
+	 *
+	 * @param Channel               The channel to query
+	 * @param WithinRange           The range within which to return key information
+	 * @param OutKeyTimes           (Optional) Array to receive key times
+	 * @param OutKeyHandles         (Optional) Array to receive key handles
+	 */
+	virtual void GetKeys_Raw(void* InChannel, const TRange<FFrameNumber>& WithinRange, TArray<FFrameNumber>* OutKeyTimes, TArray<FKeyHandle>* OutKeyHandles) const override
+	{
+		using namespace Sequencer;
+		GetKeys(static_cast<ChannelType*>(InChannel), WithinRange, OutKeyTimes, OutKeyHandles);
+	}
+
+	/**
+	 * Get all key times for the specified key handles
+	 *
+	 * @param Channel               The channel to query
+	 * @param InHandles             Array of handles to get times for
+	 * @param OutKeyTimes           Pre-sized array of key times to set. Invalid key handles will not assign to this array. Must match size of InHandles
+	 */
+	virtual void GetKeyTimes_Raw(void* InChannel, TArrayView<const FKeyHandle> InHandles, TArrayView<FFrameNumber> OutKeyTimes) const override
+	{
+		check(InHandles.Num() == OutKeyTimes.Num());
+
+		using namespace Sequencer;
+		GetKeyTimes(static_cast<ChannelType*>(InChannel), InHandles, OutKeyTimes);
+	}
+
+	/**
+	 * Set key times for the specified key handles
+	 *
+	 * @param Channel               The channel to query
+	 * @param InHandles             Array of handles to get times for
+	 * @param InKeyTimes            Array of times to apply - one per handle
+	 */
+	virtual void SetKeyTimes_Raw(void* InChannel, TArrayView<const FKeyHandle> InHandles, TArrayView<const FFrameNumber> InKeyTimes) const override
+	{
+		check(InHandles.Num() == InKeyTimes.Num());
+
+		using namespace Sequencer;
+		SetKeyTimes(static_cast<ChannelType*>(InChannel), InHandles, InKeyTimes);
+	}
+
+	/**
+	 * Duplicate the keys for the specified key handles
+	 *
+	 * @param Channel               The channel to query
+	 * @param InHandles             Array of handles to duplicate
+	 * @param OutKeyTimes           Pre-sized array to receive duplicated key handles. Invalid key handles will not be assigned to this array. Must match size of InHandles
+	 */
+	virtual void DuplicateKeys_Raw(void* InChannel, TArrayView<const FKeyHandle> InHandles, TArrayView<FKeyHandle> OutNewHandles) const override
+	{
+		check(InHandles.Num() == OutNewHandles.Num());
+
+		using namespace Sequencer;
+		DuplicateKeys(static_cast<ChannelType*>(InChannel), InHandles, OutNewHandles);
+	}
+
+	/**
+	 * Delete the keys for the specified key handles
+	 *
+	 * @param Channel               The channel to query
+	 * @param InHandles             Array of handles to delete
+	 */
+	virtual void DeleteKeys_Raw(void* InChannel, TArrayView<const FKeyHandle> InHandles) const override
+	{
+		using namespace Sequencer;
+		DeleteKeys(static_cast<ChannelType*>(InChannel), InHandles);
+	}
+
+	/**
+	 * Copy all the keys specified in KeyMask to the specified clipboard
+	 *
+	 * @param Channel               The channel to copy from
+	 * @param Section               The section that owns the channel
+	 * @param KeyAreaName           The name of the key area
+	 * @param ClipboardBuilder      The structure responsible for building clipboard information for each key
+	 * @param KeyMask               A specific set of keys to copy
+	 */
+	virtual void CopyKeys_Raw(void* InChannel, const UMovieSceneSection* Section, FName KeyAreaName, FMovieSceneClipboardBuilder& ClipboardBuilder, TArrayView<const FKeyHandle> KeyMask) const override
+	{
+		using namespace Sequencer;
+		CopyKeys(static_cast<ChannelType*>(InChannel), Section, KeyAreaName, ClipboardBuilder, KeyMask);
+	}
+
+	/**
+	 * Paste the specified key track into the specified channel
+	 *
+	 * @param Channel               The channel to copy from
+	 * @param Section               The section that owns the channel
+	 * @param KeyTrack              The source clipboard data to paste
+	 * @param SrcEnvironment        The environment the source data was copied from
+	 * @param DstEnvironment        The environment we're pasting into
+	 * @param OutPastedKeys         Array to receive key handles for any pasted keys
+	 */
+	virtual void PasteKeys_Raw(void* InChannel, UMovieSceneSection* Section, const FMovieSceneClipboardKeyTrack& KeyTrack, const FMovieSceneClipboardEnvironment& SrcEnvironment, const FSequencerPasteEnvironment& DstEnvironment, TArray<FKeyHandle>& OutPastedKeys) const override
+	{
+		using namespace Sequencer;
+		PasteKeys(static_cast<ChannelType*>(InChannel), Section, KeyTrack, SrcEnvironment, DstEnvironment, OutPastedKeys);
+	}
+
+	/**
+	 * Get an editable key struct for the specified key
+	 *
+	 * @param Channel               The channel on which the key resides
+	 * @param KeyHandle             Handle of the key to get
+	 * @return A shared editable key struct
+	 */
+	virtual TSharedPtr<FStructOnScope> GetKeyStruct_Raw(TMovieSceneChannelHandle<void> InChannel, FKeyHandle KeyHandle) const override
+	{
+		using namespace Sequencer;
+		return GetKeyStruct(InChannel.Cast<ChannelType>(), KeyHandle);
+	}
+
+	/**
+	 * Check whether an editor on the sequencer node tree can be created for the specified channel
+	 *
+	 * @param Channel               The channel to check
+	 * @return true if a key editor should be constructed, false otherwise
+	 */
+	virtual bool CanCreateKeyEditor_Raw(void* InChannel) const override
+	{
+		using namespace Sequencer;
+		return CanCreateKeyEditor(static_cast<ChannelType*>(InChannel));
+	}
+
+	/**
+	 * Extend the key context menu
+	 *
+	 * @param MenuBuilder           The menu builder used to create this context menu
+	 * @param Channels              Array of channels and handles that are being shown in the context menu
+	 * @param InSequencer           The currently active sequencer
+	 */
+	virtual void ExtendKeyMenu_Raw(FMenuBuilder& MenuBuilder, TArrayView<const TChannelAndHandles<void>> ChannelsAndHandles, TWeakPtr<ISequencer> InSequencer) const override
+	{
+		using namespace Sequencer;
+		TArray<TChannelAndHandles<ChannelType>> TypedChannels;
+
+		for (const TChannelAndHandles<void>& Ptr : ChannelsAndHandles)
+		{
+			TChannelAndHandles<ChannelType> TypedChannelAndHandles;
+			TypedChannelAndHandles.Section = Ptr.Section;
+			TypedChannelAndHandles.Handles = Ptr.Handles;
+			TypedChannelAndHandles.Channel = Ptr.Channel.Cast<ChannelType>();
+
+			TypedChannels.Add(MoveTemp(TypedChannelAndHandles));
+		}
+
+		ExtendKeyMenu(MenuBuilder, MoveTemp(TypedChannels), InSequencer);
+	}
+
+	/**
+	 * Extend the section context menu
+	 *
+	 * @param MenuBuilder           The menu builder used to create this context menu
+	 * @param Channels              Array of type specific channels that exist in the selected sections
+	 * @param Sections              Array of sections being shown on the context menu
+	 * @param InSequencer           The currently active sequencer
+	 */
+	virtual void ExtendSectionMenu_Raw(FMenuBuilder& MenuBuilder, TArrayView<TMovieSceneChannelHandle<void> const> Channels, TArrayView<UMovieSceneSection* const> Sections, TWeakPtr<ISequencer> InSequencer) const override
+	{
+		using namespace Sequencer;
+		TArray<TMovieSceneChannelHandle<ChannelType>> TypedChannels;
+
+		for (const TMovieSceneChannelHandle<void>& RawHandle : Channels)
+		{
+			TypedChannels.Add(RawHandle.Cast<ChannelType>());
+		}
+
+		ExtendSectionMenu(MenuBuilder, MoveTemp(TypedChannels), Sections, InSequencer);
+	}
+
+	/**
+	 * Gather information on how to draw the specified keys
+	 *
+	 * @param Channel               The channel to query
+	 * @param InKeyHandles          Array of handles to duplicate
+	 * @param OutKeyDrawParams      Pre-sized array to receive key draw parameters. Invalid key handles will not be assigned to this array. Must match size of InKeyHandles.
+	 */
+	virtual void DrawKeys_Raw(void* InChannel, TArrayView<const FKeyHandle> InKeyHandles, TArrayView<FKeyDrawParams> OutKeyDrawParams) const override
+	{
+		check(InKeyHandles.Num() == OutKeyDrawParams.Num());
+
+		using namespace Sequencer;
+		DrawKeys(static_cast<ChannelType*>(InChannel), InKeyHandles, OutKeyDrawParams);
+	}
+
+	/**
+	 * Create a new model for this channel that can be used on the curve editor interface
+	 *
+	 * @return (Optional) A new model to be added to a curve editor
+	 */
+	virtual TUniquePtr<FCurveModel> CreateCurveEditorModel_Raw(TMovieSceneChannelHandle<void> InChannel, UMovieSceneSection* OwningSection, TSharedRef<ISequencer> InSequencer) const override
+	{
+		using namespace Sequencer;
+		return CreateCurveEditorModel(InChannel.Cast<ChannelType>(), OwningSection, InSequencer);
+	}
+};
+
+template<typename ChannelType, bool HasSpecializedData> struct TSequencerChannelInterfaceBase;
+
+
+/**
+ * Specialized base interface for channel types that do not specify specialized editor data (ie, TMovieSceneChannelTraits<>::EditorDataType is void)
+ */
+template<typename ChannelType>
+struct TSequencerChannelInterfaceBase<ChannelType, false> : TSequencerChannelInterfaceCommon<ChannelType>
+{
+	/**
+	 * Add (or update) a key to the specified channel using it's current value at that time, or some external value specified by the specialized editor data
+	 *
+	 * @param Channel               The channel to add a key to
+	 * @param SpecializedEditorData A pointer to the specialized editor data for this channel of type TMovieSceneChannelTraits<>::EditorDataType
+	 * @param InTime                The time at which to add a key
+	 * @param InSequencer           The currently active sequencer
+	 * @param ObjectBindingID       The object binding ID for the track that this channel resides within
+	 * @param PropertyBindings      (Optional) Property bindings where this channel exists on a property track
+	 * @return A handle to the new or updated key
+	 */
+	virtual FKeyHandle AddOrUpdateKey_Raw(void* InChannel, const void* SpecializedEditorData, FFrameNumber InTime, ISequencer& Sequencer, const FGuid& ObjectBindingID, FTrackInstancePropertyBindings* PropertyBindings) const override
+	{
+		using namespace Sequencer;
+		return AddOrUpdateKey(static_cast<ChannelType*>(InChannel), InTime, Sequencer, ObjectBindingID, PropertyBindings);
+	}
+
+	/**
+	 * Create an editor on the sequencer node tree
+	 *
+	 * @param Channel               The channel to check
+	 * @param SpecializedEditorData A pointer to the specialized editor data for this channel of type TMovieSceneChannelTraits<>::EditorDataType
+	 * @param Section               The section that owns this channel
+	 * @param InObjectBindingID     The ID of the object this key area's track is bound to
+	 * @param PropertyBindings      (Optional) Property bindings where this channel exists on a property track
+	 * @param Sequencer             The currently active sequencer
+	 * @return The editor widget to display on the node tree
+	 */
+	virtual TSharedRef<SWidget> CreateKeyEditor_Raw(void* InChannel, const void* SpecializedEditorData, UMovieSceneSection* Section, const FGuid& InObjectBindingID, TWeakPtr<FTrackInstancePropertyBindings> PropertyBindings, TWeakPtr<ISequencer> Sequencer) const override
+	{
+		using namespace Sequencer;
+		return CreateKeyEditor(static_cast<ChannelType*>(InChannel), Section, InObjectBindingID, PropertyBindings, Sequencer);
+	}
+};
+
+/**
+ * Specialized base interface for channel types that specify specialized editor data (ie, TMovieSceneChannelTraits<>::EditorDataType is not void)
+ */
+template<typename ChannelType>
+struct TSequencerChannelInterfaceBase<ChannelType, true> : TSequencerChannelInterfaceCommon<ChannelType>
+{
+	typedef typename TMovieSceneChannelTraits<ChannelType>::EditorDataType EditorDataType;
+
+	/**
+	 * Add (or update) a key to the specified channel using it's current value at that time, or some external value specified by the specialized editor data
+	 *
+	 * @param Channel               The channel to add a key to
+	 * @param SpecializedEditorData A pointer to the specialized editor data for this channel of type TMovieSceneChannelTraits<>::EditorDataType
+	 * @param InTime                The time at which to add a key
+	 * @param InSequencer           The currently active sequencer
+	 * @param ObjectBindingID       The object binding ID for the track that this channel resides within
+	 * @param PropertyBindings      (Optional) Property bindings where this channel exists on a property track
+	 * @return A handle to the new or updated key
+	 */
+	virtual FKeyHandle AddOrUpdateKey_Raw(void* InChannel, const void* SpecializedEditorData, FFrameNumber InTime, ISequencer& Sequencer, const FGuid& ObjectBindingID, FTrackInstancePropertyBindings* PropertyBindings) const override
+	{
+		using namespace Sequencer;
+
+		// Specialized data must be available for this interface
+		check(SpecializedEditorData);
+
+		const auto* TypedEditorData = static_cast<const EditorDataType*>(SpecializedEditorData);
+		return AddOrUpdateKey(static_cast<ChannelType*>(InChannel), *TypedEditorData, InTime, Sequencer, ObjectBindingID, PropertyBindings);
+	}
+
+	/**
+	 * Create an editor on the sequencer node tree
+	 *
+	 * @param Channel               The channel to check
+	 * @param SpecializedEditorData A pointer to the specialized editor data for this channel of type TMovieSceneChannelTraits<>::EditorDataType
+	 * @param Section               The section that owns this channel
+	 * @param InObjectBindingID     The ID of the object this key area's track is bound to
+	 * @param PropertyBindings      (Optional) Property bindings where this channel exists on a property track
+	 * @param Sequencer             The currently active sequencer
+	 * @return The editor widget to display on the node tree
+	 */
+	virtual TSharedRef<SWidget> CreateKeyEditor_Raw(void* InChannel, const void* SpecializedEditorData, UMovieSceneSection* Section, const FGuid& InObjectBindingID, TWeakPtr<FTrackInstancePropertyBindings> PropertyBindings, TWeakPtr<ISequencer> Sequencer) const override
+	{
+		using namespace Sequencer;
+
+		// Specialized data must be available for this interface
+		check(SpecializedEditorData);
+
+		const auto* TypedEditorData = static_cast<const EditorDataType*>(SpecializedEditorData);
+		return CreateKeyEditor(static_cast<ChannelType*>(InChannel), *TypedEditorData, Section, InObjectBindingID, PropertyBindings, Sequencer);
+	}
+};
+
+/** Generic sequencer channel interface to any channel type */
+template<typename ChannelType>
+struct TSequencerChannelInterface : TSequencerChannelInterfaceBase<ChannelType, !TIsSame<typename TMovieSceneChannelTraits<ChannelType>::EditorDataType, void>::Value>
+{
+};
+
+
+/**
+ * Implementation of the RegisterChannelInterface function that is the only function that needs to see the definition of TSequencerChannelInterface
+ * Defined in here to prevent all the necessary headers being included wherever ISequencerModule is needed.
+ */
+template<typename ChannelType>
+void ISequencerModule::RegisterChannelInterface()
+{
+	const uint32 ChannelID = ChannelType::GetChannelID();
+	check(!ChannelToEditorInterfaceMap.Contains(ChannelID));
+	ChannelToEditorInterfaceMap.Add(ChannelID, TUniquePtr<ISequencerChannelInterface>(new TSequencerChannelInterface<ChannelType>()));
+}
\ No newline at end of file
diff --git a/Engine/Source/Editor/Sequencer/Public/SequencerChannelTraits.h b/Engine/Source/Editor/Sequencer/Public/SequencerChannelTraits.h
new file mode 100644
--- /dev/null
+++ b/Engine/Source/Editor/Sequencer/Public/SequencerChannelTraits.h
@@ -0,0 +1,536 @@
+// Copyright 1998-2018 Epic Games, Inc. All Rights Reserved.
+
+#pragma once
+
+#include "Templates/SharedPointer.h"
+#include "ISequencerChannelInterface.h"
+#include "KeyDrawParams.h"
+#include "Channels/MovieSceneChannel.h"
+#include "MovieSceneClipboard.h"
+#include "SequencerClipboardReconciler.h"
+#include "SequencerGenericKeyStruct.h"
+#include "Widgets/SNullWidget.h"
+#include "ISequencer.h"
+
+struct FSequencerGenericKeyStruct;
+
+/**
+ * Stub/default implementations for ISequencerChannelInterface functions.
+ * Custom behaviour should be implemented by overloading the relevant function with the necessary channel/data types.
+ * For example, to overload HasAnyKeys for FMyCustomChannelType, implement the following function in the same namespace as FMyCustomChannelType:
+ *
+ * bool HasAnyKeys(const FMyCustomChannelType* InChannel)
+ * {
+ *    return InChannel.GetKeys().Num() != 0; }
+ * }
+ */
+namespace Sequencer
+{
+	/**
+	 * Extend the specified selected section context menu
+	 *
+	 * @param MenuBuilder    The menu builder that will construct the section context menu
+	 * @param Channels       An array of all channels that are currently selected, in no particular order
+	 * @param Sections       An array of all sections that the selected channels reside in
+	 * @param InSequencer    The sequencer that is currently active
+	 */
+	template<typename ChannelType>
+	void ExtendSectionMenu(FMenuBuilder& MenuBuilder, TArray<TMovieSceneChannelHandle<ChannelType>>&& Channels, TArrayView<UMovieSceneSection* const> Sections, TWeakPtr<ISequencer> InSequencer)
+	{}
+
+
+	/**
+	 * Extend the specified selected key context menu
+	 *
+	 * @param MenuBuilder    The menu builder that will construct the section context menu
+	 * @param Channels       An array of all channels that are currently selected, in no particular order
+	 * @param Sections       An array of all sections that the selected channels reside in
+	 * @param InSequencer    The sequencer that is currently active
+	 */
+	template<typename ChannelType>
+	void ExtendKeyMenu(FMenuBuilder& MenuBuilder, TArray<TChannelAndHandles<ChannelType>>&& InChannels, TWeakPtr<ISequencer> InSequencer)
+	{}
+
+
+	/**
+	 * Determine if the specified channel has any keys
+	 *
+	 * @param InChannel      The channel to check
+	 * @return true if the channel has any keys, false otherwise
+	 */
+	template<typename ChannelType>
+	bool HasAnyKeys(const ChannelType* InChannel)
+	{
+		return InChannel->GetInterface().GetTimes().Num() != 0;
+	}
+
+
+	/**
+	 * Get a transient key structure that can be added to a details panel to enable editing of a single key
+	 *
+	 * @param ChannelHandle  Handle to the channel in which the key resides
+	 * @param KeyHandle      A handle to the key to edit
+	 * @return A shared struct object, or nullptr
+	 */
+	template<typename ChannelType>
+	TSharedPtr<FStructOnScope> GetKeyStruct(TMovieSceneChannelHandle<ChannelType> ChannelHandle, FKeyHandle KeyHandle)
+	{
+		ChannelType* Channel = ChannelHandle.Get();
+		if (Channel)
+		{
+			auto ChannelInterface = Channel->GetInterface();
+			const int32 KeyIndex = ChannelInterface.GetIndex(KeyHandle);
+
+			if (KeyIndex != INDEX_NONE)
+			{
+				TSharedRef<FStructOnScope> NewStruct = MakeShared<FStructOnScope>(FSequencerGenericKeyStruct::StaticStruct());
+				FSequencerGenericKeyStruct* StructPtr = (FSequencerGenericKeyStruct*)NewStruct->GetStructMemory();
+
+				StructPtr->Time = ChannelInterface.GetTimes()[KeyIndex];
+				StructPtr->CustomizationImpl = MakeShared<TMovieSceneKeyStructCustomization<ChannelType>>(ChannelHandle, KeyHandle);
+
+				return NewStruct;
+			}
+		}
+		return nullptr;
+	}
+
+
+
+	/**
+	 * Check whether the specified channel can create a key editor widget that should be placed on the sequencer node tree 
+	 *
+	 * @param InChannel      The channel to check
+	 * @return true if a key editor can be created, false otherwise
+	 */
+	inline bool CanCreateKeyEditor(void* InChannel)
+	{
+		return false;
+	}
+
+
+	/**
+	 * Create a key editor widget for the specified channel. Such widgets are placed on the sequencer node tree for a given key area node.
+	 *
+	 * @param InChannel          The channel to create a key editor for
+	 * @param InOwningSection    The section that owns the channel
+	 * @param InObjectBindingID  The object binding ID that this section's track is bound to
+	 * @param InPropertyBindings Optionally supplied helper for accessing an object's property pertaining to this channel
+	 * @param InSequencer        The sequencer currently active
+	 * @return true if a key editor can be created, false otherwise
+	 */
+	inline TSharedRef<SWidget> CreateKeyEditor(
+		void*                                    InChannel,
+		UMovieSceneSection*                      InOwningSection,
+		const FGuid&                             InObjectBindingID,
+		TWeakPtr<FTrackInstancePropertyBindings> InPropertyBindings,
+		TWeakPtr<ISequencer>                     InSequencer
+		)
+	{
+		return SNullWidget::NullWidget;
+	}
+
+
+	/**
+	 * Create a key editor widget for the specified channel with the channel's specialized editor data. Such widgets are placed on the sequencer node tree for a given key area node.
+	 *
+	 * @param InChannel          The channel to create a key editor for
+	 * @param InEditorData       The channel's specialized editor data
+	 * @param InOwningSection    The section that owns the channel
+	 * @param InObjectBindingID  The object binding ID that this section's track is bound to
+	 * @param InPropertyBindings Optionally supplied helper for accessing an object's property pertaining to this channel
+	 * @param InSequencer        The sequencer currently active
+	 * @return The key editor widget
+	 */
+	template<typename SpecializedEditorDataType>
+	inline TSharedRef<SWidget> CreateKeyEditor(
+		void*                                    InChannel,
+		const SpecializedEditorDataType&         InEditorData,
+		UMovieSceneSection*                      InOwningSection,
+		const FGuid&                             InObjectBindingID,
+		TWeakPtr<FTrackInstancePropertyBindings> InPropertyBindings,
+		TWeakPtr<ISequencer>                     InSequencer
+		)
+	{
+		return SNullWidget::NullWidget;
+	}
+
+
+	/**
+	 * Add a key at the specified time (or update an existing key) with the channel's current value at that time
+	 *
+	 * @param InChannel          The channel to create a key for
+	 * @param InChannelInterface The channel's interface
+	 * @param InTime             The time at which to add a key
+	 * @param InSequencer        The currently active sequencer
+	 * @param InDefaultValue     The default value to use if evaluation of the channel failed
+	 * @return A handle to the added (or updated) key
+	 */
+	template<typename ChannelType, typename ValueType>
+	FKeyHandle EvaluateAndAddKey(ChannelType* InChannel, const TMovieSceneChannel<ValueType>& InChanneInterface, FFrameNumber InTime, ISequencer& InSequencer, ValueType InDefaultValue = ValueType{})
+	{
+		using namespace MovieScene;
+
+		ValueType ValueAtTime = InDefaultValue;
+		EvaluateChannel(InChannel, InTime, ValueAtTime);
+
+		return AddKeyToChannel(InChannel, InTime, ValueAtTime, InSequencer.GetKeyInterpolation());
+	}
+
+
+	/**
+	 * Retrieve a channel's external value, and add it to the channel as a new key (or update an existing key with its value)
+	 *
+	 * @param InChannel          The channel to create a key for
+	 * @param InExternalValue    The external value definition
+	 * @param InTime             The time at which to add a key
+	 * @param InSequencer        The currently active sequencer
+	 * @param InObjectBindingID  The object binding ID that this section's track is bound to
+	 * @param InPropertyBindings Optionally supplied helper for accessing an object's property pertaining to this channel
+	 * @return (Optional) A handle to the added (or updated) key
+	 */
+	template<typename ChannelType, typename ValueType>
+	TOptional<FKeyHandle> AddKeyForExternalValue(
+		ChannelType*                               InChannel,
+		const TMovieSceneExternalValue<ValueType>& InExternalValue,
+		FFrameNumber                               InTime,
+		ISequencer&                                InSequencer,
+		const FGuid&                               InObjectBindingID,
+		FTrackInstancePropertyBindings*            InPropertyBindings
+		)
+	{
+		using namespace MovieScene;
+
+		// Add a key for the current value of the valid first object we can find
+		if (InExternalValue.OnGetExternalValue && InObjectBindingID.IsValid())
+		{
+			for (TWeakObjectPtr<> WeakObject : InSequencer.FindBoundObjects(InObjectBindingID, InSequencer.GetFocusedTemplateID()))
+			{
+				UObject* Object = WeakObject.Get();
+				if (Object)
+				{
+					TOptional<ValueType> Value = InExternalValue.OnGetExternalValue(*Object, InPropertyBindings);
+					if (Value.IsSet())
+					{
+						return AddKeyToChannel(InChannel, InTime, Value.GetValue(), InSequencer.GetKeyInterpolation());
+					}
+				}
+			}
+		}
+		return TOptional<FKeyHandle>();
+	}
+
+
+	/**
+	 * Add or update a key for this channel's current value
+	 *
+	 * @param InChannel          The channel to create a key for
+	 * @param InTime             The time at which to add a key
+	 * @param InSequencer        The currently active sequencer
+	 * @param InObjectBindingID  The object binding ID that this section's track is bound to
+	 * @param InPropertyBindings Optionally supplied helper for accessing an object's property pertaining to this channel
+	 * @return A handle to the added (or updated) key
+	 */
+	template<typename ChannelType>
+	FKeyHandle AddOrUpdateKey(
+		ChannelType*                    InChannel,
+		FFrameNumber                    InTime,
+		ISequencer&                     InSequencer,
+		const FGuid&                    InObjectBindingID,
+		FTrackInstancePropertyBindings* InPropertyBindings
+		)
+	{
+		return EvaluateAndAddKey(InChannel, InChannel->GetInterface(), InTime, InSequencer);
+	}
+
+
+	/**
+	 * Add or update a key for this channel's current value, using an external value if possible
+	 *
+	 * @param InChannel          The channel to create a key for
+	 * @param InExternalValue    The external value definition
+	 * @param InTime             The time at which to add a key
+	 * @param InSequencer        The currently active sequencer
+	 * @param InObjectBindingID  The object binding ID that this section's track is bound to
+	 * @param InPropertyBindings Optionally supplied helper for accessing an object's property pertaining to this channel
+	 * @return A handle to the added (or updated) key
+	 */
+	template<typename ChannelType, typename ValueType>
+	FKeyHandle AddOrUpdateKey(
+		ChannelType*                               InChannel,
+		const TMovieSceneExternalValue<ValueType>& InExternalValue,
+		FFrameNumber                               InTime,
+		ISequencer&                                InSequencer,
+		const FGuid&                               InObjectBindingID,
+		FTrackInstancePropertyBindings*            InPropertyBindings)
+	{
+		using namespace MovieScene;
+
+		TOptional<FKeyHandle> Handle = AddKeyForExternalValue(InChannel, InExternalValue, InTime, InSequencer, InObjectBindingID, InPropertyBindings);
+		if (!Handle.IsSet())
+		{
+			ValueType ValueAtTime{};
+			EvaluateChannel(InChannel, InTime, ValueAtTime);
+		
+			Handle = AddKeyToChannel(InChannel, InTime, ValueAtTime, InSequencer.GetKeyInterpolation());
+		}
+
+		return Handle.GetValue();
+	}
+
+
+	/**
+	 * Get all the keys in the given range. Resulting arrays must be the same size where indices correspond to both arrays.
+	 *
+	 * @param InChannel          The channel to get keys from
+	 * @param WithinRange        The bounds to get keys for
+	 * @param OutKeyTimes        Array to receive all key times within the given range
+	 * @param OutKeyHandles      Array to receive all key handles within the given range
+	 */
+	template<typename ChannelType>
+	void GetKeys(ChannelType* InChannel, const TRange<FFrameNumber>& WithinRange, TArray<FFrameNumber>* OutKeyTimes, TArray<FKeyHandle>* OutKeyHandles)
+	{
+		auto ChannelInterface = InChannel->GetInterface();
+		TArrayView<const FFrameNumber> Times = ChannelInterface.GetTimes();
+		if (!Times.Num())
+		{
+			return;
+		}
+
+		const int32 FirstIndex = WithinRange.GetLowerBound().IsClosed() ? Algo::LowerBound(Times, WithinRange.GetLowerBoundValue()) : 0;
+		const int32 LastIndex  = WithinRange.GetUpperBound().IsClosed() ? Algo::UpperBound(Times, WithinRange.GetUpperBoundValue()) : Times.Num();
+
+		const int32 NumInRange = LastIndex - FirstIndex;
+		if (NumInRange > 0)
+		{
+			if (OutKeyTimes)
+			{
+				OutKeyTimes->Reserve(OutKeyTimes->Num() + NumInRange);
+				OutKeyTimes->Append(&Times[FirstIndex], NumInRange);
+			}
+
+			if (OutKeyHandles)
+			{
+				OutKeyHandles->Reserve(OutKeyHandles->Num() + NumInRange);
+
+				for (int32 Index = FirstIndex; Index < LastIndex; ++Index)
+				{
+					OutKeyHandles->Add(ChannelInterface.GetHandle(Index));
+				}
+			}
+		}
+	}
+
+
+	/**
+	 * Get key times for a number of keys in the specified channel
+	 *
+	 * @param InChannel          The channel to set key times on
+	 * @param InHandles          Array of key handles that should have their times set
+	 * @param OutKeyTimes        Array of times that should be set for each key handle. Must be exactly the size of InHandles
+	 */
+	template<typename ChannelType>
+	void GetKeyTimes(ChannelType* InChannel, TArrayView<const FKeyHandle> InHandles, TArrayView<FFrameNumber> OutKeyTimes)
+	{
+		check(InHandles.Num() == OutKeyTimes.Num());
+
+		auto ChannelInterface = InChannel->GetInterface();
+		auto Times = ChannelInterface.GetTimes();
+		for (int32 Index = 0; Index < InHandles.Num(); ++Index)
+		{
+			const int32 KeyIndex = ChannelInterface.GetIndex(InHandles[Index]);
+			if (KeyIndex != INDEX_NONE)
+			{
+				OutKeyTimes[Index] = Times[KeyIndex];
+			}
+		}
+	}
+
+
+	/**
+	 * Set key times for a number of keys in the specified channel
+	 *
+	 * @param InChannel          The channel to set key times on
+	 * @param InHandles          Array of key handles that should have their times set
+	 * @param InKeyTimes         Array of new times for each handle of the above array
+	 */
+	template<typename ChannelType>
+	void SetKeyTimes(ChannelType* InChannel, TArrayView<const FKeyHandle> InHandles, TArrayView<const FFrameNumber> InKeyTimes)
+	{
+		check(InHandles.Num() == InKeyTimes.Num());
+
+		auto ChannelInterface = InChannel->GetInterface();
+		for (int32 Index = 0; Index < InHandles.Num(); ++Index)
+		{
+			const int32 KeyIndex = ChannelInterface.GetIndex(InHandles[Index]);
+			if (KeyIndex != INDEX_NONE)
+			{
+				ChannelInterface.MoveKey(KeyIndex, InKeyTimes[Index]);
+			}
+		}
+	}
+
+
+	/**
+	 * Duplicate a number of keys within the specified channel
+	 *
+	 * @param InChannel          The channel to duplicate keys in
+	 * @param InHandles          Array of key handles that should be duplicated
+	 * @param OutNewHandles      Array view to receive key handles for each duplicated key. Must exactly mathc the size of InHandles.
+	 */
+	template<typename ChannelType>
+	void DuplicateKeys(ChannelType* InChannel, TArrayView<const FKeyHandle> InHandles, TArrayView<FKeyHandle> OutNewHandles)
+	{
+		auto ChannelInterface = InChannel->GetInterface();
+		for (int32 Index = 0; Index < InHandles.Num(); ++Index)
+		{
+			const int32 KeyIndex = ChannelInterface.GetIndex(InHandles[Index]);
+			if (KeyIndex == INDEX_NONE)
+			{
+				// we must add a handle even if the supplied handle does not relate to a key in this channel
+				OutNewHandles[Index] = FKeyHandle::Invalid();
+			}
+			else
+			{
+				// Do not cache value and time arrays since they can be reallocated during this loop
+				auto KeyCopy = ChannelInterface.GetValues()[KeyIndex];
+				int32 NewKeyIndex = ChannelInterface.AddKey(ChannelInterface.GetTimes()[KeyIndex], MoveTemp(KeyCopy));
+				OutNewHandles[Index] = ChannelInterface.GetHandle(NewKeyIndex);
+			}
+		}
+	}
+
+
+	/**
+	 * Delete a number of keys from the specified channel
+	 *
+	 * @param InChannel          The channel to delete keys in
+	 * @param InHandles          Array of key handles that should be deleted
+	 */
+	template<typename ChannelType>
+	void DeleteKeys(ChannelType* InChannel, TArrayView<const FKeyHandle> InHandles)
+	{
+		auto ChannelInterface = InChannel->GetInterface();
+		for (int32 Index = 0; Index < InHandles.Num(); ++Index)
+		{
+			const int32 KeyIndex = ChannelInterface.GetIndex(InHandles[Index]);
+			if (KeyIndex != INDEX_NONE)
+			{
+				ChannelInterface.RemoveKey(KeyIndex);
+			}
+		}
+	}
+
+
+	/**
+	 * Gather key draw information from a channel for a specific set of keys
+	 *
+	 * @param InChannel          The channel to duplicate keys in
+	 * @param InHandles          Array of key handles that should be deleted
+	 * @param OutKeyDrawParams   Array to receive key draw information. Must be exactly the size of InHandles.
+	 */
+	SEQUENCER_API void DrawKeys(void* Channel, TArrayView<const FKeyHandle> InHandles, TArrayView<FKeyDrawParams> OutKeyDrawParams);
+
+
+	/**
+	 * Copy the specified keys from a channel
+	 *
+	 * @param InChannel          The channel to duplicate keys in
+	 * @param InSection          The section that owns this channel
+	 * @param KeyAreaName        The name of the key area representing this channel
+	 * @param ClipboardBuilder   Structure for populating the clipboard
+	 * @param InHandles          Array of key handles that should be copied
+	 */
+	template<typename ChannelType>
+	void CopyKeys(ChannelType* InChannel, const UMovieSceneSection* InSection, FName KeyAreaName, FMovieSceneClipboardBuilder& ClipboardBuilder, TArrayView<const FKeyHandle> InHandles)
+	{
+		UMovieSceneTrack* Track = InSection ? InSection->GetTypedOuter<UMovieSceneTrack>() : nullptr;
+		if (!Track)
+		{
+			return;
+		}
+
+		FMovieSceneClipboardKeyTrack* KeyTrack = nullptr;
+
+		auto ChannelInterface = InChannel->GetInterface();
+		TArrayView<const FFrameNumber> Times = ChannelInterface.GetTimes();
+		auto Values = ChannelInterface.GetValues();
+
+		for (FKeyHandle Handle : InHandles)
+		{
+			const int32 KeyIndex = ChannelInterface.GetIndex(Handle);
+			if (KeyIndex != INDEX_NONE)
+			{
+				FFrameNumber KeyTime  = Times[KeyIndex];
+				auto         KeyValue = Values[KeyIndex];
+
+				if (!KeyTrack)
+				{
+					KeyTrack = &ClipboardBuilder.FindOrAddKeyTrack<decltype(KeyValue)>(KeyAreaName, *Track);
+				}
+
+				KeyTrack->AddKey(KeyTime, KeyValue);
+			}
+		}
+	}
+
+	/**
+	 * Paste the clipboard contents onto a channel
+	 *
+	 * @param InChannel          The channel to duplicate keys in
+	 * @param Section            The section that owns this channel
+	 * @param KeyTrack           The clipboard track to paste
+	 * @param SrcEnvironment     The source clipboard environment that was originally copied
+	 * @param DstEnvironment     The destination clipboard environment that we're copying to
+	 * @param OutPastedKeys      Array of key handles that should receive any pasted keys
+	 */
+	template<typename ChannelType>
+	void PasteKeys(ChannelType* InChannel, UMovieSceneSection* Section, const FMovieSceneClipboardKeyTrack& KeyTrack, const FMovieSceneClipboardEnvironment& SrcEnvironment, const FSequencerPasteEnvironment& DstEnvironment, TArray<FKeyHandle>& OutPastedKeys)
+	{
+		if (!Section || !Section->TryModify())
+		{
+			return;
+		}
+
+		FFrameTime PasteAt = DstEnvironment.CardinalTime;
+
+		auto ChannelInterface = InChannel->GetInterface();
+
+		TRange<FFrameNumber> NewRange = Section->GetRange();
+
+		auto ForEachKey = [Section, PasteAt, &NewRange, &ChannelInterface, &OutPastedKeys, &SrcEnvironment, &DstEnvironment](const FMovieSceneClipboardKey& Key)
+		{
+			FFrameNumber Time = (PasteAt + FFrameRate::TransformTime(Key.GetTime(), SrcEnvironment.FrameResolution, DstEnvironment.FrameResolution)).FloorToFrame();
+
+			NewRange = TRange<FFrameNumber>::Hull(NewRange, TRange<FFrameNumber>(Time));
+
+			typedef typename TDecay<decltype(ChannelInterface.GetValues()[0])>::Type KeyType;
+			KeyType NewKey = Key.GetValue<KeyType>();
+
+			FKeyHandle KeyHandle = ChannelInterface.UpdateOrAddKey(Time, NewKey);
+			OutPastedKeys.Add(KeyHandle);
+			return true;
+		};
+		KeyTrack.IterateKeys(ForEachKey);
+
+		Section->SetRange(NewRange);
+	}
+
+
+
+	/**
+	 * Create a new model for the specified channel that can be used on the curve editor interface
+	 *
+	 * @return (Optional) A new model to be added to a curve editor
+	 */
+	template<typename ChannelType>
+	TUniquePtr<FCurveModel> CreateCurveEditorModel(TMovieSceneChannelHandle<ChannelType> ChannelHandle, UMovieSceneSection* OwningSection, TSharedRef<ISequencer> InSequencer)
+	{
+		return nullptr;
+	}
+
+}	// namespace Sequencer
+
+
+// include generic key struct function definitions
+#include "SequencerGenericKeyStruct.inl"
\ No newline at end of file
diff --git a/Engine/Source/Editor/Sequencer/Public/SequencerCommands.h b/Engine/Source/Editor/Sequencer/Public/SequencerCommands.h
--- a/Engine/Source/Editor/Sequencer/Public/SequencerCommands.h
+++ b/Engine/Source/Editor/Sequencer/Public/SequencerCommands.h
@@ -91,9 +91,6 @@ public:
 	/** Toggle locking the playback range. */
 	TSharedPtr< FUICommandInfo > TogglePlaybackRangeLocked;
 
-	/** Forces playback both in editor and at runtime to be evaluated at fixed frame intervals. */
-	TSharedPtr< FUICommandInfo > ToggleForceFixedFrameIntervalPlayback;
-
 	/** Reruns construction scripts on bound actors every frame. */
 	TSharedPtr< FUICommandInfo > ToggleRerunConstructionScripts;
 
@@ -118,6 +115,9 @@ public:
 	/** Expand/collapse nodes and descendants */
 	TSharedPtr< FUICommandInfo > ToggleExpandCollapseNodesAndDescendants;
 
+	/** Sort all nodes and descendants */
+	TSharedPtr< FUICommandInfo > SortAllNodesAndDescendants;
+
 	/** Sets the upper bound of the selection range */
 	TSharedPtr< FUICommandInfo > SetSelectionRangeEnd;
 
@@ -196,8 +196,8 @@ public:
 	/** Turns key all on and off. */
 	TSharedPtr< FUICommandInfo > ToggleKeyAllEnabled;
 
-	/** Turns show frame numbers on and off. */
-	TSharedPtr< FUICommandInfo > ToggleShowFrameNumbers;
+	/** Rotates through the supported formats for displaying times/frames/timecode. */
+	TSharedPtr< FUICommandInfo > ChangeTimeDisplayFormat;
 
 	/** Toggle the visibility of the goto box. */
 	TSharedPtr< FUICommandInfo > ToggleShowGotoBox;
@@ -304,9 +304,6 @@ public:
 	/** Rebinds all possessable references with their current bindings. */
 	TSharedPtr< FUICommandInfo > RebindPossessableReferences;
 
-	/** Attempts to move all time data for this sequence on to a valid frame */
-	TSharedPtr< FUICommandInfo > FixFrameTiming;
-
 	/** Record the selected actors into a sub sequence of the currently active sequence */
 	TSharedPtr< FUICommandInfo > RecordSelectedActors;
 
@@ -322,6 +319,10 @@ public:
 	/** Toggle whether we should evaluate sub sequences in isolation */
 	TSharedPtr< FUICommandInfo > ToggleEvaluateSubSequencesInIsolation;
 
+	TSharedPtr<FUICommandInfo> SetAllCurveVisibility;
+	TSharedPtr<FUICommandInfo> SetSelectedCurveVisibility;
+	TSharedPtr<FUICommandInfo> SetAnimatedCurveVisibility;
+
 	/**
 	 * Initialize commands
 	 */
diff --git a/Engine/Source/Editor/Sequencer/Public/SequencerGenericKeyStruct.h b/Engine/Source/Editor/Sequencer/Public/SequencerGenericKeyStruct.h
new file mode 100644
--- /dev/null
+++ b/Engine/Source/Editor/Sequencer/Public/SequencerGenericKeyStruct.h
@@ -0,0 +1,64 @@
+// Copyright 1998-2018 Epic Games, Inc. All Rights Reserved.
+
+#pragma once
+
+#include "MovieSceneKeyStruct.h"
+#include "SequencerGenericKeyStruct.generated.h"
+
+class IDetailLayoutBuilder;
+
+struct IMovieSceneKeyStructCustomization
+{
+	virtual ~IMovieSceneKeyStructCustomization() {}
+
+	/** Extend the specified (empty) details customization with the specified key handle */
+	virtual void Extend(IDetailLayoutBuilder& DetailBuilder) = 0;
+	virtual void Apply(FFrameNumber Time) = 0;
+};
+
+template<typename ChannelType>
+struct TMovieSceneKeyStructCustomization : IMovieSceneKeyStructCustomization
+{
+	TMovieSceneKeyStructCustomization(TMovieSceneChannelHandle<ChannelType> InChannel, FKeyHandle InHandle)
+		: KeyHandle(InHandle), ChannelHandle(InChannel)
+	{}
+
+	TSharedPtr<FStructOnScope> GetValueStruct();
+
+	/** Extend the specified (empty) details customization with the specified key handle */
+	virtual void Extend(IDetailLayoutBuilder& DetailBuilder) override;
+
+	virtual void Apply(FFrameNumber NewTime) override;
+
+private:
+
+	FKeyHandle KeyHandle;
+	TMovieSceneChannelHandle<ChannelType> ChannelHandle;
+	TSharedPtr<FStructOnScope> KeyStruct;
+};
+
+USTRUCT()
+struct SEQUENCER_API FSequencerGenericKeyStruct : public FMovieSceneKeyStruct
+{
+	GENERATED_BODY()
+
+	/** Cusomtization implementation that adds the key value data */
+	TSharedPtr<IMovieSceneKeyStructCustomization> CustomizationImpl;
+
+	/** This key's time */
+	UPROPERTY(EditAnywhere, Category="Key")
+	FFrameNumber Time;
+
+	/**
+	 * Propagate changes from this key structure to the corresponding key values.
+	 *
+	 * @param ChangeEvent The property change event.
+	 */
+	virtual void PropagateChanges(const FPropertyChangedEvent& ChangeEvent)
+	{
+		if (CustomizationImpl.IsValid())
+		{
+			CustomizationImpl->Apply(Time);
+		}
+	}
+};
\ No newline at end of file
diff --git a/Engine/Source/Editor/Sequencer/Public/SequencerGenericKeyStruct.inl b/Engine/Source/Editor/Sequencer/Public/SequencerGenericKeyStruct.inl
new file mode 100644
--- /dev/null
+++ b/Engine/Source/Editor/Sequencer/Public/SequencerGenericKeyStruct.inl
@@ -0,0 +1,72 @@
+// Copyright 1998-2018 Epic Games, Inc. All Rights Reserved.
+
+#pragma once
+
+#include "SequencerGenericKeyStruct.h"
+#include "DetailLayoutBuilder.h"
+#include "DetailCategoryBuilder.h"
+
+template<typename ChannelType>
+TSharedPtr<FStructOnScope> TMovieSceneKeyStructCustomization<ChannelType>::GetValueStruct()
+{
+	ChannelType* Channel = ChannelHandle.Get();
+	if (Channel)
+	{
+		auto ChannelInterface = Channel->GetInterface();
+		const int32 KeyIndex = ChannelInterface.GetIndex(KeyHandle);
+
+		if (KeyIndex != INDEX_NONE)
+		{
+			auto& Value = ChannelInterface.GetValues()[KeyIndex];
+
+			TSharedPtr<FStructOnScope> NewStruct = MakeShared<FStructOnScope>(Value.StaticStruct());
+
+			// Assign over the existing value
+			typedef typename TDecay<decltype(Value)>::Type ValueType;
+			*((ValueType*)NewStruct->GetStructMemory()) = Value;
+
+			return NewStruct;
+		}
+	}
+	return nullptr;
+}
+
+template<typename ChannelType>
+void TMovieSceneKeyStructCustomization<ChannelType>::Extend(IDetailLayoutBuilder& DetailBuilder)
+{
+	using namespace Sequencer;
+	if (!KeyStruct.IsValid())
+	{
+		KeyStruct = GetValueStruct();
+	}
+
+	if (KeyStruct.IsValid())
+	{
+		DetailBuilder.EditCategory("Key").AddAllExternalStructureProperties(KeyStruct.ToSharedRef());
+	}
+}
+
+template<typename ChannelType>
+void TMovieSceneKeyStructCustomization<ChannelType>::Apply(FFrameNumber NewTime)
+{
+	ChannelType* Channel = ChannelHandle.Get();
+	if (Channel)
+	{
+		auto ChannelInterface = Channel->GetInterface();
+		const int32 KeyIndex = ChannelInterface.GetIndex(KeyHandle);
+		if (KeyIndex != INDEX_NONE)
+		{
+			if (KeyStruct.IsValid())
+			{
+				auto& Value = ChannelInterface.GetValues()[KeyIndex];
+
+				// Assign over the new value
+				typedef typename TDecay<decltype(Value)>::Type ValueType;
+				Value = *((ValueType*)KeyStruct->GetStructMemory());
+			}
+
+			// Move the key to the new time
+			ChannelInterface.MoveKey(KeyIndex, NewTime);
+		}
+	}
+}
\ No newline at end of file
diff --git a/Engine/Source/Editor/Sequencer/Public/SequencerKeyCollection.h b/Engine/Source/Editor/Sequencer/Public/SequencerKeyCollection.h
new file mode 100644
--- /dev/null
+++ b/Engine/Source/Editor/Sequencer/Public/SequencerKeyCollection.h
@@ -0,0 +1,127 @@
+// Copyright 1998-2018 Epic Games, Inc. All Rights Reserved.
+
+#pragma once
+
+#include "CoreTypes.h"
+#include "Misc/Guid.h"
+#include "Templates/SharedPointer.h"
+#include "FrameNumber.h"
+
+class UMovieSceneSection;
+class FSequencerDisplayNode;
+class IKeyArea;
+
+template<typename> struct TOptional;
+template<typename> class TRange;
+template<typename> class TArrayView;
+
+/** Enumeration used to define how to search for keys */
+enum class EFindKeyDirection
+{
+	Backwards, Forwards
+};
+
+struct FSequencerKeyCollectionSignature
+{
+	FSequencerKeyCollectionSignature()
+	{}
+
+	/** Initialize this key collection from the specified nodes. Only gathers keys from those explicitly specified. */
+	SEQUENCER_API static FSequencerKeyCollectionSignature FromNodes(const TArray<FSequencerDisplayNode*>& InNodes, FFrameNumber InDuplicateThreshold);
+
+	/** Initialize this key collection from the specified nodes. Gathers keys from all child nodes. */
+	SEQUENCER_API static FSequencerKeyCollectionSignature FromNodesRecursive(const TArray<FSequencerDisplayNode*>& InNodes, FFrameNumber InDuplicateThreshold);
+
+	/** Initialize this key collection from the specified node and section index. */
+	SEQUENCER_API static FSequencerKeyCollectionSignature FromNodeRecursive(FSequencerDisplayNode& InNode, UMovieSceneSection* InSection, FFrameNumber InDuplicateThreshold);
+
+	/** Compare this signature for inequality with another */
+	SEQUENCER_API friend bool operator!=(const FSequencerKeyCollectionSignature& A, const FSequencerKeyCollectionSignature& B);
+
+	/** Compare this signature for equality with another */
+	SEQUENCER_API friend bool operator==(const FSequencerKeyCollectionSignature& A, const FSequencerKeyCollectionSignature& B);
+
+	/** Access the map of keyareas and signatures this signature was generated for */
+	const TMap<TSharedRef<IKeyArea>, FGuid>& GetKeyAreas() const
+	{
+		return KeyAreaToSignature;
+	}
+
+	/** Access duplicate threshold that this signature was generated for */
+	FFrameNumber GetDuplicateThreshold() const
+	{
+		return DuplicateThresholdTime;
+	}
+
+private:
+
+	/** Check whether this signature contains content that cannot be cached (such content causes this signature to never compare equal with another) */
+	bool HasUncachableContent() const;
+
+	/** The time at which proximal keys are considered duplicates */
+	FFrameNumber DuplicateThresholdTime;
+
+	/** Map of key areas to the section signature with with this signature was generated */
+	TMap<TSharedRef<IKeyArea>, FGuid> KeyAreaToSignature;
+};
+
+/**
+ * A collection of keys gathered recursively from a particular node or nodes
+ */
+class FSequencerKeyCollection
+{
+public:
+
+	/**
+	 * Search forwards or backwards for the first key within the specified range
+	 *
+	 * @param Range      The range to search within
+	 * @param Direction  Whether to return the first or last key that reside in the given range
+	 * @return (Optional) the time of the key that matched the range
+	 */
+	SEQUENCER_API TOptional<FFrameNumber> FindFirstKeyInRange(const TRange<FFrameNumber>& Range, EFindKeyDirection Direction) const;
+
+	/**
+	 * Get a view of all key times that reside within the specified range
+	 *
+	 * @param Range      The range to search within
+	 * @return A (possibly empty) array view of all the times that lie within the range
+	 */
+	SEQUENCER_API TArrayView<const FFrameNumber> GetKeysInRange(const TRange<FFrameNumber>& Range) const;
+
+	/**
+	* Search forwards or backwards for the next key from the specified frame number
+	* @param FrameNumber The frame number to search from
+	* @param Direction  Whether to return the next key or previous key from that time
+	* @return (Optional)  Frame number of the key that's next or previous from that time 
+	*/
+	SEQUENCER_API TOptional<FFrameNumber> GetNextKey(FFrameNumber FrameNumber, EFindKeyDirection Direction) const;
+
+	/**
+	 * Access the signature this collection was generated with
+	 *
+	 * @return The signature that this collection was generated with
+	 */
+	const FSequencerKeyCollectionSignature& GetSignature() const
+	{
+		return Signature;
+	}
+
+public:
+
+	/**
+	 * Update this key collection using the specified signature
+	 *
+	 * @param InSignature    The signature to generate keys for, containing all key areas to use for the generation
+	 * @return true if this collection was updated, or false if it was already up to date
+	 */
+	bool Update(const FSequencerKeyCollectionSignature& InSignature);
+
+private:
+
+	/** Times grouped by the supplied threshold */
+	TArray<FFrameNumber> GroupedTimes;
+
+	/** The signature with which the above array was generated */
+	FSequencerKeyCollectionSignature Signature;
+};
diff --git a/Engine/Source/Editor/Sequencer/Public/SequencerSettings.h b/Engine/Source/Editor/Sequencer/Public/SequencerSettings.h
--- a/Engine/Source/Editor/Sequencer/Public/SequencerSettings.h
+++ b/Engine/Source/Editor/Sequencer/Public/SequencerSettings.h
@@ -6,6 +6,7 @@
 #include "UObject/ObjectMacros.h"
 #include "UObject/Object.h"
 #include "UObject/Package.h"
+#include "FrameNumberDisplayFormat.h"
 #include "SequencerSettings.generated.h"
 
 struct FPropertyChangedEvent;
@@ -13,6 +14,20 @@ enum class EAutoChangeMode : uint8;
 enum class EAllowEditsMode : uint8;
 enum class EMovieSceneKeyInterpolation : uint8;
 
+
+
+/** Defines visibility states for the curves in the curve editor. */
+UENUM()
+enum class ECurveEditorCurveVisibility : uint8
+{
+	/** All curves should be visible. */
+	AllCurves,
+	/** Only curves from selected nodes should be visible. */
+	SelectedCurves,
+	/** Only curves which have keyframes should be visible. */
+	AnimatedCurves
+};
+
 UENUM()
 enum ESequencerSpawnPosition
 {
@@ -46,28 +61,6 @@ enum ESequencerLoopMode
 	SLM_LoopSelectionRange UMETA(DisplayName="Loop Selection Range"),
 };
 
-UENUM()
-enum ESequencerTimeSnapInterval
-{
-	STSI_0_001		UMETA(DisplayName="0.001s"),
-	STSI_0_01		UMETA(DisplayName="0.01s"),
-	STSI_0_1		UMETA(DisplayName="0.1s"),
-	STSI_1			UMETA(DisplayName="1s"),
-	STSI_10			UMETA(DisplayName="10s"),
-	STSI_100		UMETA(DisplayName="100s"),
-	STSI_15Fps		UMETA(DisplayName="15 fps"),
-	STSI_24Fps		UMETA(DisplayName="24 fps (film)"),
-	STSI_25Fps		UMETA(DisplayName="25 fps (PAL/25)"),
-	STSI_29_97Fps	UMETA(DisplayName="29.97 fps (NTSC/30)"),
-	STSI_30Fps		UMETA(DisplayName="30 fps"),
-	STSI_48Fps		UMETA(DisplayName="48 fps"),
-	STSI_50Fps		UMETA(DisplayName="50 fps (PAL/50)"),
-	STSI_59_94Fps	UMETA(DisplayName="59.94 fps (NTSC/60)"),
-	STSI_60Fps		UMETA(DisplayName="60 fps"),
-	STSI_120Fps		UMETA(DisplayName="120 fps"),
-	STSI_Custom		UMETA(DisplayName="Custom")
-};
-
 /** Empty class used to house multiple named USequencerSettings */
 UCLASS()
 class SEQUENCER_API USequencerSettingsContainer
@@ -112,7 +105,7 @@ public:
 	DECLARE_MULTICAST_DELEGATE( FOnEvaluateSubSequencesInIsolationChanged );
 	DECLARE_MULTICAST_DELEGATE( FOnShowSelectedNodesOnlyChanged );
 	DECLARE_MULTICAST_DELEGATE_OneParam( FOnAllowEditsModeChanged, EAllowEditsMode );
-	DECLARE_MULTICAST_DELEGATE_OneParam( FOnLockPlaybackToAudioClockChanged, bool );
+	DECLARE_MULTICAST_DELEGATE( FOnCurveEditorCurveVisibilityChanged );
 
 	virtual void PostEditChangeProperty(FPropertyChangedEvent& PropertyChangedEvent) override;
 
@@ -153,11 +146,6 @@ public:
 	/** Set whether to create spawnable cameras. */
 	void SetCreateSpawnableCameras(bool bInCreateSpawnableCameras);
 
-	/** Gets whether or not to show frame numbers. */
-	bool GetShowFrameNumbers() const;
-	/** Sets whether or not to show frame numbers. */
-	void SetShowFrameNumbers(bool InbShowFrameNumbers);
-
 	/** Gets whether or not to show the time range slider. */
 	bool GetShowRangeSlider() const;
 	/** Sets whether or not to show frame numbers. */
@@ -168,14 +156,6 @@ public:
 	/** Sets whether or not snapping is enabled. */
 	void SetIsSnapEnabled(bool InbIsSnapEnabled);
 
-	/** Gets the time in seconds used for interval snapping. */
-	float GetTimeSnapInterval() const;
-
-	/** Gets the custom time in seconds used for interval snapping. */
-	float GetCustomTimeSnapInterval() const;
-	/** Sets the custom time in seconds used for interval snapping. */
-	void SetCustomTimeSnapInterval(float InCustomTimeSnapInterval);
-
 	/** Gets whether or not to snap key times to the interval. */
 	bool GetSnapKeyTimesToInterval() const;
 	/** Sets whether or not to snap keys to the interval. */
@@ -251,11 +231,6 @@ public:
 	bool GetAutoScrollEnabled() const;
 	/** Sets whether or not auto-scroll is enabled. */
 	void SetAutoScrollEnabled(bool bInAutoScrollEnabled);
-
-	/** Gets whether or not to show curve tool tips in the curve editor. */
-	bool GetShowCurveEditorCurveToolTips() const;
-	/** Sets whether or not to show curve tool tips in the curve editor. */
-	void SetShowCurveEditorCurveToolTips(bool InbShowCurveEditorCurveToolTips);
 	
 	/** Gets whether or not to link the curve editor time range. */
 	bool GetLinkCurveEditorTimeRange() const;
@@ -342,11 +317,6 @@ public:
 	/** Snaps a time value in seconds to the currently selected interval. */
 	float SnapTimeToInterval(float InTimeValue) const;
 
-	/** @return true we're locking playback to the audio clock */
-	bool ShouldLockPlaybackToAudioClock() const;
-	/** Toggle whether to lock playback to the audio clock */
-	void SetLockPlaybackToAudioClock(bool bInLockPlaybackToAudioClock);
-
 	/** Check whether to show pre and post roll in sequencer */
 	bool ShouldShowPrePostRoll() const;
 	/** Toggle whether to show pre and post roll in sequencer */
@@ -354,9 +324,17 @@ public:
 
 	uint32 GetTrajectoryPathCap() const { return TrajectoryPathCap; }
 
-	/** Gets the multicast delegate which is invoked whenevcer the bLockPlaybackToAudioClock setting is changed. */
-	FOnLockPlaybackToAudioClockChanged& GetOnLockPlaybackToAudioClockChanged() { return OnLockPlaybackToAudioClockChanged; }
+	/** Gets the current curve visibility. */
+	ECurveEditorCurveVisibility GetCurveVisibility() const;
+	/** Sets the current curve visibility. */
+	void SetCurveVisibility(ECurveEditorCurveVisibility InCurveVisibility);
+
+	FOnCurveEditorCurveVisibilityChanged& GetOnCurveEditorCurveVisibilityChanged();
 
+	/** What format should we display the UI controls in when representing time in a sequence? */
+	EFrameNumberDisplayFormats GetTimeDisplayFormat() const { return FrameNumberDisplayFormat; }
+	/** Sets the time display format to the specified type. */
+	void SetTimeDisplayFormat(EFrameNumberDisplayFormats InFormat);
 protected:
 
 	/** The auto change mode (auto-key, auto-track or none). */
@@ -391,10 +369,6 @@ protected:
 	UPROPERTY( config, EditAnywhere, Category=General )
 	bool bCreateSpawnableCameras;
 
-	/** Show frame numbers or time in the timeline. */
-	UPROPERTY( config, EditAnywhere, Category=Timeline )
-	bool bShowFrameNumbers;
-
 	/** Show the in/out range in the timeline with respect to the start/end range. */
 	UPROPERTY( config, EditAnywhere, Category=Timeline )
 	bool bShowRangeSlider;
@@ -403,14 +377,6 @@ protected:
 	UPROPERTY( config, EditAnywhere, Category=Snapping )
 	bool bIsSnapEnabled;
 
-	/** The time snap interval mode */
-	UPROPERTY( config, EditAnywhere, Category=Snapping )
-	TEnumAsByte<ESequencerTimeSnapInterval> TimeSnapIntervalMode;
-
-	/** The custom time snapping interval in the timeline. Used if the time snap interval mode is set to Custom */
-	UPROPERTY( config, EditAnywhere, Category=Snapping )
-	float CustomTimeSnapInterval;
-
 	/** Enable or disable snapping keys to the time snapping interval. */
 	UPROPERTY( config, EditAnywhere, Category=Snapping )
 	bool bSnapKeyTimesToInterval;
@@ -470,10 +436,6 @@ protected:
 	UPROPERTY( config, EditAnywhere, Category=Timeline )
 	bool bAutoScrollEnabled;
 
-	/** Enable or disable curve editor tooltips. */
-	UPROPERTY( config, EditAnywhere, Category=CurveEditor )
-	bool bShowCurveEditorCurveToolTips;
-
 	/** Enable or disable linking the curve editor time range to the sequencer timeline's time range. */
 	UPROPERTY( config, EditAnywhere, Category=CurveEditor )
 	bool bLinkCurveEditorTimeRange;
@@ -514,10 +476,6 @@ protected:
 	UPROPERTY( config )
 	bool bShowViewportTransportControls;
 
-	/** When enabled, sequencer playback will be locked to the engine's audio clock */
-	UPROPERTY(config, EditAnywhere, Category=Playback)
-	bool bLockPlaybackToAudioClock;
-
 	/** When enabled, sequencer is able to possess viewports that represent PIE worlds */
 	UPROPERTY(config, EditAnywhere, Category=General)
 	bool bAllowPossessionOfPIEViewports;
@@ -546,7 +504,14 @@ protected:
 	UPROPERTY(config, EditAnywhere, Category=General)
 	uint32 TrajectoryPathCap;
 
-	FOnLockPlaybackToAudioClockChanged OnLockPlaybackToAudioClockChanged;
+	/** What format do we display time in to the user? */	
+	UPROPERTY(config, EditAnywhere, Category=General)
+	EFrameNumberDisplayFormats FrameNumberDisplayFormat;
+
+	/** Specifies which curves to show in the curve editor */
+	UPROPERTY(config, EditAnywhere, Category=General)
+	ECurveEditorCurveVisibility CurveVisibility;
+	FOnCurveEditorCurveVisibilityChanged OnCurveEditorCurveVisibilityChanged;
 	FOnEvaluateSubSequencesInIsolationChanged OnEvaluateSubSequencesInIsolationChangedEvent;
 	FOnShowSelectedNodesOnlyChanged OnShowSelectedNodesOnlyChangedEvent;
 	FOnAllowEditsModeChanged OnAllowEditsModeChangedEvent;
diff --git a/Engine/Source/Editor/Sequencer/Sequencer.Build.cs b/Engine/Source/Editor/Sequencer/Sequencer.Build.cs
--- a/Engine/Source/Editor/Sequencer/Sequencer.Build.cs
+++ b/Engine/Source/Editor/Sequencer/Sequencer.Build.cs
@@ -14,6 +14,12 @@ public class Sequencer : ModuleRules
 	}
         );
 
+		PublicDependencyModuleNames.AddRange(
+			new string[] {
+				"TimeManagement",
+			}
+		);
+
 		PrivateDependencyModuleNames.AddRange(
 			new string[] {
 				"AppFramework", 
@@ -21,6 +27,7 @@ public class Sequencer : ModuleRules
                 "CinematicCamera",
 				"Core", 
 				"CoreUObject", 
+                "CurveEditor",
                 "InputCore",
 				"Engine", 
 				"Slate", 
@@ -64,6 +71,7 @@ public class Sequencer : ModuleRules
 			new string[] {
 				"PropertyEditor",
                 "SceneOutliner",
+                "CurveEditor",
             }
         );
 
