Fixed localization identity conflicts

diff --git a/Engine/Source/Developer/AssetTools/Private/AssetTypeActions/AssetTypeActions_MaterialFunctionInstance.cpp b/Engine/Source/Developer/AssetTools/Private/AssetTypeActions/AssetTypeActions_MaterialFunctionInstance.cpp
--- a/Engine/Source/Developer/AssetTools/Private/AssetTypeActions/AssetTypeActions_MaterialFunctionInstance.cpp
+++ b/Engine/Source/Developer/AssetTools/Private/AssetTypeActions/AssetTypeActions_MaterialFunctionInstance.cpp
@@ -21,8 +21,8 @@ void FAssetTypeActions_MaterialFunctionInstance::GetActions( const TArray<UObjec
 	FAssetTypeActions_MaterialFunction::GetActions(InObjects, MenuBuilder);
 
 	MenuBuilder.AddMenuEntry(
-		LOCTEXT("MaterialInstanceConstant_FindParent", "Find Parent"),
-		LOCTEXT("MaterialInstanceConstant_FindParentTooltip", "Finds the function this instance is based on in the content browser."),
+		LOCTEXT("MaterialFunctionInstance_FindParent", "Find Parent"),
+		LOCTEXT("MaterialFunctionInstance_FindParentTooltip", "Finds the function this instance is based on in the content browser."),
 		FSlateIcon(FEditorStyle::GetStyleSetName(), "ContentBrowser.AssetActions.GenericFind"),
 		FUIAction(
 			FExecuteAction::CreateSP( this, &FAssetTypeActions_MaterialFunctionInstance::ExecuteFindParent, MFIs ),
diff --git a/Engine/Source/Developer/AssetTools/Private/AssetTypeActions/AssetTypeActions_StaticMesh.cpp b/Engine/Source/Developer/AssetTools/Private/AssetTypeActions/AssetTypeActions_StaticMesh.cpp
--- a/Engine/Source/Developer/AssetTools/Private/AssetTypeActions/AssetTypeActions_StaticMesh.cpp
+++ b/Engine/Source/Developer/AssetTools/Private/AssetTypeActions/AssetTypeActions_StaticMesh.cpp
@@ -50,7 +50,7 @@ void FAssetTypeActions_StaticMesh::GetActions( const TArray<UObject*>& InObjects
 
 	MenuBuilder.AddMenuEntry(
 		NSLOCTEXT("AssetTypeActions_StaticMesh", "ObjectContext_ClearVertexColors", "Remove Vertex Colors"),
-		NSLOCTEXT("AssetTypeActions_StaticMesh", "ObjectContext_ClearVertexColors", "Removes vertex colors from all LODS in all selected meshes."),
+		NSLOCTEXT("AssetTypeActions_StaticMesh", "ObjectContext_ClearVertexColorsTooltip", "Removes vertex colors from all LODS in all selected meshes."),
 		FSlateIcon(),
 		FUIAction(
 			FExecuteAction::CreateSP(this, &FAssetTypeActions_StaticMesh::ExecuteRemoveVertexColors, Meshes)
diff --git a/Engine/Source/Developer/IOS/IOSPlatformEditor/Private/IOSTargetSettingsCustomization.cpp b/Engine/Source/Developer/IOS/IOSPlatformEditor/Private/IOSTargetSettingsCustomization.cpp
--- a/Engine/Source/Developer/IOS/IOSPlatformEditor/Private/IOSTargetSettingsCustomization.cpp
+++ b/Engine/Source/Developer/IOS/IOSPlatformEditor/Private/IOSTargetSettingsCustomization.cpp
@@ -1806,7 +1806,7 @@ void FIOSTargetSettingsCustomization::SetShaderStandard(int32 Value)
 		}
 		else if (Value == 3 && (EIOSVersion)EnumValue < EIOSVersion::IOS_11)
 		{
-			Message = LOCTEXT("iOSMetalShaderVersion1_2","Enabling Metal Shader Standard v2.0 increases the minimum operating system requirement for Metal from iOS 8.0/tvOS 9.0 or later to iOS/tvOS 11.0 or later.");
+			Message = LOCTEXT("iOSMetalShaderVersion2_0","Enabling Metal Shader Standard v2.0 increases the minimum operating system requirement for Metal from iOS 8.0/tvOS 9.0 or later to iOS/tvOS 11.0 or later.");
 			SetMinVersion((int32)EIOSVersion::IOS_11);
 		}
 		
diff --git a/Engine/Source/Developer/TaskGraph/Private/SProfileVisualizer.cpp b/Engine/Source/Developer/TaskGraph/Private/SProfileVisualizer.cpp
--- a/Engine/Source/Developer/TaskGraph/Private/SProfileVisualizer.cpp
+++ b/Engine/Source/Developer/TaskGraph/Private/SProfileVisualizer.cpp
@@ -239,7 +239,7 @@ FReply SProfileVisualizer::OnSaveClicked()
 	}
 	else
 	{
-		Message = NSLOCTEXT("TaskGraph", "ExportMessage", "Could not write profile data to file");
+		Message = NSLOCTEXT("TaskGraph", "FailedExportMessage", "Could not write profile data to file");
 	}
 
 	struct Local
diff --git a/Engine/Source/Editor/BlueprintGraph/Private/K2Node_FunctionTerminator.cpp b/Engine/Source/Editor/BlueprintGraph/Private/K2Node_FunctionTerminator.cpp
--- a/Engine/Source/Editor/BlueprintGraph/Private/K2Node_FunctionTerminator.cpp
+++ b/Engine/Source/Editor/BlueprintGraph/Private/K2Node_FunctionTerminator.cpp
@@ -130,7 +130,7 @@ void UK2Node_FunctionTerminator::ValidateNodeDuringCompilation(FCompilerResultsL
 		if (Pin && Pin->PinType.bIsWeakPointer && !Pin->PinType.IsContainer())
 		{
 			const FString ErrorString = FText::Format(
-				LOCTEXT("WeakPtrNotSupportedErrorFmt", "Weak prointer is not supported as function parameter. Pin '{0}' @@"),
+				LOCTEXT("WeakPtrNotSupportedErrorFmt", "Weak pointers are not supported as function parameters. Pin '{0}' @@"),
 				FText::FromString(Pin->GetName())
 			).ToString();
 			MessageLog.Error(*ErrorString, this);
diff --git a/Engine/Source/Editor/CurveEditor/Private/CurveEditorCommands.cpp b/Engine/Source/Editor/CurveEditor/Private/CurveEditorCommands.cpp
--- a/Engine/Source/Editor/CurveEditor/Private/CurveEditorCommands.cpp
+++ b/Engine/Source/Editor/CurveEditor/Private/CurveEditorCommands.cpp
@@ -20,8 +20,8 @@ void FCurveEditorCommands::RegisterCommands()
 	UI_COMMAND(InterpolationCubicUser, "User", "Cubic interpolation - User flat tangents", EUserInterfaceActionType::ToggleButton, FInputChord(EKeys::Two));
 	UI_COMMAND(InterpolationCubicBreak, "Break", "Cubic interpolation - User broken tangents", EUserInterfaceActionType::ToggleButton, FInputChord(EKeys::Three));
 
-	UI_COMMAND(FlattenTangents, "Flatten", "Flatten tangents", EUserInterfaceActionType::Button, FInputChord());
-	UI_COMMAND(StraightenTangents, "Straighten", "Straighten tangents", EUserInterfaceActionType::Button, FInputChord());
+	UI_COMMAND(FlattenTangents, "Flatten", "Flatten Tangents", EUserInterfaceActionType::Button, FInputChord());
+	UI_COMMAND(StraightenTangents, "Straighten", "Straighten Tangents", EUserInterfaceActionType::Button, FInputChord());
 
 	UI_COMMAND(BakeCurve, "Bake", "Bake curve", EUserInterfaceActionType::Button, FInputChord());
 	UI_COMMAND(ReduceCurve, "Reduce", "Reduce curve", EUserInterfaceActionType::Button, FInputChord());
diff --git a/Engine/Source/Editor/DetailCustomizations/Private/MotionControllerPinFactory.cpp b/Engine/Source/Editor/DetailCustomizations/Private/MotionControllerPinFactory.cpp
--- a/Engine/Source/Editor/DetailCustomizations/Private/MotionControllerPinFactory.cpp
+++ b/Engine/Source/Editor/DetailCustomizations/Private/MotionControllerPinFactory.cpp
@@ -38,7 +38,7 @@ private:
 
 		if (!GraphPinObj->GetDefaultAsString().Equals(NewMotionSourceString))
 		{
-			const FScopedTransaction Transaction(NSLOCTEXT("GraphEditor", "ChangeStringPinValue", "Change Motion Source Pin Value"));
+			const FScopedTransaction Transaction(NSLOCTEXT("GraphEditor", "ChangeMotionSourcePinValue", "Change Motion Source Pin Value"));
 			GraphPinObj->Modify();
 
 			GraphPinObj->GetSchema()->TrySetDefaultValue(*GraphPinObj, NewMotionSourceString);
diff --git a/Engine/Source/Editor/MaterialEditor/Private/MaterialEditor.cpp b/Engine/Source/Editor/MaterialEditor/Private/MaterialEditor.cpp
--- a/Engine/Source/Editor/MaterialEditor/Private/MaterialEditor.cpp
+++ b/Engine/Source/Editor/MaterialEditor/Private/MaterialEditor.cpp
@@ -4365,10 +4365,10 @@ FGraphAppearanceInfo FMaterialEditor::GetGraphAppearance() const
 		switch (MaterialFunction->GetMaterialFunctionUsage())
 		{
 		case EMaterialFunctionUsage::MaterialLayer:
-			AppearanceInfo.CornerText = LOCTEXT("AppearanceCornerText_MaterialFunction", "MATERIAL LAYER");
+			AppearanceInfo.CornerText = LOCTEXT("AppearanceCornerText_MaterialLayer", "MATERIAL LAYER");
 			break;
 		case EMaterialFunctionUsage::MaterialLayerBlend:
-			AppearanceInfo.CornerText = LOCTEXT("AppearanceCornerText_MaterialFunction", "MATERIAL LAYER BLEND");
+			AppearanceInfo.CornerText = LOCTEXT("AppearanceCornerText_MaterialLayerBlend", "MATERIAL LAYER BLEND");
 			break;
 		default:
 			AppearanceInfo.CornerText = LOCTEXT("AppearanceCornerText_MaterialFunction", "MATERIAL FUNCTION");
diff --git a/Engine/Source/Editor/MaterialEditor/Private/MaterialEditorModule.cpp b/Engine/Source/Editor/MaterialEditor/Private/MaterialEditorModule.cpp
--- a/Engine/Source/Editor/MaterialEditor/Private/MaterialEditorModule.cpp
+++ b/Engine/Source/Editor/MaterialEditor/Private/MaterialEditorModule.cpp
@@ -140,7 +140,7 @@ public:
 							FUIAction(FExecuteAction::CreateRaw(this, &FMaterialEditorModule::OpenSubstancePluginGetter)),
 							NAME_None,
 							NSLOCTEXT("MaterialEditor", "GetContentText_SubstanceToolbar", "Substance"),
-							NSLOCTEXT("MaterialEditor", "GetContentText_Substance", "Add Substance textures to the project"),
+							NSLOCTEXT("MaterialEditor", "GetContentText_SubstanceTooltip", "Add Substance textures to the project"),
 							FSlateIcon(FEditorStyle::GetStyleSetName(), "MaterialEditor.AddSubstanceSpecialToolbar")
 						);
 						ToolbarBuilder.EndSection();
diff --git a/Engine/Source/Editor/MergeActors/Private/MeshProxyTool/MeshProxyTool.cpp b/Engine/Source/Editor/MergeActors/Private/MeshProxyTool/MeshProxyTool.cpp
--- a/Engine/Source/Editor/MergeActors/Private/MeshProxyTool/MeshProxyTool.cpp
+++ b/Engine/Source/Editor/MergeActors/Private/MeshProxyTool/MeshProxyTool.cpp
@@ -35,7 +35,7 @@ TSharedRef<SWidget> FMeshProxyTool::GetWidget()
 	return ProxyDialog.ToSharedRef();
 }
 
-FText  FMeshProxyTool::GetTooltipText() const
+FText FMeshProxyTool::GetTooltipText() const
 {
 	return LOCTEXT("MeshProxyToolTooltip", "Harvest geometry from selected actors and merge them into single mesh.");
 }
@@ -162,7 +162,7 @@ bool FMeshProxyTool::CanMerge() const
 
 FText FThirdPartyMeshProxyTool::GetTooltipText() const
 {
-	return LOCTEXT("MeshProxyToolTooltip", "Harvest geometry selected meshes and merge and and simplify them as a single mesh.");
+	return LOCTEXT("ThirdPartyMeshProxyToolTooltip", "Harvest geometry selected meshes and merge and and simplify them as a single mesh.");
 }
 
 
diff --git a/Engine/Source/Editor/MovieSceneTools/Private/FrameNumberDetailsCustomization.cpp b/Engine/Source/Editor/MovieSceneTools/Private/FrameNumberDetailsCustomization.cpp
--- a/Engine/Source/Editor/MovieSceneTools/Private/FrameNumberDetailsCustomization.cpp
+++ b/Engine/Source/Editor/MovieSceneTools/Private/FrameNumberDetailsCustomization.cpp
@@ -10,6 +10,8 @@
 #include "PropertyCustomizationHelpers.h"
 #include "FrameRate.h"
 
+#define LOCTEXT_NAMESPACE "TimeManagement.FrameNumber"
+
 void FFrameNumberDetailsCustomization::CustomizeChildren(TSharedRef<IPropertyHandle> PropertyHandle, IDetailChildrenBuilder& ChildBuilder, IPropertyTypeCustomizationUtils& CustomizationUtils)
 {
 	// Check for Min/Max metadata on the property itself and we'll apply it to the child when changed.
@@ -43,12 +45,12 @@ void FFrameNumberDetailsCustomization::CustomizeChildren(TSharedRef<IPropertyHan
 
 	FrameNumberProperty = CustomizedProperties.FindChecked(GET_MEMBER_NAME_CHECKED(FFrameNumber, Value));
 
-	ChildBuilder.AddCustomRow(NSLOCTEXT("TimeManagement", "TimeLabel", "Time"))
+	ChildBuilder.AddCustomRow(LOCTEXT("TimeLabel", "Time"))
 		.NameContent()
 		[
 			SNew(STextBlock)
 			.Text(PropertyHandle->GetPropertyDisplayName())
-			.ToolTipText(NSLOCTEXT("TimeManagement", "TimeLabelTooltip", "Time field which takes timecode, frames and seconds formats."))
+			.ToolTipText(LOCTEXT("TimeLabelTooltip", "Time field which takes timecode, frames and seconds formats."))
 			.Font(CustomizationUtils.GetRegularFont())
 		]
 		.ValueContent()
@@ -81,3 +83,5 @@ void FFrameNumberDetailsCustomization::OnTimeTextCommitted(const FText& InText,
 		FrameNumberProperty->SetValue((int32)ClampedValue);
 	}
 }
+
+#undef LOCTEXT_NAMESPACE
diff --git a/Engine/Source/Editor/Persona/Private/PersonaMeshDetails.cpp b/Engine/Source/Editor/Persona/Private/PersonaMeshDetails.cpp
--- a/Engine/Source/Editor/Persona/Private/PersonaMeshDetails.cpp
+++ b/Engine/Source/Editor/Persona/Private/PersonaMeshDetails.cpp
@@ -1452,7 +1452,7 @@ void FPersonaMeshDetails::ApplyChanges()
 		// if LOD is imported, ask users if they want to regenerate or just leave it
 		if (bImportedLODs)
 		{
-			const FText Text(LOCTEXT("Warning_SimplygonApplyingToImportedMesh", "Some of LODs has been imported. Are you sure you'd like to apply mesh reduction to all LODs? This will destroy imported LOD."));
+			const FText Text(LOCTEXT("Warning_SimplygonApplyingToImportedMesh_All", "Some LODs have been imported. Are you sure you'd like to apply mesh reduction to all LODs? This will destroy imported LODs."));
 			EAppReturnType::Type Ret = FMessageDialog::Open(EAppMsgType::YesNo, Text);
 			if (Ret == EAppReturnType::Yes)
 			{
diff --git a/Engine/Source/Editor/Sequencer/Private/QualifiedFrameTimeDetailsCustomization.cpp b/Engine/Source/Editor/Sequencer/Private/QualifiedFrameTimeDetailsCustomization.cpp
--- a/Engine/Source/Editor/Sequencer/Private/QualifiedFrameTimeDetailsCustomization.cpp
+++ b/Engine/Source/Editor/Sequencer/Private/QualifiedFrameTimeDetailsCustomization.cpp
@@ -10,6 +10,8 @@
 #include "PropertyCustomizationHelpers.h"
 #include "FrameRate.h"
 
+#define LOCTEXT_NAMESPACE "TimeManagement.QualifiedFrameTime"
+
 void FQualifiedFrameTimeDetailsCustomization::CustomizeChildren(TSharedRef<IPropertyHandle> PropertyHandle, IDetailChildrenBuilder& ChildBuilder, IPropertyTypeCustomizationUtils& CustomizationUtils)
 {
 	TMap<FName, TSharedRef<IPropertyHandle>> CustomizedProperties;
@@ -25,12 +27,12 @@ void FQualifiedFrameTimeDetailsCustomization::CustomizeChildren(TSharedRef<IProp
 	}
 
 	TSharedRef<IPropertyHandle> FrameNumberProperty = CustomizedProperties.FindChecked(GET_MEMBER_NAME_CHECKED(FFrameNumber, Value));
-	ChildBuilder.AddCustomRow(NSLOCTEXT("TimeManagement", "TimeLabel", "Time"))
+	ChildBuilder.AddCustomRow(LOCTEXT("TimeLabel", "Time"))
 		.NameContent()
 		[
 			SNew(STextBlock)
-			.Text(NSLOCTEXT("TimeManagement", "TimeLabel", "Time"))
-			.ToolTipText(NSLOCTEXT("TimeManagement", "TimeLabelTooltip", "Time field which takes timecode, frames and time formats."))
+			.Text(LOCTEXT("TimeLabel", "Time"))
+			.ToolTipText(LOCTEXT("TimeLabelTooltip", "Time field which takes timecode, frames and time formats."))
 			.Font(CustomizationUtils.GetRegularFont())
 		]
 		.ValueContent()
@@ -78,3 +80,5 @@ void FQualifiedFrameTimeDetailsCustomization::OnTimeTextCommitted(const FText& I
 	// 	UE_LOG(LogTemp, Log, TEXT("Unknown"));
 	// }
 }
+
+#undef LOCTEXT_NAMESPACE
diff --git a/Engine/Source/Editor/Sequencer/Private/SSequencerPlayRateCombo.cpp b/Engine/Source/Editor/Sequencer/Private/SSequencerPlayRateCombo.cpp
--- a/Engine/Source/Editor/Sequencer/Private/SSequencerPlayRateCombo.cpp
+++ b/Engine/Source/Editor/Sequencer/Private/SSequencerPlayRateCombo.cpp
@@ -407,7 +407,7 @@ FText SSequencerPlayRateCombo::GetToolTipText() const
 		FText FrameResolutionText = FrameResolutionInfo ? FrameResolutionInfo->DisplayName : FText::Format(LOCTEXT("FrameResolutionFormat", "{0} ticks every second"), FrameResolution.AsDecimal());
 
 		return FocusedMovieScene->GetEvaluationType() == EMovieSceneEvaluationType::FrameLocked
-			? FText::Format(LOCTEXT("ToolTip_Format", "This sequence is locked at runtime to {0} and uses an underlying tick resolution of {1}."), PlayRateText, FrameResolutionText)
+			? FText::Format(LOCTEXT("ToolTip_Format_FrameLocked", "This sequence is locked at runtime to {0} and uses an underlying tick resolution of {1}."), PlayRateText, FrameResolutionText)
 			: FText::Format(LOCTEXT("ToolTip_Format", "This sequence is being presented as {0} and uses an underlying tick resolution of {1}."), PlayRateText, FrameResolutionText);
 	}
 
diff --git a/Engine/Source/Editor/UMGEditor/Private/WidgetBlueprintEditor.cpp b/Engine/Source/Editor/UMGEditor/Private/WidgetBlueprintEditor.cpp
--- a/Engine/Source/Editor/UMGEditor/Private/WidgetBlueprintEditor.cpp
+++ b/Engine/Source/Editor/UMGEditor/Private/WidgetBlueprintEditor.cpp
@@ -1006,7 +1006,7 @@ public:
 void GetBindableObjects(UWidgetTree* WidgetTree, TArray<FObjectAndDisplayName>& BindableObjects)
 {
 	// Add the 'this' widget so you can animate it.
-	BindableObjects.Add(FObjectAndDisplayName(LOCTEXT("RootWidgetFormat", "[[This]]"), WidgetTree->GetOuter()));
+	BindableObjects.Add(FObjectAndDisplayName(LOCTEXT("RootWidgetThis", "[[This]]"), WidgetTree->GetOuter()));
 
 	WidgetTree->ForEachWidget([&BindableObjects] (UWidget* Widget) {
 		
diff --git a/Engine/Source/Editor/UnrealEd/Private/EditorExporters.cpp b/Engine/Source/Editor/UnrealEd/Private/EditorExporters.cpp
--- a/Engine/Source/Editor/UnrealEd/Private/EditorExporters.cpp
+++ b/Engine/Source/Editor/UnrealEd/Private/EditorExporters.cpp
@@ -1581,7 +1581,7 @@ bool ULevelExporterFBX::ExportBinary( UObject* Object, const TCHAR* Type, FArchi
 			FNotificationInfo NotificationInfo(FText::GetEmpty());
 			NotificationInfo.Text = bSelectedOnly
 				? FText(NSLOCTEXT("UnrealEd", "ExportingLevelToFBX_Selection", "The selection has nothing that can be export to fbx!"))
-				: FText(NSLOCTEXT("UnrealEd", "ExportingLevelToFBX", "The world has nothing that can be export to fbx!"));
+				: FText(NSLOCTEXT("UnrealEd", "ExportingLevelToFBX_World", "The world has nothing that can be export to fbx!"));
 			NotificationInfo.ExpireDuration = 5.0f;
 			FSlateNotificationManager::Get().AddNotification(NotificationInfo);
 		}
diff --git a/Engine/Source/Editor/UnrealEd/Private/StaticLightingSystem/StaticLightingSystem.cpp b/Engine/Source/Editor/UnrealEd/Private/StaticLightingSystem/StaticLightingSystem.cpp
--- a/Engine/Source/Editor/UnrealEd/Private/StaticLightingSystem/StaticLightingSystem.cpp
+++ b/Engine/Source/Editor/UnrealEd/Private/StaticLightingSystem/StaticLightingSystem.cpp
@@ -1883,7 +1883,7 @@ bool FStaticLightingSystem::CreateLightmassProcessor()
 		UE_LOG(LogStaticLightingSystem, Warning, TEXT("Failed to connect to Swarm."));
 		FMessageDialog::Open( EAppMsgType::Ok, 
 #if USE_LOCAL_SWARM_INTERFACE
-			LOCTEXT("FailedToConnectToSwarmDialogMessage", "Failed to connect to Swarm. Check that your network interface supports multicast.")
+			LOCTEXT("FailedToConnectToSwarmDialogMessage_CheckNetwork", "Failed to connect to Swarm. Check that your network interface supports multicast.")
 #else
 			LOCTEXT("FailedToConnectToSwarmDialogMessage", "Failed to connect to Swarm.")
 #endif	
diff --git a/Engine/Source/Editor/WorldBrowser/Private/Tiles/SWorldTileItem.cpp b/Engine/Source/Editor/WorldBrowser/Private/Tiles/SWorldTileItem.cpp
--- a/Engine/Source/Editor/WorldBrowser/Private/Tiles/SWorldTileItem.cpp
+++ b/Engine/Source/Editor/WorldBrowser/Private/Tiles/SWorldTileItem.cpp
@@ -336,11 +336,11 @@ FText SWorldTileItem::GetPositionText() const
 	FTextFormat TextFormat;
 	if (bLocked)
 	{
-		TextFormat = LOCTEXT("PositionXYFmtLocked", "{0}, {1}, {2} (Locked)");
+		TextFormat = LOCTEXT("PositionXYZFmtLocked", "{0}, {1}, {2} (Locked)");
 	}
 	else
 	{
-		TextFormat = LOCTEXT("PositionXYFmt", "{0}, {1}, {2}");
+		TextFormat = LOCTEXT("PositionXYZFmt", "{0}, {1}, {2}");
 	}
 		
 	return FText::Format(TextFormat, FText::AsNumber(Position.X), FText::AsNumber(Position.Y), FText::AsNumber(Position.Z));
diff --git a/Engine/Source/Runtime/Apple/MetalRHI/Private/MetalShaders.cpp b/Engine/Source/Runtime/Apple/MetalRHI/Private/MetalShaders.cpp
--- a/Engine/Source/Runtime/Apple/MetalRHI/Private/MetalShaders.cpp
+++ b/Engine/Source/Runtime/Apple/MetalRHI/Private/MetalShaders.cpp
@@ -180,9 +180,9 @@ static MTLLanguageVersion ValidateVersion(uint8 Version)
 #else
 		Args.Add(TEXT("RequiredOS"), FText::FromString(FString::Printf(TEXT("macOS %d.%d.%d"), MetaliOSVersions[Version][0], MetaliOSVersions[Version][1], MetaliOSVersions[Version][2])));
 #endif
-		FText LocalizedMsg = FText::Format(NSLOCTEXT("MetalRHI", "ShaderVersionUnsupported","The current OS version does not support {Version} required by the project. You must upgrade to {RequiredOS} to run this project."),Args);
+		FText LocalizedMsg = FText::Format(NSLOCTEXT("MetalRHI", "ShaderVersionUnsupported", "The current OS version does not support {Version} required by the project. You must upgrade to {RequiredOS} to run this project."),Args);
 		
-		FText Title = NSLOCTEXT("MetalRHI", "ShaderVersionUnsupported","Shader Version Unsupported");
+		FText Title = NSLOCTEXT("MetalRHI", "ShaderVersionUnsupportedTitle", "Shader Version Unsupported");
 		FMessageDialog::Open(EAppMsgType::Ok, LocalizedMsg, &Title);
 		
 		FPlatformMisc::RequestExit(true);
diff --git a/Engine/Source/Runtime/Renderer/Private/ReflectionEnvironmentCapture.cpp b/Engine/Source/Runtime/Renderer/Private/ReflectionEnvironmentCapture.cpp
--- a/Engine/Source/Runtime/Renderer/Private/ReflectionEnvironmentCapture.cpp
+++ b/Engine/Source/Runtime/Renderer/Private/ReflectionEnvironmentCapture.cpp
@@ -835,7 +835,7 @@ void BeginReflectionCaptureSlowTask(int32 NumCaptures, const TCHAR* CaptureReaso
 		
 		if (CaptureReason)
 		{
-			Status = FText::Format(NSLOCTEXT("Engine", "UpdateReflectionCaptures", "Building reflection captures for {0}"), FText::FromString(FString(CaptureReason)));
+			Status = FText::Format(NSLOCTEXT("Engine", "UpdateReflectionCapturesForX", "Building reflection captures for {0}"), FText::FromString(FString(CaptureReason)));
 		}
 		else
 		{
