Copying //UE4/Dev-Sequencer to //UE4/Dev-Main (Source: //UE4/Dev-Sequencer @ 4044339)

diff --git a/Engine/Source/Editor/SequenceRecorder/Private/SequenceRecorderModule.cpp b/Engine/Source/Editor/SequenceRecorder/Private/SequenceRecorderModule.cpp
--- a/Engine/Source/Editor/SequenceRecorder/Private/SequenceRecorderModule.cpp
+++ b/Engine/Source/Editor/SequenceRecorder/Private/SequenceRecorderModule.cpp
@@ -558,7 +558,7 @@ class FSequenceRecorderModule : public ISequenceRecorder, private FSelfRegisteri
 		UMovieScene*    MovieScene      = CurrentSequence ? CurrentSequence->GetMovieScene() : nullptr;
 		if (MovieScene)
 		{
-			return FQualifiedFrameTime(FFrameTime(MovieScene::DiscreteSize(MovieScene->GetPlaybackRange())), MovieScene->GetFrameResolution());
+			return FQualifiedFrameTime(FFrameTime(MovieScene::DiscreteSize(MovieScene->GetPlaybackRange())), MovieScene->GetTickResolution());
 		}
 		return FQualifiedFrameTime();
 	}
diff --git a/Engine/Source/Editor/SequenceRecorder/Public/SequenceRecorderSettings.h b/Engine/Source/Editor/SequenceRecorder/Public/SequenceRecorderSettings.h
--- a/Engine/Source/Editor/SequenceRecorder/Public/SequenceRecorderSettings.h
+++ b/Engine/Source/Editor/SequenceRecorder/Public/SequenceRecorderSettings.h
@@ -54,7 +54,7 @@ struct FSettingsForActorClass
 	TSubclassOf<AActor> Class;
 
 	/** Whether to record to 'possessable' (i.e. level-owned) or 'spawnable' (i.e. sequence-owned) actors. */
-	UPROPERTY(Config, EditAnywhere, AdvancedDisplay, Category = "Sequence Recording")
+	UPROPERTY(Config, EditAnywhere, Category = "Sequence Recording")
 	bool bRecordToPossessable;
 };
 
diff --git a/Engine/Source/Editor/SequenceRecorderSections/Private/MovieScene3DAttachSectionRecorder.cpp b/Engine/Source/Editor/SequenceRecorderSections/Private/MovieScene3DAttachSectionRecorder.cpp
--- a/Engine/Source/Editor/SequenceRecorderSections/Private/MovieScene3DAttachSectionRecorder.cpp
+++ b/Engine/Source/Editor/SequenceRecorderSections/Private/MovieScene3DAttachSectionRecorder.cpp
@@ -34,8 +34,8 @@ void FMovieScene3DAttachSectionRecorder::Record(float CurrentTime)
 	{
 		if(MovieSceneSection.IsValid())
 		{
-			FFrameRate FrameResolution = MovieSceneSection->GetTypedOuter<UMovieScene>()->GetFrameResolution();
-			FFrameNumber CurrentFrame = (CurrentTime * FrameResolution).FloorToFrame();
+			FFrameRate TickResolution = MovieSceneSection->GetTypedOuter<UMovieScene>()->GetTickResolution();
+			FFrameNumber CurrentFrame = (CurrentTime * TickResolution).FloorToFrame();
 
 			MovieSceneSection->ExpandToFrame(CurrentFrame);
 		}
@@ -60,8 +60,8 @@ void FMovieScene3DAttachSectionRecorder::Record(float CurrentTime)
 			{
 				MovieSceneSection = Cast<UMovieScene3DAttachSection>(AttachTrack->CreateNewSection());
 
-				FFrameRate FrameResolution = MovieSceneSection->GetTypedOuter<UMovieScene>()->GetFrameResolution();
-				FFrameNumber CurrentFrame = (CurrentTime * FrameResolution).FloorToFrame();
+				FFrameRate TickResolution = MovieSceneSection->GetTypedOuter<UMovieScene>()->GetTickResolution();
+				FFrameNumber CurrentFrame = (CurrentTime * TickResolution).FloorToFrame();
 
 				MovieSceneSection->SetRange(TRange<FFrameNumber>::Inclusive(CurrentFrame, CurrentFrame));
 				MovieSceneSection->SetConstraintId(Guid);
diff --git a/Engine/Source/Editor/SequenceRecorderSections/Private/MovieSceneParticleTrackSectionRecorder.cpp b/Engine/Source/Editor/SequenceRecorderSections/Private/MovieSceneParticleTrackSectionRecorder.cpp
--- a/Engine/Source/Editor/SequenceRecorderSections/Private/MovieSceneParticleTrackSectionRecorder.cpp
+++ b/Engine/Source/Editor/SequenceRecorderSections/Private/MovieSceneParticleTrackSectionRecorder.cpp
@@ -36,8 +36,8 @@ void FMovieSceneParticleTrackSectionRecorder::CreateSection(UObject* InObjectToR
 
 		ParticleTrack->AddSection(*MovieSceneSection);
 
-		FFrameRate FrameResolution = MovieSceneSection->GetTypedOuter<UMovieScene>()->GetFrameResolution();
-		FFrameNumber CurrentFrame = (Time * FrameResolution).FloorToFrame();
+		FFrameRate TickResolution = MovieSceneSection->GetTypedOuter<UMovieScene>()->GetTickResolution();
+		FFrameNumber CurrentFrame = (Time * TickResolution).FloorToFrame();
 		MovieSceneSection->SetRange(TRange<FFrameNumber>::Inclusive(CurrentFrame, CurrentFrame));
 
 		bWasTriggered = false;
@@ -59,8 +59,8 @@ void FMovieSceneParticleTrackSectionRecorder::Record(float CurrentTime)
 {
 	if(SystemToRecord.IsValid())
 	{
-		FFrameRate   FrameResolution = MovieSceneSection->GetTypedOuter<UMovieScene>()->GetFrameResolution();
-		FFrameNumber KeyTime         = (CurrentTime * FrameResolution).FloorToFrame();
+		FFrameRate   TickResolution  = MovieSceneSection->GetTypedOuter<UMovieScene>()->GetTickResolution();
+		FFrameNumber KeyTime         = (CurrentTime * TickResolution).FloorToFrame();
 
 		MovieSceneSection->ExpandToFrame(KeyTime);
 
diff --git a/Engine/Source/Editor/SequenceRecorderSections/Private/MovieSceneSpawnSectionRecorder.cpp b/Engine/Source/Editor/SequenceRecorderSections/Private/MovieSceneSpawnSectionRecorder.cpp
--- a/Engine/Source/Editor/SequenceRecorderSections/Private/MovieSceneSpawnSectionRecorder.cpp
+++ b/Engine/Source/Editor/SequenceRecorderSections/Private/MovieSceneSpawnSectionRecorder.cpp
@@ -36,8 +36,8 @@ void FMovieSceneSpawnSectionRecorder::CreateSection(UObject* InObjectToRecord, U
 		BoolChannel->SetDefault(false);
 		BoolChannel->GetInterface().AddKey(0, false);
 
-		FFrameRate FrameResolution = MovieSceneSection->GetTypedOuter<UMovieScene>()->GetFrameResolution();
-		FFrameNumber CurrentFrame = (Time * FrameResolution).FloorToFrame();
+		FFrameRate TickResolution = MovieSceneSection->GetTypedOuter<UMovieScene>()->GetTickResolution();
+		FFrameNumber CurrentFrame = (Time * TickResolution).FloorToFrame();
 		MovieSceneSection->SetRange(TRange<FFrameNumber>::Inclusive(CurrentFrame, CurrentFrame));
 	}
 
@@ -64,8 +64,8 @@ void FMovieSceneSpawnSectionRecorder::FinalizeSection()
 		{
 			double       OneFrameInterval = 1.0/GetDefault<USequenceRecorderSettings>()->DefaultAnimationSettings.SampleRate;
 
-			FFrameRate   FrameResolution  = MovieSceneSection->GetTypedOuter<UMovieScene>()->GetFrameResolution();
-			FFrameNumber StartTime        = MovieSceneSection->GetExclusiveEndFrame() - (OneFrameInterval * FrameResolution).CeilToFrame();
+			FFrameRate   TickResolution   = MovieSceneSection->GetTypedOuter<UMovieScene>()->GetTickResolution();
+			FFrameNumber StartTime        = MovieSceneSection->GetExclusiveEndFrame() - (OneFrameInterval * TickResolution).CeilToFrame();
 
 			Channel->GetInterface().AddKey(StartTime, true);
 			MovieSceneSection->SetStartFrame(StartTime);
@@ -77,8 +77,8 @@ void FMovieSceneSpawnSectionRecorder::Record(float CurrentTime)
 {
 	if(ObjectToRecord.IsValid())
 	{
-		FFrameRate FrameResolution = MovieSceneSection->GetTypedOuter<UMovieScene>()->GetFrameResolution();
-		MovieSceneSection->ExpandToFrame((CurrentTime * FrameResolution).FloorToFrame());
+		FFrameRate TickResolution = MovieSceneSection->GetTypedOuter<UMovieScene>()->GetTickResolution();
+		MovieSceneSection->ExpandToFrame((CurrentTime * TickResolution).FloorToFrame());
 	}
 
 	const bool bSpawned = ObjectToRecord.IsValid();
@@ -87,8 +87,8 @@ void FMovieSceneSpawnSectionRecorder::Record(float CurrentTime)
 		FMovieSceneBoolChannel* Channel = MovieSceneSection->GetChannelProxy().GetChannel<FMovieSceneBoolChannel>(0);
 		if (ensure(Channel))
 		{
-			FFrameRate   FrameResolution = MovieSceneSection->GetTypedOuter<UMovieScene>()->GetFrameResolution();
-			FFrameNumber KeyTime         = (CurrentTime * FrameResolution).FloorToFrame();
+			FFrameRate   TickResolution  = MovieSceneSection->GetTypedOuter<UMovieScene>()->GetTickResolution();
+			FFrameNumber KeyTime         = (CurrentTime * TickResolution).FloorToFrame();
 
 			Channel->GetInterface().UpdateOrAddKey(KeyTime, bSpawned);
 		}
diff --git a/Engine/Source/Editor/SequenceRecorderSections/Private/MovieSceneVisibilitySectionRecorder.cpp b/Engine/Source/Editor/SequenceRecorderSections/Private/MovieSceneVisibilitySectionRecorder.cpp
--- a/Engine/Source/Editor/SequenceRecorderSections/Private/MovieSceneVisibilitySectionRecorder.cpp
+++ b/Engine/Source/Editor/SequenceRecorderSections/Private/MovieSceneVisibilitySectionRecorder.cpp
@@ -69,8 +69,8 @@ void FMovieSceneVisibilitySectionRecorder::CreateSection(UObject* InObjectToReco
 			bWasVisible = !Actor->bHidden;
 		}
 
-		FFrameRate   FrameResolution = MovieSceneSection->GetTypedOuter<UMovieScene>()->GetFrameResolution();
-		FFrameNumber CurrentFrame    = (Time * FrameResolution).FloorToFrame();
+		FFrameRate   TickResolution = MovieSceneSection->GetTypedOuter<UMovieScene>()->GetTickResolution();
+		FFrameNumber CurrentFrame    = (Time * TickResolution).FloorToFrame();
 
 		if (ensure(Channel))
 		{
@@ -98,8 +98,8 @@ void FMovieSceneVisibilitySectionRecorder::Record(float CurrentTime)
 {
 	if(ObjectToRecord.IsValid())
 	{
-		FFrameRate   FrameResolution = MovieSceneSection->GetTypedOuter<UMovieScene>()->GetFrameResolution();
-		FFrameNumber CurrentFrame    = (CurrentTime * FrameResolution).FloorToFrame();
+		FFrameRate   TickResolution  = MovieSceneSection->GetTypedOuter<UMovieScene>()->GetTickResolution();
+		FFrameNumber CurrentFrame    = (CurrentTime * TickResolution).FloorToFrame();
 
 		MovieSceneSection->ExpandToFrame(CurrentFrame);
 
diff --git a/Engine/Source/Editor/Sequencer/Private/DisplayNodes/SequencerDisplayNode.cpp b/Engine/Source/Editor/Sequencer/Private/DisplayNodes/SequencerDisplayNode.cpp
--- a/Engine/Source/Editor/Sequencer/Private/DisplayNodes/SequencerDisplayNode.cpp
+++ b/Engine/Source/Editor/Sequencer/Private/DisplayNodes/SequencerDisplayNode.cpp
@@ -66,8 +66,8 @@ public:
 	SLATE_BEGIN_ARGS(SSequencerObjectTrack) {}
 		/** The view range of the section area */
 		SLATE_ATTRIBUTE( TRange<double>, ViewRange )
-		/** The frame resolution of the current sequence*/
-		SLATE_ATTRIBUTE( FFrameRate, FrameResolution )
+		/** The tick resolution of the current sequence*/
+		SLATE_ATTRIBUTE( FFrameRate, TickResolution )
 	SLATE_END_ARGS()
 
 	/** SLeafWidget Interface */
@@ -79,7 +79,7 @@ public:
 		RootNode = InRootNode;
 		
 		ViewRange = InArgs._ViewRange;
-		FrameResolution = InArgs._FrameResolution;
+		TickResolution = InArgs._TickResolution;
 
 		check(RootNode->GetType() == ESequencerNode::Object);
 	}
@@ -101,13 +101,13 @@ private:
 
 	/** The current view range */
 	TAttribute< TRange<double> > ViewRange;
-	/** The current frame resolution */
-	TAttribute< FFrameRate > FrameResolution;
+	/** The current tick resolution */
+	TAttribute< FFrameRate > TickResolution;
 
 	FSequencerKeyCollectionSignature KeyCollectionSignature;
 
-	/** The cached frame resolution these positions were generated with */
-	FFrameRate CachedFrameResolution;
+	/** The cached tick resolution these positions were generated with */
+	FFrameRate CachedTickResolution;
 	/** The time-range for which KeyDrawPositions was generated */
 	TRange<double> CachedViewRange;
 	/** Cached pixel positions for all keys in the current view range */
@@ -124,12 +124,12 @@ void SSequencerObjectTrack::Tick( const FGeometry& AllottedGeometry, const doubl
 	FSequencerKeyCollectionSignature NewCollectionSignature = FSequencerKeyCollectionSignature::FromNodesRecursive({ RootNode.Get() }, 0);
 
 	TRange<double> OldCachedViewRange = CachedViewRange;
-	FFrameRate OldCachedFrameResolution = CachedFrameResolution;
+	FFrameRate OldCachedTickResolution = CachedTickResolution;
 
 	CachedViewRange = ViewRange.Get();
-	CachedFrameResolution = FrameResolution.Get();
+	CachedTickResolution = TickResolution.Get();
 
-	if (NewCollectionSignature != KeyCollectionSignature || CachedViewRange != OldCachedViewRange || CachedFrameResolution != OldCachedFrameResolution)
+	if (NewCollectionSignature != KeyCollectionSignature || CachedViewRange != OldCachedViewRange || CachedTickResolution != OldCachedTickResolution)
 	{
 		KeyCollectionSignature = MoveTemp(NewCollectionSignature);
 		GenerateCachedKeyPositions(AllottedGeometry);
@@ -200,7 +200,7 @@ void SSequencerObjectTrack::GenerateCachedKeyPositions(const FGeometry& Allotted
 
 		// Generate a new cache
 		FSequencerCachedKeys TempCache;
-		TempCache.Update(KeyArea, CachedFrameResolution);
+		TempCache.Update(KeyArea, CachedTickResolution);
 
 		if (CacheKey.IsValid())
 		{
@@ -234,17 +234,17 @@ void SSequencerObjectTrack::GenerateCachedKeyPositions(const FGeometry& Allotted
 	for (auto& Pair : SectionToKeyTimeCache)
 	{
 		TArrayView<const double> Times;
-		Pair.Value.GetKeysInRange(CachedViewRange, &Times, nullptr);
+		Pair.Value.GetKeysInRange(CachedViewRange, &Times, nullptr, nullptr);
 		AllIterators.Add(Times);
 	}
 	for (auto& Uncached: UncachableKeyTimes)
 	{
 		TArrayView<const double> Times;
-		Uncached.GetKeysInRange(CachedViewRange, &Times, nullptr);
+		Uncached.GetKeysInRange(CachedViewRange, &Times, nullptr, nullptr);
 		AllIterators.Add(Times);
 	}
 
-	FTimeToPixel TimeToPixelConverter(AllottedGeometry, CachedViewRange, CachedFrameResolution);
+	FTimeToPixel TimeToPixelConverter(AllottedGeometry, CachedViewRange, CachedTickResolution);
 
 	// While any iterator is still valid, find and add the earliest time
 	while (AllIterators.ContainsByPredicate([](FIter& It){ return It; }))
@@ -494,7 +494,21 @@ FLinearColor FSequencerDisplayNode::GetDisplayNameColor() const
 			}
 			else
 			{
-				return EmptyNotActive;
+				if (TrackNode.GetSections().Num() > 0)
+				{
+					for (auto Section : TrackNode.GetSections())
+					{
+						if (Section->GetSectionObject() && Section->GetSectionObject()->IsActive())
+						{
+							// Stop traversing
+							return false;
+						}
+					}
+				}
+				else
+				{
+					return EmptyNotActive;
+				}
 			}
 		}
 		// Continue traversing
@@ -567,16 +581,16 @@ TSharedRef<SWidget> FSequencerDisplayNode::GenerateWidgetForSectionArea(const TA
 	{
 		return SNew(SSequencerObjectTrack, SharedThis(this))
 			.ViewRange(ViewRange)
-			.FrameResolution(this, &FSequencerDisplayNode::GetFrameResolution);
+			.TickResolution(this, &FSequencerDisplayNode::GetTickResolution);
 	}
 
 	// currently only section areas display widgets
 	return SNullWidget::NullWidget;
 }
 
-FFrameRate FSequencerDisplayNode::GetFrameResolution() const
+FFrameRate FSequencerDisplayNode::GetTickResolution() const
 {
-	return GetSequencer().GetFocusedFrameResolution();
+	return GetSequencer().GetFocusedTickResolution();
 }
 
 TSharedPtr<FSequencerDisplayNode> FSequencerDisplayNode::GetSectionAreaAuthority() const
diff --git a/Engine/Source/Editor/Sequencer/Private/DisplayNodes/SequencerDisplayNode.h b/Engine/Source/Editor/Sequencer/Private/DisplayNodes/SequencerDisplayNode.h
--- a/Engine/Source/Editor/Sequencer/Private/DisplayNodes/SequencerDisplayNode.h
+++ b/Engine/Source/Editor/Sequencer/Private/DisplayNodes/SequencerDisplayNode.h
@@ -262,7 +262,7 @@ public:
 	 */
 	TSharedPtr<FSequencerDisplayNode> GetSectionAreaAuthority() const;
 
-	FFrameRate GetFrameResolution() const;
+	FFrameRate GetTickResolution() const;
 
 	/**
 	 * @return the path to this node starting with the outermost parent
diff --git a/Engine/Source/Editor/Sequencer/Private/DisplayNodes/SequencerKeyTimeCache.cpp b/Engine/Source/Editor/Sequencer/Private/DisplayNodes/SequencerKeyTimeCache.cpp
--- a/Engine/Source/Editor/Sequencer/Private/DisplayNodes/SequencerKeyTimeCache.cpp
+++ b/Engine/Source/Editor/Sequencer/Private/DisplayNodes/SequencerKeyTimeCache.cpp
@@ -8,22 +8,23 @@
 void FSequencerCachedKeys::Update(TSharedRef<IKeyArea> InKeyArea, FFrameRate SourceResolution)
 {
 	UMovieSceneSection* Section = InKeyArea->GetOwningSection();
-	if (!Section || !CachedSignature.IsValid() || Section->GetSignature() != CachedSignature || SourceResolution != CachedFrameResolution)
+	if (!Section || !CachedSignature.IsValid() || Section->GetSignature() != CachedSignature || SourceResolution != CachedTickResolution)
 	{
 		CachedSignature = Section ? Section->GetSignature() : FGuid();
-		CachedFrameResolution = SourceResolution;
+		CachedTickResolution = SourceResolution;
+
+		CachedKeyFrames.Reset();
 
-		TArray<FFrameNumber> Times;
 		TArray<FKeyHandle> Handles;
-		InKeyArea->GetKeyInfo(&Handles, &Times);
+		InKeyArea->GetKeyInfo(&Handles, &CachedKeyFrames);
 
-		CachedKeyTimes.Reset(Times.Num());
-		CachedKeyHandles.Reset(Times.Num());
+		CachedKeyTimes.Reset(CachedKeyFrames.Num());
+		CachedKeyHandles.Reset(CachedKeyFrames.Num());
 
 		// Generate and cache
-		for (int32 Index = 0; Index < Times.Num(); ++Index)
+		for (int32 Index = 0; Index < CachedKeyFrames.Num(); ++Index)
 		{
-			CachedKeyTimes.Add(Times[Index] / SourceResolution);
+			CachedKeyTimes.Add(CachedKeyFrames[Index] / SourceResolution);
 			CachedKeyHandles.Add(Handles[Index]);
 		}
 
@@ -31,7 +32,7 @@ void FSequencerCachedKeys::Update(TSharedRef<IKeyArea> InKeyArea, FFrameRate Sou
 	}
 }
 
-void FSequencerCachedKeys::GetKeysInRange(const TRange<double>& Range, TArrayView<const double>* OutTimes, TArrayView<const FKeyHandle>* OutHandles) const
+void FSequencerCachedKeys::GetKeysInRange(const TRange<double>& Range, TArrayView<const double>* OutTimes, TArrayView<const FFrameNumber>* OutKeyFrames, TArrayView<const FKeyHandle>* OutHandles) const
 {
 	// Binary search the first time that's >= the lower bound
 	int32 FirstVisibleIndex = Algo::LowerBound(CachedKeyTimes, Range.GetLowerBoundValue());
@@ -46,6 +47,11 @@ void FSequencerCachedKeys::GetKeysInRange(const TRange<double>& Range, TArrayVie
 			*OutTimes = MakeArrayView(&CachedKeyTimes[FirstVisibleIndex], Num);
 		}
 
+		if (OutKeyFrames)
+		{
+			*OutKeyFrames = MakeArrayView(&CachedKeyFrames[FirstVisibleIndex], Num);
+		}
+
 		if (OutHandles)
 		{
 			*OutHandles = MakeArrayView(&CachedKeyHandles[FirstVisibleIndex], Num);
diff --git a/Engine/Source/Editor/Sequencer/Private/DisplayNodes/SequencerKeyTimeCache.h b/Engine/Source/Editor/Sequencer/Private/DisplayNodes/SequencerKeyTimeCache.h
--- a/Engine/Source/Editor/Sequencer/Private/DisplayNodes/SequencerKeyTimeCache.h
+++ b/Engine/Source/Editor/Sequencer/Private/DisplayNodes/SequencerKeyTimeCache.h
@@ -19,7 +19,7 @@ struct FSequencerCachedKeys
 	void Update(TSharedRef<IKeyArea> KeyArea, FFrameRate SourceResolution);
 
 	/** Get an view of the cached array for keys that fall within the specified range */
-	void GetKeysInRange(const TRange<double>& Range, TArrayView<const double>* OutTimes, TArrayView<const FKeyHandle>* OutHandles) const;
+	void GetKeysInRange(const TRange<double>& Range, TArrayView<const double>* OutTimes, TArrayView<const FFrameNumber>* OutKeyFrames, TArrayView<const FKeyHandle>* OutHandles) const;
 
 	/** Get the key area this cache was generated for, or nullptr if the cache has never been updated */
 	TSharedPtr<IKeyArea> GetKeyArea() const
@@ -30,13 +30,14 @@ struct FSequencerCachedKeys
 private:
 	/** Cached key information */
 	TArray<double> CachedKeyTimes;
+	TArray<FFrameNumber> CachedKeyFrames;
 	TArray<FKeyHandle> CachedKeyHandles;
 
 	/** The guid with which the above array was generated */
 	FGuid CachedSignature;
 
-	/** The frame resolution of the sequence that this cache was generated with */
-	FFrameRate CachedFrameResolution;
+	/** The tick resolution of the sequence that this cache was generated with */
+	FFrameRate CachedTickResolution;
 
 	/** The key area this cache is for */
 	TSharedPtr<IKeyArea> KeyArea;
diff --git a/Engine/Source/Editor/Sequencer/Private/DisplayNodes/SequencerObjectBindingNode.cpp b/Engine/Source/Editor/Sequencer/Private/DisplayNodes/SequencerObjectBindingNode.cpp
--- a/Engine/Source/Editor/Sequencer/Private/DisplayNodes/SequencerObjectBindingNode.cpp
+++ b/Engine/Source/Editor/Sequencer/Private/DisplayNodes/SequencerObjectBindingNode.cpp
@@ -39,6 +39,9 @@
 
 #include "Tracks/MovieSceneSpawnTrack.h"
 #include "Sections/MovieSceneSpawnSection.h"
+#include "LevelUtils.h"
+#include "Engine/LevelStreaming.h"
+#include "Engine/World.h"
 
 #define LOCTEXT_NAMESPACE "FObjectBindingNode"
 
@@ -161,15 +164,24 @@ void FSequencerObjectBindingNode::BuildContextMenu(FMenuBuilder& MenuBuilder)
 
 		if (Spawnable)
 		{
+			MenuBuilder.BeginSection("Spawnable", LOCTEXT("SpawnableMenuSectionName", "Spawnable"));
+	
 			MenuBuilder.AddSubMenu(
 				LOCTEXT("OwnerLabel", "Spawned Object Owner"),
 				LOCTEXT("OwnerTooltip", "Specifies how the spawned object is to be owned"),
 				FNewMenuDelegate::CreateSP(this, &FSequencerObjectBindingNode::AddSpawnOwnershipMenu)
 			);
 
-			MenuBuilder.AddMenuEntry( FSequencerCommands::Get().SaveCurrentSpawnableState );
+			MenuBuilder.AddSubMenu(
+				LOCTEXT("SubLevelLabel", "Spawnable Level"),
+				LOCTEXT("SubLevelTooltip", "Specifies which level the spawnable should be spawned into"),
+				FNewMenuDelegate::CreateSP(this, &FSequencerObjectBindingNode::AddSpawnLevelMenu)
+			);
 
+			MenuBuilder.AddMenuEntry( FSequencerCommands::Get().SaveCurrentSpawnableState );
 			MenuBuilder.AddMenuEntry( FSequencerCommands::Get().ConvertToPossessable );
+
+			MenuBuilder.EndSection();
 		}
 		else
 		{
@@ -188,6 +200,8 @@ void FSequencerObjectBindingNode::BuildContextMenu(FMenuBuilder& MenuBuilder)
 			MenuBuilder.AddMenuEntry( FSequencerCommands::Get().ConvertToSpawnable );
 		}
 
+		MenuBuilder.BeginSection("Import/Export", LOCTEXT("ImportExportMenuSectionName", "Import/Export"));
+		
 		MenuBuilder.AddMenuEntry(
 			LOCTEXT("Import FBX", "Import..."),
 			LOCTEXT("ImportFBXTooltip", "Import FBX animation to this object"),
@@ -211,6 +225,7 @@ void FSequencerObjectBindingNode::BuildContextMenu(FMenuBuilder& MenuBuilder)
 			FUIAction(
 				FExecuteAction::CreateLambda([=]{ GetSequencer().ExportToCameraAnim(); })
 			));
+		MenuBuilder.EndSection();
 	}
 
 	MenuBuilder.BeginSection("Organize", LOCTEXT("OrganizeContextMenuSectionName", "Organize"));
@@ -297,6 +312,51 @@ void FSequencerObjectBindingNode::AddSpawnOwnershipMenu(FMenuBuilder& MenuBuilde
 	);
 }
 
+
+void FSequencerObjectBindingNode::AddSpawnLevelMenu(FMenuBuilder& MenuBuilder)
+{
+	UMovieScene* MovieScene = GetSequencer().GetFocusedMovieSceneSequence()->GetMovieScene();
+	FMovieSceneSpawnable* Spawnable = MovieScene->FindSpawnable(ObjectBinding);
+	if (!Spawnable)
+	{
+		return;
+	}
+
+	MenuBuilder.AddMenuEntry(
+		NSLOCTEXT("UnrealEd", "PersistentLevel", "Persistent Level"),
+		NSLOCTEXT("UnrealEd", "PersistentLevel", "Persistent Level"),
+		FSlateIcon(),
+		FUIAction(
+			FExecuteAction::CreateLambda([=] { GetSequencer().SetSelectedNodesSpawnableLevel(NAME_None); }),
+			FCanExecuteAction(),
+			FIsActionChecked::CreateLambda([=] { return Spawnable->GetLevelName() == NAME_None; })
+		),
+		NAME_None,
+		EUserInterfaceActionType::ToggleButton
+	);
+
+	for (ULevelStreaming* LevelStreaming : GWorld->GetStreamingLevels())
+	{
+		if (LevelStreaming)
+		{
+			FName LevelName = FPackageName::GetShortFName( LevelStreaming->GetWorldAssetPackageFName() );
+
+			MenuBuilder.AddMenuEntry(
+				FText::FromName(LevelName),
+				FText::FromName(LevelName),
+				FSlateIcon(),
+				FUIAction(
+					FExecuteAction::CreateLambda([=] { GetSequencer().SetSelectedNodesSpawnableLevel(LevelName); }),
+					FCanExecuteAction(),
+					FIsActionChecked::CreateLambda([=] { return Spawnable->GetLevelName() == LevelName; })
+				),
+				NAME_None,
+				EUserInterfaceActionType::ToggleButton
+			);
+		}
+	}
+}
+
 bool FSequencerObjectBindingNode::CanRenameNode() const
 {
 	return true;
@@ -672,6 +732,8 @@ TSharedRef<SWidget> FSequencerObjectBindingNode::HandleAddTrackComboButtonGetMen
 	TSharedRef<FUICommandList> CommandList(new FUICommandList);
 	FMenuBuilder AddTrackMenuBuilder(true, nullptr, SequencerModule.GetAddTrackMenuExtensibilityManager()->GetAllExtenders(CommandList, TArrayBuilder<UObject*>().Add(BoundObject)));
 
+	const int32 NumStartingBlocks = AddTrackMenuBuilder.GetMultiBox()->GetBlocks().Num();
+
 	const UClass* ObjectClass = GetClassForObjectBinding();
 	AddTrackMenuBuilder.BeginSection(NAME_None, LOCTEXT("TracksMenuHeader" , "Tracks"));
 	GetSequencer().BuildObjectBindingTrackMenu(AddTrackMenuBuilder, ObjectBinding, ObjectClass);
@@ -760,6 +822,19 @@ TSharedRef<SWidget> FSequencerObjectBindingNode::HandleAddTrackComboButtonGetMen
 	}
 	AddTrackMenuBuilder.EndSection();
 
+	if (AddTrackMenuBuilder.GetMultiBox()->GetBlocks().Num() == NumStartingBlocks)
+	{
+		TSharedRef<SWidget> EmptyTip = SNew(SBox)
+			.Padding(FMargin(15.f, 7.5f))
+			[
+				SNew(STextBlock)
+				.Text(LOCTEXT("NoKeyablePropertiesFound", "No keyable properties or tracks"))
+				.ColorAndOpacity(FSlateColor::UseSubduedForeground())
+			];
+
+		AddTrackMenuBuilder.AddWidget(EmptyTip, FText(), true, false);
+	}
+
 	return AddTrackMenuBuilder.MakeWidget();
 }
 
@@ -771,7 +846,8 @@ void FSequencerObjectBindingNode::HandleAddTrackSubMenuNew(FMenuBuilder& AddTrac
 	// [PostProcessSettings] [ColorGrading]
 
 	// Create property menu data based on keyable property paths
-	TSet<UProperty*> PropertiesTraversed;
+	TArray<UProperty*> PropertiesTraversed;
+	TArray<int32> ArrayIndicesTraversed;
 	TArray<PropertyMenuData> KeyablePropertyMenuData;
 	for (const FPropertyPath& KeyablePropertyPath : KeyablePropertyPaths)
 	{
@@ -783,14 +859,37 @@ void FSequencerObjectBindingNode::HandleAddTrackSubMenuNew(FMenuBuilder& AddTrac
 		{
 			const FPropertyInfo& PropertyInfo = KeyablePropertyPath.GetPropertyInfo(1);
 			UProperty* Property = PropertyInfo.Property.Get();
-			if (PropertiesTraversed.Find(Property) != nullptr)
+
+			// Search for any array elements
+			int32 ArrayIndex = INDEX_NONE;
+			for (int32 PropertyInfoIndex = 0; PropertyInfoIndex < KeyablePropertyPath.GetNumProperties(); ++PropertyInfoIndex)
+			{
+				const FPropertyInfo& ArrayPropertyInfo = KeyablePropertyPath.GetPropertyInfo(PropertyInfoIndex);
+				if (ArrayPropertyInfo.ArrayIndex != INDEX_NONE)
+				{
+					ArrayIndex = ArrayPropertyInfo.ArrayIndex;
+					break;
+				}
+			}
+
+			bool bFound = false;
+			for (int32 TraversedIndex = 0; TraversedIndex < PropertiesTraversed.Num(); ++TraversedIndex)
+			{
+				if (PropertiesTraversed[TraversedIndex] == Property && ArrayIndicesTraversed[TraversedIndex] == ArrayIndex)
+				{
+					bFound = true;
+					break;
+				}
+			}
+
+			if (bFound)
 			{
 				continue;
 			}
 
-			if (PropertyInfo.ArrayIndex != INDEX_NONE)
+			if (ArrayIndex != INDEX_NONE)
 			{
-				KeyableMenuData.MenuName = FText::Format(LOCTEXT("ArrayElementFormat", "Element {0}"), FText::AsNumber(PropertyInfo.ArrayIndex)).ToString();
+				KeyableMenuData.MenuName = FText::Format(LOCTEXT("ArrayElementFormat", "{0} [{1}]"), Property->GetDisplayNameText(), FText::AsNumber(ArrayIndex)).ToString();
 			}
 			else
 			{
@@ -798,6 +897,7 @@ void FSequencerObjectBindingNode::HandleAddTrackSubMenuNew(FMenuBuilder& AddTrac
 			}
 
 			PropertiesTraversed.Add(Property);
+			ArrayIndicesTraversed.Add(ArrayIndex);
 		}
 		else
 		{
diff --git a/Engine/Source/Editor/Sequencer/Private/DisplayNodes/SequencerObjectBindingNode.h b/Engine/Source/Editor/Sequencer/Private/DisplayNodes/SequencerObjectBindingNode.h
--- a/Engine/Source/Editor/Sequencer/Private/DisplayNodes/SequencerObjectBindingNode.h
+++ b/Engine/Source/Editor/Sequencer/Private/DisplayNodes/SequencerObjectBindingNode.h
@@ -87,6 +87,7 @@ protected:
 	void AddPropertyMenuItems(FMenuBuilder& AddTrackMenuBuilder, TArray<FPropertyPath> KeyableProperties, int32 PropertyNameIndexStart = 0, int32 PropertyNameIndexEnd = -1);
 
 	void AddSpawnOwnershipMenu(FMenuBuilder& MenuBuilder);
+	void AddSpawnLevelMenu(FMenuBuilder& MenuBuilder);
 
 	/** Get class for object binding */
 	const UClass* GetClassForObjectBinding() const;
diff --git a/Engine/Source/Editor/Sequencer/Private/ISequencer.cpp b/Engine/Source/Editor/Sequencer/Private/ISequencer.cpp
--- a/Engine/Source/Editor/Sequencer/Private/ISequencer.cpp
+++ b/Engine/Source/Editor/Sequencer/Private/ISequencer.cpp
@@ -4,48 +4,48 @@
 #include "MovieSceneSequence.h"
 #include "MovieScene.h"
 
-FFrameRate ISequencer::GetRootFrameResolution() const
+FFrameRate ISequencer::GetRootTickResolution() const
 {
 	UMovieSceneSequence* RootSequence = GetRootMovieSceneSequence();
 	if (RootSequence)
 	{
-		return RootSequence->GetMovieScene()->GetFrameResolution();
+		return RootSequence->GetMovieScene()->GetTickResolution();
 	}
 
 	ensureMsgf(false, TEXT("No valid sequence found."));
 	return FFrameRate();
 }
 
-FFrameRate ISequencer::GetRootPlayRate() const
+FFrameRate ISequencer::GetRootDisplayRate() const
 {
 	UMovieSceneSequence* RootSequence = GetRootMovieSceneSequence();
 	if (RootSequence)
 	{
-		return RootSequence->GetMovieScene()->GetPlaybackFrameRate();
+		return RootSequence->GetMovieScene()->GetDisplayRate();
 	}
 
 	ensureMsgf(false, TEXT("No valid sequence found."));
 	return FFrameRate();
 }
 
-FFrameRate ISequencer::GetFocusedFrameResolution() const
+FFrameRate ISequencer::GetFocusedTickResolution() const
 {
 	UMovieSceneSequence* FocusedSequence = GetFocusedMovieSceneSequence();
 	if (FocusedSequence)
 	{
-		return FocusedSequence->GetMovieScene()->GetFrameResolution();
+		return FocusedSequence->GetMovieScene()->GetTickResolution();
 	}
 
 	ensureMsgf(false, TEXT("No valid sequence found."));
 	return FFrameRate();
 }
 
-FFrameRate ISequencer::GetFocusedPlayRate() const
+FFrameRate ISequencer::GetFocusedDisplayRate() const
 {
 	UMovieSceneSequence* FocusedSequence = GetFocusedMovieSceneSequence();
 	if (FocusedSequence)
 	{
-		return FocusedSequence->GetMovieScene()->GetPlaybackFrameRate();
+		return FocusedSequence->GetMovieScene()->GetDisplayRate();
 	}
 
 	ensureMsgf(false, TEXT("No valid sequence found."));
diff --git a/Engine/Source/Editor/Sequencer/Private/SAnimationOutlinerTreeNode.cpp b/Engine/Source/Editor/Sequencer/Private/SAnimationOutlinerTreeNode.cpp
--- a/Engine/Source/Editor/Sequencer/Private/SAnimationOutlinerTreeNode.cpp
+++ b/Engine/Source/Editor/Sequencer/Private/SAnimationOutlinerTreeNode.cpp
@@ -232,7 +232,7 @@ void SAnimationOutlinerTreeNode::Construct( const FArguments& InArgs, TSharedRef
 				.ContentPadding(0)
 				.VAlign(VAlign_Fill)
 				.HasDownArrow(false)
-				.IsFocusable(false)
+				.IsFocusable(true)
 				.IsEnabled(!DisplayNode->GetSequencer().IsReadOnly())
 				.ButtonStyle(FEditorStyle::Get(), "Sequencer.AnimationOutliner.ColorStrip")
 				.OnGetMenuContent(this, &SAnimationOutlinerTreeNode::OnGetColorPicker)
diff --git a/Engine/Source/Editor/Sequencer/Private/SSequencer.cpp b/Engine/Source/Editor/Sequencer/Private/SSequencer.cpp
--- a/Engine/Source/Editor/Sequencer/Private/SSequencer.cpp
+++ b/Engine/Source/Editor/Sequencer/Private/SSequencer.cpp
@@ -10,6 +10,7 @@
 #include "Misc/FeedbackContext.h"
 #include "Modules/ModuleManager.h"
 #include "Editor.h"
+#include "IDetailsView.h"
 #include "Widgets/Layout/SBorder.h"
 #include "ISequencerEditTool.h"
 #include "Layout/WidgetPath.h"
@@ -55,8 +56,8 @@
 #include "SequencerContextMenus.h"
 #include "Math/UnitConversion.h"
 #include "Widgets/Input/NumericUnitTypeInterface.inl"
+#include "FrameNumberDetailsCustomization.h"
 #include "SequencerSettings.h"
-#include "SSequencerGotoBox.h"
 #include "SSequencerTransformBox.h"
 #include "SSequencerDebugVisualizer.h"
 #include "ISequencerModule.h"
@@ -69,9 +70,16 @@
 #include "MovieSceneTimeHelpers.h"
 #include "FrameNumberNumericInterface.h"
 #include "LevelSequence.h"
+#include "SequencerLog.h"
 
 #define LOCTEXT_NAMESPACE "Sequencer"
 
+TSharedRef<IPropertyTypeCustomization> CreateFrameNumberCustomization(TWeakPtr<FSequencer> WeakSequencer)
+{
+	TSharedPtr<ISequencer> SequencerPtr = WeakSequencer.Pin();
+	return MakeShared<FFrameNumberDetailsCustomization>(SequencerPtr->GetNumericTypeInterface());
+}
+
 /* SSequencer interface
  *****************************************************************************/
 PRAGMA_DISABLE_OPTIMIZATION
@@ -114,19 +122,19 @@ void SSequencer::Construct(const FArguments& InArgs, TSharedRef<FSequencer> InSe
 		return 0;
 	});
 
-	auto GetFrameResolutionDelegate = FOnGetFrameRate::CreateSP(InSequencer, &FSequencer::GetFocusedFrameResolution);
-	auto GetPlayRateDelegate = FOnGetFrameRate::CreateSP(InSequencer, &FSequencer::GetFocusedPlayRate);
+	auto GetTickResolutionDelegate = FOnGetFrameRate::CreateSP(InSequencer, &FSequencer::GetFocusedTickResolution);
+	auto GetDisplayRateDelegate = FOnGetFrameRate::CreateSP(InSequencer, &FSequencer::GetFocusedDisplayRate);
 
 	// Create our numeric type interface so we can pass it to the time slider below.
-	NumericTypeInterface = MakeShareable(new FFrameNumberInterface(GetDisplayFormatDelegate, GetZeroPadFramesDelegate, GetFrameResolutionDelegate, GetPlayRateDelegate));
+	NumericTypeInterface = MakeShareable(new FFrameNumberInterface(GetDisplayFormatDelegate, GetZeroPadFramesDelegate, GetTickResolutionDelegate, GetDisplayRateDelegate));
 
 	FTimeSliderArgs TimeSliderArgs;
 	{
 		TimeSliderArgs.ViewRange = InArgs._ViewRange;
 		TimeSliderArgs.ClampRange = InArgs._ClampRange;
 		TimeSliderArgs.PlaybackRange = InArgs._PlaybackRange;
-		TimeSliderArgs.PlayRate = TAttribute<FFrameRate>(InSequencer, &FSequencer::GetFocusedPlayRate);
-		TimeSliderArgs.FrameResolution = TAttribute<FFrameRate>(InSequencer, &FSequencer::GetFocusedFrameResolution);
+		TimeSliderArgs.DisplayRate = TAttribute<FFrameRate>(InSequencer, &FSequencer::GetFocusedDisplayRate);
+		TimeSliderArgs.TickResolution = TAttribute<FFrameRate>(InSequencer, &FSequencer::GetFocusedTickResolution);
 		TimeSliderArgs.SelectionRange = InArgs._SelectionRange;
 		TimeSliderArgs.OnPlaybackRangeChanged = InArgs._OnPlaybackRangeChanged;
 		TimeSliderArgs.OnPlaybackRangeBeginDrag = OnPlaybackRangeBeginDrag;
@@ -188,7 +196,6 @@ void SSequencer::Construct(const FArguments& InArgs, TSharedRef<FSequencer> InSe
 
 	TSharedRef<SScrollBar> ScrollBar = SNew(SScrollBar)
 		.Thickness(FVector2D(5.0f, 5.0f));
-
 	SAssignNew(TrackOutliner, SSequencerTrackOutliner);
 	SAssignNew(TrackArea, SSequencerTrackArea, TimeSliderControllerRef, InSequencer);
 	SAssignNew(TreeView, SSequencerTreeView, InSequencer->GetNodeTree(), TrackArea.ToSharedRef())
@@ -205,6 +212,8 @@ void SSequencer::Construct(const FArguments& InArgs, TSharedRef<FSequencer> InSe
 		// Grid lines match the color specified in FSequencerTimeSliderController::OnPaintSectionView
 		.GridLineTint(FLinearColor(0.f, 0.f, 0.f, 0.3f));
 
+	CurveEditorPanel->GetKeyDetailsView()->RegisterInstancedCustomPropertyTypeLayout("FrameNumber", FOnGetPropertyTypeCustomizationInstance::CreateStatic(CreateFrameNumberCustomization, SequencerPtr));
+
 	const int32 Column0 = 0, Column1 = 1;
 	const int32 Row0 = 0, Row1 = 1, Row2 = 2, Row3 = 3, Row4 = 4;
 
@@ -320,6 +329,32 @@ void SSequencer::Construct(const FArguments& InArgs, TSharedRef<FSequencer> InSe
 								.HintText(LOCTEXT("FilterNodesHint", "Filter"))
 								.OnTextChanged( this, &SSequencer::OnOutlinerSearchChanged )
 							]
+							+ SHorizontalBox::Slot()
+							.AutoWidth()
+							.VAlign(VAlign_Center)
+							.HAlign(HAlign_Right)
+							.Padding(FMargin(CommonPadding + 2.0, 0.f, 0.f, 0.f))
+							[
+								SNew(SBorder)
+								.BorderImage(nullptr)
+								[
+									// Current Play Time 
+									SAssignNew(PlayTimeDisplay, SSpinBox<double>)
+									.Style(&FEditorStyle::GetWidgetStyle<FSpinBoxStyle>("Sequencer.PlayTimeSpinBox"))
+									.Value_Lambda([this]() -> double {
+										return SequencerPtr.Pin()->GetLocalTime().Time.GetFrame().Value;
+									})
+									.OnValueChanged(this, &SSequencer::SetPlayTimeClampedByWorkingRange)
+									.OnValueCommitted_Lambda([this](double InFrame, ETextCommit::Type) {
+										SetPlayTimeClampedByWorkingRange(InFrame);
+									})
+									.MinValue(TOptional<double>())
+									.MaxValue(TOptional<double>())
+									.TypeInterface(NumericTypeInterface)
+									.Delta(this, &SSequencer::GetSpinboxDelta)
+									.LinearDeltaSensitivity(25)
+								]
+							]
 						]
 					]
 
@@ -451,27 +486,14 @@ void SSequencer::Construct(const FArguments& InArgs, TSharedRef<FSequencer> InSe
 						.HAlign(HAlign_Left)
 						.VAlign(VAlign_Top)
 					[
-						SNew(SHorizontalBox)
-
-						// Goto box
-						+SHorizontalBox::Slot()
-						.AutoWidth()
-						[
-							SAssignNew(GotoBox, SSequencerGotoBox, SequencerPtr.Pin().ToSharedRef(), *Settings, NumericTypeInterface.ToSharedRef())
-						]
-
 						// Transform box
-						+SHorizontalBox::Slot()
-						.AutoWidth()
-						[
-							SAssignNew(TransformBox, SSequencerTransformBox, SequencerPtr.Pin().ToSharedRef(), *Settings, NumericTypeInterface.ToSharedRef())
-						]
+						SAssignNew(TransformBox, SSequencerTransformBox, SequencerPtr.Pin().ToSharedRef(), *Settings, NumericTypeInterface.ToSharedRef())
 					]
 
 					+ SGridPanel::Slot(Column1, Row2, SGridPanel::Layer(50))
 					.Padding(ResizeBarPadding)
 					[
-						SAssignNew(FrameResolutionOverlay, SSequencerTimePanel, SequencerPtr)
+						SAssignNew(TickResolutionOverlay, SSequencerTimePanel, SequencerPtr)
 					]
 
 					// debug vis
@@ -533,7 +555,7 @@ void SSequencer::Construct(const FArguments& InArgs, TSharedRef<FSequencer> InSe
 		]
 	];
 
-	HideFrameResolutionOverlay();
+	HideTickResolutionOverlay();
 
 	InSequencer->GetSelection().GetOnKeySelectionChanged().AddSP(this, &SSequencer::HandleKeySelectionChanged);
 	InSequencer->GetSelection().GetOnSectionSelectionChanged().AddSP(this, &SSequencer::HandleSectionSelectionChanged);
@@ -568,7 +590,7 @@ void SSequencer::BindCommands(TSharedRef<FUICommandList> SequencerCommandBinding
 
 	SequencerCommandBindings->MapAction(
 		FSequencerCommands::Get().ToggleShowGotoBox,
-		FExecuteAction::CreateLambda([this]{ GotoBox->ToggleVisibility(); })
+		FExecuteAction::CreateLambda([this] { FSlateApplication::Get().SetKeyboardFocus(PlayTimeDisplay, EFocusCause::SetDirectly); })
 	);
 
 	SequencerCommandBindings->MapAction(
@@ -577,14 +599,14 @@ void SSequencer::BindCommands(TSharedRef<FUICommandList> SequencerCommandBinding
 	);
 }
 
-void SSequencer::ShowFrameResolutionOverlay()
+void SSequencer::ShowTickResolutionOverlay()
 {
-	FrameResolutionOverlay->SetVisibility(EVisibility::Visible);
+	TickResolutionOverlay->SetVisibility(EVisibility::Visible);
 }
 
-void SSequencer::HideFrameResolutionOverlay()
+void SSequencer::HideTickResolutionOverlay()
 {
-	FrameResolutionOverlay->SetVisibility(EVisibility::Collapsed);
+	TickResolutionOverlay->SetVisibility(EVisibility::Collapsed);
 }
 
 const ISequencerEditTool* SSequencer::GetEditTool() const
@@ -784,15 +806,39 @@ TSharedRef<SWidget> SSequencer::MakeToolBar()
 
 			if( SequencerPtr.Pin()->IsLevelEditorSequencer() )
 			{
-				TAttribute<FSlateIcon> KeyAllIcon;
-				KeyAllIcon.Bind(TAttribute<FSlateIcon>::FGetter::CreateLambda([&]{
-					static FSlateIcon KeyAllEnabledIcon(FEditorStyle::GetStyleSetName(), "Sequencer.KeyAllEnabled");
-					static FSlateIcon KeyAllDisabledIcon(FEditorStyle::GetStyleSetName(), "Sequencer.KeyAllDisabled");
+				TAttribute<FSlateIcon> KeyGroupModeIcon;
+				KeyGroupModeIcon.Bind(TAttribute<FSlateIcon>::FGetter::CreateLambda([&] {
+					switch (SequencerPtr.Pin()->GetKeyGroupMode())
+					{
+					case EKeyGroupMode::KeyAll:
+						return FSequencerCommands::Get().SetKeyAll->GetIcon();
+					case EKeyGroupMode::KeyGroup:
+						return FSequencerCommands::Get().SetKeyGroup->GetIcon();
+					default: // EKeyGroupMode::KeyChanged
+						return FSequencerCommands::Get().SetKeyChanged->GetIcon();
+					}
+				}));
 
-					return SequencerPtr.Pin()->GetKeyAllEnabled() ? KeyAllEnabledIcon : KeyAllDisabledIcon;
+				TAttribute<FText> KeyGroupModeToolTip;
+				KeyGroupModeToolTip.Bind(TAttribute<FText>::FGetter::CreateLambda([&] {
+					switch (SequencerPtr.Pin()->GetKeyGroupMode())
+					{
+					case EKeyGroupMode::KeyAll:
+						return FSequencerCommands::Get().SetKeyAll->GetDescription();
+					case EKeyGroupMode::KeyGroup:
+						return FSequencerCommands::Get().SetKeyGroup->GetDescription();
+					default: // EKeyGroupMode::KeyChanged
+						return FSequencerCommands::Get().SetKeyChanged->GetDescription();
+					}
 				}));
 
-				ToolBarBuilder.AddToolBarButton( FSequencerCommands::Get().ToggleKeyAllEnabled, NAME_None, TAttribute<FText>(), TAttribute<FText>(), KeyAllIcon );
+				ToolBarBuilder.AddComboButton(
+					FUIAction(),
+					FOnGetContent::CreateSP(this, &SSequencer::MakeKeyGroupMenu),
+					LOCTEXT("KeyGroup", "Key All"),
+					KeyGroupModeToolTip,
+					KeyGroupModeIcon);
+
 			}
 
 			if (IVREditorModule::Get().IsVREditorModeActive() || (SequencerPtr.Pin()->IsLevelEditorSequencer() && ExactCast<ULevelSequence>(SequencerPtr.Pin()->GetFocusedMovieSceneSequence()) == nullptr))
@@ -1028,6 +1074,46 @@ void SSequencer::FillPlaybackSpeedMenu(FMenuBuilder& InMenuBarBuilder)
 	InMenuBarBuilder.EndSection();
 }
 
+void SSequencer::FillTimeDisplayFormatMenu(FMenuBuilder& MenuBuilder)
+{
+	TSharedPtr<FSequencer> Sequencer = SequencerPtr.Pin();
+	bool bSupportsDropFormatDisplay = FTimecode::IsDropFormatTimecodeSupported(Sequencer->GetFocusedDisplayRate());
+
+	const UEnum* FrameNumberDisplayEnum = FindObject<UEnum>(ANY_PACKAGE, TEXT("EFrameNumberDisplayFormats"), true);
+	check(FrameNumberDisplayEnum);
+
+	if (Settings)
+	{
+		EFrameNumberDisplayFormats CurrentDisplay = Settings->GetTimeDisplayFormat();
+
+		for (int32 Index = 0; Index < FrameNumberDisplayEnum->NumEnums() - 1; Index++)
+		{
+			if (!FrameNumberDisplayEnum->HasMetaData(TEXT("Hidden"), Index))
+			{
+				EFrameNumberDisplayFormats Value = (EFrameNumberDisplayFormats)FrameNumberDisplayEnum->GetValueByIndex(Index);
+
+				// Don't show Drop Frame Timecode when they're in a format that doesn't support it.
+				if (Value == EFrameNumberDisplayFormats::DropFrameTimecode && !bSupportsDropFormatDisplay)
+					continue;
+
+				bool bIsSet = CurrentDisplay == Value;
+				MenuBuilder.AddMenuEntry(
+					FrameNumberDisplayEnum->GetDisplayNameTextByIndex(Index),
+					FrameNumberDisplayEnum->GetToolTipTextByIndex(Index),
+					FSlateIcon(),
+					FUIAction(
+						FExecuteAction::CreateUObject(Settings, &USequencerSettings::SetTimeDisplayFormat, Value),
+						FCanExecuteAction(),
+						FIsActionChecked::CreateLambda([=] { return CurrentDisplay == Value; })
+					),
+					NAME_None,
+					EUserInterfaceActionType::RadioButton
+				);
+			}
+		}
+	}
+}
+
 TSharedRef<SWidget> SSequencer::MakePlaybackMenu()
 {
 	FMenuBuilder MenuBuilder( true, SequencerPtr.Pin()->GetCommandBindings() );
@@ -1036,9 +1122,17 @@ TSharedRef<SWidget> SSequencer::MakePlaybackMenu()
 	MenuBuilder.BeginSection("PlaybackThisSequence", LOCTEXT("PlaybackThisSequenceHeader", "Playback - This Sequence"));
 	{
 		// Menu entry for the start position
-		auto OnStartChanged = [=](FFrameNumber NewValue){
+		auto OnStartChanged = [=](double NewValue){
+
+			// We clamp the new value when the value is set. We can't clamp in the UI because we need an unset Min/Max for linear scaling to work.
+			double Min = (SequencerPtr.Pin()->GetClampRange().GetLowerBoundValue() * SequencerPtr.Pin()->GetFocusedTickResolution()).GetFrame().Value;
+			double Max = SequencerPtr.Pin()->GetPlaybackRange().GetUpperBoundValue().Value; 
+
+			NewValue = FMath::Clamp(NewValue, Min, Max);
+			FFrameNumber ValueAsFrame = FFrameTime::FromDecimal(NewValue).GetFrame();
+
 			FFrameNumber Upper = MovieScene::DiscreteExclusiveUpper(SequencerPtr.Pin()->GetPlaybackRange());
-			SequencerPtr.Pin()->SetPlaybackRange(TRange<FFrameNumber>(FMath::Min(NewValue, Upper-1), Upper));
+			SequencerPtr.Pin()->SetPlaybackRange(TRange<FFrameNumber>(FMath::Min(ValueAsFrame, Upper-1), Upper));
 		};
 
 		MenuBuilder.AddWidget(
@@ -1056,28 +1150,32 @@ TSharedRef<SWidget> SSequencer::MakePlaybackMenu()
 							return !SequencerPtr.Pin()->IsPlaybackRangeLocked();
 						})
 						.Style(&FEditorStyle::GetWidgetStyle<FSpinBoxStyle>("Sequencer.HyperlinkSpinBox"))
-						.OnValueCommitted_Lambda([=](double Value, ETextCommit::Type){ OnStartChanged(FFrameTime::FromDecimal(Value).GetFrame()); })
-						.OnValueChanged_Lambda([=](double Value) { OnStartChanged(FFrameTime::FromDecimal(Value).GetFrame()); })
+						.OnValueCommitted_Lambda([=](double Value, ETextCommit::Type){ OnStartChanged(Value); })
+						.OnValueChanged_Lambda([=](double Value) { OnStartChanged(Value); })
 						.OnBeginSliderMovement(OnPlaybackRangeBeginDrag)
-						.OnEndSliderMovement_Lambda([=](double Value){ OnStartChanged(FFrameTime::FromDecimal(Value).GetFrame()); OnPlaybackRangeEndDrag.ExecuteIfBound(); })
-						.MinValue_Lambda([=]() -> double {
-							// Clamp range is in seconds so we need to convert it to frame resolution.
-							return (SequencerPtr.Pin()->GetClampRange().GetLowerBoundValue() * SequencerPtr.Pin()->GetFocusedFrameResolution()).GetFrame().Value;
-						})
-						.MaxValue_Lambda([=]() -> double {
-							return SequencerPtr.Pin()->GetPlaybackRange().GetUpperBoundValue().Value; 
-						})
+						.OnEndSliderMovement_Lambda([=](double Value){ OnStartChanged(Value); OnPlaybackRangeEndDrag.ExecuteIfBound(); })
+						.MinValue(TOptional<double>())
+						.MaxValue(TOptional<double>())
 						.Value_Lambda([=]() -> double {
 							return SequencerPtr.Pin()->GetPlaybackRange().GetLowerBoundValue().Value;
 						})
 						.Delta(this, &SSequencer::GetSpinboxDelta)
+						.LinearDeltaSensitivity(25)
 			],
 			LOCTEXT("PlaybackStartLabel", "Start"));
 
 		// Menu entry for the end position
-		auto OnEndChanged = [=](FFrameNumber NewValue){
+		auto OnEndChanged = [=](double NewValue) {
+
+			// We clamp the new value when the value is set. We can't clamp in the UI because we need an unset Min/Max for linear scaling to work.
+			double Min = SequencerPtr.Pin()->GetPlaybackRange().GetLowerBoundValue().Value;
+			double Max = (SequencerPtr.Pin()->GetClampRange().GetUpperBoundValue() * SequencerPtr.Pin()->GetFocusedTickResolution()).GetFrame().Value;
+
+			NewValue = FMath::Clamp(NewValue, Min, Max);
+			FFrameNumber ValueAsFrame = FFrameTime::FromDecimal(NewValue).GetFrame();
+
 			FFrameNumber Lower = MovieScene::DiscreteInclusiveLower(SequencerPtr.Pin()->GetPlaybackRange());
-			SequencerPtr.Pin()->SetPlaybackRange(TRange<FFrameNumber>(Lower, FMath::Max(NewValue, Lower)));
+			SequencerPtr.Pin()->SetPlaybackRange(TRange<FFrameNumber>(Lower, FMath::Max(ValueAsFrame, Lower)));
 		};
 
 		MenuBuilder.AddWidget(
@@ -1095,21 +1193,17 @@ TSharedRef<SWidget> SSequencer::MakePlaybackMenu()
 					 		return !SequencerPtr.Pin()->IsPlaybackRangeLocked();
 						})
 						.Style(&FEditorStyle::GetWidgetStyle<FSpinBoxStyle>("Sequencer.HyperlinkSpinBox"))
-						.OnValueCommitted_Lambda([=](double Value, ETextCommit::Type){ OnEndChanged(FFrameTime::FromDecimal(Value).GetFrame()); })
-						.OnValueChanged_Lambda([=](double Value) { OnEndChanged(FFrameTime::FromDecimal(Value).GetFrame()); })
+						.OnValueCommitted_Lambda([=](double Value, ETextCommit::Type){ OnEndChanged(Value); })
+						.OnValueChanged_Lambda([=](double Value) { OnEndChanged(Value); })
 						.OnBeginSliderMovement(OnPlaybackRangeBeginDrag)
-						.OnEndSliderMovement_Lambda([=](double Value){ OnEndChanged(FFrameTime::FromDecimal(Value).GetFrame()); OnPlaybackRangeEndDrag.ExecuteIfBound(); })
-						.MinValue_Lambda([=]() -> double {
-					 		return SequencerPtr.Pin()->GetPlaybackRange().GetLowerBoundValue().Value;
-						})
-						.MaxValue_Lambda([=]() -> double {
-							// Clamp range is in seconds so we need to convert it to frame resolution.
-							return (SequencerPtr.Pin()->GetClampRange().GetUpperBoundValue() * SequencerPtr.Pin()->GetFocusedFrameResolution()).GetFrame().Value;
-						})
+						.OnEndSliderMovement_Lambda([=](double Value){ OnEndChanged(Value); OnPlaybackRangeEndDrag.ExecuteIfBound(); })
+						.MinValue(TOptional<double>())
+						.MaxValue(TOptional<double>())
 						.Value_Lambda([=]() -> double {
 					 		return SequencerPtr.Pin()->GetPlaybackRange().GetUpperBoundValue().Value;
 						})
 						.Delta(this, &SSequencer::GetSpinboxDelta)
+						.LinearDeltaSensitivity(25)
 				],
 			LOCTEXT("PlaybackStartEnd", "End"));
 
@@ -1274,6 +1368,18 @@ TSharedRef<SWidget> SSequencer::MakeAllowEditsMenu()
 
 }
 
+TSharedRef<SWidget> SSequencer::MakeKeyGroupMenu()
+{
+	FMenuBuilder MenuBuilder(false, SequencerPtr.Pin()->GetCommandBindings());
+
+	MenuBuilder.AddMenuEntry(FSequencerCommands::Get().SetKeyAll);
+	MenuBuilder.AddMenuEntry(FSequencerCommands::Get().SetKeyGroup);
+	MenuBuilder.AddMenuEntry(FSequencerCommands::Get().SetKeyChanged);
+
+	return MenuBuilder.MakeWidget();
+
+}
+
 TSharedRef<SWidget> SSequencer::MakeTimeRange(const TSharedRef<SWidget>& InnerContent, bool bShowWorkingRange, bool bShowViewRange, bool bShowPlaybackRange)
 {
 	ISequencerWidgetsModule& SequencerWidgets = FModuleManager::Get().LoadModuleChecked<ISequencerWidgetsModule>( "SequencerWidgets" );
@@ -1454,7 +1560,6 @@ void SSequencer::UpdateLayoutTree()
 		// Add any additional paths that have been added by the user for nodes that may not exist yet but we want them to be selected
 		// after the node tree is updated and we restore selections.
 		SelectedPathNames.Append(AdditionalSelectionsToAdd);
-		AdditionalSelectionsToAdd.Empty();
 
 		// Suspend broadcasting selection changes because we don't want unnecessary rebuilds.
 		Sequencer->GetSelection().SuspendBroadcast();
@@ -1471,6 +1576,25 @@ void SSequencer::UpdateLayoutTree()
 
 		RestoreKeySelection(SelectedKeys, Sequencer->GetSelection(), *Sequencer->GetNodeTree());
 		RestoreSectionSelection(SelectedSections, Sequencer->GetSelection());
+		
+		// If we've manually specified an additional selection to add it's because the item was newly created.
+		// Now that the treeview has been refreshed and selection restored, we'll try to focus the first item
+		// so that the view scrolls down when things are added to the bottom.
+		if (AdditionalSelectionsToAdd.Num() > 0)
+		{
+			FString NodePath = AdditionalSelectionsToAdd[0];
+
+			for (TSharedRef<FSequencerDisplayNode> Node : Sequencer->GetNodeTree()->GetAllNodes())
+			{
+				if (Node->GetPathName() == NodePath)
+				{
+					TreeView->RequestScrollIntoView(Node);
+					break;
+				}
+			}
+		}
+
+		AdditionalSelectionsToAdd.Empty();
 
 		// Continue broadcasting selection changes
 		Sequencer->GetSelection().ResumeBroadcast();
@@ -1704,6 +1828,13 @@ void SSequencer::OnAssetsDropped( const FAssetDragDropOp& DragDropOp )
 
 		if (!SequencerRef.OnHandleAssetDropped(CurObject, TargetObjectGuid))
 		{
+			// Doesn't make sense to drop a level sequence asset into sequencer as a spawnable actor
+			if (CurObject->IsA<ULevelSequence>())
+			{
+				UE_LOG(LogSequencer, Warning, TEXT("Can't add '%s' as a spawnable"), *CurObject->GetName());
+				continue;
+			}
+
 			FGuid NewGuid = SequencerRef.MakeNewSpawnable( *CurObject, DragDropOp.GetActorFactory() );
 
 			UMovieScene* MovieScene = SequencerRef.GetFocusedMovieSceneSequence()->GetMovieScene();
@@ -2110,7 +2241,7 @@ FPasteContextMenuArgs SSequencer::GeneratePasteArgs(FFrameNumber PasteAtTime, TS
 	if (Settings->GetIsSnapEnabled())
 	{
 		// @todo: sequencer-timecode: play rate override
-		//PasteAtTime = FQualifiedFrameTime(PasteAtTime, GetFrameResolution()).ToFrameRate();
+		//PasteAtTime = FQualifiedFrameTime(PasteAtTime, GetTickResolution()).ToFrameRate();
 	}
 
 	// Open a paste menu at the current mouse position
@@ -2283,7 +2414,25 @@ void SSequencer::BuildCustomContextMenuForGuid(FMenuBuilder& MenuBuilder, FGuid
 double SSequencer::GetSpinboxDelta() const
 {
 	TSharedPtr<FSequencer> Sequencer = SequencerPtr.Pin();
-	return Sequencer->GetPlayRateDeltaFrameCount();
+	return Sequencer->GetDisplayRateDeltaFrameCount();
+}
+
+void SSequencer::SetPlayTimeClampedByWorkingRange(double Frame)
+{
+	if (SequencerPtr.IsValid())
+	{
+		TSharedPtr<FSequencer> Sequencer = SequencerPtr.Pin();
+		// Some of our spin boxes need to use an unbounded min/max so that they can drag linearly instead of based on the current value.
+		// We clamp the value here by the working range to emulate the behavior of the Cinematic Level Viewport
+		FFrameRate   PlayRate = Sequencer->GetLocalTime().Rate;
+		UMovieScene* MovieScene = Sequencer->GetFocusedMovieSceneSequence()->GetMovieScene();
+		double StartInSeconds = MovieScene->GetEditorData().WorkStart;
+		double EndInSeconds = MovieScene->GetEditorData().WorkEnd;
+
+		Frame = FMath::Clamp(Frame, (double)(StartInSeconds*PlayRate).GetFrame().Value, (double)(EndInSeconds*PlayRate).GetFrame().Value);
+
+		Sequencer->SetLocalTime(FFrameTime::FromDecimal(Frame));
+	}
 }
 
 #undef LOCTEXT_NAMESPACE
diff --git a/Engine/Source/Editor/Sequencer/Private/SSequencer.h b/Engine/Source/Editor/Sequencer/Private/SSequencer.h
--- a/Engine/Source/Editor/Sequencer/Private/SSequencer.h
+++ b/Engine/Source/Editor/Sequencer/Private/SSequencer.h
@@ -14,6 +14,7 @@
 #include "ITimeSlider.h"
 #include "Framework/Commands/UICommandList.h"
 #include "Widgets/Input/NumericTypeInterface.h"
+#include "Widgets/Input/SSpinBox.h"
 #include "Sequencer.h"
 
 class FActorDragDropGraphEdOp;
@@ -234,9 +235,11 @@ public:
 	/** Access the currently active track area edit tool */
 	const ISequencerEditTool* GetEditTool() const;
 
-	void ShowFrameResolutionOverlay();
-	void HideFrameResolutionOverlay();
+	void ShowTickResolutionOverlay();
+	void HideTickResolutionOverlay();
 
+	/** Sets the play time for the sequence but clamped by the working range. This is useful for cases where we can't clamp via the UI control. */
+	void SetPlayTimeClampedByWorkingRange(double Frame);
 public:
 
 	// FNotifyHook overrides
@@ -308,9 +311,16 @@ private:
 	/** Makes the allow edits menu for the toolbar. */
 	TSharedRef<SWidget> MakeAllowEditsMenu();
 
+	/** Makes the key group menu for the toolbar. */
+	TSharedRef<SWidget> MakeKeyGroupMenu();
+
 	/** Makes the playback speed menu for the toolbar. */
 	void FillPlaybackSpeedMenu(FMenuBuilder& InMenuBuilder);
 
+public:
+	/** Makes the time display format menu for the toolbar and the play rate menu. */
+	void FillTimeDisplayFormatMenu(FMenuBuilder& MenuBuilder);
+
 public:	
 
 	/** Makes a time range widget with the specified inner content */
@@ -463,6 +473,9 @@ private:
 	/** The search box for filtering tracks. */
 	TSharedPtr<SSearchBox> SearchBox;
 
+	/** The current playback time display.*/
+	TSharedPtr<SSpinBox<double>> PlayTimeDisplay;
+
 	/** The sequencer tree view responsible for the outliner and track areas */
 	TSharedPtr<SSequencerTreeView> TreeView;
 
@@ -525,7 +538,7 @@ private:
 	 * This can be used to highlight nodes that may not exist until the rebuild. Cleared after the tree is rebuilt
 	 * and the selection list is restored.
 	*/
-	TSet<FString> AdditionalSelectionsToAdd;
+	TArray<FString> AdditionalSelectionsToAdd;
 
-	TSharedPtr<SWidget> FrameResolutionOverlay;
+	TSharedPtr<SWidget> TickResolutionOverlay;
 };
diff --git a/Engine/Source/Editor/Sequencer/Private/SSequencerCurveEditorToolBar.cpp b/Engine/Source/Editor/Sequencer/Private/SSequencerCurveEditorToolBar.cpp
--- a/Engine/Source/Editor/Sequencer/Private/SSequencerCurveEditorToolBar.cpp
+++ b/Engine/Source/Editor/Sequencer/Private/SSequencerCurveEditorToolBar.cpp
@@ -75,6 +75,7 @@ void SSequencerCurveEditorToolBar::Construct( const FArguments& InArgs, TSharedR
 
 	ToolBarBuilder.BeginSection( "Tangents" );
 	{
+		ToolBarBuilder.AddToolBarButton( FCurveEditorCommands::Get().InterpolationToggleWeighted );
 		ToolBarBuilder.AddToolBarButton( FCurveEditorCommands::Get().FlattenTangents );
 		ToolBarBuilder.AddToolBarButton( FCurveEditorCommands::Get().StraightenTangents );
 	}
diff --git a/Engine/Source/Editor/Sequencer/Private/SSequencerDebugVisualizer.cpp b/Engine/Source/Editor/Sequencer/Private/SSequencerDebugVisualizer.cpp
--- a/Engine/Source/Editor/Sequencer/Private/SSequencerDebugVisualizer.cpp
+++ b/Engine/Source/Editor/Sequencer/Private/SSequencerDebugVisualizer.cpp
@@ -47,7 +47,7 @@ void SSequencerDebugVisualizer::Refresh()
 		return;
 	}
 
-	const FFrameRate                  SequenceResolution = ActiveSequence->GetMovieScene()->GetFrameResolution();
+	const FFrameRate                  SequenceResolution = ActiveSequence->GetMovieScene()->GetTickResolution();
 	const FMovieSceneEvaluationField& EvaluationField    = ActiveTemplate->EvaluationField;
 
 	CachedSignature = EvaluationField.GetSignature();
@@ -222,7 +222,7 @@ void SSequencerDebugVisualizer::OnArrangeChildren( const FGeometry& AllottedGeom
 		return;
 	}
 
-	FTimeToPixel TimeToPixelConverter = FTimeToPixel(AllottedGeometry, ViewRange.Get(), ActiveSequence->GetMovieScene()->GetFrameResolution());
+	FTimeToPixel TimeToPixelConverter = FTimeToPixel(AllottedGeometry, ViewRange.Get(), ActiveSequence->GetMovieScene()->GetTickResolution());
 
 	for (int32 WidgetIndex = 0; WidgetIndex < Children.Num(); ++WidgetIndex)
 	{
diff --git a/Engine/Source/Editor/Sequencer/Private/SSequencerPlayRateCombo.cpp b/Engine/Source/Editor/Sequencer/Private/SSequencerPlayRateCombo.cpp
--- a/Engine/Source/Editor/Sequencer/Private/SSequencerPlayRateCombo.cpp
+++ b/Engine/Source/Editor/Sequencer/Private/SSequencerPlayRateCombo.cpp
@@ -91,7 +91,7 @@ EVisibility SSequencerPlayRateCombo::GetFrameLockedVisibility() const
 EVisibility SSequencerPlayRateCombo::GetFrameRateErrorVisibility() const
 {
 	TSharedPtr<FSequencer> Sequencer  = WeakSequencer.Pin();
-	if (!Sequencer.IsValid() || Sequencer->GetFocusedPlayRate().IsMultipleOf(Sequencer->GetFocusedFrameResolution()))
+	if (!Sequencer.IsValid() || Sequencer->GetFocusedDisplayRate().IsMultipleOf(Sequencer->GetFocusedTickResolution()))
 	{
 		return EVisibility::Collapsed;
 	}
@@ -103,16 +103,16 @@ FText SSequencerPlayRateCombo::GetFrameRateErrorDescription() const
 	TSharedPtr<FSequencer> Sequencer  = WeakSequencer.Pin();
 	if (Sequencer.IsValid())
 	{
-		FFrameRate PlayRate = Sequencer->GetFocusedPlayRate();
-		FFrameRate FrameResolution = Sequencer->GetFocusedFrameResolution();
+		FFrameRate DisplayRate = Sequencer->GetFocusedDisplayRate();
+		FFrameRate TickResolution = Sequencer->GetFocusedTickResolution();
 
-		const FCommonFrameRateInfo* PlayRateInfo        = FCommonFrameRates::Find(PlayRate);
-		const FCommonFrameRateInfo* FrameResolutionInfo = FCommonFrameRates::Find(FrameResolution);
+		const FCommonFrameRateInfo* DisplayRateInfo     = FCommonFrameRates::Find(DisplayRate);
+		const FCommonFrameRateInfo* TickResolutionInfo  = FCommonFrameRates::Find(TickResolution);
 
-		FText PlayRateText        = PlayRateInfo        ? PlayRateInfo->DisplayName        : FText::Format(LOCTEXT("PlayRateFormat", "{0} fps"), PlayRate.AsDecimal());
-		FText FrameResolutionText = FrameResolutionInfo ? FrameResolutionInfo->DisplayName : FText::Format(LOCTEXT("FrameResolutionFormat", "{0} ticks every second"), FrameResolution.AsDecimal());
+		FText DisplayRateText     = DisplayRateInfo     ? DisplayRateInfo->DisplayName    : FText::Format(LOCTEXT("DisplayRateFormat", "{0} fps"), DisplayRate.AsDecimal());
+		FText TickResolutionText  = TickResolutionInfo  ? TickResolutionInfo->DisplayName : FText::Format(LOCTEXT("TickResolutionFormat", "{0} ticks every second"), TickResolution.AsDecimal());
 
-		return FText::Format(LOCTEXT("FrameRateErrorDescription", "The current display rate of {0} is incompatible with this sequence's tick resolution of {1} ticks per second."), PlayRateText, FrameResolutionText);
+		return FText::Format(LOCTEXT("FrameRateErrorDescription", "The current display rate of {0} is incompatible with this sequence's tick resolution of {1} ticks per second."), DisplayRateText, TickResolutionText);
 	}
 
 	return FText();
@@ -129,12 +129,12 @@ TSharedRef<SWidget> SSequencerPlayRateCombo::OnCreateMenu()
 
 	FMenuBuilder MenuBuilder(true, nullptr);
 
-	FFrameRate FrameResolution = Sequencer->GetFocusedFrameResolution();
+	FFrameRate TickResolution = Sequencer->GetFocusedTickResolution();
 
 	TArray<FCommonFrameRateInfo> CompatibleRates;
 	for (const FCommonFrameRateInfo& Info : FCommonFrameRates::GetAll())
 	{
-		if (Info.FrameRate.IsMultipleOf(FrameResolution))
+		if (Info.FrameRate.IsMultipleOf(TickResolution))
 		{
 			CompatibleRates.Add(Info);
 		}
@@ -160,8 +160,8 @@ TSharedRef<SWidget> SSequencerPlayRateCombo::OnCreateMenu()
 			.MaxDesiredWidth(100.f)
 			[
 				SNew(SFrameRateEntryBox)
-				.Value(this, &SSequencerPlayRateCombo::GetPlaybackFrameRate)
-				.OnValueChanged(this, &SSequencerPlayRateCombo::SetPlaybackFrameRate)
+				.Value(this, &SSequencerPlayRateCombo::GetDisplayRate)
+				.OnValueChanged(this, &SSequencerPlayRateCombo::SetDisplayRate)
 			],
 			LOCTEXT("CustomFramerateDisplayLabel", "Custom")
 		);
@@ -170,7 +170,7 @@ TSharedRef<SWidget> SSequencerPlayRateCombo::OnCreateMenu()
 		{
 			MenuBuilder.AddSubMenu(
 				LOCTEXT("IncompatibleRates", "Incompatible Rates"),
-				FText::Format(LOCTEXT("IncompatibleRates_Description", "Choose from a list of display rates that are incompatible with a resolution of {0} ticks per second"), FrameResolution.AsDecimal()),
+				FText::Format(LOCTEXT("IncompatibleRates_Description", "Choose from a list of display rates that are incompatible with a resolution of {0} ticks per second"), TickResolution.AsDecimal()),
 				FNewMenuDelegate::CreateSP(this, &SSequencerPlayRateCombo::PopulateIncompatibleRatesMenu)
 				);
 		}
@@ -182,7 +182,12 @@ TSharedRef<SWidget> SSequencerPlayRateCombo::OnCreateMenu()
 	MenuBuilder.AddSubMenu(
 		LOCTEXT("ShowTimesAs", "Show Time As"),
 		LOCTEXT("ShowTimesAs_Description", "Change how to display times in Sequencer"),
-		FNewMenuDelegate::CreateSP(this, &SSequencerPlayRateCombo::PopulateTimeDisplayMenu)
+		FNewMenuDelegate::CreateLambda([this](FMenuBuilder& InMenuBuilder) {
+			if (WeakSequencerWidget.IsValid())
+			{
+				WeakSequencerWidget.Pin()->FillTimeDisplayFormatMenu(InMenuBuilder);
+			}
+		})
 	);
 
 	if (Sequencer->GetRootMovieSceneSequence() == Sequencer->GetFocusedMovieSceneSequence())
@@ -212,7 +217,7 @@ TSharedRef<SWidget> SSequencerPlayRateCombo::OnCreateMenu()
 		LOCTEXT("AdvancedOptions_Description", "Open advanced time-related properties for this sequence"),
 		FSlateIcon(),
 		FUIAction(
-			FExecuteAction::CreateSP(SequencerWidget.Get(), &SSequencer::ShowFrameResolutionOverlay)
+			FExecuteAction::CreateSP(SequencerWidget.Get(), &SSequencer::ShowTickResolutionOverlay)
 		)
 	);
 
@@ -224,12 +229,12 @@ void SSequencerPlayRateCombo::PopulateIncompatibleRatesMenu(FMenuBuilder& MenuBu
 	TSharedPtr<FSequencer> Sequencer = WeakSequencer.Pin();
 	if (Sequencer.IsValid())
 	{
-		FFrameRate FrameResolution = Sequencer->GetFocusedFrameResolution();
+		FFrameRate TickResolution = Sequencer->GetFocusedTickResolution();
 
 		TArray<FCommonFrameRateInfo> IncompatibleRates;
 		for (const FCommonFrameRateInfo& Info : FCommonFrameRates::GetAll())
 		{
-			if (!Info.FrameRate.IsMultipleOf(FrameResolution))
+			if (!Info.FrameRate.IsMultipleOf(TickResolution))
 			{
 				IncompatibleRates.Add(Info);
 			}
@@ -249,47 +254,6 @@ void SSequencerPlayRateCombo::PopulateIncompatibleRatesMenu(FMenuBuilder& MenuBu
 	}
 }
 
-void SSequencerPlayRateCombo::PopulateTimeDisplayMenu(FMenuBuilder& MenuBuilder)
-{
-	TSharedPtr<FSequencer> Sequencer = WeakSequencer.Pin();
-	USequencerSettings*    Settings = Sequencer.IsValid() ? Sequencer->GetSequencerSettings() : nullptr;
-	bool bSupportsDropFormatDisplay = FTimecode::IsDropFormatTimecodeSupported(Sequencer->GetFocusedPlayRate());
-
-	const UEnum* FrameNumberDisplayEnum = FindObject<UEnum>(ANY_PACKAGE, TEXT("EFrameNumberDisplayFormats"), true);
-	check(FrameNumberDisplayEnum);
-
-	if (Settings)
-	{
-		EFrameNumberDisplayFormats CurrentDisplay = Settings->GetTimeDisplayFormat();
-
-		for (int32 Index = 0; Index < FrameNumberDisplayEnum->NumEnums() - 1; Index++)
-		{
-			if (!FrameNumberDisplayEnum->HasMetaData(TEXT("Hidden"), Index))
-			{
-				EFrameNumberDisplayFormats Value = (EFrameNumberDisplayFormats)FrameNumberDisplayEnum->GetValueByIndex(Index);
-
-				// Don't show Drop Frame Timecode when they're in a format that doesn't support it.
-				if (Value == EFrameNumberDisplayFormats::DropFrameTimecode && !bSupportsDropFormatDisplay)
-					continue;
-
-				bool bIsSet = CurrentDisplay == Value;
-				MenuBuilder.AddMenuEntry(
-					FrameNumberDisplayEnum->GetDisplayNameTextByIndex(Index),
-					FrameNumberDisplayEnum->GetToolTipTextByIndex(Index),
-					FSlateIcon(),
-					FUIAction(
-						FExecuteAction::CreateUObject(Settings, &USequencerSettings::SetTimeDisplayFormat, Value),
-						FCanExecuteAction(),
-						FIsActionChecked::CreateLambda([=]{ return CurrentDisplay == Value; })
-					),
-					NAME_None,
-					EUserInterfaceActionType::RadioButton
-				);
-			}
-		}
-	}
-}
-
 void SSequencerPlayRateCombo::PopulateClockSourceMenu(FMenuBuilder& MenuBuilder)
 {
 	TSharedPtr<FSequencer> Sequencer    = WeakSequencer.Pin();
@@ -334,9 +298,9 @@ void SSequencerPlayRateCombo::AddMenuEntry(FMenuBuilder& MenuBuilder, const FCom
 		Info.Description,
 		FSlateIcon(),
 		FUIAction(
-			FExecuteAction::CreateSP(this, &SSequencerPlayRateCombo::SetPlaybackFrameRate, Info.FrameRate),
+			FExecuteAction::CreateSP(this, &SSequencerPlayRateCombo::SetDisplayRate, Info.FrameRate),
 			FCanExecuteAction(),
-			FIsActionChecked::CreateSP(this, &SSequencerPlayRateCombo::IsSamePlaybackFrameRate, Info.FrameRate)
+			FIsActionChecked::CreateSP(this, &SSequencerPlayRateCombo::IsSameDisplayRate, Info.FrameRate)
 		),
 		NAME_None,
 		EUserInterfaceActionType::RadioButton
@@ -358,7 +322,7 @@ void SSequencerPlayRateCombo::SetClockSource(EUpdateClockSource NewClockSource)
 	}
 }
 
-void SSequencerPlayRateCombo::SetPlaybackFrameRate(FFrameRate InFrameRate)
+void SSequencerPlayRateCombo::SetDisplayRate(FFrameRate InFrameRate)
 {
 	TSharedPtr<FSequencer> Sequencer = WeakSequencer.Pin();
 	UMovieSceneSequence* FocusedSequence = Sequencer.IsValid() ? Sequencer->GetFocusedMovieSceneSequence() : nullptr;
@@ -367,26 +331,26 @@ void SSequencerPlayRateCombo::SetPlaybackFrameRate(FFrameRate InFrameRate)
 		FScopedTransaction ScopedTransaction(FText::Format(LOCTEXT("SetDisplayRate", "Set Display Rate to {0}"), InFrameRate.ToPrettyText()));
 
 		FocusedSequence->GetMovieScene()->Modify();
-		FocusedSequence->GetMovieScene()->SetPlaybackFrameRate(InFrameRate);
+		FocusedSequence->GetMovieScene()->SetDisplayRate(InFrameRate);
 	}
 }
 
-FFrameRate SSequencerPlayRateCombo::GetPlaybackFrameRate() const
+FFrameRate SSequencerPlayRateCombo::GetDisplayRate() const
 {
 	TSharedPtr<FSequencer> Sequencer = WeakSequencer.Pin();
-	return Sequencer.IsValid() ? Sequencer->GetFocusedPlayRate() : FFrameRate();
+	return Sequencer.IsValid() ? Sequencer->GetFocusedDisplayRate() : FFrameRate();
 }
 
-bool SSequencerPlayRateCombo::IsSamePlaybackFrameRate(FFrameRate InFrameRate) const
+bool SSequencerPlayRateCombo::IsSameDisplayRate(FFrameRate InFrameRate) const
 {
-	return GetPlaybackFrameRate() == InFrameRate;
+	return GetDisplayRate() == InFrameRate;
 }
 
 FText SSequencerPlayRateCombo::GetFrameRateText() const
 {
 	TSharedPtr<FSequencer> Sequencer = WeakSequencer.Pin();
 
-	return Sequencer.IsValid() ? Sequencer->GetFocusedPlayRate().ToPrettyText() : FText();
+	return Sequencer.IsValid() ? Sequencer->GetFocusedDisplayRate().ToPrettyText() : FText();
 }
 
 FText SSequencerPlayRateCombo::GetToolTipText() const
@@ -397,18 +361,18 @@ FText SSequencerPlayRateCombo::GetToolTipText() const
 
 	if (FocusedMovieScene)
 	{
-		FFrameRate PlayRate        = FocusedMovieScene->GetPlaybackFrameRate();
-		FFrameRate FrameResolution = FocusedMovieScene->GetFrameResolution();
+		FFrameRate DisplayRate     = FocusedMovieScene->GetDisplayRate();
+		FFrameRate TickResolution  = FocusedMovieScene->GetTickResolution();
 
-		const FCommonFrameRateInfo* PlayRateInfo        = FCommonFrameRates::Find(PlayRate);
-		const FCommonFrameRateInfo* FrameResolutionInfo = FCommonFrameRates::Find(FrameResolution);
+		const FCommonFrameRateInfo* DisplayRateInfo    = FCommonFrameRates::Find(DisplayRate);
+		const FCommonFrameRateInfo* TickResolutionInfo = FCommonFrameRates::Find(TickResolution);
 
-		FText PlayRateText        = PlayRateInfo        ? PlayRateInfo->DisplayName        : FText::Format(LOCTEXT("PlayRateFormat", "{0} fps"), PlayRate.AsDecimal());
-		FText FrameResolutionText = FrameResolutionInfo ? FrameResolutionInfo->DisplayName : FText::Format(LOCTEXT("FrameResolutionFormat", "{0} ticks every second"), FrameResolution.AsDecimal());
+		FText DisplayRateText     = DisplayRateInfo     ? DisplayRateInfo->DisplayName    : FText::Format(LOCTEXT("DisplayRateFormat", "{0} fps"), DisplayRate.AsDecimal());
+		FText TickResolutionText  = TickResolutionInfo  ? TickResolutionInfo->DisplayName : FText::Format(LOCTEXT("TickResolutionFormat", "{0} ticks every second"), TickResolution.AsDecimal());
 
 		return FocusedMovieScene->GetEvaluationType() == EMovieSceneEvaluationType::FrameLocked
-			? FText::Format(LOCTEXT("ToolTip_Format_FrameLocked", "This sequence is locked at runtime to {0} and uses an underlying tick resolution of {1}."), PlayRateText, FrameResolutionText)
-			: FText::Format(LOCTEXT("ToolTip_Format", "This sequence is being presented as {0} and uses an underlying tick resolution of {1}."), PlayRateText, FrameResolutionText);
+			? FText::Format(LOCTEXT("ToolTip_Format_FrameLocked", "This sequence is locked at runtime to {0} and uses an underlying tick resolution of {1}."), DisplayRateText, TickResolutionText)
+			: FText::Format(LOCTEXT("ToolTip_Format", "This sequence is being presented as {0} and uses an underlying tick resolution of {1}."), DisplayRateText, TickResolutionText);
 	}
 
 	return FText();
diff --git a/Engine/Source/Editor/Sequencer/Private/SSequencerPlayRateCombo.h b/Engine/Source/Editor/Sequencer/Private/SSequencerPlayRateCombo.h
--- a/Engine/Source/Editor/Sequencer/Private/SSequencerPlayRateCombo.h
+++ b/Engine/Source/Editor/Sequencer/Private/SSequencerPlayRateCombo.h
@@ -44,16 +44,15 @@ private:
 
 	TSharedRef<SWidget> OnCreateMenu();
 	void PopulateIncompatibleRatesMenu(FMenuBuilder& MenuBuilder);
-	void PopulateTimeDisplayMenu(FMenuBuilder& MenuBuilder);
 	void PopulateClockSourceMenu(FMenuBuilder& MenuBuilder);
 	void AddMenuEntry(FMenuBuilder& MenuBuilder, const FCommonFrameRateInfo& Info);
 
 	void OnToggleFrameLocked();
 	ECheckBoxState OnGetFrameLockedCheckState() const;
 
-	void SetPlaybackFrameRate(FFrameRate InFrameRate);
-	FFrameRate GetPlaybackFrameRate() const;
-	bool IsSamePlaybackFrameRate(FFrameRate InFrameRate) const;
+	void SetDisplayRate(FFrameRate InFrameRate);
+	FFrameRate GetDisplayRate() const;
+	bool IsSameDisplayRate(FFrameRate InFrameRate) const;
 
 	void SetClockSource(EUpdateClockSource NewClockSource);
 
diff --git a/Engine/Source/Editor/Sequencer/Private/SSequencerSection.cpp b/Engine/Source/Editor/Sequencer/Private/SSequencerSection.cpp
--- a/Engine/Source/Editor/Sequencer/Private/SSequencerSection.cpp
+++ b/Engine/Source/Editor/Sequencer/Private/SSequencerSection.cpp
@@ -8,6 +8,7 @@
 #include "Editor.h"
 #include "Sequencer.h"
 #include "SequencerSectionPainter.h"
+#include "MovieSceneSequence.h"
 #include "CommonMovieSceneTools.h"
 #include "ISequencerEditTool.h"
 #include "ISequencerSection.h"
@@ -75,11 +76,11 @@ FTimeToPixel ConstructTimeConverterForSection(const FGeometry& InSectionGeometry
 {
 	TRange<double> ViewRange       = InSequencer.GetViewRange();
 
-	FFrameRate     FrameResolution = InSection.GetTypedOuter<UMovieScene>()->GetFrameResolution();
-	double         LowerTime       = InSection.HasStartFrame() ? InSection.GetInclusiveStartFrame() / FrameResolution : ViewRange.GetLowerBoundValue();
-	double         UpperTime       = InSection.HasEndFrame()   ? InSection.GetExclusiveEndFrame()   / FrameResolution : ViewRange.GetUpperBoundValue();
+	FFrameRate     TickResolution  = InSection.GetTypedOuter<UMovieScene>()->GetTickResolution();
+	double         LowerTime       = InSection.HasStartFrame() ? InSection.GetInclusiveStartFrame() / TickResolution : ViewRange.GetLowerBoundValue();
+	double         UpperTime       = InSection.HasEndFrame()   ? InSection.GetExclusiveEndFrame()   / TickResolution : ViewRange.GetUpperBoundValue();
 
-	return FTimeToPixel(InSectionGeometry, TRange<double>(LowerTime, UpperTime), FrameResolution);
+	return FTimeToPixel(InSectionGeometry, TRange<double>(LowerTime, UpperTime), TickResolution);
 }
 
 
@@ -400,7 +401,7 @@ struct FSequencerSectionPainterImpl : FSequencerSectionPainter
 		SectionHandle.GetSectionObject()->EvaluateEasing(Time, EaseInValue, EaseOutValue, &EaseInInterp, &EaseOutInterp);
 
 		return FEasingCurvePoint(
-			FVector2D(Time / TimeToPixelConverter.GetFrameResolution(), EaseInValue.Get(1.f) * EaseOutValue.Get(1.f)),
+			FVector2D(Time / TimeToPixelConverter.GetTickResolution(), EaseInValue.Get(1.f) * EaseOutValue.Get(1.f)),
 			FMath::Lerp(FinalTint, EaseSelectionColor, GetEaseHighlightAmount(SectionHandle, EaseInInterp, EaseOutInterp))
 		);
 	}
@@ -423,12 +424,12 @@ struct FSequencerSectionPainterImpl : FSequencerSectionPainter
 			if ((Upper.Location.X - Lower.Location.X)*.5f > MinTimeSize)
 			{
 				float      NewPointTime  = (Upper.Location.X + Lower.Location.X)*.5f;
-				FFrameTime FrameTime     = NewPointTime * TimeToPixelConverter.GetFrameResolution();
+				FFrameTime FrameTime     = NewPointTime * TimeToPixelConverter.GetTickResolution();
 				float      NewPointValue = SectionObject->EvaluateEasing(FrameTime);
 
 				// Check that the gradient is changing significantly
 				float LinearValue = (Upper.Location.Y + Lower.Location.Y) * .5f;
-				float PointGradient = NewPointValue - SectionObject->EvaluateEasing(FMath::Lerp(Lower.Location.X, NewPointTime, 0.9f) * TimeToPixelConverter.GetFrameResolution());
+				float PointGradient = NewPointValue - SectionObject->EvaluateEasing(FMath::Lerp(Lower.Location.X, NewPointTime, 0.9f) * TimeToPixelConverter.GetTickResolution());
 				float OuterGradient = Upper.Location.Y - Lower.Location.Y;
 				if (!FMath::IsNearlyEqual(OuterGradient, PointGradient, GradientThreshold) ||
 					!FMath::IsNearlyEqual(LinearValue, NewPointValue, ValueThreshold))
@@ -508,8 +509,8 @@ struct FSequencerSectionPainterImpl : FSequencerSectionPainter
 
 			for (const FEasingCurvePoint& Point : CurvePoints)
 			{
-				float SegmentStartTime = MovieScene::DiscreteInclusiveLower(Segment.Range) / TimeToPixelConverter.GetFrameResolution();
-				float U = (Point.Location.X - SegmentStartTime) / ( FFrameNumber(MovieScene::DiscreteSize(Segment.Range)) / TimeToPixelConverter.GetFrameResolution() );
+				float SegmentStartTime = MovieScene::DiscreteInclusiveLower(Segment.Range) / TimeToPixelConverter.GetTickResolution();
+				float U = (Point.Location.X - SegmentStartTime) / ( FFrameNumber(MovieScene::DiscreteSize(Segment.Range)) / TimeToPixelConverter.GetTickResolution() );
 
 				// Add verts top->bottom
 				FVector2D UV(U, 0.f);
@@ -720,14 +721,15 @@ void SSequencerSection::GetKeysUnderMouse( const FVector2D& MousePosition, const
 
 	UMovieSceneSection& Section = *SectionInterface->GetSectionObject();
 
-	FTimeToPixel TimeToPixelConverter = ConstructTimeConverterForSection(SectionGeometry, Section, GetSequencer());
-	const float  MousePixel           = SectionGeometry.AbsoluteToLocal( MousePosition ).X;
+	FTimeToPixel    TimeToPixelConverter = ConstructTimeConverterForSection(SectionGeometry, Section, GetSequencer());
+	const FVector2D MousePixel           = SectionGeometry.AbsoluteToLocal( MousePosition );
 
 	// HitTest 
 	const FFrameTime HalfKeySizeFrames = TimeToPixelConverter.PixelDeltaToFrame( SequencerSectionConstants::KeySize.X*.5f );
-	const FFrameTime MouseFrameTime    = TimeToPixelConverter.PixelToFrame( MousePixel ) - FFrameTime(0, 0.5f);
+	const FFrameTime MouseFrameTime    = TimeToPixelConverter.PixelToFrame( MousePixel.X );
 
 	TRange<FFrameNumber> HitTestRange = TRange<FFrameNumber>( (MouseFrameTime-HalfKeySizeFrames).CeilToFrame(), (MouseFrameTime+HalfKeySizeFrames).CeilToFrame());
+	
 	if (HitTestRange.IsEmpty())
 	{
 		return;
@@ -736,8 +738,14 @@ void SSequencerSection::GetKeysUnderMouse( const FVector2D& MousePosition, const
 	// Search every key area until we find the one under the mouse
 	for (const FSectionLayoutElement& Element : Layout->GetElements())
 	{
-		// Is the key area under the mouse
-		if( !GetKeyAreaGeometry( Element, AllottedGeometry ).IsUnderLocation( MousePosition ) )
+		const FGeometry KeyAreaGeometry = GetKeyAreaGeometry(Element, AllottedGeometry);
+		const FVector2D LocalMousePixel = KeyAreaGeometry.AbsoluteToLocal(MousePosition);
+		const float LocalKeyPosY = KeyAreaGeometry.GetLocalSize().Y * .5f;
+
+		// Check that this section is under our mouse, and discard it from potential selection if the mouse is higher than the key's height. We have to
+		// check keys on a per-section basis (and not for the overall SectionGeometry) because keys are offset on tracks that have expandable 
+		// ranges (ie: Audio, Animation) which otherwise makes them fail the height-threshold check.
+		if (!KeyAreaGeometry.IsUnderLocation(MousePosition) || FMath::Abs(LocalKeyPosY - LocalMousePixel.Y) > SequencerSectionConstants::KeySize.Y *.5f)
 		{
 			continue;
 		}
@@ -773,7 +781,7 @@ void SSequencerSection::CreateKeysUnderMouse( const FVector2D& MousePosition, co
 		Section.Modify();
 
 		// Offset by 1 pixel worth of time if possible
-		const FFrameTime TimeFuzz = (GetSequencer().GetViewRange().Size<double>() / ParentGeometry.GetLocalSize().X) * Section.GetTypedOuter<UMovieScene>()->GetFrameResolution();
+		const FFrameTime TimeFuzz = (GetSequencer().GetViewRange().Size<double>() / ParentGeometry.GetLocalSize().X) * Section.GetTypedOuter<UMovieScene>()->GetTickResolution();
 
 		for (const FSequencerSelectedKey& PressedKey : InPressedKeys)
 		{
@@ -868,7 +876,7 @@ bool SSequencerSection::CheckForEasingHandleInteraction( const FPointerEvent& Mo
 		if (EasingSectionObj->HasStartFrame())
 		{
 			TRange<FFrameNumber> EaseInRange      = EasingSectionObj->GetEaseInRange();
-			double               HandlePositionIn = ( EaseInRange.IsEmpty() ? EasingSectionObj->GetInclusiveStartFrame() : EaseInRange.GetUpperBoundValue() ) / TimeToPixelConverter.GetFrameResolution();
+			double               HandlePositionIn = ( EaseInRange.IsEmpty() ? EasingSectionObj->GetInclusiveStartFrame() : EaseInRange.GetUpperBoundValue() ) / TimeToPixelConverter.GetTickResolution();
 
 			if (FMath::IsNearlyEqual(MouseTime, HandlePositionIn, HalfHandleSizeX))
 			{
@@ -880,7 +888,7 @@ bool SSequencerSection::CheckForEasingHandleInteraction( const FPointerEvent& Mo
 		if (EasingSectionObj->HasEndFrame())
 		{
 			TRange<FFrameNumber> EaseOutRange      = EasingSectionObj->GetEaseOutRange();
-			double               HandlePositionOut = (EaseOutRange.IsEmpty() ? EasingSectionObj->GetExclusiveEndFrame() : EaseOutRange.GetLowerBoundValue() ) / TimeToPixelConverter.GetFrameResolution();
+			double               HandlePositionOut = (EaseOutRange.IsEmpty() ? EasingSectionObj->GetExclusiveEndFrame() : EaseOutRange.GetLowerBoundValue() ) / TimeToPixelConverter.GetTickResolution();
 
 			if (FMath::IsNearlyEqual(MouseTime, HandlePositionOut, HalfHandleSizeX))
 			{
@@ -1040,6 +1048,8 @@ int32 SSequencerSection::OnPaint( const FPaintArgs& Args, const FGeometry& Allot
 	Painter.LayerId = LayerId;
 	Painter.bParentEnabled = bEnabled;
 	Painter.bIsHighlighted = IsSectionHighlighted(FSectionHandle(ParentSectionArea, SectionIndex), Hotspot);
+	auto& Selection = ParentSectionArea->GetSequencer().GetSelection();
+	Painter.bIsSelected = Selection.IsSelected(&SectionObject);
 
 	FSlateClippingZone ClippingZone(Painter.SectionClippingRect);
 	OutDrawElements.PushClip(ClippingZone);
@@ -1146,6 +1156,9 @@ void SSequencerSection::PaintKeys( FSequencerSectionPainter& InPainter, const FW
 	// draw all keys in each key area
 	UMovieSceneSection& SectionObject = *SectionInterface->GetSectionObject();
 
+	// Use the sub sequence range to draw valid keys, else use the current sequence's playback range
+	TRange<FFrameNumber> ValidKeyRange = Sequencer.GetSubSequenceRange().Get(Sequencer.GetFocusedMovieSceneSequence()->GetMovieScene()->GetPlaybackRange());
+
 	const ESlateDrawEffect DrawEffects = InPainter.bParentEnabled
 		? ESlateDrawEffect::None
 		: ESlateDrawEffect::DisabledEffect;
@@ -1238,17 +1251,15 @@ void SSequencerSection::PaintKeys( FSequencerSectionPainter& InPainter, const FW
 			);
 		}
 
-		double HalfFrameInterval = Sequencer.GetFocusedFrameResolution().AsInterval() * 0.5f;
-
 		// Gather keys for a region larger than the view range to ensure we draw keys that are only just offscreen.
 		TRange<double> PaddedViewRange;
 		{
 			// Compute visible range taking into account a half-frame offset for keys, plus half a key width for keys that are partially offscreen
 			TRange<FFrameNumber> SectionRange   = SectionObject.GetRange();
 			const double         HalfKeyWidth   = 0.5f * (TimeToPixelConverter.PixelToSeconds(SequencerSectionConstants::KeySize.X) - TimeToPixelConverter.PixelToSeconds(0));
-			TRange<double>       VisibleRange   = MovieScene::DilateRange(GetSequencer().GetViewRange(), -(HalfFrameInterval + HalfKeyWidth), -HalfFrameInterval + HalfKeyWidth);
+			TRange<double>       VisibleRange   = MovieScene::DilateRange(GetSequencer().GetViewRange(), -HalfKeyWidth, HalfKeyWidth);
 
-			PaddedViewRange = TRange<double>::Intersection(SectionRange / Sequencer.GetFocusedFrameResolution(), VisibleRange);
+			PaddedViewRange = TRange<double>::Intersection(SectionRange / Sequencer.GetFocusedTickResolution(), VisibleRange);
 		}
 
 		// For each key area, we gather draw params for the visible time range, then combine
@@ -1256,8 +1267,8 @@ void SSequencerSection::PaintKeys( FSequencerSectionPainter& InPainter, const FW
 		auto& Found = CachedKeyAreaPositions.FindChecked(LayoutElement);
 		TArrayView<const FSequencerCachedKeys> CachedKeys = Found;
 
-		// Can't do any of the rest if there are no keys
-		if (CachedKeys.Num() == 0)
+		// Can't do any of the rest if there are no keys, or there's no range to draw within
+		if (CachedKeys.Num() == 0 || PaddedViewRange.IsEmpty())
 		{
 			continue;
 		}
@@ -1275,6 +1286,7 @@ void SSequencerSection::PaintKeys( FSequencerSectionPainter& InPainter, const FW
 
 			int32 NextIndex;
 			TArrayView<const double> TimesInRange;
+			TArrayView<const FFrameNumber> FramesInRange;
 			TArrayView<const FKeyHandle> HandlesInRange;
 			TArray<FKeyDrawParams> DrawParams;
 		};
@@ -1286,7 +1298,7 @@ void SSequencerSection::PaintKeys( FSequencerSectionPainter& InPainter, const FW
 			// Gather all the key handles in this view range
 			FKeyDrawInformation KeyDrawInfo;
 
-			Cache.GetKeysInRange(PaddedViewRange, &KeyDrawInfo.TimesInRange, &KeyDrawInfo.HandlesInRange);
+			Cache.GetKeysInRange(PaddedViewRange, &KeyDrawInfo.TimesInRange, &KeyDrawInfo.FramesInRange, &KeyDrawInfo.HandlesInRange);
 
 			// Generate draw data for this key area
 			if (KeyDrawInfo.TimesInRange.Num())
@@ -1328,6 +1340,9 @@ void SSequencerSection::PaintKeys( FSequencerSectionPainter& InPainter, const FW
 			// and that we avoid floating point precision issues where there is only one key [(KeyTime + TimeOverlapThreshold) - KeyTime != TimeOverlapThreshold] for some floats
 			CardinalKeyTime += TimeOverlapThreshold*0.9994f;
 
+			// Generate the hull of frame numbers that contribute to this key so we can draw it enabled/disabled depending on whether it is outside of the valid range or not
+			TRange<FFrameNumber> KeyRange = TRange<FFrameNumber>::Empty();
+
 			float AverageKeyTime = 0.f;
 			int32 NumKeyTimes = 0;
 			int32 NumOverlaps = 0;
@@ -1345,12 +1360,14 @@ void SSequencerSection::PaintKeys( FSequencerSectionPainter& InPainter, const FW
 				{
 					int32 FinalIndexInThreshold = Info.NextIndex + 1;
 
+					KeyRange = TRange<FFrameNumber>::Hull(KeyRange, TRange<FFrameNumber>(Info.FramesInRange[Info.NextIndex]));
 					AverageKeyTime += Info.TimesInRange[Info.NextIndex];
 					++NumKeyTimes;
 
 					// Count the number of overlapping keys
 					while (FinalIndexInThreshold < Info.TimesInRange.Num() && FMath::IsNearlyEqual(Info.TimesInRange[FinalIndexInThreshold], CardinalKeyTime, TimeOverlapThreshold))
 					{
+						KeyRange = TRange<FFrameNumber>::Hull(KeyRange, TRange<FFrameNumber>(Info.FramesInRange[FinalIndexInThreshold]));
 						AverageKeyTime += Info.TimesInRange[FinalIndexInThreshold];
 						++NumKeyTimes;
 
@@ -1366,7 +1383,7 @@ void SSequencerSection::PaintKeys( FSequencerSectionPainter& InPainter, const FW
 				}
 			}
 
-			const float FinalKeyPosition = TimeToPixelConverter.SecondsToPixel((AverageKeyTime / NumKeyTimes) + HalfFrameInterval);
+			const float FinalKeyPosition = TimeToPixelConverter.SecondsToPixel(AverageKeyTime / NumKeyTimes);
 
 			static const FSlateBrush* PartialKeyBrush = FEditorStyle::GetBrush("Sequencer.PartialKey");
 
@@ -1464,6 +1481,8 @@ void SSequencerSection::PaintKeys( FSequencerSectionPainter& InPainter, const FW
 				KeyDrawParam->BorderTint = FLinearColor(0.83f, 0.12f, 0.12f, 1.0f); // Red
 			}
 
+			const ESlateDrawEffect KeyDrawEffects = ValidKeyRange.Contains(KeyRange) ? DrawEffects : ESlateDrawEffect::DisabledEffect;
+
 			// draw border
 			static FVector2D ThrobAmount(12.f, 12.f);
 			const FVector2D KeySize = bSelected ? SequencerSectionConstants::KeySize + ThrobAmount * KeyThrobScaleValue : SequencerSectionConstants::KeySize;
@@ -1481,7 +1500,7 @@ void SSequencerSection::PaintKeys( FSequencerSectionPainter& InPainter, const FW
 					KeySize
 				),
 				KeyDrawParam->BorderBrush,
-				DrawEffects,
+				KeyDrawEffects,
 				KeyDrawParam->BorderTint
 			);
 
@@ -1500,7 +1519,7 @@ void SSequencerSection::PaintKeys( FSequencerSectionPainter& InPainter, const FW
 					KeySize - 2.0f * BrushBorderWidth
 				),
 				KeyDrawParam->FillBrush,
-				DrawEffects,
+				KeyDrawEffects,
 				KeyDrawParam->FillTint
 			);
 		}
@@ -1773,7 +1792,7 @@ void SSequencerSection::Tick( const FGeometry& AllottedGeometry, const double In
 		TMap<FSectionLayoutElement, TArray<FSequencerCachedKeys, TInlineAllocator<1>>, FDefaultSetAllocator, FLayoutElementKeyFuncs> OldCachedKeyAreaPositions;
 		Swap(OldCachedKeyAreaPositions, CachedKeyAreaPositions);
 
-		FFrameRate FrameResolution = GetSequencer().GetFocusedFrameResolution();
+		FFrameRate TickResolution = GetSequencer().GetFocusedTickResolution();
 
 		// Move over any existing still valid caches 
 		for (const FSectionLayoutElement& LayoutElement : Layout->GetElements())
@@ -1786,7 +1805,7 @@ void SSequencerSection::Tick( const FGeometry& AllottedGeometry, const double In
 				for (TSharedRef<IKeyArea> KeyArea : LayoutElement.GetKeyAreas())
 				{
 					NewCachedKeys.Emplace();
-					NewCachedKeys.Last().Update(KeyArea, FrameResolution);
+					NewCachedKeys.Last().Update(KeyArea, TickResolution);
 				}
 				CachedKeyAreaPositions.Add(LayoutElement, MoveTemp(NewCachedKeys));
 			}
@@ -1795,7 +1814,7 @@ void SSequencerSection::Tick( const FGeometry& AllottedGeometry, const double In
 				// We can reuse this one
 				for (FSequencerCachedKeys& CachedKeys : *CacheArray)
 				{
-					CachedKeys.Update(CachedKeys.GetKeyArea().ToSharedRef(), FrameResolution);
+					CachedKeys.Update(CachedKeys.GetKeyArea().ToSharedRef(), TickResolution);
 				}
 				CachedKeyAreaPositions.Add(LayoutElement, MoveTemp(*CacheArray));
 			}
@@ -1804,7 +1823,7 @@ void SSequencerSection::Tick( const FGeometry& AllottedGeometry, const double In
 		UMovieSceneSection* Section = SectionInterface->GetSectionObject();
 		if (Section && Section->HasStartFrame() && Section->HasEndFrame())
 		{
-			FTimeToPixel TimeToPixelConverter(ParentGeometry, GetSequencer().GetViewRange(), FrameResolution);
+			FTimeToPixel TimeToPixelConverter(ParentGeometry, GetSequencer().GetViewRange(), TickResolution);
 
 			const int32 SectionLengthPx = FMath::Max(0,
 				FMath::RoundToInt(
diff --git a/Engine/Source/Editor/Sequencer/Private/SSequencerSectionAreaView.cpp b/Engine/Source/Editor/Sequencer/Private/SSequencerSectionAreaView.cpp
--- a/Engine/Source/Editor/Sequencer/Private/SSequencerSectionAreaView.cpp
+++ b/Engine/Source/Editor/Sequencer/Private/SSequencerSectionAreaView.cpp
@@ -27,7 +27,7 @@ namespace SequencerSectionUtils
 	FGeometry GetSectionGeometry( const FGeometry& AllottedGeometry, int32 RowIndex, int32 MaxTracks, float NodeHeight, TSharedPtr<ISequencerSection> SectionInterface, const FTimeToPixel& TimeToPixelConverter )
 	{
 		const UMovieSceneSection* Section    = SectionInterface->GetSectionObject();
-		const FFrameRate          Resolution = Section->GetTypedOuter<UMovieScene>()->GetFrameResolution();
+		const FFrameRate          Resolution = Section->GetTypedOuter<UMovieScene>()->GetTickResolution();
 
 		// Sections with an infite (open) start bound use the currently visible geometry and time range
 		float PixelStartX = Section->HasStartFrame() ? TimeToPixelConverter.FrameToPixel( Section->GetInclusiveStartFrame() )     : AllottedGeometry.Position.X;
@@ -186,7 +186,7 @@ void SSequencerSectionAreaView::OnArrangeChildren( const FGeometry& AllottedGeom
 FTimeToPixel SSequencerSectionAreaView::GetTimeToPixel( const FGeometry& AllottedGeometry ) const
 {
 	const UMovieSceneTrack* Track      = SectionAreaNode->GetTrack();
-	const FFrameRate        Resolution = Track->GetTypedOuter<UMovieScene>()->GetFrameResolution();
+	const FFrameRate        Resolution = Track->GetTypedOuter<UMovieScene>()->GetTickResolution();
 
 	return FTimeToPixel( AllottedGeometry, ViewRange.Get(), Resolution );
 }
diff --git a/Engine/Source/Editor/Sequencer/Private/SSequencerShotFilterOverlay.cpp b/Engine/Source/Editor/Sequencer/Private/SSequencerShotFilterOverlay.cpp
--- a/Engine/Source/Editor/Sequencer/Private/SSequencerShotFilterOverlay.cpp
+++ b/Engine/Source/Editor/Sequencer/Private/SSequencerShotFilterOverlay.cpp
@@ -24,7 +24,7 @@ int32 SSequencerShotFilterOverlay::OnPaint( const FPaintArgs& Args, const FGeome
 
 	if (Alpha > 0.f)
 	{
-		FTimeToPixel TimeToPixelConverter(AllottedGeometry, ViewRange.Get(), Sequencer.Pin()->GetFocusedFrameResolution());
+		FTimeToPixel TimeToPixelConverter(AllottedGeometry, ViewRange.Get(), Sequencer.Pin()->GetFocusedTickResolution());
 		
 		TRange<float> TimeBounds = TRange<float>(
 			TimeToPixelConverter.PixelToSeconds(0),
