Copying //UE4/Dev-Sequencer to //UE4/Dev-Main (Source: //UE4/Dev-Sequencer @ 4044339)

diff --git a/Engine/Source/Editor/Sequencer/Private/SSequencerTimePanel.cpp b/Engine/Source/Editor/Sequencer/Private/SSequencerTimePanel.cpp
--- a/Engine/Source/Editor/Sequencer/Private/SSequencerTimePanel.cpp
+++ b/Engine/Source/Editor/Sequencer/Private/SSequencerTimePanel.cpp
@@ -68,21 +68,15 @@ void SSequencerTimePanel::Construct(const FArguments& InArgs, TWeakPtr<FSequence
 {
 	WeakSequencer = InSequencer;
 
-	TSharedPtr<FSequencer> Sequencer = InSequencer.Pin();
-	if (Sequencer.IsValid())
-	{
-		CurrentFrameResolution = Sequencer->GetFocusedFrameResolution();
-	}
-
-	TArray<FCommonFrameRateInfo> FrameResolutionRates;
+	TArray<FCommonFrameRateInfo> TickResolutionRates;
 	{
 		TArrayView<const FCommonFrameRateInfo> CommonRates = FCommonFrameRates::GetAll();
-		FrameResolutionRates.Append(CommonRates.GetData(), CommonRates.Num());
+		TickResolutionRates.Append(CommonRates.GetData(), CommonRates.Num());
 
-		FrameResolutionRates.Add(FCommonFrameRateInfo{ FFrameRate(1000, 1),   LOCTEXT("1000_Name",   "1000 fps (ms precision)"),                       LOCTEXT("1000_Description",   "Allows placement of sequence keys and sections with millisecond precision") });
-		FrameResolutionRates.Add(FCommonFrameRateInfo{ FFrameRate(24000, 1),  LOCTEXT("24000_Name",  "24000 fps (all integer rates + 23.976)"),        LOCTEXT("24000_Description",  "A very high framerate that allows frame-accurate evaluation of all common integer frame rates as well as NTSC 24.") });
-		FrameResolutionRates.Add(FCommonFrameRateInfo{ FFrameRate(60000, 1),  LOCTEXT("60000_Name",  "60000 fps (all integer rates + 29.97 & 59.94)"), LOCTEXT("60000_Description",  "A very high framerate that allows frame-accurate evaluation of all common integer frame rates as well as NTSC 30 and 60.") });
-		FrameResolutionRates.Add(FCommonFrameRateInfo{ FFrameRate(120000, 1), LOCTEXT("120000_Name", "120000 fps (all common rates)"),                 LOCTEXT("120000_Description", "A very high framerate that allows frame-accurate evaluation of all common integer and NTSC frame rates.") });
+		TickResolutionRates.Add(FCommonFrameRateInfo{ FFrameRate(1000, 1),   LOCTEXT("1000_Name",   "1000 fps (ms precision)"),                       LOCTEXT("1000_Description",   "Allows placement of sequence keys and sections with millisecond precision") });
+		TickResolutionRates.Add(FCommonFrameRateInfo{ FFrameRate(24000, 1),  LOCTEXT("24000_Name",  "24000 fps (all integer rates + 23.976)"),        LOCTEXT("24000_Description",  "A very high framerate that allows frame-accurate evaluation of all common integer frame rates as well as NTSC 24.") });
+		TickResolutionRates.Add(FCommonFrameRateInfo{ FFrameRate(60000, 1),  LOCTEXT("60000_Name",  "60000 fps (all integer rates + 29.97 & 59.94)"), LOCTEXT("60000_Description",  "A very high framerate that allows frame-accurate evaluation of all common integer frame rates as well as NTSC 30 and 60.") });
+		TickResolutionRates.Add(FCommonFrameRateInfo{ FFrameRate(120000, 1), LOCTEXT("120000_Name", "120000 fps (all common rates)"),                 LOCTEXT("120000_Description", "A very high framerate that allows frame-accurate evaluation of all common integer and NTSC frame rates.") });
 	}
 
 	FText Description = LOCTEXT("Description", "Sequences store section start times and keys at points in time called 'ticks'.\n\nThe size of a single tick is defined per-sequence; it is recommended that you choose a tick-interval that fits into your desired display rate or content frame rates. Increasing the resolution will reduce the total supported time range.");
@@ -119,6 +113,14 @@ void SSequencerTimePanel::Construct(const FArguments& InArgs, TWeakPtr<FSequence
 					[
 						SNew(SVerticalBox)
 
+						+ SVerticalBox::Slot()
+						.AutoHeight()
+						[
+							SNew(STextBlock)
+							.TextStyle(FEditorStyle::Get(), "LargeText")
+							.Text(LOCTEXT("Title", "Advanced Time Properties"))
+						]
+
 						+ SVerticalBox::Slot()
 						[
 							SNew(SBox)
@@ -126,13 +128,6 @@ void SSequencerTimePanel::Construct(const FArguments& InArgs, TWeakPtr<FSequence
 							[
 								SNew(SScrollBox)
 
-								+ SScrollBox::Slot()
-								[
-									SNew(STextBlock)
-									.TextStyle(FEditorStyle::Get(), "LargeText")
-									.Text(LOCTEXT("Title", "Advanced Time Properties"))
-								]
-
 								+ SScrollBox::Slot()
 								.Padding(FMargin(0.f, 0.f, 0.f, 10.f))
 								[
@@ -165,9 +160,9 @@ void SSequencerTimePanel::Construct(const FArguments& InArgs, TWeakPtr<FSequence
 										.NotRecommendedText(LOCTEXT("NotCompatibleWithDisplayRate", "Other"))
 										.NotRecommendedToolTip(LOCTEXT("NotCompatibleWithDisplayRate_Tip", "All other preset frame rates that are not compatible with the current display and tick rate"))
 										.IsPresetRecommended(this, &SSequencerTimePanel::IsRecommendedResolution)
-										.PresetValues(MoveTemp(FrameResolutionRates))
-										.Value(this, &SSequencerTimePanel::OnGetFrameResolution)
-										.OnValueChanged(this, &SSequencerTimePanel::OnSetFrameResolution)
+										.PresetValues(MoveTemp(TickResolutionRates))
+										.Value(this, &SSequencerTimePanel::GetCurrentTickResolution)
+										.OnValueChanged(this, &SSequencerTimePanel::OnSetTickResolution)
 									]
 
 									+ SGridPanel::Slot(0, 1)
@@ -249,11 +244,30 @@ void SSequencerTimePanel::Construct(const FArguments& InArgs, TWeakPtr<FSequence
 						.HAlign(HAlign_Center)
 						.VAlign(VAlign_Center)
 						[
-							SNew(SButton)
-							.OnClicked(this, &SSequencerTimePanel::Apply)
+							SNew(SHorizontalBox)
+
+							+SHorizontalBox::Slot()
+							.Padding(FMargin(0, 0, 2, 0))
+							.AutoWidth()
 							[
-								SNew(STextBlock)
-								.Text(LOCTEXT("ApplyButtonText", "Apply"))
+								SNew(SButton)
+								.OnClicked(this, &SSequencerTimePanel::Apply)
+								[
+									SNew(STextBlock)
+									.Text(LOCTEXT("ApplyButtonText", "Apply"))
+								]
+							]
+
+							+SHorizontalBox::Slot()
+							.Padding(FMargin(2, 0, 0, 0))
+							.AutoWidth()
+							[
+								SNew(SButton)
+								.OnClicked(this, &SSequencerTimePanel::Close)
+								[
+									SNew(STextBlock)
+									.Text(LOCTEXT("CancelButtonText", "Cancel"))
+								]
 							]
 						]
 					]
@@ -281,6 +295,7 @@ FReply SSequencerTimePanel::OnBorderFadeClicked(const FGeometry& MyGeometry, con
 
 FReply SSequencerTimePanel::Close()
 {
+	CurrentTickResolution.Reset();
 	SetVisibility(EVisibility::Collapsed);
 	return FReply::Handled();
 }
@@ -292,8 +307,8 @@ FReply SSequencerTimePanel::Apply()
 
 	if (MovieScene)
 	{
-		FFrameRate Src = MovieScene->GetFrameResolution();
-		FFrameRate Dst = CurrentFrameResolution;
+		FFrameRate Src = MovieScene->GetTickResolution();
+		FFrameRate Dst = GetCurrentTickResolution();
 
 		FScopedTransaction ScopedTransaction(FText::Format(LOCTEXT("MigrateFrameTimes", "Convert sequence tick interval from {0} to {1}"), Src.ToPrettyText(), Dst.ToPrettyText()));
 
@@ -305,8 +320,8 @@ FReply SSequencerTimePanel::Apply()
 
 EVisibility SSequencerTimePanel::GetWarningVisibility() const
 {
-	TSharedPtr<FSequencer> Sequencer  = WeakSequencer.Pin();
-	return Sequencer.IsValid() && Sequencer->GetFocusedFrameResolution().IsMultipleOf(CurrentFrameResolution) ? EVisibility::Collapsed : EVisibility::Visible;
+	TSharedPtr<FSequencer> Sequencer = WeakSequencer.Pin();
+	return Sequencer.IsValid() && Sequencer->GetFocusedTickResolution().IsMultipleOf(GetCurrentTickResolution()) ? EVisibility::Collapsed : EVisibility::Visible;
 }
 
 void SSequencerTimePanel::UpdateCommonFrameRates()
@@ -314,7 +329,7 @@ void SSequencerTimePanel::UpdateCommonFrameRates()
 	TArray<FCommonFrameRateInfo> CompatibleRates;
 	for (const FCommonFrameRateInfo& Info : FCommonFrameRates::GetAll())
 	{
-		if (Info.FrameRate.IsMultipleOf(CurrentFrameResolution))
+		if (Info.FrameRate.IsMultipleOf(GetCurrentTickResolution()))
 		{
 			CompatibleRates.Add(Info);
 		}
@@ -340,7 +355,7 @@ void SSequencerTimePanel::UpdateCommonFrameRates()
 
 FText SSequencerTimePanel::GetSupportedTimeRange() const
 {
-	double FrameRate = CurrentFrameResolution.AsDecimal();
+	double FrameRate = GetCurrentTickResolution().AsDecimal();
 
 	int64 TotalMaxSeconds = static_cast<int64>(TNumericLimits<int32>::Max() / FrameRate);
 
@@ -365,17 +380,18 @@ UMovieSceneSequence* SSequencerTimePanel::GetFocusedSequence() const
 bool SSequencerTimePanel::IsRecommendedResolution(FFrameRate InFrameRate) const
 {
 	UMovieSceneSequence* Sequence = GetFocusedSequence();
-	return !Sequence || (InFrameRate.IsFactorOf(Sequence->GetMovieScene()->GetPlaybackFrameRate()) && InFrameRate.IsFactorOf(Sequence->GetMovieScene()->GetFrameResolution()));
+	return !Sequence || (InFrameRate.IsFactorOf(Sequence->GetMovieScene()->GetDisplayRate()) && InFrameRate.IsFactorOf(Sequence->GetMovieScene()->GetTickResolution()));
 }
 
-FFrameRate SSequencerTimePanel::OnGetFrameResolution() const
+FFrameRate SSequencerTimePanel::GetCurrentTickResolution() const
 {
-	return CurrentFrameResolution;
+	TSharedPtr<FSequencer> Sequencer = WeakSequencer.Pin();
+	return CurrentTickResolution.Get(Sequencer.IsValid() ? Sequencer->GetFocusedTickResolution() : FFrameRate(24000, 1));
 }
 
-void SSequencerTimePanel::OnSetFrameResolution(FFrameRate InFrameResolution)
+void SSequencerTimePanel::OnSetTickResolution(FFrameRate InTickResolution)
 {
-	CurrentFrameResolution = InFrameResolution;
+	CurrentTickResolution = InTickResolution;
 	UpdateCommonFrameRates();
 }
 
@@ -387,7 +403,7 @@ void SSequencerTimePanel::MigrateFrameTimes(FFrameRate SourceRate, FFrameRate De
 		TotalNumTracks += Binding.GetTracks().Num();
 	}
 
-	FScopedSlowTask SlowTask(TotalNumTracks, LOCTEXT("ChangingFrameResolution", "Migrating sequence frame timing"));
+	FScopedSlowTask SlowTask(TotalNumTracks, LOCTEXT("ChangingTickResolution", "Migrating sequence frame timing"));
 	SlowTask.MakeDialogDelayed(0.25f, true);
 
 	MovieScene->Modify();
@@ -416,7 +432,7 @@ void SSequencerTimePanel::MigrateFrameTimes(FFrameRate SourceRate, FFrameRate De
 		}
 	}
 
-	MovieScene->SetFrameResolutionDirectly(DestinationRate);
+	MovieScene->SetTickResolutionDirectly(DestinationRate);
 }
 
 void SSequencerTimePanel::MigrateFrameTimes(FFrameRate SourceRate, FFrameRate DestinationRate, UMovieSceneTrack* Track)
@@ -479,7 +495,7 @@ void SSequencerTimePanel::MigrateFrameTimes(FFrameRate SourceRate, FFrameRate De
 
 	for (const FMovieSceneChannelEntry& Entry : Section->GetChannelProxy().GetAllEntries())
 	{
-		Entry.GetBatchChannelInterface().ChangeFrameResolution_Batch(Entry.GetChannels(), SourceRate, DestinationRate);
+		Entry.GetBatchChannelInterface().ChangeTickResolution_Batch(Entry.GetChannels(), SourceRate, DestinationRate);
 	}
 }
 
diff --git a/Engine/Source/Editor/Sequencer/Private/SSequencerTimePanel.h b/Engine/Source/Editor/Sequencer/Private/SSequencerTimePanel.h
--- a/Engine/Source/Editor/Sequencer/Private/SSequencerTimePanel.h
+++ b/Engine/Source/Editor/Sequencer/Private/SSequencerTimePanel.h
@@ -35,12 +35,12 @@ private:
 
 	EVisibility GetWarningVisibility() const;
 
-	FFrameRate OnGetFrameResolution() const;
+	FFrameRate GetCurrentTickResolution() const;
 	bool IsRecommendedResolution(FFrameRate InFrameRate) const;
 
 	FText GetSupportedTimeRange() const;
 
-	void OnSetFrameResolution(FFrameRate InFrameResolution);
+	void OnSetTickResolution(FFrameRate InTickResolution);
 
 	void UpdateCommonFrameRates();
 
@@ -50,7 +50,7 @@ private:
 	static void MigrateFrameTimes(FFrameRate SourceRate, FFrameRate DestinationRate, UMovieSceneTrack* Track);
 	static void MigrateFrameTimes(FFrameRate SourceRate, FFrameRate DestinationRate, UMovieSceneSection* Section);
 
-	FFrameRate CurrentFrameResolution;
+	TOptional<FFrameRate> CurrentTickResolution;
 	TSharedPtr<SVerticalBox> CommonFrameRates;
 	TWeakPtr<FSequencer> WeakSequencer;
 };
\ No newline at end of file
diff --git a/Engine/Source/Editor/Sequencer/Private/SSequencerTreeView.cpp b/Engine/Source/Editor/Sequencer/Private/SSequencerTreeView.cpp
--- a/Engine/Source/Editor/Sequencer/Private/SSequencerTreeView.cpp
+++ b/Engine/Source/Editor/Sequencer/Private/SSequencerTreeView.cpp
@@ -181,6 +181,7 @@ void SSequencerTreeView::Construct(const FArguments& InArgs, const TSharedRef<FS
 		.OnExpansionChanged(this, &SSequencerTreeView::OnExpansionChanged)
 		.AllowOverscroll(EAllowOverscroll::No)
 		.OnContextMenuOpening( this, &SSequencerTreeView::OnContextMenuOpening )
+		.OnSetExpansionRecursive(this, &SSequencerTreeView::SetItemExpansionRecursive)
 	);
 }
 
@@ -728,6 +729,11 @@ void SSequencerTreeView::OnExpansionChanged(FDisplayNodeRef InItem, bool bIsExpa
 	}
 }
 
+void SSequencerTreeView::SetItemExpansionRecursive(FDisplayNodeRef InItem, bool bIsExpanded)
+{
+	ExpandCollapseNode(InItem, bIsExpanded, ETreeRecursion::Recursive);
+}
+
 void SSequencerTreeView::OnGetChildren(FDisplayNodeRef InParent, TArray<FDisplayNodeRef>& OutChildren) const
 {
 	for (const auto& Node : InParent->GetChildNodes())
diff --git a/Engine/Source/Editor/Sequencer/Private/SSequencerTreeView.h b/Engine/Source/Editor/Sequencer/Private/SSequencerTreeView.h
--- a/Engine/Source/Editor/Sequencer/Private/SSequencerTreeView.h
+++ b/Engine/Source/Editor/Sequencer/Private/SSequencerTreeView.h
@@ -141,6 +141,8 @@ private:
 	/** Handles the context menu opening when right clicking on the tree view. */
 	TSharedPtr<SWidget> OnContextMenuOpening();
 
+	void SetItemExpansionRecursive(FDisplayNodeRef InItem, bool bIsExpanded);
+
 	FOnGetContextMenuContent OnGetContextMenuContent;
 
 public:
diff --git a/Engine/Source/Editor/Sequencer/Private/Sequencer.cpp b/Engine/Source/Editor/Sequencer/Private/Sequencer.cpp
--- a/Engine/Source/Editor/Sequencer/Private/Sequencer.cpp
+++ b/Engine/Source/Editor/Sequencer/Private/Sequencer.cpp
@@ -188,9 +188,9 @@ struct FSequencerCurveEditorBounds : ICurveEditorBounds
 		TSharedPtr<FSequencer> Sequencer = WeakSequencer.Pin();
 		if (Sequencer.IsValid())
 		{
-			FFrameRate FrameResolution = Sequencer->GetFocusedFrameResolution();
+			FFrameRate TickResolution = Sequencer->GetFocusedTickResolution();
 
-			if (InMin * FrameResolution > TNumericLimits<int32>::Lowest() && InMax * FrameResolution < TNumericLimits<int32>::Max())
+			if (InMin * TickResolution > TNumericLimits<int32>::Lowest() && InMax * TickResolution < TNumericLimits<int32>::Max())
 			{
 				Sequencer->SetViewRange(TRange<double>(InMin, InMax), EViewRangeInterpolation::Immediate);
 			}
@@ -248,6 +248,7 @@ public:
 void FSequencer::InitSequencer(const FSequencerInitParams& InitParams, const TSharedRef<ISequencerObjectChangeListener>& InObjectChangeListener, const TArray<FOnCreateTrackEditor>& TrackEditorDelegates, const TArray<FOnCreateEditorObjectBinding>& EditorObjectBindingDelegates)
 {
 	bIsEditingWithinLevelEditor = InitParams.bEditWithinLevelEditor;
+	ScrubStyle = InitParams.ViewParams.ScrubberStyle;
 
 	SilentModeCount = 0;
 	bReadOnly = InitParams.ViewParams.bReadOnly;
@@ -294,7 +295,7 @@ void FSequencer::InitSequencer(const FSequencerInitParams& InitParams, const TSh
 		CurveEditorModel->OnOutputSnapEnabledChanged  = FOnSetBoolean::CreateLambda([this](bool NewValue){ Settings->SetSnapCurveValueToInterval(NewValue); });
 
 		CurveEditorModel->OutputSnapIntervalAttribute = MakeAttributeLambda([this]{ return (double)Settings->GetCurveValueSnapInterval(); });
-		CurveEditorModel->InputSnapRateAttribute      = MakeAttributeSP(this, &FSequencer::GetFocusedPlayRate);
+		CurveEditorModel->InputSnapRateAttribute      = MakeAttributeSP(this, &FSequencer::GetFocusedDisplayRate);
 
 		CurveEditorModel->DefaultKeyAttributes        = MakeAttributeSP(this, &FSequencer::GetDefaultKeyAttributes);
 	}
@@ -589,7 +590,7 @@ void FSequencer::Tick(float InDeltaTime)
 		UMovieSceneSubSection* Section = UMovieSceneSubSection::GetRecordingSection();
 		if(Section != nullptr && Section->HasStartFrame())
 		{
-			FFrameRate   SectionResolution = Section->GetTypedOuter<UMovieScene>()->GetFrameResolution();
+			FFrameRate   SectionResolution = Section->GetTypedOuter<UMovieScene>()->GetTickResolution();
 			FFrameNumber RecordingLength   = SequenceRecorder.GetCurrentRecordingLength().ConvertTo(SectionResolution).CeilToFrame();
 
 			if (RecordingLength > 0)
@@ -778,7 +779,7 @@ void FSequencer::GetKeysFromSelection(TUniquePtr<FSequencerKeyCollection>& KeyCo
 		SelectedNodes.Add(&Node.Get());
 	}
 
-	FFrameNumber ThresholdFrames = (DuplicateThresholdSeconds * GetFocusedFrameResolution()).FloorToFrame();
+	FFrameNumber ThresholdFrames = (DuplicateThresholdSeconds * GetFocusedTickResolution()).FloorToFrame();
 	KeyCollection->Update(FSequencerKeyCollectionSignature::FromNodesRecursive(SelectedNodes, ThresholdFrames));
 }
 
@@ -796,7 +797,7 @@ void FSequencer::GetAllKeys(TUniquePtr<FSequencerKeyCollection>& KeyCollection,
 		AllNodes.Add(&Node.Get());
 	}
 
-	FFrameNumber ThresholdFrames = (DuplicateThresholdSeconds * GetFocusedFrameResolution()).FloorToFrame();
+	FFrameNumber ThresholdFrames = (DuplicateThresholdSeconds * GetFocusedTickResolution()).FloorToFrame();
 	KeyCollection->Update(FSequencerKeyCollectionSignature::FromNodesRecursive(AllNodes, ThresholdFrames));
 }
 
@@ -843,7 +844,7 @@ void FSequencer::UpdateSubSequenceData()
 	RootToLocalTransform = FMovieSceneSequenceTransform();
 
 	// Find the parent sub section and set up the sub sequence range, if necessary
-	if (ActiveTemplateIDs.Num() <= 1 || Settings->ShouldEvaluateSubSequencesInIsolation())
+	if (ActiveTemplateIDs.Num() <= 1)
 	{
 		return;
 	}
@@ -1057,6 +1058,84 @@ void FSequencer::SetInterpTangentMode(ERichCurveInterpMode InterpMode, ERichCurv
 	}
 }
 
+void FSequencer::ToggleInterpTangentWeightMode()
+{
+	// @todo: sequencer-timecode: move this float-specific logic elsewhere to make it extensible for any channel type
+
+	TArray<FSequencerSelectedKey> SelectedKeysArray = Selection.GetSelectedKeys().Array();
+	if (SelectedKeysArray.Num() == 0)
+	{
+		return;
+	}
+
+	FScopedTransaction SetInterpTangentWeightModeTransaction(NSLOCTEXT("Sequencer", "ToggleInterpTangentWeightMode_Transaction", "Toggle Tangent Weight Mode"));
+	bool bAnythingChanged = false;
+
+	FSelectedKeysByChannelType KeysByChannel(SelectedKeysArray);
+	TSet<UMovieSceneSection*> ModifiedSections;
+
+	const uint32 FloatChannelID = FMovieSceneFloatChannel::GetChannelID();
+
+	// Remove all tangent weights unless we find a compatible key that does not have weights yet
+	ERichCurveTangentWeightMode WeightModeToApply = RCTWM_WeightedNone;
+
+	// First off iterate all the current keys and find any that don't have weights
+	for (TTuple<void*, FSelectedChannelInfo>& Pair : KeysByChannel.ChannelToKeyHandleMap)
+	{
+		if (Pair.Value.ChannelTypeID == FloatChannelID)
+		{
+			TMovieSceneChannel<FMovieSceneFloatValue> ChannelInterface = static_cast<FMovieSceneFloatChannel*>(Pair.Key)->GetInterface();
+			TArrayView<FMovieSceneFloatValue> Values = ChannelInterface.GetValues();
+
+			for (FKeyHandle Handle : Pair.Value.KeyHandles)
+			{
+				const int32 KeyIndex = ChannelInterface.GetIndex(Handle);
+				if (KeyIndex != INDEX_NONE && Values[KeyIndex].InterpMode == RCIM_Cubic && Values[KeyIndex].Tangent.TangentWeightMode == RCTWM_WeightedNone)
+				{
+					WeightModeToApply = RCTWM_WeightedBoth;
+					goto assign_weights;
+				}
+			}
+		}
+	}
+
+assign_weights:
+
+	// Assign the new weight mode for all cubic keys
+	for (TTuple<void*, FSelectedChannelInfo>& Pair : KeysByChannel.ChannelToKeyHandleMap)
+	{
+		if (Pair.Value.ChannelTypeID == FloatChannelID)
+		{
+			if (!ModifiedSections.Contains(Pair.Value.OwningSection))
+			{
+				Pair.Value.OwningSection->Modify();
+				ModifiedSections.Add(Pair.Value.OwningSection);
+			}
+
+			FMovieSceneFloatChannel* Channel = static_cast<FMovieSceneFloatChannel*>(Pair.Key);
+			TMovieSceneChannel<FMovieSceneFloatValue> ChannelInterface = Channel->GetInterface();
+			TArrayView<FMovieSceneFloatValue> Values = ChannelInterface.GetValues();
+
+			for (FKeyHandle Handle : Pair.Value.KeyHandles)
+			{
+				const int32 KeyIndex = ChannelInterface.GetIndex(Handle);
+				if (KeyIndex != INDEX_NONE && Values[KeyIndex].InterpMode == RCIM_Cubic)
+				{
+					Values[KeyIndex].Tangent.TangentWeightMode = WeightModeToApply;
+					bAnythingChanged = true;
+				}
+			}
+
+			Channel->AutoSetTangents();
+		}
+	}
+
+	if (bAnythingChanged)
+	{
+		NotifyMovieSceneDataChanged(EMovieSceneDataChangeType::TrackValueChanged);
+	}
+}
+
 void FSequencer::SnapToFrame()
 {
 	ISequencerModule& SequencerModule = FModuleManager::LoadModuleChecked<ISequencerModule>("Sequencer");
@@ -1085,14 +1164,14 @@ void FSequencer::SnapToFrame()
 
 			ChannelInterface->GetKeyTimes_Raw(Pair.Key, Pair.Value.KeyHandles, KeyTimesScratch);
 
-			FFrameRate FrameResolution = GetFocusedFrameResolution();
-			FFrameRate PlayRate        = GetFocusedPlayRate();
+			FFrameRate TickResolution  = GetFocusedTickResolution();
+			FFrameRate DisplayRate     = GetFocusedDisplayRate();
 
 			for (FFrameNumber& Time : KeyTimesScratch)
 			{
 				// Convert to frame
-				FFrameNumber PlayFrame    = FFrameRate::TransformTime(Time,      FrameResolution, PlayRate).RoundToFrame();
-				FFrameNumber SnappedFrame = FFrameRate::TransformTime(PlayFrame, PlayRate, FrameResolution).RoundToFrame();
+				FFrameNumber PlayFrame    = FFrameRate::TransformTime(Time,      TickResolution, DisplayRate).RoundToFrame();
+				FFrameNumber SnappedFrame = FFrameRate::TransformTime(PlayFrame, DisplayRate, TickResolution).RoundToFrame();
 
 				Time = SnappedFrame;
 			}
@@ -1124,6 +1203,7 @@ void FSequencer::TransformSelectedKeysAndSections(FFrameTime InDeltaTime, float
 	bool bAnythingChanged = false;
 
 	TArray<FSequencerSelectedKey> SelectedKeysArray = Selection.GetSelectedKeys().Array();
+	TArray<TWeakObjectPtr<UMovieSceneSection>> SelectedSectionsArray = Selection.GetSelectedSections().Array();
 
 	const FFrameTime OriginTime = GetLocalTime().Time;
 
@@ -1146,34 +1226,90 @@ void FSequencer::TransformSelectedKeysAndSections(FFrameTime InDeltaTime, float
 				// Populate the key times scratch buffer with the times for these handles
 				ChannelInterface->GetKeyTimes_Raw(Pair.Key, Pair.Value.KeyHandles, KeyTimesScratch);
 
+				// We have to find the lowest key time and the highest key time. They're added based on selection order so we can't rely on their order in the array.
+				FFrameTime LowestFrameTime = KeyTimesScratch[0];
+				FFrameTime HighestFrameTime = KeyTimesScratch[0];
+
 				// Perform the transformation
 				for (FFrameNumber& Time : KeyTimesScratch)
 				{
 					FFrameTime KeyTime = Time;
 					Time = (OriginTime + InDeltaTime + (KeyTime - OriginTime) * InScale).FloorToFrame();
+
+					if (Time < LowestFrameTime)
+					{
+						LowestFrameTime = Time;
+					}
+
+					if (Time > HighestFrameTime)
+					{
+						HighestFrameTime = Time;
+					}
 				}
 
 				TRange<FFrameNumber>* NewSectionBounds = SectionToNewBounds.Find(Pair.Value.OwningSection);
 				if (!NewSectionBounds)
 				{
+					// Call Modify on the owning section before we call SetKeyTimes so that our section bounds/key times stay in sync.
 					Pair.Value.OwningSection->Modify();
-					NewSectionBounds = &SectionToNewBounds.Add(Pair.Value.OwningSection, TRange<FFrameNumber>(KeyTimesScratch[0], KeyTimesScratch.Last()+1));
+					NewSectionBounds = &SectionToNewBounds.Add(Pair.Value.OwningSection, Pair.Value.OwningSection->GetRange());
 				}
 
+
+				// Expand the range by ensuring the new range contains the range our keys are in. We add one because the upper time is exclusive
+				// for sections, but HighestFrameTime is measuring only the key's time.
+				*NewSectionBounds = TRange<FFrameNumber>::Hull(*NewSectionBounds, TRange<FFrameNumber>(LowestFrameTime.GetFrame(), HighestFrameTime.GetFrame() + 1));
+
 				// Apply the new, transformed key times
 				ChannelInterface->SetKeyTimes_Raw(Pair.Key, Pair.Value.KeyHandles, KeyTimesScratch);
 				bAnythingChanged = true;
 			}
 		}
+
+		// Dilate the sections
+		for (TWeakObjectPtr<UMovieSceneSection> Section : SelectedSectionsArray)
+		{
+			if (!Section.IsValid())
+			{
+				continue;
+			}
+
+			TRangeBound<FFrameNumber> LowerBound = Section->GetRange().GetLowerBound();
+			TRangeBound<FFrameNumber> UpperBound = Section->GetRange().GetUpperBound();
+
+			if (Section->HasStartFrame())
+			{
+				FFrameTime StartTime = Section->GetInclusiveStartFrame();
+				FFrameNumber StartFrame = (OriginTime + InDeltaTime + (StartTime - OriginTime) * InScale).FloorToFrame();
+				LowerBound = TRangeBound<FFrameNumber>::Inclusive(StartFrame);
+			}
+
+			if (Section->HasEndFrame())
+			{
+				FFrameTime EndTime = Section->GetExclusiveEndFrame();
+				FFrameNumber EndFrame = (OriginTime + InDeltaTime + (EndTime - OriginTime) * InScale).FloorToFrame();
+				UpperBound = TRangeBound<FFrameNumber>::Exclusive(EndFrame);
+			}
+
+			TRange<FFrameNumber>* NewSectionBounds = SectionToNewBounds.Find(Section.Get());
+			if (!NewSectionBounds)
+			{
+				NewSectionBounds = &SectionToNewBounds.Add( Section.Get(), TRange<FFrameNumber>(LowerBound, UpperBound) );
+			}
+
+			// If keys have already modified the section, we're applying the same modification to the section so we can
+			// overwrite the (possibly) existing bound without a TRange::Hull.
+			*NewSectionBounds = TRange<FFrameNumber>(LowerBound, UpperBound);
+			bAnythingChanged = true;
+		}
 	}
 	
+	// Remove any null sections so we don't need a null check inside the loop.
 	SectionToNewBounds.Remove(nullptr);
 	for (TTuple<UMovieSceneSection*, TRange<FFrameNumber>>& Pair : SectionToNewBounds)
 	{
-		if (!Pair.Key->GetRange().Contains(Pair.Value))
-		{
-			Pair.Key->SetRange(Pair.Value);
-		}
+		// Set the range of each section that has been modified to their new bounds.
+		Pair.Key->SetRange(Pair.Value);
 	}
 
 	if (bAnythingChanged)
@@ -1185,7 +1321,7 @@ void FSequencer::TransformSelectedKeysAndSections(FFrameTime InDeltaTime, float
 void FSequencer::TranslateSelectedKeysAndSections(bool bTranslateLeft)
 {
 	int32 Shift = bTranslateLeft ? -1 : 1;
-	FFrameTime Delta = FQualifiedFrameTime(Shift, GetFocusedPlayRate()).ConvertTo(GetFocusedFrameResolution());
+	FFrameTime Delta = FQualifiedFrameTime(Shift, GetFocusedDisplayRate()).ConvertTo(GetFocusedTickResolution());
 	TransformSelectedKeysAndSections(Delta, 1.f);
 }
 
@@ -1222,8 +1358,8 @@ void FSequencer::BakeTransform()
 			TArray<FRotator> Rotations;
 			TArray<FFrameNumber> KeyTimes;
 
-			FFrameRate   Resolution  = FocusedMovieScene->GetFrameResolution();
-			FFrameRate   SnapRate    = FocusedMovieScene->GetEvaluationType() == EMovieSceneEvaluationType::FrameLocked ? FocusedMovieScene->GetPlaybackFrameRate() : Resolution;
+			FFrameRate   Resolution  = FocusedMovieScene->GetTickResolution();
+			FFrameRate   SnapRate    = FocusedMovieScene->GetEvaluationType() == EMovieSceneEvaluationType::FrameLocked ? FocusedMovieScene->GetDisplayRate() : Resolution;
 
 			FFrameNumber InFrame     = MovieScene::DiscreteInclusiveLower(GetPlaybackRange());
 			FFrameNumber OutFrame    = MovieScene::DiscreteExclusiveUpper(GetPlaybackRange());
@@ -1378,6 +1514,18 @@ void FSequencer::NotifyMovieSceneDataChangedInternal()
 
 void FSequencer::NotifyMovieSceneDataChanged( EMovieSceneDataChangeType DataChangeType )
 {
+	if (!GetFocusedMovieSceneSequence()->GetMovieScene())
+	{
+		if (RootSequence.IsValid())
+		{
+			ResetToNewRootSequence(*RootSequence.Get());
+		}
+		else
+		{
+			UE_LOG(LogSequencer, Error, TEXT("Fatal error, focused movie scene no longer valid and there is no root sequence to default to."));
+		}
+	}
+
 	if ( DataChangeType == EMovieSceneDataChangeType::ActiveMovieSceneChanged ||
 		DataChangeType == EMovieSceneDataChangeType::Unknown )
 	{
@@ -1499,7 +1647,7 @@ void FSequencer::SetSelectionRange(TRange<FFrameNumber> Range)
 
 void FSequencer::SetSelectionRangeEnd()
 {
-	const FFrameNumber LocalTime = GetLocalTime().Time.FrameNumber + 1;
+	const FFrameNumber LocalTime = GetLocalTime().Time.FrameNumber;
 
 	if (GetSelectionRange().GetLowerBoundValue() >= LocalTime)
 	{
@@ -1765,7 +1913,7 @@ void FSequencer::TogglePlaybackRangeLocked()
 
 void FSequencer::ResetViewRange()
 {
-	TRange<double> PlayRangeSeconds = GetPlaybackRange() / GetFocusedFrameResolution();
+	TRange<double> PlayRangeSeconds = GetPlaybackRange() / GetFocusedTickResolution();
 	const double OutputViewSize = PlayRangeSeconds.Size<double>();
 	const double OutputChange = OutputViewSize * 0.1f;
 
@@ -1786,7 +1934,7 @@ void FSequencer::ZoomViewRange(float InZoomDelta)
 
 	const double CurrentTime = GetLocalTime().AsSeconds();
 	const double OutputViewSize = LocalViewRangeMax - LocalViewRangeMin;
-	const double OutputChange = OutputViewSize * 0.1f;
+	const double OutputChange = OutputViewSize * InZoomDelta;
 
 	float CurrentPositionFraction = (CurrentTime - LocalViewRangeMin) / OutputViewSize;
 
@@ -1861,15 +2009,15 @@ void FSequencer::SetAllowEditsMode(EAllowEditsMode AllowEditsMode)
 }
 
 
-bool FSequencer::GetKeyAllEnabled() const 
+EKeyGroupMode FSequencer::GetKeyGroupMode() const
 {
-	return Settings->GetKeyAllEnabled();
+	return Settings->GetKeyGroupMode();
 }
 
 
-void FSequencer::SetKeyAllEnabled(bool bKeyAllEnabled) 
+void FSequencer::SetKeyGroupMode(EKeyGroupMode Mode)
 {
-	Settings->SetKeyAllEnabled(bKeyAllEnabled);
+	Settings->SetKeyGroupMode(Mode);
 }
 
 
@@ -1917,7 +2065,7 @@ bool FSequencer::GetAutoSetTrackDefaults() const
 
 FQualifiedFrameTime FSequencer::GetLocalTime() const
 {
-	FFrameRate FocusedResolution = GetFocusedFrameResolution();
+	FFrameRate FocusedResolution = GetFocusedTickResolution();
 	FFrameTime CurrentPosition   = PlayPosition.GetCurrentPosition();
 
 	FFrameTime RootTime = ConvertFrameTime(CurrentPosition, PlayPosition.GetInputRate(), PlayPosition.GetOutputRate());
@@ -1933,7 +2081,7 @@ FQualifiedFrameTime FSequencer::GetGlobalTime() const
 
 void FSequencer::SetLocalTime( FFrameTime NewTime, ESnapTimeMode SnapTimeMode )
 {
-	FFrameRate LocalResolution = GetFocusedFrameResolution();
+	FFrameRate LocalResolution = GetFocusedTickResolution();
 
 	// Ensure the time is in the current view
 	ScrollIntoView(NewTime / LocalResolution);
@@ -1941,14 +2089,14 @@ void FSequencer::SetLocalTime( FFrameTime NewTime, ESnapTimeMode SnapTimeMode )
 	// Perform snapping
 	if ((SnapTimeMode & ESnapTimeMode::STM_Interval) && Settings->GetIsSnapEnabled())
 	{
-		FFrameRate LocalPlayRate = GetFocusedPlayRate();
+		FFrameRate LocalDisplayRate = GetFocusedDisplayRate();
 
-		NewTime = FFrameRate::TransformTime(FFrameRate::TransformTime(NewTime, LocalResolution, LocalPlayRate).RoundToFrame(), LocalPlayRate, LocalResolution);
+		NewTime = FFrameRate::TransformTime(FFrameRate::TransformTime(NewTime, LocalResolution, LocalDisplayRate).RoundToFrame(), LocalDisplayRate, LocalResolution);
 	}
 
 	if ((SnapTimeMode & ESnapTimeMode::STM_Keys) && (Settings->GetSnapPlayTimeToKeys() || FSlateApplication::Get().GetModifierKeys().IsShiftDown()))
 	{
-		NewTime = OnGetNearestKey(NewTime);
+		NewTime = OnGetNearestKey(NewTime, true);
 	}
 
 	SetLocalTimeDirectly(NewTime);
@@ -1964,7 +2112,7 @@ void FSequencer::SetLocalTimeDirectly(FFrameTime NewTime)
 
 void FSequencer::SetGlobalTime( FFrameTime NewTime )
 {
-	NewTime = ConvertFrameTime(NewTime, GetRootFrameResolution(), PlayPosition.GetInputRate());
+	NewTime = ConvertFrameTime(NewTime, GetRootTickResolution(), PlayPosition.GetInputRate());
 	if (PlayPosition.GetEvaluationType() == EMovieSceneEvaluationType::FrameLocked)
 	{
 		NewTime = NewTime.FloorToFrame();
@@ -2152,18 +2300,18 @@ void FSequencer::RenderMovieInternal(TRange<FFrameNumber> Range, bool bSetFrameO
 
 	MovieSceneCapture->SetLevelSequenceAsset(GetCurrentAsset()->GetPathName());
 
-	FFrameRate PlayRate = GetFocusedPlayRate();
-	FFrameRate FrameResolution = GetFocusedFrameResolution();
+	FFrameRate DisplayRate = GetFocusedDisplayRate();
+	FFrameRate TickResolution = GetFocusedTickResolution();
 
-	MovieSceneCapture->Settings.FrameRate = PlayRate;
+	MovieSceneCapture->Settings.FrameRate = DisplayRate;
 	MovieSceneCapture->Settings.ZeroPadFrameNumbers = Settings->GetZeroPadFrames();
 	MovieSceneCapture->Settings.bUseRelativeFrameNumbers = false;
 
 	FFrameNumber StartFrame = MovieScene::DiscreteInclusiveLower(Range);
 	FFrameNumber EndFrame = MovieScene::DiscreteExclusiveUpper(Range);
 
-	FFrameNumber RoundedStartFrame = FFrameRate::TransformTime(StartFrame, FrameResolution, PlayRate).FloorToFrame();
-	FFrameNumber RoundedEndFrame = FFrameRate::TransformTime(EndFrame, FrameResolution, PlayRate).CeilToFrame();
+	FFrameNumber RoundedStartFrame = FFrameRate::TransformTime(StartFrame, TickResolution, DisplayRate).CeilToFrame();
+	FFrameNumber RoundedEndFrame = FFrameRate::TransformTime(EndFrame, TickResolution, DisplayRate).CeilToFrame();
 
 	if (bSetFrameOverrides)
 	{
@@ -2203,7 +2351,7 @@ void FSequencer::RenderMovieInternal(TRange<FFrameNumber> Range, bool bSetFrameO
 		return 0;
 	});
 
-	// By using a FrameResolution/PlayRateResolution that match the numbers entered via the numeric interface don't change frames of reference.
+	// By using a TickResolution/DisplayRate that match the numbers entered via the numeric interface don't change frames of reference.
 	// This is used here because the movie scene capture works entirely on play rate resolution and has no knowledge of the internal resolution
 	// so we don't need to convert the user's input into internal resolution.
 	auto GetFrameRateDelegate = FOnGetFrameRate::CreateLambda([MovieSceneCapture]() -> FFrameRate
@@ -3002,9 +3150,9 @@ void FSequencer::HandleRecordingStarted(UMovieSceneSequence* Sequence)
 	UMovieSceneSubSection* Section = UMovieSceneSubSection::GetRecordingSection();
 	if(Section != nullptr)
 	{
-		FFrameRate   FrameResolution = GetFocusedFrameResolution();
-		FFrameNumber StartFrame      = GetLocalTime().ConvertTo(FrameResolution).FloorToFrame();
-		int32        Duration        = FFrameRate::TransformTime(1, GetFocusedPlayRate(), FrameResolution).CeilToFrame().Value;
+		FFrameRate   TickResolution  = GetFocusedTickResolution();
+		FFrameNumber StartFrame      = GetLocalTime().ConvertTo(TickResolution).CeilToFrame();
+		int32        Duration        = FFrameRate::TransformTime(1, GetFocusedDisplayRate(), TickResolution).CeilToFrame().Value;
 
 		Section->SetRange(TRange<FFrameNumber>(StartFrame, StartFrame + Duration));
 	}
@@ -3091,10 +3239,10 @@ FReply FSequencer::OnStepForward()
 {
 	SetPlaybackStatus(EMovieScenePlayerStatus::Stepping);
 
-	FFrameRate          PlayRate    = GetFocusedPlayRate();
+	FFrameRate          DisplayRate = GetFocusedDisplayRate();
 	FQualifiedFrameTime CurrentTime = GetLocalTime();
 
-	FFrameTime NewPosition = FFrameRate::TransformTime(CurrentTime.ConvertTo(PlayRate).FloorToFrame() + 1, PlayRate, CurrentTime.Rate);
+	FFrameTime NewPosition = FFrameRate::TransformTime(CurrentTime.ConvertTo(DisplayRate).FloorToFrame() + 1, DisplayRate, CurrentTime.Rate);
 	SetLocalTime(NewPosition, ESnapTimeMode::STM_Interval);
 	return FReply::Handled();
 }
@@ -3104,10 +3252,10 @@ FReply FSequencer::OnStepBackward()
 {
 	SetPlaybackStatus(EMovieScenePlayerStatus::Stepping);
 
-	FFrameRate          PlayRate    = GetFocusedPlayRate();
+	FFrameRate          DisplayRate = GetFocusedDisplayRate();
 	FQualifiedFrameTime CurrentTime = GetLocalTime();
 
-	FFrameTime NewPosition = FFrameRate::TransformTime(CurrentTime.ConvertTo(PlayRate).FloorToFrame() - 1, PlayRate, CurrentTime.Rate);
+	FFrameTime NewPosition = FFrameRate::TransformTime(CurrentTime.ConvertTo(DisplayRate).FloorToFrame() - 1, DisplayRate, CurrentTime.Rate);
 
 	SetLocalTime(NewPosition, ESnapTimeMode::STM_Interval);
 	return FReply::Handled();
@@ -3117,7 +3265,7 @@ FReply FSequencer::OnStepBackward()
 FReply FSequencer::OnJumpToStart()
 {
 	SetPlaybackStatus(EMovieScenePlayerStatus::Stepping);
-	SetLocalTime(GetPlaybackRange().GetLowerBoundValue());
+	SetLocalTime(MovieScene::DiscreteInclusiveLower(GetPlaybackRange()));
 	return FReply::Handled();
 }
 
@@ -3125,7 +3273,17 @@ FReply FSequencer::OnJumpToStart()
 FReply FSequencer::OnJumpToEnd()
 {
 	SetPlaybackStatus(EMovieScenePlayerStatus::Stepping);
-	SetLocalTime(GetPlaybackRange().GetUpperBoundValue());
+	const bool bInsetDisplayFrame = ScrubStyle == ESequencerScrubberStyle::FrameBlock && Settings->GetSnapPlayTimeToInterval() && Settings->GetIsSnapEnabled();
+
+	FFrameRate LocalResolution = GetFocusedTickResolution();
+	FFrameRate DisplayRate = GetFocusedDisplayRate();
+
+	// Calculate an offset from the end to go to. If they have snapping on (and the scrub style is a block) the last valid frame is represented as one
+	// whole display rate frame before the end, otherwise we just subtract a single frame which matches the behavior of hitting play and letting it run to the end.
+	FFrameTime OneFrame = bInsetDisplayFrame ? FFrameRate::TransformTime(FFrameTime(1), DisplayRate, LocalResolution) : FFrameTime(1);
+	FFrameTime NewTime = MovieScene::DiscreteExclusiveUpper(GetPlaybackRange()) - OneFrame;
+
+	SetLocalTime(NewTime);
 	return FReply::Handled();
 }
 
@@ -3245,8 +3403,8 @@ void FSequencer::SetLocalTimeLooped(FFrameTime NewLocalTime)
 	TRange<FFrameNumber> TimeBounds = GetTimeBounds();
 
 	bool         bResetPosition       = false;
-	FFrameRate   LocalFrameResolution = GetFocusedFrameResolution();
-	FFrameRate   RootFrameResolution  = GetRootFrameResolution();
+	FFrameRate   LocalTickResolution  = GetFocusedTickResolution();
+	FFrameRate   RootTickResolution   = GetRootTickResolution();
 	FFrameNumber MinInclusiveTime     = MovieScene::DiscreteInclusiveLower(TimeBounds);
 	FFrameNumber MaxInclusiveTime     = MovieScene::DiscreteExclusiveUpper(TimeBounds)-1;
 
@@ -3298,10 +3456,10 @@ void FSequencer::SetLocalTimeLooped(FFrameTime NewLocalTime)
 			}
 		}
 		// Ensure the time is within the working range
-		else if (!WorkingRange.Contains(NewLocalTime / LocalFrameResolution))
+		else if (!WorkingRange.Contains(NewLocalTime / LocalTickResolution))
 		{
-			FFrameTime WorkingMin = (WorkingRange.GetLowerBoundValue() * LocalFrameResolution).CeilToFrame();
-			FFrameTime WorkingMax = (WorkingRange.GetUpperBoundValue() * LocalFrameResolution).FloorToFrame();
+			FFrameTime WorkingMin = (WorkingRange.GetLowerBoundValue() * LocalTickResolution).CeilToFrame();
+			FFrameTime WorkingMax = (WorkingRange.GetUpperBoundValue() * LocalTickResolution).FloorToFrame();
 
 			NewGlobalTime = FMath::Clamp(NewLocalTime, WorkingMin, WorkingMax) * RootToLocalTransform.Inverse();
 
@@ -3312,15 +3470,15 @@ void FSequencer::SetLocalTimeLooped(FFrameTime NewLocalTime)
 
 	// Ensure the time is in the current view - must occur before the time cursor changes
 	UMovieScene* MovieScene = GetFocusedMovieSceneSequence()->GetMovieScene();
-	ScrollIntoView( (NewGlobalTime * RootToLocalTransform) / RootFrameResolution );
+	ScrollIntoView( (NewGlobalTime * RootToLocalTransform) / RootTickResolution );
 
-	FFrameTime NewPlayPosition = ConvertFrameTime(NewGlobalTime, RootFrameResolution, PlayPosition.GetInputRate());
+	FFrameTime NewPlayPosition = ConvertFrameTime(NewGlobalTime, RootTickResolution, PlayPosition.GetInputRate());
 
 	// Reset the play cursor if we're looping or have otherwise jumpted to a new position in the sequence
 	if (bResetPosition)
 	{
 		PlayPosition.Reset(NewPlayPosition);
-		TimeController->Reset(FQualifiedFrameTime(NewGlobalTime, RootFrameResolution));
+		TimeController->Reset(FQualifiedFrameTime(NewGlobalTime, RootTickResolution));
 	}
 
 	// Evaluate the sequence
@@ -3463,10 +3621,18 @@ void FSequencer::OnClampRangeChanged( TRange<double> NewClampRange )
 	}
 }
 
-FFrameNumber FSequencer::OnGetNearestKey(FFrameTime InTime)
+FFrameNumber FSequencer::OnGetNearestKey(FFrameTime InTime, bool bSearchAllTracks)
 {
 	FFrameNumber NearestKeyTime = InTime.FloorToFrame();
-	GetKeysFromSelection(SelectedKeyCollection, SMALL_NUMBER);
+
+	if (bSearchAllTracks)
+	{
+		GetAllKeys(SelectedKeyCollection, SMALL_NUMBER);
+	}
+	else
+	{
+		GetKeysFromSelection(SelectedKeyCollection, SMALL_NUMBER);
+	}
 
 	if (SelectedKeyCollection.IsValid())
 	{
@@ -3516,7 +3682,7 @@ void FSequencer::OnScrubPositionChanged( FFrameTime NewScrubPosition, bool bScru
 			// Clamp to the view range when not auto-scrolling
 			bClampToViewRange = false;
 	
-			UpdateAutoScroll(NewScrubPosition / GetFocusedFrameResolution());
+			UpdateAutoScroll(NewScrubPosition / GetFocusedTickResolution());
 			
 			// When scrubbing, we animate auto-scrolled scrub position in Tick()
 			if (AutoscrubOffset.IsSet())
@@ -3528,16 +3694,16 @@ void FSequencer::OnScrubPositionChanged( FFrameTime NewScrubPosition, bool bScru
 
 	if (bClampToViewRange)
 	{
-		FFrameRate PlayRate        = GetFocusedPlayRate();
-		FFrameRate FrameResolution = GetFocusedFrameResolution();
+		FFrameRate DisplayRate    = GetFocusedDisplayRate();
+		FFrameRate TickResolution = GetFocusedTickResolution();
 
-		FFrameTime LowerBound = (TargetViewRange.GetLowerBoundValue() * FrameResolution).CeilToFrame();
-		FFrameTime UpperBound = (TargetViewRange.GetUpperBoundValue() * FrameResolution).FloorToFrame();
+		FFrameTime LowerBound = (TargetViewRange.GetLowerBoundValue() * TickResolution).CeilToFrame();
+		FFrameTime UpperBound = (TargetViewRange.GetUpperBoundValue() * TickResolution).FloorToFrame();
 
 		if (Settings->GetIsSnapEnabled() && Settings->GetSnapPlayTimeToInterval())
 		{
-			LowerBound = FFrameRate::Snap(LowerBound, FrameResolution, PlayRate);
-			UpperBound = FFrameRate::Snap(UpperBound, FrameResolution, PlayRate);
+			LowerBound = FFrameRate::Snap(LowerBound, TickResolution, DisplayRate);
+			UpperBound = FFrameRate::Snap(UpperBound, TickResolution, DisplayRate);
 		}
 
 		NewScrubPosition = FMath::Clamp(NewScrubPosition, LowerBound, UpperBound);
@@ -3738,8 +3904,8 @@ void FSequencer::AddSubSequence(UMovieSceneSequence* Sequence)
 
 	FFrameNumber Duration = ConvertFrameTime(
 		Sequence->GetMovieScene()->GetPlaybackRange().Size<FFrameNumber>(),
-		Sequence->GetMovieScene()->GetFrameResolution(),
-		OwnerMovieScene->GetFrameResolution()).FloorToFrame();
+		Sequence->GetMovieScene()->GetTickResolution(),
+		OwnerMovieScene->GetTickResolution()).FloorToFrame();
 
 	SubTrack->AddSequence(Sequence, GetLocalTime().Time.FloorToFrame(), Duration.Value);
 }
@@ -3797,6 +3963,11 @@ bool FSequencer::OnRequestNodeDeleted( TSharedRef<const FSequencerDisplayNode> N
 	UMovieSceneSequence* Sequence = GetFocusedMovieSceneSequence();
 	UMovieScene* OwnerMovieScene = Sequence->GetMovieScene();
 
+	// Remove the selected object from our selection otherwise invisible objects are still selected and it causes confusion with
+	// things that are based on having a selection or not.
+	TSharedRef<FSequencerDisplayNode> SelectionNodeToRemove = ConstCastSharedRef<FSequencerDisplayNode>(NodeToBeDeleted);
+	Selection.RemoveFromSelection(SelectionNodeToRemove);
+
 	if ( NodeToBeDeleted->GetType() == ESequencerNode::Folder )
 	{
 		// Delete Children
@@ -4779,24 +4950,26 @@ void FSequencer::SyncCurveEditorToSelection(bool bOutlinerSelectionChanged)
 		{
 			continue;
 		}
-
-		if (ExistingCurveIDs.Contains(Channel))
-		{
-			ExistingCurveIDs.Remove(Channel);
-			continue;
-		}
-		else
+		if (KeyArea->GetOwningSection()->ShowCurveForChannel(Channel))
 		{
-			TUniquePtr<FCurveModel> NewCurve = KeyArea->CreateCurveEditorModel(AsShared());
-			if (NewCurve.IsValid())
+			if (ExistingCurveIDs.Contains(Channel))
 			{
-				bAnythingChanged = true;
-				CurveEditorModel->AddCurve(MoveTemp(NewCurve));
+				ExistingCurveIDs.Remove(Channel);
+				continue;
+			}
+			else
+			{
+				TUniquePtr<FCurveModel> NewCurve = KeyArea->CreateCurveEditorModel(AsShared());
+				if (NewCurve.IsValid())
+				{
+					bAnythingChanged = true;
+					CurveEditorModel->AddCurve(MoveTemp(NewCurve));
+				}
 			}
 		}
 	}
 
-	// Remove anything that's no longer selected
+	// Remove anything that's no longer selected or shown
 	for (const TTuple<const void*, FCurveModelID>& Pair : ExistingCurveIDs)
 	{
 		bAnythingChanged = true;
@@ -4847,7 +5020,7 @@ void FSequencer::SyncCurveEditorToSelection(bool bOutlinerSelectionChanged)
 
 void FSequencer::ZoomToSelectedSections()
 {
-	FFrameRate FrameResolution = GetFocusedFrameResolution();
+	FFrameRate TickResolution = GetFocusedTickResolution();
 
 	TRange<FFrameNumber> BoundsHull = TRange<FFrameNumber>::Empty();
 	for (TWeakObjectPtr<UMovieSceneSection> SelectedSection : Selection.GetSelectedSections())
@@ -4866,8 +5039,8 @@ void FSequencer::ZoomToSelectedSections()
 
 		// Zoom back to last view range if already expanded
 		if (!ViewRangeBeforeZoom.IsEmpty() &&
-			FMath::IsNearlyEqual(BoundsHull.GetLowerBoundValue() / FrameResolution, GetViewRange().GetLowerBoundValue(), Tolerance) &&
-			FMath::IsNearlyEqual(BoundsHull.GetUpperBoundValue() / FrameResolution, GetViewRange().GetUpperBoundValue(), Tolerance))
+			FMath::IsNearlyEqual(BoundsHull.GetLowerBoundValue() / TickResolution, GetViewRange().GetLowerBoundValue(), Tolerance) &&
+			FMath::IsNearlyEqual(BoundsHull.GetUpperBoundValue() / TickResolution, GetViewRange().GetUpperBoundValue(), Tolerance))
 		{
 			SetViewRange(ViewRangeBeforeZoom, EViewRangeInterpolation::Animated);
 		}
@@ -4875,7 +5048,7 @@ void FSequencer::ZoomToSelectedSections()
 		{
 			ViewRangeBeforeZoom = GetViewRange();
 
-			SetViewRange(BoundsHull / FrameResolution, EViewRangeInterpolation::Animated);
+			SetViewRange(BoundsHull / TickResolution, EViewRangeInterpolation::Animated);
 		}
 	}
 }
@@ -5163,6 +5336,7 @@ FGuid FSequencer::DoAssignActor(AActor*const* InActors, int32 NumActors, FGuid I
 
 		// Replace
 		OwnerMovieScene->ReplacePossessable( OldComponentGuid, NewPossessable );
+		OwnerSequence->UnbindPossessableObjects( OldComponentGuid );
 		State.Invalidate(OldComponentGuid, ActiveTemplateIDs.Top());
 
 		FMovieScenePossessable* ThisPossessable = OwnerMovieScene->FindPossessable( NewPossessable.GetGuid() );
@@ -5217,6 +5391,7 @@ FGuid FSequencer::DoAssignActor(AActor*const* InActors, int32 NumActors, FGuid I
 
 	// Replace the actor itself after components have been updated
 	OwnerMovieScene->ReplacePossessable(InObjectBinding, NewPossessableActor);
+	OwnerSequence->UnbindPossessableObjects(InObjectBinding);
 
 	State.Invalidate(InObjectBinding, ActiveTemplateIDs.Top());
 
@@ -5892,6 +6067,29 @@ void FSequencer::SaveSelectedNodesSpawnableState()
 	NotifyMovieSceneDataChanged(EMovieSceneDataChangeType::MovieSceneStructureItemsChanged);
 }
 
+void FSequencer::SetSelectedNodesSpawnableLevel(FName InLevelName)
+{							
+	const FScopedTransaction Transaction( LOCTEXT("SetSpawnableLevel", "Set Spawnable Level") );
+
+	UMovieScene* MovieScene = GetFocusedMovieSceneSequence()->GetMovieScene();
+
+	MovieScene->Modify();
+
+	TArray<FMovieSceneSpawnable*> Spawnables;
+
+	for (const TSharedRef<FSequencerDisplayNode>& Node : Selection.GetSelectedOutlinerNodes())
+	{
+		if (Node->GetType() == ESequencerNode::Object)
+		{
+			FMovieSceneSpawnable* Spawnable = MovieScene->FindSpawnable(StaticCastSharedRef<FSequencerObjectBindingNode>(Node)->GetObjectBinding());
+			if (Spawnable)
+			{
+				Spawnable->SetLevelName(InLevelName);
+			}
+		}
+	}
+}
+
 void FSequencer::ConvertToSpawnable(TSharedRef<FSequencerObjectBindingNode> NodeToBeConverted)
 {
 	const FScopedTransaction Transaction( LOCTEXT("ConvertSelectedNodeSpawnable", "Convert Node to Spawnables") );
@@ -5994,9 +6192,6 @@ FMovieSceneSpawnable* FSequencer::ConvertToSpawnableInternal(FGuid PossessableGu
 	UMovieSceneSequence* Sequence = GetFocusedMovieSceneSequence();
 	UMovieScene* MovieScene = Sequence->GetMovieScene();
 
-	// Find the object in the environment
-	FMovieScenePossessable* Possessable = MovieScene->FindPossessable(PossessableGuid);
-
 	//@todo: this code doesn't work where multiple objects are bound
 	TArrayView<TWeakObjectPtr<>> FoundObjects = FindBoundObjects(PossessableGuid, ActiveTemplateIDs.Top());
 	if (FoundObjects.Num() != 1)
@@ -6011,43 +6206,27 @@ FMovieSceneSpawnable* FSequencer::ConvertToSpawnableInternal(FGuid PossessableGu
 	}
 
 	Sequence->Modify();
+	MovieScene->Modify();
 
 	FMovieSceneSpawnable* Spawnable = MovieScene->FindSpawnable(AddSpawnable(*FoundObject));
-	if (!Spawnable)
-	{
-		return nullptr;
-	}
-
-	// Swap the guids, so the possessable's tracks now belong to the spawnable
+	if (Spawnable)
 	{
-		FGuid BenignSpawnableGuid = Spawnable->GetGuid();
-		FGuid PersistentGuid = Possessable->GetGuid();
+		FGuid SpawnableGuid = Spawnable->GetGuid();
 
-		Possessable->SetGuid(BenignSpawnableGuid);
-		Spawnable->SetGuid(PersistentGuid);
+		// Remap all the spawnable's tracks and child bindings onto the new possessable
+		MovieScene->MoveBindingContents(PossessableGuid, SpawnableGuid);
 
-		if (MovieScene->RemovePossessable(BenignSpawnableGuid))
+		if (MovieScene->RemovePossessable(PossessableGuid))
 		{
-			Sequence->UnbindPossessableObjects(PersistentGuid);
+			Sequence->UnbindPossessableObjects(PossessableGuid);
 		}
 
-		for (int32 Index = 0; Index < MovieScene->GetPossessableCount(); ++Index)
-		{
-			FMovieScenePossessable& MovieScenePossessable = MovieScene->GetPossessable(Index);
-			bool bBelongsToNewSpawnable = MovieScenePossessable.GetParent() == PersistentGuid;
-			if (bBelongsToNewSpawnable)
-			{
-				MovieScenePossessable.SetParent(PersistentGuid);
-				Spawnable->AddChildPossessable(MovieScenePossessable.GetGuid());
-			}
-		}
-	}
-
-	TOptional<FTransformData> TransformData;
-	SpawnRegister->HandleConvertPossessableToSpawnable(FoundObject, *this, TransformData);
-	SpawnRegister->SetupDefaultsForSpawnable(nullptr, Spawnable->GetGuid(), TransformData, AsShared(), Settings);
+		TOptional<FTransformData> TransformData;
+		SpawnRegister->HandleConvertPossessableToSpawnable(FoundObject, *this, TransformData);
+		SpawnRegister->SetupDefaultsForSpawnable(nullptr, Spawnable->GetGuid(), TransformData, AsShared(), Settings);
 
-	ForceEvaluate();
+		ForceEvaluate();
+	}
 
 	return Spawnable;
 }
@@ -6163,17 +6342,20 @@ FMovieScenePossessable* FSequencer::ConvertToPossessableInternal(FGuid Spawnable
 		return nullptr;
 	}
 
+	Sequence->Modify();
+	MovieScene->Modify();
+
 	// Delete the spawn track
 	UMovieSceneSpawnTrack* SpawnTrack = Cast<UMovieSceneSpawnTrack>(MovieScene->FindTrack(UMovieSceneSpawnTrack::StaticClass(), SpawnableGuid, NAME_None));
 	if (SpawnTrack)
 	{
-		MovieScene->Modify();
 		MovieScene->RemoveTrack(*SpawnTrack);
 	}
 
 	FTransform SpawnTransform = SpawnableActorTemplate->GetActorTransform();
 	FActorSpawnParameters SpawnInfo;
 	SpawnInfo.bDeferConstruction = true;
+	SpawnInfo.Template = SpawnableActorTemplate;
 
 	UWorld* PlaybackContext = Cast<UWorld>(GetPlaybackContext());
 	AActor* PossessedActor = PlaybackContext->SpawnActor(Spawnable->GetObjectTemplate()->GetClass(), &SpawnTransform, SpawnInfo);
@@ -6183,41 +6365,31 @@ FMovieScenePossessable* FSequencer::ConvertToPossessableInternal(FGuid Spawnable
 		return nullptr;
 	}
 
-	UEngine::FCopyPropertiesForUnrelatedObjectsParams CopyParams;
-	PossessedActor->UnregisterAllComponents();
-	UEngine::CopyPropertiesForUnrelatedObjects(SpawnableActorTemplate, PossessedActor, CopyParams);
-	PossessedActor->RegisterAllComponents();
-
 	const bool bIsDefaultTransform = true;
 	PossessedActor->FinishSpawning(SpawnTransform, bIsDefaultTransform);
 
-	const FGuid PossessableGuid = CreateBinding(*PossessedActor, PossessedActor->GetActorLabel());
+	const FGuid NewPossessableGuid = CreateBinding(*PossessedActor, PossessedActor->GetActorLabel());
+	const FGuid OldSpawnableGuid = Spawnable->GetGuid();
 
-	FMovieScenePossessable* Possessable = MovieScene->FindPossessable(PossessableGuid);
-	if (!Possessable)
-	{
-		return nullptr;
-	}
-
-	// Swap the guids, so the spawnable's tracks now belong to the possessable
+	FMovieScenePossessable* Possessable = MovieScene->FindPossessable(NewPossessableGuid);
+	if (Possessable)
 	{
-		FGuid BenignSpawnableGuid = Spawnable->GetGuid();
-		FGuid PersistentGuid = Possessable->GetGuid();
-
-		Spawnable->SetGuid(PersistentGuid);
-		Possessable->SetGuid(BenignSpawnableGuid);
+		// Remap all the spawnable's tracks and child bindings onto the new possessable
+		MovieScene->MoveBindingContents(OldSpawnableGuid, NewPossessableGuid);
 
-		if (MovieScene->RemoveSpawnable(PersistentGuid))
+		// Remove the spawnable and all it's sub tracks
+		if (MovieScene->RemoveSpawnable(OldSpawnableGuid))
 		{
-			SpawnRegister->DestroySpawnedObject(BenignSpawnableGuid, ActiveTemplateIDs.Top(), *this);
+			SpawnRegister->DestroySpawnedObject(OldSpawnableGuid, ActiveTemplateIDs.Top(), *this);
 		}
-
-		Sequence->BindPossessableObject(BenignSpawnableGuid, *PossessedActor, PlaybackContext);
-	}
 	
-	GEditor->SelectActor(PossessedActor, false, true);
+		static const FName SequencerActorTag(TEXT("SequencerActor"));
+		PossessedActor->Tags.Remove(SequencerActorTag);
 
-	ForceEvaluate();
+		GEditor->SelectActor(PossessedActor, false, true);
+
+		ForceEvaluate();
+	}
 
 	return Possessable;
 }
@@ -6393,13 +6565,13 @@ void FSequencer::Pause()
 	// When stopping a sequence, we always evaluate a non-empty range if possible. This ensures accurate paused motion blur effects.
 	if (Settings->GetIsSnapEnabled())
 	{
-		FQualifiedFrameTime LocalTime       = GetLocalTime();
-		FFrameRate          FocusedPlayRate = GetFocusedPlayRate();
+		FQualifiedFrameTime LocalTime          = GetLocalTime();
+		FFrameRate          FocusedDisplayRate = GetFocusedDisplayRate();
 
 		// Snap to the focused play rate
-		FFrameTime RootPosition  = FFrameRate::Snap(LocalTime.Time, LocalTime.Rate, FocusedPlayRate) * RootToLocalTransform.Inverse();
+		FFrameTime RootPosition  = FFrameRate::Snap(LocalTime.Time, LocalTime.Rate, FocusedDisplayRate) * RootToLocalTransform.Inverse();
 
-		// Convert the root position from frame resolution time base (the output rate), to the play position input rate
+		// Convert the root position from tick resolution time base (the output rate), to the play position input rate
 		FFrameTime InputPosition = ConvertFrameTime(RootPosition, PlayPosition.GetOutputRate(), PlayPosition.GetInputRate());
 		EvaluateInternal(PlayPosition.PlayTo(InputPosition));
 	}
@@ -6670,6 +6842,41 @@ void FSequencer::OnSetKeyTimeTextCommitted(const FText& InText, ETextCommit::Typ
 	}
 }
 
+bool FSequencer::CanRekey() const
+{
+	return Selection.GetSelectedKeys().Num() > 0;
+}
+
+
+void FSequencer::Rekey()
+{
+	bool bAnythingChanged = false;
+
+	FQualifiedFrameTime CurrentTime = GetLocalTime();
+
+	FScopedTransaction RekeyTransaction(NSLOCTEXT("Sequencer", "Rekey_Transaction", "Rekey"));
+	TArray<FSequencerSelectedKey> SelectedKeysArray = Selection.GetSelectedKeys().Array();
+	
+	for ( const FSequencerSelectedKey& Key : SelectedKeysArray )
+	{
+		if (Key.IsValid())
+		{
+	 		if (Key.Section->TryModify())
+	 		{
+	 			Key.KeyArea->SetKeyTime(Key.KeyHandle.GetValue(), CurrentTime.Time.FrameNumber);
+	 			bAnythingChanged = true;
+
+				Key.Section->ExpandToFrame(CurrentTime.Time.FrameNumber);
+	 		}
+		}
+	}
+
+	if (bAnythingChanged)
+	{
+		NotifyMovieSceneDataChanged(EMovieSceneDataChangeType::TrackValueChanged);
+	}
+}
+
 TArray<TSharedPtr<FMovieSceneClipboard>> GClipboardStack;
 
 void FSequencer::CopySelection()
@@ -6775,7 +6982,7 @@ void FSequencer::CopySelectedKeys()
 
 	TSharedRef<FMovieSceneClipboard> Clipboard = MakeShareable( new FMovieSceneClipboard(Builder.Commit(CopyRelativeTo)) );
 	
-	Clipboard->GetEnvironment().FrameResolution = GetFocusedFrameResolution();
+	Clipboard->GetEnvironment().TickResolution = GetFocusedTickResolution();
 
 	if (Clipboard->GetKeyTrackGroups().Num())
 	{
@@ -7212,10 +7419,7 @@ void FSequencer::ImportFBX()
 		}
 	}
 	
-	if (MovieSceneToolHelpers::ImportFBX(MovieScene, *this, ObjectBindingNameMap))
-	{
-		NotifyMovieSceneDataChanged( EMovieSceneDataChangeType::MovieSceneStructureItemAdded );
-	}
+	MovieSceneToolHelpers::ImportFBX(MovieScene, *this, ObjectBindingNameMap);
 }
 
 
@@ -7352,7 +7556,7 @@ void FSequencer::CloseEntryPopupMenu()
 void FSequencer::TrimSection(bool bTrimLeft)
 {
 	FScopedTransaction TrimSectionTransaction( NSLOCTEXT("Sequencer", "TrimSection_Transaction", "Trim Section") );
-	MovieSceneToolHelpers::TrimSection(Selection.GetSelectedSections(), GetLocalTime().Time.FloorToFrame(), bTrimLeft);
+	MovieSceneToolHelpers::TrimSection(Selection.GetSelectedSections(), GetLocalTime(), bTrimLeft);
 	NotifyMovieSceneDataChanged( EMovieSceneDataChangeType::TrackValueChanged );
 }
 
@@ -7360,7 +7564,7 @@ void FSequencer::TrimSection(bool bTrimLeft)
 void FSequencer::SplitSection()
 {
 	FScopedTransaction SplitSectionTransaction( NSLOCTEXT("Sequencer", "SplitSection_Transaction", "Split Section") );
-	MovieSceneToolHelpers::SplitSection(Selection.GetSelectedSections(), GetLocalTime().Time.FloorToFrame());
+	MovieSceneToolHelpers::SplitSection(Selection.GetSelectedSections(), GetLocalTime());
 	NotifyMovieSceneDataChanged( EMovieSceneDataChangeType::MovieSceneStructureItemAdded );
 }
 
@@ -7569,10 +7773,22 @@ void FSequencer::BindCommands()
 		FIsActionChecked::CreateLambda( [this]{ return Settings->GetAutoChangeMode() == EAutoChangeMode::AutoKey; } ) );
 
 	SequencerCommandBindings->MapAction(
-		Commands.ToggleKeyAllEnabled,
-		FExecuteAction::CreateLambda( [this]{ Settings->SetKeyAllEnabled( !Settings->GetKeyAllEnabled() ); } ),
-		FCanExecuteAction::CreateLambda( [this]{ return Settings->GetAllowEditsMode() != EAllowEditsMode::AllowLevelEditsOnly; } ),
-		FIsActionChecked::CreateLambda( [this]{ return Settings->GetKeyAllEnabled(); } ) );
+		Commands.SetKeyChanged,
+		FExecuteAction::CreateLambda([this] { Settings->SetKeyGroupMode(EKeyGroupMode::KeyChanged); }),
+		FCanExecuteAction::CreateLambda([] { return true; }),
+		FIsActionChecked::CreateLambda([this] { return Settings->GetKeyGroupMode() == EKeyGroupMode::KeyChanged; }));
+
+	SequencerCommandBindings->MapAction(
+		Commands.SetKeyGroup,
+		FExecuteAction::CreateLambda([this] { Settings->SetKeyGroupMode(EKeyGroupMode::KeyGroup); }),
+		FCanExecuteAction::CreateLambda([] { return true; }),
+		FIsActionChecked::CreateLambda([this] { return Settings->GetKeyGroupMode() == EKeyGroupMode::KeyGroup; }));
+
+	SequencerCommandBindings->MapAction(
+		Commands.SetKeyAll,
+		FExecuteAction::CreateLambda([this] { Settings->SetKeyGroupMode(EKeyGroupMode::KeyAll); }),
+		FCanExecuteAction::CreateLambda([] { return true; }),
+		FIsActionChecked::CreateLambda([this] { return Settings->GetKeyGroupMode() == EKeyGroupMode::KeyAll; }));
 
 	SequencerCommandBindings->MapAction(
 		Commands.ToggleAutoScroll,
@@ -7627,7 +7843,7 @@ void FSequencer::BindCommands()
 
 			// If the next framerate in the list is drop format timecode and we're not in a play rate that supports drop format timecode,
 			// then we will skip over it.
-			bool bCanShowDropFrameTimecode = FTimecode::IsDropFormatTimecodeSupported(GetFocusedPlayRate());
+			bool bCanShowDropFrameTimecode = FTimecode::IsDropFormatTimecodeSupported(GetFocusedDisplayRate());
 			if (!bCanShowDropFrameTimecode && NextFormat == EFrameNumberDisplayFormats::DropFrameTimecode)
 			{
 				NextFormat = EFrameNumberDisplayFormats::Seconds;
@@ -7793,6 +8009,18 @@ void FSequencer::BindCommands()
 		return !EditedSequence->GetMovieScene()->GetSelectionRange().IsEmpty();
 	};
 
+	SequencerCommandBindings->MapAction(
+		FGenericCommands::Get().Rename,
+		FExecuteAction::CreateLambda([this]
+		{
+			Selection.GetSelectedOutlinerNodes().Array()[0]->OnRenameRequested().Broadcast();
+		}),
+		FCanExecuteAction::CreateLambda([this]
+		{
+			return (Selection.GetSelectedOutlinerNodes().Num() > 0) && (Selection.GetSelectedOutlinerNodes().Array()[0]->CanRenameNode());
+		})
+	);
+
 	SequencerCommandBindings->MapAction(
 		FGenericCommands::Get().Cut,
 		FExecuteAction::CreateSP(this, &FSequencer::CutSelection),
@@ -7938,6 +8166,10 @@ void FSequencer::BindCommands()
 		Commands.SetInterpolationCubicBreak,
 		FExecuteAction::CreateSP(this, &FSequencer::SetInterpTangentMode, ERichCurveInterpMode::RCIM_Cubic, ERichCurveTangentMode::RCTM_Break));
 
+	SequencerCommandBindings->MapAction(
+		Commands.ToggleWeightedTangents,
+		FExecuteAction::CreateSP(this, &FSequencer::ToggleInterpTangentWeightMode));
+
 	SequencerCommandBindings->MapAction(
 		Commands.SetInterpolationLinear,
 		FExecuteAction::CreateSP(this, &FSequencer::SetInterpTangentMode, ERichCurveInterpMode::RCIM_Linear, ERichCurveTangentMode::RCTM_Auto));
@@ -8134,18 +8366,18 @@ void FSequencer::UpdateTimeBases()
 
 	if (RootMovieScene)
 	{
-		EMovieSceneEvaluationType EvaluationType    = RootMovieScene->GetEvaluationType();
-		FFrameRate                FrameResolution   = RootMovieScene->GetFrameResolution();
-		FFrameRate                PlaybackFrameRate = EvaluationType == EMovieSceneEvaluationType::FrameLocked ? RootMovieScene->GetPlaybackFrameRate() : FrameResolution;
+		EMovieSceneEvaluationType EvaluationType  = RootMovieScene->GetEvaluationType();
+		FFrameRate                TickResolution  = RootMovieScene->GetTickResolution();
+		FFrameRate                DisplayRate     = EvaluationType == EMovieSceneEvaluationType::FrameLocked ? RootMovieScene->GetDisplayRate() : TickResolution;
 
-		if (PlaybackFrameRate != PlayPosition.GetInputRate())
+		if (DisplayRate != PlayPosition.GetInputRate())
 		{
 			bNeedsEvaluate = true;
 		}
 
-		// We set the play position in terms of the playback frame rate,
-		// but want evaluation ranges in the moviescene's frame resolution
-		PlayPosition.SetTimeBase(PlaybackFrameRate, FrameResolution, EvaluationType);
+		// We set the play position in terms of the display rate,
+		// but want evaluation ranges in the moviescene's tick resolution
+		PlayPosition.SetTimeBase(DisplayRate, TickResolution, EvaluationType);
 	}
 }
 
@@ -8185,7 +8417,7 @@ bool FSequencer::GetGridMetrics(float PhysicalWidth, double& OutMajorInterval, i
 	TSharedRef<FSlateFontMeasure> FontMeasureService = FSlateApplication::Get().GetRenderer()->GetFontMeasureService();
 
 	double BiggestTime = GetViewRange().GetUpperBoundValue();
-	FString TickString = GetNumericTypeInterface()->ToString((BiggestTime * GetFocusedPlayRate()).FrameNumber.Value);
+	FString TickString = GetNumericTypeInterface()->ToString((BiggestTime * GetFocusedDisplayRate()).FrameNumber.Value);
 	FVector2D MaxTextSize = FontMeasureService->Measure(TickString, SmallLayoutFont);
 
 	static float MajorTickMultiplier = 2.f;
@@ -8195,7 +8427,7 @@ bool FSequencer::GetGridMetrics(float PhysicalWidth, double& OutMajorInterval, i
 
 	if (PhysicalWidth > 0)
 	{
-		return GetFocusedPlayRate().ComputeGridSpacing(
+		return GetFocusedDisplayRate().ComputeGridSpacing(
 			PhysicalWidth / GetViewRange().Size<double>(),
 			OutMajorInterval,
 			OutMinorDivisions,
@@ -8206,9 +8438,9 @@ bool FSequencer::GetGridMetrics(float PhysicalWidth, double& OutMajorInterval, i
 	return false;
 }
 
-double FSequencer::GetPlayRateDeltaFrameCount() const
+double FSequencer::GetDisplayRateDeltaFrameCount() const
 {
-	return GetFocusedFrameResolution().AsDecimal() * GetFocusedPlayRate().AsInterval();
+	return GetFocusedTickResolution().AsDecimal() * GetFocusedDisplayRate().AsInterval();
 }
 
 
diff --git a/Engine/Source/Editor/Sequencer/Private/Sequencer.h b/Engine/Source/Editor/Sequencer/Private/Sequencer.h
--- a/Engine/Source/Editor/Sequencer/Private/Sequencer.h
+++ b/Engine/Source/Editor/Sequencer/Private/Sequencer.h
@@ -248,6 +248,9 @@ public:
 	/** Set interpolation modes. */
 	void SetInterpTangentMode(ERichCurveInterpMode InterpMode, ERichCurveTangentMode TangentMode);
 
+	/** Toggle tangent weight mode. */
+	void ToggleInterpTangentWeightMode();
+
 	/** Snap the currently selected keys to frame. */
 	void SnapToFrame();
 
@@ -297,6 +300,11 @@ public:
 	 */
 	void ConvertSelectedNodesToPossessables();
 
+	/*
+	 * Set the spawnable level for the currently selected objects
+	 */
+	void SetSelectedNodesSpawnableLevel(FName InLevelName);
+
 protected:
 
 	/**
@@ -469,6 +477,10 @@ public:
 	void SetKeyTime();
 	void OnSetKeyTimeTextCommitted(const FText& InText, ETextCommit::Type CommitInfo);
 
+	/** Called when a user executes the rekey for selected keys */
+	bool CanRekey() const;
+	void Rekey();
+
 	void SelectKey(UMovieSceneSection* InSection, TSharedPtr<IKeyArea> InKeyArea, FKeyHandle KeyHandle, bool bToggle);
 
 	FSequencerLabelManager& GetLabelManager()
@@ -611,8 +623,8 @@ public:
 	virtual void SetAutoChangeMode(EAutoChangeMode AutoChangeMode) override;
 	virtual EAllowEditsMode GetAllowEditsMode() const override;
 	virtual void SetAllowEditsMode(EAllowEditsMode AllowEditsMode) override;
-	virtual bool GetKeyAllEnabled() const override;
-	virtual void SetKeyAllEnabled(bool bKeyAllEnabled) override;
+	virtual EKeyGroupMode GetKeyGroupMode() const override;
+	virtual void SetKeyGroupMode(EKeyGroupMode) override;
 	virtual bool GetKeyInterpPropertiesOnly() const override;
 	virtual void SetKeyInterpPropertiesOnly(bool bKeyInterpPropertiesOnly) override;
 	virtual EMovieSceneKeyInterpolation GetKeyInterpolation() const override;
@@ -737,9 +749,11 @@ protected:
 	 * Called to get the nearest key
 	 *
 	 * @param InTime The time to get the nearest key to
+	 * @param bSearchAllTracks If true this will search all tracks for a potential nearest.
+	 *						 False will return keys only from the currently selected track.
 	 * @return NearestKey
 	 */
-	FFrameNumber OnGetNearestKey(FFrameTime InTime);
+	FFrameNumber OnGetNearestKey(FFrameTime InTime, bool bSearchAllTracks);
 
 	/**
 	 * Called when the scrub position is changed by the user
@@ -929,8 +943,15 @@ public:
 	/** Reset the timing manager to the clock source specified by the root movie scene */
 	void ResetTimeController();
 
-	/** Helper function which returns how many frames (in frame resolution) one display rate frame represents. */
-	double GetPlayRateDeltaFrameCount() const;
+	/** Helper function which returns how many frames (in tick resolution) one display rate frame represents. */
+	double GetDisplayRateDeltaFrameCount() const;
+
+	/** Retrieve the desired scrubber style for this instance. */
+	ESequencerScrubberStyle GetScrubStyle() const
+	{
+		return ScrubStyle;
+	}
+
 private:
 
 	/** Update the time bases for the current movie scene */
@@ -1023,6 +1044,9 @@ private:
 	/** Whether the sequence should be editable or read only */
 	bool bReadOnly;
 
+	/** Scrub style provided on construction */
+	ESequencerScrubberStyle ScrubStyle;
+
 	/** Generic Popup Entry */
 	TWeakPtr<IMenu> EntryPopupMenu;
 
diff --git a/Engine/Source/Editor/Sequencer/Private/SequencerCommands.cpp b/Engine/Source/Editor/Sequencer/Private/SequencerCommands.cpp
--- a/Engine/Source/Editor/Sequencer/Private/SequencerCommands.cpp
+++ b/Engine/Source/Editor/Sequencer/Private/SequencerCommands.cpp
@@ -55,9 +55,12 @@ void FSequencerCommands::RegisterCommands()
 	UI_COMMAND( SetInterpolationCubicAuto, "Set Key Auto", "Cubic interpolation - Automatic tangents", EUserInterfaceActionType::Button, FInputChord(EKeys::One));
 	UI_COMMAND( SetInterpolationCubicUser, "Set Key User", "Cubic interpolation - User flat tangents", EUserInterfaceActionType::Button, FInputChord(EKeys::Two));
 	UI_COMMAND( SetInterpolationCubicBreak, "Set Key Break", "Cubic interpolation - User broken tangents", EUserInterfaceActionType::Button, FInputChord(EKeys::Three));
+
 	UI_COMMAND( SetInterpolationLinear, "Set Key Linear", "Linear interpolation", EUserInterfaceActionType::Button, FInputChord(EKeys::Four));
 	UI_COMMAND( SetInterpolationConstant, "Set Key Constant", "Constant interpolation", EUserInterfaceActionType::Button, FInputChord(EKeys::Five));
 
+	UI_COMMAND( ToggleWeightedTangents, "Toggle Weighted Tangents", "Toggles cubic tangents to be weighted/non-weighted", EUserInterfaceActionType::ToggleButton, FInputChord(EModifierKey::Control, EKeys::W));
+
 	UI_COMMAND( TrimSectionLeft, "Trim Section Left", "Trim section at current time to the left (keeps the right)", EUserInterfaceActionType::Button, FInputChord(EModifierKey::Control, EKeys::Comma) );
 	UI_COMMAND( TrimSectionRight, "Trim Section Right", "Trim section at current time to the right (keeps the left)", EUserInterfaceActionType::Button, FInputChord(EModifierKey::Control, EKeys::Period) );
 	UI_COMMAND( SplitSection, "Split Section", "Split section at current time", EUserInterfaceActionType::Button, FInputChord(EModifierKey::Control, EKeys::Slash) );
@@ -74,8 +77,11 @@ void FSequencerCommands::RegisterCommands()
 	UI_COMMAND( AllowSequencerEditsOnly, "Allow Sequencer Edits Only", "All edits will produce either a track or a key.", EUserInterfaceActionType::ToggleButton, FInputChord());
 	UI_COMMAND( AllowLevelEditsOnly, "Allow Level Edits Only", "Properties in the details panel will be disabled if they have a track.", EUserInterfaceActionType::ToggleButton, FInputChord());
 
-	UI_COMMAND( ToggleAutoKeyEnabled, "Auto-key", "Create a key when channels/properties change. Only automatically adds a key when there's already a track and at least one key.", EUserInterfaceActionType::ToggleButton, FInputChord());
-	UI_COMMAND( ToggleKeyAllEnabled, "Key All", "Key all channels/properties when only one of them changes. ie. Keys all translation, rotation, scale channels when only translation Y changes", EUserInterfaceActionType::ToggleButton, FInputChord());
+	UI_COMMAND(ToggleAutoKeyEnabled, "Auto-key", "Create a key when channels/properties change. Only automatically adds a key when there's already a track and at least one key.", EUserInterfaceActionType::ToggleButton, FInputChord());
+
+	UI_COMMAND( SetKeyChanged, "Key Changed", "Key just the changed channel when it changes.", EUserInterfaceActionType::ToggleButton, FInputChord());
+	UI_COMMAND( SetKeyGroup, "Key Group", "Key the groups channels/properties when only one of them changes. ie. Keys all three translation channels when only translation Y changes", EUserInterfaceActionType::ToggleButton, FInputChord());
+	UI_COMMAND( SetKeyAll, "Key All", "Key all channels/properties when only one of them changes. ie. Keys all translation, rotation, scale channels when only translation Y changes", EUserInterfaceActionType::ToggleButton, FInputChord());
 
 	UI_COMMAND( ToggleAutoScroll, "Auto Scroll", "Toggle auto-scroll: When enabled, automatically scrolls the sequencer view to keep the current time visible", EUserInterfaceActionType::ToggleButton, FInputChord(EModifierKey::Shift, EKeys::S) );
 
@@ -87,10 +93,10 @@ void FSequencerCommands::RegisterCommands()
 	UI_COMMAND( ToggleIsSnapEnabled, "Enable Snapping", "Enables and disables snapping", EUserInterfaceActionType::ToggleButton, FInputChord() );
 
 	UI_COMMAND( ToggleSnapKeyTimesToInterval, "Snap to the Interval", "Snap keys to the time snapping interval", EUserInterfaceActionType::ToggleButton, FInputChord() );
-	UI_COMMAND( ToggleSnapKeyTimesToKeys, "Snap to Other Keys", "Snap keys to other keys in the section", EUserInterfaceActionType::ToggleButton, FInputChord() );
+	UI_COMMAND( ToggleSnapKeyTimesToKeys, "Snap to Keys and Sections", "Snap keys to other keys and sections in this sequence", EUserInterfaceActionType::ToggleButton, FInputChord() );
 
 	UI_COMMAND( ToggleSnapSectionTimesToInterval, "Snap to the Interval", "Snap sections to the time snapping interval", EUserInterfaceActionType::ToggleButton, FInputChord() );
-	UI_COMMAND( ToggleSnapSectionTimesToSections, "Snap to Other Sections", "Snap sections to other sections", EUserInterfaceActionType::ToggleButton, FInputChord() );
+	UI_COMMAND( ToggleSnapSectionTimesToSections, "Snap to Keys and Sections", "Snap sections to other keys and sections in this sequence", EUserInterfaceActionType::ToggleButton, FInputChord() );
 
 	UI_COMMAND( ToggleSnapPlayTimeToKeys, "Snap to Keys While Scrubbing", "Snap the current time to keys of the selected track while scrubbing", EUserInterfaceActionType::ToggleButton, FInputChord() );
 	UI_COMMAND( ToggleSnapPlayTimeToInterval, "Snap to the Interval While Scrubbing", "Snap the current time to the time snapping interval while scrubbing", EUserInterfaceActionType::ToggleButton, FInputChord() );
diff --git a/Engine/Source/Editor/Sequencer/Private/SequencerCommonHelpers.cpp b/Engine/Source/Editor/Sequencer/Private/SequencerCommonHelpers.cpp
--- a/Engine/Source/Editor/Sequencer/Private/SequencerCommonHelpers.cpp
+++ b/Engine/Source/Editor/Sequencer/Private/SequencerCommonHelpers.cpp
@@ -75,7 +75,7 @@ int32 SequencerHelpers::GetSectionFromTime(TArrayView<UMovieSceneSection* const>
 				}
 			}
 			// Check for nearby sections if there is nothing overlapping
-			else if (!MaxOverlapPriority.IsSet())
+			else if (!MaxOverlapPriority.IsSet() && SectionRange.HasLowerBound())
 			{
 				const FFrameNumber LowerBoundValue = SectionRange.GetLowerBoundValue();
 				// If this section exists beyond the current time, we can choose it if its closest to the time
diff --git a/Engine/Source/Editor/Sequencer/Private/SequencerContextMenus.cpp b/Engine/Source/Editor/Sequencer/Private/SequencerContextMenus.cpp
--- a/Engine/Source/Editor/Sequencer/Private/SequencerContextMenus.cpp
+++ b/Engine/Source/Editor/Sequencer/Private/SequencerContextMenus.cpp
@@ -167,6 +167,13 @@ void FKeyContextMenu::PopulateMenu(FMenuBuilder& MenuBuilder)
 				FCanExecuteAction::CreateSP(SequencerPtr, &FSequencer::CanSetKeyTime))
 		);
 
+		MenuBuilder.AddMenuEntry(LOCTEXT("Rekey", "Rekey"), LOCTEXT("RekeyTooltip", "Set the selected key's time to the current time"),
+			FSlateIcon(),
+			FUIAction(
+				FExecuteAction::CreateSP(SequencerPtr, &FSequencer::Rekey),
+				FCanExecuteAction::CreateSP(SequencerPtr, &FSequencer::CanRekey))
+		);
+
 		MenuBuilder.AddMenuEntry(
 			LOCTEXT("SnapToFrame", "Snap to Frame"),
 			LOCTEXT("SnapToFrameToolTip", "Snap selected keys to frame"),
@@ -416,6 +423,15 @@ void FSectionContextMenu::AddEditMenu(FMenuBuilder& MenuBuilder)
 			FCanExecuteAction::CreateLambda([=]{ return Shared->IsTrimmable(); }))
 	);
 		
+	MenuBuilder.AddMenuEntry(
+		LOCTEXT("AutoSizeSection", "Auto Size"),
+		LOCTEXT("AutoSizeSectionTooltip", "Auto size the section length to the duration of the source of this section (ie. audio, animation or shot length)"),
+		FSlateIcon(),
+		FUIAction(
+			FExecuteAction::CreateLambda([=]{ Shared->AutoSizeSection(); }),
+			FCanExecuteAction::CreateLambda([=]{ return Shared->CanAutoSize(); }))
+	);
+		
 	MenuBuilder.AddMenuEntry(
 		LOCTEXT("ReduceKeysSection", "Reduce Keys"),
 		LOCTEXT("ReduceKeysTooltip", "Reduce keys in this section"),
@@ -490,6 +506,7 @@ void FSectionContextMenu::AddPropertiesMenu(FMenuBuilder& MenuBuilder)
 		DetailsViewArgs.bShowOptions = false;
 		DetailsViewArgs.bShowModifiedPropertiesOption = false;
 		DetailsViewArgs.NotifyHook = &DetailsNotifyWrapper.Get();
+		DetailsViewArgs.ColumnWidth = 0.45f;
 	}
 
 	TArray<TWeakObjectPtr<UObject>> Sections;
@@ -669,8 +686,7 @@ void FSectionContextMenu::TrimSection(bool bTrimLeft)
 {
 	FScopedTransaction TrimSectionTransaction(LOCTEXT("TrimSection_Transaction", "Trim Section"));
 
-	FFrameNumber CurrentFrame = Sequencer->GetLocalTime().Time.FrameNumber;
-	MovieSceneToolHelpers::TrimSection(Sequencer->GetSelection().GetSelectedSections(), CurrentFrame, bTrimLeft);
+	MovieSceneToolHelpers::TrimSection(Sequencer->GetSelection().GetSelectedSections(), Sequencer->GetLocalTime(), bTrimLeft);
 	Sequencer->NotifyMovieSceneDataChanged( EMovieSceneDataChangeType::TrackValueChanged );
 }
 
@@ -680,7 +696,30 @@ void FSectionContextMenu::SplitSection()
 	FScopedTransaction SplitSectionTransaction(LOCTEXT("SplitSection_Transaction", "Split Section"));
 
 	FFrameNumber CurrentFrame = Sequencer->GetLocalTime().Time.FrameNumber;
-	MovieSceneToolHelpers::SplitSection(Sequencer->GetSelection().GetSelectedSections(), CurrentFrame);
+	FQualifiedFrameTime SplitFrame = FQualifiedFrameTime(CurrentFrame, Sequencer->GetFocusedTickResolution());
+
+	MovieSceneToolHelpers::SplitSection(Sequencer->GetSelection().GetSelectedSections(), SplitFrame);
+	Sequencer->NotifyMovieSceneDataChanged( EMovieSceneDataChangeType::MovieSceneStructureItemAdded );
+}
+
+
+void FSectionContextMenu::AutoSizeSection()
+{
+	FScopedTransaction AutoSizeSectionTransaction(LOCTEXT("AutoSizeSection_Transaction", "Auto Size Section"));
+
+	for (auto Section : Sequencer->GetSelection().GetSelectedSections())
+	{
+		if (Section.IsValid() && Section->GetAutoSizeRange().IsSet())
+		{
+			TOptional<TRange<FFrameNumber> > DefaultSectionLength = Section->GetAutoSizeRange();
+
+			if (DefaultSectionLength.IsSet())
+			{
+				Section->SetRange(DefaultSectionLength.GetValue());
+			}
+		}
+	}
+
 	Sequencer->NotifyMovieSceneDataChanged( EMovieSceneDataChangeType::MovieSceneStructureItemAdded );
 }
 
@@ -740,6 +779,18 @@ bool FSectionContextMenu::IsTrimmable() const
 	return false;
 }
 
+bool FSectionContextMenu::CanAutoSize() const
+{
+	for (auto Section : Sequencer->GetSelection().GetSelectedSections())
+	{
+		if (Section.IsValid() && Section->GetAutoSizeRange().IsSet())
+		{
+			return true;
+		}
+	}
+	return false;
+}
+
 bool FSectionContextMenu::CanReduceKeys() const
 {
 	TSet<TSharedPtr<IKeyArea> > KeyAreas;
@@ -1278,7 +1329,15 @@ void FPasteContextMenu::Setup()
 				continue;
 			}
 
-			const TArray<UMovieSceneSection*>& Sections = TrackNode->GetTrack()->GetAllSections();
+			TArray<UMovieSceneSection*> Sections;
+			for (auto Section : TrackNode->GetSections())
+			{
+				if (Section.Get().GetSectionObject())
+				{
+					Sections.Add(Section.Get().GetSectionObject());
+				}
+			}
+
 			UMovieSceneSection* Section = MovieSceneHelpers::FindNearestSectionAtTime(Sections, Args.PasteAtTime);
 			int32 SectionIndex = INDEX_NONE;
 			if (Section)
@@ -1419,7 +1478,7 @@ void FPasteContextMenu::PasteInto(int32 DestinationIndex, FName KeyAreaName)
 	TSet<FSequencerSelectedKey> NewSelection;
 
 	FSequencerPasteEnvironment PasteEnvironment;
-	PasteEnvironment.FrameResolution = Sequencer->GetFocusedFrameResolution();
+	PasteEnvironment.TickResolution = Sequencer->GetFocusedTickResolution();
 	PasteEnvironment.CardinalTime = Args.PasteAtTime;
 	PasteEnvironment.OnKeyPasted = [&](FKeyHandle Handle, IKeyArea& KeyArea){
 		NewSelection.Add(FSequencerSelectedKey(*KeyArea.GetOwningSection(), KeyArea.AsShared(), Handle));
@@ -1561,7 +1620,7 @@ void FEasingContextMenu::PopulateMenu(FMenuBuilder& MenuBuilder)
 					.MaxValue(TOptional<double>())
 					.MaxSliderValue(TOptional<double>())
 					.MinSliderValue(0.f)
-					.Delta_Lambda([=]() -> double { return Sequencer->GetPlayRateDeltaFrameCount(); })
+					.Delta_Lambda([=]() -> double { return Sequencer->GetDisplayRateDeltaFrameCount(); })
 					.Value_Lambda([=] {
 						TOptional<int32> Current = Shared->GetCurrentLength();
 						if (Current.IsSet())
diff --git a/Engine/Source/Editor/Sequencer/Private/SequencerContextMenus.h b/Engine/Source/Editor/Sequencer/Private/SequencerContextMenus.h
--- a/Engine/Source/Editor/Sequencer/Private/SequencerContextMenus.h
+++ b/Engine/Source/Editor/Sequencer/Private/SequencerContextMenus.h
@@ -59,10 +59,14 @@ private:
 
 	void SplitSection();
 
+	void AutoSizeSection();
+
 	void ReduceKeys();
 
 	bool IsTrimmable() const;
 
+	bool CanAutoSize() const;
+
 	bool CanReduceKeys() const;
 
 	void ToggleSectionActive();
diff --git a/Engine/Source/Editor/Sequencer/Private/SequencerEdMode.cpp b/Engine/Source/Editor/Sequencer/Private/SequencerEdMode.cpp
--- a/Engine/Source/Editor/Sequencer/Private/SequencerEdMode.cpp
+++ b/Engine/Source/Editor/Sequencer/Private/SequencerEdMode.cpp
@@ -229,7 +229,7 @@ void FSequencerEdMode::CleanUpMeshTrails()
 namespace SequencerEdMode_Draw3D
 {
 static const FColor	KeySelectedColor(255,128,0);
-static const int32	DrawTrackTimeRes = 10;
+static const float	DrawTrackTimeRes = 0.1f;
 static const float	CurveHandleScale = 0.5f;
 }
 
@@ -347,7 +347,7 @@ void FSequencerEdMode::GetLocationAtTime(FMovieSceneEvaluationTrack* Track, UObj
 	FMovieSceneInterrogationData InterrogationData;
 	Sequencer->GetEvaluationTemplate().CopyActuators(InterrogationData.GetAccumulator());
 
-	FMovieSceneContext Context(FMovieSceneEvaluationRange(KeyTime, Sequencer->GetFocusedFrameResolution()));
+	FMovieSceneContext Context(FMovieSceneEvaluationRange(KeyTime, Sequencer->GetFocusedTickResolution()));
 	Track->Interrogate(Context, InterrogationData, Object);
 
 	for (const FTransform& Transform : InterrogationData.Iterate<FTransform>(UMovieScene3DTransformTrack::GetInterrogationKey()))
@@ -398,7 +398,7 @@ void FSequencerEdMode::DrawTransformTrack(const TSharedPtr<FSequencer>& Sequence
 		}
 	);
 	
-	FFrameRate FrameResolution = Sequencer->GetFocusedFrameResolution();
+	FFrameRate TickResolution = Sequencer->GetFocusedTickResolution();
 	FMovieSceneEvaluationTemplate* Template = Sequencer->GetEvaluationTemplate().FindTemplate(Sequencer->GetFocusedTemplateID());
 	if (!bShowTrajectory || !Template || !TransformTrack->GetAllSections().ContainsByPredicate([](UMovieSceneSection* In){ return In->IsActive(); }))
 	{
@@ -435,10 +435,9 @@ void FSequencerEdMode::DrawTransformTrack(const TSharedPtr<FSequencer>& Sequence
 				// If not the first keypoint, draw a line to the last keypoint.
 				if(KeyTimeIndex > 0)
 				{
-					int32 NumSteps = ((NewKeyTime - OldKeyTime)*SequencerEdMode_Draw3D::DrawTrackTimeRes).CeilToFrame().Value;
+					int32 NumSteps = FMath::CeilToInt((TickResolution.AsSeconds(NewKeyTime) - TickResolution.AsSeconds(OldKeyTime)) / SequencerEdMode_Draw3D::DrawTrackTimeRes);
 					// Limit the number of steps to prevent a rendering performance hit
 					NumSteps = FMath::Min( 100, NumSteps );
-
 					FFrameTime DrawSubstep = NumSteps == 0 ? 0 : (NewKeyTime - OldKeyTime)*(1.f/NumSteps);
 
 					// Find position on first keyframe.
@@ -491,7 +490,7 @@ void FSequencerEdMode::DrawTransformTrack(const TSharedPtr<FSequencer>& Sequence
 								}
 								else if (TrailActor != nullptr)
 								{
-									TrailActor->AddFrameMeshComponent(NewTime / FrameResolution, FTransform(NewRot, NewPos, FVector(3.0f)));
+									TrailActor->AddFrameMeshComponent(NewTime / TickResolution, FTransform(NewRot, NewPos, FVector(3.0f)));
 								}
 							}
 							OldTime = NewTime;
@@ -543,7 +542,7 @@ void FSequencerEdMode::DrawTransformTrack(const TSharedPtr<FSequencer>& Sequence
 
 					for (UMovieScene3DTransformSection* Section : AllSections)
 					{
-						TrailActor->AddKeyMeshActor(NewKeyTime / FrameResolution, FTransform(NewKeyRot, NewKeyPos, FVector(3.0f)), Section);
+						TrailActor->AddKeyMeshActor(NewKeyTime / TickResolution, FTransform(NewKeyRot, NewKeyPos, FVector(3.0f)), Section);
 					}
 				}
 
diff --git a/Engine/Source/Editor/Sequencer/Private/SequencerKeyActor.cpp b/Engine/Source/Editor/Sequencer/Private/SequencerKeyActor.cpp
--- a/Engine/Source/Editor/Sequencer/Private/SequencerKeyActor.cpp
+++ b/Engine/Source/Editor/Sequencer/Private/SequencerKeyActor.cpp
@@ -77,8 +77,8 @@ void ASequencerKeyActor::PropagateKeyChange()
 		// Mark the track section as dirty
 		TrackSection->Modify();
 
-		FFrameRate   FrameResolution = TrackSection->GetTypedOuter<UMovieScene>()->GetFrameResolution();
-		FFrameNumber FrameNumber     = (KeyTime * FrameResolution).RoundToFrame();
+		FFrameRate   TickResolution  = TrackSection->GetTypedOuter<UMovieScene>()->GetTickResolution();
+		FFrameNumber FrameNumber     = (KeyTime * TickResolution).RoundToFrame();
 
 		TArrayView<FMovieSceneFloatChannel*> FloatChannels = TrackSection->GetChannelProxy().GetChannels<FMovieSceneFloatChannel>();
 
diff --git a/Engine/Source/Editor/Sequencer/Private/SequencerNodeSortingMethods.h b/Engine/Source/Editor/Sequencer/Private/SequencerNodeSortingMethods.h
--- a/Engine/Source/Editor/Sequencer/Private/SequencerNodeSortingMethods.h
+++ b/Engine/Source/Editor/Sequencer/Private/SequencerNodeSortingMethods.h
@@ -17,28 +17,6 @@ struct FDisplayNodeTreePositionSorter
 	}
 };
 
-/**
-* This sorting delegate sorts based only on the sorting order specified by the node. This respects
-* a user-defined sorting order and ignores the node's type.
-*/
-struct FDisplayNodeSortingOrderSorter
-{
-	bool operator()(const TSharedRef<FSequencerDisplayNode>& A, const TSharedRef<FSequencerDisplayNode>& B) const
-	{
-		// If both nodes have been sorted before we just compare their sorting orders.
-		if (A->GetSortingOrder() >= 0 && B->GetSortingOrder() >= 0)
-		{
-			return A->GetSortingOrder() < B->GetSortingOrder();
-		}
-
-		// Otherwise if one of them has not placed before, we want the lower number higher so that
-		// the unsorted node shows up at the end of the list.
-		return A->GetSortingOrder() > B->GetSortingOrder();
-	}
-};
-
-
-
 /**
 * This sorting delegate sorts based on category and then alphabetically after that. This replicates
 * existing behavior where all folders come first, then all tracks, and finally all object bindings.
@@ -113,6 +91,35 @@ struct FDisplayNodeCategoricalSorter
 };
 
 
+/**
+* This sorting delegate sorts based on the sorting order specified by the node. This sorter falls back to the 
+* FDisplayNodeCategoricalSorter for Object nodes because Object nodes do not support sorting order for child nodes.
+*/
+struct FDisplayNodeSortingOrderSorter
+{
+	bool operator()(const TSharedRef<FSequencerDisplayNode>& A, const TSharedRef<FSequencerDisplayNode>& B) const
+	{
+		// If the parent of the object is a Object Binding then we use legacy sorting as we don't allow the user
+		// to reorder tracks within object binding nodes. 
+		TSharedPtr<FSequencerDisplayNode> ParentNode = A->GetParent();
+		if (ParentNode.IsValid() && ParentNode->GetType() == ESequencerNode::Object)
+		{
+			FDisplayNodeCategoricalSorter CategoricalSorter;
+			return CategoricalSorter(A, B);
+		}
+		
+		// If both nodes have been sorted before we just compare their sorting orders.
+		if (A->GetSortingOrder() >= 0 && B->GetSortingOrder() >= 0)
+		{
+			return A->GetSortingOrder() < B->GetSortingOrder();
+		}
+
+		// Otherwise if one of them has not placed before, we want the lower number higher so that
+		// the unsorted node shows up at the end of the list.
+		return A->GetSortingOrder() > B->GetSortingOrder();
+	}
+};
+
 /**
 * Sorts the supplied unsorted nodes and inserts them into the existing sorted nodes before assigning
 * a sorting order to the newly combined array. Inserted nodes can optionally be relative to an existing
diff --git a/Engine/Source/Editor/Sequencer/Private/SequencerSettings.cpp b/Engine/Source/Editor/Sequencer/Private/SequencerSettings.cpp
--- a/Engine/Source/Editor/Sequencer/Private/SequencerSettings.cpp
+++ b/Engine/Source/Editor/Sequencer/Private/SequencerSettings.cpp
@@ -9,7 +9,7 @@ USequencerSettings::USequencerSettings( const FObjectInitializer& ObjectInitiali
 {
 	AutoChangeMode = EAutoChangeMode::None;
 	AllowEditsMode = EAllowEditsMode::AllEdits;
-	bKeyAllEnabled = false;
+	KeyGroupMode = EKeyGroupMode::KeyChanged;
 	bKeyInterpPropertiesOnly = false;
 	KeyInterpolation = EMovieSceneKeyInterpolation::Auto;
 	bAutoSetTrackDefaults = false;
@@ -87,16 +87,16 @@ void USequencerSettings::SetAllowEditsMode(EAllowEditsMode InAllowEditsMode)
 	}
 }
 
-bool USequencerSettings::GetKeyAllEnabled() const
+EKeyGroupMode USequencerSettings::GetKeyGroupMode() const
 {
-	return bKeyAllEnabled;
+	return KeyGroupMode;
 }
 
-void USequencerSettings::SetKeyAllEnabled(bool InbKeyAllEnabled)
+void USequencerSettings::SetKeyGroupMode(EKeyGroupMode InKeyGroupMode)
 {
-	if ( bKeyAllEnabled != InbKeyAllEnabled )
+	if (KeyGroupMode != InKeyGroupMode)
 	{
-		bKeyAllEnabled = InbKeyAllEnabled;
+		KeyGroupMode = InKeyGroupMode;
 		SaveConfig();
 	}
 }
