Copying //UE4/Dev-Sequencer to //UE4/Dev-Main (Source: //UE4/Dev-Sequencer @ 4044339)

diff --git a/Engine/Source/Editor/Sequencer/Private/SequencerTimeSliderController.cpp b/Engine/Source/Editor/Sequencer/Private/SequencerTimeSliderController.cpp
--- a/Engine/Source/Editor/Sequencer/Private/SequencerTimeSliderController.cpp
+++ b/Engine/Source/Editor/Sequencer/Private/SequencerTimeSliderController.cpp
@@ -19,6 +19,7 @@
 #include "MovieSceneTimeHelpers.h"
 #include "CommonFrameRates.h"
 #include "Sequencer.h"
+#include "SequencerDisplayNode.h"
 
 #define LOCTEXT_NAMESPACE "TimeSlider"
 
@@ -34,53 +35,6 @@ namespace ScrubConstants
 	const float ScrollPanFraction = 0.1f;
 }
 
-/** Utility struct for converting between scrub range space and local/absolute screen space */
-struct FSequencerTimeSliderController::FScrubRangeToScreen
-{
-	double ViewStart;
-	float  PixelsPerInput;
-
-	FScrubRangeToScreen(const TRange<double>& InViewInput, const FVector2D& InWidgetSize )
-	{
-		float ViewInputRange = InViewInput.Size<double>();
-
-		ViewStart      = InViewInput.GetLowerBoundValue();
-		PixelsPerInput = ViewInputRange > 0 ? ( InWidgetSize.X / ViewInputRange ) : 0;
-	}
-
-	/** Local Widget Space -> Curve Input domain. */
-	double LocalXToInput(float ScreenX) const
-	{
-		return PixelsPerInput > 0 ? (ScreenX/PixelsPerInput) + ViewStart : ViewStart;
-	}
-
-	/** Curve Input domain -> local Widget Space */
-	float InputToLocalX(double Input) const
-	{
-		return (Input - ViewStart) * PixelsPerInput;
-	}
-};
-
-
-/**
- * Gets the the next spacing value in the series 
- * to determine a good spacing value
- * E.g, .001,.005,.010,.050,.100,.500,1.000,etc
- */
-static float GetNextSpacing( uint32 CurrentStep )
-{
-	if(CurrentStep & 0x01) 
-	{
-		// Odd numbers
-		return FMath::Pow( 10.f, 0.5f*((float)(CurrentStep-1)) + 1.f );
-	}
-	else 
-	{
-		// Even numbers
-		return 0.5f * FMath::Pow( 10.f, 0.5f*((float)(CurrentStep)) + 1.f );
-	}
-}
-
 FSequencerTimeSliderController::FSequencerTimeSliderController( const FTimeSliderArgs& InArgs, TWeakPtr<FSequencer> InWeakSequencer )
 	: WeakSequencer(InWeakSequencer)
 	, TimeSliderArgs( InArgs )
@@ -88,11 +42,11 @@ FSequencerTimeSliderController::FSequencerTimeSliderController( const FTimeSlide
 	, MouseDragType( DRAG_NONE )
 	, bPanning( false )
 {
-	ScrubFillBrush              = FEditorStyle::GetBrush( TEXT( "Sequencer.Timeline.ScrubFill" ) );
-	ScrubHandleUpBrush          = FEditorStyle::GetBrush( TEXT( "Sequencer.Timeline.ScrubHandleUp" ) ); 
-	ScrubHandleDownBrush        = FEditorStyle::GetBrush( TEXT( "Sequencer.Timeline.ScrubHandleDown" ) );
-	ClampedScrubHandleUpBrush   = FEditorStyle::GetBrush( TEXT( "Sequencer.Timeline.ClampedScrubHandleUp" ) ); 
-	ClampedScrubHandleDownBrush = FEditorStyle::GetBrush( TEXT( "Sequencer.Timeline.ClampedScrubHandleDown" ) );
+	ScrubFillBrush                 = FEditorStyle::GetBrush( TEXT( "Sequencer.Timeline.ScrubFill" ) );
+	FrameBlockScrubHandleUpBrush   = FEditorStyle::GetBrush( TEXT( "Sequencer.Timeline.FrameBlockScrubHandleUp" ) ); 
+	FrameBlockScrubHandleDownBrush = FEditorStyle::GetBrush( TEXT( "Sequencer.Timeline.FrameBlockScrubHandleDown" ) );
+	VanillaScrubHandleUpBrush      = FEditorStyle::GetBrush( TEXT( "Sequencer.Timeline.VanillaScrubHandleUp" ) ); 
+	VanillaScrubHandleDownBrush    = FEditorStyle::GetBrush( TEXT( "Sequencer.Timeline.VanillaScrubHandleDown" ) );
 	ContextMenuSuppression = 0;
 }
 
@@ -100,17 +54,30 @@ FFrameTime FSequencerTimeSliderController::ComputeScrubTimeFromMouse(const FGeom
 {
 	FVector2D           CursorPos     = Geometry.AbsoluteToLocal( ScreenSpacePosition );
 	double              MouseSeconds  = RangeToScreen.LocalXToInput( CursorPos.X );
-	FFrameTime          ScrubTime     = MouseSeconds * GetFrameResolution();
+	FFrameTime          ScrubTime     = MouseSeconds * GetTickResolution();
 
 	if ( TimeSliderArgs.Settings->GetIsSnapEnabled() )
 	{
-		if (TimeSliderArgs.Settings->GetSnapPlayTimeToKeys())
+		if (TimeSliderArgs.Settings->GetSnapPlayTimeToInterval())
 		{
-			ScrubTime = SnapTimeToNearestKey(RangeToScreen, CursorPos.X, ScrubTime);
+			// Set the style of the scrub handle
+			TSharedPtr<FSequencer> Sequencer = WeakSequencer.Pin();
+			if (Sequencer.IsValid() && Sequencer->GetScrubStyle() == ESequencerScrubberStyle::FrameBlock)
+			{
+				// Floor to the display frame
+				ScrubTime = ConvertFrameTime(ConvertFrameTime(ScrubTime, GetTickResolution(), GetDisplayRate()).FloorToFrame(), GetDisplayRate(), GetTickResolution());
+			}
+			else
+			{
+				// Snap (round) to display rate
+				ScrubTime = FFrameRate::Snap(ScrubTime, GetTickResolution(), GetDisplayRate());
+			}
 		}
-		else if (TimeSliderArgs.Settings->GetSnapPlayTimeToInterval())
+
+		if (TimeSliderArgs.Settings->GetSnapPlayTimeToKeys())
 		{
-			ScrubTime = FFrameRate::Snap(ScrubTime, GetFrameResolution(), GetPlayRate());
+			// SnapTimeToNearestKey will return ScrubTime unmodified if there is no key within range.
+			ScrubTime = SnapTimeToNearestKey(RangeToScreen, CursorPos.X, ScrubTime);
 		}
 	}
 
@@ -129,59 +96,67 @@ FFrameTime FSequencerTimeSliderController::ComputeFrameTimeFromMouse(const FGeom
 
 	if (CheckSnapping && TimeSliderArgs.Settings->GetIsSnapEnabled())
 	{
-		FFrameNumber        SnappedFrameNumber = (MouseValue * GetPlayRate()).FloorToFrame();
-		FQualifiedFrameTime RoundedPlayFrame   = FQualifiedFrameTime(SnappedFrameNumber, GetPlayRate());
-		return RoundedPlayFrame.ConvertTo(GetFrameResolution());
+		FFrameNumber        SnappedFrameNumber = (MouseValue * GetDisplayRate()).FloorToFrame();
+		FQualifiedFrameTime RoundedPlayFrame   = FQualifiedFrameTime(SnappedFrameNumber, GetDisplayRate());
+		return RoundedPlayFrame.ConvertTo(GetTickResolution());
 	}
 	else
 	{
-		return MouseValue * GetFrameResolution();
+		return MouseValue * GetTickResolution();
 	}
 }
 
-FSequencerTimeSliderController::FScrubPixelRange FSequencerTimeSliderController::GetHitTestScrubberPixelRange(FFrameTime ScrubTime, const FScrubRangeToScreen& RangeToScreen) const
+FSequencerTimeSliderController::FScrubberMetrics FSequencerTimeSliderController::GetHitTestScrubPixelMetrics(const FScrubRangeToScreen& RangeToScreen) const
 {
 	static const float DragToleranceSlateUnits = 2.f, MouseTolerance = 2.f;
-	return GetScrubberPixelRange(ScrubTime, GetFrameResolution(), GetPlayRate(), RangeToScreen, DragToleranceSlateUnits + MouseTolerance);
+	return GetScrubPixelMetrics(FQualifiedFrameTime(TimeSliderArgs.ScrubPosition.Get(), GetTickResolution()), RangeToScreen, DragToleranceSlateUnits + MouseTolerance);
 }
 
-FSequencerTimeSliderController::FScrubPixelRange FSequencerTimeSliderController::GetScrubberPixelRange(FFrameTime ScrubTime, const FScrubRangeToScreen& RangeToScreen) const
+FSequencerTimeSliderController::FScrubberMetrics FSequencerTimeSliderController::GetScrubPixelMetrics(const FQualifiedFrameTime& ScrubTime, const FScrubRangeToScreen& RangeToScreen, float DilationPixels) const
 {
-	return GetScrubberPixelRange(ScrubTime, GetFrameResolution(), GetPlayRate(), RangeToScreen);
-}
+	FFrameRate DisplayRate = GetDisplayRate();
 
-FSequencerTimeSliderController::FScrubPixelRange FSequencerTimeSliderController::GetScrubberPixelRange(FFrameTime ScrubTime, FFrameRate Resolution, FFrameRate PlayRate, const FScrubRangeToScreen& RangeToScreen, float DilationPixels) const
-{
-	const FFrameNumber Frame = ScrubTime.FloorToFrame();// ConvertFrameTime(ScrubTime, Resolution, PlayRate).FloorToFrame();
+	FScrubberMetrics Metrics;
 
-	float StartPixel = RangeToScreen.InputToLocalX(  Frame    / Resolution);
-	float EndPixel   = RangeToScreen.InputToLocalX( (Frame+1) / Resolution);
+	static float MinScrubSize = 14.f;
+
+	const FFrameNumber Frame = ScrubTime.ConvertTo(DisplayRate).FloorToFrame();
+
+	float FrameStartPixel = RangeToScreen.InputToLocalX(  Frame    / DisplayRate );
+	float FrameEndPixel   = RangeToScreen.InputToLocalX( (Frame+1) / DisplayRate ) - 1;
 
 	{
-		float RoundedStartPixel = FMath::RoundToInt(StartPixel);
-		EndPixel -= (StartPixel - RoundedStartPixel);
+		float RoundedStartPixel = FMath::RoundToInt(FrameStartPixel);
+		FrameEndPixel -= (FrameStartPixel - RoundedStartPixel);
 
-		StartPixel = RoundedStartPixel;
-		EndPixel = FMath::Max(EndPixel, StartPixel + 1);
+		FrameStartPixel = RoundedStartPixel;
+		FrameEndPixel   = FMath::Max(FrameEndPixel, FrameStartPixel + 1);
 	}
 
-	FScrubPixelRange Range;
+	// Store off the pixel width of the frame
+	Metrics.FrameExtentsPx = TRange<float>(FrameStartPixel, FrameEndPixel);
+
+	// Set the style of the scrub handle
+	TSharedPtr<FSequencer> Sequencer = WeakSequencer.Pin();
+	Metrics.Style = Sequencer.IsValid() ? Sequencer->GetScrubStyle() : ESequencerScrubberStyle::Vanilla;
 
-	float MinScrubSize = 14.f;
-	Range.bClamped    = EndPixel - StartPixel < MinScrubSize;
-	Range.Range       = TRange<float>(StartPixel, EndPixel);
-	if (Range.bClamped)
+	// Always draw the extents on the section area for frame block styles
+	Metrics.bDrawExtents = Metrics.Style == ESequencerScrubberStyle::FrameBlock;
+
+	// If it's vanilla style or too small to show the frame width, set that up
+	if (Metrics.Style == ESequencerScrubberStyle::Vanilla || FrameEndPixel - FrameStartPixel < MinScrubSize)
 	{
-		Range.HandleRange = TRange<float>(
-			(StartPixel + EndPixel - MinScrubSize) * .5f,
-			(StartPixel + EndPixel + MinScrubSize) * .5f);
+		Metrics.Style = ESequencerScrubberStyle::Vanilla;
+
+		float ScrubPixel = RangeToScreen.InputToLocalX(ScrubTime.AsSeconds());
+		Metrics.HandleRangePx = TRange<float>(ScrubPixel - MinScrubSize*.5f, ScrubPixel + MinScrubSize*.5f);
 	}
 	else
 	{
-		Range.HandleRange = Range.Range;
+		Metrics.HandleRangePx = Metrics.FrameExtentsPx;
 	}
 
-	return Range;
+	return Metrics;
 }
 
 struct FSequencerTimeSliderController::FDrawTickArgs
@@ -215,7 +190,8 @@ void FSequencerTimeSliderController::DrawTicks( FSlateWindowElementList& OutDraw
 		return;
 	}
 
-	FFrameRate     FrameResolution = GetFrameResolution();
+	FFrameRate     TickResolution  = GetTickResolution();
+	FFrameRate     DisplayRate     = GetDisplayRate();
 	FPaintGeometry PaintGeometry   = InArgs.AllottedGeometry.ToPaintGeometry();
 	FSlateFontInfo SmallLayoutFont = FCoreStyle::GetDefaultFontStyle("Regular", 8);
 
@@ -235,9 +211,11 @@ void FSequencerTimeSliderController::DrawTicks( FSlateWindowElementList& OutDraw
 	LinePoints.SetNumUninitialized(2);
 
 	const bool bAntiAliasLines = false;
+	
+	const double FirstMajorLine  = FMath::FloorToDouble(ViewRange.GetLowerBoundValue() / MajorGridStep) * MajorGridStep;
+	const double LastMajorLine   = FMath::CeilToDouble(ViewRange.GetUpperBoundValue() / MajorGridStep) * MajorGridStep;
 
-	const double FirstMajorLine = FMath::FloorToDouble(ViewRange.GetLowerBoundValue() / MajorGridStep) * MajorGridStep;
-	const double LastMajorLine  = FMath::CeilToDouble(ViewRange.GetUpperBoundValue() / MajorGridStep) * MajorGridStep;
+	const float  FlooredScrubPx  = RangeToScreen.InputToLocalX(ConvertFrameTime(TimeSliderArgs.ScrubPosition.Get(), TickResolution, GetDisplayRate()).FloorToFrame() / DisplayRate);
 
 	for (double CurrentMajorLine = FirstMajorLine; CurrentMajorLine < LastMajorLine; CurrentMajorLine += MajorGridStep)
 	{
@@ -257,12 +235,12 @@ void FSequencerTimeSliderController::DrawTicks( FSlateWindowElementList& OutDraw
 			bAntiAliasLines
 			);
 
-		if (!InArgs.bOnlyDrawMajorTicks)
+		if (!InArgs.bOnlyDrawMajorTicks && !FMath::IsNearlyEqual(MajorLinePx, FlooredScrubPx, 3.f))
 		{
-			FString FrameString = TimeSliderArgs.NumericTypeInterface->ToString((CurrentMajorLine * FrameResolution).RoundToFrame().Value);
+			FString FrameString = TimeSliderArgs.NumericTypeInterface->ToString((CurrentMajorLine * TickResolution).RoundToFrame().Value);
 
 			// Space the text between the tick mark but slightly above
-			FVector2D TextOffset( MajorLinePx + 5.f, InArgs.bMirrorLabels ? 3.f : FMath::Abs( InArgs.AllottedGeometry.Size.Y - (InArgs.MajorTickHeight+3.f) ) );
+			FVector2D TextOffset( MajorLinePx + 5.f, InArgs.bMirrorLabels ? 1.f : FMath::Abs( InArgs.AllottedGeometry.Size.Y - (InArgs.MajorTickHeight+3.f) ) );
 			FSlateDrawElement::MakeText(
 				OutDrawElements,
 				InArgs.StartLayer+1, 
@@ -345,10 +323,10 @@ int32 FSequencerTimeSliderController::OnPaintTimeSlider( bool bMirrorLabels, con
 		LayerId = DrawSelectionRange(AllottedGeometry, MyCullingRect, OutDrawElements, LayerId, RangeToScreen, PlaybackRangeArgs);
 
 		// Draw the scrub handle
-		FFrameTime       ScrubPosition      = TimeSliderArgs.ScrubPosition.Get();
-		FScrubPixelRange ScrubberPixelRange = GetScrubberPixelRange(ScrubPosition, RangeToScreen);
-		const float      HandleStart        = ScrubberPixelRange.HandleRange.GetLowerBoundValue();
-		const float      HandleEnd          = ScrubberPixelRange.HandleRange.GetUpperBoundValue();
+		FQualifiedFrameTime ScrubPosition = FQualifiedFrameTime(TimeSliderArgs.ScrubPosition.Get(), GetTickResolution());
+		FScrubberMetrics    ScrubMetrics  = GetScrubPixelMetrics(ScrubPosition, RangeToScreen);
+		const float         HandleStart   = ScrubMetrics.HandleRangePx.GetLowerBoundValue();
+		const float         HandleEnd     = ScrubMetrics.HandleRangePx.GetUpperBoundValue();
 
 		const int32 ArrowLayer = LayerId + 2;
 		FPaintGeometry MyGeometry =	AllottedGeometry.ToPaintGeometry( FVector2D( HandleStart, 0 ), FVector2D( HandleEnd - HandleStart, AllottedGeometry.Size.Y ) );
@@ -360,9 +338,9 @@ int32 FSequencerTimeSliderController::OnPaintTimeSlider( bool bMirrorLabels, con
 			ScrubColor.G *= 0.2f;
 		}
 
-		const FSlateBrush* Brush = ScrubberPixelRange.bClamped
-			? ( bMirrorLabels ? ClampedScrubHandleUpBrush : ClampedScrubHandleDownBrush )
-			: ( bMirrorLabels ? ScrubHandleUpBrush        : ScrubHandleDownBrush );
+		const FSlateBrush* Brush = ScrubMetrics.Style == ESequencerScrubberStyle::Vanilla
+			? ( bMirrorLabels ? VanillaScrubHandleUpBrush    : VanillaScrubHandleDownBrush )
+			: ( bMirrorLabels ? FrameBlockScrubHandleUpBrush : FrameBlockScrubHandleDownBrush );
 
 		FSlateDrawElement::MakeBox(
 			OutDrawElements,
@@ -389,7 +367,7 @@ int32 FSequencerTimeSliderController::OnPaintTimeSlider( bool bMirrorLabels, con
 			bool  bDrawLeft    = (AllottedGeometry.Size.X - HandleEnd) < (TextSize.X + 14.f) - TextOffsetPx;
 			float TextPosition = bDrawLeft ? HandleStart - TextSize.X - TextOffsetPx : HandleEnd + TextOffsetPx;
 
-			FVector2D TextOffset( TextPosition, Args.bMirrorLabels ? TextSize.Y-6.f : Args.AllottedGeometry.Size.Y - (Args.MajorTickHeight+TextSize.Y) );
+			FVector2D TextOffset( TextPosition, Args.bMirrorLabels ? Args.AllottedGeometry.Size.Y - TextSize.Y : 0.f );
 
 			FSlateDrawElement::MakeText(
 				OutDrawElements,
@@ -404,7 +382,7 @@ int32 FSequencerTimeSliderController::OnPaintTimeSlider( bool bMirrorLabels, con
 		
 		if (MouseDragType == DRAG_SETTING_RANGE)
 		{
-			FFrameRate Resolution     = GetFrameResolution();
+			FFrameRate Resolution = GetTickResolution();
 			FFrameTime MouseDownTime[2];
 
 			FScrubRangeToScreen MouseDownRange(TimeSliderArgs.ViewRange.Get(), MouseDownGeometry.Size);
@@ -421,7 +399,7 @@ int32 FSequencerTimeSliderController::OnPaintTimeSlider( bool bMirrorLabels, con
 				OutDrawElements,
 				LayerId+1,
 				AllottedGeometry.ToPaintGeometry( FVector2D(RangePosX, 0.f), FVector2D(RangeSizeX, AllottedGeometry.Size.Y) ),
-				bMirrorLabels ? ClampedScrubHandleDownBrush : ClampedScrubHandleUpBrush,
+				bMirrorLabels ? VanillaScrubHandleDownBrush : VanillaScrubHandleUpBrush,
 				DrawEffects,
 				MouseStartPosX < MouseEndPosX ? FLinearColor(0.5f, 0.5f, 0.5f) : FLinearColor(0.25f, 0.3f, 0.3f)
 			);
@@ -436,7 +414,7 @@ int32 FSequencerTimeSliderController::OnPaintTimeSlider( bool bMirrorLabels, con
 
 int32 FSequencerTimeSliderController::DrawSelectionRange(const FGeometry& AllottedGeometry, const FSlateRect& MyCullingRect, FSlateWindowElementList& OutDrawElements, int32 LayerId, const FScrubRangeToScreen& RangeToScreen, const FPaintPlaybackRangeArgs& Args) const
 {
-	TRange<double> SelectionRange = TimeSliderArgs.SelectionRange.Get() / GetFrameResolution();
+	TRange<double> SelectionRange = TimeSliderArgs.SelectionRange.Get() / GetTickResolution();
 
 	if (!SelectionRange.IsEmpty())
 	{
@@ -488,9 +466,10 @@ int32 FSequencerTimeSliderController::DrawPlaybackRange(const FGeometry& Allotte
 
 	const uint8 OpacityBlend = TimeSliderArgs.SubSequenceRange.Get().IsSet() ? 128 : 255;
 
-	const TRange<double> PlaybackRange = TimeSliderArgs.PlaybackRange.Get() / GetFrameResolution();
-	const float PlaybackRangeL = RangeToScreen.InputToLocalX(PlaybackRange.GetLowerBoundValue()) - 1;
-	const float PlaybackRangeR = RangeToScreen.InputToLocalX(PlaybackRange.GetUpperBoundValue()) + 1;
+	TRange<FFrameNumber> PlaybackRange = TimeSliderArgs.PlaybackRange.Get();
+	FFrameRate TickResolution = GetTickResolution();
+	const float PlaybackRangeL = RangeToScreen.InputToLocalX(MovieScene::DiscreteInclusiveLower(PlaybackRange) / TickResolution);
+	const float PlaybackRangeR = RangeToScreen.InputToLocalX(MovieScene::DiscreteExclusiveUpper(PlaybackRange) / TickResolution) - 1;
 
 	FSlateDrawElement::MakeBox(
 		OutDrawElements,
@@ -542,7 +521,7 @@ int32 FSequencerTimeSliderController::DrawSubSequenceRange(const FGeometry& Allo
 		return LayerId;
 	}
 
-	const FFrameRate   Resolution = GetFrameResolution();
+	const FFrameRate   Resolution = GetTickResolution();
 	const FFrameNumber LowerFrame = MovieScene::DiscreteInclusiveLower(RangeValue.GetValue());
 	const FFrameNumber UpperFrame = MovieScene::DiscreteExclusiveUpper(RangeValue.GetValue());
 
@@ -697,7 +676,7 @@ FReply FSequencerTimeSliderController::OnMouseButtonUp( SWidget& WidgetOwner, co
 					// push the current value onto the stack
 					ViewRangeStack.Add(ViewRange);
 
-					ViewRange = TRange<double>(MouseDownStart.FrameNumber / GetFrameResolution(), MouseTime.FrameNumber / GetFrameResolution());
+					ViewRange = TRange<double>(MouseDownStart.FrameNumber / GetTickResolution(), MouseTime.FrameNumber / GetTickResolution());
 				}
 				
 				TimeSliderArgs.OnViewRangeChanged.ExecuteIfBound(ViewRange, EViewRangeInterpolation::Immediate);
@@ -782,13 +761,13 @@ FReply FSequencerTimeSliderController::OnMouseMove( SWidget& WidgetOwner, const
 			{
 				FFrameTime MouseDownFree = ComputeFrameTimeFromMouse(MyGeometry, MouseDownPosition[0], RangeToScreen, false);
 
-				const FFrameRate FrameResolution    = GetFrameResolution();
+				const FFrameRate TickResolution     = GetTickResolution();
 				const bool       bLockedPlayRange   = TimeSliderArgs.IsPlaybackRangeLocked.Get();
-				const float      MouseDownPixel     = RangeToScreen.InputToLocalX(MouseDownFree / FrameResolution);
-				const bool       bHitScrubber       = GetHitTestScrubberPixelRange(TimeSliderArgs.ScrubPosition.Get(), RangeToScreen).HandleRange.Contains(MouseDownPixel);
+				const float      MouseDownPixel     = RangeToScreen.InputToLocalX(MouseDownFree / TickResolution);
+				const bool       bHitScrubber       = GetHitTestScrubPixelMetrics(RangeToScreen).HandleRangePx.Contains(MouseDownPixel);
 
-				TRange<double>   SelectionRange   = TimeSliderArgs.SelectionRange.Get() / FrameResolution;
-				TRange<double>   PlaybackRange    = TimeSliderArgs.PlaybackRange.Get()  / FrameResolution;
+				TRange<double>   SelectionRange   = TimeSliderArgs.SelectionRange.Get() / TickResolution;
+				TRange<double>   PlaybackRange    = TimeSliderArgs.PlaybackRange.Get()  / TickResolution;
 
 				// Disable selection range test if it's empty so that the playback range scrubbing gets priority
 				if (!SelectionRange.IsEmpty() && !bHitScrubber && HitTestRangeEnd(RangeToScreen, SelectionRange, MouseDownPixel))
@@ -839,7 +818,7 @@ FReply FSequencerTimeSliderController::OnMouseMove( SWidget& WidgetOwner, const
 			// Set the end range time?
 			else if(MouseDragType == DRAG_PLAYBACK_END)
 			{
-				SetPlaybackRangeEnd(MouseTime.FrameNumber-1);
+				SetPlaybackRangeEnd(MouseTime.FrameNumber);
 			}
 			else if (MouseDragType == DRAG_SELECTION_START)
 			{
@@ -894,7 +873,7 @@ FReply FSequencerTimeSliderController::OnMouseWheel( SWidget& WidgetOwner, const
 		// If zooming on the current time, adjust mouse fractionX
 		if (TimeSliderArgs.Settings->GetZoomPosition() == ESequencerZoomPosition::SZP_CurrentTime)
 		{
-			const double ScrubPosition = TimeSliderArgs.ScrubPosition.Get() / GetFrameResolution();
+			const double ScrubPosition = TimeSliderArgs.ScrubPosition.Get() / GetTickResolution();
 			if (TimeSliderArgs.ViewRange.Get().Contains(ScrubPosition))
 			{
 				FScrubRangeToScreen RangeToScreen(TimeSliderArgs.ViewRange.Get(), MyGeometry.Size);
@@ -922,13 +901,13 @@ FCursorReply FSequencerTimeSliderController::OnCursorQuery( TSharedRef<const SWi
 {
 	FScrubRangeToScreen RangeToScreen(TimeSliderArgs.ViewRange.Get(), MyGeometry.Size);
 
-	const FFrameRate FrameResolution  = GetFrameResolution();
+	const FFrameRate TickResolution   = GetTickResolution();
 	const bool       bLockedPlayRange = TimeSliderArgs.IsPlaybackRangeLocked.Get();
 	const float      HitTestPixel     = MyGeometry.AbsoluteToLocal(CursorEvent.GetScreenSpacePosition()).X;
-	const bool       bHitScrubber     = GetHitTestScrubberPixelRange(TimeSliderArgs.ScrubPosition.Get(), RangeToScreen).HandleRange.Contains(HitTestPixel);
+	const bool       bHitScrubber     = GetHitTestScrubPixelMetrics(RangeToScreen).HandleRangePx.Contains(HitTestPixel);
 
-	TRange<double>   SelectionRange   = TimeSliderArgs.SelectionRange.Get() / FrameResolution;
-	TRange<double>   PlaybackRange    = TimeSliderArgs.PlaybackRange.Get()  / FrameResolution;
+	TRange<double>   SelectionRange   = TimeSliderArgs.SelectionRange.Get() / TickResolution;
+	TRange<double>   PlaybackRange    = TimeSliderArgs.PlaybackRange.Get()  / TickResolution;
 
 	if (MouseDragType == DRAG_SCRUBBING_TIME)
 	{
@@ -992,42 +971,41 @@ int32 FSequencerTimeSliderController::OnPaintSectionView( const FGeometry& Allot
 
 	if( Args.bDisplayScrubPosition )
 	{
-		FScrubPixelRange ScrubberPixelRange = GetScrubberPixelRange(TimeSliderArgs.ScrubPosition.Get(), RangeToScreen);
-
-		if (ScrubberPixelRange.Range.Size<float>() < 2.f)
-		{
-			// Draw a line for the scrub position
-			const float LinePos = ScrubberPixelRange.Range.GetLowerBoundValue() + ScrubberPixelRange.Range.Size<float>() * .5f;
-
-			TArray<FVector2D> LinePoints;
-			{
-				LinePoints.AddUninitialized(2);
-				LinePoints[0] = FVector2D( 0.0f, 0.0f );
-				LinePoints[1] = FVector2D( 0.0f, FMath::FloorToFloat( AllottedGeometry.Size.Y ) );
-			}
+		FQualifiedFrameTime ScrubPosition = FQualifiedFrameTime(TimeSliderArgs.ScrubPosition.Get(), GetTickResolution());
+		FScrubberMetrics    ScrubMetrics = GetScrubPixelMetrics(ScrubPosition, RangeToScreen);
 
-			FSlateDrawElement::MakeLines(
-				OutDrawElements,
-				LayerId+1,
-				AllottedGeometry.ToPaintGeometry( FVector2D(LinePos, 0.0f ), FVector2D(1.0f,1.0f) ),
-				LinePoints,
-				DrawEffects,
-				FLinearColor(1.f, 1.f, 1.f, .5f),
-				false
-			);
-		}
-		else
+		if (ScrubMetrics.bDrawExtents)
 		{
 			// Draw a box for the scrub position
 			FSlateDrawElement::MakeBox(
 				OutDrawElements,
-				LayerId+1,
-				AllottedGeometry.ToPaintGeometry( FVector2D(ScrubberPixelRange.Range.GetLowerBoundValue(), 0.0f ), FVector2D(ScrubberPixelRange.Range.Size<float>(), AllottedGeometry.Size.Y) ),
+				LayerId + 1,
+				AllottedGeometry.ToPaintGeometry(FVector2D(ScrubMetrics.FrameExtentsPx.GetLowerBoundValue(), 0.0f), FVector2D(ScrubMetrics.FrameExtentsPx.Size<float>(), AllottedGeometry.Size.Y)),
 				ScrubFillBrush,
-				DrawEffects, 
-				FLinearColor::White
+				DrawEffects,
+				FLinearColor::White.CopyWithNewOpacity(0.5f)
 			);
 		}
+
+		// Draw a line for the scrub position
+		TArray<FVector2D> LinePoints;
+		{
+			float LinePos = RangeToScreen.InputToLocalX(ScrubPosition.AsSeconds());
+
+			LinePoints.AddUninitialized(2);
+			LinePoints[0] = FVector2D( LinePos, 0.0f );
+			LinePoints[1] = FVector2D( LinePos, FMath::FloorToFloat( AllottedGeometry.Size.Y ) );
+		}
+
+		FSlateDrawElement::MakeLines(
+			OutDrawElements,
+			LayerId+1,
+			AllottedGeometry.ToPaintGeometry(),
+			LinePoints,
+			DrawEffects,
+			FLinearColor(1.f, 1.f, 1.f, .5f),
+			false
+		);
 	}
 
 	return LayerId;
@@ -1145,7 +1123,7 @@ void FSequencerTimeSliderController::ClampViewRange(double& NewRangeMin, double&
 void FSequencerTimeSliderController::SetViewRange( double NewRangeMin, double NewRangeMax, EViewRangeInterpolation Interpolation )
 {
 	// Clamp to a minimum size to avoid zero-sized or negative visible ranges
-	double MinVisibleTimeRange = FFrameNumber(1) / GetFrameResolution();
+	double MinVisibleTimeRange = FFrameNumber(1) / GetTickResolution();
 	TRange<double> ExistingViewRange  = TimeSliderArgs.ViewRange.Get();
 	TRange<double> ExistingClampRange = TimeSliderArgs.ClampRange.Get();
 
@@ -1261,10 +1239,21 @@ bool FSequencerTimeSliderController::HitTestRangeEnd(const FScrubRangeToScreen&
 
 FFrameTime FSequencerTimeSliderController::SnapTimeToNearestKey(const FScrubRangeToScreen& RangeToScreen, float CursorPos, FFrameTime InTime) const
 {
+	if (!WeakSequencer.IsValid())
+	{
+		return InTime;
+	}
+
 	if (TimeSliderArgs.OnGetNearestKey.IsBound())
 	{
-		FFrameNumber NearestKey = TimeSliderArgs.OnGetNearestKey.Execute(InTime);
-		float LocalKeyPos = RangeToScreen.InputToLocalX( NearestKey / GetFrameResolution() );
+		// If there are any tracks selected we'll find the nearest key only on that track. If there are no keys selected,
+		// we will try to find the nearest keys on all tracks. This mirrors the behavior of the Jump to Next Keyframe commands.
+		const TSet< TSharedRef<FSequencerDisplayNode> >& SelectedNodes = WeakSequencer.Pin()->GetSelection().GetSelectedOutlinerNodes();
+		const bool bSearchAllTracks = SelectedNodes.Num() == 0;
+
+		FFrameNumber NearestKey = TimeSliderArgs.OnGetNearestKey.Execute(InTime, bSearchAllTracks);
+
+		float LocalKeyPos = RangeToScreen.InputToLocalX( NearestKey / GetTickResolution() );
 		static float MouseTolerance = 20.f;
 
 		if (FMath::IsNearlyEqual(LocalKeyPos, CursorPos, MouseTolerance))
@@ -1292,7 +1281,7 @@ void FSequencerTimeSliderController::SetPlaybackRangeEnd(FFrameNumber NewEnd)
 
 	if (NewEnd >= MovieScene::DiscreteInclusiveLower(PlaybackRange))
 	{
-		TimeSliderArgs.OnPlaybackRangeChanged.ExecuteIfBound(TRange<FFrameNumber>(PlaybackRange.GetLowerBound(), NewEnd));
+		TimeSliderArgs.OnPlaybackRangeChanged.ExecuteIfBound(TRange<FFrameNumber>(PlaybackRange.GetLowerBound(), TRangeBound<FFrameNumber>::Exclusive(NewEnd)));
 	}
 }
 
diff --git a/Engine/Source/Editor/Sequencer/Private/Tools/EditToolDragOperations.cpp b/Engine/Source/Editor/Sequencer/Private/Tools/EditToolDragOperations.cpp
--- a/Engine/Source/Editor/Sequencer/Private/Tools/EditToolDragOperations.cpp
+++ b/Engine/Source/Editor/Sequencer/Private/Tools/EditToolDragOperations.cpp
@@ -14,53 +14,46 @@
 #include "Channels/MovieSceneChannelProxy.h"
 #include "ISequencerModule.h"
 
-struct FDefaultKeySnappingCandidates : ISequencerSnapCandidate
+struct FInvalidKeyAndSectionSnappingCandidates : ISequencerSnapCandidate
 {
-	FDefaultKeySnappingCandidates(const TSet<FSequencerSelectedKey>& InKeysToExclude)
-		: KeysToExclude(InKeysToExclude)
-	{}
-
-	virtual bool IsKeyApplicable(FKeyHandle KeyHandle, const TSharedPtr<IKeyArea>& KeyArea, UMovieSceneSection* Section) override
-	{
-		return !KeysToExclude.Contains(FSequencerSelectedKey(*Section, KeyArea, KeyHandle));
-	}
-
-	const TSet<FSequencerSelectedKey>& KeysToExclude;
-};
-
-struct FDefaultSectionSnappingCandidates : ISequencerSnapCandidate
-{
-	FDefaultSectionSnappingCandidates(FSectionHandle InSectionToIgnore)
+	/**
+	 * Keys and Sections added to this ISequencerSnapCandidate will be ignored as potential candidates for snapping.
+	 */
+	FInvalidKeyAndSectionSnappingCandidates(const TSet<FSequencerSelectedKey>& InKeysToIgnore, const TArray<FSectionHandle>& InSectionsToIgnore)
 	{
-		SectionsToIgnore.Add(InSectionToIgnore.GetSectionObject());
+		KeysToExclude = InKeysToIgnore;
+		for (const FSectionHandle& SectionHandle : InSectionsToIgnore)
+		{
+			SectionsToExclude.Add(SectionHandle.GetSectionObject());
+		}
 	}
 
-	FDefaultSectionSnappingCandidates(const TArray<FSectionHandle>& InSectionsToIgnore)
+	virtual bool IsKeyApplicable(FKeyHandle KeyHandle, const TSharedPtr<IKeyArea>& KeyArea, UMovieSceneSection* Section) override
 	{
-		for (auto& SectionHandle : InSectionsToIgnore)
-		{
-			SectionsToIgnore.Add(SectionHandle.GetSectionObject());
-		}
+		return !KeysToExclude.Contains(FSequencerSelectedKey(*Section, KeyArea, KeyHandle));
 	}
 
 	virtual bool AreSectionBoundsApplicable(UMovieSceneSection* Section) override
 	{
-		return !SectionsToIgnore.Contains(Section);
+		return !SectionsToExclude.Contains(Section);
 	}
-
-	TSet<UMovieSceneSection*> SectionsToIgnore;
+	
+protected:
+	TSet<FSequencerSelectedKey> KeysToExclude;
+	TSet<UMovieSceneSection*> SectionsToExclude;
 };
 
-TOptional<FSequencerSnapField::FSnapResult> SnapToInterval(const TArray<FFrameNumber>& InTimes, int32 FrameThreshold, FFrameRate Resolution, FFrameRate PlayRate)
+TOptional<FSequencerSnapField::FSnapResult> SnapToInterval(const TArray<FFrameNumber>& InTimes, int32 FrameThreshold, FFrameRate Resolution, FFrameRate DisplayRate, ESequencerScrubberStyle ScrubStyle)
 {
 	TOptional<FSequencerSnapField::FSnapResult> Result;
 
 	FFrameNumber SnapAmount(0);
 	for (FFrameNumber Time : InTimes)
 	{
-		// Convert from resolution to playrate, round to frame, then back again
-		FFrameNumber PlayIntervalTime = FFrameRate::TransformTime(Time,             Resolution, PlayRate  ).RoundToFrame();
-		FFrameNumber IntervalSnap     = FFrameRate::TransformTime(PlayIntervalTime, PlayRate,   Resolution).FloorToFrame();
+		// Convert from resolution to DisplayRate, round to frame, then back again. We floor to frames when using the frame block scrubber, and round using the vanilla scrubber
+		FFrameTime   DisplayTime      = FFrameRate::TransformTime(Time, Resolution, DisplayRate);
+		FFrameNumber PlayIntervalTime = ScrubStyle == ESequencerScrubberStyle::FrameBlock ? DisplayTime.FloorToFrame() : DisplayTime.RoundToFrame();
+		FFrameNumber IntervalSnap     = FFrameRate::TransformTime(PlayIntervalTime, DisplayRate, Resolution  ).FloorToFrame();
 
 		FFrameNumber ThisSnapAmount   = IntervalSnap - Time;
 		if (FMath::Abs(ThisSnapAmount) <= FrameThreshold)
@@ -173,7 +166,8 @@ void FResizeSection::OnBeginDrag(const FPointerEvent& MouseEvent, FVector2D Loca
 	MouseDownTime = VirtualTrackArea.PixelToFrame(LocalMousePos.X);
 
 	// Construct a snap field of unselected sections
-	FDefaultSectionSnappingCandidates SnapCandidates(Sections);
+	TSet<FSequencerSelectedKey> EmptyKeySet;
+	FInvalidKeyAndSectionSnappingCandidates SnapCandidates(EmptyKeySet, Sections);
 	SnapField = FSequencerSnapField(Sequencer, SnapCandidates, ESequencerEntity::Section);
 
 	SectionInitTimes.Empty();
@@ -231,8 +225,10 @@ void FResizeSection::OnDrag(const FPointerEvent& MouseEvent, FVector2D LocalMous
 
 	bool bIsDilating = MouseEvent.IsControlDown();
 
-	FFrameRate   FrameResolution = Sequencer.GetFocusedFrameResolution();
-	FFrameRate   PlayRate        = Sequencer.GetFocusedPlayRate();
+	ESequencerScrubberStyle ScrubStyle = Sequencer.GetScrubStyle();
+
+	FFrameRate   TickResolution  = Sequencer.GetFocusedTickResolution();
+	FFrameRate   DisplayRate     = Sequencer.GetFocusedDisplayRate();
 
 	// Convert the current mouse position to a time
 	FFrameNumber DeltaTime = (VirtualTrackArea.PixelToFrame(LocalMousePos.X) - MouseDownTime).RoundToFrame();
@@ -248,7 +244,7 @@ void FResizeSection::OnDrag(const FPointerEvent& MouseEvent, FVector2D LocalMous
 		}
 
 		float SnapThresholdPx = VirtualTrackArea.PixelToSeconds(PixelSnapWidth) - VirtualTrackArea.PixelToSeconds(0.f);
-		int32 SnapThreshold   = ( SnapThresholdPx * FrameResolution ).FloorToFrame().Value;
+		int32 SnapThreshold   = ( SnapThresholdPx * TickResolution ).FloorToFrame().Value;
 
 		TOptional<FSequencerSnapField::FSnapResult> SnappedTime;
 
@@ -259,8 +255,8 @@ void FResizeSection::OnDrag(const FPointerEvent& MouseEvent, FVector2D LocalMous
 
 		if (!SnappedTime.IsSet() && Settings->GetSnapSectionTimesToInterval())
 		{
-			int32 IntervalSnapThreshold = FMath::RoundToInt( ( FrameResolution / PlayRate ).AsDecimal() );
-			SnappedTime = SnapToInterval(SectionTimes, IntervalSnapThreshold, FrameResolution, PlayRate);
+			int32 IntervalSnapThreshold = FMath::RoundToInt( ( TickResolution / DisplayRate ).AsDecimal() );
+			SnappedTime = SnapToInterval(SectionTimes, IntervalSnapThreshold, TickResolution, DisplayRate, ScrubStyle);
 		}
 
 		if (SnappedTime.IsSet())
@@ -275,11 +271,48 @@ void FResizeSection::OnDrag(const FPointerEvent& MouseEvent, FVector2D LocalMous
 	{
 		for(FPreDragSectionData Data: PreDragSectionData)
 		{
+			// It is only valid to dilate a fixed bound. Tracks can have mixed bounds types (ie: infinite upper, closed lower)
+			check(bDraggingByEnd ? Data.InitialRange.GetUpperBound().IsClosed() : Data.InitialRange.GetLowerBound().IsClosed());
+
 			FFrameNumber StartPosition  = bDraggingByEnd ? MovieScene::DiscreteExclusiveUpper(Data.InitialRange) : MovieScene::DiscreteInclusiveLower(Data.InitialRange);
-			FFrameNumber DilationOrigin = bDraggingByEnd ? MovieScene::DiscreteInclusiveLower(Data.InitialRange) : MovieScene::DiscreteExclusiveUpper(Data.InitialRange);
+
+			FFrameNumber DilationOrigin;
+			if (bDraggingByEnd)
+			{
+				if (Data.InitialRange.GetLowerBound().IsClosed())
+				{
+					DilationOrigin = MovieScene::DiscreteInclusiveLower(Data.InitialRange);
+				}
+				else
+				{
+					// We're trying to dilate a track that has an infinite lower bound as its origin.
+					// Sections already compute an effective range for UMG's auto-playback range, so we'll use that to have it handle finding either the
+					// uppermost key or the overall length of the section.
+					DilationOrigin = Data.MovieSection->ComputeEffectiveRange().GetLowerBoundValue();
+				}
+			}
+			else
+			{
+				if (Data.InitialRange.GetUpperBound().IsClosed())
+				{
+					DilationOrigin = MovieScene::DiscreteExclusiveUpper(Data.InitialRange);
+				}
+				else
+				{
+					// We're trying to dilate a track that has an infinite upper bound as its origin. 
+					DilationOrigin = Data.MovieSection->ComputeEffectiveRange().GetUpperBoundValue();
+
+				}
+			}
+
+			// Because we can have an one-sided infinite data range, we calculate a new range using our clamped values. 
+			TRange<FFrameNumber> DataRange;
+			DataRange.SetLowerBound(TRangeBound<FFrameNumber>(DilationOrigin < StartPosition ? DilationOrigin : StartPosition));
+			DataRange.SetUpperBound(TRangeBound<FFrameNumber>(DilationOrigin > StartPosition ? DilationOrigin : StartPosition));
+
 			FFrameNumber NewPosition    = bDraggingByEnd ? FMath::Max(StartPosition + DeltaTime, DilationOrigin) : FMath::Min(StartPosition + DeltaTime, DilationOrigin);
 
-			float DilationFactor = FMath::Abs(NewPosition.Value - DilationOrigin.Value) / float(MovieScene::DiscreteSize(Data.InitialRange));
+			float DilationFactor = FMath::Abs(NewPosition.Value - DilationOrigin.Value) / float(MovieScene::DiscreteSize(DataRange));
 
 			if (bDraggingByEnd)
 			{
@@ -333,7 +366,7 @@ void FResizeSection::OnDrag(const FPointerEvent& MouseEvent, FVector2D LocalMous
 					NewTime = FMath::Max( NewTime, MinFrame );
 				    if (bIsSlipping)
 					{
-						SequencerSection->SlipSection( NewTime/FrameResolution );
+						SequencerSection->SlipSection( NewTime/TickResolution );
 					}
 					else
 					{
@@ -350,7 +383,7 @@ void FResizeSection::OnDrag(const FPointerEvent& MouseEvent, FVector2D LocalMous
 
 					if (bIsSlipping)
 					{
-						SequencerSection->SlipSection( NewTime/FrameResolution );
+						SequencerSection->SlipSection( NewTime/TickResolution );
 					}
 					else
 					{
@@ -388,224 +421,591 @@ void FResizeSection::OnDrag(const FPointerEvent& MouseEvent, FVector2D LocalMous
 	Sequencer.NotifyMovieSceneDataChanged( EMovieSceneDataChangeType::TrackValueChanged );
 }
 
-FMoveSection::FMoveSection( FSequencer& InSequencer, TArray<FSectionHandle> InSections )
-	: FEditToolDragOperation( InSequencer )
+void CollateTrackNodesByTrack(const TArray<TSharedRef<FSequencerDisplayNode>>& DisplayNodes, TMap<UMovieSceneTrack*, TArray<TSharedRef<FSequencerTrackNode>>>& TrackToTrackNodesMap)
 {
-	// Only allow sections that are not infinite to be movable.
-	for (const FSectionHandle& InSection : InSections)
+	for (TSharedRef<FSequencerDisplayNode> DisplayNode : DisplayNodes)
 	{
-		const UMovieSceneSection* Section = InSection.GetSectionObject();
-		if (Section->HasStartFrame() && Section->HasEndFrame())
+		if (DisplayNode->GetType() == ESequencerNode::Track)
 		{
-			Sections.Add(InSection);
+			TSharedRef<FSequencerTrackNode> TrackNode = StaticCastSharedRef<FSequencerTrackNode>(DisplayNode);
+			TArray<TSharedRef<FSequencerTrackNode>>* TrackNodes = TrackToTrackNodesMap.Find(TrackNode->GetTrack());
+			if (TrackNodes == nullptr)
+			{
+				TrackNodes = &TrackToTrackNodesMap.Add(TrackNode->GetTrack());
+			}
+			TrackNodes->Add(TrackNode);
 		}
-	}
-
-	SequencerNodeTreeUpdatedHandle = InSequencer.GetNodeTree()->OnUpdated().AddRaw(this, &FMoveSection::OnSequencerNodeTreeUpdated);
-}
 
-FMoveSection::~FMoveSection()
-{
-	Sequencer.GetNodeTree()->OnUpdated().Remove(SequencerNodeTreeUpdatedHandle);
+		CollateTrackNodesByTrack(DisplayNode->GetChildNodes(), TrackToTrackNodesMap);
+	}
 }
 
-void FMoveSection::OnBeginDrag(const FPointerEvent& MouseEvent, FVector2D LocalMousePos, const FVirtualTrackArea& VirtualTrackArea)
+bool TryUpdateHandleFromNewTrackNodes(const TArray<TSharedRef<FSequencerTrackNode>>& NewTrackNodes, FSectionHandle& SectionHandle)
 {
-	if (!Sections.Num())
+	UMovieSceneSection* MovieSceneSection = SectionHandle.GetSectionObject();
+	for (TSharedRef<FSequencerTrackNode> NewTrackNode : NewTrackNodes)
 	{
-		return;
+		const TArray<TSharedRef<ISequencerSection>> SequencerSections = NewTrackNode->GetSections();
+		for (int32 i = 0; i < SequencerSections.Num(); i++)
+		{
+			if (SequencerSections[i]->GetSectionObject() == MovieSceneSection)
+			{
+				SectionHandle.TrackNode = NewTrackNode;
+				SectionHandle.SectionIndex = i;
+				return true;
+			}
+		}
 	}
+	return false;
+}
 
-	BeginTransaction( Sections, NSLOCTEXT("Sequencer", "MoveSectionTransaction", "Move Section") );
+void FDuplicateKeysAndSections::OnBeginDrag(const FPointerEvent& MouseEvent, FVector2D LocalMousePos, const FVirtualTrackArea& VirtualTrackArea)
+{
+	// Begin an editor transaction and mark the section as transactional so it's state will be saved
+	BeginTransaction( Sections, NSLOCTEXT("Sequencer", "DuplicateKeysTransaction", "Duplicate Keys or Sections") );
 
-	// Construct a snap field of unselected sections
-	FDefaultSectionSnappingCandidates SnapCandidates(Sections);
-	SnapField = FSequencerSnapField(Sequencer, SnapCandidates, ESequencerEntity::Section);
+	// Call Modify on all of the sections that own keys we have selected so that when we duplicate keys we can restore them properly.
+	ModifyNonSelectedSections();
 
-	const FFrameTime InitialPosition = VirtualTrackArea.PixelToFrame(LocalMousePos.X);
+	// We're going to take our current selection and make a duplicate of each item in it and leave those items behind.
+	// This means our existing selection will still refer to the same keys, so we're duplicating and moving the originals.
+	// This saves us from modifying the user's selection when duplicating. We can't move the duplicates as we can't get
+	// section handles for sections until the tree is rebuilt.
+	TArray<FKeyHandle> NewKeyHandles;
+	NewKeyHandles.SetNumZeroed(KeysAsArray.Num());
 
-	RelativeOffsets.Reserve(Sections.Num());
-	for (int32 Index = 0; Index < Sections.Num(); ++Index)
-	{
-		UMovieSceneSection* Section = Sections[Index].GetSectionObject();
-		// Must have a start frame and end frame to be in the Sections array
-		RelativeOffsets.Add(FRelativeOffset{ Section->GetInclusiveStartFrame() - InitialPosition, Section->GetExclusiveEndFrame() - InitialPosition });
-	}
+	// Duplicate our keys into the NewKeyHandles array. Duplicating keys automatically updates their sections,
+	// so we don't need to actually use the new key handles.
+	DuplicateKeys(KeysAsArray, NewKeyHandles);
 
-	TSet<UMovieSceneTrack*> Tracks;
-	for (int32 Index = 0; Index < Sections.Num(); ++Index)
-	{
-		Tracks.Add(Sections[Index].TrackNode->GetTrack());
-	}
-	for (UMovieSceneTrack* Track : Tracks)
+	// Duplicate our selections as well.
+	bool bDelayedStructureRebuild = false;
+	for (const FSectionHandle& SectionHandle : Sections)
 	{
-		for (auto Section : Track->GetAllSections())
-		{
-			InitialRowIndices.Add(FInitialRowIndex{ Section, Section->GetRowIndex() });
-		}
-	}
-}
+		UMovieSceneSection* DuplicatedSection = DuplicateObject<UMovieSceneSection>(SectionHandle.GetSectionObject(), SectionHandle.GetSectionObject()->GetOuter());
+		UMovieSceneTrack* OwningTrack = SectionHandle.GetSectionObject()->GetTypedOuter<UMovieSceneTrack>();
+		OwningTrack->Modify();
+		OwningTrack->AddSection(*DuplicatedSection);
 
-void FMoveSection::OnEndDrag(const FPointerEvent& MouseEvent, FVector2D LocalMousePos, const FVirtualTrackArea& VirtualTrackArea)
-{
-	if (!Sections.Num())
-	{
-		return;
+		bDelayedStructureRebuild = true;
 	}
 
-	InitialRowIndices.Empty();
+	// Now start the move drag
+	FMoveKeysAndSections::OnBeginDrag(MouseEvent, LocalMousePos, VirtualTrackArea);
 
-	TSet<UMovieSceneTrack*> Tracks;
-	bool bRowIndicesFixed = false;
-	for (auto& Handle : Sections)
-	{
-		Tracks.Add(Handle.TrackNode->GetTrack());
-	}
-	for (UMovieSceneTrack* Track : Tracks)
-	{
-		bRowIndicesFixed |= Track->FixRowIndices();
-	}
-	if (bRowIndicesFixed)
+	if (bDelayedStructureRebuild)
 	{
+		// We need to rebuild the track layout now so that the newly added section shows up, otherwise it won't show up until a section is vertically rearranged.
 		Sequencer.NotifyMovieSceneDataChanged(EMovieSceneDataChangeType::MovieSceneStructureItemsChanged);
 	}
+}
 
-	for (auto& Handle : Sections)
-	{
-		UMovieSceneSection* Section = Handle.GetSectionObject();
-		UMovieSceneTrack* OuterTrack = Cast<UMovieSceneTrack>(Section->GetOuter());
-
-		if (OuterTrack)
-		{
-			OuterTrack->Modify();
-			OuterTrack->OnSectionMoved(*Section);
-		}
-	}
+void FDuplicateKeysAndSections::OnEndDrag(const FPointerEvent& MouseEvent, FVector2D LocalMousePos, const FVirtualTrackArea& VirtualTrackArea)
+{
+	FMoveKeysAndSections::OnEndDrag(MouseEvent, LocalMousePos, VirtualTrackArea);
 
 	EndTransaction();
 }
 
-void FMoveSection::OnDrag(const FPointerEvent& MouseEvent, FVector2D LocalMousePos, const FVirtualTrackArea& VirtualTrackArea)
+FManipulateSectionEasing::FManipulateSectionEasing( FSequencer& InSequencer, FSectionHandle InSection, bool _bEaseIn )
+	: FEditToolDragOperation(InSequencer)
+	, Handle(InSection)
+	, bEaseIn(_bEaseIn)
+	, MouseDownTime(0)
+{
+}
+
+void FManipulateSectionEasing::OnBeginDrag(const FPointerEvent& MouseEvent, FVector2D LocalMousePos, const FVirtualTrackArea& VirtualTrackArea)
 {
-	if (!Sections.Num())
+	Transaction.Reset( new FScopedTransaction(NSLOCTEXT("Sequencer", "DragSectionEasing", "Change Section Easing")) );
+
+	UMovieSceneSection* Section = Handle.GetSectionObject();
+	Section->SetFlags( RF_Transactional );
+	Section->Modify();
+
+	MouseDownTime = VirtualTrackArea.PixelToFrame(LocalMousePos.X);
+
+	if (Settings->GetSnapSectionTimesToSections())
 	{
-		return;
+		// Construct a snap field of all section bounds
+		ISequencerSnapCandidate SnapCandidates;
+		SnapField = FSequencerSnapField(Sequencer, SnapCandidates, ESequencerEntity::Section);
 	}
 
-	LocalMousePos.Y = FMath::Clamp(LocalMousePos.Y, 0.f, VirtualTrackArea.GetPhysicalSize().Y);
+	InitValue = bEaseIn ? Section->Easing.GetEaseInDuration() : Section->Easing.GetEaseOutDuration();
+}
+
+void FManipulateSectionEasing::OnDrag(const FPointerEvent& MouseEvent, FVector2D LocalMousePos, const FVirtualTrackArea& VirtualTrackArea)
+{
+	ESequencerScrubberStyle ScrubStyle = Sequencer.GetScrubStyle();
 
-	FFrameRate FrameResolution = Sequencer.GetFocusedFrameResolution();
-	FFrameRate PlayRate        = Sequencer.GetFocusedPlayRate();
+	FFrameRate TickResolution  = Sequencer.GetFocusedTickResolution();
+	FFrameRate DisplayRate     = Sequencer.GetFocusedDisplayRate();
 
 	// Convert the current mouse position to a time
-	FVector2D  VirtualMousePos = VirtualTrackArea.PhysicalToVirtual(LocalMousePos);
-	FFrameTime MouseTime       = VirtualTrackArea.PixelToFrame(LocalMousePos.X);
+	FFrameTime  DeltaTime = VirtualTrackArea.PixelToFrame(LocalMousePos.X) - MouseDownTime;
 
 	// Snapping
-	if ( Settings->GetIsSnapEnabled() )
+	if (Settings->GetIsSnapEnabled())
 	{
+		TArray<FFrameNumber> SnapTimes;
 
-		float SnapThresholdPx = VirtualTrackArea.PixelToSeconds(PixelSnapWidth) - VirtualTrackArea.PixelToSeconds(0.f);
-		int32 SnapThreshold   = ( SnapThresholdPx * FrameResolution ).FloorToFrame().Value;
-
-		TArray<FFrameNumber> SectionTimes;
-		SectionTimes.Reserve(RelativeOffsets.Num());
-		for (const auto& Offset : RelativeOffsets)
+		UMovieSceneSection* Section = Handle.GetSectionObject();
+		if (bEaseIn)
+		{
+			FFrameNumber DesiredTime = (DeltaTime + Section->GetInclusiveStartFrame() + InitValue.Get(0)).RoundToFrame();
+			SnapTimes.Add(DesiredTime);
+		}
+		else
 		{
-			SectionTimes.Add((Offset.StartTime + MouseTime).FloorToFrame());
-			SectionTimes.Add((Offset.EndTime   + MouseTime).FloorToFrame());
+			FFrameNumber DesiredTime = (Section->GetExclusiveEndFrame() - InitValue.Get(0) + DeltaTime).RoundToFrame();
+			SnapTimes.Add(DesiredTime);
 		}
 
+		float SnapThresholdPx = VirtualTrackArea.PixelToSeconds(PixelSnapWidth) - VirtualTrackArea.PixelToSeconds(0.f);
+		int32 SnapThreshold   = ( SnapThresholdPx * TickResolution ).FloorToFrame().Value;
+
 		TOptional<FSequencerSnapField::FSnapResult> SnappedTime;
 
 		if (Settings->GetSnapSectionTimesToSections())
 		{
-			SnappedTime = SnapField->Snap(SectionTimes, SnapThreshold);
+			SnappedTime = SnapField->Snap(SnapTimes, SnapThreshold);
 		}
 
 		if (!SnappedTime.IsSet() && Settings->GetSnapSectionTimesToInterval())
 		{
-			int32 IntervalSnapThreshold = FMath::RoundToInt( ( FrameResolution / PlayRate ).AsDecimal() );
-			SnappedTime = SnapToInterval(SectionTimes, IntervalSnapThreshold, FrameResolution, PlayRate);
+			int32 IntervalSnapThreshold = FMath::RoundToInt( ( TickResolution / DisplayRate ).AsDecimal() );
+			SnappedTime = SnapToInterval(SnapTimes, IntervalSnapThreshold, TickResolution, DisplayRate, ScrubStyle);
 		}
 
 		if (SnappedTime.IsSet())
 		{
 			// Add the snapped amount onto the delta
-			MouseTime += (SnappedTime->Snapped - SnappedTime->Original);
+			DeltaTime += SnappedTime->Snapped - SnappedTime->Original;
 		}
 	}
 
-	// If sections are all on different rows, don't set row indices for anything because it leads to odd behavior.
-	bool bSectionsAreOnDifferentRows = false;
-	int32 FirstRowIndex = Sections[0].GetSectionObject()->GetRowIndex();
-	TArray<const UMovieSceneSection*> SectionsBeingMoved;
-	for (auto SectionHandle : Sections)
+	UMovieSceneSection* Section = Handle.GetSectionObject();
+
+	const int32 MaxEasingDuration = Section->HasStartFrame() && Section->HasEndFrame() ? MovieScene::DiscreteSize(Section->GetRange()) : TNumericLimits<int32>::Max() / 2;
+
+	if (bEaseIn)
 	{
-		if (FirstRowIndex != SectionHandle.GetSectionObject()->GetRowIndex())
-		{
-			bSectionsAreOnDifferentRows = true;
-		}
-		SectionsBeingMoved.Add(SectionHandle.GetSectionObject());
+		Section->Easing.bManualEaseIn = true;
+		Section->Easing.ManualEaseInDuration  = FMath::Clamp(InitValue.Get(0) + DeltaTime.RoundToFrame().Value, 0, MaxEasingDuration);
+	}
+	else
+	{
+		Section->Easing.bManualEaseOut = true;
+		Section->Easing.ManualEaseOutDuration = FMath::Clamp(InitValue.Get(0) - DeltaTime.RoundToFrame().Value, 0, MaxEasingDuration);
+	}
+
+	UMovieSceneTrack* OuterTrack = Section->GetTypedOuter<UMovieSceneTrack>();
+	if (OuterTrack)
+	{
+		OuterTrack->MarkAsChanged();
 	}
 
-	TOptional<FFrameNumber> MinDeltaXTime;
+	Sequencer.NotifyMovieSceneDataChanged( EMovieSceneDataChangeType::TrackValueChanged );
+}
 
-	// Disallow movement if any of the sections can't move
-	for (int32 Index = 0; Index < Sections.Num(); ++Index)
+void FManipulateSectionEasing::OnEndDrag(const FPointerEvent& MouseEvent, FVector2D LocalMousePos, const FVirtualTrackArea& VirtualTrackArea)
+{
+	EndTransaction();
+}
+
+
+FMoveKeysAndSections::FMoveKeysAndSections(FSequencer& InSequencer, const TSet<FSequencerSelectedKey>& InSelectedKeys, TArray<FSectionHandle> InSelectedSections, bool InbHotspotWasSection)
+	: FEditToolDragOperation(InSequencer)
+	, bHotspotWasSection(InbHotspotWasSection)
+{
+	// Keys are always movable so we simply pass the selected set on to be tracked
+	Keys = InSelectedKeys;
+	KeysAsArray = InSelectedKeys.Array();
+
+	// However, we don't want infinite sections to be movable, so we discard them from our selection.
+	// We support partially infinite (infinite on one side) sections however.
+	for (const FSectionHandle& SectionHandle : InSelectedSections)
 	{
-		auto& Handle = Sections[Index];
-		UMovieSceneSection* Section = Handle.GetSectionObject();
-		if (Section->GetBlendType().IsValid())
+		const UMovieSceneSection* Section = SectionHandle.GetSectionObject();
+		if (Section->HasStartFrame() || Section->HasEndFrame())
 		{
-			continue;
+			Sections.Add(SectionHandle);
 		}
+	}
 
-		const FFrameNumber DeltaTime = (MouseTime + RelativeOffsets[Index].StartTime - Section->GetInclusiveStartFrame()).FloorToFrame();
+	// Register a callback for when the node tree is updated so we can update our local Section Handle array.
+	SequencerNodeTreeUpdatedHandle = InSequencer.GetNodeTree()->OnUpdated().AddRaw(this, &FMoveKeysAndSections::OnSequencerNodeTreeUpdated);
+}
 
-		// Find the borders of where you can move to
-		TRange<FFrameNumber> SectionBoundaries = GetSectionBoundaries(Section, Sections, Handle.TrackNode);
+FMoveKeysAndSections::~FMoveKeysAndSections()
+{
+	Sequencer.GetNodeTree()->OnUpdated().Remove(SequencerNodeTreeUpdatedHandle);
+}
 
-		FFrameNumber LeftMovementMaximum  = MovieScene::DiscreteInclusiveLower(SectionBoundaries);
-		FFrameNumber RightMovementMaximum = MovieScene::DiscreteExclusiveUpper(SectionBoundaries);
-		FFrameNumber NewStartTime         = Section->GetInclusiveStartFrame() + DeltaTime;
-		FFrameNumber NewEndTime           = Section->GetExclusiveEndFrame()   + DeltaTime;
+void FMoveKeysAndSections::OnBeginDrag(const FPointerEvent& MouseEvent, FVector2D LocalMousePos, const FVirtualTrackArea& VirtualTrackArea)
+	{
+	// Early out if we've somehow started a drag operation without any sections or keys. This prevents an empty Undo/Redo Transaction from being created.
+	if (!Sections.Num() && !Keys.Num())
+	{
+		return;
+	}
 
-		if (NewStartTime < LeftMovementMaximum || NewEndTime > RightMovementMaximum)
-		{
-			FFrameNumber ClampedDeltaTime = NewStartTime < LeftMovementMaximum ? LeftMovementMaximum - Section->GetInclusiveStartFrame() : RightMovementMaximum - Section->GetExclusiveEndFrame();
+	BeginTransaction(Sections, NSLOCTEXT("Sequencer", "MoveKeyAndSectionTransaction", "Move Keys or Sections"));
 
-			if (!MinDeltaXTime.IsSet() || MinDeltaXTime.GetValue() > ClampedDeltaTime)
-			{
-				MinDeltaXTime = ClampedDeltaTime;
-			}
-		}
-	}
+	// Tell the Snap Field to ignore our currently selected keys and sections. We can snap to the edges of non-selected
+	// sections and keys. The actual snapping field will add other sequencer data (play ranges, playheads, etc.) as snap targets.
+	FInvalidKeyAndSectionSnappingCandidates AvoidSnapCanidates(Keys, Sections); 
+	SnapField = FSequencerSnapField(Sequencer, AvoidSnapCanidates );
 
-	bool bRowIndexChanged = false;
+	// Store the frame time of the mouse so we can see how far we've moved from the starting point.
+	MouseTimePrev = VirtualTrackArea.PixelToFrame(LocalMousePos.X).FloorToFrame();
+
+	// Now we store a relative offset to each key and section from the start position. This allows us to know how far away from
+	// the mouse each valid key/section was so we can restore their offset if needed.
+	RelativeOffsets.Reserve(Sections.Num() + Keys.Num());
 	for (int32 Index = 0; Index < Sections.Num(); ++Index)
 	{
-		auto& Handle = Sections[Index];
-		UMovieSceneSection* Section = Handle.GetSectionObject();
+		UMovieSceneSection* Section = Sections[Index].GetSectionObject();
+		FRelativeOffset Offset;
 
-		const FFrameNumber DeltaTime = (MouseTime + RelativeOffsets[Index].StartTime - Section->GetInclusiveStartFrame()).FloorToFrame();
+		if (Section->HasStartFrame())
+		{
+			Offset.StartOffset = Section->GetInclusiveStartFrame() - MouseTimePrev;
+		}
+		if (Section->HasEndFrame())
+		{
+			Offset.EndOffset = Section->GetExclusiveEndFrame() - MouseTimePrev;
+		}
 
-		const TArray<UMovieSceneSection*>& AllSections = Handle.TrackNode->GetTrack()->GetAllSections();
+		RelativeOffsets.Add(Offset);
+	}
 
-		TArray<UMovieSceneSection*> NonDraggedSections;
-		for (UMovieSceneSection* TrackSection : AllSections)
+	// Sections can be dragged vertically to adjust their row up or down, so we need to store what row each section is currently on. A section
+	// can be dragged above all other sections - this is accomplished by moving all other sections down. We store the row indices for all sections
+	// in all tracks that we're modifying so we can get them later to move them.
+	TSet<UMovieSceneTrack*> Tracks;
+	for (int32 Index = 0; Index < Sections.Num(); ++Index)
+	{
+		Tracks.Add(Sections[Index].TrackNode->GetTrack());
+	}
+	for (UMovieSceneTrack* Track : Tracks)
+	{
+		for (auto Section : Track->GetAllSections())
 		{
-			if (!SectionsBeingMoved.Contains(TrackSection))
-			{
-				NonDraggedSections.Add(TrackSection);
-			}
+			InitialSectionRowIndicies.Add(FInitialRowIndex{ Section, Section->GetRowIndex() });
 		}
+	}
 
-		int32 TargetRowIndex = Section->GetRowIndex();
+	// Our Key Handles don't store their times so we need to convert the handles into an array of times
+	// so that we can store the relative offset to each one.
+	TArray<FFrameNumber> KeyTimes;
+	KeyTimes.SetNum(Keys.Num());
+	GetKeyTimes(KeysAsArray, KeyTimes);
 
-		// vertical dragging
-		if (Handle.TrackNode->GetTrack()->SupportsMultipleRows() && AllSections.Num() > 1)
-		{
+	for (int32 Index = 0; Index < KeyTimes.Num(); ++Index)
+	{
+		// Key offsets use only the Start offset and don't set the End offset as they do not represent ranges.
+		FRelativeOffset KeyOffset;
+		KeyOffset.StartOffset = KeyTimes[Index] - MouseTimePrev;
+
+		// These are packed with our Section Offsets (ToDo: Is this actually a good idea?)
+		RelativeOffsets.Add(KeyOffset);
+	}
+
+	// Keys can be moved within sections without the section itself being moved, so we need to call Modify on any section that owns a key that isn't also being moved.
+	ModifyNonSelectedSections();
+}
+
+void FMoveKeysAndSections::OnDrag(const FPointerEvent& MouseEvent, FVector2D LocalMousePos, const FVirtualTrackArea& VirtualTrackArea)
+{
+	if (!Sections.Num() && !Keys.Num())
+	{
+		return;
+	}
+
+	ESequencerScrubberStyle ScrubStyle = Sequencer.GetScrubStyle();
+
+	FFrameRate TickResolution = Sequencer.GetFocusedTickResolution();
+	FFrameRate DisplayRate = Sequencer.GetFocusedDisplayRate();
+
+	// Convert the current mouse position to a time
+	FVector2D  VirtualMousePos = VirtualTrackArea.PhysicalToVirtual(LocalMousePos);
+	FFrameTime MouseTime = VirtualTrackArea.PixelToFrame(LocalMousePos.X);
+
+	// Calculate snapping first which modifies our MouseTime to reflect where it would have to be for the closest snap to work.
+	if (Settings->GetIsSnapEnabled())
+	{
+		float SnapThresholdPx = VirtualTrackArea.PixelToSeconds(PixelSnapWidth) - VirtualTrackArea.PixelToSeconds(0.f);
+		int32 SnapThreshold = (SnapThresholdPx * TickResolution).FloorToFrame().Value;
+
+		// The edge of each bounded section as well as each individual key is a valid marker to try and snap to intervals/sections/etc.
+		// We take our stored offsets and add them to our current time to figure out where on the timeline the are currently.
+		TArray<FFrameNumber> ValidSnapMarkers;
+
+		// If they have both keys and settings selected then we snap to the interval if either one of them is enabled, otherwise respect the individual setting.
+		const bool bSnapToInterval = (KeysAsArray.Num() > 0 && Settings->GetSnapKeyTimesToInterval()) || (Sections.Num() > 0 && Settings->GetSnapSectionTimesToInterval());
+		const bool bSnapToLikeTypes = (KeysAsArray.Num() > 0 && Settings->GetSnapKeyTimesToKeys()) || (Sections.Num() > 0 && Settings->GetSnapSectionTimesToSections());
+
+		// RelativeOffsets contains both our sections and our keys, and we add them all as potential things that can snap to stuff.
+		for (const FRelativeOffset& Offset : RelativeOffsets)
+		{
+			if (Offset.StartOffset.IsSet())
+			{
+				ValidSnapMarkers.Add((Offset.StartOffset.GetValue() + MouseTime).FloorToFrame());
+			}
+			if (Offset.EndOffset.IsSet())
+			{
+				ValidSnapMarkers.Add((Offset.EndOffset.GetValue() + MouseTime).FloorToFrame());
+			}
+		}
+
+		// Now we'll try and snap all of these points to the closest valid snap marker (which may be a section or interval)
+		TOptional<FSequencerSnapField::FSnapResult> SnappedTime;
+
+		if (bSnapToLikeTypes)
+		{
+			// This may or may not set the SnappedTime depending on if there are any sections within the threshold.
+			SnappedTime = SnapField->Snap(ValidSnapMarkers, SnapThreshold);
+		}
+
+		if (!SnappedTime.IsSet() && bSnapToInterval)
+		{
+			// Snap to the nearest interval (if enabled). Snapping to other objects has over interval.
+			int32 IntervalSnapThreshold = FMath::RoundToInt((TickResolution / DisplayRate).AsDecimal());
+			SnappedTime = SnapToInterval(ValidSnapMarkers, IntervalSnapThreshold, TickResolution, DisplayRate, ScrubStyle);
+		}
+
+		// If they actually snapped to something (snapping may be on but settings might dictate nothing to snap to) add the difference
+		// to our current MouseTime so that MouseTime reflects the amount needed to move to get to the whole snap point.
+		if (SnappedTime.IsSet())
+		{
+			// Add the snapped amount onto the mouse time so the resulting delta brings us in alignment.
+			MouseTime += (SnappedTime->Snapped - SnappedTime->Original);
+		}
+	}
+
+	// We'll calculate a DeltaX based on limits on movement (snapping, section collision) and then use them on keys and sections below.
+	TOptional<FFrameNumber> MaxDeltaX = GetMovementDeltaX(MouseTime);
+
+	FFrameNumber MouseDeltaTime = (MouseTime - MouseTimePrev).FloorToFrame();
+	MouseTimePrev = MouseTimePrev + MaxDeltaX.Get(MouseDeltaTime);
+
+	// Move sections horizontally (limited by our calculated delta) and vertically based on mouse cursor.
+	bool bSectionMovementModifiedStructure = HandleSectionMovement(MouseTime, VirtualMousePos, LocalMousePos, MaxDeltaX, MouseDeltaTime);
+
+	// Update our key times by moving them by our delta.
+	HandleKeyMovement(MaxDeltaX, MouseDeltaTime);
+
+	// Get a list of the unique tracks in this selection and update their easing so previews draw interactively as you drag.
+	TSet<UMovieSceneTrack*> Tracks;
+	for (auto SectionHandle : Sections)
+	{
+		if (UMovieSceneTrack* Track = SectionHandle.GetSectionObject()->GetTypedOuter<UMovieSceneTrack>())
+		{
+			Tracks.Add(Track);
+		}
+	}
+
+	for (UMovieSceneTrack* Track : Tracks)
+	{
+		Track->UpdateEasing();
+	}
+
+	// If we changed the layout by rearranging sections we need to tell the Sequencer to rebuild things, otherwise just re-evaluate existing tracks.
+	if (bSectionMovementModifiedStructure)
+	{
+		Sequencer.NotifyMovieSceneDataChanged(EMovieSceneDataChangeType::MovieSceneStructureItemsChanged);
+	}
+	else
+	{
+		Sequencer.NotifyMovieSceneDataChanged(EMovieSceneDataChangeType::TrackValueChanged);
+	}
+}
+
+void FMoveKeysAndSections::OnEndDrag(const FPointerEvent& MouseEvent, FVector2D LocalMousePos, const FVirtualTrackArea& VirtualTrackArea)
+{
+	if (!Sections.Num() && !Keys.Num())
+	{
+		return;
+	}
+
+	InitialSectionRowIndicies.Empty();
+	ModifiedNonSelectedSections.Empty();
+
+	// Tracks can tell us if the row indexes for any sections were changed during our drag/drop operation.
+	bool bRowIndicesChanged = false;
+	TSet<UMovieSceneTrack*> Tracks;
+
+	for (auto& SectionHandle : Sections)
+	{
+		// Grab only unique tracks as multiple sections can reside on the same track.
+		Tracks.Add(SectionHandle.TrackNode->GetTrack());
+	}
+
+	for (UMovieSceneTrack* Track : Tracks)
+	{
+		// Ensure all of the tracks have updated the row indices for their sections
+		bRowIndicesChanged |= Track->FixRowIndices();
+	}
+
+	if (bRowIndicesChanged)
+	{
+		Sequencer.NotifyMovieSceneDataChanged(EMovieSceneDataChangeType::MovieSceneStructureItemsChanged);
+	}
+
+	for (const FSectionHandle& SectionHandle : Sections)
+	{
+		UMovieSceneSection* Section = SectionHandle.GetSectionObject();
+		UMovieSceneTrack* OuterTrack = Cast<UMovieSceneTrack>(Section->GetOuter());
+
+		if (OuterTrack)
+		{
+			OuterTrack->Modify();
+			OuterTrack->OnSectionMoved(*Section);
+		}
+	}
+
+	EndTransaction();
+}
+
+void FMoveKeysAndSections::OnSequencerNodeTreeUpdated()
+{
+	TMap<UMovieSceneTrack*, TArray<TSharedRef<FSequencerTrackNode>>> TrackToTrackNodesMap;
+	CollateTrackNodesByTrack(Sequencer.GetNodeTree()->GetRootNodes(), TrackToTrackNodesMap);
+
+	// Update the track nodes in the handles based on the original track and section index.
+	for (FSectionHandle& SectionHandle : Sections)
+	{
+		TArray<TSharedRef<FSequencerTrackNode>>* NewTrackNodes = TrackToTrackNodesMap.Find(SectionHandle.TrackNode->GetTrack());
+		ensureMsgf(NewTrackNodes != nullptr, TEXT("Error rebuilding section handles:  Track not found after node tree update."));
+
+		if (NewTrackNodes != nullptr)
+		{
+			bool bHandleUpdated = TryUpdateHandleFromNewTrackNodes(*NewTrackNodes, SectionHandle);
+			ensureMsgf(bHandleUpdated, TEXT("Error rebuilding section handles: Track node with correct track and section index could not be found."));
+		}
+	}
+}
+
+void FMoveKeysAndSections::ModifyNonSelectedSections()
+{
+	for (const FSequencerSelectedKey& Key : Keys)
+	{
+		UMovieSceneSection* OwningSection = Key.Section;
+		const bool bHasBeenModified = ModifiedNonSelectedSections.Contains(OwningSection);
+		const bool bIsAlreadySelected = Sections.ContainsByPredicate([OwningSection](FSectionHandle Handle) { return Handle.GetSectionObject() == OwningSection; });
+		if (!bHasBeenModified && !bIsAlreadySelected)
+		{
+			OwningSection->SetFlags(RF_Transactional);
+			if (OwningSection->TryModify())
+			{
+				ModifiedNonSelectedSections.Add(OwningSection);
+			}
+		}
+	}
+}
+
+TOptional<FFrameNumber> FMoveKeysAndSections::GetMovementDeltaX(FFrameTime MouseTime)
+{
+	TOptional<FFrameNumber> DeltaX;
+
+	// The delta of the mouse is the difference in the current mouse time vs when we started dragging
+	const FFrameNumber MouseDeltaTime = (MouseTime - MouseTimePrev).FloorToFrame();
+
+	// Disallow movement if any of the sections can't move
+	for (int32 Index = 0; Index < Sections.Num(); ++Index)
+	{
+		const FSectionHandle& SectionHandle = Sections[Index];
+
+		// If we're moving a section that is blending with something then it's OK if it overlaps stuff, the blend amount will get updated at the end.
+		UMovieSceneSection* Section = SectionHandle.GetSectionObject();
+		if (Section->GetBlendType().IsValid())
+		{
+			continue;
+		}
+
+		// We'll calculate this section's borders and clamp the possible delta time to be less than that
+		TRange<FFrameNumber> SectionBoundaries = GetSectionBoundaries(Section, Sections, SectionHandle.TrackNode);
+
+		FFrameNumber LeftMovementMaximum = MovieScene::DiscreteInclusiveLower(SectionBoundaries);
+		FFrameNumber RightMovementMaximum = MovieScene::DiscreteExclusiveUpper(SectionBoundaries);
+		
+		if (Section->HasStartFrame())
+		{
+			FFrameNumber NewStartTime = Section->GetInclusiveStartFrame() + MouseDeltaTime;
+			if (NewStartTime < LeftMovementMaximum)
+			{
+				FFrameNumber ClampedDeltaTime = LeftMovementMaximum - Section->GetInclusiveStartFrame();
+				if (!DeltaX.IsSet() || DeltaX.GetValue() > ClampedDeltaTime)
+				{
+					DeltaX = ClampedDeltaTime;
+				}
+			}
+		}
+		
+		if (Section->HasEndFrame())
+		{
+			FFrameNumber NewEndTime = Section->GetExclusiveEndFrame() + MouseDeltaTime;
+			if (NewEndTime > RightMovementMaximum)
+			{
+				FFrameNumber ClampedDeltaTime = RightMovementMaximum - Section->GetExclusiveEndFrame();
+				if (!DeltaX.IsSet() || DeltaX.GetValue() > ClampedDeltaTime)
+				{
+					DeltaX = ClampedDeltaTime;
+				}
+			}
+		}
+	}
+
+	return DeltaX;
+}
+
+bool FMoveKeysAndSections::HandleSectionMovement(FFrameTime MouseTime, FVector2D VirtualMousePos, FVector2D LocalMousePos, TOptional<FFrameNumber> MaxDeltaX, FFrameNumber DesiredDeltaX)
+{
+	// Don't try to process moving sections if we don't have any sections.
+	if (Sections.Num() == 0)
+	{
+		return false;
+	}
+
+	// If sections are all on different rows, don't set row indices for anything because it leads to odd behavior.
+	bool bSectionsAreOnDifferentRows = false;
+	int32 FirstRowIndex = Sections[0].GetSectionObject()->GetRowIndex();
+	TArray<const UMovieSceneSection*> SectionsBeingMoved;
+	for (auto SectionHandle : Sections)
+	{
+		if (FirstRowIndex != SectionHandle.GetSectionObject()->GetRowIndex())
+		{
+			bSectionsAreOnDifferentRows = true;
+		}
+		SectionsBeingMoved.Add(SectionHandle.GetSectionObject());
+	}
+
+	bool bRowIndexChanged = false;
+	for (int32 Index = 0; Index < Sections.Num(); ++Index)
+	{
+		auto& Handle = Sections[Index];
+		UMovieSceneSection* Section = Handle.GetSectionObject();
+
+
+		const TArray<UMovieSceneSection*>& AllSections = Handle.TrackNode->GetTrack()->GetAllSections();
+
+		TArray<UMovieSceneSection*> NonDraggedSections;
+		for (UMovieSceneSection* TrackSection : AllSections)
+		{
+			if (!SectionsBeingMoved.Contains(TrackSection))
+			{
+				NonDraggedSections.Add(TrackSection);
+			}
+		}
+
+		int32 TargetRowIndex = Section->GetRowIndex();
+
+		// Handle vertical dragging to re-arrange tracks. We don't support vertical rearranging if you're dragging via
+		// a key, as the built in offset causes it to always jump down a row even without moving the mouse.
+		if (Handle.TrackNode->GetTrack()->SupportsMultipleRows() && AllSections.Num() > 1 && bHotspotWasSection)
+		{
 			// Compute the max row index whilst disregarding the one we're dragging
 			int32 MaxRowIndex = 0;
 			for (UMovieSceneSection* NonDraggedSection : NonDraggedSections)
@@ -623,7 +1023,7 @@ void FMoveSection::OnDrag(const FPointerEvent& MouseEvent, FVector2D LocalMouseP
 
 				// Find the total height of the track - this is necessary because tracks may contain key areas, but they will not use sub tracks unless there is more than one row
 				float VirtualSectionBottom = 0.f;
-				Handle.TrackNode->TraverseVisible_ParentFirst([&](FSequencerDisplayNode& Node){ VirtualSectionBottom = Node.GetVirtualBottom(); return true; }, true);
+				Handle.TrackNode->TraverseVisible_ParentFirst([&](FSequencerDisplayNode& Node) { VirtualSectionBottom = Node.GetVirtualBottom(); return true; }, true);
 
 				// Assume same height rows
 				const float VirtualSectionTop = Handle.TrackNode->GetVirtualTop();
@@ -637,14 +1037,14 @@ void FMoveSection::OnDrag(const FPointerEvent& MouseEvent, FVector2D LocalMouseP
 					const int32 NewIndex = FMath::FloorToInt((VirtualMousePos.Y - VirtualSectionTop) / VirtualRowHeight);
 					TargetRowIndex = FMath::Clamp(NewIndex, 0, MaxRowIndex);
 				}
-				
+
 				// If close to the top of the row, move else everything down
-				if (VirtualMousePos.Y <= VirtualSectionTop)
+				if (VirtualMousePos.Y <= VirtualSectionTop || LocalMousePos.Y <= 0)
 				{
 					TargetRowIndex = -1;
 				}
 			}
-			else if(Handle.TrackNode->GetSubTrackMode() == FSequencerTrackNode::ESubTrackMode::SubTrack)
+			else if (Handle.TrackNode->GetSubTrackMode() == FSequencerTrackNode::ESubTrackMode::SubTrack)
 			{
 				TSharedPtr<FSequencerTrackNode> ParentTrack = StaticCastSharedPtr<FSequencerTrackNode>(Handle.TrackNode->GetParent());
 				if (ensure(ParentTrack.IsValid()))
@@ -654,7 +1054,7 @@ void FMoveSection::OnDrag(const FPointerEvent& MouseEvent, FVector2D LocalMouseP
 						TSharedRef<FSequencerDisplayNode> ChildNode = ParentTrack->GetChildNodes()[ChildIndex];
 						float VirtualSectionTop = ChildNode->GetVirtualTop();
 						float VirtualSectionBottom = 0.f;
-						ChildNode->TraverseVisible_ParentFirst([&](FSequencerDisplayNode& Node){ VirtualSectionBottom = Node.GetVirtualBottom(); return true; }, true);
+						ChildNode->TraverseVisible_ParentFirst([&](FSequencerDisplayNode& Node) { VirtualSectionBottom = Node.GetVirtualBottom(); return true; }, true);
 
 						if (VirtualMousePos.Y < VirtualSectionBottom)
 						{
@@ -670,22 +1070,22 @@ void FMoveSection::OnDrag(const FPointerEvent& MouseEvent, FVector2D LocalMouseP
 			}
 		}
 
-		bool bDeltaX = DeltaTime != 0;
+		bool bDeltaX = DesiredDeltaX != 0;
 		bool bDeltaY = TargetRowIndex != Section->GetRowIndex();
 
 		// Horizontal movement
 		if (bDeltaX)
 		{
-			Section->MoveSection(MinDeltaXTime.Get(DeltaTime));
+			Section->MoveSection(MaxDeltaX.Get(DesiredDeltaX));
 		}
 
 		// Vertical movement
 		if (bDeltaY && !bSectionsAreOnDifferentRows &&
 			(
 				Section->GetBlendType().IsValid() ||
-				!Section->OverlapsWithSections(NonDraggedSections, TargetRowIndex - Section->GetRowIndex(), DeltaTime.Value)
+				!Section->OverlapsWithSections(NonDraggedSections, TargetRowIndex - Section->GetRowIndex(), DesiredDeltaX.Value)
+				)
 			)
-		)
 		{
 			// Reached the top, move everything else we're not moving downwards
 			if (TargetRowIndex == -1)
@@ -694,7 +1094,7 @@ void FMoveSection::OnDrag(const FPointerEvent& MouseEvent, FVector2D LocalMouseP
 				{
 					// If the sections being moved are all at the top, and all others are below it, do nothing
 					bool bSectionsBeingMovedAreAtTop = true;
-					for (auto InitialRowIndex : InitialRowIndices)
+					for (auto InitialRowIndex : InitialSectionRowIndicies)
 					{
 						if (!SectionsBeingMoved.Contains(InitialRowIndex.Section))
 						{
@@ -708,12 +1108,12 @@ void FMoveSection::OnDrag(const FPointerEvent& MouseEvent, FVector2D LocalMouseP
 
 					if (!bSectionsBeingMovedAreAtTop)
 					{
-						for (auto InitialRowIndex : InitialRowIndices)
+						for (auto InitialRowIndex : InitialSectionRowIndicies)
 						{
 							if (!SectionsBeingMoved.Contains(InitialRowIndex.Section))
 							{
 								InitialRowIndex.Section->Modify();
-								InitialRowIndex.Section->SetRowIndex(InitialRowIndex.RowIndex+1);
+								InitialRowIndex.Section->SetRowIndex(InitialRowIndex.RowIndex + 1);
 								bRowIndexChanged = true;
 							}
 						}
@@ -729,199 +1129,54 @@ void FMoveSection::OnDrag(const FPointerEvent& MouseEvent, FVector2D LocalMouseP
 		}
 	}
 
-	{
-		TSet<UMovieSceneTrack*> Tracks;
-		for (auto SectionHandle : Sections)
-		{
-			if (UMovieSceneTrack* Track = SectionHandle.GetSectionObject()->GetTypedOuter<UMovieSceneTrack>())
-			{
-				Tracks.Add(Track);
-			}
-		}
-		for (UMovieSceneTrack* Track : Tracks)
-		{
-			Track->UpdateEasing();
-		}
-	}
-
-	Sequencer.NotifyMovieSceneDataChanged( bRowIndexChanged 
-		? EMovieSceneDataChangeType::MovieSceneStructureItemsChanged 
-		: EMovieSceneDataChangeType::TrackValueChanged );
+	return bRowIndexChanged;
 }
 
-void CollateTrackNodesByTrack(const TArray<TSharedRef<FSequencerDisplayNode>>& DisplayNodes, TMap<UMovieSceneTrack*, TArray<TSharedRef<FSequencerTrackNode>>>& TrackToTrackNodesMap)
+void FMoveKeysAndSections::HandleKeyMovement(TOptional<FFrameNumber> MaxDeltaX, FFrameNumber DesiredDeltaX)
 {
-	for (TSharedRef<FSequencerDisplayNode> DisplayNode : DisplayNodes)
-	{
-		if (DisplayNode->GetType() == ESequencerNode::Track)
-		{
-			TSharedRef<FSequencerTrackNode> TrackNode = StaticCastSharedRef<FSequencerTrackNode>(DisplayNode);
-			TArray<TSharedRef<FSequencerTrackNode>>* TrackNodes = TrackToTrackNodesMap.Find(TrackNode->GetTrack());
-			if (TrackNodes == nullptr)
-			{
-				TrackNodes = &TrackToTrackNodesMap.Add(TrackNode->GetTrack());
-			}
-			TrackNodes->Add(TrackNode);
-		}
-
-		CollateTrackNodesByTrack(DisplayNode->GetChildNodes(), TrackToTrackNodesMap);
-	}
-}
-
-bool TryUpdateHandleFromNewTrackNodes(const TArray<TSharedRef<FSequencerTrackNode>>& NewTrackNodes, FSectionHandle& SectionHandle)
-{
-	UMovieSceneSection* MovieSceneSection = SectionHandle.GetSectionObject();
-	for (TSharedRef<FSequencerTrackNode> NewTrackNode : NewTrackNodes)
-	{
-		const TArray<TSharedRef<ISequencerSection>> SequencerSections = NewTrackNode->GetSections();
-		for (int32 i = 0; i < SequencerSections.Num(); i++)
-		{
-			if (SequencerSections[i]->GetSectionObject() == MovieSceneSection)
-			{
-				SectionHandle.TrackNode = NewTrackNode;
-				SectionHandle.SectionIndex = i;
-				return true;
-			}
-		}
-	}
-	return false;
-}
-
-void FMoveSection::OnSequencerNodeTreeUpdated()
-{
-	TMap<UMovieSceneTrack*, TArray<TSharedRef<FSequencerTrackNode>>> TrackToTrackNodesMap;
-	CollateTrackNodesByTrack(Sequencer.GetNodeTree()->GetRootNodes(), TrackToTrackNodesMap);
-
-	// Update the track nodes in the handles based on the original track and section index.
-	for (FSectionHandle& SectionHandle : Sections)
-	{
-		TArray<TSharedRef<FSequencerTrackNode>>* NewTrackNodes = TrackToTrackNodesMap.Find(SectionHandle.TrackNode->GetTrack());
-		ensureMsgf(NewTrackNodes != nullptr, TEXT("Error rebuilding section handles:  Track not found after node tree update."));
-
-		if (NewTrackNodes != nullptr)
-		{
-			bool bHandleUpdated = TryUpdateHandleFromNewTrackNodes(*NewTrackNodes, SectionHandle);
-			ensureMsgf(bHandleUpdated, TEXT("Error rebuilding section handles: Track node with correct track and section index could not be found."));
-		}
-	}
-}
-
-void FMoveKeys::OnBeginDrag(const FPointerEvent& MouseEvent, FVector2D LocalMousePos, const FVirtualTrackArea& VirtualTrackArea)
-{
-	check( SelectedKeys.Num() > 0 )
-
-
-	FDefaultKeySnappingCandidates SnapCandidates(SelectedKeys);
-	SnapField = FSequencerSnapField(Sequencer, SnapCandidates);
-
-	SelectedKeyArray = SelectedKeys.Array();
-
-	// Begin an editor transaction and mark the section as transactional so it's state will be saved
-	TArray<FSectionHandle> DummySections;
-	BeginTransaction( DummySections, NSLOCTEXT("Sequencer", "MoveKeysTransaction", "Move Keys") );
-
-	// Populate the relative offset for each key
-	TArray<FFrameNumber> KeyTimes;
-	KeyTimes.SetNum(SelectedKeyArray.Num());
-	GetKeyTimes(SelectedKeyArray, KeyTimes);
-
-	const FFrameTime MouseTime = VirtualTrackArea.PixelToFrame(LocalMousePos.X);
-
-	RelativeOffsets.Reserve(KeyTimes.Num());
-	for (FFrameNumber Time : KeyTimes)
-	{
-		RelativeOffsets.Add(Time - MouseTime);
-	}
-
-	for( const FSequencerSelectedKey& SelectedKey : SelectedKeyArray )
-	{
-		UMovieSceneSection* OwningSection = SelectedKey.Section;
-
-		// Only modify sections once
-		if( !ModifiedSections.Contains( OwningSection ) )
-		{
-			OwningSection->SetFlags( RF_Transactional );
-
-			// Save the current state of the section
-			if (OwningSection->TryModify())
-			{
-				// Section has been modified
-				ModifiedSections.Add( OwningSection );
-			}
-		}
-	}
-}
-
-void FMoveKeys::OnDrag(const FPointerEvent& MouseEvent, FVector2D LocalMousePos, const FVirtualTrackArea& VirtualTrackArea)
-{
-	if( MouseEvent.GetCursorDelta().X == 0.0f )
+	if (KeysAsArray.Num() == 0)
 	{
 		return;
 	}
 
-	FFrameRate FrameResolution = Sequencer.GetFocusedFrameResolution();
-	FFrameRate PlayRate        = Sequencer.GetFocusedPlayRate();
-
-	FFrameTime MouseTime       = VirtualTrackArea.PixelToFrame(LocalMousePos.X);
-
-	TArray<FFrameNumber> KeyTimes;
-	KeyTimes.Reserve(RelativeOffsets.Num());
-
-	for( FFrameTime Time : RelativeOffsets )
-	{
-		KeyTimes.Add( (MouseTime + Time).FloorToFrame() );
-	}
+	// Apply the delta to our key times. We need to get our key time so that we can add the delta
+	// to each one so that we come up with a new absolute time for it.
+	TArray<FFrameNumber> CurrentKeyTimes;
+	CurrentKeyTimes.SetNum(KeysAsArray.Num());
+	GetKeyTimes(KeysAsArray, CurrentKeyTimes);
 
-	// Snapping
-	if ( Settings->GetIsSnapEnabled() )
+	for (int32 Index = 0; Index < CurrentKeyTimes.Num(); ++Index)
 	{
+		FSequencerSelectedKey& SelectedKey = KeysAsArray[Index];
+		const bool bOwningSectionIsSelected = Sections.ContainsByPredicate([SelectedKey](FSectionHandle Handle) { return Handle.GetSectionObject() == SelectedKey.Section; });
 
-		float SnapThresholdPx = VirtualTrackArea.PixelToSeconds(PixelSnapWidth) - VirtualTrackArea.PixelToSeconds(0.f);
-		int32 SnapThreshold   = ( SnapThresholdPx * FrameResolution ).FloorToFrame().Value;
-
-		TOptional<FSequencerSnapField::FSnapResult> SnappedTime;
-
-		if (Settings->GetSnapKeyTimesToKeys())
-		{
-			SnappedTime = SnapField->Snap(KeyTimes, SnapThreshold);
-		}
-
-		if (!SnappedTime.IsSet() && Settings->GetSnapKeyTimesToInterval())
+		// We don't want to apply delta if we have the key's section selected as well, otherwise they get double
+		// transformed (moving the section moves the keys + we add the delta to the key positions).
+		if (!bOwningSectionIsSelected)
 		{
-			int32 IntervalSnapThreshold = FMath::RoundToInt( ( FrameResolution / PlayRate ).AsDecimal() );
-			SnappedTime = SnapToInterval(KeyTimes, IntervalSnapThreshold, FrameResolution, PlayRate);
-		}
-
-		if (SnappedTime.IsSet())
-		{
-			MouseTime += SnappedTime->Snapped - SnappedTime->Original;
-
-			// Reset the new key times to account for the snap position
-			for( int32 Index = 0; Index < KeyTimes.Num(); ++Index )
-			{
-				KeyTimes[Index] = (MouseTime + RelativeOffsets[Index]).FloorToFrame();
-			}
+			CurrentKeyTimes[Index] += MaxDeltaX.Get(DesiredDeltaX);
 		}
 	}
 
-	// Apply new key times to the selection
-	SetKeyTimes(SelectedKeyArray, KeyTimes);
+	// Now set the times back to the keys.
+	SetKeyTimes(KeysAsArray, CurrentKeyTimes);
 
-	for (int32 Index = 0; Index < KeyTimes.Num(); ++Index)
+	// Expand any sections containing those keys to encompass their new location
+	for (int32 Index = 0; Index < CurrentKeyTimes.Num(); ++Index)
 	{
-		FSequencerSelectedKey SelectedKey = SelectedKeyArray[Index];
+		FSequencerSelectedKey SelectedKey = KeysAsArray[Index];
 
 		UMovieSceneSection* Section = SelectedKey.Section;
-		if (ModifiedSections.Contains(Section))
+		if (ModifiedNonSelectedSections.Contains(Section))
 		{
 			// If the key moves outside of the section resize the section to fit the key
 			// @todo Sequencer - Doesn't account for hitting other sections 
-			const FFrameNumber   NewKeyTime   = KeyTimes[Index];
+			const FFrameNumber   NewKeyTime = CurrentKeyTimes[Index];
 			TRange<FFrameNumber> SectionRange = Section->GetRange();
 
 			if (!SectionRange.Contains(NewKeyTime))
 			{
-				TRange<FFrameNumber> NewRange = TRange<FFrameNumber>::Hull( SectionRange, TRange<FFrameNumber>(NewKeyTime) );
+				TRange<FFrameNumber> NewRange = TRange<FFrameNumber>::Hull(SectionRange, TRange<FFrameNumber>(NewKeyTime));
 				Section->SetRange(NewRange);
 			}
 		}
@@ -929,194 +1184,22 @@ void FMoveKeys::OnDrag(const FPointerEvent& MouseEvent, FVector2D LocalMousePos,
 
 
 	// Snap the play time to the new dragged key time if all the keyframes were dragged to the same time
-	if (Settings->GetSnapPlayTimeToDraggedKey() && KeyTimes.Num())
+	if (Settings->GetSnapPlayTimeToDraggedKey() && CurrentKeyTimes.Num())
 	{
-		FFrameNumber FirstFrame        = KeyTimes[0];
-		auto         EqualsFirstFrame  = [=](FFrameNumber In)
+		FFrameNumber FirstFrame = CurrentKeyTimes[0];
+		auto         EqualsFirstFrame = [=](FFrameNumber In)
 		{
 			return In == FirstFrame;
 		};
 
-		if (Algo::AllOf(KeyTimes, EqualsFirstFrame))
+		if (Algo::AllOf(CurrentKeyTimes, EqualsFirstFrame))
 		{
 			Sequencer.SetLocalTime(FirstFrame);
 		}
 	}
 
-	for (UMovieSceneSection* Section : ModifiedSections)
+	for (UMovieSceneSection* Section : ModifiedNonSelectedSections)
 	{
 		Section->MarkAsChanged();
 	}
-	Sequencer.NotifyMovieSceneDataChanged( EMovieSceneDataChangeType::TrackValueChanged );
-}
-
-void FMoveKeys::OnEndDrag(const FPointerEvent& MouseEvent, FVector2D LocalMousePos, const FVirtualTrackArea& VirtualTrackArea)
-{
-	ModifiedSections.Empty();
-	EndTransaction();
-}
-
-void FDuplicateKeys::OnBeginDrag(const FPointerEvent& MouseEvent, FVector2D LocalMousePos, const FVirtualTrackArea& VirtualTrackArea)
-{
-	// Begin an editor transaction and mark the section as transactional so it's state will be saved
-	TArray<FSectionHandle> DummySections;
-	BeginTransaction( DummySections, NSLOCTEXT("Sequencer", "DuplicateKeysTransaction", "Duplicate Keys") );
-
-	// Modify all the sections first
-	for (const FSequencerSelectedKey& SelectedKey : SelectedKeys)
-	{
-		UMovieSceneSection* OwningSection = SelectedKey.Section;
-
-		// Only modify sections once
-		if (!ModifiedSections.Contains(OwningSection ))
-		{
-			OwningSection->SetFlags(RF_Transactional);
-			// Save the current state of the section
-			if (OwningSection->TryModify())
-			{
-				// Section has been modified
-				ModifiedSections.Add( OwningSection );
-			}
-		}
-	}
-
-	// Then duplicate the keys
-
-	// @todo sequencer: selection in transactions
-	FSequencerSelection& Selection = Sequencer.GetSelection();
-
-	TArray<FSequencerSelectedKey> DuplicatedKeyArray = SelectedKeys.Array();
-	TArray<FKeyHandle> NewKeyHandles;
-	// Ideally we'd memset here, but there's no existing method to copy n bytes to an array n times
-	NewKeyHandles.SetNumUninitialized(DuplicatedKeyArray.Num());
-	for (FKeyHandle& Handle : NewKeyHandles)
-	{
-		Handle = FKeyHandle::Invalid();
-	}
-
-	Selection.EmptySelectedKeys();
-	DuplicateKeys(DuplicatedKeyArray, NewKeyHandles);
-	
-	for (int32 Index = 0; Index < NewKeyHandles.Num(); ++Index)
-	{
-		FSequencerSelectedKey NewKey = DuplicatedKeyArray[Index];
-		NewKey.KeyHandle = NewKeyHandles[Index];
-		Selection.AddToSelection(NewKey);
-	}
-
-	// Now start the move drag
-	FMoveKeys::OnBeginDrag(MouseEvent, LocalMousePos, VirtualTrackArea);
-}
-
-void FDuplicateKeys::OnEndDrag(const FPointerEvent& MouseEvent, FVector2D LocalMousePos, const FVirtualTrackArea& VirtualTrackArea)
-{
-	FMoveKeys::OnEndDrag(MouseEvent, LocalMousePos, VirtualTrackArea);
-
-	EndTransaction();
-}
-
-FManipulateSectionEasing::FManipulateSectionEasing( FSequencer& InSequencer, FSectionHandle InSection, bool _bEaseIn )
-	: FEditToolDragOperation(InSequencer)
-	, Handle(InSection)
-	, bEaseIn(_bEaseIn)
-	, MouseDownTime(0)
-{
-}
-
-void FManipulateSectionEasing::OnBeginDrag(const FPointerEvent& MouseEvent, FVector2D LocalMousePos, const FVirtualTrackArea& VirtualTrackArea)
-{
-	Transaction.Reset( new FScopedTransaction(NSLOCTEXT("Sequencer", "DragSectionEasing", "Change Section Easing")) );
-
-	UMovieSceneSection* Section = Handle.GetSectionObject();
-	Section->SetFlags( RF_Transactional );
-	Section->Modify();
-
-	MouseDownTime = VirtualTrackArea.PixelToFrame(LocalMousePos.X);
-
-	if (Settings->GetSnapSectionTimesToSections())
-	{
-		// Construct a snap field of all section bounds
-		ISequencerSnapCandidate SnapCandidates;
-		SnapField = FSequencerSnapField(Sequencer, SnapCandidates, ESequencerEntity::Section);
-	}
-
-	InitValue = bEaseIn ? Section->Easing.GetEaseInDuration() : Section->Easing.GetEaseOutDuration();
-}
-
-void FManipulateSectionEasing::OnDrag(const FPointerEvent& MouseEvent, FVector2D LocalMousePos, const FVirtualTrackArea& VirtualTrackArea)
-{
-	FFrameRate FrameResolution = Sequencer.GetFocusedFrameResolution();
-	FFrameRate PlayRate        = Sequencer.GetFocusedPlayRate();
-
-	// Convert the current mouse position to a time
-	FFrameTime  DeltaTime = VirtualTrackArea.PixelToFrame(LocalMousePos.X) - MouseDownTime;
-
-	// Snapping
-	if (Settings->GetIsSnapEnabled())
-	{
-		TArray<FFrameNumber> SnapTimes;
-
-		UMovieSceneSection* Section = Handle.GetSectionObject();
-		if (bEaseIn)
-		{
-			FFrameNumber DesiredTime = (DeltaTime + Section->GetInclusiveStartFrame() + InitValue.Get(0)).RoundToFrame();
-			SnapTimes.Add(DesiredTime);
-		}
-		else
-		{
-			FFrameNumber DesiredTime = (Section->GetExclusiveEndFrame() - InitValue.Get(0) + DeltaTime).RoundToFrame();
-			SnapTimes.Add(DesiredTime);
-		}
-
-		float SnapThresholdPx = VirtualTrackArea.PixelToSeconds(PixelSnapWidth) - VirtualTrackArea.PixelToSeconds(0.f);
-		int32 SnapThreshold   = ( SnapThresholdPx * FrameResolution ).FloorToFrame().Value;
-
-		TOptional<FSequencerSnapField::FSnapResult> SnappedTime;
-
-		if (Settings->GetSnapSectionTimesToSections())
-		{
-			SnappedTime = SnapField->Snap(SnapTimes, SnapThreshold);
-		}
-
-		if (!SnappedTime.IsSet() && Settings->GetSnapSectionTimesToInterval())
-		{
-			int32 IntervalSnapThreshold = FMath::RoundToInt( ( FrameResolution / PlayRate ).AsDecimal() );
-			SnappedTime = SnapToInterval(SnapTimes, IntervalSnapThreshold, FrameResolution, PlayRate);
-		}
-
-		if (SnappedTime.IsSet())
-		{
-			// Add the snapped amount onto the delta
-			DeltaTime += SnappedTime->Snapped - SnappedTime->Original;
-		}
-	}
-
-	UMovieSceneSection* Section = Handle.GetSectionObject();
-
-	const int32 MaxEasingDuration = Section->HasStartFrame() && Section->HasEndFrame() ? MovieScene::DiscreteSize(Section->GetRange()) : TNumericLimits<int32>::Max() / 2;
-
-	if (bEaseIn)
-	{
-		Section->Easing.bManualEaseIn = true;
-		Section->Easing.ManualEaseInDuration  = FMath::Clamp(InitValue.Get(0) + DeltaTime.RoundToFrame().Value, 0, MaxEasingDuration);
-	}
-	else
-	{
-		Section->Easing.bManualEaseOut = true;
-		Section->Easing.ManualEaseOutDuration = FMath::Clamp(InitValue.Get(0) - DeltaTime.RoundToFrame().Value, 0, MaxEasingDuration);
-	}
-
-	UMovieSceneTrack* OuterTrack = Section->GetTypedOuter<UMovieSceneTrack>();
-	if (OuterTrack)
-	{
-		OuterTrack->MarkAsChanged();
-	}
-
-	Sequencer.NotifyMovieSceneDataChanged( EMovieSceneDataChangeType::TrackValueChanged );
-}
-
-void FManipulateSectionEasing::OnEndDrag(const FPointerEvent& MouseEvent, FVector2D LocalMousePos, const FVirtualTrackArea& VirtualTrackArea)
-{
-	EndTransaction();
-}
-
+}
\ No newline at end of file
diff --git a/Engine/Source/Editor/Sequencer/Private/Tools/EditToolDragOperations.h b/Engine/Source/Editor/Sequencer/Private/Tools/EditToolDragOperations.h
--- a/Engine/Source/Editor/Sequencer/Private/Tools/EditToolDragOperations.h
+++ b/Engine/Source/Editor/Sequencer/Private/Tools/EditToolDragOperations.h
@@ -123,96 +123,101 @@ private:
 	TOptional<FSequencerSnapField> SnapField;
 };
 
-
-/** Operation to move the currently selected sections */
-class FMoveSection
+/**
+ * This drag operation handles moving both keys and sections depending on what you have selected.
+ */
+class FMoveKeysAndSections
 	: public FEditToolDragOperation
 {
 public:
-	FMoveSection( FSequencer& InSequencer, TArray<FSectionHandle> Sections );
-
-	~FMoveSection();
-
-public:
+	FMoveKeysAndSections(FSequencer& InSequencer, const TSet<FSequencerSelectedKey>& InSelectedKeys, TArray<FSectionHandle> InSelectedSections, bool InbHotspotWasSection);
+	~FMoveKeysAndSections();
 
 	// FEditToolDragOperation interface
-
 	virtual void OnBeginDrag(const FPointerEvent& MouseEvent, FVector2D LocalMousePos, const FVirtualTrackArea& VirtualTrackArea) override;
 	virtual void OnDrag(const FPointerEvent& MouseEvent, FVector2D LocalMousePos, const FVirtualTrackArea& VirtualTrackArea) override;
 	virtual void OnEndDrag(const FPointerEvent& MouseEvent, FVector2D LocalMousePos, const FVirtualTrackArea& VirtualTrackArea) override;
-	virtual FCursorReply GetCursor() const override { return FCursorReply::Cursor( EMouseCursor::CardinalCross ); }
-
-private:
-	/** Callback for when the node tree is updated in sequencer. */
-	void OnSequencerNodeTreeUpdated();
-
-private:
-	/** The sections we are interacting with */
-	TArray<FSectionHandle> Sections;
-
-	/** Array of desired offsets relative to the mouse position. Representes a contiguous array of start/end times directly corresponding to Sections array */
-	struct FRelativeOffset { FFrameTime StartTime, EndTime; };
-
-	TArray<FRelativeOffset> RelativeOffsets;
-
-	/** Optional snap field to use when dragging */
-	TOptional<FSequencerSnapField> SnapField;
+	virtual FCursorReply GetCursor() const override { return FCursorReply::Cursor(EMouseCursor::CardinalCross); }
+	// ~FEditToolDragOperation interface
 
-	/** A handle for the sequencer node tree updated delegate. */
-	FDelegateHandle SequencerNodeTreeUpdatedHandle;
+protected:
+	/** Calculate the possible horizontal movement we can, constrained by sections running into things. */
+	TOptional<FFrameNumber> GetMovementDeltaX(FFrameTime MouseTime);
+	/** Move selected sections, if any. */
+	bool HandleSectionMovement(FFrameTime MouseTime, FVector2D VirtualMousePos, FVector2D LocalMousePos, TOptional<FFrameNumber> MaxDeltaX, FFrameNumber DesiredDeltaX);
+	/** Move selected keys, if any. */
+	void HandleKeyMovement(TOptional<FFrameNumber> MaxDeltaX, FFrameNumber DesiredDeltaX);
 
-	struct FInitialRowIndex { UMovieSceneSection* Section; int32 RowIndex; };
-	TArray<FInitialRowIndex> InitialRowIndices;
-};
+	void OnSequencerNodeTreeUpdated();
 
+	/** Calls Modify on sections that own keys we're moving, as the need to be notified the data is about to change too. */
+	void ModifyNonSelectedSections();
 
-/**
- * Operation to move the currently selected keys
- */
-class FMoveKeys
-	: public FEditToolDragOperation
-{
-public:
-	FMoveKeys( FSequencer& InSequencer, const TSet<FSequencerSelectedKey>& InSelectedKeys )
-		: FEditToolDragOperation(InSequencer)
-		, SelectedKeys( InSelectedKeys )
-	{ }
+protected:
+	/** Array of sections that we're moving. */
+	TArray<FSectionHandle> Sections;
 
-public:
+	/** Set of keys that are being moved. */
+	TSet<FSequencerSelectedKey> Keys;
+	TArray<FSequencerSelectedKey> KeysAsArray;
 
-	// FEditToolDragOperation interface
+	/** What was the time of the mouse for the previous frame? Used to calculate a per-frame delta. */
+	FFrameTime MouseTimePrev;
 
-	virtual void OnBeginDrag(const FPointerEvent& MouseEvent, FVector2D LocalMousePos, const FVirtualTrackArea& VirtualTrackArea) override;
-	virtual void OnDrag(const FPointerEvent& MouseEvent, FVector2D LocalMousePos, const FVirtualTrackArea& VirtualTrackArea) override;
-	virtual void OnEndDrag(const FPointerEvent& MouseEvent, FVector2D LocalMousePos, const FVirtualTrackArea& VirtualTrackArea) override;
+	struct FRelativeOffset
+	{
+		FRelativeOffset()
+			: StartOffset()
+			, EndOffset()
+		{
+		}
+
+		/**
+		 * The offset for the start of the section. Can be unset in the case of a section with no lower bound.
+		 * Keys are represented only by StartOffset and do not have an End Offset (which would imply a range).
+		 */
+		TOptional<FFrameTime> StartOffset;
+
+		/**
+		 * The offset for the end of the section. Can be unset in the case of a section with no upper bound.
+		 */
+		TOptional<FFrameTime> EndOffset;
+	};
 
-protected:
+	/** Array of relative offsets for each selected item. Keys + Sections are both added to this array. */
+	TArray<FRelativeOffset> RelativeOffsets;
 
-	/** The selected keys being moved. */
-	const TSet<FSequencerSelectedKey>& SelectedKeys;
+	struct FInitialRowIndex
+	{
+		UMovieSceneSection* Section;
+		int32 RowIndex;
+	};
 
-	/** Sequential copy of the above set used during drag */
-	TArray<FSequencerSelectedKey> SelectedKeyArray;
+	/** Store the row each section starts on when we start dragging. */
+	TArray<FInitialRowIndex> InitialSectionRowIndicies;
 
-	/** Map of relative offsets from the original mouse position */
-	TArray<FFrameTime> RelativeOffsets;
+	/** Array of sections that we called Modify on because we're editing keys that belong to these sections, but not actually moving these sections. */
+	TArray<UMovieSceneSection*> ModifiedNonSelectedSections;
 
-	/** Snap field used to assist in snapping calculations */
+	/** Optional snap field to use when dragging */
 	TOptional<FSequencerSnapField> SnapField;
 
-	/** The set of sections being modified */
-	TSet<UMovieSceneSection*> ModifiedSections;
+	/** A handle for the sequencer node tree updated delegate. */
+	FDelegateHandle SequencerNodeTreeUpdatedHandle;
+
+	/** If the user is moving them via clicking on the Section then we'll allow vertical re-arranging, otherwise not. */
+	bool bHotspotWasSection;
 };
 
 /**
- * Operation to drag-duplicate the currently selected keys
+ * Operation to drag-duplicate the currently selected keys and sections.
  */
-class FDuplicateKeys : public FMoveKeys
+class FDuplicateKeysAndSections : public FMoveKeysAndSections
 {
 public:
 
-	FDuplicateKeys( FSequencer& InSequencer, const TSet<FSequencerSelectedKey>& InSelectedKeys )
-		: FMoveKeys(InSequencer, InSelectedKeys)
+	FDuplicateKeysAndSections( FSequencer& InSequencer, const TSet<FSequencerSelectedKey>& InSelectedKeys, TArray<FSectionHandle> InSelectedSections, bool InbHotspotWasSection)
+		: FMoveKeysAndSections(InSequencer, InSelectedKeys, InSelectedSections, InbHotspotWasSection)
 	{}
 
 public:
@@ -221,7 +226,6 @@ public:
 	virtual void OnEndDrag(const FPointerEvent& MouseEvent, FVector2D LocalMousePos, const FVirtualTrackArea& VirtualTrackArea) override;
 };
 
-
 /**
  * An operation to change a section's ease in/out by dragging its left or right handle
  */
