Copying //UE4/Dev-Sequencer to //UE4/Dev-Main (Source: //UE4/Dev-Sequencer @ 4044339)

diff --git a/Engine/Source/Editor/Sequencer/Private/Tools/SequencerEditTool_Movement.cpp b/Engine/Source/Editor/Sequencer/Private/Tools/SequencerEditTool_Movement.cpp
--- a/Engine/Source/Editor/Sequencer/Private/Tools/SequencerEditTool_Movement.cpp
+++ b/Engine/Source/Editor/Sequencer/Private/Tools/SequencerEditTool_Movement.cpp
@@ -144,10 +144,29 @@ TSharedPtr<ISequencerEditToolDragOperation> FSequencerEditTool_Movement::CreateD
 		{
 			return HotspotDrag;
 		}
-
-		// Ok, the hotspot doesn't know how to drag - let's decide for ourselves
 		auto HotspotType = DelayedDrag->Hotspot->GetType();
 
+		const bool bSectionsSelected = Selection.GetSelectedSections().Num() > 0;
+		const bool bKeySelected = Selection.GetSelectedKeys().Num() > 0;
+		// @todo sequencer: Make this a customizable UI command modifier?
+		const bool bIsDuplicateEvent = MouseEvent.IsAltDown() || MouseEvent.GetEffectingButton() == EKeys::MiddleMouseButton;
+		const bool bHotspotIsSection = HotspotType == ESequencerHotspot::Section;
+
+		// If they have both keys and sections selected then we only support moving them right now, so we
+		// check for that first before trying to figure out if they're resizing or dilating.
+		if (bSectionsSelected && bKeySelected && !bIsDuplicateEvent)
+		{
+			TArray<FSectionHandle> SelectedSectionHandles = SequencerWidget->GetSectionHandles(Selection.GetSelectedSections());
+			TSet<FSequencerSelectedKey> SelectedKeys = Selection.GetSelectedKeys();
+
+			return MakeShareable(new FMoveKeysAndSections(Sequencer, SelectedKeys, SelectedSectionHandles, bHotspotIsSection));
+		}
+		else if (bIsDuplicateEvent)
+		{
+			return MakeShareable(new FDuplicateKeysAndSections(Sequencer, Selection.GetSelectedKeys(), SequencerWidget->GetSectionHandles(Selection.GetSelectedSections()), bHotspotIsSection));
+		}
+
+
 		TOptional<FSectionHandle> SectionToDrag;
 		if (HotspotType == ESequencerHotspot::Section || HotspotType == ESequencerHotspot::EasingArea)
 		{
@@ -181,7 +200,8 @@ TSharedPtr<ISequencerEditToolDragOperation> FSequencerEditTool_Movement::CreateD
 			}
 			else
 			{
-				return MakeShareable( new FMoveSection( Sequencer, SectionHandles ) );
+				TSet<FSequencerSelectedKey> EmptyKeySet;
+				return MakeShareable( new FMoveKeysAndSections( Sequencer, EmptyKeySet, SectionHandles, true) );
 			}
 		}
 		// Moving key(s)?
@@ -207,23 +227,20 @@ TSharedPtr<ISequencerEditToolDragOperation> FSequencerEditTool_Movement::CreateD
 				SequencerHelpers::UpdateHoveredNodeFromSelectedKeys(Sequencer);
 			}
 
-			// @todo sequencer: Make this a customizable UI command modifier?
-			if (MouseEvent.IsAltDown() || MouseEvent.GetEffectingButton() == EKeys::MiddleMouseButton)
-			{
-				return MakeShareable( new FDuplicateKeys( Sequencer, Selection.GetSelectedKeys() ) );
-			}
-
-			return MakeShareable( new FMoveKeys( Sequencer, Selection.GetSelectedKeys() ) );
+			TArray<FSectionHandle> EmptySectionHandles;
+			return MakeShareable( new FMoveKeysAndSections( Sequencer, Selection.GetSelectedKeys(), EmptySectionHandles, false) );
 		}
 	}
 	// If we're not dragging a hotspot, sections take precedence over keys
 	else if (Selection.GetSelectedSections().Num())
 	{
-		return MakeShareable( new FMoveSection( Sequencer, SequencerWidget->GetSectionHandles(Selection.GetSelectedSections()) ) );
+		TSet<FSequencerSelectedKey> EmptyKeySet;
+		return MakeShareable( new FMoveKeysAndSections( Sequencer, EmptyKeySet, SequencerWidget->GetSectionHandles(Selection.GetSelectedSections()), true ) );
 	}
 	else if (Selection.GetSelectedKeys().Num())
 	{
-		return MakeShareable( new FMoveKeys( Sequencer, Selection.GetSelectedKeys() ) );
+		TArray<FSectionHandle> EmptySectionHandles;
+		return MakeShareable( new FMoveKeysAndSections( Sequencer, Selection.GetSelectedKeys(), EmptySectionHandles, false) );
 	}
 
 	return nullptr;
@@ -464,41 +481,41 @@ FString FSequencerEditTool_Movement::TimeToString(FFrameTime Time, bool IsDelta)
 	{
 		case EFrameNumberDisplayFormats::Seconds:
 		{
-			FFrameRate FrameResolution = Sequencer.GetFocusedFrameResolution();
-			double TimeInSeconds = FrameResolution.AsSeconds(Time);
+			FFrameRate TickResolution = Sequencer.GetFocusedTickResolution();
+			double TimeInSeconds = TickResolution.AsSeconds(Time);
 			return IsDelta ? FString::Printf(TEXT("[%+.2fs]"), TimeInSeconds) : FString::Printf(TEXT("%.2fs"), TimeInSeconds);
 		}
 		case EFrameNumberDisplayFormats::Frames:
 		{
-			FFrameRate FrameResolution = Sequencer.GetFocusedFrameResolution();
-			FFrameRate PlayRate = Sequencer.GetFocusedPlayRate();
+			FFrameRate TickResolution = Sequencer.GetFocusedTickResolution();
+			FFrameRate DisplayRate    = Sequencer.GetFocusedDisplayRate();
 
 			// Convert from sequence resolution into display rate frames.
-			FFrameTime DisplayTime = FFrameRate::TransformTime(Time, FrameResolution, PlayRate);
+			FFrameTime DisplayTime = FFrameRate::TransformTime(Time, TickResolution, DisplayRate);
 			FString SubframeIndicator = FMath::IsNearlyZero(DisplayTime.GetSubFrame()) ? TEXT("") : TEXT("*");
 			int32 ZeroPadFrames = Sequencer.GetSequencerSettings()->GetZeroPadFrames();
 			return IsDelta ? FString::Printf(TEXT("[%+0*d%s]"), ZeroPadFrames, DisplayTime.GetFrame().Value, *SubframeIndicator) : FString::Printf(TEXT("%0*d%s"), ZeroPadFrames, DisplayTime.GetFrame().Value, *SubframeIndicator);
 		}
 		case EFrameNumberDisplayFormats::NonDropFrameTimecode:
 		{
-			FFrameRate SourceFrameRate = Sequencer.GetFocusedFrameResolution();
-			FFrameRate DestinationFrameRate = Sequencer.GetFocusedPlayRate();
+			FFrameRate SourceFrameRate = Sequencer.GetFocusedTickResolution();
+			FFrameRate DestinationFrameRate = Sequencer.GetFocusedDisplayRate();
 
-			FFrameNumber PlayRateFrameNumber = FFrameRate::TransformTime(Time, SourceFrameRate, DestinationFrameRate).FloorToFrame();
+			FFrameNumber DisplayRateFrameNumber = FFrameRate::TransformTime(Time, SourceFrameRate, DestinationFrameRate).FloorToFrame();
 
-			FTimecode AsNonDropTimecode = FTimecode::FromFrameNumber(PlayRateFrameNumber, DestinationFrameRate, false);
+			FTimecode AsNonDropTimecode = FTimecode::FromFrameNumber(DisplayRateFrameNumber, DestinationFrameRate, false);
 
 			const bool bForceSignDisplay = IsDelta;
 			return IsDelta ? FString::Printf(TEXT("[%s]"), *AsNonDropTimecode.ToString(bForceSignDisplay)) : FString::Printf(TEXT("%s"), *AsNonDropTimecode.ToString(bForceSignDisplay));
 		}
 		case EFrameNumberDisplayFormats::DropFrameTimecode:
 		{
-			FFrameRate SourceFrameRate = Sequencer.GetFocusedFrameResolution();
-			FFrameRate DestinationFrameRate = Sequencer.GetFocusedPlayRate();
+			FFrameRate SourceFrameRate = Sequencer.GetFocusedTickResolution();
+			FFrameRate DestinationFrameRate = Sequencer.GetFocusedDisplayRate();
 
-			FFrameNumber PlayRateFrameNumber = FFrameRate::TransformTime(Time, SourceFrameRate, DestinationFrameRate).FloorToFrame();
+			FFrameNumber DisplayRateFrameNumber = FFrameRate::TransformTime(Time, SourceFrameRate, DestinationFrameRate).FloorToFrame();
 
-			FTimecode AsDropTimecode = FTimecode::FromFrameNumber(PlayRateFrameNumber, DestinationFrameRate, true);
+			FTimecode AsDropTimecode = FTimecode::FromFrameNumber(DisplayRateFrameNumber, DestinationFrameRate, true);
 
 			const bool bForceSignDisplay = IsDelta;
 			return IsDelta ? FString::Printf(TEXT("[%s]"), *AsDropTimecode.ToString(bForceSignDisplay)) : FString::Printf(TEXT("%s"), *AsDropTimecode.ToString(bForceSignDisplay));
diff --git a/Engine/Source/Editor/Sequencer/Private/Tools/SequencerEditTool_Selection.cpp b/Engine/Source/Editor/Sequencer/Private/Tools/SequencerEditTool_Selection.cpp
--- a/Engine/Source/Editor/Sequencer/Private/Tools/SequencerEditTool_Selection.cpp
+++ b/Engine/Source/Editor/Sequencer/Private/Tools/SequencerEditTool_Selection.cpp
@@ -23,10 +23,14 @@ struct FSelectionPreviewVisitor final
 	{
 		FSequencerSelectedKey Key(*Section, KeyArea, KeyHandle);
 
-		// If we're trying to change this key's selection state, we go into 'key selection mode', thus we reset the selection state of any nodes that weren't selected by keys
+		// Under default behavior keys have priority, so if a key is changing selection state then we remove any sections from the selection. The user can bypass this
+		// by holding down the control key which will allow selecting both keys and sections.
+		bool bKeySelectionHasPriority = !FSlateApplication::Get().GetModifierKeys().IsControlDown();
 		bool bKeyIsSelected = ExistingSelection.IsSelected(Key);
-		if ( (bKeyIsSelected && SetStateTo == ESelectionPreviewState::NotSelected) ||
-			(!bKeyIsSelected && SetStateTo == ESelectionPreviewState::Selected) )
+
+		if (bKeySelectionHasPriority && 
+			((bKeyIsSelected && SetStateTo == ESelectionPreviewState::NotSelected) ||
+			(!bKeyIsSelected && SetStateTo == ESelectionPreviewState::Selected)))
 		{
 			// Clear selected nodes
 			for (TSharedRef<FSequencerDisplayNode> SelectedNode : NodesSelectedBySections)
@@ -48,9 +52,12 @@ struct FSelectionPreviewVisitor final
 
 	virtual void VisitSection(UMovieSceneSection* Section, TSharedRef<FSequencerDisplayNode> Node) const
 	{
-		// Never select a combination of sections and keys
-		// Never allow infinite sections to be selected (they're only selectable through right click)
-		if (SelectionPreview.GetDefinedKeyStates().Num() == 0 && Section->GetRange() != TRange<FFrameNumber>::All())
+		// If key selection has priority then we check to see if there are any keys selected. If there are key selected, we don't add this section.
+		// Otherwise, we bypass this check and only care that the range isn't infinite (those are selectable via right click)
+		bool bKeySelectionHasPriority = !FSlateApplication::Get().GetModifierKeys().IsControlDown();
+		bool bKeyStateCheck = bKeySelectionHasPriority ? SelectionPreview.GetDefinedKeyStates().Num() == 0 : true;
+
+		if (bKeyStateCheck && Section->GetRange() != TRange<FFrameNumber>::All())
 		{
 			SelectionPreview.SetSelectionState(Section, SetStateTo);
 			SelectionPreview.SetSelectionState(Node, SetStateTo);
@@ -181,7 +188,7 @@ public:
 		const auto& RootNodes = SequencerWidget->GetTreeView()->GetNodeTree()->GetRootNodes();
 
 		// Now walk everything within the current marquee range, setting preview selection states as we go
-		FSequencerEntityWalker Walker(FSequencerEntityRange(TopLeft(), BottomRight(), VirtualTrackArea.GetFrameResolution()), VirtualKeySize);
+		FSequencerEntityWalker Walker(FSequencerEntityRange(TopLeft(), BottomRight(), VirtualTrackArea.GetTickResolution()), VirtualKeySize);
 		Walker.Traverse(FSelectionPreviewVisitor(SelectionPreview, Sequencer.GetSelection(), PreviewState), RootNodes);
 	}
 
diff --git a/Engine/Source/Editor/Sequencer/Private/Tools/SequencerEntityVisitor.cpp b/Engine/Source/Editor/Sequencer/Private/Tools/SequencerEntityVisitor.cpp
--- a/Engine/Source/Editor/Sequencer/Private/Tools/SequencerEntityVisitor.cpp
+++ b/Engine/Source/Editor/Sequencer/Private/Tools/SequencerEntityVisitor.cpp
@@ -6,14 +6,14 @@
 #include "MovieSceneTrack.h"
 #include "MovieSceneTimeHelpers.h"
 
-FSequencerEntityRange::FSequencerEntityRange(const TRange<double>& InRange, FFrameRate InFrameResolution)
-	: FrameResolution(InFrameResolution)
+FSequencerEntityRange::FSequencerEntityRange(const TRange<double>& InRange, FFrameRate InTickResolution)
+	: TickResolution(InTickResolution)
 	, Range(InRange)
 {
 }
 
-FSequencerEntityRange::FSequencerEntityRange(FVector2D TopLeft, FVector2D BottomRight, FFrameRate InFrameResolution)
-	: FrameResolution(InFrameResolution)
+FSequencerEntityRange::FSequencerEntityRange(FVector2D TopLeft, FVector2D BottomRight, FFrameRate InTickResolution)
+	: TickResolution(InTickResolution)
 	, Range(TopLeft.X, BottomRight.X)
 	, VerticalTop(TopLeft.Y), VerticalBottom(BottomRight.Y)
 {
@@ -22,7 +22,7 @@ FSequencerEntityRange::FSequencerEntityRange(FVector2D TopLeft, FVector2D Bottom
 bool FSequencerEntityRange::IntersectSection(const UMovieSceneSection* InSection, const TSharedRef<FSequencerTrackNode>& InTrackNode, int32 MaxRowIndex) const
 {
 	// Test horizontal bounds
-	if (!(InSection->GetRange() / FrameResolution).Overlaps(Range))
+	if (!(InSection->GetRange() / TickResolution).Overlaps(Range))
 	{
 		return false;
 	}
@@ -231,9 +231,9 @@ void FSequencerEntityWalker::HandleKeyArea(const ISequencerEntityVisitor& Visito
 	TArray<FKeyHandle> Handles;
 	TArray<FFrameNumber> Times;
 
-	const FFrameTime HalfKeySizeFrames   = (VirtualKeySize.X*.5f) * Range.FrameResolution;
-	const FFrameTime RangeStartFrame     = Range.Range.GetLowerBoundValue() * Range.FrameResolution - FFrameTime(0, 0.5f);
-	const FFrameTime RangeEndFrame       = Range.Range.GetUpperBoundValue() * Range.FrameResolution - FFrameTime(0, 0.5f);
+	const FFrameTime HalfKeySizeFrames   = (VirtualKeySize.X*.5f) * Range.TickResolution;
+	const FFrameTime RangeStartFrame     = Range.Range.GetLowerBoundValue() * Range.TickResolution;
+	const FFrameTime RangeEndFrame       = Range.Range.GetUpperBoundValue() * Range.TickResolution;
 
 	TRange<FFrameNumber> VisitRangeFrames( (RangeStartFrame-HalfKeySizeFrames).CeilToFrame(), (RangeEndFrame+HalfKeySizeFrames).FloorToFrame() );
 	KeyArea->GetKeyInfo(&Handles, &Times, VisitRangeFrames);
diff --git a/Engine/Source/Editor/Sequencer/Private/Tools/SequencerEntityVisitor.h b/Engine/Source/Editor/Sequencer/Private/Tools/SequencerEntityVisitor.h
--- a/Engine/Source/Editor/Sequencer/Private/Tools/SequencerEntityVisitor.h
+++ b/Engine/Source/Editor/Sequencer/Private/Tools/SequencerEntityVisitor.h
@@ -44,8 +44,8 @@ protected:
 /** A range specifying time (and possibly vertical) bounds in the sequencer */
 struct FSequencerEntityRange
 {
-	FSequencerEntityRange(const TRange<double>& InRange, FFrameRate InFrameResolution);
-	FSequencerEntityRange(FVector2D TopLeft, FVector2D BottomRight, FFrameRate InFrameResolution);
+	FSequencerEntityRange(const TRange<double>& InRange, FFrameRate InTickResolution);
+	FSequencerEntityRange(FVector2D TopLeft, FVector2D BottomRight, FFrameRate InTickResolution);
 
 	/** Check whether the specified section intersects this range */
 	bool IntersectSection(const UMovieSceneSection* InSection, const TSharedRef<FSequencerTrackNode>& InTrackNode, int32 MaxRowIndex) const;
@@ -56,8 +56,8 @@ struct FSequencerEntityRange
 	/** Check whether the specified node's key area intersects this range */
 	bool IntersectKeyArea(TSharedRef<FSequencerDisplayNode> InNode, float VirtualKeyHeight) const;
 
-	/** Frame resolution of the current time-base */
-	FFrameRate FrameResolution;
+	/** tick resolution of the current time-base */
+	FFrameRate TickResolution;
 
 	/** Start/end times */
 	TRange<double> Range;
@@ -70,7 +70,7 @@ struct FSequencerEntityRange
 struct FSequencerEntityWalker
 {
 	/** Construction from the range itself, and an optional virtual key size, where key bounds must be taken into consideration */
-	FSequencerEntityWalker(const FSequencerEntityRange& InRange, FVector2D InVirtualKeySize = FVector2D());
+	FSequencerEntityWalker(const FSequencerEntityRange& InRange, FVector2D InVirtualKeySize);
 
 	/** Visit the specified nodes (recursively) with this range and a user-supplied visitor */
 	void Traverse(const ISequencerEntityVisitor& Visitor, const TArray< TSharedRef<FSequencerDisplayNode> >& Nodes);
diff --git a/Engine/Source/Editor/Sequencer/Private/Tools/SequencerSnapField.cpp b/Engine/Source/Editor/Sequencer/Private/Tools/SequencerSnapField.cpp
--- a/Engine/Source/Editor/Sequencer/Private/Tools/SequencerSnapField.cpp
+++ b/Engine/Source/Editor/Sequencer/Private/Tools/SequencerSnapField.cpp
@@ -6,6 +6,7 @@
 #include "SSequencerTreeView.h"
 #include "MovieSceneTimeHelpers.h"
 #include "MovieSceneSequence.h"
+#include "ISequencerSection.h"
 
 struct FSnapGridVisitor : ISequencerEntityVisitor
 {
@@ -62,7 +63,9 @@ FSequencerSnapField::FSequencerSnapField(const ISequencer& InSequencer, ISequenc
 	}
 
 	TRange<double> ViewRange = InSequencer.GetViewRange();
-	FSequencerEntityWalker Walker(FSequencerEntityRange(ViewRange, InSequencer.GetFocusedFrameResolution()));
+	FSequencerEntityWalker Walker(
+		FSequencerEntityRange(ViewRange, InSequencer.GetFocusedTickResolution()),
+		FVector2D(SequencerSectionConstants::KeySize));
 
 	// Traverse the visible space, collecting snapping times as we go
 	FSnapGridVisitor Visitor(Candidate, EntityMask);
@@ -70,16 +73,22 @@ FSequencerSnapField::FSequencerSnapField(const ISequencer& InSequencer, ISequenc
 
 	// Add the playback range start/end bounds as potential snap candidates
 	TRange<FFrameNumber> PlaybackRange = InSequencer.GetFocusedMovieSceneSequence()->GetMovieScene()->GetPlaybackRange();
-	Visitor.Snaps.Add(FSequencerSnapPoint{ FSequencerSnapPoint::PlaybackRange, PlaybackRange.GetLowerBoundValue() });
-	Visitor.Snaps.Add(FSequencerSnapPoint{ FSequencerSnapPoint::PlaybackRange, PlaybackRange.GetUpperBoundValue() });
+	if(MovieScene::DiscreteSize(PlaybackRange) > 0)
+	{
+		Visitor.Snaps.Add(FSequencerSnapPoint{ FSequencerSnapPoint::PlaybackRange, MovieScene::DiscreteInclusiveLower(PlaybackRange)});
+		Visitor.Snaps.Add(FSequencerSnapPoint{ FSequencerSnapPoint::PlaybackRange, MovieScene::DiscreteExclusiveUpper(PlaybackRange) - 1});
+	}
 
 	// Add the current time as a potential snap candidate
 	Visitor.Snaps.Add(FSequencerSnapPoint{ FSequencerSnapPoint::CurrentTime, InSequencer.GetLocalTime().Time.FrameNumber });
 
 	// Add the selection range bounds as a potential snap candidate
 	TRange<FFrameNumber> SelectionRange = InSequencer.GetFocusedMovieSceneSequence()->GetMovieScene()->GetSelectionRange();
-	Visitor.Snaps.Add(FSequencerSnapPoint{ FSequencerSnapPoint::InOutRange, MovieScene::DiscreteInclusiveLower(SelectionRange) });
-	Visitor.Snaps.Add(FSequencerSnapPoint{ FSequencerSnapPoint::InOutRange, MovieScene::DiscreteExclusiveUpper(SelectionRange) });
+	if (MovieScene::DiscreteSize(SelectionRange) > 0)
+	{
+		Visitor.Snaps.Add(FSequencerSnapPoint{ FSequencerSnapPoint::InOutRange, MovieScene::DiscreteInclusiveLower(SelectionRange)});
+		Visitor.Snaps.Add(FSequencerSnapPoint{ FSequencerSnapPoint::InOutRange, MovieScene::DiscreteExclusiveUpper(SelectionRange) - 1});
+	}
 
 	// Sort
 	Visitor.Snaps.Sort([](const FSequencerSnapPoint& A, const FSequencerSnapPoint& B){
diff --git a/Engine/Source/Editor/Sequencer/Private/VirtualTrackArea.cpp b/Engine/Source/Editor/Sequencer/Private/VirtualTrackArea.cpp
--- a/Engine/Source/Editor/Sequencer/Private/VirtualTrackArea.cpp
+++ b/Engine/Source/Editor/Sequencer/Private/VirtualTrackArea.cpp
@@ -5,7 +5,7 @@
 #include "SequencerTrackNode.h"
 
 FVirtualTrackArea::FVirtualTrackArea(const FSequencer& InSequencer, SSequencerTreeView& InTreeView, const FGeometry& InTrackAreaGeometry)
-	: FTimeToPixel(InTrackAreaGeometry, InSequencer.GetViewRange(), InSequencer.GetFocusedFrameResolution())
+	: FTimeToPixel(InTrackAreaGeometry, InSequencer.GetViewRange(), InSequencer.GetFocusedTickResolution())
 	, TreeView(InTreeView)
 	, TrackAreaGeometry(InTrackAreaGeometry)
 {
@@ -121,7 +121,7 @@ FSequencerSelectedKey FVirtualTrackArea::HitTestKey(FVector2D InPhysicalPosition
 
 	const double   KeyLeft  = PixelToSeconds(InPhysicalPosition.X - SequencerSectionConstants::KeySize.X/2);
 	const double   KeyRight = PixelToSeconds(InPhysicalPosition.X + SequencerSectionConstants::KeySize.X/2);
-	TRange<FFrameNumber> KeyRange((KeyLeft * GetFrameResolution()).FloorToFrame(), (KeyRight * GetFrameResolution()).CeilToFrame());
+	TRange<FFrameNumber> KeyRange((KeyLeft * GetTickResolution()).FloorToFrame(), (KeyRight * GetTickResolution()).CeilToFrame());
 
 	TArray<TSharedRef<IKeyArea>> KeyAreas;
 
diff --git a/Engine/Source/Editor/Sequencer/Public/ISequencer.h b/Engine/Source/Editor/Sequencer/Public/ISequencer.h
--- a/Engine/Source/Editor/Sequencer/Public/ISequencer.h
+++ b/Engine/Source/Editor/Sequencer/Public/ISequencer.h
@@ -66,6 +66,22 @@ enum class EAllowEditsMode : uint8
 	AllowLevelEditsOnly
 };
 
+/**
+* Defines set key gruops mode.
+*/
+UENUM()
+enum class EKeyGroupMode : uint8
+{
+	/** Key just changed channel */
+	KeyChanged,
+
+	/** Key just one, the parent translation, rotation or scale, when one changes */
+	KeyGroup,
+
+	/** Key All (translation, rotation, scale) when one changes */
+	KeyAll
+};
+
 
 /**
  * Enumerates types of UI Command bindings.
@@ -232,11 +248,15 @@ public:
 	/** Sets where edits are allowed */
 	virtual void SetAllowEditsMode(EAllowEditsMode AllowEditsMode) = 0;
 
-	/** @return Returns whether key all is enabled in this sequencer */
-	virtual bool GetKeyAllEnabled() const = 0;
+	/** @returns Returns what channels will get keyed when one channel changes */
+	virtual EKeyGroupMode GetKeyGroupMode() const = 0;
 
-	/** Sets whether key all is enabled in this sequencer. */
-	virtual void SetKeyAllEnabled(bool bKeyAllEnabled) = 0;
+	/** Sets which channels are keyed when a channel is keyed 
+	* @param Mode Specifies which channels to key, all (EKeyGroupMode::KeyAll), 
+	* just changed (EKeyGroup::KeyChanged), 
+	* just the group,e.g. Rotation X,Y and Z if any of those are changed (EKeyGroup::KeyGroup)
+	*/
+	virtual void SetKeyGroupMode(EKeyGroupMode Mode) = 0;
 
 	/** @return Returns whether or not to key only interp properties in this sequencer */
 	virtual bool GetKeyInterpPropertiesOnly() const = 0;
@@ -517,24 +537,24 @@ public:
 public:
 
 	/**
-	 * Get the frame resolution of the currently root sequence
+	 * Get the tick resolution of the currently root sequence
 	 */
-	SEQUENCER_API FFrameRate GetRootFrameResolution() const;
+	SEQUENCER_API FFrameRate GetRootTickResolution() const;
 
 	/**
-	 * Get the playback frame rate of the currently root sequence
+	 * Get the display rate of the currently root sequence
 	 */
-	SEQUENCER_API FFrameRate GetRootPlayRate() const;
+	SEQUENCER_API FFrameRate GetRootDisplayRate() const;
 
 	/**
-	 * Get the frame resolution of the currently focused sequence
+	 * Get the tick resolution of the currently focused sequence
 	 */
-	SEQUENCER_API FFrameRate GetFocusedFrameResolution() const;
+	SEQUENCER_API FFrameRate GetFocusedTickResolution() const;
 
 	/**
-	 * Get the playback frame rate of the currently focused sequence
+	 * Get the display rate of the currently focused sequence
 	 */
-	SEQUENCER_API FFrameRate GetFocusedPlayRate() const;
+	SEQUENCER_API FFrameRate GetFocusedDisplayRate() const;
 
 protected:
 	FOnInitializeDetailsPanel InitializeDetailsPanelEvent;
diff --git a/Engine/Source/Editor/Sequencer/Public/ISequencerModule.h b/Engine/Source/Editor/Sequencer/Public/ISequencerModule.h
--- a/Engine/Source/Editor/Sequencer/Public/ISequencerModule.h
+++ b/Engine/Source/Editor/Sequencer/Public/ISequencerModule.h
@@ -24,6 +24,15 @@ namespace SequencerMenuExtensionPoints
 	static const FName AddTrackMenu_PropertiesSection("AddTrackMenu_PropertiesSection");
 }
 
+/** Enum representing supported scrubber styles */
+enum class ESequencerScrubberStyle : uint8
+{
+	/** Scrubber is represented as a single thin line for the current time, with a constant-sized thumb. */
+	Vanilla,
+
+	/** Scrubber thumb occupies a full 'display rate' frame, with a single thin line for the current time. Tailored to frame-accuracy scenarios. */
+	FrameBlock,
+};
 
 /** A delegate which will create an auto-key handler. */
 DECLARE_DELEGATE_RetVal_OneParam(TSharedRef<ISequencerTrackEditor>, FOnCreateTrackEditor, TSharedRef<ISequencer>);
@@ -64,9 +73,13 @@ struct FSequencerViewParams
 	/** Whether the sequencer is read-only */
 	bool bReadOnly;
 
+	/** Style of scrubber to use */
+	ESequencerScrubberStyle ScrubberStyle;
+
 	FSequencerViewParams(FString InName = FString())
 		: UniqueName(MoveTemp(InName))
 		, bReadOnly(false)
+		, ScrubberStyle(ESequencerScrubberStyle::Vanilla)
 	{ }
 };
 
diff --git a/Engine/Source/Editor/Sequencer/Public/SequencerChannelTraits.h b/Engine/Source/Editor/Sequencer/Public/SequencerChannelTraits.h
--- a/Engine/Source/Editor/Sequencer/Public/SequencerChannelTraits.h
+++ b/Engine/Source/Editor/Sequencer/Public/SequencerChannelTraits.h
@@ -500,7 +500,7 @@ namespace Sequencer
 
 		auto ForEachKey = [Section, PasteAt, &NewRange, &ChannelInterface, &OutPastedKeys, &SrcEnvironment, &DstEnvironment](const FMovieSceneClipboardKey& Key)
 		{
-			FFrameNumber Time = (PasteAt + FFrameRate::TransformTime(Key.GetTime(), SrcEnvironment.FrameResolution, DstEnvironment.FrameResolution)).FloorToFrame();
+			FFrameNumber Time = (PasteAt + FFrameRate::TransformTime(Key.GetTime(), SrcEnvironment.TickResolution, DstEnvironment.TickResolution)).FloorToFrame();
 
 			NewRange = TRange<FFrameNumber>::Hull(NewRange, TRange<FFrameNumber>(Time));
 
diff --git a/Engine/Source/Editor/Sequencer/Public/SequencerCommands.h b/Engine/Source/Editor/Sequencer/Public/SequencerCommands.h
--- a/Engine/Source/Editor/Sequencer/Public/SequencerCommands.h
+++ b/Engine/Source/Editor/Sequencer/Public/SequencerCommands.h
@@ -148,6 +148,9 @@ public:
 	/** Sets the interp tangent mode for the selected keys to break */
 	TSharedPtr< FUICommandInfo > SetInterpolationCubicBreak;
 
+	/** Toggles the interp tangent weight mode for the selected keys */
+	TSharedPtr< FUICommandInfo > ToggleWeightedTangents;
+
 	/** Sets the interp tangent mode for the selected keys to linear */
 	TSharedPtr< FUICommandInfo > SetInterpolationLinear;
 
@@ -193,8 +196,14 @@ public:
 	/** Turns autokey on and off. */
 	TSharedPtr< FUICommandInfo > ToggleAutoKeyEnabled;
 
-	/** Turns key all on and off. */
-	TSharedPtr< FUICommandInfo > ToggleKeyAllEnabled;
+	/** Set mode to just key changed attribute. */
+	TSharedPtr< FUICommandInfo > SetKeyChanged;
+
+	/** Set mode to key changed attribute and others in it's group. */
+	TSharedPtr< FUICommandInfo > SetKeyGroup;
+
+	/** Set mode to key all. */
+	TSharedPtr< FUICommandInfo > SetKeyAll;
 
 	/** Rotates through the supported formats for displaying times/frames/timecode. */
 	TSharedPtr< FUICommandInfo > ChangeTimeDisplayFormat;
diff --git a/Engine/Source/Editor/Sequencer/Public/SequencerSectionPainter.h b/Engine/Source/Editor/Sequencer/Public/SequencerSectionPainter.h
--- a/Engine/Source/Editor/Sequencer/Public/SequencerSectionPainter.h
+++ b/Engine/Source/Editor/Sequencer/Public/SequencerSectionPainter.h
@@ -60,5 +60,8 @@ public:
 
 	/** Whether the section is hovered or not */
 	bool bIsHighlighted;
+	
+	/** Whether the section is selected or not */
+	bool bIsSelected;
 };
 
diff --git a/Engine/Source/Editor/Sequencer/Public/SequencerSettings.h b/Engine/Source/Editor/Sequencer/Public/SequencerSettings.h
--- a/Engine/Source/Editor/Sequencer/Public/SequencerSettings.h
+++ b/Engine/Source/Editor/Sequencer/Public/SequencerSettings.h
@@ -12,6 +12,7 @@
 struct FPropertyChangedEvent;
 enum class EAutoChangeMode : uint8;
 enum class EAllowEditsMode : uint8;
+enum class EKeyGroupMode : uint8;
 enum class EMovieSceneKeyInterpolation : uint8;
 
 
@@ -121,10 +122,10 @@ public:
 	/** Gets the multicast delegate which is run whenever the allow edits mode is changed. */
 	FOnAllowEditsModeChanged& GetOnAllowEditsModeChanged() { return OnAllowEditsModeChangedEvent; }
 
-	/** Gets whether or not key all is enabled. */
-	bool GetKeyAllEnabled() const;
-	/** Sets whether or not key all is enabled. */
-	void SetKeyAllEnabled(bool InbKeyAllEnabled);
+	/** Returns what channels will get keyed when one channel changes */
+	EKeyGroupMode GetKeyGroupMode() const;
+	/** Sets which channels are keyed when a channel is keyed */
+	void SetKeyGroupMode(EKeyGroupMode);
 
 	/** Gets whether or not to key interp properties only. */
 	bool GetKeyInterpPropertiesOnly() const;
@@ -345,9 +346,9 @@ protected:
 	UPROPERTY( config, EditAnywhere, Category=Keyframing )
 	EAllowEditsMode AllowEditsMode;
 
-	/** Enable or disable keying all channels when any are keyed. */
-	UPROPERTY( config, EditAnywhere, Category=Keyframing )
-	bool bKeyAllEnabled;
+	/**Key group mode. */
+	UPROPERTY(config, EditAnywhere, Category = Keyframing)
+	EKeyGroupMode KeyGroupMode;
 
 	/** Enable or disable only keyframing properties marked with the 'Interp' keyword. */
 	UPROPERTY( config, EditAnywhere, Category=Keyframing )
diff --git a/Engine/Source/Editor/Sequencer/Public/SequencerTimeSliderController.h b/Engine/Source/Editor/Sequencer/Public/SequencerTimeSliderController.h
new file mode 100644
--- /dev/null
+++ b/Engine/Source/Editor/Sequencer/Public/SequencerTimeSliderController.h
@@ -0,0 +1,337 @@
+// Copyright 1998-2018 Epic Games, Inc. All Rights Reserved.
+
+#pragma once
+
+#include "CoreMinimal.h"
+#include "Input/CursorReply.h"
+#include "Input/Reply.h"
+#include "Widgets/SWidget.h"
+#include "Editor/SequencerWidgets/Public/ITimeSlider.h"
+#include "ISequencerModule.h"
+
+class FSlateWindowElementList;
+struct FContextMenuSuppressor;
+struct FSlateBrush;
+class FSequencer;
+
+struct FPaintPlaybackRangeArgs
+{
+	FPaintPlaybackRangeArgs()
+		: StartBrush(nullptr), EndBrush(nullptr), BrushWidth(0.f), SolidFillOpacity(0.f)
+	{}
+
+	FPaintPlaybackRangeArgs(const FSlateBrush* InStartBrush, const FSlateBrush* InEndBrush, float InBrushWidth)
+		: StartBrush(InStartBrush), EndBrush(InEndBrush), BrushWidth(InBrushWidth)
+	{}
+	/** Brush to use for the start bound */
+	const FSlateBrush* StartBrush;
+	/** Brush to use for the end bound */
+	const FSlateBrush* EndBrush;
+	/** The width of the above brushes, in slate units */
+	float BrushWidth;
+	/** level of opacity for the fill color between the range markers */
+	float SolidFillOpacity;
+};
+
+struct FPaintSectionAreaViewArgs
+{
+	FPaintSectionAreaViewArgs()
+		: bDisplayTickLines(false), bDisplayScrubPosition(false)
+	{}
+
+	/** Whether to display tick lines */
+	bool bDisplayTickLines;
+	/** Whether to display the scrub position */
+	bool bDisplayScrubPosition;
+	/** Optional Paint args for the playback range*/
+	TOptional<FPaintPlaybackRangeArgs> PlaybackRangeArgs;
+};
+
+/**
+ * A time slider controller for sequencer
+ * Draws and manages time data for a Sequencer
+ */
+class FSequencerTimeSliderController : public ITimeSliderController
+{
+public:
+	FSequencerTimeSliderController( const FTimeSliderArgs& InArgs, TWeakPtr<FSequencer> InWeakSequencer );
+
+	/**
+	* Determines the optimal spacing between tick marks in the slider for a given pixel density
+	* Increments until a minimum amount of slate units specified by MinTick is reached
+	*
+	* @param InPixelsPerInput	The density of pixels between each input
+	* @param MinTick			The minimum slate units per tick allowed
+	* @param MinTickSpacing	The minimum tick spacing in time units allowed
+	* @return the optimal spacing in time units
+	*/
+	float DetermineOptimalSpacing(float InPixelsPerInput, uint32 MinTick, float MinTickSpacing) const;
+
+	/** ITimeSliderController Interface */
+	virtual int32 OnPaintTimeSlider( bool bMirrorLabels, const FGeometry& AllottedGeometry, const FSlateRect& MyCullingRect, FSlateWindowElementList& OutDrawElements, int32 LayerId, const FWidgetStyle& InWidgetStyle, bool bParentEnabled ) const override;
+	virtual FReply OnMouseButtonDown( SWidget& WidgetOwner, const FGeometry& MyGeometry, const FPointerEvent& MouseEvent ) override;
+	virtual FReply OnMouseButtonUp( SWidget& WidgetOwner, const FGeometry& MyGeometry, const FPointerEvent& MouseEvent ) override;
+	virtual FReply OnMouseMove( SWidget& WidgetOwner, const FGeometry& MyGeometry, const FPointerEvent& MouseEvent ) override;
+	virtual FReply OnMouseWheel( SWidget& WidgetOwner, const FGeometry& MyGeometry, const FPointerEvent& MouseEvent ) override;
+	virtual FCursorReply OnCursorQuery( TSharedRef<const SWidget> WidgetOwner, const FGeometry& MyGeometry, const FPointerEvent& CursorEvent ) const override;
+	/** End ITimeSliderController Interface */
+
+	/** Get the current play rate for this controller */
+	virtual FFrameRate GetDisplayRate() const override { return TimeSliderArgs.DisplayRate.Get(); }
+
+	/** Get the current tick resolution for this controller */
+	virtual FFrameRate GetTickResolution() const override { return TimeSliderArgs.TickResolution.Get(); }
+
+	/** Get the current view range for this controller */
+	virtual FAnimatedRange GetViewRange() const override { return TimeSliderArgs.ViewRange.Get(); }
+
+	/** Get the current clamp range for this controller in seconds. */
+	virtual FAnimatedRange GetClampRange() const override { return TimeSliderArgs.ClampRange.Get(); }
+
+	/** Get the current play range for this controller */
+	virtual TRange<FFrameNumber> GetPlayRange() const override { return TimeSliderArgs.PlaybackRange.Get(TRange<FFrameNumber>()); }
+
+	/**
+	 * Clamp the given range to the clamp range 
+	 *
+	 * @param NewRangeMin		The new lower bound of the range
+	 * @param NewRangeMax		The new upper bound of the range
+	 */	
+	void ClampViewRange(double& NewRangeMin, double& NewRangeMax);
+
+	/**
+	 * Set a new range based on a min, max and an interpolation mode
+	 * 
+	 * @param NewRangeMin		The new lower bound of the range
+	 * @param NewRangeMax		The new upper bound of the range
+	 * @param Interpolation		How to set the new range (either immediately, or animated)
+	 */
+	virtual void SetViewRange( double NewRangeMin, double NewRangeMax, EViewRangeInterpolation Interpolation ) override;
+
+	/**
+	 * Set a new clamp range based on a min, max
+	 * 
+	 * @param NewRangeMin		The new lower bound of the clamp range
+	 * @param NewRangeMax		The new upper bound of the clamp range
+	 */
+	virtual void SetClampRange( double NewRangeMin, double NewRangeMax ) override;
+
+	/**
+	 * Set a new playback range based on a min, max
+	 * 
+	 * @param RangeStart		The new lower bound of the playback range
+	 * @param RangeDuration		The total number of frames that we play for
+	 */
+	virtual void SetPlayRange( FFrameNumber RangeStart, int32 RangeDuration ) override;
+
+	/**
+	 * Zoom the range by a given delta.
+	 * 
+	 * @param InDelta		The total amount to zoom by (+ve = zoom out, -ve = zoom in)
+	 * @param ZoomBias		Bias to apply to lower/upper extents of the range. (0 = lower, 0.5 = equal, 1 = upper)
+	 */
+	bool ZoomByDelta( float InDelta, float ZoomBias = 0.5f );
+
+	/**
+	 * Pan the range by a given delta
+	 * 
+	 * @param InDelta		The total amount to pan by (+ve = pan forwards in time, -ve = pan backwards in time)
+	 */
+	void PanByDelta( float InDelta );
+
+	/**
+	 * Draws major tick lines in the section view                                                              
+	 */
+	int32 OnPaintSectionView( const FGeometry& AllottedGeometry, const FSlateRect& MyCullingRect, FSlateWindowElementList& OutDrawElements, int32 LayerId, bool bEnabled, const FPaintSectionAreaViewArgs& Args ) const;
+
+public:
+
+	struct FScrubberMetrics
+	{
+		/** The extents of the current frame that the scrubber is on, in pixels */
+		TRange<float> FrameExtentsPx;
+		/** The pixel range that the scrubber handle (thumb) occupies */
+		TRange<float> HandleRangePx;
+		/** The style of the scrubber handle */
+		ESequencerScrubberStyle Style;
+		/** The style of the scrubber handle */
+		bool bDrawExtents;
+	};
+
+	/** Utility struct for converting between scrub range space and local/absolute screen space */
+	struct FScrubRangeToScreen
+	{
+		double ViewStart;
+		float  PixelsPerInput;
+
+		FScrubRangeToScreen(const TRange<double>& InViewInput, const FVector2D& InWidgetSize)
+		{
+			float ViewInputRange = InViewInput.Size<double>();
+
+			ViewStart = InViewInput.GetLowerBoundValue();
+			PixelsPerInput = ViewInputRange > 0 ? (InWidgetSize.X / ViewInputRange) : 0;
+		}
+
+		/** Local Widget Space -> Curve Input domain. */
+		double LocalXToInput(float ScreenX) const
+		{
+			return PixelsPerInput > 0 ? (ScreenX / PixelsPerInput) + ViewStart : ViewStart;
+		}
+
+		/** Curve Input domain -> local Widget Space */
+		float InputToLocalX(double Input) const
+		{
+			return (Input - ViewStart) * PixelsPerInput;
+		}
+	};
+
+private:
+	// forward declared as class members to prevent name collision with similar types defined in other units
+	struct FDrawTickArgs;
+
+	/**
+	 * Call this method when the user's interaction has changed the scrub position
+	 *
+	 * @param NewValue				Value resulting from the user's interaction
+	 * @param bIsScrubbing			True if done via scrubbing, false if just releasing scrubbing
+	 */
+	void CommitScrubPosition( FFrameTime NewValue, bool bIsScrubbing );
+
+	/**
+	 * Draw time tick marks
+	 *
+	 * @param OutDrawElements	List to add draw elements to
+	 * @param ViewRange			The currently visible time range in seconds
+	 * @param RangeToScreen		Time range to screen space converter
+	 * @param InArgs			Parameters for drawing the tick lines
+	 */
+	void DrawTicks( FSlateWindowElementList& OutDrawElements, const TRange<double>& ViewRange, const FScrubRangeToScreen& RangeToScreen, FDrawTickArgs& InArgs ) const;
+
+	/**
+	 * Draw the selection range.
+	 *
+	 * @return The new layer ID.
+	 */
+	int32 DrawSelectionRange(const FGeometry& AllottedGeometry, const FSlateRect& MyCullingRect, FSlateWindowElementList& OutDrawElements, int32 LayerId, const FScrubRangeToScreen& RangeToScreen, const FPaintPlaybackRangeArgs& Args) const;
+
+	/**
+	 * Draw the playback range.
+	 *
+	 * @return the new layer ID
+	 */
+	int32 DrawPlaybackRange(const FGeometry& AllottedGeometry, const FSlateRect& MyCullingRect, FSlateWindowElementList& OutDrawElements, int32 LayerId, const FScrubRangeToScreen& RangeToScreen, const FPaintPlaybackRangeArgs& Args) const;
+
+	/**
+	 * Draw the playback range.
+	 *
+	 * @return the new layer ID
+	 */
+	int32 DrawSubSequenceRange(const FGeometry& AllottedGeometry, const FSlateRect& MyCullingRect, FSlateWindowElementList& OutDrawElements, int32 LayerId, const FScrubRangeToScreen& RangeToScreen, const FPaintPlaybackRangeArgs& Args) const;
+
+private:
+
+	/**
+	 * Hit test the lower bound of a range
+	 */
+	bool HitTestRangeStart(const FScrubRangeToScreen& RangeToScreen, const TRange<double>& Range, float HitPixel) const;
+
+	/**
+	 * Hit test the upper bound of a range
+	 */
+	bool HitTestRangeEnd(const FScrubRangeToScreen& RangeToScreen, const TRange<double>& Range, float HitPixel) const;
+
+	FFrameTime SnapTimeToNearestKey(const FScrubRangeToScreen& RangeToScreen, float CursorPos, FFrameTime InTime) const;
+
+	void SetPlaybackRangeStart(FFrameNumber NewStart);
+	void SetPlaybackRangeEnd(FFrameNumber NewEnd);
+
+	void SetSelectionRangeStart(FFrameNumber NewStart);
+	void SetSelectionRangeEnd(FFrameNumber NewEnd);
+
+	TSharedRef<SWidget> OpenSetPlaybackRangeMenu(FFrameNumber FrameNumber);
+	FFrameTime ComputeScrubTimeFromMouse(const FGeometry& Geometry, FVector2D ScreenSpacePosition, FScrubRangeToScreen RangeToScreen) const;
+	FFrameTime ComputeFrameTimeFromMouse(const FGeometry& Geometry, FVector2D ScreenSpacePosition, FScrubRangeToScreen RangeToScreen, bool CheckSnapping = true) const;
+
+private:
+
+	/**
+	 * Get the pixel matrics of the Scrubber
+	 * @param ScrubTime			The qualified time of the scrubber
+	 * @param RangeToScreen		Range to screen helper
+	 * @param DilationPixels	Number of pixels to dilate the handle by
+	 * return FScrubberMetrics struct
+	 */
+	FScrubberMetrics GetScrubPixelMetrics(const FQualifiedFrameTime& ScrubTime, const FScrubRangeToScreen& RangeToScreen, float DilationPixels = 0.f) const;
+
+	FScrubberMetrics GetHitTestScrubPixelMetrics(const FScrubRangeToScreen& RangeToScreen) const;
+
+private:
+
+	/** Pointer back to the sequencer object */
+	TWeakPtr<FSequencer> WeakSequencer;
+
+	FTimeSliderArgs TimeSliderArgs;
+
+	/** Brush for drawingthe fill area on the scrubber */
+	const FSlateBrush* ScrubFillBrush;
+
+	/** Brush for drawing an upwards facing scrub handles */
+	const FSlateBrush* FrameBlockScrubHandleUpBrush, *VanillaScrubHandleUpBrush;
+	
+	/** Brush for drawing a downwards facing scrub handle */
+	const FSlateBrush* FrameBlockScrubHandleDownBrush, *VanillaScrubHandleDownBrush;
+	
+	/** Total mouse delta during dragging **/
+	float DistanceDragged;
+	
+	/** If we are dragging a scrubber or dragging to set the time range */
+	enum DragType
+	{
+		DRAG_SCRUBBING_TIME,
+		DRAG_SETTING_RANGE,
+		DRAG_PLAYBACK_START,
+		DRAG_PLAYBACK_END,
+		DRAG_SELECTION_START,
+		DRAG_SELECTION_END,
+		DRAG_NONE
+	};
+	
+	DragType MouseDragType;
+	
+	/** If we are currently panning the panel */
+	bool bPanning;
+
+	/** Mouse down position range */
+	FVector2D MouseDownPosition[2];
+
+	/** Geometry on mouse down */
+	FGeometry MouseDownGeometry;
+
+	/** Range stack */
+	TArray<TRange<double>> ViewRangeStack;
+
+	/** When > 0, we should not show context menus */
+	int32 ContextMenuSuppression;
+	
+	friend FContextMenuSuppressor;
+	
+};
+
+struct FContextMenuSuppressor
+{
+	FContextMenuSuppressor(TSharedRef<FSequencerTimeSliderController> InTimeSliderController)
+		: TimeSliderController(InTimeSliderController)
+	{
+		++TimeSliderController->ContextMenuSuppression;
+	}
+	~FContextMenuSuppressor()
+	{
+		--TimeSliderController->ContextMenuSuppression;
+	}
+
+private:
+	FContextMenuSuppressor(const FContextMenuSuppressor&);
+	FContextMenuSuppressor& operator=(const FContextMenuSuppressor&);
+
+	TSharedRef<FSequencerTimeSliderController> TimeSliderController;
+};
diff --git a/Engine/Source/Editor/SequencerWidgets/Private/STimeRange.cpp b/Engine/Source/Editor/SequencerWidgets/Private/STimeRange.cpp
--- a/Engine/Source/Editor/SequencerWidgets/Private/STimeRange.cpp
+++ b/Engine/Source/Editor/SequencerWidgets/Private/STimeRange.cpp
@@ -52,7 +52,7 @@ void STimeRange::Construct( const STimeRange::FArguments& InArgs, TSharedRef<ITi
 	{
 		ViewRangeStart = SNew(SSpinBox<double>)
 		.Value(this, &STimeRange::ViewStartTime)
-		.ToolTipText(ViewStartTimeTooltip())
+		.ToolTipText(LOCTEXT("ViewStartTimeTooltip", "View Range Start Time"))
 		.OnValueCommitted( this, &STimeRange::OnViewStartTimeCommitted )
 		.OnValueChanged( this, &STimeRange::OnViewStartTimeChanged )
 		.MinValue(TOptional<double>())
@@ -66,7 +66,7 @@ void STimeRange::Construct( const STimeRange::FArguments& InArgs, TSharedRef<ITi
 
 		ViewRangeEnd = SNew(SSpinBox<double>)
 		.Value(this, &STimeRange::ViewEndTime)
-		.ToolTipText(ViewEndTimeTooltip())
+		.ToolTipText(LOCTEXT("ViewEndTimeTooltip", "View Range End Time"))
 		.OnValueCommitted( this, &STimeRange::OnViewEndTimeCommitted )
 		.OnValueChanged( this, &STimeRange::OnViewEndTimeChanged )
 		.MinValue(this, &STimeRange::MinViewEndTime)
@@ -82,7 +82,7 @@ void STimeRange::Construct( const STimeRange::FArguments& InArgs, TSharedRef<ITi
 	{
 		PlaybackRangeStart = SNew(SSpinBox<double>)
 		.Value(this, &STimeRange::PlayStartTime)
-		.ToolTipText(PlayStartTimeTooltip())
+		.ToolTipText(LOCTEXT("PlayStartTimeTooltip", "In Time"))
 		.OnValueCommitted(this, &STimeRange::OnPlayStartTimeCommitted)
 		.OnValueChanged(this, &STimeRange::OnPlayStartTimeChanged)
 		.MinValue(TOptional<double>())
@@ -96,7 +96,7 @@ void STimeRange::Construct( const STimeRange::FArguments& InArgs, TSharedRef<ITi
 
 		PlaybackRangeEnd = SNew(SSpinBox<double>)
 		.Value(this, &STimeRange::PlayEndTime)
-		.ToolTipText(PlayEndTimeTooltip())
+		.ToolTipText(LOCTEXT("PlayEndTimeTooltip", "Out Time"))
 		.OnValueCommitted( this, &STimeRange::OnPlayEndTimeCommitted )
 		.OnValueChanged( this, &STimeRange::OnPlayEndTimeChanged )
 		.MinValue(TOptional<double>())
@@ -221,39 +221,39 @@ void STimeRange::Construct( const STimeRange::FArguments& InArgs, TSharedRef<ITi
 
 double STimeRange::WorkingStartTime() const
 {
-	FFrameRate Rate = TimeSliderController->GetFrameResolution();
+	FFrameRate Rate = TimeSliderController->GetTickResolution();
 	FFrameTime Time = TimeSliderController->GetClampRange().GetLowerBoundValue() * Rate;
 	return Time.GetFrame().Value;
 }
 
 double STimeRange::WorkingEndTime() const
 {
-	FFrameRate Rate = TimeSliderController->GetFrameResolution();
+	FFrameRate Rate = TimeSliderController->GetTickResolution();
 	FFrameTime Time = TimeSliderController->GetClampRange().GetUpperBoundValue() * Rate;
 	return Time.GetFrame().Value;
 }
 
 double STimeRange::ViewStartTime() const
 {
-	FFrameRate FrameResolution = TimeSliderController->GetFrameResolution();
+	FFrameRate TickResolution = TimeSliderController->GetTickResolution();
 
-	// View range is in seconds so we convert it to frame resolution
-	FFrameTime Time = TimeSliderController->GetViewRange().GetLowerBoundValue() * FrameResolution;
+	// View range is in seconds so we convert it to tick resolution
+	FFrameTime Time = TimeSliderController->GetViewRange().GetLowerBoundValue() * TickResolution;
 	return Time.GetFrame().Value;
 }
 
 double STimeRange::ViewEndTime() const
 {
-	FFrameRate FrameResolution = TimeSliderController->GetFrameResolution();
+	FFrameRate TickResolution = TimeSliderController->GetTickResolution();
 
-	// View range is in seconds so we convert it to frame resolution
-	FFrameTime Time = TimeSliderController->GetViewRange().GetUpperBoundValue() * FrameResolution;
+	// View range is in seconds so we convert it to tick resolution
+	FFrameTime Time = TimeSliderController->GetViewRange().GetUpperBoundValue() * TickResolution;
 	return Time.GetFrame().Value;
 }
 
 double STimeRange::GetSpinboxDelta() const
 {
-	return TimeSliderController->GetFrameResolution().AsDecimal() * TimeSliderController->GetPlayRate().AsInterval();
+	return TimeSliderController->GetTickResolution().AsDecimal() * TimeSliderController->GetDisplayRate().AsInterval();
 }
 
 double STimeRange::PlayStartTime() const
@@ -340,9 +340,8 @@ void STimeRange::OnPlayEndTimeCommitted(double NewValue, ETextCommit::Type InTex
 
 void STimeRange::OnWorkingStartTimeChanged(double NewValue)
 {
-	FFrameRate PlayRate = TimeSliderController->GetPlayRate();
-	FFrameRate FrameResolution = TimeSliderController->GetFrameResolution();
-	double Time = FrameResolution.AsSeconds(FFrameTime::FromDecimal(NewValue));
+	FFrameRate TickResolution = TimeSliderController->GetTickResolution();
+	double Time = TickResolution.AsSeconds(FFrameTime::FromDecimal(NewValue));
 
 	// Clamp range is in seconds
 	TimeSliderController->SetClampRange(Time, TimeSliderController->GetClampRange().GetUpperBoundValue());
@@ -355,9 +354,8 @@ void STimeRange::OnWorkingStartTimeChanged(double NewValue)
 
 void STimeRange::OnWorkingEndTimeChanged(double NewValue)
 {
-	FFrameRate PlayRate = TimeSliderController->GetPlayRate();
-	FFrameRate FrameResolution = TimeSliderController->GetFrameResolution();
-	double Time = FrameResolution.AsSeconds(FFrameTime::FromDecimal(NewValue));
+	FFrameRate TickResolution = TimeSliderController->GetTickResolution();
+	double Time = TickResolution.AsSeconds(FFrameTime::FromDecimal(NewValue));
 
 	// Clamp range is in seconds
 	TimeSliderController->SetClampRange(TimeSliderController->GetClampRange().GetLowerBoundValue(), Time);
@@ -370,8 +368,8 @@ void STimeRange::OnWorkingEndTimeChanged(double NewValue)
 
 void STimeRange::OnViewStartTimeChanged(double NewValue)
 {
-	FFrameRate FrameResolution = TimeSliderController->GetFrameResolution();
-	double Time = FrameResolution.AsSeconds(FFrameTime::FromDecimal(NewValue));
+	FFrameRate TickResolution = TimeSliderController->GetTickResolution();
+	double Time = TickResolution.AsSeconds(FFrameTime::FromDecimal(NewValue));
 
 	if (Time < TimeSliderController.Get()->GetClampRange().GetLowerBoundValue())
 	{
@@ -383,8 +381,8 @@ void STimeRange::OnViewStartTimeChanged(double NewValue)
 
 void STimeRange::OnViewEndTimeChanged(double NewValue)
 {
-	FFrameRate FrameResolution = TimeSliderController->GetFrameResolution();
-	double Time = FrameResolution.AsSeconds(FFrameTime::FromDecimal(NewValue));
+	FFrameRate TickResolution = TimeSliderController->GetTickResolution();
+	double Time = TickResolution.AsSeconds(FFrameTime::FromDecimal(NewValue));
 
 	if (Time > TimeSliderController->GetClampRange().GetUpperBoundValue())
 	{
@@ -400,9 +398,9 @@ void STimeRange::OnPlayStartTimeChanged(double NewValue)
 	// UI adjustment to work in reasonable deltas, so instead we clamp it here.
 	NewValue = FMath::Clamp(NewValue, MinPlayStartTime().GetValue(), MaxPlayStartTime().GetValue());
 
-	FFrameRate FrameResolution = TimeSliderController->GetFrameResolution();
+	FFrameRate TickResolution = TimeSliderController->GetTickResolution();
 	FFrameTime Time = FFrameTime::FromDecimal(NewValue);
-	double     TimeInSeconds = FrameResolution.AsSeconds(Time);
+	double     TimeInSeconds = TickResolution.AsSeconds(Time);
 
 	if (TimeInSeconds < TimeSliderController.Get()->GetClampRange().GetLowerBoundValue())
 	{
@@ -419,9 +417,9 @@ void STimeRange::OnPlayEndTimeChanged(double NewValue)
 	// UI adjustment to work in reasonable deltas, so instead we clamp it here.
 	NewValue = FMath::Clamp(NewValue, MinPlayEndTime().GetValue(), MaxPlayEndTime().GetValue());
 
-	FFrameRate FrameResolution = TimeSliderController->GetFrameResolution();
+	FFrameRate TickResolution = TimeSliderController->GetTickResolution();
 	FFrameTime Time = FFrameTime::FromDecimal(NewValue);
-	double     TimeInSeconds = FrameResolution.AsSeconds(Time);
+	double     TimeInSeconds = TickResolution.AsSeconds(Time);
 
 	if (TimeInSeconds > TimeSliderController->GetClampRange().GetUpperBoundValue())
 	{
@@ -434,65 +432,4 @@ void STimeRange::OnPlayEndTimeChanged(double NewValue)
 	TimeSliderController->SetPlayRange(StartFrame, PlayDuration);
 }
 
-FText STimeRange::PlayStartTimeTooltip() const
-{
-	// @todo sequencer-timecode We need a NumericTypeInterface independent way to get the current type of display.
-	// bool bShowFrameNumbers = ShowFrameNumbers.IsBound() ? ShowFrameNumbers.Get() : false;
-	// if (bShowFrameNumbers)
-	// {
-	// 	return LOCTEXT("PlayStartFrameTooltip", "In Frame");
-	// }
-	// else
-	// {
-	// 	return LOCTEXT("PlayStartTimeTooltip", "In Time");
-	// }
-	return LOCTEXT("PlayStartTimeTooltip", "In Time");
-}
-
-FText STimeRange::PlayEndTimeTooltip() const
-{
-	// @todo sequencer-timecode We need a NumericTypeInterface independent way to get the current type of display.
-	// bool bShowFrameNumbers = ShowFrameNumbers.IsBound() ? ShowFrameNumbers.Get() : false;
-	// if (bShowFrameNumbers)
-	// {
-	// 	return LOCTEXT("PlayEndFrameTooltip", "Out Frame");
-	// }
-	// else
-	// {
-	// 	return LOCTEXT("PlayEndTimeTooltip", "Out Time");
-	// }
-	return LOCTEXT("PlayEndTimeTooltip", "Out Time");
-}
-
-FText STimeRange::ViewStartTimeTooltip() const
-{
-	// @todo sequencer-timecode We need a NumericTypeInterface independent way to get the current type of display.
-	// bool bShowFrameNumbers = ShowFrameNumbers.IsBound() ? ShowFrameNumbers.Get() : false;
-	// if (bShowFrameNumbers)
-	// {
-	// 	return LOCTEXT("ViewStartFrameTooltip", "View Range Start Frame");
-	// }
-	// else
-	// {
-	// 	return LOCTEXT("ViewStartTimeTooltip", "View Range Start Time");
-	// }
-	return LOCTEXT("ViewStartTimeTooltip", "View Range Start Time");
-
-}
-
-FText STimeRange::ViewEndTimeTooltip() const
-{
-	// @todo sequencer-timecode We need a NumericTypeInterface independent way to get the current type of display.
-	// bool bShowFrameNumbers = ShowFrameNumbers.IsBound() ? ShowFrameNumbers.Get() : false;
-	// if (bShowFrameNumbers)
-	// {
-	// 	return LOCTEXT("ViewEndFrameTooltip", "View Range End Frame");
-	// }
-	// else
-	// {
-	// 	return LOCTEXT("ViewEndTimeTooltip", "View Range End Time");
-	// }
-	return LOCTEXT("ViewEndTimeTooltip", "View Range End Time");
-}
-
 #undef LOCTEXT_NAMESPACE
diff --git a/Engine/Source/Editor/SequencerWidgets/Private/STimeRange.h b/Engine/Source/Editor/SequencerWidgets/Private/STimeRange.h
--- a/Engine/Source/Editor/SequencerWidgets/Private/STimeRange.h
+++ b/Engine/Source/Editor/SequencerWidgets/Private/STimeRange.h
@@ -51,14 +51,11 @@ protected:
 	void OnPlayStartTimeChanged(double NewValue);
 	void OnPlayEndTimeChanged(double NewValue);
 
-	FText PlayStartTimeTooltip() const;
-	FText PlayEndTimeTooltip() const;
-
 protected:
 
 	double ViewStartTime() const;
 	double ViewEndTime() const;
-
+	
 	TOptional<double> MaxViewStartTime() const;
 	TOptional<double> MinViewEndTime() const;
 
@@ -68,9 +65,6 @@ protected:
 	void OnViewStartTimeChanged(double NewValue);
 	void OnViewEndTimeChanged(double NewValue);
 
-	FText ViewStartTimeTooltip() const;
-	FText ViewEndTimeTooltip() const;
-
 protected:
 
 	double WorkingStartTime() const;
diff --git a/Engine/Source/Editor/SequencerWidgets/Private/STimeRangeSlider.cpp b/Engine/Source/Editor/SequencerWidgets/Private/STimeRangeSlider.cpp
--- a/Engine/Source/Editor/SequencerWidgets/Private/STimeRangeSlider.cpp
+++ b/Engine/Source/Editor/SequencerWidgets/Private/STimeRangeSlider.cpp
@@ -192,11 +192,26 @@ FReply STimeRangeSlider::OnMouseMove( const FGeometry& MyGeometry, const FPointe
 		}
 		else if (bLeftHandleDragged || bRightHandleDragged)
 		{
-			double InDelta  = bLeftHandleDragged  ? DragDelta : 0.0;
-			double OutDelta = bRightHandleDragged ? DragDelta : 0.0;
+			double NewIn = 0.0f;
+			double NewOut = 0.0f;
 
-			double NewIn  = MouseDownViewRange.GetLowerBoundValue() + InDelta;
-			double NewOut = MouseDownViewRange.GetUpperBoundValue() + OutDelta;
+			if (bLeftHandleDragged)
+			{
+				NewIn = MouseDownViewRange.GetLowerBoundValue() + DragDelta;
+				NewOut = MouseDownViewRange.GetUpperBoundValue() - DragDelta;
+			}
+			else
+			{
+				NewIn = MouseDownViewRange.GetLowerBoundValue() - DragDelta;
+				NewOut = MouseDownViewRange.GetUpperBoundValue() + DragDelta;
+			}
+
+			// In cases of extreme zoom the drag delta will be greater than the difference between In/Out.
+			// This causes zooming to then become pan at extreme levels which is undesirable.
+			if (NewIn >= NewOut)
+			{
+				return FReply::Handled();
+			}
 
 			TimeSliderControllerPtr->SetViewRange(NewIn, NewOut, EViewRangeInterpolation::Immediate);
 		}
diff --git a/Engine/Source/Editor/SequencerWidgets/Public/ITimeSlider.h b/Engine/Source/Editor/SequencerWidgets/Public/ITimeSlider.h
--- a/Engine/Source/Editor/SequencerWidgets/Public/ITimeSlider.h
+++ b/Engine/Source/Editor/SequencerWidgets/Public/ITimeSlider.h
@@ -28,7 +28,7 @@ DECLARE_DELEGATE_TwoParams( FOnScrubPositionChanged, FFrameTime, bool )
 DECLARE_DELEGATE_TwoParams( FOnViewRangeChanged, TRange<double>, EViewRangeInterpolation )
 DECLARE_DELEGATE_OneParam( FOnTimeRangeChanged, TRange<double> )
 DECLARE_DELEGATE_OneParam( FOnFrameRangeChanged, TRange<FFrameNumber> )
-DECLARE_DELEGATE_RetVal_OneParam( FFrameNumber, FOnGetNearestKey, FFrameTime )
+DECLARE_DELEGATE_RetVal_TwoParams( FFrameNumber, FOnGetNearestKey, FFrameTime, bool )
 
 /** Structure used to wrap up a range, and an optional animation target */
 struct FAnimatedRange : public TRange<double>
@@ -107,11 +107,11 @@ struct FTimeSliderArgs
 	/** Attribute defining the playback range for this controller */
 	TAttribute<TRange<FFrameNumber>> PlaybackRange;
 
-	/** Attribute for the current sequence's playback rate */
-	TAttribute<FFrameRate> PlayRate;
+	/** Attribute for the current sequence's display rate */
+	TAttribute<FFrameRate> DisplayRate;
 
-	/** Attribute for the current sequence's frame resolution */
-	TAttribute<FFrameRate> FrameResolution;
+	/** Attribute for the current sequence's tick resolution */
+	TAttribute<FFrameRate> TickResolution;
 
 	/** Delegate that is called when the playback range wants to change */
 	FOnFrameRangeChanged OnPlaybackRangeChanged;
@@ -163,10 +163,10 @@ public:
 	virtual FCursorReply OnCursorQuery( TSharedRef<const SWidget> WidgetOwner, const FGeometry& MyGeometry, const FPointerEvent& CursorEvent ) const = 0;
 
 	/** Get the current play rate for this controller */
-	virtual FFrameRate GetPlayRate() const = 0;
+	virtual FFrameRate GetDisplayRate() const = 0;
 
-	/** Get the current frame resolution for this controller */
-	virtual FFrameRate GetFrameResolution() const = 0;
+	/** Get the current tick resolution for this controller */
+	virtual FFrameRate GetTickResolution() const = 0;
 
 	/** Get the current view range for this controller */
 	virtual FAnimatedRange GetViewRange() const { return FAnimatedRange(); }
diff --git a/Engine/Source/Editor/SequencerWidgets/SequencerWidgets.Build.cs b/Engine/Source/Editor/SequencerWidgets/SequencerWidgets.Build.cs
--- a/Engine/Source/Editor/SequencerWidgets/SequencerWidgets.Build.cs
+++ b/Engine/Source/Editor/SequencerWidgets/SequencerWidgets.Build.cs
@@ -31,7 +31,8 @@ public class SequencerWidgets : ModuleRules
 				"Slate",
 				"SlateCore",
                 "MovieScene",
-                "UnrealEd"
+                "UnrealEd",
+				"TimeManagement"
 			}
 		);
 	}
diff --git a/Engine/Source/Editor/UMGEditor/Private/Animation/Sequencer2DTransformTrackEditor.cpp b/Engine/Source/Editor/UMGEditor/Private/Animation/Sequencer2DTransformTrackEditor.cpp
--- a/Engine/Source/Editor/UMGEditor/Private/Animation/Sequencer2DTransformTrackEditor.cpp
+++ b/Engine/Source/Editor/UMGEditor/Private/Animation/Sequencer2DTransformTrackEditor.cpp
@@ -1,6 +1,7 @@
 // Copyright 1998-2018 Epic Games, Inc. All Rights Reserved.
 
 #include "Animation/Sequencer2DTransformTrackEditor.h"
+#include "Animation/Sequencer2DTransformSection.h"
 #include "Slate/WidgetTransform.h"
 #include "ISectionLayoutBuilder.h"
 
@@ -18,6 +19,13 @@ TSharedRef<ISequencerTrackEditor> F2DTransformTrackEditor::CreateTrackEditor( TS
 }
 
 
+TSharedRef<ISequencerSection> F2DTransformTrackEditor::MakeSectionInterface(UMovieSceneSection& SectionObject, UMovieSceneTrack& Track, FGuid ObjectBinding)
+{
+	check(SupportsType(SectionObject.GetOuter()->GetClass()));
+	return MakeShared<F2DTransformSection>(SectionObject, GetSequencer());
+}
+
+
 void F2DTransformTrackEditor::GenerateKeysFromPropertyChanged( const FPropertyChangedParams& PropertyChangedParams, FGeneratedTrackKeys& OutGeneratedKeys )
 {
 	FPropertyPath StructPath = PropertyChangedParams.StructPathToKey;
diff --git a/Engine/Source/Editor/UMGEditor/Private/Animation/Sequencer2DTransformTrackEditor.h b/Engine/Source/Editor/UMGEditor/Private/Animation/Sequencer2DTransformTrackEditor.h
--- a/Engine/Source/Editor/UMGEditor/Private/Animation/Sequencer2DTransformTrackEditor.h
+++ b/Engine/Source/Editor/UMGEditor/Private/Animation/Sequencer2DTransformTrackEditor.h
@@ -44,6 +44,10 @@ public:
 
 protected:
 
+	// ISequencerTrackEditor interface
+
+	virtual TSharedRef<ISequencerSection> MakeSectionInterface( UMovieSceneSection& SectionObject, UMovieSceneTrack& Track, FGuid ObjectBinding ) override;
+
 	// FPropertyTrackEditor interface
 
 	virtual void GenerateKeysFromPropertyChanged( const FPropertyChangedParams& PropertyChangedParams, FGeneratedTrackKeys& OutGeneratedKeys ) override;
diff --git a/Engine/Source/Editor/UMGEditor/Private/TabFactory/AnimationTabSummoner.cpp b/Engine/Source/Editor/UMGEditor/Private/TabFactory/AnimationTabSummoner.cpp
--- a/Engine/Source/Editor/UMGEditor/Private/TabFactory/AnimationTabSummoner.cpp
+++ b/Engine/Source/Editor/UMGEditor/Private/TabFactory/AnimationTabSummoner.cpp
@@ -403,10 +403,10 @@ private:
 		NewAnimation->MovieScene = NewObject<UMovieScene>(NewAnimation, NewAnimation->GetFName(), RF_Transactional);
 
 		// Default to 20 fps display rate (as was the previous default in USequencerSettings)
-		NewAnimation->MovieScene->SetPlaybackFrameRate(FFrameRate(20, 1));
+		NewAnimation->MovieScene->SetDisplayRate(FFrameRate(20, 1));
 
-		const FFrameTime InFrame  = InTime  * NewAnimation->MovieScene->GetFrameResolution();
-		const FFrameTime OutFrame = OutTime * NewAnimation->MovieScene->GetFrameResolution();
+		const FFrameTime InFrame  = InTime  * NewAnimation->MovieScene->GetTickResolution();
+		const FFrameTime OutFrame = OutTime * NewAnimation->MovieScene->GetTickResolution();
 		NewAnimation->MovieScene->SetPlaybackRange(TRange<FFrameNumber>(InFrame.FrameNumber, OutFrame.FrameNumber+1));
 		NewAnimation->MovieScene->GetEditorData().WorkStart = InTime;
 		NewAnimation->MovieScene->GetEditorData().WorkEnd   = OutTime;
diff --git a/Engine/Source/Editor/UMGEditor/Private/WidgetBlueprintEditor.cpp b/Engine/Source/Editor/UMGEditor/Private/WidgetBlueprintEditor.cpp
--- a/Engine/Source/Editor/UMGEditor/Private/WidgetBlueprintEditor.cpp
+++ b/Engine/Source/Editor/UMGEditor/Private/WidgetBlueprintEditor.cpp
@@ -680,9 +680,9 @@ TSharedPtr<ISequencer>& FWidgetBlueprintEditor::GetSequencer()
 		FSequencerInitParams SequencerInitParams;
 		{
 			UWidgetAnimation* NullAnimation = UWidgetAnimation::GetNullAnimation();
-			FFrameRate FrameResolution = NullAnimation->MovieScene->GetFrameResolution();
-			FFrameNumber StartFrame = (InTime  * FrameResolution).FloorToFrame();
-			FFrameNumber EndFrame   = (OutTime * FrameResolution).CeilToFrame();
+			FFrameRate TickResolution = NullAnimation->MovieScene->GetTickResolution();
+			FFrameNumber StartFrame = (InTime  * TickResolution).FloorToFrame();
+			FFrameNumber EndFrame   = (OutTime * TickResolution).CeilToFrame();
 			NullAnimation->MovieScene->SetPlaybackRange(StartFrame, (EndFrame-StartFrame).Value);
 			FMovieSceneEditorData& EditorData = NullAnimation->MovieScene->GetEditorData();
 			EditorData.WorkStart = InTime;
diff --git a/Engine/Source/Editor/UMGEditor/UMGEditor.Build.cs b/Engine/Source/Editor/UMGEditor/UMGEditor.Build.cs
--- a/Engine/Source/Editor/UMGEditor/UMGEditor.Build.cs
+++ b/Engine/Source/Editor/UMGEditor/UMGEditor.Build.cs
@@ -56,7 +56,8 @@ public class UMGEditor : ModuleRules
 				"DetailCustomizations",
                 "Settings",
 				"RenderCore",
-                "TargetPlatform"
+                "TargetPlatform",
+				"TimeManagement"
 			}
 			);
 	}
diff --git a/Engine/Source/Editor/UnrealEd/Private/EditorEngine.cpp b/Engine/Source/Editor/UnrealEd/Private/EditorEngine.cpp
--- a/Engine/Source/Editor/UnrealEd/Private/EditorEngine.cpp
+++ b/Engine/Source/Editor/UnrealEd/Private/EditorEngine.cpp
@@ -88,6 +88,7 @@
 #include "Interfaces/IPluginManager.h"
 #include "UObject/PackageReload.h"
 #include "HAL/PlatformApplicationMisc.h"
+#include "IMediaModule.h"
 
 // needed for the RemotePropagator
 #include "AudioDevice.h"
@@ -1484,7 +1485,6 @@ void UEditorEngine::Tick( float DeltaSeconds, bool bIdleMode )
 		}
 	}
 
-
 	// Perform editor level streaming previs if no PIE session is currently in progress.
 	if( !PlayWorld )
 	{
@@ -1694,6 +1694,15 @@ void UEditorEngine::Tick( float DeltaSeconds, bool bIdleMode )
 		FTickableGameObject::TickObjects(nullptr, TickType, false, DeltaSeconds);
 	}
 
+	// tick media framework
+	static const FName MediaModuleName(TEXT("Media"));
+	IMediaModule* MediaModule = FModuleManager::LoadModulePtr<IMediaModule>(MediaModuleName);
+
+	if (MediaModule != nullptr)
+	{
+		MediaModule->TickPostEngine();
+	}
+
 	if (bFirstTick)
 	{
 		bFirstTick = false;
diff --git a/Engine/Source/Editor/UnrealEd/Private/Fbx/FbxMainExport.cpp b/Engine/Source/Editor/UnrealEd/Private/Fbx/FbxMainExport.cpp
--- a/Engine/Source/Editor/UnrealEd/Private/Fbx/FbxMainExport.cpp
+++ b/Engine/Source/Editor/UnrealEd/Private/Fbx/FbxMainExport.cpp
@@ -1486,20 +1486,20 @@ FFbxExporter::FLevelSequenceAnimTrackAdapter::FLevelSequenceAnimTrackAdapter( IM
 
 float FFbxExporter::FLevelSequenceAnimTrackAdapter::GetAnimationStart() const
 {
-	return MovieScene->GetPlaybackRange().GetLowerBoundValue() / MovieScene->GetFrameResolution();
+	return MovieScene->GetPlaybackRange().GetLowerBoundValue() / MovieScene->GetTickResolution();
 }
 
 float FFbxExporter::FLevelSequenceAnimTrackAdapter::GetAnimationLength() const
 {
-	return FFrameNumber(MovieScene::DiscreteSize(MovieScene->GetPlaybackRange())) / MovieScene->GetFrameResolution();
+	return FFrameNumber(MovieScene::DiscreteSize(MovieScene->GetPlaybackRange())) / MovieScene->GetTickResolution();
 }
 
 
 void FFbxExporter::FLevelSequenceAnimTrackAdapter::UpdateAnimation( float Time )
 {
 	// Convert the time in seconds to a frame time
-	FFrameTime FrameTime = Time * MovieScene->GetFrameResolution();
-	FMovieSceneContext Context = FMovieSceneContext(FMovieSceneEvaluationRange(FrameTime, MovieScene->GetFrameResolution()), MovieScenePlayer->GetPlaybackStatus()).SetHasJumped(true);
+	FFrameTime FrameTime = Time * MovieScene->GetTickResolution();
+	FMovieSceneContext Context = FMovieSceneContext(FMovieSceneEvaluationRange(FrameTime, MovieScene->GetTickResolution()), MovieScenePlayer->GetPlaybackStatus()).SetHasJumped(true);
 	MovieScenePlayer->GetEvaluationTemplate().Evaluate( Context, *MovieScenePlayer );
 }
 
@@ -2328,7 +2328,7 @@ void RichCurveInterpolationToFbxInterpolation(ERichCurveInterpMode InInterpolati
 	}
 }
 
-void FFbxExporter::ExportChannelToFbxCurve(FbxAnimCurve& InFbxCurve, const FMovieSceneFloatChannel& InChannel, FFrameRate FrameResolution, ERichCurveValueMode ValueMode, bool bNegative)
+void FFbxExporter::ExportChannelToFbxCurve(FbxAnimCurve& InFbxCurve, const FMovieSceneFloatChannel& InChannel, FFrameRate TickResolution, ERichCurveValueMode ValueMode, bool bNegative)
 {
 	const float NegateFactor = bNegative ? -1.f : 1.f;
 
@@ -2346,7 +2346,7 @@ void FFbxExporter::ExportChannelToFbxCurve(FbxAnimCurve& InFbxCurve, const FMovi
 
 		FbxTime FbxTime;
 		FbxAnimCurveKey FbxKey;
-		FbxTime.SetSecondDouble(KeyTime / FrameResolution);
+		FbxTime.SetSecondDouble(KeyTime / TickResolution);
 
 		int FbxKeyIndex = InFbxCurve.KeyAdd(FbxTime);
 
@@ -2359,8 +2359,8 @@ void FFbxExporter::ExportChannelToFbxCurve(FbxAnimCurve& InFbxCurve, const FMovi
 		{
 			if (Index < Times.Num() - 1)
 			{
-				float LeaveTangent = KeyValue.Tangent.LeaveTangent;
-				float NextArriveTangent = Values[Index+1].Tangent.ArriveTangent;
+				float LeaveTangent = KeyValue.Tangent.LeaveTangent * TickResolution.AsDecimal();
+				float NextArriveTangent = Values[Index+1].Tangent.ArriveTangent * TickResolution.AsDecimal();
 
 				if (bNegative)
 				{
@@ -2401,7 +2401,7 @@ void FFbxExporter::ExportLevelSequence3DTransformTrack( FbxNode& FbxActor, UMovi
 		return;
 	}
 
-	FFrameRate FrameResolution = TransformTrack.GetTypedOuter<UMovieScene>()->GetFrameResolution();
+	FFrameRate TickResolution = TransformTrack.GetTypedOuter<UMovieScene>()->GetTickResolution();
 
 	FbxAnimCurveNode* TranslationNode = FbxActor.LclTranslation.GetCurveNode( BaseLayer, true );
 	FbxAnimCurveNode* RotationNode = FbxActor.LclRotation.GetCurveNode( BaseLayer, true );
@@ -2422,24 +2422,24 @@ void FFbxExporter::ExportLevelSequence3DTransformTrack( FbxNode& FbxActor, UMovi
 	TArrayView<FMovieSceneFloatChannel*> FloatChannels = TransformSection->GetChannelProxy().GetChannels<FMovieSceneFloatChannel>();
 
 	// Translation
-	ExportChannelToFbxCurve(*FbxCurveTransX, *FloatChannels[0], FrameResolution, ERichCurveValueMode::Default, false);
-	ExportChannelToFbxCurve(*FbxCurveTransY, *FloatChannels[1], FrameResolution, ERichCurveValueMode::Default, true);
-	ExportChannelToFbxCurve(*FbxCurveTransZ, *FloatChannels[2], FrameResolution, ERichCurveValueMode::Default, false);
+	ExportChannelToFbxCurve(*FbxCurveTransX, *FloatChannels[0], TickResolution, ERichCurveValueMode::Default, false);
+	ExportChannelToFbxCurve(*FbxCurveTransY, *FloatChannels[1], TickResolution, ERichCurveValueMode::Default, true);
+	ExportChannelToFbxCurve(*FbxCurveTransZ, *FloatChannels[2], TickResolution, ERichCurveValueMode::Default, false);
 
 	// Scale - don't generate scale keys for cameras
 	if (!bIsCameraActor)
 	{
-		ExportChannelToFbxCurve(*FbxCurveScaleX, *FloatChannels[3], FrameResolution, ERichCurveValueMode::Default, false);
-		ExportChannelToFbxCurve(*FbxCurveScaleY, *FloatChannels[4], FrameResolution, ERichCurveValueMode::Default, false);
-		ExportChannelToFbxCurve(*FbxCurveScaleZ, *FloatChannels[5], FrameResolution, ERichCurveValueMode::Default, false);
+		ExportChannelToFbxCurve(*FbxCurveScaleX, *FloatChannels[6], TickResolution, ERichCurveValueMode::Default, false);
+		ExportChannelToFbxCurve(*FbxCurveScaleY, *FloatChannels[7], TickResolution, ERichCurveValueMode::Default, false);
+		ExportChannelToFbxCurve(*FbxCurveScaleZ, *FloatChannels[8], TickResolution, ERichCurveValueMode::Default, false);
 	}
 
 	// Rotation - bake rotation for cameras and lights
 	if (!bBakeRotations)
 	{
-		ExportChannelToFbxCurve(*FbxCurveRotX, *FloatChannels[6], FrameResolution, ERichCurveValueMode::Default, false);
-		ExportChannelToFbxCurve(*FbxCurveRotY, *FloatChannels[7], FrameResolution, ERichCurveValueMode::Default, true);
-		ExportChannelToFbxCurve(*FbxCurveRotZ, *FloatChannels[8], FrameResolution, ERichCurveValueMode::Default, true);
+		ExportChannelToFbxCurve(*FbxCurveRotX, *FloatChannels[3], TickResolution, ERichCurveValueMode::Default, false);
+		ExportChannelToFbxCurve(*FbxCurveRotY, *FloatChannels[4], TickResolution, ERichCurveValueMode::Default, true);
+		ExportChannelToFbxCurve(*FbxCurveRotZ, *FloatChannels[5], TickResolution, ERichCurveValueMode::Default, true);
 	}
 	else
 	{
@@ -2460,12 +2460,12 @@ void FFbxExporter::ExportLevelSequence3DTransformTrack( FbxNode& FbxActor, UMovi
 		FbxCurveRotZ->KeyModifyBegin();
 
 		FFrameNumber PlayStart      = MovieScene::DiscreteInclusiveLower(InPlaybackRange);
-		double       PlayDuration   = FFrameTime(MovieScene::DiscreteSize(InPlaybackRange)) / FrameResolution;
+		double       PlayDuration   = FFrameTime(MovieScene::DiscreteSize(InPlaybackRange)) / TickResolution;
 		int32        NumKeys        = PlayDuration * BakeTransformsFPS;
 
 		for ( int KeyIndex = 0; KeyIndex < NumKeys; KeyIndex++ )
 		{
-			FFrameTime RelativeKeyTime = ( KeyIndex * PlayDuration / NumKeys ) * FrameResolution;
+			FFrameTime RelativeKeyTime = ( KeyIndex * PlayDuration / NumKeys ) * TickResolution;
 			FFrameTime KeyTime         = RelativeKeyTime + PlayStart;
 
 			FVector Trans = FVector::ZeroVector;
@@ -2495,7 +2495,7 @@ void FFbxExporter::ExportLevelSequence3DTransformTrack( FbxNode& FbxActor, UMovi
 			FbxVector4 KeyScale = Converter.ConvertToFbxScale(RelativeTransform.GetScale3D());
 
 			FbxTime FbxTime;
-			FbxTime.SetSecondDouble(KeyTime / FrameResolution);
+			FbxTime.SetSecondDouble(KeyTime / TickResolution);
 
 			FbxCurveRotX->KeySet(FbxCurveRotX->KeyAdd(FbxTime), FbxTime, KeyRot[0]);
 			FbxCurveRotY->KeySet(FbxCurveRotY->KeyAdd(FbxTime), FbxTime, KeyRot[1]);
@@ -2579,8 +2579,8 @@ void FFbxExporter::ExportLevelSequenceFloatTrack( FbxNode& FbxActor, UMovieScene
 	CurveNode->SetChannelValue<double>( 0U, FloatChannel->GetDefault().Get(MAX_flt) );
 	CurveNode->ConnectToChannel( AnimCurve, 0U );
 
-	FFrameRate FrameResolution = FloatTrack.GetTypedOuter<UMovieScene>()->GetFrameResolution();
-	ExportChannelToFbxCurve(*AnimCurve, *FloatChannel, FrameResolution, IsFoV ? ERichCurveValueMode::Fov : ERichCurveValueMode::Default);
+	FFrameRate TickResolution = FloatTrack.GetTypedOuter<UMovieScene>()->GetTickResolution();
+	ExportChannelToFbxCurve(*AnimCurve, *FloatChannel, TickResolution, IsFoV ? ERichCurveValueMode::Fov : ERichCurveValueMode::Default);
 }
 		
 /**
diff --git a/Engine/Source/Editor/UnrealEd/Private/FbxExporter.h b/Engine/Source/Editor/UnrealEd/Private/FbxExporter.h
--- a/Engine/Source/Editor/UnrealEd/Private/FbxExporter.h
+++ b/Engine/Source/Editor/UnrealEd/Private/FbxExporter.h
@@ -431,7 +431,7 @@ private:
 	};
 
 	/** Exports a movie scene channel to an fbx animation curve. */
-	void ExportChannelToFbxCurve(FbxAnimCurve& InFbxCurve, const FMovieSceneFloatChannel& InChannel, FFrameRate FrameResolution, ERichCurveValueMode ValueMode = ERichCurveValueMode::Default, bool bNegative = false);
+	void ExportChannelToFbxCurve(FbxAnimCurve& InFbxCurve, const FMovieSceneFloatChannel& InChannel, FFrameRate TickResolution, ERichCurveValueMode ValueMode = ERichCurveValueMode::Default, bool bNegative = false);
 
 	/**
 	 * Finds the given actor in the already-exported list of structures
diff --git a/Engine/Source/Editor/UnrealEd/Private/LevelEditorViewport.cpp b/Engine/Source/Editor/UnrealEd/Private/LevelEditorViewport.cpp
--- a/Engine/Source/Editor/UnrealEd/Private/LevelEditorViewport.cpp
+++ b/Engine/Source/Editor/UnrealEd/Private/LevelEditorViewport.cpp
@@ -2200,6 +2200,13 @@ void FLevelEditorViewportClient::UpdateViewForLockedActor(float DeltaTime)
 					}
 				}
 			}
+
+			const float DistanceToCurrentLookAt = FVector::Dist( GetViewLocation() , GetLookAtLocation() );
+
+			const FQuat CameraOrientation = FQuat::MakeFromEuler( GetViewRotation().Euler() );
+			FVector Direction = CameraOrientation.RotateVector( FVector(1,0,0) );
+
+			SetLookAtLocation( GetViewLocation() + Direction * DistanceToCurrentLookAt );
 		}
 	}
 }
diff --git a/Engine/Source/Editor/UnrealEd/UnrealEd.Build.cs b/Engine/Source/Editor/UnrealEd/UnrealEd.Build.cs
--- a/Engine/Source/Editor/UnrealEd/UnrealEd.Build.cs
+++ b/Engine/Source/Editor/UnrealEd/UnrealEd.Build.cs
@@ -62,6 +62,7 @@ public class UnrealEd : ModuleRules
 				"PhysicsAssetEditor",
 				"ClothingSystemEditorInterface",
 				"NavigationSystem",
+				"Media",
 			}
 		);
 
@@ -160,7 +161,8 @@ public class UnrealEd : ModuleRules
 				"ClothingSystemRuntime",
 				"ClothingSystemRuntimeInterface",
 				"PIEPreviewDeviceProfileSelector",
-            }
+				"TimeManagement",
+			}
 		);
 
 		DynamicallyLoadedModuleNames.AddRange(
@@ -221,6 +223,7 @@ public class UnrealEd : ModuleRules
 				"OverlayEditor",
 				"AnimationModifiers",
 				"ClothPainter",
+				"Media",
 			}
 		);
 
diff --git a/Engine/Source/Runtime/AVIWriter/Private/AVIWriter.cpp b/Engine/Source/Runtime/AVIWriter/Private/AVIWriter.cpp
--- a/Engine/Source/Runtime/AVIWriter/Private/AVIWriter.cpp
+++ b/Engine/Source/Runtime/AVIWriter/Private/AVIWriter.cpp
@@ -276,7 +276,24 @@ public:
 			}
 
 			IAMVideoCompression* CompressionImpl = nullptr;
-			if(EncodingFilter && SUCCEEDED(EncodingFilter->QueryInterface(IID_IAMVideoCompression, (void **)&CompressionImpl)))
+			if (EncodingFilter)
+			{
+				IEnumPins* pEnum = nullptr;
+				IPin* pPin = nullptr;
+				EncodingFilter->EnumPins(&pEnum);
+				while (S_OK == pEnum->Next(1, &pPin, NULL))
+				{
+					hr = pPin->QueryInterface(IID_IAMVideoCompression, (void**)&CompressionImpl);
+					pPin->Release();
+					if (SUCCEEDED(hr)) // Found the interface.
+					{
+						break;
+					}
+				}
+				pEnum->Release();
+			}
+
+			if (CompressionImpl)
 			{
 				CompressionImpl->put_Quality(Options.CompressionQuality.GetValue());
 				CompressionImpl->Release();
diff --git a/Engine/Source/Runtime/CinematicCamera/Private/CameraRig_Rail.cpp b/Engine/Source/Runtime/CinematicCamera/Private/CameraRig_Rail.cpp
--- a/Engine/Source/Runtime/CinematicCamera/Private/CameraRig_Rail.cpp
+++ b/Engine/Source/Runtime/CinematicCamera/Private/CameraRig_Rail.cpp
@@ -17,6 +17,7 @@ ACameraRig_Rail::ACameraRig_Rail(const FObjectInitializer& ObjectInitializer)
 	PrimaryActorTick.bStartWithTickEnabled = true;
 
 	CurrentPositionOnRail = 0.f;
+	bLockOrientationToRail = false;
 
 	// Create components
 	TransformComponent = CreateDefaultSubobject<USceneComponent>(TEXT("TransformComponent"));
@@ -154,6 +155,12 @@ void ACameraRig_Rail::UpdateRailComponents()
 		float const SplineLen = RailSplineComponent->GetSplineLength();
 		FVector const MountPos = RailSplineComponent->GetLocationAtDistanceAlongSpline(CurrentPositionOnRail*SplineLen, ESplineCoordinateSpace::World);
 		RailCameraMount->SetWorldLocation(MountPos);
+
+		if (bLockOrientationToRail)
+		{
+			FQuat const RailRot = RailSplineComponent->GetQuaternionAtDistanceAlongSpline(CurrentPositionOnRail*SplineLen, ESplineCoordinateSpace::World);
+			RailCameraMount->SetWorldRotation(RailRot);
+		}
 	}
 
 #if WITH_EDITOR
diff --git a/Engine/Source/Runtime/CinematicCamera/Public/CameraRig_Rail.h b/Engine/Source/Runtime/CinematicCamera/Public/CameraRig_Rail.h
--- a/Engine/Source/Runtime/CinematicCamera/Public/CameraRig_Rail.h
+++ b/Engine/Source/Runtime/CinematicCamera/Public/CameraRig_Rail.h
@@ -31,6 +31,10 @@ public:
 	UPROPERTY(Interp, EditAnywhere, BlueprintReadWrite, Category = "Rail Controls", meta=(ClampMin="0.0", ClampMax = "1.0"))
 	float CurrentPositionOnRail;
 
+	/** Determines whether the orientation of the mount should be in the direction of the rail. */
+	UPROPERTY(Interp, EditAnywhere, BlueprintReadWrite, Category = "Rail Controls")
+	bool bLockOrientationToRail;
+
 	virtual class USceneComponent* GetDefaultAttachComponent() const override;
 #if WITH_EDITOR
 	virtual void PostEditChangeProperty(FPropertyChangedEvent& PropertyChangedEvent) override;
diff --git a/Engine/Source/Runtime/CinematicCamera/Public/CineCameraActor.h b/Engine/Source/Runtime/CinematicCamera/Public/CineCameraActor.h
--- a/Engine/Source/Runtime/CinematicCamera/Public/CineCameraActor.h
+++ b/Engine/Source/Runtime/CinematicCamera/Public/CineCameraActor.h
@@ -17,7 +17,7 @@ struct FCameraLookatTrackingSettings
 	GENERATED_USTRUCT_BODY()
 
 	/** True to enable lookat tracking, false otherwise. */
-	UPROPERTY(/*Interp, */EditAnywhere, BlueprintReadWrite, Category = "LookAt")
+	UPROPERTY(Interp, EditAnywhere, BlueprintReadWrite, Category = "LookAt")
 	uint8 bEnableLookAtTracking : 1;
 
 	/** True to draw a debug representation of the lookat location */
diff --git a/Engine/Source/Runtime/Core/Private/Misc/Timespan.cpp b/Engine/Source/Runtime/Core/Private/Misc/Timespan.cpp
--- a/Engine/Source/Runtime/Core/Private/Misc/Timespan.cpp
+++ b/Engine/Source/Runtime/Core/Private/Misc/Timespan.cpp
@@ -81,6 +81,7 @@ FString FTimespan::ToString(const TCHAR* Format) const
 			case TCHAR('s'): Result += FString::Printf(TEXT("%02i"), FMath::Abs(GetSeconds())); break;
 			case TCHAR('f'): Result += FString::Printf(TEXT("%03i"), FMath::Abs(GetFractionMilli())); break;
 			case TCHAR('u'): Result += FString::Printf(TEXT("%06i"), FMath::Abs(GetFractionMicro())); break;
+			case TCHAR('t'): Result += FString::Printf(TEXT("%07i"), FMath::Abs(GetFractionTicks())); break;
 			case TCHAR('n'): Result += FString::Printf(TEXT("%09i"), FMath::Abs(GetFractionNano())); break;
 			default:
 				Result += *Format;
diff --git a/Engine/Source/Runtime/Core/Private/Tests/Math/RangeSetTest.cpp b/Engine/Source/Runtime/Core/Private/Tests/Math/RangeSetTest.cpp
--- a/Engine/Source/Runtime/Core/Private/Tests/Math/RangeSetTest.cpp
+++ b/Engine/Source/Runtime/Core/Private/Tests/Math/RangeSetTest.cpp
@@ -3,12 +3,34 @@
 #include "CoreTypes.h"
 #include "Misc/AutomationTest.h"
 
+#include "Math/Range.h"
+#include "Math/RangeSet.h"
+#include "Misc/Timespan.h"
+
 #if WITH_DEV_AUTOMATION_TESTS
 
 IMPLEMENT_SIMPLE_AUTOMATION_TEST(FRangeSetTest, "System.Core.Math.RangeSet", EAutomationTestFlags::Disabled | EAutomationTestFlags::EditorContext | EAutomationTestFlags::ClientContext | EAutomationTestFlags::SmokeFilter)
 
-bool FRangeSetTest::RunTest( const FString& Parameters )
+bool FRangeSetTest::RunTest(const FString& Parameters)
 {
+	// @todo: Write an actual unit test (the following just checks compilation of certain use cases)
+
+	{
+		TRange<int32> Range = TRange<int32>::AtLeast(0);
+		TRangeSet<int32> RangeSet;
+
+		RangeSet.Add(Range);
+		int32 Value = RangeSet.GetMinBoundValue();
+	}
+
+	{
+		TRange<FTimespan> Range = TRange<FTimespan>::AtLeast(FTimespan::Zero());
+		TRangeSet<FTimespan> RangeSet;
+
+		RangeSet.Add(Range);
+		FTimespan Value = RangeSet.GetMinBoundValue();
+	}
+
 	return true;
 }
 
diff --git a/Engine/Source/Runtime/Core/Public/Math/RangeSet.h b/Engine/Source/Runtime/Core/Public/Math/RangeSet.h
--- a/Engine/Source/Runtime/Core/Public/Math/RangeSet.h
+++ b/Engine/Source/Runtime/Core/Public/Math/RangeSet.h
@@ -18,6 +18,9 @@ template<typename ElementType> class TRangeSet
 	typedef TRangeBound<ElementType> BoundsType;
 	typedef TRange<ElementType> RangeType;
 
+	/*~ Typedef used to pass/return small element types by value rather than const& */
+	typedef typename TCallTraits<ElementType>::ParamType ElementValueOrConstRef;
+
 public:
 
 	/** Default constructor. */
@@ -113,7 +116,7 @@ public:
 	 * @param Element The element to check.
 	 * @return true if the element is in the set, false otherwise.
 	 */
-	bool Contains(const ElementType& Element) const
+	bool Contains(ElementValueOrConstRef Element) const
 	{
 		for (typename TArray<RangeType>::TConstIterator It(Ranges); It; ++It)
 		{
@@ -171,7 +174,7 @@ public:
 	 * @return Bound value.
 	 * @see GetMaxBoundValue, GetMinBound, HasMinBound
 	 */
-	const ElementType& GetMinBoundValue() const
+	ElementType GetMinBoundValue() const
 	{
 		return GetMinBound().GetValue();
 	}
@@ -202,7 +205,7 @@ public:
 	 * @return Bound value.
 	 * @see GetMaxBound, GetMinBoundValue, HasMaxBound
 	 */
-	const ElementType& GetMaxBoundValue() const
+	ElementType GetMaxBoundValue() const
 	{
 		return GetMaxBound().GetValue();
 	}
diff --git a/Engine/Source/Runtime/Core/Public/Math/UnrealMathUtility.h b/Engine/Source/Runtime/Core/Public/Math/UnrealMathUtility.h
--- a/Engine/Source/Runtime/Core/Public/Math/UnrealMathUtility.h
+++ b/Engine/Source/Runtime/Core/Public/Math/UnrealMathUtility.h
@@ -299,7 +299,9 @@ struct FMath : public FPlatformMath
 	template <class T>
 	static FORCEINLINE T DivideAndRoundNearest(T Dividend, T Divisor)
 	{
-		return (Dividend > 0) ? 1 + ((Dividend - 1) / Divisor) : (Dividend / Divisor);
+		return (Dividend >= 0)
+			? (Dividend + Divisor / 2) / Divisor
+			: (Dividend - Divisor / 2 + 1) / Divisor;
 	}
 
 	/**
diff --git a/Engine/Source/Runtime/Core/Public/Misc/Timespan.h b/Engine/Source/Runtime/Core/Public/Misc/Timespan.h
--- a/Engine/Source/Runtime/Core/Public/Misc/Timespan.h
+++ b/Engine/Source/Runtime/Core/Public/Misc/Timespan.h
@@ -3,6 +3,7 @@
 #pragma once
 
 #include "CoreTypes.h"
+#include "Math/UnrealMathUtility.h"
 #include "Misc/AssertionMacros.h"
 #include "Math/Interval.h"
 
@@ -396,6 +397,16 @@ public:
 		return (int32)((Ticks % ETimespan::TicksPerSecond) * ETimespan::NanosecondsPerTick);
 	}
 
+	/**
+	 * Gets the fractional ticks (in 100 nanosecond resolution).
+	 *
+	 * @return Number of ticks in fractional part.
+	 */
+	int32 GetFractionTicks() const
+	{
+		return (int32)(Ticks % ETimespan::TicksPerSecond);
+	}
+
 	/**
 	 * Gets the hours component of this time span.
 	 *
@@ -590,7 +601,7 @@ public:
 	 */
 	static FTimespan FromDays(double Days)
 	{
-		return FTimespan(Days * ETimespan::TicksPerDay);
+		return FTimespan(FMath::RoundToZero(Days * ETimespan::TicksPerDay + 0.5));
 	}
 
 	/**
@@ -602,7 +613,7 @@ public:
 	 */
 	static FTimespan FromHours(double Hours)
 	{
-		return FTimespan(Hours * ETimespan::TicksPerHour);
+		return FTimespan(FMath::RoundToZero(Hours * ETimespan::TicksPerHour + 0.5));
 	}
 
 	/**
@@ -614,7 +625,7 @@ public:
 	 */
 	static FTimespan FromMicroseconds(double Microseconds)
 	{
-		return FTimespan(Microseconds * ETimespan::TicksPerMicrosecond);
+		return FTimespan(FMath::RoundToZero(Microseconds * ETimespan::TicksPerMicrosecond + 0.5));
 	}
 
 	/**
@@ -626,7 +637,7 @@ public:
 	 */
 	static FTimespan FromMilliseconds(double Milliseconds)
 	{
-		return FTimespan(Milliseconds * ETimespan::TicksPerMillisecond);
+		return FTimespan(FMath::RoundToZero(Milliseconds * ETimespan::TicksPerMillisecond + 0.5));
 	}
 
 	/**
@@ -638,7 +649,7 @@ public:
 	 */
 	static FTimespan FromMinutes(double Minutes)
 	{
-		return FTimespan(Minutes * ETimespan::TicksPerMinute);
+		return FTimespan(FMath::RoundToZero(Minutes * ETimespan::TicksPerMinute + 0.5));
 	}
 
 	/**
@@ -650,7 +661,7 @@ public:
 	 */
 	static FTimespan FromSeconds(double Seconds)
 	{
-		return FTimespan(Seconds * ETimespan::TicksPerSecond);
+		return FTimespan(FMath::RoundToZero(Seconds * ETimespan::TicksPerSecond + 0.5));
 	}
 
 	/**
diff --git a/Engine/Source/Runtime/Core/Public/UObject/SequencerObjectVersion.h b/Engine/Source/Runtime/Core/Public/UObject/SequencerObjectVersion.h
--- a/Engine/Source/Runtime/Core/Public/UObject/SequencerObjectVersion.h
+++ b/Engine/Source/Runtime/Core/Public/UObject/SequencerObjectVersion.h
@@ -33,6 +33,12 @@ struct CORE_API FSequencerObjectVersion
 		// When finished now defaults to project default
 		FloatToIntConversion,
 
+		// Purged old spawnable blueprint classes from level sequence assets
+		PurgeSpawnableBlueprints,
+
+		// Finish UMG evaluation on end
+		FinishUMGEvaluation,
+
 		// -----<new versions can be added above this line>-------------------------------------------------
 		VersionPlusOne,
 		LatestVersion = VersionPlusOne - 1
diff --git a/Engine/Source/Runtime/Engine/Engine.Build.cs b/Engine/Source/Runtime/Engine/Engine.Build.cs
--- a/Engine/Source/Runtime/Engine/Engine.Build.cs
+++ b/Engine/Source/Runtime/Engine/Engine.Build.cs
@@ -188,6 +188,7 @@ public class Engine : ModuleRules
 		{
 			PrivateIncludePathModuleNames.AddRange(
 				new string[] {
+					"Media",
 					"SlateNullRenderer",
 					"SlateRHIRenderer"
 				}
@@ -195,6 +196,7 @@ public class Engine : ModuleRules
 
 			DynamicallyLoadedModuleNames.AddRange(
 				new string[] {
+					"Media",
 					"SlateNullRenderer",
 					"SlateRHIRenderer"
 				}
diff --git a/Engine/Source/Runtime/Engine/Private/GameEngine.cpp b/Engine/Source/Runtime/Engine/Private/GameEngine.cpp
--- a/Engine/Source/Runtime/Engine/Private/GameEngine.cpp
+++ b/Engine/Source/Runtime/Engine/Private/GameEngine.cpp
@@ -61,6 +61,10 @@
 #include "IPIEPreviewDeviceModule.h"
 #endif
 
+#if !UE_SERVER
+	#include "IMediaModule.h"
+#endif
+
 CSV_DECLARE_CATEGORY_MODULE_EXTERN(CORE_API, Basic);
 
 ENGINE_API bool GDisallowNetworkTravel = false;
@@ -1342,6 +1346,17 @@ void UGameEngine::Tick( float DeltaSeconds, bool bIdleMode )
 		GWorld = GetWorldContextFromHandleChecked(OriginalGWorldContext).World();
 	}
 
+#if !UE_SERVER
+	// tick media framework
+	static const FName MediaModuleName(TEXT("Media"));
+	IMediaModule* MediaModule = FModuleManager::LoadModulePtr<IMediaModule>(MediaModuleName);
+
+	if (MediaModule != nullptr)
+	{
+		MediaModule->TickPostEngine();
+	}
+#endif
+
 	// Tick the viewport
 	if ( GameViewport != NULL && !bIdleMode )
 	{
diff --git a/Engine/Source/Runtime/Foliage/Public/InstancedFoliageActor.h b/Engine/Source/Runtime/Foliage/Public/InstancedFoliageActor.h
--- a/Engine/Source/Runtime/Foliage/Public/InstancedFoliageActor.h
+++ b/Engine/Source/Runtime/Foliage/Public/InstancedFoliageActor.h
@@ -22,7 +22,7 @@ struct FFoliageMeshInfo;
 // Function for filtering out hit components during FoliageTrace
 typedef TFunction<bool(const UPrimitiveComponent*)> FFoliageTraceFilterFunc;
 
-UCLASS(notplaceable, hidecategories = (Object, Rendering), MinimalAPI, NotBlueprintable)
+UCLASS(notplaceable, hidecategories = (Object, Rendering, Mobility), MinimalAPI, NotBlueprintable)
 class AInstancedFoliageActor : public AActor
 {
 	GENERATED_UCLASS_BODY()
diff --git a/Engine/Source/Runtime/Landscape/Classes/Landscape.h b/Engine/Source/Runtime/Landscape/Classes/Landscape.h
--- a/Engine/Source/Runtime/Landscape/Classes/Landscape.h
+++ b/Engine/Source/Runtime/Landscape/Classes/Landscape.h
@@ -24,7 +24,7 @@ enum ELandscapeSetupErrors
 	LSE_MAX,
 };
 
-UCLASS(MinimalAPI, showcategories=(Display, Movement, Collision, Lighting, LOD, Input))
+UCLASS(MinimalAPI, showcategories=(Display, Movement, Collision, Lighting, LOD, Input), hidecategories=(Mobility))
 class ALandscape : public ALandscapeProxy
 {
 	GENERATED_BODY()
diff --git a/Engine/Source/Runtime/Landscape/Classes/LandscapeComponent.h b/Engine/Source/Runtime/Landscape/Classes/LandscapeComponent.h
--- a/Engine/Source/Runtime/Landscape/Classes/LandscapeComponent.h
+++ b/Engine/Source/Runtime/Landscape/Classes/LandscapeComponent.h
@@ -235,7 +235,7 @@ struct FLandscapeComponentGrassData
 	friend FArchive& operator<<(FArchive& Ar, FLandscapeComponentGrassData& Data);
 };
 
-UCLASS(hidecategories=(Display, Attachment, Physics, Debug, Collision, Movement, Rendering, PrimitiveComponent, Object, Transform), showcategories=("Rendering|Material"), MinimalAPI, Within=LandscapeProxy)
+UCLASS(hidecategories=(Display, Attachment, Physics, Debug, Collision, Movement, Rendering, PrimitiveComponent, Object, Transform, Mobility), showcategories=("Rendering|Material"), MinimalAPI, Within=LandscapeProxy)
 class ULandscapeComponent : public UPrimitiveComponent
 {
 	GENERATED_UCLASS_BODY()
diff --git a/Engine/Source/Runtime/Landscape/Classes/LandscapeProxy.h b/Engine/Source/Runtime/Landscape/Classes/LandscapeProxy.h
--- a/Engine/Source/Runtime/Landscape/Classes/LandscapeProxy.h
+++ b/Engine/Source/Runtime/Landscape/Classes/LandscapeProxy.h
@@ -326,7 +326,7 @@ public:
 	~FAsyncGrassTask();
 };
 
-UCLASS(Abstract, MinimalAPI, NotBlueprintable, hidecategories=(Display, Attachment, Physics, Debug, Lighting, LOD), showcategories=(Lighting, Rendering, "Utilities|Transformation"))
+UCLASS(Abstract, MinimalAPI, NotBlueprintable, hidecategories=(Display, Attachment, Physics, Debug, Lighting, LOD), showcategories=(Lighting, Rendering, "Utilities|Transformation"), hidecategories=(Mobility))
 class ALandscapeProxy : public AActor
 {
 	GENERATED_BODY()
diff --git a/Engine/Source/Runtime/Launch/Launch.Build.cs b/Engine/Source/Runtime/Launch/Launch.Build.cs
--- a/Engine/Source/Runtime/Launch/Launch.Build.cs
+++ b/Engine/Source/Runtime/Launch/Launch.Build.cs
@@ -9,38 +9,32 @@ public class Launch : ModuleRules
 	{
 		PrivateIncludePaths.Add("Runtime/Launch/Private");
 
-		PrivateIncludePathModuleNames.AddRange(
-			new string[] {
-				"AutomationController",
-				"Media",
-                "MediaUtils",
-				"TaskGraph",
-			}
-		);
-
-		PrivateDependencyModuleNames.AddRange(
-			new string[] {
-				"Core",
-				"CoreUObject",
-				"Engine",
-				"InputCore",
-				"MoviePlayer",
-				"Networking",
-				"PakFile",
-				"Projects",
-				"RenderCore",
-				"RHI",
-				"SandboxFile",
-				"Serialization",
-				"ShaderCore",
-				"ApplicationCore",
-				"Slate",
-				"SlateCore",
-				"Sockets",
-                "Overlay",
-				"UtilityShaders",
-			}
-		);
+		PrivateIncludePathModuleNames.AddRange(new string[] {
+			"AutomationController",
+			"TaskGraph",
+		});
+
+		PrivateDependencyModuleNames.AddRange(new string[] {
+			"Core",
+			"CoreUObject",
+			"Engine",
+			"InputCore",
+			"MoviePlayer",
+			"Networking",
+			"PakFile",
+			"Projects",
+			"RenderCore",
+			"RHI",
+			"SandboxFile",
+			"Serialization",
+			"ShaderCore",
+			"ApplicationCore",
+			"Slate",
+			"SlateCore",
+			"Sockets",
+            "Overlay",
+			"UtilityShaders",
+		});
 
 		// Enable the LauncherCheck module to be used for platforms that support the Launcher.
 		// Projects should set Target.bUseLauncherChecks in their Target.cs to enable the functionality.
@@ -59,25 +53,28 @@ public class Launch : ModuleRules
 
 		if (Target.Type != TargetType.Server)
 		{
-			PrivateDependencyModuleNames.AddRange(
-				new string[] {
-					"HeadMountedDisplay",
-					"MRMesh",
-				}
-			);
+			PrivateDependencyModuleNames.AddRange(new string[] {
+				"HeadMountedDisplay",
+				"MediaUtils",
+				"MRMesh",
+			});
 
 			if ((Target.Platform == UnrealTargetPlatform.Win32) ||
 				(Target.Platform == UnrealTargetPlatform.Win64))
 			{
-				DynamicallyLoadedModuleNames.Add("D3D12RHI");
-				DynamicallyLoadedModuleNames.Add("D3D11RHI");
-				DynamicallyLoadedModuleNames.Add("XAudio2");
-				DynamicallyLoadedModuleNames.Add("AudioMixerXAudio2");
+				DynamicallyLoadedModuleNames.AddRange(new string[] {
+					"AudioMixerXAudio2",
+					"D3D11RHI",
+					"D3D12RHI",
+					"XAudio2",
+				});
 			}
 			else if (Target.Platform == UnrealTargetPlatform.Mac)
 			{
-				DynamicallyLoadedModuleNames.Add("CoreAudio");
-				DynamicallyLoadedModuleNames.Add("AudioMixerAudioUnit");
+				DynamicallyLoadedModuleNames.AddRange(new string[] {
+					"AudioMixerAudioUnit",
+					"CoreAudio",
+				});
 			}
 			else if (Target.IsInPlatformGroup(UnrealPlatformGroup.Unix))
 			{
@@ -85,65 +82,66 @@ public class Launch : ModuleRules
 				PrivateDependencyModuleNames.Add("Json");
 			}
 
-			PrivateIncludePathModuleNames.AddRange(
-				new string[] {
-					"SlateNullRenderer",
-					"SlateRHIRenderer"
-				}
-			);
-
-			DynamicallyLoadedModuleNames.AddRange(
-				new string[] {
-					"SlateNullRenderer",
-					"SlateRHIRenderer"
-				}
-			);
+			PrivateIncludePathModuleNames.AddRange(new string[] {
+				"Media",
+				"SlateNullRenderer",
+				"SlateRHIRenderer",
+			});
+
+			DynamicallyLoadedModuleNames.AddRange(new string[] {
+				"Media",
+				"SlateNullRenderer",
+				"SlateRHIRenderer",
+			});
 		}
 
 		// UFS clients are not available in shipping builds
 		if (Target.Configuration != UnrealTargetConfiguration.Shipping)
 		{
-			PrivateDependencyModuleNames.AddRange(
-				new string[] {
-					"NetworkFile",
-					"StreamingFile",
-					"CookedIterativeFile",
-					"AutomationWorker",
-				}
-			);
+			PrivateDependencyModuleNames.AddRange(new string[] {
+				"NetworkFile",
+				"StreamingFile",
+				"CookedIterativeFile",
+				"AutomationWorker",
+			});
 		}
 
-		DynamicallyLoadedModuleNames.AddRange(
-			new string[] {
-				"Media",
-				"Renderer",
-			}
-		);
+		DynamicallyLoadedModuleNames.AddRange(new string[] {
+			"Renderer",
+		});
 
 		if (Target.bCompileAgainstEngine)
 		{
-			PrivateIncludePathModuleNames.AddRange(
-				new string[] {
-					"MessagingCommon",
-					"MediaAssets"
-				}
-			);
+			PrivateIncludePathModuleNames.AddRange(new string[] {
+				"MessagingCommon",
+			});
 
 			PublicDependencyModuleNames.Add("SessionServices");
 			PrivateIncludePaths.Add("Developer/DerivedDataCache/Public");
 
 			// LaunchEngineLoop.cpp will still attempt to load XMPP but not all projects require it so it will silently fail unless referenced by the project's build.cs file.
 			// DynamicallyLoadedModuleNames.Add("XMPP");
-			DynamicallyLoadedModuleNames.Add("HTTP");
 
-			PrivateDependencyModuleNames.Add("ClothingSystemRuntimeInterface");
-			PrivateDependencyModuleNames.Add("ClothingSystemRuntime");
+			DynamicallyLoadedModuleNames.AddRange(new string[] {
+				"HTTP",
+				"MediaAssets",
+			});
+
+			PrivateDependencyModuleNames.AddRange(new string[] {
+				"ClothingSystemRuntime",
+				"ClothingSystemRuntimeInterface"
+			});
 		}
 
 		if (Target.Configuration != UnrealTargetConfiguration.Shipping)
 		{
 			PublicIncludePathModuleNames.Add("ProfilerService");
-			DynamicallyLoadedModuleNames.AddRange(new string[] { "TaskGraph", "RealtimeProfiler", "ProfilerService" });
+
+			DynamicallyLoadedModuleNames.AddRange(new string[] {
+				"TaskGraph",
+				"RealtimeProfiler",
+				"ProfilerService"
+			});
 		}
 
 		// The engine can use AutomationController in any connfiguration besides shipping.  This module is loaded
@@ -157,35 +155,33 @@ public class Launch : ModuleRules
 		{
 			PublicIncludePathModuleNames.Add("ProfilerClient");
 
-			PrivateDependencyModuleNames.AddRange(
-				new string[] {
-					"SourceControl",
-					"UnrealEd",
-					"DesktopPlatform",
-					"PIEPreviewDeviceProfileSelector",
-				}
-			);
+			PrivateDependencyModuleNames.AddRange(new string[] {
+				"SourceControl",
+				"UnrealEd",
+				"DesktopPlatform",
+				"PIEPreviewDeviceProfileSelector",
+			});
 
 
 			// ExtraModules that are loaded when WITH_EDITOR=1 is true
-			DynamicallyLoadedModuleNames.AddRange(
-				new string[] {
-					"AutomationWindow",
-					"ProfilerClient",
-					"Toolbox",
-					"GammaUI",
-					"ModuleUI",
-					"OutputLog",
-					"TextureCompressor",
-					"MeshUtilities",
-					"SourceCodeAccess"
-				}
-			);
+			DynamicallyLoadedModuleNames.AddRange(new string[] {
+				"AutomationWindow",
+				"ProfilerClient",
+				"Toolbox",
+				"GammaUI",
+				"ModuleUI",
+				"OutputLog",
+				"TextureCompressor",
+				"MeshUtilities",
+				"SourceCodeAccess"
+			});
 
 			if (Target.Platform == UnrealTargetPlatform.Mac)
 			{
-				PrivateDependencyModuleNames.Add("MainFrame");
-				PrivateDependencyModuleNames.Add("Settings");
+				PrivateDependencyModuleNames.AddRange(new string[] {
+					"MainFrame",
+					"Settings",
+				});
 			}
 			else
 			{
@@ -196,25 +192,35 @@ public class Launch : ModuleRules
 		if (Target.Platform == UnrealTargetPlatform.IOS ||
 			Target.Platform == UnrealTargetPlatform.TVOS)
 		{
-			PrivateDependencyModuleNames.Add("OpenGLDrv");
-			PrivateDependencyModuleNames.Add("IOSAudio");
-			PrivateDependencyModuleNames.Add("AudioMixerAudioUnit");
-			DynamicallyLoadedModuleNames.Add("IOSRuntimeSettings");
-			DynamicallyLoadedModuleNames.Add("IOSLocalNotification");
+			PrivateDependencyModuleNames.AddRange(new string[] {
+				"AudioMixerAudioUnit",
+				"IOSAudio",
+				"LaunchDaemonMessages",
+				"OpenGLDrv",
+			});
+
+			DynamicallyLoadedModuleNames.AddRange(new string[] {
+				"IOSLocalNotification",
+				"IOSRuntimeSettings",
+			});
+
 			PublicFrameworks.Add("OpenGLES");
 			// this is weak for IOS8 support for CAMetalLayer that is in QuartzCore
 			PublicWeakFrameworks.Add("QuartzCore");
-
-			PrivateDependencyModuleNames.Add("LaunchDaemonMessages");
 		}
 
 		if (Target.Platform == UnrealTargetPlatform.Android)
 		{
-			PrivateDependencyModuleNames.Add("OpenGLDrv");
-			PrivateDependencyModuleNames.Add("AndroidAudio");
-			PrivateDependencyModuleNames.Add("AudioMixerAndroid");
-			DynamicallyLoadedModuleNames.Add("AndroidRuntimeSettings");
-			DynamicallyLoadedModuleNames.Add("AndroidLocalNotification");
+			PrivateDependencyModuleNames.AddRange(new string[] {
+				"AndroidAudio",
+				"AudioMixerAndroid",
+				"OpenGLDrv",
+			});
+
+			DynamicallyLoadedModuleNames.AddRange(new string[] {
+				"AndroidLocalNotification",
+				"AndroidRuntimeSettings",
+			});
 		}
 
 		if ((Target.Platform == UnrealTargetPlatform.Win32) ||
@@ -240,14 +246,10 @@ public class Launch : ModuleRules
 
 		if (Target.IsInPlatformGroup(UnrealPlatformGroup.Unix))
 		{
-			PrivateDependencyModuleNames.AddRange(
-				new string[] {
-					"UnixCommonStartup"
-				}
-			);
+			PrivateDependencyModuleNames.Add("UnixCommonStartup");
 		}
 
-		if(Target.LinkType == TargetLinkType.Monolithic && !Target.bFormalBuild)
+		if (Target.LinkType == TargetLinkType.Monolithic && !Target.bFormalBuild)
 		{
 			PrivateDefinitions.Add(string.Format("COMPILED_IN_CL={0}", Target.Version.Changelist));
 			PrivateDefinitions.Add(string.Format("COMPILED_IN_COMPATIBLE_CL={0}", Target.Version.EffectiveCompatibleChangelist));
diff --git a/Engine/Source/Runtime/Launch/Private/LaunchEngineLoop.cpp b/Engine/Source/Runtime/Launch/Private/LaunchEngineLoop.cpp
--- a/Engine/Source/Runtime/Launch/Private/LaunchEngineLoop.cpp
+++ b/Engine/Source/Runtime/Launch/Private/LaunchEngineLoop.cpp
@@ -3350,7 +3350,6 @@ void FEngineLoop::Tick()
 
 		if (MediaModule != nullptr)
 		{
-			QUICK_SCOPE_CYCLE_COUNTER(STAT_FEngineLoop_MediaTickPreEngine);
 			MediaModule->TickPreEngine();
 		}
 #endif
@@ -3358,15 +3357,6 @@ void FEngineLoop::Tick()
 		// main game engine tick (world, game objects, etc.)
 		GEngine->Tick(FApp::GetDeltaTime(), bIdleMode);
 
-#if !UE_SERVER
-		// tick media framework
-		if (MediaModule != nullptr)
-		{
-			QUICK_SCOPE_CYCLE_COUNTER(STAT_FEngineLoop_MediaTickPostEngine);
-			MediaModule->TickPostEngine();
-		}
-#endif
-		
 		// If a movie that is blocking the game thread has been playing,
 		// wait for it to finish before we continue to tick or tick again
 		// We do this right after GEngine->Tick() because that is where user code would initiate a load / movie.
@@ -3392,7 +3382,6 @@ void FEngineLoop::Tick()
 		// tick media framework
 		if (MediaModule != nullptr)
 		{
-			QUICK_SCOPE_CYCLE_COUNTER(STAT_FEngineLoop_MediaTickPreSlate);
 			MediaModule->TickPreSlate();
 		}
 #endif
