Sequencer: Add checks for valid sequencer.

diff --git a/Engine/Source/Editor/Sequencer/Private/SequencerTimeSliderController.cpp b/Engine/Source/Editor/Sequencer/Private/SequencerTimeSliderController.cpp
--- a/Engine/Source/Editor/Sequencer/Private/SequencerTimeSliderController.cpp
+++ b/Engine/Source/Editor/Sequencer/Private/SequencerTimeSliderController.cpp
@@ -278,6 +278,12 @@ void FSequencerTimeSliderController::DrawTicks( FSlateWindowElementList& OutDraw
 
 int32 FSequencerTimeSliderController::OnPaintTimeSlider( bool bMirrorLabels, const FGeometry& AllottedGeometry, const FSlateRect& MyCullingRect, FSlateWindowElementList& OutDrawElements, int32 LayerId, const FWidgetStyle& InWidgetStyle, bool bParentEnabled ) const
 {
+	TSharedPtr<FSequencer> Sequencer = WeakSequencer.Pin();
+	if (!Sequencer.IsValid())
+	{
+		return LayerId;
+	}
+
 	const bool bEnabled = bParentEnabled;
 	const ESlateDrawEffect DrawEffects = bEnabled ? ESlateDrawEffect::None : ESlateDrawEffect::DisabledEffect;
 
@@ -414,6 +420,12 @@ int32 FSequencerTimeSliderController::OnPaintTimeSlider( bool bMirrorLabels, con
 
 int32 FSequencerTimeSliderController::DrawSelectionRange(const FGeometry& AllottedGeometry, const FSlateRect& MyCullingRect, FSlateWindowElementList& OutDrawElements, int32 LayerId, const FScrubRangeToScreen& RangeToScreen, const FPaintPlaybackRangeArgs& Args) const
 {
+	TSharedPtr<FSequencer> Sequencer = WeakSequencer.Pin();
+	if (!Sequencer.IsValid())
+	{
+		return LayerId;
+	}
+
 	TRange<double> SelectionRange = TimeSliderArgs.SelectionRange.Get() / GetTickResolution();
 
 	if (!SelectionRange.IsEmpty())
@@ -459,6 +471,12 @@ int32 FSequencerTimeSliderController::DrawSelectionRange(const FGeometry& Allott
 
 int32 FSequencerTimeSliderController::DrawPlaybackRange(const FGeometry& AllottedGeometry, const FSlateRect& MyCullingRect, FSlateWindowElementList& OutDrawElements, int32 LayerId, const FScrubRangeToScreen& RangeToScreen, const FPaintPlaybackRangeArgs& Args) const
 {
+	TSharedPtr<FSequencer> Sequencer = WeakSequencer.Pin();
+	if (!Sequencer.IsValid())
+	{
+		return LayerId;
+	}
+
 	if (!TimeSliderArgs.PlaybackRange.IsSet())
 	{
 		return LayerId;
@@ -513,6 +531,12 @@ int32 FSequencerTimeSliderController::DrawPlaybackRange(const FGeometry& Allotte
 
 int32 FSequencerTimeSliderController::DrawSubSequenceRange(const FGeometry& AllottedGeometry, const FSlateRect& MyCullingRect, FSlateWindowElementList& OutDrawElements, int32 LayerId, const FScrubRangeToScreen& RangeToScreen, const FPaintPlaybackRangeArgs& Args) const
 {
+	TSharedPtr<FSequencer> Sequencer = WeakSequencer.Pin();
+	if (!Sequencer.IsValid())
+	{
+		return LayerId;
+	}
+
 	TOptional<TRange<FFrameNumber>> RangeValue;
 	RangeValue = TimeSliderArgs.SubSequenceRange.Get(RangeValue);
 
@@ -718,6 +742,12 @@ FReply FSequencerTimeSliderController::OnMouseButtonUp( SWidget& WidgetOwner, co
 
 FReply FSequencerTimeSliderController::OnMouseMove( SWidget& WidgetOwner, const FGeometry& MyGeometry, const FPointerEvent& MouseEvent )
 {
+	TSharedPtr<FSequencer> Sequencer = WeakSequencer.Pin();
+	if (!Sequencer.IsValid())
+	{
+		return FReply::Unhandled();
+	}
+
 	bool bHandleLeftMouseButton  = MouseEvent.IsMouseButtonDown( EKeys::LeftMouseButton  );
 	bool bHandleRightMouseButton = MouseEvent.IsMouseButtonDown( EKeys::RightMouseButton ) && TimeSliderArgs.AllowZoom;
 
@@ -899,6 +929,12 @@ FReply FSequencerTimeSliderController::OnMouseWheel( SWidget& WidgetOwner, const
 
 FCursorReply FSequencerTimeSliderController::OnCursorQuery( TSharedRef<const SWidget> WidgetOwner, const FGeometry& MyGeometry, const FPointerEvent& CursorEvent ) const
 {
+	TSharedPtr<FSequencer> Sequencer = WeakSequencer.Pin();
+	if (!Sequencer.IsValid())
+	{
+		return FCursorReply::Unhandled();
+	}
+
 	FScrubRangeToScreen RangeToScreen(TimeSliderArgs.ViewRange.Get(), MyGeometry.Size);
 
 	const FFrameRate TickResolution   = GetTickResolution();
@@ -932,6 +968,12 @@ FCursorReply FSequencerTimeSliderController::OnCursorQuery( TSharedRef<const SWi
 
 int32 FSequencerTimeSliderController::OnPaintSectionView( const FGeometry& AllottedGeometry, const FSlateRect& MyCullingRect, FSlateWindowElementList& OutDrawElements, int32 LayerId, bool bEnabled, const FPaintSectionAreaViewArgs& Args ) const
 {
+	TSharedPtr<FSequencer> Sequencer = WeakSequencer.Pin();
+	if (!Sequencer.IsValid())
+	{
+		return LayerId;
+	}
+
 	const ESlateDrawEffect DrawEffects = bEnabled ? ESlateDrawEffect::None : ESlateDrawEffect::DisabledEffect;
 
 	TRange<double> LocalViewRange = TimeSliderArgs.ViewRange.Get();
