Fix for performance issues with motion trails.

Found that the main issue was the GetRefTM call.  Refactored the code so that we only call that once per tick, instead of up to  4 times per tick. Also we cache the calculation of the key positions, so we don't need to calculate them again when calculating keys.  Finally the GetRefTM call was recursive which was a decent hint, so we cache out it's parent's first and just get the transform's linearly.

Also we don't select parents or children to show their trails also, this seemed buggy and possibly very very slow.

diff --git a/Engine/Source/Editor/Sequencer/Private/SequencerEdMode.cpp b/Engine/Source/Editor/Sequencer/Private/SequencerEdMode.cpp
--- a/Engine/Source/Editor/Sequencer/Private/SequencerEdMode.cpp
+++ b/Engine/Source/Editor/Sequencer/Private/SequencerEdMode.cpp
@@ -238,6 +238,115 @@ static const float	DrawTrackTimeRes = 0.1f;
 static const float	CurveHandleScale = 0.5f;
 }
 
+void FSequencerEdMode::GetParents(TArray<const UObject *>& Parents, const UObject* InObject)
+{
+	const AActor* Actor = Cast<AActor>(InObject);
+	if (Actor != nullptr)
+	{
+		Parents.Emplace(InObject);
+	}
+	else
+	{
+		const USceneComponent* SceneComponent = Cast<USceneComponent>(InObject);
+
+		if (SceneComponent != nullptr)
+		{
+			Parents.Emplace(InObject);
+			if (SceneComponent->GetAttachParent() == SceneComponent->GetOwner()->GetRootComponent())
+			{
+				GetParents(Parents, SceneComponent->GetAttachParent()->GetOwner());
+			}
+			else
+			{
+				GetParents(Parents, SceneComponent->GetAttachParent());
+			}
+		}
+	}
+}
+/** This is not that scalable moving forward with stuff like the control rig , need a better caching solution there */
+bool FSequencerEdMode::GetParentTM(FTransform& CurrentRefTM, const TSharedPtr<FSequencer>& Sequencer, UObject* ParentObject, FFrameTime KeyTime)
+{
+	FGuid ObjectBinding = Sequencer->FindObjectId(*ParentObject, Sequencer->GetFocusedTemplateID());
+
+	if (ObjectBinding.IsValid())
+	{
+		const TSharedPtr< FSequencerDisplayNode >& ObjectNode = Sequencer->GetNodeTree()->GetObjectBindingMap()[ObjectBinding];
+
+		for (const TSharedRef< FSequencerDisplayNode >& ChildNode : ObjectNode->GetChildNodes())
+		{
+			if (ChildNode->GetType() == ESequencerNode::Track)
+			{
+				const TSharedRef<FSequencerTrackNode> TrackNode = StaticCastSharedRef<FSequencerTrackNode>(ChildNode);
+				const UMovieSceneTrack* TrackNodeTrack = TrackNode->GetTrack();
+				const UMovieScene3DTransformTrack* TransformTrack = Cast<UMovieScene3DTransformTrack>(TrackNodeTrack);
+
+				if (TransformTrack != nullptr)
+				{
+					for (const UMovieSceneSection* Section : TransformTrack->GetAllSections())
+					{
+						if (Section->IsTimeWithinSection(KeyTime.FrameNumber))
+						{
+							const UMovieScene3DTransformSection* ParentSection = Cast<UMovieScene3DTransformSection>(Section);
+
+							if (ParentSection != nullptr)
+							{
+								FVector ParentKeyPos;
+								FRotator ParentKeyRot;
+
+								FMovieSceneEvaluationTemplate* Template = Sequencer->GetEvaluationTemplate().FindTemplate(Sequencer->GetFocusedTemplateID());
+								FMovieSceneEvaluationTrack* EvalTrack = Template ? Template->FindTrack(TransformTrack->GetSignature()) : nullptr;
+
+								if (EvalTrack)
+								{
+									GetLocationAtTime(EvalTrack, ParentObject, KeyTime, ParentKeyPos, ParentKeyRot, Sequencer);
+									CurrentRefTM = FTransform(ParentKeyRot, ParentKeyPos);
+									return true;
+								}
+							}
+						}
+					}
+				}
+			}
+		}
+	}
+
+	return false;
+}
+
+FTransform FSequencerEdMode::GetRefFrame(const TSharedPtr<FSequencer>& Sequencer, const TArray<const UObject *>& Parents, FFrameTime KeyTime)
+{
+	FTransform RefTM = FTransform::Identity;
+	FTransform ParentRefTM = FTransform::Identity;
+
+	for (const UObject* Object : Parents)
+	{
+		const AActor* Actor = Cast<AActor>(Object);
+		if (Actor != nullptr)
+		{
+			if (Actor->GetRootComponent() != nullptr && Actor->GetRootComponent()->GetAttachParent() != nullptr)
+			{
+				if (!GetParentTM(ParentRefTM, Sequencer, Actor->GetRootComponent()->GetAttachParent()->GetOwner(), KeyTime))
+				{
+					RefTM = Actor->GetRootComponent()->GetAttachParent()->GetSocketTransform(Actor->GetRootComponent()->GetAttachSocketName());
+				}
+				RefTM = ParentRefTM * RefTM;
+			}
+		}
+		else
+		{
+			const USceneComponent* SceneComponent = Cast<USceneComponent>(Object);
+			FTransform CurrentRefTM = FTransform::Identity;
+			UObject* ParentObject = SceneComponent->GetAttachParent() == SceneComponent->GetOwner()->GetRootComponent() ? static_cast<UObject*>(SceneComponent->GetOwner()) : SceneComponent->GetAttachParent();
+			if (!GetParentTM(CurrentRefTM, Sequencer, ParentObject, KeyTime))
+			{
+				CurrentRefTM = RefTM * SceneComponent->GetAttachParent()->GetRelativeTransform();
+			}
+			RefTM = CurrentRefTM * RefTM;
+		}
+	}
+	return RefTM;
+}
+
 FTransform FSequencerEdMode::GetRefFrame(const TSharedPtr<FSequencer>& Sequencer, const UObject* InObject, FFrameTime KeyTime)
 {
 	FTransform RefTM = FTransform::Identity;
@@ -278,7 +387,7 @@ FTransform FSequencerEdMode::GetRefFrame(const TSharedPtr<FSequencer>& Sequencer
 
 	if (SceneComponent != nullptr && SceneComponent->GetAttachParent() != nullptr)
 	{
-		FTransform ParentRefTM = GetRefFrame(Sequencer, SceneComponent->GetAttachParent(), KeyTime);
+		FTransform ParentRefTM;
 
 		// If our parent is the root component, get the RefFrame from the Actor
 		if (SceneComponent->GetAttachParent() == SceneComponent->GetOwner()->GetRootComponent())
@@ -364,14 +473,14 @@ void FSequencerEdMode::GetLocationAtTime(FMovieSceneEvaluationTrack* Track, UObj
 }
 
 void FSequencerEdMode::DrawTransformTrack(const TSharedPtr<FSequencer>& Sequencer, FPrimitiveDrawInterface* PDI,
-											UMovieScene3DTransformTrack* TransformTrack, const TArray<TWeakObjectPtr<UObject>>& BoundObjects, const bool bIsSelected)
+	UMovieScene3DTransformTrack* TransformTrack, const TArray<TWeakObjectPtr<UObject>>& BoundObjects, const bool bIsSelected)
 {
 	bool bHitTesting = true;
-	if( PDI != nullptr )
+	if (PDI != nullptr)
 	{
 		bHitTesting = PDI->IsHitTesting();
 	}
-	
+
 	ASequencerMeshTrail* TrailActor = nullptr;
 	// Get the Trail Actor associated with this track if we are drawing mesh trails
 	if (bDrawMeshTrails)
@@ -388,30 +497,29 @@ void FSequencerEdMode::DrawTransformTrack(const TSharedPtr<FSequencer>& Sequence
 
 	bool bShowTrajectory = TransformTrack->GetAllSections().ContainsByPredicate(
 		[bIsSelected](UMovieSceneSection* Section)
+	{
+		UMovieScene3DTransformSection* TransformSection = Cast<UMovieScene3DTransformSection>(Section);
+		if (TransformSection)
 		{
-			UMovieScene3DTransformSection* TransformSection = Cast<UMovieScene3DTransformSection>(Section);
-			if (TransformSection)
+			switch (TransformSection->GetShow3DTrajectory())
 			{
-				switch (TransformSection->GetShow3DTrajectory())
-				{
-				case EShow3DTrajectory::EST_Always:				return true;
-				case EShow3DTrajectory::EST_Never:				return false;
-				case EShow3DTrajectory::EST_OnlyWhenSelected:	return bIsSelected;
-				}
+			case EShow3DTrajectory::EST_Always:				return true;
+			case EShow3DTrajectory::EST_Never:				return false;
+			case EShow3DTrajectory::EST_OnlyWhenSelected:	return bIsSelected;
 			}
-			return false;
 		}
+		return false;
+	}
 	);
-	
+
 	FFrameRate TickResolution = Sequencer->GetFocusedTickResolution();
 	FMovieSceneEvaluationTemplate* Template = Sequencer->GetEvaluationTemplate().FindTemplate(Sequencer->GetFocusedTemplateID());
-	if (!bShowTrajectory || !Template || !TransformTrack->GetAllSections().ContainsByPredicate([](UMovieSceneSection* In){ return In->IsActive(); }))
+	if (!bShowTrajectory || !Template || !TransformTrack->GetAllSections().ContainsByPredicate([](UMovieSceneSection* In) { return In->IsActive(); }))
 	{
 		return;
 	}
 
 	FLinearColor TrackColor = TransformTrack->GetColorTint();
-
 	// Draw one line per-track (should only really ever be one)
 	if (FMovieSceneEvaluationTrack* EvalTrack = Template->FindTrack(TransformTrack->GetSignature()))
 	{
@@ -423,47 +531,59 @@ void FSequencerEdMode::DrawTransformTrack(const TSharedPtr<FSequencer>& Sequence
 			{
 				continue;
 			}
+			TArray<const UObject *> Parents;
+			GetParents(Parents, BoundObject);
 
 			FVector OldKeyPos(0);
 			FFrameTime OldKeyTime = 0;
 			int KeyTimeIndex = 0;
-
+			FTransform OldPosRefTM;;
+			FVector OldPos_G;
+			//We cache the key times and positions so we don't have to calculate it twice.
+			//In particularly in same cases calling GetRefFrame may be heavy.
+			//We cache time also so we get cache coherence when iterating over them later.
+			struct FKeyPositionRotation
+			{
+				FVector Position;
+				FRotator Rotation;
+				FVector WorldPosition;
+				FTrajectoryKey TrajectoryKey;
+
+				FKeyPositionRotation(const FTrajectoryKey &InTrajKey, const FVector &InPos, const FRotator &InRot, const FVector &InWorldPos) :
+					TrajectoryKey(InTrajKey), Position(InPos), Rotation(InRot), WorldPosition(InWorldPos) {}
+			};
+			TArray<FKeyPositionRotation> KeyPosRots;
+			KeyPosRots.Reserve(TrajectoryKeys.Num());
 			for (const FTrajectoryKey& NewTrajectoryKey : TrajectoryKeys)
 			{
 				FFrameTime NewKeyTime = NewTrajectoryKey.Time;
 
 				FVector NewKeyPos(0);
-				FRotator NewKeyRot(0,0,0);
+				FRotator NewKeyRot(0, 0, 0);
 
 				GetLocationAtTime(EvalTrack, BoundObject, NewKeyTime, NewKeyPos, NewKeyRot, Sequencer);
-
+				FTransform NewPosRefTM = GetRefFrame(Sequencer, Parents, NewKeyTime);
+				FVector NewKeyPos_G = NewPosRefTM.TransformPosition(NewKeyPos);
+				FKeyPositionRotation KeyPosRot(NewTrajectoryKey, NewKeyPos, NewKeyRot, NewKeyPos_G);
+				KeyPosRots.Push(KeyPosRot);
 				// If not the first keypoint, draw a line to the last keypoint.
-				if(KeyTimeIndex > 0)
+				if (KeyTimeIndex > 0)
 				{
 					int32 NumSteps = FMath::CeilToInt((TickResolution.AsSeconds(NewKeyTime) - TickResolution.AsSeconds(OldKeyTime)) / SequencerEdMode_Draw3D::DrawTrackTimeRes);
 					// Limit the number of steps to prevent a rendering performance hit
-					NumSteps = FMath::Min( 100, NumSteps );
-					FFrameTime DrawSubstep = NumSteps == 0 ? 0 : (NewKeyTime - OldKeyTime)*(1.f/NumSteps);
-
+					NumSteps = FMath::Min(100, NumSteps);
+					FFrameTime DrawSubstep = NumSteps == 0 ? 0 : (NewKeyTime - OldKeyTime)*(1.f / NumSteps);
 					// Find position on first keyframe.
 					FFrameTime OldTime = OldKeyTime;
-
 					FVector OldPos(0);
-					FRotator OldRot(0,0,0);
+					FRotator OldRot(0, 0, 0);
 					GetLocationAtTime(EvalTrack, BoundObject, OldKeyTime, OldPos, OldRot, Sequencer);
 
 					const bool bIsConstantKey = NewTrajectoryKey.Is(ERichCurveInterpMode::RCIM_Constant);
-
-					FTransform OldPosRefTM = GetRefFrame(Sequencer, BoundObject, OldKeyTime);
-					FTransform NewPosRefTM = GetRefFrame(Sequencer, BoundObject, NewKeyTime);
-
-					FVector OldPos_G = OldPosRefTM.TransformPosition(OldPos);
-					FVector NewKeyPos_G = NewPosRefTM.TransformPosition(NewKeyPos);
-
 					// For constant interpolation - don't draw ticks - just draw dotted line.
 					if (bIsConstantKey)
 					{
-						if(PDI != nullptr)
+						if (PDI != nullptr)
 						{
 							DrawDashedLine(PDI, OldPos_G, NewKeyPos_G, TrackColor, 20, SDPG_Foreground);
 						}
@@ -471,15 +591,16 @@ void FSequencerEdMode::DrawTransformTrack(const TSharedPtr<FSequencer>& Sequence
 					else
 					{
 						// Then draw a line for each substep.
-						for (int32 j=1; j<NumSteps+1; j++)
+						for (int32 j = 1; j<NumSteps + 1; j++)
 						{
 							FFrameTime NewTime = OldKeyTime + DrawSubstep*j;
 
 							FVector NewPos(0);
-							FRotator NewRot(0,0,0);
+							FRotator NewRot(0, 0, 0);
 							GetLocationAtTime(EvalTrack, BoundObject, NewTime, NewPos, NewRot, Sequencer);
 
-							FTransform RefTM = GetRefFrame(Sequencer, BoundObject, NewTime);
+							FTransform RefTM = GetRefFrame(Sequencer, Parents, NewTime);
+
 							FVector NewPos_G = RefTM.TransformPosition(NewPos);
 							if (PDI != nullptr)
 							{
@@ -503,41 +624,31 @@ void FSequencerEdMode::DrawTransformTrack(const TSharedPtr<FSequencer>& Sequence
 						}
 					}
 				}
-					
+				OldPosRefTM = NewPosRefTM;
+				OldPos_G = NewKeyPos_G;
 				OldKeyTime = NewKeyTime;
 				OldKeyPos = NewKeyPos;
 				++KeyTimeIndex;
 			}
 
 			// Draw keypoints on top of curve
-			for (const FTrajectoryKey& TrajectoryKey : TrajectoryKeys)
+			FColor KeyColor = TrackColor.ToFColor(true);
+			for (const FKeyPositionRotation& KeyPosRot : KeyPosRots)
 			{
-				FFrameTime NewKeyTime = TrajectoryKey.Time;
-
-				// Find the time, position and orientation of this Key.
-				FVector NewKeyPos(0);
-				FRotator NewKeyRot(0,0,0);
-				GetLocationAtTime(EvalTrack, BoundObject, NewKeyTime, NewKeyPos, NewKeyRot, Sequencer);
-
-				FTransform RefTM = GetRefFrame(Sequencer, BoundObject, NewKeyTime);
-
-				FColor KeyColor = TrackColor.ToFColor(true);
-
-				if (bHitTesting && PDI) 
+				if (bHitTesting && PDI)
 				{
-					PDI->SetHitProxy(new HMovieSceneKeyProxy(TransformTrack, TrajectoryKey));
+					PDI->SetHitProxy(new HMovieSceneKeyProxy(TransformTrack, KeyPosRot.TrajectoryKey));
 				}
 
-				FVector NewKeyPos_G = RefTM.TransformPosition(NewKeyPos);
 				// Drawing keys
 				if (PDI != nullptr)
 				{
-					PDI->DrawPoint(NewKeyPos_G, KeyColor, 6.f, SDPG_Foreground);
+					PDI->DrawPoint(KeyPosRot.WorldPosition, KeyColor, 6.f, SDPG_Foreground);
 				}
 				else if (TrailActor != nullptr)
 				{
 					TArray<UMovieScene3DTransformSection*> AllSections;
-					for (const FTrajectoryKey::FData& Value : TrajectoryKey.KeyData)
+					for (const FTrajectoryKey::FData& Value : KeyPosRot.TrajectoryKey.KeyData)
 					{
 						if (UMovieScene3DTransformSection* Section = Value.Section.Get())
 						{
@@ -547,11 +658,11 @@ void FSequencerEdMode::DrawTransformTrack(const TSharedPtr<FSequencer>& Sequence
 
 					for (UMovieScene3DTransformSection* Section : AllSections)
 					{
-						TrailActor->AddKeyMeshActor(NewKeyTime / TickResolution, FTransform(NewKeyRot, NewKeyPos, FVector(3.0f)), Section);
+						TrailActor->AddKeyMeshActor(KeyPosRot.TrajectoryKey.Time / TickResolution, FTransform(KeyPosRot.Rotation, KeyPosRot.Position, FVector(3.0f)), Section);
 					}
 				}
 
-				if (bHitTesting && PDI) 
+				if (bHitTesting && PDI)
 				{
 					PDI->SetHitProxy(nullptr);
 				}
@@ -576,48 +687,14 @@ void FSequencerEdMode::DrawTracks3D(FPrimitiveDrawInterface* PDI)
 		// Map between object binding nodes and selection
 		TMap<TSharedRef<FSequencerDisplayNode>, bool > ObjectBindingNodesSelectionMap;
 
-		for (auto ObjectBinding : Sequencer->GetNodeTree()->GetObjectBindingMap() )
+		for (auto ObjectBinding : Sequencer->GetNodeTree()->GetObjectBindingMap())
 		{
 			if (!ObjectBinding.Value.IsValid())
 			{
 				continue;
 			}
-
 			TSharedRef<FSequencerObjectBindingNode> ObjectBindingNode = ObjectBinding.Value.ToSharedRef();
-
-			TSet<TSharedRef<FSequencerDisplayNode> > DescendantNodes;
-			SequencerHelpers::GetDescendantNodes(ObjectBindingNode, DescendantNodes);
-
 			bool bSelected = Sequencer->GetSelection().IsSelected(ObjectBindingNode);
-
-			if (!bSelected)
-			{
-				// If one of our child is selected, we're considered selected
-				for (auto& DescendantNode : DescendantNodes)
-				{
-					if (Sequencer->GetSelection().IsSelected(DescendantNode) ||
-						Sequencer->GetSelection().NodeHasSelectedKeysOrSections(DescendantNode))
-					{
-						bSelected = true;
-						break;
-					}
-				}
-			}
-
-			// If one of our parent is selected, we're considered selected
-			TSharedPtr<FSequencerDisplayNode> ParentNode = ObjectBindingNode->GetParent();
-
-			while (!bSelected && ParentNode.IsValid())
-			{
-				if (Sequencer->GetSelection().IsSelected(ParentNode.ToSharedRef()) ||
-					Sequencer->GetSelection().NodeHasSelectedKeysOrSections(ParentNode.ToSharedRef()))
-				{
-					bSelected = true;
-				}
-
-				ParentNode = ParentNode->GetParent();
-			}
-
 			ObjectBindingNodesSelectionMap.Add(ObjectBindingNode, bSelected);
 		}
 
@@ -649,8 +726,8 @@ void FSequencerEdMode::DrawTracks3D(FPrimitiveDrawInterface* PDI)
 							});
 							if (TrailPtr == nullptr)
 							{
-								UViewportWorldInteraction* WorldInteraction = Cast<UViewportWorldInteraction>( GEditor->GetEditorWorldExtensionsManager()->GetEditorWorldExtensions( GetWorld() )->FindExtension( UViewportWorldInteraction::StaticClass() ) );
-								if( WorldInteraction != nullptr )
+								UViewportWorldInteraction* WorldInteraction = Cast<UViewportWorldInteraction>(GEditor->GetEditorWorldExtensionsManager()->GetEditorWorldExtensions(GetWorld())->FindExtension(UViewportWorldInteraction::StaticClass()));
+								if (WorldInteraction != nullptr)
 								{
 									ASequencerMeshTrail* TrailActor = WorldInteraction->SpawnTransientSceneActor<ASequencerMeshTrail>(TEXT("SequencerMeshTrail"), true);
 									FMeshTrailData MeshTrail = FMeshTrailData(TransformTrack, TrailActor);
@@ -667,6 +744,7 @@ void FSequencerEdMode::DrawTracks3D(FPrimitiveDrawInterface* PDI)
 	}
 }
 
+
 FSequencerEdModeTool::FSequencerEdModeTool(FSequencerEdMode* InSequencerEdMode) :
 	SequencerEdMode(InSequencerEdMode)
 {
diff --git a/Engine/Source/Editor/Sequencer/Private/SequencerEdMode.h b/Engine/Source/Editor/Sequencer/Private/SequencerEdMode.h
--- a/Engine/Source/Editor/Sequencer/Private/SequencerEdMode.h
+++ b/Engine/Source/Editor/Sequencer/Private/SequencerEdMode.h
@@ -78,6 +78,10 @@ protected:
 	static FTransform GetRefFrame(const TSharedPtr<FSequencer>& Sequencer, const UObject* InObject, FFrameTime KeyTime);
 	static FTransform GetRefFrame(const TSharedPtr<FSequencer>& Sequencer, const AActor* Actor, FFrameTime KeyTime);
 	static FTransform GetRefFrame(const TSharedPtr<FSequencer>& Sequencer, const USceneComponent* SceneComponent, FFrameTime KeyTime);
+	static void GetParents(TArray<const UObject *>& Parents, const UObject* InObject);
+	static FTransform GetRefFrame(const TSharedPtr<FSequencer>& Sequencer, const TArray<const UObject *>& Parents, FFrameTime KeyTime);
+	static bool GetParentTM(FTransform& CurrentRefTM, const TSharedPtr<FSequencer>& Sequencer, UObject* ParentObject, FFrameTime KeyTime);
+
 
 private:
 	TArray<TWeakPtr<FSequencer>> Sequencers;
