Fix for build warning, order of initialization.
#trivial

diff --git a/Engine/Source/Editor/Sequencer/Private/SequencerEdMode.cpp b/Engine/Source/Editor/Sequencer/Private/SequencerEdMode.cpp
--- a/Engine/Source/Editor/Sequencer/Private/SequencerEdMode.cpp
+++ b/Engine/Source/Editor/Sequencer/Private/SequencerEdMode.cpp
@@ -550,7 +550,7 @@ void FSequencerEdMode::DrawTransformTrack(const TSharedPtr<FSequencer>& Sequence
 				FTrajectoryKey TrajectoryKey;
 
 				FKeyPositionRotation(const FTrajectoryKey &InTrajKey, const FVector &InPos, const FRotator &InRot, const FVector &InWorldPos) :
-					TrajectoryKey(InTrajKey), Position(InPos), Rotation(InRot), WorldPosition(InWorldPos) {}
+					 Position(InPos), Rotation(InRot), WorldPosition(InWorldPos), TrajectoryKey(InTrajKey) {}
 			};
 			TArray<FKeyPositionRotation> KeyPosRots;
 			KeyPosRots.Reserve(TrajectoryKeys.Num());
