Sequencer: Fix crash with stale section pointer

diff --git a/Engine/Source/Editor/Sequencer/Private/SSequencerSectionAreaView.cpp b/Engine/Source/Editor/Sequencer/Private/SSequencerSectionAreaView.cpp
--- a/Engine/Source/Editor/Sequencer/Private/SSequencerSectionAreaView.cpp
+++ b/Engine/Source/Editor/Sequencer/Private/SSequencerSectionAreaView.cpp
@@ -155,7 +155,10 @@ void SSequencerSectionAreaView::OnArrangeChildren( const FGeometry& AllottedGeom
 
 			TSharedPtr<ISequencerSection> SectionInterface = Widget->GetSectionInterface();
 
-			MaxRowIndex = FMath::Max(MaxRowIndex, SectionInterface->GetSectionObject()->GetRowIndex());
+			if (SectionInterface->GetSectionObject())
+			{
+				MaxRowIndex = FMath::Max(MaxRowIndex, SectionInterface->GetSectionObject()->GetRowIndex());
+			}
 		}
 	}
 	int32 MaxTracks = MaxRowIndex + 1;
@@ -168,6 +171,11 @@ void SSequencerSectionAreaView::OnArrangeChildren( const FGeometry& AllottedGeom
 
 		TSharedPtr<ISequencerSection> SectionInterface = Widget->GetSectionInterface();
 
+		if (!SectionInterface->GetSectionObject())
+		{
+			continue;
+		}
+
 		int32 RowIndex = SectionAreaNode->GetSubTrackMode() == FSequencerTrackNode::ESubTrackMode::None ? SectionInterface->GetSectionObject()->GetRowIndex() : 0;
 
 		EVisibility WidgetVisibility = Widget->GetVisibility();
