Sequencer: Expand the clamp range when necessary if the sequencer window expands.

This fixes a new bug where the view range would not update because TimeSliderController::SetViewRange now clamps the incoming view range to the clamp range. In one case, tick marks in the display were cached using the view range and weren't getting updated as the window was resizing because the view range remained clamped.

diff --git a/Engine/Source/Editor/Sequencer/Private/SSequencerTrackArea.cpp b/Engine/Source/Editor/Sequencer/Private/SSequencerTrackArea.cpp
--- a/Engine/Source/Editor/Sequencer/Private/SSequencerTrackArea.cpp
+++ b/Engine/Source/Editor/Sequencer/Private/SSequencerTrackArea.cpp
@@ -444,9 +444,21 @@ void SSequencerTrackArea::Tick( const FGeometry& AllottedGeometry, const double
 		const float Difference = Size.X - SizeLastFrame->X;
 		TRange<double> OldRange = TimeSliderController->GetViewRange().GetAnimationTarget();
 
+		double NewRangeMin = OldRange.GetLowerBoundValue();
+		double NewRangeMax = OldRange.GetUpperBoundValue() + (Difference * OldRange.Size<double>() / SizeLastFrame->X);
+
+		TRange<double> ClampRange = TimeSliderController->GetClampRange();
+
+		if (NewRangeMin < ClampRange.GetLowerBoundValue() || NewRangeMax > ClampRange.GetUpperBoundValue())
+		{
+			double NewClampRangeMin = NewRangeMin < ClampRange.GetLowerBoundValue() ? NewRangeMin : ClampRange.GetLowerBoundValue();
+			double NewClampRangeMax = NewRangeMax > ClampRange.GetUpperBoundValue() ? NewRangeMax : ClampRange.GetUpperBoundValue();
+
+			TimeSliderController->SetClampRange(NewClampRangeMin, NewClampRangeMax);
+		}
+
 		TimeSliderController->SetViewRange(
-			OldRange.GetLowerBoundValue(),
-			OldRange.GetUpperBoundValue() + (Difference * OldRange.Size<double>() / SizeLastFrame->X),
+			NewRangeMin, NewRangeMax,
 			EViewRangeInterpolation::Immediate
 		);
 	}
