Sequencer: Fix crash in drawing path without attach parent.

diff --git a/Engine/Source/Editor/Sequencer/Private/SequencerEdMode.cpp b/Engine/Source/Editor/Sequencer/Private/SequencerEdMode.cpp
--- a/Engine/Source/Editor/Sequencer/Private/SequencerEdMode.cpp
+++ b/Engine/Source/Editor/Sequencer/Private/SequencerEdMode.cpp
@@ -337,9 +337,13 @@ FTransform FSequencerEdMode::GetRefFrame(const TSharedPtr<FSequencer>& Sequencer
 			const USceneComponent* SceneComponent = Cast<USceneComponent>(Object);
 			FTransform CurrentRefTM = FTransform::Identity;
 			UObject* ParentObject = SceneComponent->GetAttachParent() == SceneComponent->GetOwner()->GetRootComponent() ? static_cast<UObject*>(SceneComponent->GetOwner()) : SceneComponent->GetAttachParent();
-			if (!GetParentTM(CurrentRefTM, Sequencer, ParentObject, KeyTime))
+
+			if (SceneComponent->GetAttachParent() != nullptr)
 			{
-				CurrentRefTM = RefTM * SceneComponent->GetAttachParent()->GetRelativeTransform();
+				if (!GetParentTM(CurrentRefTM, Sequencer, ParentObject, KeyTime))
+				{
+					CurrentRefTM = RefTM * SceneComponent->GetAttachParent()->GetRelativeTransform();
+				}
 			}
 			RefTM = CurrentRefTM * RefTM;
 		}
