Sequencer: Fixes various issues copy and pasting Spawnables from one sequence to another and within the same sequence.
No longer prints superfluous warnings on paste and properly initializes the new object template so that it does not get garbage collected/cause issues on save.

This used to crash on garbage collection and around saving, changing the text serialization and deserialization seems to have fixed it.

diff --git a/Engine/Source/Editor/Sequencer/Private/MovieSceneCopyableBinding.h b/Engine/Source/Editor/Sequencer/Private/MovieSceneCopyableBinding.h
--- a/Engine/Source/Editor/Sequencer/Private/MovieSceneCopyableBinding.h
+++ b/Engine/Source/Editor/Sequencer/Private/MovieSceneCopyableBinding.h
@@ -11,15 +11,26 @@
 
 class UMovieSceneTrack;
 
-UCLASS()
+UCLASS(Transient)
 class UMovieSceneCopyableBinding : public UObject
 {
 public:
 	GENERATED_BODY()
 	
-	UPROPERTY()
+	/** 
+	* Spawnables need to know about their Object Template but we cannot rely on automatic serialization due to the object
+	* template belonging to the Movie Scene (it gets serialized as a reference). Instead we manually serialize the object
+	* so that we can duplicate it into a new object (which is stored in this variable) but we don't want this exported with
+	* the rest of the text as it'll fall back to the same reference issue. Marking this as TextExportTransient solves this.
+	*/
+	UPROPERTY(TextExportTransient)
 	UObject* SpawnableObjectTemplate;
 
+	/**
+	 * Tracks are also owned by the owning Movie Sequence. We manually copy the tracks out of a binding when we copy,
+	 * because the binding stores them as a reference to a privately owned object. We store these copied tracks here,
+	 * and then restore them upon paste to re-create the tracks with the correct owner.
+	 */
 	UPROPERTY()
 	TArray<UMovieSceneTrack*> Tracks;
 
diff --git a/Engine/Source/Editor/Sequencer/Private/SSequencer.cpp b/Engine/Source/Editor/Sequencer/Private/SSequencer.cpp
--- a/Engine/Source/Editor/Sequencer/Private/SSequencer.cpp
+++ b/Engine/Source/Editor/Sequencer/Private/SSequencer.cpp
@@ -71,6 +71,7 @@
 #include "FrameNumberNumericInterface.h"
 #include "LevelSequence.h"
 #include "SequencerLog.h"
+#include "MovieSceneCopyableBinding.h"
 
 #define LOCTEXT_NAMESPACE "Sequencer"
 
@@ -2296,30 +2297,27 @@ void SSequencer::OnPaste()
 bool SSequencer::CanPaste()
 {
 	TSharedPtr<FSequencer> Sequencer = SequencerPtr.Pin();
-	TSet<TSharedRef<FSequencerDisplayNode>> SelectedNodes = Sequencer->GetSelection().GetSelectedOutlinerNodes();
-	if (SelectedNodes.Num() != 0)
+	
+	FString TextToImport;
+	FPlatformApplicationMisc::ClipboardPaste(TextToImport);
+
+	// Attempts to deserialize the text into object bindings/tracks that Sequencer understands.
+	if (Sequencer->CanPaste(TextToImport))
 	{
-		FString TexttoImport;
-		FPlatformApplicationMisc::ClipboardPaste(TexttoImport);
+		TArray<UMovieSceneTrack*> ImportedTrack;
+		TArray<UMovieSceneCopyableBinding*> ImportedObjects;
+		Sequencer->ImportTracksFromText(TextToImport, ImportedTrack);
+		Sequencer->ImportObjectsFromText(TextToImport, ImportedObjects);
 
-		if (Sequencer->CanPaste(TexttoImport))
+		// If we couldn't deserialize any tracks or objects then the data isn't valid for sequencer,
+		// and we'll block a paste attempt.
+		if (ImportedTrack.Num() == 0 && ImportedObjects.Num() == 0)
 		{
-			TArray<UMovieSceneTrack*> ImportedTrack;
-			Sequencer->ImportTracksFromText(TexttoImport, ImportedTrack);
-			if (ImportedTrack.Num() == 0)
-			{
-				return false;
-			}
-
-			for (TSharedRef<FSequencerDisplayNode> Node : SelectedNodes)
-			{
-				if (Node->GetType() == ESequencerNode::Object)
-				{
-					return true;
-				}
-			}
 			return false;
 		}
+
+		// Otherwise, as long as they have one or the other, there is something to paste.
+		return true;
 	}
 
 	return SequencerPtr.Pin()->GetClipboardStack().Num() != 0;
diff --git a/Engine/Source/Editor/Sequencer/Private/Sequencer.cpp b/Engine/Source/Editor/Sequencer/Private/Sequencer.cpp
--- a/Engine/Source/Editor/Sequencer/Private/Sequencer.cpp
+++ b/Engine/Source/Editor/Sequencer/Private/Sequencer.cpp
@@ -5569,9 +5569,20 @@ void ExportObjectsToText(const TArray<UMovieSceneCopyableBinding*>& ObjectsToExp
 		check((LastOuter == ThisOuter) || (LastOuter == nullptr));
 		LastOuter = ThisOuter;
 
+		// We can't use TextExportTransient on USTRUCTS (which our object contains) so we're going to manually null out some references before serializing them. These references are
+		// serialized manually into the archive, as the auto-serialization will only store a reference (to a privately owned object) which creates issues on deserialization. Attempting 
+		// to deserialize these private objects throws a superflous error in the console that makes it look like things went wrong when they're actually OK and expected.
+		TArray<UMovieSceneTrack*> OldTracks = ObjectToExport->Binding.StealTracks();
+		UObject* OldSpawnableTemplate = ObjectToExport->Spawnable.GetObjectTemplate();
+		ObjectToExport->Spawnable.SetObjectTemplate(nullptr);
+
 		UExporter::ExportToOutputDevice(&Context, ObjectToExport, nullptr, Archive, TEXT("copy"), 0, PPF_ExportsNotFullyQualified | PPF_Copy | PPF_Delimited, false, ThisOuter);
 
-		// cbb: Copy the object template separately for now. There's a bug in reimporting the object template from the CopyableBinding.
+		// Restore the references (as we don't want to modify the original in the event of a copy operation!)
+		ObjectToExport->Binding.SetTracks(OldTracks);
+		ObjectToExport->Spawnable.SetObjectTemplate(OldSpawnableTemplate);
+
+		// We manually export the object template for the same private-ownership reason as above. Templates need to be re-created anyways as each Spawnable contains its own copy of the template.
 		if (ObjectToExport->SpawnableObjectTemplate)
 		{
 			UExporter::ExportToOutputDevice(&Context, ObjectToExport->SpawnableObjectTemplate, nullptr, Archive, TEXT("copy"), 0, PPF_ExportsNotFullyQualified | PPF_Copy | PPF_Delimited, false, ThisOuter);
@@ -5626,22 +5637,25 @@ private:
 };
 
 
-void ImportObjectsFromText(FSequencer& InSequencer, const FString& TextToImport, /*out*/ TArray<UMovieSceneCopyableBinding*>& ImportedObjects)
+void FSequencer::ImportObjectsFromText(const FString& TextToImport, /*out*/ TArray<UMovieSceneCopyableBinding*>& ImportedObjects)
 {
 	UPackage* TempPackage = NewObject<UPackage>(nullptr, TEXT("/Engine/Sequencer/Editor/Transient"), RF_Transient);
 	TempPackage->AddToRoot();
 
 	// Turn the text buffer into objects
-	FObjectBindingTextFactory Factory(InSequencer);
+	FObjectBindingTextFactory Factory(*this);
 	Factory.ProcessBuffer(TempPackage, RF_Transactional, TextToImport);
 	ImportedObjects = Factory.NewCopyableBindings;
 
-	// cbb: Process the object templates separately for now. There's a bug in reimporting the object template from the CopyableBinding.
+	// We had to explicitly serialize object templates due to them being a reference to a privately owned object. We now deserialize these object template copies
+	// and match them up with their MovieSceneCopyableBinding again.
+	
 	int32 SpawnableObjectTemplateIndex = 0;
 	for (auto ImportedObject : ImportedObjects)
 	{
 		if (ImportedObject->Spawnable.GetGuid().IsValid() && SpawnableObjectTemplateIndex < Factory.NewSpawnableObjectTemplates.Num())
 		{
+			// This Spawnable Object Template is owned by our transient package, so you'll need to change the owner if you want to keep it later.
 			ImportedObject->SpawnableObjectTemplate = Factory.NewSpawnableObjectTemplates[SpawnableObjectTemplateIndex++];
 		}
 	}
@@ -5693,6 +5707,9 @@ void FSequencer::CopySelectedObjects(TArray<TSharedPtr<FSequencerObjectBindingNo
 			if (Spawnable)
 			{
 				CopyableBinding->Spawnable = *Spawnable;
+				
+				// We manually serialize the spawnable object template so that it's not a reference to a privately owned object. Spawnables all have unique copies of their template objects anyways.
+				// Object Templates are re-created on paste (based on these templates) with the correct ownership set up.
 				CopyableBinding->SpawnableObjectTemplate = Spawnable->GetObjectTemplate();
 			}
 		}
@@ -5703,6 +5720,8 @@ void FSequencer::CopySelectedObjects(TArray<TSharedPtr<FSequencerObjectBindingNo
 			CopyableBinding->Binding = *Binding;
 			for (auto Track : Binding->GetTracks())
 			{
+				// Tracks suffer from the same issues as Spawnable's Object Templates (reference to a privately owned object). We'll manually serialize the tracks to copy them,
+				// and then restore them on paste.
 				UMovieSceneTrack* DuplicatedTrack = Cast<UMovieSceneTrack>(StaticDuplicateObject(Track, CopyableBinding));
 
 				CopyableBinding->Tracks.Add(DuplicatedTrack);
@@ -5765,7 +5784,7 @@ void FSequencer::PasteCopiedTracks()
 	UObject* BindingContext = GetPlaybackContext();
 
 	TArray<UMovieSceneCopyableBinding*> ImportedBindings;
-	ImportObjectsFromText(*this, TextToImport, ImportedBindings);
+	ImportObjectsFromText(TextToImport, ImportedBindings);
 
 	if (ImportedBindings.Num() != 0)
 	{
@@ -5776,12 +5795,13 @@ void FSequencer::PasteCopiedTracks()
 		TArray<FMovieSceneBinding> BindingsPasted;
 		for (UMovieSceneCopyableBinding* CopyableBinding : ImportedBindings)
 		{
-			FGuid NewGuid = FGuid::NewGuid();
-
-			FMovieSceneBinding NewBinding(NewGuid, CopyableBinding->Binding.GetName(), CopyableBinding->Tracks);
 			
 			if (CopyableBinding->Possessable.GetGuid().IsValid())
 			{
+				FGuid NewGuid = FGuid::NewGuid();
+
+				FMovieSceneBinding NewBinding(NewGuid, CopyableBinding->Binding.GetName(), CopyableBinding->Tracks);
+
 				FMovieScenePossessable NewPossessable = CopyableBinding->Possessable;
 				NewPossessable.SetGuid(NewGuid);
 
@@ -5795,13 +5815,31 @@ void FSequencer::PasteCopiedTracks()
 			}
 			else if (CopyableBinding->Spawnable.GetGuid().IsValid())
 			{
-				FMovieSceneSpawnable NewSpawnable = CopyableBinding->Spawnable;
-				NewSpawnable.SetGuid(NewGuid);
-
-				MovieScene->AddSpawnable(NewSpawnable, NewBinding);
+				// We need to let the sequence create the spawnable so that it has everything set up properly internally.
+				// This is required to get spawnables with the correct references to object templates, object templates with
+				// correct owners, etc. However, making a new spawnable also creates the binding for us - this is a problem
+				// because we need to use our binding (which has tracks associated with it). To solve this, we let it create
+				// an object template based off of our (transient package owned) template, then find the newly created binding
+				// and update it.
+				FGuid NewGuid = MakeNewSpawnable(*CopyableBinding->SpawnableObjectTemplate);
+				FMovieSceneBinding NewBinding(NewGuid, CopyableBinding->Binding.GetName(), CopyableBinding->Tracks);
 
 				FMovieSceneSpawnable* Spawnable = MovieScene->FindSpawnable(NewGuid);
-				Spawnable->CopyObjectTemplate(*CopyableBinding->SpawnableObjectTemplate, *GetFocusedMovieSceneSequence());
+				// Copy the name of the original spawnable too.
+				Spawnable->SetName(CopyableBinding->Spawnable.GetName());
+
+				// Replace the auto-generated binding with our deserialized bindings (which has our tracks)
+				FMovieSceneBinding* Binding = MovieScene->FindBinding(NewGuid);
+				if (Binding)
+				{
+					*Binding = NewBinding;
+
+					// We also need to change the track's owners to be the MovieScene. Typically MovieScene::AddSpawnable/MovieScene::AddSpawnable/AddPossessable does this.
+					for (auto Track : Binding->GetTracks())
+					{
+						Track->Rename(nullptr, MovieScene);
+					}
+				}
 
 				OldToNewGuidMap.Add(CopyableBinding->Spawnable.GetGuid(), NewGuid);
 
diff --git a/Engine/Source/Editor/Sequencer/Private/Sequencer.h b/Engine/Source/Editor/Sequencer/Private/Sequencer.h
--- a/Engine/Source/Editor/Sequencer/Private/Sequencer.h
+++ b/Engine/Source/Editor/Sequencer/Private/Sequencer.h
@@ -60,6 +60,7 @@ class ULevel;
 class UMovieSceneSequence;
 class UMovieSceneSubSection;
 class USequencerSettings;
+class UMovieSceneCopyableBinding;
 struct FMovieSceneTimeController;
 struct FMovieScenePossessable;
 struct FSequencerTemplateStore;
@@ -466,6 +467,7 @@ public:
 	bool CanPaste(const FString& TextToImport);
 	void PasteCopiedTracks();
 	void ImportTracksFromText(const FString& TextToImport, /*out*/ TArray<UMovieSceneTrack*>& ImportedTrack);
+	void ImportObjectsFromText(const FString& TextToImport, /*out*/ TArray<UMovieSceneCopyableBinding*>& ImportedObjects);
 
 	/** Called when a user executes the active node menu item */
 	void ToggleNodeActive();
diff --git a/Engine/Source/Runtime/MovieScene/Public/MovieScene.h b/Engine/Source/Runtime/MovieScene/Public/MovieScene.h
--- a/Engine/Source/Runtime/MovieScene/Public/MovieScene.h
+++ b/Engine/Source/Runtime/MovieScene/Public/MovieScene.h
@@ -451,6 +451,16 @@ public:
 	 */
 	void MoveBindingContents(const FGuid& SourceBindingId, const FGuid& DestinationBindingId);
 
+	/**
+	 * Tries to find an FMovieSceneBinding for the specified Guid.
+	 * 
+	 * @param ForGuid	The binding's Guid to look for.
+	 * @return			Pointer to the binding, otherwise nullptr.
+	 */
+	FMovieSceneBinding* FindBinding(const FGuid& ForGuid)
+	{
+		return ObjectBindings.FindByPredicate([ForGuid](const FMovieSceneBinding& Binding) { return Binding.GetObjectGuid() == ForGuid; });
+	}
 public:
 
 	// @todo sequencer: the following methods really shouldn't be here
diff --git a/Engine/Source/Runtime/MovieScene/Public/MovieSceneBinding.h b/Engine/Source/Runtime/MovieScene/Public/MovieSceneBinding.h
--- a/Engine/Source/Runtime/MovieScene/Public/MovieSceneBinding.h
+++ b/Engine/Source/Runtime/MovieScene/Public/MovieSceneBinding.h
@@ -129,6 +129,14 @@ struct FMovieSceneBinding
 		Tracks = MoveTemp(InTracks);
 	}
 
+	/**
+	* Assign all tracks in this binding
+	*/
+	void SetTracks(TArray<UMovieSceneTrack*>& InTracks)
+	{
+		Tracks = InTracks;
+	}
+
 #if WITH_EDITOR
 	/**
 	 * Perform cook-time optimization on this object binding
diff --git a/Engine/Source/Runtime/MovieScene/Public/MovieSceneSpawnable.h b/Engine/Source/Runtime/MovieScene/Public/MovieSceneSpawnable.h
--- a/Engine/Source/Runtime/MovieScene/Public/MovieSceneSpawnable.h
+++ b/Engine/Source/Runtime/MovieScene/Public/MovieSceneSpawnable.h
@@ -89,6 +89,15 @@ public:
 		return ObjectTemplate;
 	}
 
+	/**
+	 * Set the object template to the specified object directly.
+	 * Used for Copy/Paste, typically you should use CopyObjectTemplate.
+	 */
+	void SetObjectTemplate(UObject* InObjectTemplate)
+	{
+		ObjectTemplate = InObjectTemplate;
+	}
+
 	/**
 	 * Copy the specified object into this spawnable's template
 	 *
