Sequencer: Set actor label when converting a spawnable to a possessable.

diff --git a/Engine/Source/Editor/Sequencer/Private/Sequencer.cpp b/Engine/Source/Editor/Sequencer/Private/Sequencer.cpp
--- a/Engine/Source/Editor/Sequencer/Private/Sequencer.cpp
+++ b/Engine/Source/Editor/Sequencer/Private/Sequencer.cpp
@@ -6483,6 +6483,8 @@ FMovieScenePossessable* FSequencer::ConvertToPossessableInternal(FGuid Spawnable
 		return nullptr;
 	}
 
+	PossessedActor->SetActorLabel(Spawnable->GetName());
+
 	const bool bIsDefaultTransform = true;
 	PossessedActor->FinishSpawning(SpawnTransform, bIsDefaultTransform);
 
