Sequencer: Fix multiple selection transactions when selecting keys.

diff --git a/Engine/Source/Editor/Sequencer/Private/SequencerEdMode.cpp b/Engine/Source/Editor/Sequencer/Private/SequencerEdMode.cpp
--- a/Engine/Source/Editor/Sequencer/Private/SequencerEdMode.cpp
+++ b/Engine/Source/Editor/Sequencer/Private/SequencerEdMode.cpp
@@ -137,6 +137,8 @@ void FSequencerEdMode::OnKeySelected(FViewport* Viewport, HMovieSceneKeyProxy* K
 
 	for (TWeakPtr<FSequencer> WeakSequencer : Sequencers)
 	{
+		bool bChangedSelection = false;
+
 		TSharedPtr<FSequencer> Sequencer = WeakSequencer.Pin();
 		if (Sequencer.IsValid())
 		{
@@ -145,6 +147,12 @@ void FSequencerEdMode::OnKeySelected(FViewport* Viewport, HMovieSceneKeyProxy* K
 			FSequencerSelection& Selection = Sequencer->GetSelection();
 			if (!bAddToSelection && !bToggleSelection)
 			{
+				if (!bChangedSelection)
+				{
+					Sequencer->GetSelection().SuspendBroadcast();
+					bChangedSelection = true;
+				}
+
 				Sequencer->GetSelection().EmptySelectedKeys();
 			}
 
@@ -169,6 +177,12 @@ void FSequencerEdMode::OnKeySelected(FViewport* Viewport, HMovieSceneKeyProxy* K
 							TSharedPtr<IKeyArea> KeyArea = KeyAreaNode->GetKeyArea(Section);
 							if (KeyArea.IsValid() && KeyArea->GetName() == KeyData.ChannelName)
 							{
+								if (!bChangedSelection)
+								{
+									Sequencer->GetSelection().SuspendBroadcast();
+									bChangedSelection = true;
+								}
+
 								Sequencer->SelectKey(Section, KeyArea, KeyData.KeyHandle.GetValue(), bToggleSelection);
 								goto next_key;
 							}
@@ -178,6 +192,12 @@ void FSequencerEdMode::OnKeySelected(FViewport* Viewport, HMovieSceneKeyProxy* K
 next_key:
 				continue;
 			}
+			if (bChangedSelection)
+			{
+				Sequencer->GetSelection().ResumeBroadcast();
+				Sequencer->GetSelection().GetOnKeySelectionChanged().Broadcast();
+				Sequencer->GetSelection().GetOnOutlinerNodeSelectionChangedObjectGuids().Broadcast();
+			}
 		}
 	}
 }
