Sequencer: Always bake at the display rate

(Otherwise, baking at tick resolution can be pretty heavy handed)

diff --git a/Engine/Source/Editor/Sequencer/Private/Sequencer.cpp b/Engine/Source/Editor/Sequencer/Private/Sequencer.cpp
--- a/Engine/Source/Editor/Sequencer/Private/Sequencer.cpp
+++ b/Engine/Source/Editor/Sequencer/Private/Sequencer.cpp
@@ -1359,7 +1359,7 @@ void FSequencer::BakeTransform()
 			TArray<FFrameNumber> KeyTimes;
 
 			FFrameRate   Resolution  = FocusedMovieScene->GetTickResolution();
-			FFrameRate   SnapRate    = FocusedMovieScene->GetEvaluationType() == EMovieSceneEvaluationType::FrameLocked ? FocusedMovieScene->GetDisplayRate() : Resolution;
+			FFrameRate   SnapRate    = FocusedMovieScene->GetDisplayRate();
 
 			FFrameNumber InFrame     = MovieScene::DiscreteInclusiveLower(GetPlaybackRange());
 			FFrameNumber OutFrame    = MovieScene::DiscreteExclusiveUpper(GetPlaybackRange());
