Sequencer: Remove Ctrl+O and Shift+O bindings to toggle node collapsing and expanding since they conflict with Open Level...

diff --git a/Engine/Source/Editor/Sequencer/Private/SequencerCommands.cpp b/Engine/Source/Editor/Sequencer/Private/SequencerCommands.cpp
--- a/Engine/Source/Editor/Sequencer/Private/SequencerCommands.cpp
+++ b/Engine/Source/Editor/Sequencer/Private/SequencerCommands.cpp
@@ -40,8 +40,8 @@ void FSequencerCommands::RegisterCommands()
 
 	UI_COMMAND( ExpandAllNodesAndDescendants, "Expand All Nodes", "Expand all nodes", EUserInterfaceActionType::Button, FInputChord() );
 	UI_COMMAND( CollapseAllNodesAndDescendants, "Collapse All Nodes", "Collapse all selected nodes", EUserInterfaceActionType::Button, FInputChord() );
-	UI_COMMAND( ToggleExpandCollapseNodes, "Expand/Collapse Nodes", "Toggle expand or collapse selected nodes", EUserInterfaceActionType::Button, FInputChord(EModifierKey::Control, EKeys::O) );
-	UI_COMMAND( ToggleExpandCollapseNodesAndDescendants, "Expand/Collapse Nodes and Descendants", "Toggle expand or collapse selected nodes and descendants", EUserInterfaceActionType::Button, FInputChord(EModifierKey::Shift, EKeys::O) );
+	UI_COMMAND( ToggleExpandCollapseNodes, "Expand/Collapse Nodes", "Toggle expand or collapse selected nodes", EUserInterfaceActionType::Button, FInputChord() );
+	UI_COMMAND( ToggleExpandCollapseNodesAndDescendants, "Expand/Collapse Nodes and Descendants", "Toggle expand or collapse selected nodes and descendants", EUserInterfaceActionType::Button, FInputChord() );
 	UI_COMMAND( SortAllNodesAndDescendants, "Sort All Nodes", "Sorts all nodes by type and then alphabetically.", EUserInterfaceActionType::Button, FInputChord());
 
 	UI_COMMAND( SetSelectionRangeEnd, "Set Selection End", "Sets the end of the selection range", EUserInterfaceActionType::Button, FInputChord(EKeys::O) );
