Sequencer: Intersect the section range with the visit range to ensure keys that fall outside the section range aren't snapped against.

diff --git a/Engine/Source/Editor/Sequencer/Private/Tools/SequencerEntityVisitor.cpp b/Engine/Source/Editor/Sequencer/Private/Tools/SequencerEntityVisitor.cpp
--- a/Engine/Source/Editor/Sequencer/Private/Tools/SequencerEntityVisitor.cpp
+++ b/Engine/Source/Editor/Sequencer/Private/Tools/SequencerEntityVisitor.cpp
@@ -236,6 +236,9 @@ void FSequencerEntityWalker::HandleKeyArea(const ISequencerEntityVisitor& Visito
 	const FFrameTime RangeEndFrame       = Range.Range.GetUpperBoundValue() * Range.TickResolution;
 
 	TRange<FFrameNumber> VisitRangeFrames( (RangeStartFrame-HalfKeySizeFrames).CeilToFrame(), (RangeEndFrame+HalfKeySizeFrames).FloorToFrame() );
+
+	VisitRangeFrames = TRange<FFrameNumber>::Intersection(Section->GetRange(), VisitRangeFrames);
+
 	KeyArea->GetKeyInfo(&Handles, &Times, VisitRangeFrames);
 
 	for (int32 Index = 0; Index < Times.Num(); ++Index)
