Sequencer: Added ability to manually invalidate eveluation field segments through the debug visualizer
   - Also fixed layout issues with DPI scaling turned on

diff --git a/Engine/Source/Editor/Sequencer/Private/SSequencerDebugVisualizer.cpp b/Engine/Source/Editor/Sequencer/Private/SSequencerDebugVisualizer.cpp
--- a/Engine/Source/Editor/Sequencer/Private/SSequencerDebugVisualizer.cpp
+++ b/Engine/Source/Editor/Sequencer/Private/SSequencerDebugVisualizer.cpp
@@ -7,6 +7,8 @@
 #include "Widgets/Layout/SBorder.h"
 #include "Widgets/Text/STextBlock.h"
 #include "Layout/ArrangedChildren.h"
+#include "Framework/MultiBox/MultiBoxBuilder.h"
+#include "Framework/Application/SlateApplication.h"
 #include "MovieSceneSequence.h"
 #include "MovieScene.h"
 #include "MovieSceneTimeHelpers.h"
@@ -107,6 +109,7 @@ void SSequencerDebugVisualizer::Refresh()
 				SNew(SBorder)
 				.BorderImage(SectionBackgroundBrush)
 				.Padding(FMargin(1.f))
+				.OnMouseButtonUp(this, &SSequencerDebugVisualizer::OnSlotMouseButtonUp, Index)
 				[
 					SNew(SBorder)
 					.BorderImage(SectionBackgroundTintBrush)
@@ -139,7 +142,7 @@ FGeometry SSequencerDebugVisualizer::GetSegmentGeometry(const FGeometry& Allotte
 	TRange<FFrameNumber> SegmentRange = ActiveTemplate->EvaluationField.GetRange(Slot.GetSegmentIndex());
 
 	float PixelStartX = SegmentRange.GetLowerBound().IsOpen() ? 0.f : TimeToPixelConverter.FrameToPixel(MovieScene::DiscreteInclusiveLower(SegmentRange));
-	float PixelEndX   = SegmentRange.GetUpperBound().IsOpen() ? AllottedGeometry.GetDrawSize().X : TimeToPixelConverter.FrameToPixel(MovieScene::DiscreteExclusiveUpper(SegmentRange));
+	float PixelEndX   = SegmentRange.GetUpperBound().IsOpen() ? AllottedGeometry.GetLocalSize().X : TimeToPixelConverter.FrameToPixel(MovieScene::DiscreteExclusiveUpper(SegmentRange));
 
 	const float MinSectionWidth = 0.f;
 	float SectionLength = FMath::Max(MinSectionWidth, PixelEndX - PixelStartX);
@@ -236,11 +239,55 @@ void SSequencerDebugVisualizer::OnArrangeChildren( const FGeometry& AllottedGeom
 			{
 				ArrangedChildren.AddWidget( 
 					WidgetVisibility, 
-					AllottedGeometry.MakeChild(Child, SegmentGeometry.Position, SegmentGeometry.GetDrawSize())
+					AllottedGeometry.MakeChild(Child, SegmentGeometry.Position, SegmentGeometry.GetLocalSize())
 					);
 			}
 		}
 	}
 }
 
+FReply SSequencerDebugVisualizer::OnSlotMouseButtonUp(const FGeometry& MyGeometry, const FPointerEvent& MouseEvent, int32 SlotIndex)
+{
+	FMenuBuilder MenuBuilder(true, nullptr);
+
+	MenuBuilder.AddMenuEntry(
+		LOCTEXT("InvalidateSegment", "Invalidate Segment"),
+		FText(),
+		FSlateIcon(),
+		FUIAction(FExecuteAction::CreateSP(this, &SSequencerDebugVisualizer::InvalidateSegment, SlotIndex))
+	);
+
+	int32 IndexNone = INDEX_NONE;
+
+	MenuBuilder.AddMenuEntry(
+		LOCTEXT("InvalidateAll", "Invalidate All"),
+		FText(),
+		FSlateIcon(),
+		FUIAction(FExecuteAction::CreateSP(this, &SSequencerDebugVisualizer::InvalidateSegment, IndexNone))
+	);
+
+	FWidgetPath WidgetPath = MouseEvent.GetEventPath() != nullptr ? *MouseEvent.GetEventPath() : FWidgetPath();
+	FSlateApplication::Get().PushMenu(AsShared(), WidgetPath, MenuBuilder.MakeWidget(), MouseEvent.GetScreenSpacePosition(), FPopupTransitionEffect(FPopupTransitionEffect::ContextMenu));
+
+	return FReply::Handled();
+}
+
+void SSequencerDebugVisualizer::InvalidateSegment(int32 SlotIndex)
+{
+	TSharedPtr<FSequencer> Sequencer = this->WeakSequencer.Pin();
+	FMovieSceneEvaluationTemplate* Template = Sequencer.IsValid() ? Sequencer->GetEvaluationTemplate().FindTemplate(Sequencer->GetFocusedTemplateID()) : nullptr;
+
+	if (Template)
+	{
+		if (SlotIndex == INDEX_NONE)
+		{
+			Template->EvaluationField.Invalidate(TRange<FFrameNumber>::All());
+		}
+		else if (SlotIndex < Template->EvaluationField.Size())
+		{
+			Template->EvaluationField.Invalidate(Template->EvaluationField.GetRange(SlotIndex));
+		}
+	}
+}
+
 #undef LOCTEXT_NAMESPACE
diff --git a/Engine/Source/Editor/Sequencer/Private/SSequencerDebugVisualizer.h b/Engine/Source/Editor/Sequencer/Private/SSequencerDebugVisualizer.h
--- a/Engine/Source/Editor/Sequencer/Private/SSequencerDebugVisualizer.h
+++ b/Engine/Source/Editor/Sequencer/Private/SSequencerDebugVisualizer.h
@@ -75,6 +75,10 @@ protected:
 
 	const FMovieSceneEvaluationTemplate* GetTemplate() const;
 
+	FReply OnSlotMouseButtonUp(const FGeometry& MyGeometry, const FPointerEvent& MouseEvent, int32 SlotIndex);
+
+	void InvalidateSegment(int32 SlotIndex);
+
 private:
 
 	/** The current view range */
