Removes "Show Transport Controls" from the Level Editor Viewport UI as the controls have been completely surpassed by the Cinematic Viewport Preview and you had to have Sequenecer open for the controls to be visible (so they were just a duplicate of the Sequencer window!)

diff --git a/Engine/Plugins/MovieScene/LevelSequenceEditor/Source/LevelSequenceEditor/Private/CinematicViewport/SCinematicLevelViewport.h b/Engine/Plugins/MovieScene/LevelSequenceEditor/Source/LevelSequenceEditor/Private/CinematicViewport/SCinematicLevelViewport.h
--- a/Engine/Plugins/MovieScene/LevelSequenceEditor/Source/LevelSequenceEditor/Private/CinematicViewport/SCinematicLevelViewport.h
+++ b/Engine/Plugins/MovieScene/LevelSequenceEditor/Source/LevelSequenceEditor/Private/CinematicViewport/SCinematicLevelViewport.h
@@ -31,9 +31,6 @@ struct FCinematicViewportClient : FLevelEditorViewportClient
 	FCinematicViewportClient();
 
 	void SetViewportWidget(const TSharedPtr<SEditorViewport>& InViewportWidget) { EditorViewportWidget = InViewportWidget; }
-
-	/** FLevelEditorViewportClient Interface*/
-	virtual bool CanAttachTransportControls() const override { return false; }
 };
 
 
diff --git a/Engine/Plugins/MovieScene/LevelSequenceEditor/Source/LevelSequenceEditor/Private/LevelSequenceEditorToolkit.h b/Engine/Plugins/MovieScene/LevelSequenceEditor/Source/LevelSequenceEditor/Private/LevelSequenceEditorToolkit.h
--- a/Engine/Plugins/MovieScene/LevelSequenceEditor/Source/LevelSequenceEditor/Private/LevelSequenceEditorToolkit.h
+++ b/Engine/Plugins/MovieScene/LevelSequenceEditor/Source/LevelSequenceEditor/Private/LevelSequenceEditorToolkit.h
@@ -154,9 +154,6 @@ private:
 	/** The sequencer used by this editor. */
 	TSharedPtr<ISequencer> Sequencer;
 
-	/** A map of all the transport controls to viewports that this sequencer has made */
-	TMap<TWeakPtr<ILevelViewport>, TSharedPtr<SWidget>> TransportControls;
-
 	FDelegateHandle SequencerExtenderHandle;
 
 	/** Pointer to the style set to use for toolkits. */
diff --git a/Engine/Source/Editor/Sequencer/Private/LevelEditorSequencerIntegration.cpp b/Engine/Source/Editor/Sequencer/Private/LevelEditorSequencerIntegration.cpp
--- a/Engine/Source/Editor/Sequencer/Private/LevelEditorSequencerIntegration.cpp
+++ b/Engine/Source/Editor/Sequencer/Private/LevelEditorSequencerIntegration.cpp
@@ -216,7 +216,6 @@ void FLevelEditorSequencerIntegration::Initialize()
 
 	AddLevelViewportMenuExtender();
 	ActivateDetailHandler();
-	AttachTransportControlsToViewports();
 	ActivateSequencerEditorMode();
 	BindLevelEditorCommands();
 	AttachOutlinerColumn();
@@ -849,277 +848,6 @@ private:
 	FCurveSequence FadeInSequence;
 };
 
-class SViewportTransportControls : public SCompoundWidget
-{
-public:
-	SLATE_BEGIN_ARGS(SViewportTransportControls){}
-	SLATE_END_ARGS()
-
-	void Construct(const FArguments& InArgs, TSharedPtr<ILevelViewport> Viewport)
-	{
-		check(Viewport.IsValid());
-
-		WeakViewport = Viewport;
-
-		SetVisibility(EVisibility::SelfHitTestInvisible);
-
-		ChildSlot
-		[
-			SNew(SHorizontalBox)
-
-			+SHorizontalBox::Slot()
-			.FillWidth(1)
-			.HAlign(HAlign_Center)
-			.VAlign(VAlign_Bottom)
-			.Padding(4.f)
-			[
-				SNew(SFader)
-				.Content()
-				[
-					SNew(SBorder)
-					.Padding(4.f)
-					.Cursor( EMouseCursor::Default )
-					.BorderImage( FEditorStyle::GetBrush( "ToolPanel.GroupBorder" ) )
-					.Visibility(this, &SViewportTransportControls::GetVisibility)
-					.Content()
-					[
-						SNew(SHorizontalBox)
-
-						+SHorizontalBox::Slot()
-						.AutoWidth()
-						.HAlign(HAlign_Center)
-						.VAlign(VAlign_Center)
-						.Padding(FMargin(0,0,4.f,0))
-						[
-							SNew(SComboButton)
-							.Visibility(this, &SViewportTransportControls::GetComboVisibility)
-							.OnGetMenuContent(this, &SViewportTransportControls::GetBoundSequencerMenu)
-							.ButtonContent()
-							[
-								SNew(STextBlock)
-								.Text(this, &SViewportTransportControls::GetBoundSequencerName)
-							]
-						]
-
-						+SHorizontalBox::Slot()
-						.AutoWidth()
-						.HAlign(HAlign_Center)
-						.VAlign(VAlign_Center)
-						[
-							SAssignNew(ControlContent, SBox)
-						]
-						+SHorizontalBox::Slot()
-						.AutoWidth()
-						.HAlign(HAlign_Center)
-						.VAlign(VAlign_Top)
-						[
-							SNew(SButton)
-							.ButtonStyle(&FEditorStyle::Get().GetWidgetStyle< FButtonStyle >("Sequencer.Transport.CloseButton"))
-							.ToolTipText(LOCTEXT("CloseTransportControlsToolTip", "Hide the transport controls. You can re-enable transport controls from the viewport menu."))
-							.OnClicked(this, &SViewportTransportControls::Close)
-						]
-					]
-				]
-			]
-		];
-	}
-
-	TSharedPtr<FSequencer> GetSequencer() const
-	{
-		return WeakSequencer.Pin();
-	}
-
-	void AssignSequencer(TSharedRef<FSequencer> InSequencer)
-	{
-		WeakSequencer = InSequencer;
-		bool bExtendedControls = false;
-		ControlContent->SetContent(InSequencer->MakeTransportControls(bExtendedControls));
-	}
-
-	EVisibility GetComboVisibility() const
-	{
-		int32 NumSequencers = 0;
-		FLevelEditorSequencerIntegration::Get().IterateAllSequencers(
-			[&NumSequencers](FSequencer& In, const FLevelEditorSequencerIntegrationOptions& Options)
-			{
-				if (UMovieSceneSequence* RootSequence = In.GetRootMovieSceneSequence())
-				{
-					++NumSequencers;
-				}
-			}
-		);
-
-		return NumSequencers == 1 ? EVisibility::Collapsed : EVisibility::Visible;
-	}
-
-	FText GetBoundSequencerName() const
-	{
-		TSharedPtr<FSequencer> Sequencer = WeakSequencer.Pin();
-		if (Sequencer.IsValid() && Sequencer->GetRootMovieSceneSequence())
-		{
-			return Sequencer->GetRootMovieSceneSequence()->GetDisplayName();
-		}
-
-		return LOCTEXT("SelectSequencer", "Choose Sequence...");
-	}
-
-	TSharedRef<SWidget> GetBoundSequencerMenu()
-	{
-		FMenuBuilder MenuBuilder(true, nullptr);
-
-		for (const FLevelEditorSequencerIntegration::FSequencerAndOptions& SequencerAndOptions : FLevelEditorSequencerIntegration::Get().BoundSequencers)
-		{
-			TSharedPtr<FSequencer> PinnedSequencer = SequencerAndOptions.Sequencer.Pin();
-			if (PinnedSequencer.IsValid())
-			{
-				UMovieSceneSequence* RootSequence = PinnedSequencer->GetRootMovieSceneSequence();
-
-				// Be careful not to hold a strong reference in the lamdba below
-				TWeakPtr<FSequencer> ThisWeakSequencer = SequencerAndOptions.Sequencer;
-
-				MenuBuilder.AddMenuEntry(
-					RootSequence->GetDisplayName(),
-					FText(),
-					FSlateIconFinder::FindIconForClass(RootSequence->GetClass(), "MovieSceneSequence"),
-					FUIAction(
-						FExecuteAction::CreateLambda(
-							[this, ThisWeakSequencer]
-							{
-								TSharedPtr<FSequencer> LocalPinnedSequence = ThisWeakSequencer.Pin();
-								if (LocalPinnedSequence.IsValid())
-								{
-									AssignSequencer(LocalPinnedSequence.ToSharedRef());
-								}
-							}
-						)
-					));
-			}
-		}
-
-		return MenuBuilder.MakeWidget();
-	}
-
-	FReply Close()
-	{
-		TSharedPtr<FSequencer> Sequencer = WeakSequencer.Pin();
-		if (Sequencer.IsValid())
-		{
-			Sequencer->GetSequencerSettings()->SetShowViewportTransportControls(!Sequencer->GetSequencerSettings()->GetShowViewportTransportControls());
-		}
-		return FReply::Handled();
-	}
-
-	EVisibility GetVisibility() const
-	{
-		TSharedPtr<ILevelViewport> Viewport = WeakViewport.Pin();
-		TSharedPtr<FSequencer> Sequencer = WeakSequencer.Pin();
-		FLevelEditorViewportClient* ViewportClient = Viewport.IsValid() ? &Viewport->GetLevelViewportClient() : nullptr;
-
-		return (
-				Sequencer.IsValid() &&
-				ViewportClient &&
-				Sequencer->GetSequencerSettings()->GetShowViewportTransportControls() &&
-				ViewportClient->ViewportType == LVT_Perspective &&
-				ViewportClient->AllowsCinematicPreview()
-			) ? EVisibility::Visible : EVisibility::Collapsed;
-	}
-
-	virtual bool SupportsKeyboardFocus() const override
-	{
-		// Transport controls in the viewport need to have something that is focusable to prevent mouse input dropping through to the viewport.
-		// We don't want the buttons themselves to be focusable, so we just add them to a parent box that is.
-		return true;
-	}
-
-private:
-
-	TSharedPtr<SBox> ControlContent;
-	TWeakPtr<ILevelViewport> WeakViewport;
-	TWeakPtr<FSequencer> WeakSequencer;
-};
-
-void FLevelEditorSequencerIntegration::AttachTransportControlsToViewports()
-{
-	FLevelEditorModule* Module = FModuleManager::Get().LoadModulePtr<FLevelEditorModule>("LevelEditor");
-	
-	if (Module == nullptr)
-	{
-		return;
-	}
-
-	// register level editor viewport menu extenders
-	{
-		FLevelEditorModule::FLevelEditorMenuExtender ViewMenuExtender = FLevelEditorModule::FLevelEditorMenuExtender::CreateRaw(this, &FLevelEditorSequencerIntegration::OnExtendLevelEditorViewMenu);
-		Module->GetAllLevelViewportOptionsMenuExtenders().Add(ViewMenuExtender);
-
-		FDelegateHandle Handle = Module->GetAllLevelViewportOptionsMenuExtenders().Last().GetHandle();
-		AcquiredResources.Add(
-			[Handle]{
-				FLevelEditorModule* LevelEditorModule = FModuleManager::Get().GetModulePtr<FLevelEditorModule>("LevelEditor");
-				if (LevelEditorModule)
-				{
-					LevelEditorModule->GetAllLevelViewportOptionsMenuExtenders().RemoveAll([Handle](const FLevelEditorModule::FLevelEditorMenuExtender& In){ return In.GetHandle() == Handle; });
-				}
-			}
-		);
-	}
-
-	TSharedPtr<ILevelEditor> LevelEditor = Module->GetFirstLevelEditor();
-
-	for (TSharedPtr<ILevelViewport> LevelViewport : LevelEditor->GetViewports())
-	{
-		if (LevelViewport->GetLevelViewportClient().CanAttachTransportControls())
-		{
-			TSharedRef<SViewportTransportControls> TransportControl = SNew(SViewportTransportControls, LevelViewport);
-			LevelViewport->AddOverlayWidget(TransportControl);
-			TransportControls.Add(FTransportControl{ LevelViewport, TransportControl });
-		}
-	}
-
-	AcquiredResources.Add([=]{ this->DetachTransportControlsFromViewports(); });
-}
-
-void FLevelEditorSequencerIntegration::SetViewportTransportControlsVisibility(bool bVisible)
-{
-	IterateAllSequencers(
-		[=](FSequencer& In, const FLevelEditorSequencerIntegrationOptions& Options)
-		{
-			In.GetSequencerSettings()->SetShowViewportTransportControls(bVisible);
-		}
-	);
-}
-
-bool FLevelEditorSequencerIntegration::GetViewportTransportControlsVisibility() const
-{
-	bool bVisible = false;
-
-	IterateAllSequencers(
-		[&](FSequencer& In, const FLevelEditorSequencerIntegrationOptions& Options)
-		{
-			if (In.GetSequencerSettings()->GetShowViewportTransportControls())
-			{
-				bVisible = true;
-			}
-		}
-	);
-
-	return bVisible;
-}
-
-void FLevelEditorSequencerIntegration::DetachTransportControlsFromViewports()
-{
-	for (const FTransportControl& Control : TransportControls)
-	{
-		TSharedPtr<ILevelViewport> Viewport = Control.Viewport.Pin();
-		if (Viewport.IsValid())
-		{
-			Viewport->RemoveOverlayWidget(Control.Widget.ToSharedRef());
-		}
-	}
-	TransportControls.Reset();
-}
-
-
 TSharedRef< ISceneOutlinerColumn > FLevelEditorSequencerIntegration::CreateSequencerInfoColumn( ISceneOutliner& SceneOutliner ) const
 {
 	//@todo only supports the first bound sequencer
@@ -1223,71 +951,14 @@ void FLevelEditorSequencerIntegration::RestoreRealtimeViewports()
 
 TSharedRef<FExtender> FLevelEditorSequencerIntegration::OnExtendLevelEditorViewMenu(const TSharedRef<FUICommandList> CommandList)
 {
+	// This would be where you added an extension to the Level Editor dropdown menus if needed.
 	TSharedRef<FExtender> Extender(new FExtender());
-
-	Extender->AddMenuExtension(
-		"LevelViewportViewportOptions2",
-		EExtensionHook::First,
-		NULL,
-		FMenuExtensionDelegate::CreateRaw(this, &FLevelEditorSequencerIntegration::CreateTransportToggleMenuEntry));
-
 	return Extender;
 }
 
-void FLevelEditorSequencerIntegration::CreateTransportToggleMenuEntry(FMenuBuilder& MenuBuilder)
-{
-	MenuBuilder.AddMenuEntry(
-		LOCTEXT("ShowTransportControls", "Show Transport Controls"),
-		LOCTEXT("ShowTransportControlsToolTip", "Show or hide the Sequencer transport controls when a sequence is active."),
-		FSlateIcon(),
-		FUIAction(
-			FExecuteAction::CreateLambda([this](){ SetViewportTransportControlsVisibility(!GetViewportTransportControlsVisibility()); }),
-			FCanExecuteAction(), 
-			FGetActionCheckState::CreateLambda([this](){ return GetViewportTransportControlsVisibility() ? ECheckBoxState::Checked : ECheckBoxState::Unchecked; })),
-		NAME_None,
-		EUserInterfaceActionType::ToggleButton);
-}
-
 void FLevelEditorSequencerIntegration::OnTabContentChanged()
-{	
-	for (const FTransportControl& Control : TransportControls)
-	{
-		TSharedPtr<ILevelViewport> Viewport = Control.Viewport.Pin();
-		if (Viewport.IsValid())
-		{
-			Viewport->RemoveOverlayWidget(Control.Widget.ToSharedRef());
-		}
-	}
-	TransportControls.Reset();
-
-	FLevelEditorModule* Module = FModuleManager::Get().LoadModulePtr<FLevelEditorModule>("LevelEditor");
-	
-	if (Module == nullptr)
-	{
-		return;
-	}
-
-	TSharedPtr<ILevelEditor> LevelEditor = Module->GetFirstLevelEditor();
-
-	for (TSharedPtr<ILevelViewport> LevelViewport : LevelEditor->GetViewports())
-	{
-		if (LevelViewport->GetLevelViewportClient().CanAttachTransportControls())
-		{
-			TSharedRef<SViewportTransportControls> TransportControl = SNew(SViewportTransportControls, LevelViewport);
-			LevelViewport->AddOverlayWidget(TransportControl);
-
-			auto IsValidSequencer = [](const FSequencerAndOptions& In){ return In.Sequencer.IsValid(); };
-			if (FSequencerAndOptions* FirstValidSequencer = BoundSequencers.FindByPredicate(IsValidSequencer))
-			{
-				TSharedPtr<FSequencer> SequencerPtr = FirstValidSequencer->Sequencer.Pin();
-				check(SequencerPtr.IsValid());
-				
-				TransportControl->AssignSequencer(SequencerPtr.ToSharedRef());
-			}
+{
 
-			TransportControls.Add(FTransportControl{ LevelViewport, TransportControl });
-		}
-	}
 }
 
 void FLevelEditorSequencerIntegration::AddSequencer(TSharedRef<ISequencer> InSequencer, const FLevelEditorSequencerIntegrationOptions& Options)
@@ -1345,15 +1016,6 @@ void FLevelEditorSequencerIntegration::AddSequencer(TSharedRef<ISequencer> InSeq
 		SequencerEdMode->AddSequencer(DerivedSequencerPtr);
 	}
 
-	// Set up any transport controls
-	for (const FTransportControl& Control : TransportControls)
-	{
-		if (!Control.Widget->GetSequencer().IsValid())
-		{
-			Control.Widget->AssignSequencer(DerivedSequencerPtr);
-		}
-	}
-
 	ActivateRealtimeViewports();
 }
 
@@ -1385,21 +1047,7 @@ void FLevelEditorSequencerIntegration::RemoveSequencer(TSharedRef<ISequencer> In
 	KeyFrameHandler->Remove(InSequencer);
 
 	auto IsValidSequencer = [](const FSequencerAndOptions& In){ return In.Sequencer.IsValid(); };
-	if (FSequencerAndOptions* FirstValidSequencer = BoundSequencers.FindByPredicate(IsValidSequencer))
-	{
-		TSharedPtr<FSequencer> SequencerPtr = FirstValidSequencer->Sequencer.Pin();
-		check(SequencerPtr.IsValid());
-
-		// Assign any transport controls
-		for (const FTransportControl& Control : TransportControls)
-		{
-			if (!Control.Widget->GetSequencer().IsValid())
-			{
-				Control.Widget->AssignSequencer(SequencerPtr.ToSharedRef());
-			}
-		}
-	}
-	else
+	if (!BoundSequencers.FindByPredicate(IsValidSequencer))
 	{
 		AcquiredResources.Release();
 	}
diff --git a/Engine/Source/Editor/Sequencer/Private/SequencerSettings.cpp b/Engine/Source/Editor/Sequencer/Private/SequencerSettings.cpp
--- a/Engine/Source/Editor/Sequencer/Private/SequencerSettings.cpp
+++ b/Engine/Source/Editor/Sequencer/Private/SequencerSettings.cpp
@@ -41,7 +41,6 @@ USequencerSettings::USequencerSettings( const FObjectInitializer& ObjectInitiali
 	bShowCombinedKeyframes = true;
 	bInfiniteKeyAreas = false;
 	bShowChannelColors = false;
-	bShowViewportTransportControls = true;
 	bAllowPossessionOfPIEViewports = false;
 	bActivateRealtimeViewports = true;
 	bEvaluateSubSequencesInIsolation = false;
@@ -527,21 +526,6 @@ void USequencerSettings::SetShowChannelColors(bool InbShowChannelColors)
 	}
 }
 
-bool USequencerSettings::GetShowViewportTransportControls() const
-{
-	return bShowViewportTransportControls;
-}
-
-void USequencerSettings::SetShowViewportTransportControls(bool bVisible)
-{
-	if (bShowViewportTransportControls != bVisible)
-	{
-		bShowViewportTransportControls = bVisible;
-		SaveConfig();
-	}
-}
-
-
 bool USequencerSettings::ShouldAllowPossessionOfPIEViewports() const
 {
 	return bAllowPossessionOfPIEViewports;
diff --git a/Engine/Source/Editor/Sequencer/Public/LevelEditorSequencerIntegration.h b/Engine/Source/Editor/Sequencer/Public/LevelEditorSequencerIntegration.h
--- a/Engine/Source/Editor/Sequencer/Public/LevelEditorSequencerIntegration.h
+++ b/Engine/Source/Editor/Sequencer/Public/LevelEditorSequencerIntegration.h
@@ -17,7 +17,6 @@ class FSequencer;
 class FUICommandList;
 class ILevelViewport;
 class ISequencer;
-class SViewportTransportControls;
 class ULevel;
 struct FPropertyAndParent;
 
@@ -82,10 +81,6 @@ public:
 
 	void RemoveSequencer(TSharedRef<ISequencer> InSequencer);
 
-	void SetViewportTransportControlsVisibility(bool bVisible);
-
-	bool GetViewportTransportControlsVisibility() const;
-
 private:
 
 	/** Called before the world is going to be saved. The sequencer puts everything back to its initial state. */
@@ -159,11 +154,6 @@ private:
 
 	void RecordSelectedActors();
 
-	EVisibility GetTransportControlVisibility(TSharedPtr<ILevelViewport> LevelViewport) const;
-	
-	/** Create a menu entry we can use to toggle the transport controls */
-	void CreateTransportToggleMenuEntry(FMenuBuilder& MenuBuilder);
-
 	bool IsPropertyReadOnly(const FPropertyAndParent& InPropertyAndParent);
 
 private:
@@ -171,8 +161,6 @@ private:
 	void ActivateSequencerEditorMode();
 	void AddLevelViewportMenuExtender();
 	void ActivateDetailHandler();
-	void AttachTransportControlsToViewports();
-	void DetachTransportControlsFromViewports();
 	void AttachOutlinerColumn();
 	void DetachOutlinerColumn();
 	void ActivateRealtimeViewports();
@@ -198,17 +186,6 @@ private:
 	FLevelEditorSequencerIntegration();
 
 private:
-
-	friend SViewportTransportControls;
-	
-	/** A map of all the transport controls to viewports that this sequencer has made */
-	struct FTransportControl
-	{
-		TWeakPtr<ILevelViewport> Viewport;
-		TSharedPtr<SViewportTransportControls> Widget;
-	};
-	TArray<FTransportControl> TransportControls;
-
 	FAcquiredResources AcquiredResources;
 
 	TSharedPtr<class FDetailKeyframeHandlerWrapper> KeyFrameHandler;
diff --git a/Engine/Source/Editor/Sequencer/Public/SequencerSettings.h b/Engine/Source/Editor/Sequencer/Public/SequencerSettings.h
--- a/Engine/Source/Editor/Sequencer/Public/SequencerSettings.h
+++ b/Engine/Source/Editor/Sequencer/Public/SequencerSettings.h
@@ -278,11 +278,6 @@ public:
 	/** Set whether to show channel colors */
 	void SetShowChannelColors(bool bInShowChannelColors);
 
-	/** @return true if showing transport controls in level editor viewports */
-	bool GetShowViewportTransportControls() const;
-	/** Toggle whether to show transport controls in level editor viewports */
-	void SetShowViewportTransportControls(bool bVisible);
-
 	/** @return Whether to allow possession of PIE viewports */
 	bool ShouldAllowPossessionOfPIEViewports() const;
 	/** Toggle whether to allow possession of PIE viewports */
@@ -473,10 +468,6 @@ protected:
 	UPROPERTY( config, EditAnywhere, Category=Timeline )
 	bool bShowChannelColors;
 
-	/** Enable or disable transport controls in the viewport. */
-	UPROPERTY( config )
-	bool bShowViewportTransportControls;
-
 	/** When enabled, sequencer is able to possess viewports that represent PIE worlds */
 	UPROPERTY(config, EditAnywhere, Category=General)
 	bool bAllowPossessionOfPIEViewports;
diff --git a/Engine/Source/Editor/UnrealEd/Public/LevelEditorViewport.h b/Engine/Source/Editor/UnrealEd/Public/LevelEditorViewport.h
--- a/Engine/Source/Editor/UnrealEd/Public/LevelEditorViewport.h
+++ b/Engine/Source/Editor/UnrealEd/Public/LevelEditorViewport.h
@@ -552,9 +552,6 @@ public:
 	 */
 	static UObject* GetOrCreateMaterialFromTexture( UTexture* UnrealTexture );
 
-	/** Whether transport controls can be attached */
-	virtual bool CanAttachTransportControls() const { return true; }
-
 protected:
 	/** 
 	 * Checks the viewport to see if the given blueprint asset can be dropped on the viewport.
