Renaming Object Bindings in Sequencer now marks the package as modified (so they prompt for save requests) and pass the rename onto the Spawnable so that pasting object bindings doesn't revert to the original name when they were created.

diff --git a/Engine/Source/Editor/Sequencer/Private/DisplayNodes/SequencerObjectBindingNode.cpp b/Engine/Source/Editor/Sequencer/Private/DisplayNodes/SequencerObjectBindingNode.cpp
--- a/Engine/Source/Editor/Sequencer/Private/DisplayNodes/SequencerObjectBindingNode.cpp
+++ b/Engine/Source/Editor/Sequencer/Private/DisplayNodes/SequencerObjectBindingNode.cpp
@@ -485,7 +485,21 @@ void FSequencerObjectBindingNode::SetDisplayName(const FText& NewDisplayName)
 
 	if (MovieScene != nullptr)
 	{
+		// Modify the movie scene so that it gets marked dirty and renames are saved consistently.
+		MovieScene->Modify();
 		MovieScene->SetObjectDisplayName(ObjectBinding, NewDisplayName);
+
+		FMovieSceneSpawnable* Spawnable = MovieScene->FindSpawnable(GetObjectBinding());
+		if (Spawnable)
+		{
+			Spawnable->SetName(NewDisplayName.ToString());
+		}
+
+		FMovieScenePossessable* Possessable = MovieScene->FindPossessable(GetObjectBinding());
+		if (Possessable)
+		{
+			Possessable->SetName(NewDisplayName.ToString());
+		}
 	}
 }
 
