Sequencer Marks Implementation

Feature Documenation: https://docs.google.com/document/d/1dLcasYUWWQzQ_r-6lq90LiNmzT96EaoWuxq-3gHv4-Y/edit?usp=sharing

diff --git a/Engine/Source/Editor/Sequencer/Private/SSequencer.cpp b/Engine/Source/Editor/Sequencer/Private/SSequencer.cpp
--- a/Engine/Source/Editor/Sequencer/Private/SSequencer.cpp
+++ b/Engine/Source/Editor/Sequencer/Private/SSequencer.cpp
@@ -154,6 +154,9 @@ void SSequencer::Construct(const FArguments& InArgs, TSharedRef<FSequencer> InSe
 		TimeSliderArgs.OnScrubPositionChanged = InArgs._OnScrubPositionChanged;
 		TimeSliderArgs.PlaybackStatus = InArgs._PlaybackStatus;
 		TimeSliderArgs.SubSequenceRange = InArgs._SubSequenceRange;
+		TimeSliderArgs.MarkedFrames = InArgs._MarkedFrames;
+		TimeSliderArgs.OnMarkedFrameChanged = InArgs._OnMarkedFrameChanged;
+		TimeSliderArgs.OnClearAllMarkedFrames = InArgs._OnClearAllMarkedFrames;
 
 		TimeSliderArgs.Settings = Settings;
 		TimeSliderArgs.NumericTypeInterface = GetNumericTypeInterface();
@@ -478,6 +481,7 @@ void SSequencer::Construct(const FArguments& InArgs, TSharedRef<FSequencer> InSe
 						.Visibility( EVisibility::HitTestInvisible )
 						.DisplayScrubPosition( true )
 						.DisplayTickLines( false )
+						.DisplayMarkedFrames( true )
 						.PaintPlaybackRangeArgs(this, &SSequencer::GetSectionPlaybackRangeArgs)
 						.Clipping(EWidgetClipping::ClipToBounds)
 					]
diff --git a/Engine/Source/Editor/Sequencer/Private/SSequencer.h b/Engine/Source/Editor/Sequencer/Private/SSequencer.h
--- a/Engine/Source/Editor/Sequencer/Private/SSequencer.h
+++ b/Engine/Source/Editor/Sequencer/Private/SSequencer.h
@@ -109,6 +109,9 @@ public:
 		/** The selection range */
 		SLATE_ATTRIBUTE( TRange<FFrameNumber>, SelectionRange)
 
+		/** The Marked Frames */
+		SLATE_ATTRIBUTE( TSet<FFrameNumber>, MarkedFrames )
+
 		/** The current sub sequence range */
 		SLATE_ATTRIBUTE( TOptional<TRange<FFrameNumber>>, SubSequenceRange)
 
@@ -145,6 +148,12 @@ public:
 		/** Called when the user changes the view range */
 		SLATE_EVENT( FOnViewRangeChanged, OnViewRangeChanged )
 
+		/** Called when the user changes on the set of marked frames */
+		SLATE_EVENT( FOnMarkedFrameChanged, OnMarkedFrameChanged )
+
+		/** Called when all marked frames should be cleared */
+		SLATE_EVENT( FSimpleDelegate, OnClearAllMarkedFrames)
+
 		/** Called when the user changes the clamp range */
 		SLATE_EVENT( FOnTimeRangeChanged, OnClampRangeChanged )
 
diff --git a/Engine/Source/Editor/Sequencer/Private/SSequencerSectionOverlay.cpp b/Engine/Source/Editor/Sequencer/Private/SSequencerSectionOverlay.cpp
--- a/Engine/Source/Editor/Sequencer/Private/SSequencerSectionOverlay.cpp
+++ b/Engine/Source/Editor/Sequencer/Private/SSequencerSectionOverlay.cpp
@@ -7,6 +7,7 @@ int32 SSequencerSectionOverlay::OnPaint( const FPaintArgs& Args, const FGeometry
 	FPaintSectionAreaViewArgs PaintArgs;
 	PaintArgs.bDisplayTickLines = bDisplayTickLines.Get();
 	PaintArgs.bDisplayScrubPosition = bDisplayScrubPosition.Get();
+	PaintArgs.bDisplayMarkedFrames = bDisplayMarkedFrames.Get();
 
 	if (PaintPlaybackRangeArgs.IsSet())
 	{
diff --git a/Engine/Source/Editor/Sequencer/Private/SSequencerSectionOverlay.h b/Engine/Source/Editor/Sequencer/Private/SSequencerSectionOverlay.h
--- a/Engine/Source/Editor/Sequencer/Private/SSequencerSectionOverlay.h
+++ b/Engine/Source/Editor/Sequencer/Private/SSequencerSectionOverlay.h
@@ -20,9 +20,11 @@ public:
 	SLATE_BEGIN_ARGS( SSequencerSectionOverlay )
 		: _DisplayTickLines( true )
 		, _DisplayScrubPosition( false )
+		, _DisplayMarkedFrames( false )
 	{}
 		SLATE_ATTRIBUTE( bool, DisplayTickLines )
 		SLATE_ATTRIBUTE( bool, DisplayScrubPosition )
+		SLATE_ATTRIBUTE( bool, DisplayMarkedFrames )
 		SLATE_ATTRIBUTE( FPaintPlaybackRangeArgs, PaintPlaybackRangeArgs )
 
 	SLATE_END_ARGS()
@@ -31,6 +33,7 @@ public:
 	{
 		bDisplayScrubPosition = InArgs._DisplayScrubPosition;
 		bDisplayTickLines = InArgs._DisplayTickLines;
+		bDisplayMarkedFrames = InArgs._DisplayMarkedFrames;
 		PaintPlaybackRangeArgs = InArgs._PaintPlaybackRangeArgs;
 		TimeSliderController = InTimeSliderController;
 	}
@@ -46,6 +49,8 @@ private:
 	TAttribute<bool> bDisplayScrubPosition;
 	/** Whether or not to display tick lines */
 	TAttribute<bool> bDisplayTickLines;
+	/** Whether or not to display marked frames */
+	TAttribute<bool> bDisplayMarkedFrames;
 	/** User-supplied options for drawing playback range */
 	TAttribute<FPaintPlaybackRangeArgs> PaintPlaybackRangeArgs;
 
diff --git a/Engine/Source/Editor/Sequencer/Private/Sequencer.cpp b/Engine/Source/Editor/Sequencer/Private/Sequencer.cpp
--- a/Engine/Source/Editor/Sequencer/Private/Sequencer.cpp
+++ b/Engine/Source/Editor/Sequencer/Private/Sequencer.cpp
@@ -340,6 +340,9 @@ void FSequencer::InitSequencer(const FSequencerInitParams& InitParams, const TSh
 		.PlaybackRange( this, &FSequencer::GetPlaybackRange )
 		.PlaybackStatus( this, &FSequencer::GetPlaybackStatus )
 		.SelectionRange( this, &FSequencer::GetSelectionRange )
+		.MarkedFrames( this, &FSequencer::GetMarkedFrames )
+		.OnMarkedFrameChanged(this, &FSequencer::SetMarkedFrame )
+		.OnClearAllMarkedFrames(this, &FSequencer::ClearAllMarkedFrames )
 		.SubSequenceRange( this, &FSequencer::GetSubSequenceRange )
 		.OnPlaybackRangeChanged( this, &FSequencer::SetPlaybackRange )
 		.OnPlaybackRangeBeginDrag( this, &FSequencer::OnPlaybackRangeBeginDrag )
@@ -6997,6 +7000,143 @@ void FSequencer::Rekey()
 	}
 }
 
+TSet<FFrameNumber> FSequencer::GetMarkedFrames() const 
+{
+	UMovieSceneSequence* FocusedMovieSequence = GetFocusedMovieSceneSequence();
+	if (FocusedMovieSequence != nullptr)
+	{
+		UMovieScene* FocusedMovieScene = FocusedMovieSequence->GetMovieScene();
+		if (FocusedMovieScene != nullptr)
+		{
+			return FocusedMovieScene->GetEditorData().MarkedFrames;
+		}
+	}
+
+	return TSet<FFrameNumber>();
+}
+
+void FSequencer::ToggleMarkAtPlayPosition()
+{
+
+	UMovieSceneSequence* FocusedMovieSequence = GetFocusedMovieSceneSequence();
+	if (FocusedMovieSequence != nullptr)
+	{
+		UMovieScene* FocusedMovieScene = FocusedMovieSequence->GetMovieScene();
+		if (FocusedMovieScene != nullptr)
+		{
+			FFrameNumber TickFrameNumber = GetLocalTime().Time.FloorToFrame();
+
+			FMovieSceneEditorData& EditorData = FocusedMovieScene->GetEditorData();
+			if ( EditorData.MarkedFrames.Contains(TickFrameNumber) )
+			{
+				EditorData.MarkedFrames.Remove(TickFrameNumber);
+			}
+			else 
+			{
+				EditorData.MarkedFrames.Add(TickFrameNumber);	
+			}
+		}
+	}
+}
+
+void FSequencer::SetMarkedFrame(FFrameNumber FrameNumber, bool bSetMark)
+{
+
+	UMovieSceneSequence* FocusedMovieSequence = GetFocusedMovieSceneSequence();
+	if (FocusedMovieSequence != nullptr)
+	{
+		UMovieScene* FocusedMovieScene = FocusedMovieSequence->GetMovieScene();
+		if (FocusedMovieScene != nullptr)
+		{
+			FMovieSceneEditorData& EditorData = FocusedMovieScene->GetEditorData();
+			if (bSetMark)
+			{
+				EditorData.MarkedFrames.Add(FrameNumber);	
+			}
+			else
+			{
+				EditorData.MarkedFrames.Remove(FrameNumber);
+			}
+		}
+	}
+}
+
+void FSequencer::ClearAllMarkedFrames()
+{
+
+	UMovieSceneSequence* FocusedMovieSequence = GetFocusedMovieSceneSequence();
+	if (FocusedMovieSequence != nullptr)
+	{
+		UMovieScene* FocusedMovieScene = FocusedMovieSequence->GetMovieScene();
+		if (FocusedMovieScene != nullptr)
+		{
+			FMovieSceneEditorData& EditorData = FocusedMovieScene->GetEditorData();
+			EditorData.MarkedFrames.Empty();
+		}
+	}
+}
+
+void FSequencer::StepToNextMark()
+{
+	UMovieSceneSequence* FocusedMovieSequence = GetFocusedMovieSceneSequence();
+	if (FocusedMovieSequence != nullptr)
+	{
+		UMovieScene* FocusedMovieScene = FocusedMovieSequence->GetMovieScene();
+		if (FocusedMovieScene != nullptr)
+		{
+			FMovieSceneEditorData& EditorData = FocusedMovieScene->GetEditorData();
+			if (EditorData.MarkedFrames.Num() > 0)
+			{
+				FFrameNumber CurrentTickFrameNumber = GetLocalTime().Time.FloorToFrame();
+
+				TArray<FFrameNumber> SortedMarkedFrames = EditorData.MarkedFrames.Array();
+				SortedMarkedFrames.Sort();
+
+				for (FFrameNumber TickFrame : SortedMarkedFrames)
+				{
+					if (TickFrame > CurrentTickFrameNumber)
+					{
+						SetLocalTime(TickFrame.Value);
+						break;
+					}
+				}
+			}
+		}
+	}
+}
+
+void FSequencer::StepToPreviousMark()
+{
+	UMovieSceneSequence* FocusedMovieSequence = GetFocusedMovieSceneSequence();
+	if (FocusedMovieSequence != nullptr)
+	{
+		UMovieScene* FocusedMovieScene = FocusedMovieSequence->GetMovieScene();
+		if (FocusedMovieScene != nullptr)
+		{
+			FMovieSceneEditorData& EditorData = FocusedMovieScene->GetEditorData();
+			if (EditorData.MarkedFrames.Num() > 0)
+			{
+
+				FFrameNumber CurrentTickFrameNumber = GetLocalTime().Time.FloorToFrame();
+
+				//Intentionally sorted in reverse order to find the one just smaller than the current time
+				TArray<FFrameNumber> SortedMarkedFrames = EditorData.MarkedFrames.Array();
+				SortedMarkedFrames.Sort([](const FFrameNumber& A, const FFrameNumber& B) {return A > B;});
+
+				for (FFrameNumber TickFrame : SortedMarkedFrames)
+				{
+					if (TickFrame < CurrentTickFrameNumber)
+					{
+						SetLocalTime(TickFrame.Value);
+						break;
+					}
+				}
+			}
+		}
+	}
+}
+
+
 TArray<TSharedPtr<FMovieSceneClipboard>> GClipboardStack;
 
 void FSequencer::CopySelection()
@@ -7915,6 +8055,18 @@ void FSequencer::BindCommands()
 		FCanExecuteAction::CreateLambda([] { return true; }),
 		FIsActionChecked::CreateLambda([this] { return Settings->GetKeyGroupMode() == EKeyGroupMode::KeyAll; }));
 
+	SequencerCommandBindings->MapAction(
+		Commands.ToggleMarkAtPlayPosition,
+		FExecuteAction::CreateSP( this, &FSequencer::ToggleMarkAtPlayPosition));
+
+	SequencerCommandBindings->MapAction(
+		Commands.StepToNextMark,
+		FExecuteAction::CreateSP( this, &FSequencer::StepToNextMark));
+
+	SequencerCommandBindings->MapAction(
+		Commands.StepToPreviousMark,
+		FExecuteAction::CreateSP( this, &FSequencer::StepToPreviousMark));
+
 	SequencerCommandBindings->MapAction(
 		Commands.ToggleAutoScroll,
 		FExecuteAction::CreateLambda( [this]{ Settings->SetAutoScrollEnabled( !Settings->GetAutoScrollEnabled() ); } ),
diff --git a/Engine/Source/Editor/Sequencer/Private/Sequencer.h b/Engine/Source/Editor/Sequencer/Private/Sequencer.h
--- a/Engine/Source/Editor/Sequencer/Private/Sequencer.h
+++ b/Engine/Source/Editor/Sequencer/Private/Sequencer.h
@@ -278,6 +278,25 @@ public:
 		return TrackEditors;
 	}
 
+public:
+	/** @return The set of marked frames */
+	TSet<FFrameNumber> GetMarkedFrames() const;
+
+protected:
+
+	/** Set/Clear a Mark at the current time */
+	void ToggleMarkAtPlayPosition();
+	void StepToNextMark();
+	void StepToPreviousMark();
+
+	/**
+	 * @param	FrameNumber The FrameNumber in Ticks 
+	 * @param	bSetMark  true to set the mark, false to clear the mark
+	 */
+	void SetMarkedFrame(FFrameNumber FrameNumber, bool bSetMark);
+
+	void ClearAllMarkedFrames();
+
 public:
 
 	/**
diff --git a/Engine/Source/Editor/Sequencer/Private/SequencerCommands.cpp b/Engine/Source/Editor/Sequencer/Private/SequencerCommands.cpp
--- a/Engine/Source/Editor/Sequencer/Private/SequencerCommands.cpp
+++ b/Engine/Source/Editor/Sequencer/Private/SequencerCommands.cpp
@@ -83,6 +83,10 @@ void FSequencerCommands::RegisterCommands()
 	UI_COMMAND( SetKeyGroup, "Key Group", "Key the groups channels/properties when only one of them changes. ie. Keys all three translation channels when only translation Y changes", EUserInterfaceActionType::ToggleButton, FInputChord());
 	UI_COMMAND( SetKeyAll, "Key All", "Key all channels/properties when only one of them changes. ie. Keys all translation, rotation, scale channels when only translation Y changes", EUserInterfaceActionType::ToggleButton, FInputChord());
 
+	UI_COMMAND( ToggleMarkAtPlayPosition, "Toggle Mark", "Sets or clears a mark at the current play position.", EUserInterfaceActionType::Button, FInputChord(EKeys::M) );
+	UI_COMMAND( StepToNextMark, "Step to Next Marked Frame", "Step to the next marked frame", EUserInterfaceActionType::Button, FInputChord(EModifierKey::Control | EModifierKey::Shift, EKeys::Period) );
+	UI_COMMAND( StepToPreviousMark, "Step to Previous Marked Frame", "Step to the previous marked frame", EUserInterfaceActionType::Button, FInputChord(EModifierKey::Control | EModifierKey::Shift, EKeys::Comma) );
+
 	UI_COMMAND( ToggleAutoScroll, "Auto Scroll", "Toggle auto-scroll: When enabled, automatically scrolls the sequencer view to keep the current time visible", EUserInterfaceActionType::ToggleButton, FInputChord(EModifierKey::Shift, EKeys::S) );
 
 	UI_COMMAND( ChangeTimeDisplayFormat, "Change Time Display Format", "Rotates through supported display formats for time", EUserInterfaceActionType::Button, FInputChord(EModifierKey::Control, EKeys::T) );
diff --git a/Engine/Source/Editor/Sequencer/Private/SequencerTimeSliderController.cpp b/Engine/Source/Editor/Sequencer/Private/SequencerTimeSliderController.cpp
--- a/Engine/Source/Editor/Sequencer/Private/SequencerTimeSliderController.cpp
+++ b/Engine/Source/Editor/Sequencer/Private/SequencerTimeSliderController.cpp
@@ -275,6 +275,37 @@ void FSequencerTimeSliderController::DrawTicks( FSlateWindowElementList& OutDraw
 	}
 }
 
+int32 FSequencerTimeSliderController::DrawMarkedFrames( const FGeometry& AllottedGeometry, const FScrubRangeToScreen& RangeToScreen, FSlateWindowElementList& OutDrawElements, int32 LayerId, const ESlateDrawEffect& DrawEffects ) const
+{
+	TSet<FFrameNumber> MarkedFrames = TimeSliderArgs.MarkedFrames.Get();
+	if (MarkedFrames.Num() < 1)
+	{
+		return LayerId;
+	}
+
+	for (FFrameNumber TickFrame : MarkedFrames)
+	{
+		double Seconds = TickFrame / GetTickResolution();
+
+		const float  LinePos = RangeToScreen.InputToLocalX( Seconds );
+		TArray<FVector2D> LinePoints;
+		LinePoints.AddUninitialized(2);
+		LinePoints[0] = FVector2D( LinePos, 0.0f );
+		LinePoints[1] = FVector2D( LinePos, FMath::FloorToFloat( AllottedGeometry.Size.Y ) );
+
+		FSlateDrawElement::MakeLines(
+			OutDrawElements,
+			LayerId+1,
+			AllottedGeometry.ToPaintGeometry(),
+			LinePoints,
+			DrawEffects,
+			FLinearColor(0.f, 1.f, 1.f, 0.4f),
+			false
+			);
+	}
+
+	return LayerId + 1;
+}
 
 int32 FSequencerTimeSliderController::OnPaintTimeSlider( bool bMirrorLabels, const FGeometry& AllottedGeometry, const FSlateRect& MyCullingRect, FSlateWindowElementList& OutDrawElements, int32 LayerId, const FWidgetStyle& InWidgetStyle, bool bParentEnabled ) const
 {
@@ -357,6 +388,8 @@ int32 FSequencerTimeSliderController::OnPaintTimeSlider( bool bMirrorLabels, con
 			ScrubColor
 		);
 
+		LayerId = DrawMarkedFrames(AllottedGeometry, RangeToScreen, OutDrawElements, LayerId, DrawEffects);
+
 		{
 			// Draw the current time next to the scrub handle
 			FString FrameString;
@@ -1011,6 +1044,11 @@ int32 FSequencerTimeSliderController::OnPaintSectionView( const FGeometry& Allot
 		DrawTicks( OutDrawElements, LocalViewRange, RangeToScreen, DrawTickArgs );
 	}
 
+	if ( Args.bDisplayMarkedFrames )
+	{
+		LayerId = DrawMarkedFrames(AllottedGeometry, RangeToScreen, OutDrawElements, LayerId, DrawEffects);
+	}
+
 	if( Args.bDisplayScrubPosition )
 	{
 		FQualifiedFrameTime ScrubPosition = FQualifiedFrameTime(TimeSliderArgs.ScrubPosition.Get(), GetTickResolution());
@@ -1136,6 +1174,44 @@ TSharedRef<SWidget> FSequencerTimeSliderController::OpenSetPlaybackRangeMenu(FFr
 	}
 	MenuBuilder.EndSection(); // SequencerPlaybackRangeMenu
 
+	MenuBuilder.BeginSection("SequencerMarkMenu", FText::Format(LOCTEXT("MarkTextFormat", "Mark ({0}):"), CurrentTimeText));
+	{
+		FFrameNumber DisplayFrameNumber = GetDisplayRate().AsFrameNumber(FrameNumber / GetTickResolution());
+
+		const TSet<FFrameNumber>& MarkedFrames = TimeSliderArgs.MarkedFrames.Get();
+		bool HasMarkAtFrame = MarkedFrames.Contains(FrameNumber);
+		if (!HasMarkAtFrame)
+		{
+			MenuBuilder.AddMenuEntry( 
+				LOCTEXT("AddMark", "Add Mark"),
+				FText(),
+				FSlateIcon(),
+				FUIAction(FExecuteAction::CreateLambda( [=]{ AddMarkAtFrame(FrameNumber); }))
+			);
+		}
+		else 
+		{
+			MenuBuilder.AddMenuEntry( 
+				LOCTEXT("ClearMark", "Clear Mark"),
+				FText(),
+				FSlateIcon(),
+				FUIAction(FExecuteAction::CreateLambda([=]{ ClearMarkAtFrame(FrameNumber); }))
+			);
+		}
+
+		MenuBuilder.AddMenuEntry( 
+			LOCTEXT("Clear All Marks", "Clear All Marks"),
+			FText(),
+			FSlateIcon(),
+
+			FUIAction(
+				FExecuteAction::CreateLambda([=]{ ClearAllMarks(); }),
+				FCanExecuteAction::CreateLambda([=]{ return (MarkedFrames.Num() > 0); })
+			)
+		);
+	}
+	MenuBuilder.EndSection(); // SequencerMarkMenu
+
 	return MenuBuilder.MakeWidget();
 }
 
@@ -1355,4 +1431,20 @@ void FSequencerTimeSliderController::SetSelectionRangeEnd(FFrameNumber NewEnd)
 	}
 }
 
+void FSequencerTimeSliderController::AddMarkAtFrame(FFrameNumber FrameNumber)
+{
+	TimeSliderArgs.OnMarkedFrameChanged.ExecuteIfBound(FrameNumber, true);
+}
+
+void FSequencerTimeSliderController::ClearMarkAtFrame(FFrameNumber FrameNumber)
+{
+	TimeSliderArgs.OnMarkedFrameChanged.ExecuteIfBound(FrameNumber, false);
+}
+
+void FSequencerTimeSliderController::ClearAllMarks()
+{
+	TimeSliderArgs.OnClearAllMarkedFrames.ExecuteIfBound();
+}
+	
+
 #undef LOCTEXT_NAMESPACE
diff --git a/Engine/Source/Editor/Sequencer/Private/Tools/SequencerSnapField.cpp b/Engine/Source/Editor/Sequencer/Private/Tools/SequencerSnapField.cpp
--- a/Engine/Source/Editor/Sequencer/Private/Tools/SequencerSnapField.cpp
+++ b/Engine/Source/Editor/Sequencer/Private/Tools/SequencerSnapField.cpp
@@ -90,6 +90,13 @@ FSequencerSnapField::FSequencerSnapField(const ISequencer& InSequencer, ISequenc
 		Visitor.Snaps.Add(FSequencerSnapPoint{ FSequencerSnapPoint::InOutRange, MovieScene::DiscreteExclusiveUpper(SelectionRange) - 1});
 	}
 
+	// Add in the marked frames
+	const TSet<FFrameNumber>& MarkedFrames = InSequencer.GetFocusedMovieSceneSequence()->GetMovieScene()->GetEditorData().MarkedFrames;
+	for (FFrameNumber MarkedFrame : MarkedFrames)
+	{
+		Visitor.Snaps.Add( FSequencerSnapPoint{ FSequencerSnapPoint::Mark, MarkedFrame } );
+	}
+
 	// Sort
 	Visitor.Snaps.Sort([](const FSequencerSnapPoint& A, const FSequencerSnapPoint& B){
 		return A.Time < B.Time;
diff --git a/Engine/Source/Editor/Sequencer/Private/Tools/SequencerSnapField.h b/Engine/Source/Editor/Sequencer/Private/Tools/SequencerSnapField.h
--- a/Engine/Source/Editor/Sequencer/Private/Tools/SequencerSnapField.h
+++ b/Engine/Source/Editor/Sequencer/Private/Tools/SequencerSnapField.h
@@ -12,7 +12,7 @@ class ISequencer;
 /** Structure defining a point to snap to in the sequencer */
 struct FSequencerSnapPoint
 {
-	enum ESnapType { Key, SectionBounds, CustomSection, PlaybackRange, CurrentTime, InOutRange };
+	enum ESnapType { Key, SectionBounds, CustomSection, PlaybackRange, CurrentTime, InOutRange, Mark };
 
 	/** The type of snap */
 	ESnapType Type;
diff --git a/Engine/Source/Editor/Sequencer/Public/SequencerCommands.h b/Engine/Source/Editor/Sequencer/Public/SequencerCommands.h
--- a/Engine/Source/Editor/Sequencer/Public/SequencerCommands.h
+++ b/Engine/Source/Editor/Sequencer/Public/SequencerCommands.h
@@ -205,6 +205,15 @@ public:
 	/** Set mode to key all. */
 	TSharedPtr< FUICommandInfo > SetKeyAll;
 
+	/** Toggle on/off a Mark at the current time **/
+	TSharedPtr< FUICommandInfo> ToggleMarkAtPlayPosition;
+
+	/** Step to next mark */
+	TSharedPtr< FUICommandInfo > StepToNextMark;
+
+	/** Step to previous mark */
+	TSharedPtr< FUICommandInfo > StepToPreviousMark;
+
 	/** Rotates through the supported formats for displaying times/frames/timecode. */
 	TSharedPtr< FUICommandInfo > ChangeTimeDisplayFormat;
 
diff --git a/Engine/Source/Editor/Sequencer/Public/SequencerTimeSliderController.h b/Engine/Source/Editor/Sequencer/Public/SequencerTimeSliderController.h
--- a/Engine/Source/Editor/Sequencer/Public/SequencerTimeSliderController.h
+++ b/Engine/Source/Editor/Sequencer/Public/SequencerTimeSliderController.h
@@ -36,13 +36,15 @@ struct FPaintPlaybackRangeArgs
 struct FPaintSectionAreaViewArgs
 {
 	FPaintSectionAreaViewArgs()
-		: bDisplayTickLines(false), bDisplayScrubPosition(false)
+		: bDisplayTickLines(false), bDisplayScrubPosition(false), bDisplayMarkedFrames(false)
 	{}
 
 	/** Whether to display tick lines */
 	bool bDisplayTickLines;
 	/** Whether to display the scrub position */
 	bool bDisplayScrubPosition;
+	/** Whether to display the marked frames */
+	bool bDisplayMarkedFrames;
 	/** Optional Paint args for the playback range*/
 	TOptional<FPaintPlaybackRangeArgs> PlaybackRangeArgs;
 };
@@ -228,6 +230,13 @@ private:
 	 */
 	int32 DrawSubSequenceRange(const FGeometry& AllottedGeometry, const FSlateRect& MyCullingRect, FSlateWindowElementList& OutDrawElements, int32 LayerId, const FScrubRangeToScreen& RangeToScreen, const FPaintPlaybackRangeArgs& Args) const;
 
+	/**
+	 * Draw the marked frames. 
+	 *
+	 * @return the new layer ID
+	 */
+	int32 FSequencerTimeSliderController::DrawMarkedFrames( const FGeometry& AllottedGeometry, const FScrubRangeToScreen& RangeToScreen, FSlateWindowElementList& OutDrawElements, int32 LayerId, const ESlateDrawEffect& DrawEffects ) const;
+
 private:
 
 	/**
@@ -252,6 +261,10 @@ private:
 	FFrameTime ComputeScrubTimeFromMouse(const FGeometry& Geometry, FVector2D ScreenSpacePosition, FScrubRangeToScreen RangeToScreen) const;
 	FFrameTime ComputeFrameTimeFromMouse(const FGeometry& Geometry, FVector2D ScreenSpacePosition, FScrubRangeToScreen RangeToScreen, bool CheckSnapping = true) const;
 
+	void AddMarkAtFrame(FFrameNumber FrameNumber);
+	void ClearMarkAtFrame(FFrameNumber FrameNumber);
+	void ClearAllMarks();
+
 private:
 
 	/**
diff --git a/Engine/Source/Editor/SequencerWidgets/Public/ITimeSlider.h b/Engine/Source/Editor/SequencerWidgets/Public/ITimeSlider.h
--- a/Engine/Source/Editor/SequencerWidgets/Public/ITimeSlider.h
+++ b/Engine/Source/Editor/SequencerWidgets/Public/ITimeSlider.h
@@ -28,6 +28,7 @@ DECLARE_DELEGATE_TwoParams( FOnScrubPositionChanged, FFrameTime, bool )
 DECLARE_DELEGATE_TwoParams( FOnViewRangeChanged, TRange<double>, EViewRangeInterpolation )
 DECLARE_DELEGATE_OneParam( FOnTimeRangeChanged, TRange<double> )
 DECLARE_DELEGATE_OneParam( FOnFrameRangeChanged, TRange<FFrameNumber> )
+DECLARE_DELEGATE_TwoParams( FOnMarkedFrameChanged, FFrameNumber, bool )
 DECLARE_DELEGATE_RetVal_TwoParams( FFrameNumber, FOnGetNearestKey, FFrameTime, bool )
 
 /** Structure used to wrap up a range, and an optional animation target */
@@ -133,6 +134,16 @@ struct FTimeSliderArgs
 
 	/** Called right after the selection range has finished being dragged */
 	FSimpleDelegate OnSelectionRangeEndDrag;
+
+	/** Attribute for the current sequence's marked frames */
+	TAttribute<TSet<FFrameNumber>> MarkedFrames;
+
+	/** Called when the marked frames need to be updated */
+	FOnMarkedFrameChanged OnMarkedFrameChanged;
+
+	/** Called when all marked frames should be cleared */
+	FSimpleDelegate OnClearAllMarkedFrames;
+
 	/** Round the scrub position to an integer during playback */
 	TAttribute<EMovieScenePlayerStatus::Type> PlaybackStatus;
 
diff --git a/Engine/Source/Runtime/MovieScene/Public/MovieScene.h b/Engine/Source/Runtime/MovieScene/Public/MovieScene.h
--- a/Engine/Source/Runtime/MovieScene/Public/MovieScene.h
+++ b/Engine/Source/Runtime/MovieScene/Public/MovieScene.h
@@ -102,6 +102,10 @@ struct FMovieSceneEditorData
 	UPROPERTY()
 	double WorkEnd;
 
+	/** The set of user-marked frames for display convenience */
+	UPROPERTY()
+	TSet<FFrameNumber> MarkedFrames;
+
 	UPROPERTY()
 	FFloatRange WorkingRange_DEPRECATED;
 	UPROPERTY()
