Spawnable Object Bindings now show up Gray in Sequencer instead of Red when they have not been spawned yet due to the evaluation of the Spawn track.

diff --git a/Engine/Source/Editor/Sequencer/Private/DisplayNodes/SequencerObjectBindingNode.cpp b/Engine/Source/Editor/Sequencer/Private/DisplayNodes/SequencerObjectBindingNode.cpp
--- a/Engine/Source/Editor/Sequencer/Private/DisplayNodes/SequencerObjectBindingNode.cpp
+++ b/Engine/Source/Editor/Sequencer/Private/DisplayNodes/SequencerObjectBindingNode.cpp
@@ -417,6 +417,13 @@ FLinearColor FSequencerObjectBindingNode::GetDisplayNameColor() const
 		}
 	}
 
+	// Spawnables don't have valid object bindings when their track hasn't spawned them yet,
+	// so we override the default behavior of red with a gray so that users don't think there is something wrong.
+	if (GetBindingType() == EObjectBindingType::Spawnable)
+	{
+		return FLinearColor::Gray;
+	}
+
 	return FLinearColor::Red;
 }
 
