Sequencer: Fixed additional edit widgets not appearing on top-level key area nodes

diff --git a/Engine/Source/Editor/Sequencer/Private/DisplayNodes/SequencerSectionKeyAreaNode.cpp b/Engine/Source/Editor/Sequencer/Private/DisplayNodes/SequencerSectionKeyAreaNode.cpp
--- a/Engine/Source/Editor/Sequencer/Private/DisplayNodes/SequencerSectionKeyAreaNode.cpp
+++ b/Engine/Source/Editor/Sequencer/Private/DisplayNodes/SequencerSectionKeyAreaNode.cpp
@@ -50,16 +50,34 @@ bool FSequencerSectionKeyAreaNode::CanRenameNode() const
 	return false;
 }
 
+TSharedRef<SWidget> FSequencerSectionKeyAreaNode::GetOrCreateKeyAreaEditorSwitcher()
+{
+	if (!KeyEditorSwitcher.IsValid())
+	{
+		KeyEditorSwitcher = SNew(SKeyAreaEditorSwitcher, SharedThis(this));
+	}
+	return KeyEditorSwitcher.ToSharedRef();
+}
 
 TSharedRef<SWidget> FSequencerSectionKeyAreaNode::GetCustomOutlinerContent()
 {
 	if (GetAllKeyAreas().Num() > 0)
 	{
-		if (!KeyEditorSwitcher.IsValid())
-		{
-			KeyEditorSwitcher = SNew(SKeyAreaEditorSwitcher, SharedThis(this));
-		}
-		return KeyEditorSwitcher.ToSharedRef();
+		return SNew(SHorizontalBox)
+
+		+ SHorizontalBox::Slot()
+		.HAlign(HAlign_Right)
+		.VAlign(VAlign_Center)
+		[
+			GetOrCreateKeyAreaEditorSwitcher()
+		]
+
+		+ SHorizontalBox::Slot()
+		.AutoWidth()
+		.VAlign(VAlign_Center)
+		[
+			SNew(SKeyNavigationButtons, SharedThis(this))
+		];
 	}
 
 	return FSequencerDisplayNode::GetCustomOutlinerContent();
diff --git a/Engine/Source/Editor/Sequencer/Private/DisplayNodes/SequencerSectionKeyAreaNode.h b/Engine/Source/Editor/Sequencer/Private/DisplayNodes/SequencerSectionKeyAreaNode.h
--- a/Engine/Source/Editor/Sequencer/Private/DisplayNodes/SequencerSectionKeyAreaNode.h
+++ b/Engine/Source/Editor/Sequencer/Private/DisplayNodes/SequencerSectionKeyAreaNode.h
@@ -64,6 +64,9 @@ public:
 		return bTopLevel;
 	}
 
+	/** Retrieve the key area editor switcher widget, creating it if it doesn't yet exist */
+	TSharedRef<SWidget> GetOrCreateKeyAreaEditorSwitcher();
+
 public:
 
 	// FSequencerDisplayNode interface
diff --git a/Engine/Source/Editor/Sequencer/Private/DisplayNodes/SequencerTrackNode.cpp b/Engine/Source/Editor/Sequencer/Private/DisplayNodes/SequencerTrackNode.cpp
--- a/Engine/Source/Editor/Sequencer/Private/DisplayNodes/SequencerTrackNode.cpp
+++ b/Engine/Source/Editor/Sequencer/Private/DisplayNodes/SequencerTrackNode.cpp
@@ -34,6 +34,17 @@ namespace SequencerNodeConstants
 	extern const float CommonPadding;
 }
 
+bool ContainsKeyableArea(TSharedRef<FSequencerSectionKeyAreaNode> InKeyAreaNode)
+{
+	for (const TSharedRef<IKeyArea>& KeyArea : InKeyAreaNode->GetAllKeyAreas())
+	{
+		if (KeyArea->CanCreateKeyEditor())
+		{
+			return true;
+		}
+	}
+	return false;
+}
 
 /* FTrackNode structors
  *****************************************************************************/
@@ -161,9 +172,10 @@ FReply FSequencerTrackNode::CreateNewSection() const
 TSharedRef<SWidget> FSequencerTrackNode::GetCustomOutlinerContent()
 {
 	TSharedPtr<FSequencerSectionKeyAreaNode> KeyAreaNode = GetTopLevelKeyNode();
+	TSharedPtr<SWidget> KeyEditorWidget;
 	if (KeyAreaNode.IsValid())
 	{
-		return KeyAreaNode->GetCustomOutlinerContent();
+		KeyEditorWidget = KeyAreaNode->GetOrCreateKeyAreaEditorSwitcher();
 	}
 
 	TAttribute<bool> NodeIsHovered = TAttribute<bool>::Create(TAttribute<bool>::FGetter::CreateSP(this, &FSequencerDisplayNode::IsHovered));
@@ -191,44 +203,59 @@ TSharedRef<SWidget> FSequencerTrackNode::GetCustomOutlinerContent()
 		Params.TrackInsertRowIndex = Track->GetMaxRowIndex()+1;
 	}
 
-	TSharedPtr<SWidget> Widget = GetSequencer().IsReadOnly() ? SNullWidget::NullWidget : AssociatedEditor.BuildOutlinerEditWidget(ObjectBinding, Track, Params);
-
-	bool bHasKeyableAreas = false;
-
-	TArray<TSharedRef<FSequencerSectionKeyAreaNode>> ChildKeyAreaNodes;
-	FSequencerDisplayNode::GetChildKeyAreaNodesRecursively(ChildKeyAreaNodes);
-	for (int32 ChildIndex = 0; ChildIndex < ChildKeyAreaNodes.Num() && !bHasKeyableAreas; ++ChildIndex)
-	{
-		TArray< TSharedRef<IKeyArea> > ChildKeyAreas = ChildKeyAreaNodes[ChildIndex]->GetAllKeyAreas();
-
-		for (int32 ChildKeyAreaIndex = 0; ChildKeyAreaIndex < ChildKeyAreas.Num() && !bHasKeyableAreas; ++ChildKeyAreaIndex)
-		{
-			if (ChildKeyAreas[ChildKeyAreaIndex]->CanCreateKeyEditor())
-			{
-				bHasKeyableAreas = true;
-			}
-		}
-	}
+	TSharedPtr<SWidget> CustomWidget = GetSequencer().IsReadOnly() ? SNullWidget::NullWidget : AssociatedEditor.BuildOutlinerEditWidget(ObjectBinding, Track, Params);
 
-	if (Widget.IsValid())
+	if (KeyEditorWidget.IsValid())
 	{
 		BoxPanel->AddSlot()
 		.AutoWidth()
 		.VAlign(VAlign_Center)
-		.HAlign(HAlign_Right)
 		[
-			Widget.ToSharedRef()
+			KeyEditorWidget.ToSharedRef()
 		];
-	}
 
-	if (bHasKeyableAreas)
-	{
+		if (CustomWidget.IsValid())
+		{
+			BoxPanel->AddSlot()
+			.AutoWidth()
+			.VAlign(VAlign_Center)
+			[
+				CustomWidget.ToSharedRef()
+			];
+		}
+
 		BoxPanel->AddSlot()
 		.VAlign(VAlign_Center)
+		.AutoWidth()
 		[
-			SNew(SKeyNavigationButtons, AsShared())
+			SNew(SKeyNavigationButtons, KeyAreaNode.ToSharedRef())
 		];
 	}
+	else
+	{
+		if (CustomWidget.IsValid())
+		{
+			BoxPanel->AddSlot()
+			.AutoWidth()
+			.VAlign(VAlign_Center)
+			[
+				CustomWidget.ToSharedRef()
+			];
+		}
+
+		TArray<TSharedRef<FSequencerSectionKeyAreaNode>> ChildKeyAreaNodes;
+		FSequencerDisplayNode::GetChildKeyAreaNodesRecursively(ChildKeyAreaNodes);
+
+		if (ChildKeyAreaNodes.FindByPredicate(ContainsKeyableArea))
+		{
+			BoxPanel->AddSlot()
+			.VAlign(VAlign_Center)
+			.AutoWidth()
+			[
+				SNew(SKeyNavigationButtons, AsShared())
+			];
+		}
+	}
 
 	return SNew(SBox)
 		.VAlign(VAlign_Center)
diff --git a/Engine/Source/Editor/Sequencer/Private/SKeyAreaEditorSwitcher.cpp b/Engine/Source/Editor/Sequencer/Private/SKeyAreaEditorSwitcher.cpp
--- a/Engine/Source/Editor/Sequencer/Private/SKeyAreaEditorSwitcher.cpp
+++ b/Engine/Source/Editor/Sequencer/Private/SKeyAreaEditorSwitcher.cpp
@@ -14,27 +14,8 @@ void SKeyAreaEditorSwitcher::Construct(const FArguments& InArgs, TSharedRef<FSeq
 	WeakKeyAreaNode = InKeyAreaNode;
 	ChildSlot
 	[
-		SNew(SBox)
-		.HAlign(HAlign_Right)
-		.VAlign(VAlign_Center)
-		[
-			SNew(SHorizontalBox)
-
-			+ SHorizontalBox::Slot()
-			.HAlign(HAlign_Right)
-			.VAlign(VAlign_Center)
-			[
-				SAssignNew(Overlay, SOverlay)
-				.IsEnabled(!InKeyAreaNode->GetSequencer().IsReadOnly())
-			]
-
-			+ SHorizontalBox::Slot()
-			.AutoWidth()
-			.VAlign(VAlign_Center)
-			[
-				SNew(SKeyNavigationButtons, InKeyAreaNode)
-			]
-		]
+		SAssignNew(Overlay, SOverlay)
+		.IsEnabled(!InKeyAreaNode->GetSequencer().IsReadOnly())
 	];
 
 	Rebuild();
