Sequencer: Fix crash in step to next/previous shot.

Added check for null subsequence data for the root sequence id.

Copy from Release-4.20

diff --git a/Engine/Source/Editor/Sequencer/Private/Sequencer.cpp b/Engine/Source/Editor/Sequencer/Private/Sequencer.cpp
--- a/Engine/Source/Editor/Sequencer/Private/Sequencer.cpp
+++ b/Engine/Source/Editor/Sequencer/Private/Sequencer.cpp
@@ -6738,7 +6738,11 @@ void FSequencer::StepToNextShot()
 	UMovieSceneSequence* Sequence = RootTemplateInstance.GetSequence(OuterSequenceID);
 
 	FFrameTime StartTime = FFrameTime(0) * RootToLocalTransform.Inverse();
-	FFrameTime CurrentTime = StartTime * RootTemplateInstance.GetHierarchy().FindSubData(OuterSequenceID)->RootToSequenceTransform;
+	FFrameTime CurrentTime = StartTime;
+	if (FMovieSceneSubSequenceData* SubSequenceData = RootTemplateInstance.GetHierarchy().FindSubData(OuterSequenceID))
+	{
+		CurrentTime = StartTime * SubSequenceData->RootToSequenceTransform;
+	}
 
 	UMovieSceneSubSection* NextShot = Cast<UMovieSceneSubSection>(FindNextOrPreviousShot(Sequence, CurrentTime.FloorToFrame(), true));
 	if (!NextShot)
@@ -6766,7 +6770,11 @@ void FSequencer::StepToPreviousShot()
 	UMovieSceneSequence* Sequence = RootTemplateInstance.GetSequence(OuterSequenceID);
 
 	FFrameTime StartTime = FFrameTime(0) * RootToLocalTransform.Inverse();
-	FFrameTime CurrentTime = StartTime * RootTemplateInstance.GetHierarchy().FindSubData(OuterSequenceID)->RootToSequenceTransform;
+	FFrameTime CurrentTime = StartTime;
+	if (FMovieSceneSubSequenceData* SubSequenceData = RootTemplateInstance.GetHierarchy().FindSubData(OuterSequenceID))
+	{
+		CurrentTime = StartTime * SubSequenceData->RootToSequenceTransform;
+	}
 
 	UMovieSceneSubSection* PreviousShot = Cast<UMovieSceneSubSection>(FindNextOrPreviousShot(Sequence, CurrentTime.FloorToFrame(), false));
 	if (!PreviousShot)
