Sequencer: Pause before begin scrubbing.

This fixes an issue where if you play the sequence with audio and scrub and release, the audio keeps playing. This is because there's no explicit evaluation in the stopped state OnEndScrubbing.

diff --git a/Engine/Source/Editor/Sequencer/Private/Sequencer.cpp b/Engine/Source/Editor/Sequencer/Private/Sequencer.cpp
--- a/Engine/Source/Editor/Sequencer/Private/Sequencer.cpp
+++ b/Engine/Source/Editor/Sequencer/Private/Sequencer.cpp
@@ -3793,6 +3793,9 @@ void FSequencer::OnScrubPositionChanged( FFrameTime NewScrubPosition, bool bScru
 
 void FSequencer::OnBeginScrubbing()
 {
+	// Pause first since there's no explicit evaluation in the stopped state when OnEndScrubbing() is called
+	Pause();
+
 	SetPlaybackStatus(EMovieScenePlayerStatus::Scrubbing);
 	SequencerWidget->RegisterActiveTimerForPlayback();
 
