Sequencer: Don't attempt to select actors/components on locked levels.

diff --git a/Engine/Source/Editor/Sequencer/Private/Sequencer.cpp b/Engine/Source/Editor/Sequencer/Private/Sequencer.cpp
--- a/Engine/Source/Editor/Sequencer/Private/Sequencer.cpp
+++ b/Engine/Source/Editor/Sequencer/Private/Sequencer.cpp
@@ -120,6 +120,7 @@
 #include "Fonts/FontMeasure.h"
 #include "MovieSceneTimeHelpers.h"
 #include "FrameNumberNumericInterface.h"
+#include "LevelUtils.h"
 #include "Engine/Blueprint.h"
 #include "MovieSceneSequenceEditor.h"
 #include "Kismet2/KismetEditorUtilities.h"
@@ -4771,9 +4772,13 @@ void FSequencer::SynchronizeExternalSelectionWithSequencerSelection()
 
 	GEditor->SelectNone( bNotifySelectionChanged, bDeselectBSP, bWarnAboutTooManyActors );
 
-	for ( AActor* SelectedSequencerActor : SelectedSequencerActors )
+	for (AActor* SelectedSequencerActor : SelectedSequencerActors)
 	{
-		GEditor->SelectActor( SelectedSequencerActor, true, bNotifySelectionChanged, bSelectEvenIfHidden );
+		ULevel* ActorLevel = SelectedSequencerActor->GetLevel();
+		if (!FLevelUtils::IsLevelLocked(ActorLevel))
+		{
+			GEditor->SelectActor(SelectedSequencerActor, true, bNotifySelectionChanged, bSelectEvenIfHidden);
+		}
 	}
 
 	GEditor->GetSelectedActors()->EndBatchSelectOperation();
@@ -4783,9 +4788,12 @@ void FSequencer::SynchronizeExternalSelectionWithSequencerSelection()
 		GEditor->GetSelectedComponents()->Modify();
 		GEditor->GetSelectedComponents()->BeginBatchSelectOperation();
 
-		for ( UActorComponent* SelectedSequencerComponent : SelectedSequencerComponents )
+		for (UActorComponent* SelectedSequencerComponent : SelectedSequencerComponents)
 		{
-			GEditor->SelectComponent( SelectedSequencerComponent, true, bNotifySelectionChanged, bSelectEvenIfHidden );
+			if (!FLevelUtils::IsLevelLocked(SelectedSequencerComponent->GetOwner()->GetLevel()))
+			{
+				GEditor->SelectComponent(SelectedSequencerComponent, true, bNotifySelectionChanged, bSelectEvenIfHidden);
+			}
 		}
 
 		GEditor->GetSelectedComponents()->EndBatchSelectOperation();
