Sequencer: Fix duplicated spawnble tracks not getting saved. Clear transient flags from the tracks/sections coming from the FMovieSceneCopyableBinding.

diff --git a/Engine/Source/Editor/Sequencer/Private/Sequencer.cpp b/Engine/Source/Editor/Sequencer/Private/Sequencer.cpp
--- a/Engine/Source/Editor/Sequencer/Private/Sequencer.cpp
+++ b/Engine/Source/Editor/Sequencer/Private/Sequencer.cpp
@@ -5835,6 +5835,16 @@ void FSequencer::PasteCopiedTracks()
 				// Copy the name of the original spawnable too.
 				Spawnable->SetName(CopyableBinding->Spawnable.GetName());
 
+				// Clear the transient flags on the copyable binding before assigning to the new spawnable
+				for (auto Track : NewBinding.GetTracks())
+				{
+					Track->ClearFlags(RF_Transient);
+					for (auto Section : Track->GetAllSections())
+					{
+						Section->ClearFlags(RF_Transient);
+					}
+				}
+
 				// Replace the auto-generated binding with our deserialized bindings (which has our tracks)
 				MovieScene->ReplaceBinding(NewGuid, NewBinding);
 
