Sequencer: Overlay the custom widget (ie. +Section) on the key editors

diff --git a/Engine/Source/Editor/Sequencer/Private/DisplayNodes/SequencerTrackNode.cpp b/Engine/Source/Editor/Sequencer/Private/DisplayNodes/SequencerTrackNode.cpp
--- a/Engine/Source/Editor/Sequencer/Private/DisplayNodes/SequencerTrackNode.cpp
+++ b/Engine/Source/Editor/Sequencer/Private/DisplayNodes/SequencerTrackNode.cpp
@@ -207,8 +207,10 @@ TSharedRef<SWidget> FSequencerTrackNode::GetCustomOutlinerContent()
 
 	if (KeyEditorWidget.IsValid())
 	{
-		BoxPanel->AddSlot()
-		.AutoWidth()
+		TSharedRef<SOverlay> Overlay = SNew(SOverlay);
+
+		Overlay->AddSlot()
+		.HAlign(HAlign_Left)
 		.VAlign(VAlign_Center)
 		[
 			KeyEditorWidget.ToSharedRef()
@@ -216,8 +218,8 @@ TSharedRef<SWidget> FSequencerTrackNode::GetCustomOutlinerContent()
 
 		if (CustomWidget.IsValid())
 		{
-			BoxPanel->AddSlot()
-			.AutoWidth()
+			Overlay->AddSlot()
+			.HAlign(HAlign_Right)
 			.VAlign(VAlign_Center)
 			[
 				CustomWidget.ToSharedRef()
@@ -225,12 +227,20 @@ TSharedRef<SWidget> FSequencerTrackNode::GetCustomOutlinerContent()
 		}
 
 		BoxPanel->AddSlot()
+		.AutoWidth()
 		.VAlign(VAlign_Center)
+		[
+			Overlay
+		];
+
+		BoxPanel->AddSlot()
 		.AutoWidth()
+		.VAlign(VAlign_Center)
 		[
 			SNew(SKeyNavigationButtons, KeyAreaNode.ToSharedRef())
 		];
 	}
+
 	else
 	{
 		if (CustomWidget.IsValid())
