Sequencer: Clarified tooltip for export fbx

diff --git a/Engine/Source/Editor/Sequencer/Private/SequencerCommands.cpp b/Engine/Source/Editor/Sequencer/Private/SequencerCommands.cpp
--- a/Engine/Source/Editor/Sequencer/Private/SequencerCommands.cpp
+++ b/Engine/Source/Editor/Sequencer/Private/SequencerCommands.cpp
@@ -143,8 +143,8 @@ void FSequencerCommands::RegisterCommands()
 
 	UI_COMMAND( RecordSelectedActors, "Record Selected Actors", "Records the selected actors into a new sub sequence of the currently active sequence in Sequencer.", EUserInterfaceActionType::Button, FInputChord(EModifierKey::Alt, EKeys::R) );
 
-	UI_COMMAND( ImportFBX, "Import...", "Imports the animation from an FBX file.", EUserInterfaceActionType::Button, FInputChord() );
-	UI_COMMAND( ExportFBX, "Export...", "Exports the animation to an FBX file. (Shots and sub-scenes not supported)", EUserInterfaceActionType::Button, FInputChord() );
+	UI_COMMAND( ImportFBX, "Import...", "Import the animation from an FBX file.", EUserInterfaceActionType::Button, FInputChord() );
+	UI_COMMAND( ExportFBX, "Export...", "Export the selected objects (all if none selected) and animation to an FBX file. (Shots and sub-scenes not supported)", EUserInterfaceActionType::Button, FInputChord() );
 
 	UI_COMMAND( ExportToCameraAnim, "Export to Camera Anim...", "Exports the animation to a camera anim asset.", EUserInterfaceActionType::Button, FInputChord() );
 
