Sequencer: Remove ctrl-w as a hotkey to toggle weighted/non-weighted because it conflicts with the more generic binding to duplicate objects/tracks.

diff --git a/Engine/Source/Editor/Sequencer/Private/SequencerCommands.cpp b/Engine/Source/Editor/Sequencer/Private/SequencerCommands.cpp
--- a/Engine/Source/Editor/Sequencer/Private/SequencerCommands.cpp
+++ b/Engine/Source/Editor/Sequencer/Private/SequencerCommands.cpp
@@ -59,7 +59,7 @@ void FSequencerCommands::RegisterCommands()
 	UI_COMMAND( SetInterpolationLinear, "Set Key Linear", "Linear interpolation", EUserInterfaceActionType::Button, FInputChord(EKeys::Four));
 	UI_COMMAND( SetInterpolationConstant, "Set Key Constant", "Constant interpolation", EUserInterfaceActionType::Button, FInputChord(EKeys::Five));
 
-	UI_COMMAND( ToggleWeightedTangents, "Toggle Weighted Tangents", "Toggles cubic tangents to be weighted/non-weighted", EUserInterfaceActionType::ToggleButton, FInputChord(EModifierKey::Control, EKeys::W));
+	UI_COMMAND( ToggleWeightedTangents, "Toggle Weighted Tangents", "Toggles cubic tangents to be weighted/non-weighted", EUserInterfaceActionType::ToggleButton, FInputChord());
 
 	UI_COMMAND( TrimSectionLeft, "Trim Section Left", "Trim section at current time to the left (keeps the right)", EUserInterfaceActionType::Button, FInputChord(EModifierKey::Control, EKeys::Comma) );
 	UI_COMMAND( TrimSectionRight, "Trim Section Right", "Trim section at current time to the right (keeps the left)", EUserInterfaceActionType::Button, FInputChord(EModifierKey::Control, EKeys::Period) );
