Sequencer: Fix some newly added objects not getting added to selected folders.

Call OnAddTrack for some newly added tracks so that they automatically get added to the selected folders.
Add spawnables to selected folders.

diff --git a/Engine/Source/Editor/MovieSceneTools/Private/TrackEditors/AudioTrackEditor.cpp b/Engine/Source/Editor/MovieSceneTools/Private/TrackEditors/AudioTrackEditor.cpp
--- a/Engine/Source/Editor/MovieSceneTools/Private/TrackEditors/AudioTrackEditor.cpp
+++ b/Engine/Source/Editor/MovieSceneTools/Private/TrackEditors/AudioTrackEditor.cpp
@@ -1007,6 +1007,11 @@ FKeyPropertyResult FAudioTrackEditor::AddNewMasterSound( FFrameNumber KeyTime, U
 	if (TrackResult.bWasCreated)
 	{
 		AudioTrack->SetDisplayName(LOCTEXT("AudioTrackName", "Audio"));
+
+		if (GetSequencer().IsValid())
+		{
+			GetSequencer()->OnAddTrack(AudioTrack);
+		}
 	}
 
 	KeyPropertyResult.bTrackModified = true;
diff --git a/Engine/Source/Editor/Sequencer/Private/Sequencer.cpp b/Engine/Source/Editor/Sequencer/Private/Sequencer.cpp
--- a/Engine/Source/Editor/Sequencer/Private/Sequencer.cpp
+++ b/Engine/Source/Editor/Sequencer/Private/Sequencer.cpp
@@ -3950,6 +3950,15 @@ FGuid FSequencer::MakeNewSpawnable( UObject& Object, UActorFactory* ActorFactory
 		return FGuid();
 	}
 
+	TArray<UMovieSceneFolder*> SelectedParentFolders;
+	FString NewNodePath;
+	CalculateSelectedFolderAndPath(SelectedParentFolders, NewNodePath);
+
+	if (SelectedParentFolders.Num() > 0)
+	{
+		SelectedParentFolders[0]->AddChildObjectBinding(NewGuid);
+	}
+
 	FMovieSceneSpawnable* Spawnable = GetFocusedMovieSceneSequence()->GetMovieScene()->FindSpawnable(NewGuid);
 	if (!Spawnable)
 	{
