Sequencer: Fix event repeater sections not displaying on multiple rows since it doesn't have an explicit layout. Removed restriction on using sub tracks if there's no section layout.

Throb and select the created event section.

diff --git a/Engine/Source/Editor/MovieSceneTools/Private/TrackEditors/EventTrackEditor.cpp b/Engine/Source/Editor/MovieSceneTools/Private/TrackEditors/EventTrackEditor.cpp
--- a/Engine/Source/Editor/MovieSceneTools/Private/TrackEditors/EventTrackEditor.cpp
+++ b/Engine/Source/Editor/MovieSceneTools/Private/TrackEditors/EventTrackEditor.cpp
@@ -306,6 +306,9 @@ void FEventTrackEditor::CreateNewSection(UMovieSceneTrack* Track, int32 RowIndex
 		Track->UpdateEasing();
 
 		SequencerPtr->NotifyMovieSceneDataChanged(EMovieSceneDataChangeType::MovieSceneStructureItemAdded);
+		SequencerPtr->EmptySelection();
+		SequencerPtr->SelectSection(NewSection);
+		SequencerPtr->ThrobSectionSelection();
 	}
 }
 
diff --git a/Engine/Source/Editor/Sequencer/Private/SequencerNodeTree.cpp b/Engine/Source/Editor/Sequencer/Private/SequencerNodeTree.cpp
--- a/Engine/Source/Editor/Sequencer/Private/SequencerNodeTree.cpp
+++ b/Engine/Source/Editor/Sequencer/Private/SequencerNodeTree.cpp
@@ -278,7 +278,7 @@ TSharedRef<FSequencerTrackNode> FSequencerNodeTree::MakeSubTracksAndSectionInter
 		FSequencerSectionLayoutBuilder Builder(CurrentTrackNode);
 		Section->GenerateSectionLayout(Builder);
 
-		if (Builder.HasAnyLayout() && bHasMultipleRows && CurrentTrackNode == ParentNode)
+		if (bHasMultipleRows && CurrentTrackNode == ParentNode)
 		{
 			// Create a new parent node
 			ParentNode = MakeShared<FSequencerTrackNode>(*Track, *Editor, CurrentTrackNode->CanDrag(), nullptr, *this);
