Sequencer: Throb and select newly created sections so that it's clear if there's an overlapping section that gets created.

diff --git a/Engine/Source/Editor/Sequencer/Private/SequencerUtilities.cpp b/Engine/Source/Editor/Sequencer/Private/SequencerUtilities.cpp
--- a/Engine/Source/Editor/Sequencer/Private/SequencerUtilities.cpp
+++ b/Engine/Source/Editor/Sequencer/Private/SequencerUtilities.cpp
@@ -122,6 +122,9 @@ void FSequencerUtilities::PopulateMenu_CreateNewSection(FMenuBuilder& MenuBuilde
 			Track->UpdateEasing();
 
 			Sequencer->NotifyMovieSceneDataChanged(EMovieSceneDataChangeType::MovieSceneStructureItemAdded);
+			Sequencer->EmptySelection();
+			Sequencer->SelectSection(NewSection);
+			Sequencer->ThrobSectionSelection();
 		}
 		else
 		{
