Sequencer: Change so adding a track through the Add Track menu doesn't automatically create a key

diff --git a/Engine/Source/Editor/Sequencer/Private/DisplayNodes/SequencerObjectBindingNode.cpp b/Engine/Source/Editor/Sequencer/Private/DisplayNodes/SequencerObjectBindingNode.cpp
--- a/Engine/Source/Editor/Sequencer/Private/DisplayNodes/SequencerObjectBindingNode.cpp
+++ b/Engine/Source/Editor/Sequencer/Private/DisplayNodes/SequencerObjectBindingNode.cpp
@@ -999,7 +999,7 @@ void FSequencerObjectBindingNode::HandlePropertyMenuItemExecute(FPropertyPath Pr
 		}
 	}
 
-	FKeyPropertyParams KeyPropertyParams(KeyableBoundObjects, PropertyPath, ESequencerKeyMode::ManualKeyForced);
+	FKeyPropertyParams KeyPropertyParams(KeyableBoundObjects, PropertyPath, ESequencerKeyMode::ManualKey);
 
 	Sequencer.KeyProperty(KeyPropertyParams);
 }
