Copying //UE4/Dev-Sequencer to //UE4/Dev-Main (Source: //UE4/Dev-Sequencer @ 4323925)

============================
  MAJOR FEATURES & CHANGES
============================

Change 4046322 by andrew.porter

	QAGame: Adding a sequence for a curve bug

Change 4053249 by andrew.porter

	QAGame: Adding perflevel and dolly sequence to test UE-58712

Change 4053275 by [REDACTED_PERSON]

	Sequencer: Added 'Write FCPXML' option to Render Movie Settings dialog.

	Changed track menu FCPXML export to get default resolution and frame rate from MovieSceneCaptureSettings.

Change 4056521 by [REDACTED_PERSON]

	Sequencer: Fix bOverridesDefault tooltip so that it reads correctly.

Change 4056560 by [REDACTED_PERSON]

	Sequencer: Fix multiple selection transactions when selecting keys.

Change 4057639 by [REDACTED_PERSON]

	QAGame: Fixing marcie material so that its unlit and two sided

Change 4058065 by [REDACTED_PERSON]

	Sequencer: Upon export, log warnings when filename format is not '{shot}'

Change 4058496 by [REDACTED_PERSON]

	Sequencer: Speculative fix for crash OnKeySelected

Change 4062604 by [REDACTED_PERSON]

	Import and Export support for Weighted Tangents

	Main tricky thing was the difference between how we store weighted tangents (length of hypotenuse) and FBX stores weighted, in addtion to the silly FBX fact that the current key stores the next key's arrive tangent.

	Also had to utilize RichCurve since that has support for storing the weight values and the old Matinee InterpFloatValue does not.  Therefore as suggested set a bunch of functions as deprecated and since they call each other warnings appears.

	Also fixed some issues with import and export not supporting auto tangents.

	Originally was going to add weighted tangent interpolation to FRichCurve also but we don't actually touch the SkeletalMesh code here, it's all within the Sequencer code path, so that can go into a seperate changelist, with it's own set of tests.

Change 4063506 by [REDACTED_PERSON]

	Sequencer: Resolve merge conflicts

Change 4067778 by [REDACTED_PERSON]

	Fix for deprecrated warnings caused by deprecated functions calling other deprecated functions. Added pragma's to disable the warnings. Noticed it in other parts of the code also.

Change 4069703 by [REDACTED_PERSON]

	PR #4728: Description Correction (Contributed by [REDACTED_PERSON])

Change 4070512 by [REDACTED_PERSON]

	Fix static analysis warning

Change 4073274 by [REDACTED_PERSON]

	Fix for performance issues with motion trails.

	Found that the main issue was the GetRefTM call.  Refactored the code so that we only call that once per tick, instead of up to  4 times per tick. Also we cache the calculation of the key positions, so we don't need to calculate them again when calculating keys.  Finally the GetRefTM call was recursive which was a decent hint, so we cache out it's parent's first and just get the transform's linearly.

	Also we don't select parents or children to show their trails also, this seemed buggy and possibly very very slow.

Change 4076431 by [REDACTED_PERSON]

	Fix for regression when we added  imported tangent weights.

	There was an Euler Filter to make sure the angles stayed consistent. It seems to still be needed since we could get flips from FBX resampling.
	Was removed from changelist 4062604. Just putting it back in.  Need this in before the next for for the bad maya import.

Change 4077878 by [REDACTED_PERSON]

	Fix for Mac warning, ordering of intializers.

	#trivial

Change 4085194 by [REDACTED_PERSON]

	Sequencer: Add GetObjectBindings() to retrieve the object bindings for the passed in UObject.

Change 4085236 by [REDACTED_PERSON]

	Sequencer: Refresh actor list to clean up "deleted actor" in the outliner when converting to spawnable.

Change 4085899 by [REDACTED_PERSON]

	Sequencer: Export manual focus distance and ensure that the property is animatable.

Change 4094586 by [REDACTED_PERSON]

	Sequencer: Remove Ctrl+O and Shift+O bindings to toggle node collapsing and expanding since they conflict with Open Level...

Change 4096172 by [REDACTED_PERSON]

	We save out Start and End Frames with the Geometry Cache so the component can calculate what frame is getting played at an arbirtrary time similar to skeletal animation assets. We display this in the UI.

	GeometryCacheComponent contains a boolean to ignore the default TickComponent with instead a different  function (TickComponentAtTime) that let's you manually control what time to run the Cache. Jurre mentioned that maybe this should be editor only, but I assume it's needed for sequencer recorder and PIE workflows?

	The rest of the code is the plumbing for the sequencer and movie scene connections, works like a paired down version of skeletal animation.

	One thing to note since we control this manual tick override we needed to set up the template to use the token evaluation mechanism
	so we can Restore State back to the previous override tick value when we aren't evaluating a section.

	We hav use EvalOptions.bCanEvaluateNearestSection = true and EvalOptions.bEvaluateInPreroll = true; in order to support pre-roll.

Change 4097739 by [REDACTED_PERSON]

	PR #4577: UE-56250: Remove debug from shipping builds (Contributed by [REDACTED_PERSON])

Change 4097758 by [REDACTED_PERSON]

	Sequencer: Refresh actor list to clean up "deleted actor" in the outliner when drawing thumbnails

Change 4097852 by [REDACTED_PERSON]

	MovieSceneCapture: Normalize filename for the MovieFolder path and the Output path. Output an error and disable the UI from allowing movie scene capture if the output directory is read only.

Change 4097876 by [REDACTED_PERSON]

	Opening the Add Track menu in Sequencer and trying to add an actor from the map to the sequence should take significantly less time now on maps that have tens of thousands of actors.

	The old method (through a series of complicated events) cleared the binding cache in sequencer and re-made it for every binding in the sequence for every possible actor. The new method gets all bindings once, finds all actors they're bound to, and then checks each actor in the map against that set.

Change 4098527 by [REDACTED_PERSON]

	Sequence Recorder: If removing root animation from the animation asset, push the root animation to the transform recorder, which records to the transform track only when the animation is not in world space.

Change 4098530 by [REDACTED_PERSON]

	Sequence Recorder: When recording animation in world space, the transform keys are compensated by the root bone. There was a bug when not recording animation in world space where there would be no transform keys. This is because the root bone values are all identity (since they're not in world space). To fix this, the transform recorder should record as usual, filling the buffered transforms from the actor, rather than swapping them in during finalize from the root bone.

Change 4098658 by [REDACTED_PERSON]

	AutoKey didn't work with position UMG widgets since we changes are done to a preview widget that is deleted after the drop happens. So to fix we just re-import the changes on the the new preview widget that we get from the widget reference. This will retrigger the auto-key onto the correct preview widget.

Change 4098728 by [REDACTED_PERSON]

	Fix for Win 64 warnings.

	#trivial

Change 4098813 by [REDACTED_PERSON]

	Fix CIS

Change 4098975 by [REDACTED_PERSON]

	CIS Fixes from 4.20 from Max.

Change 4120177 by [REDACTED_PERSON]

	Sequencer: Allow particle activation/deactivation keys to fire even if they're not in the playback range.

	Trigger keys remain uneffected - they still fire only if they're within the playback range.

Change 4122949 by [REDACTED_PERSON]

	Sequencer: Don't reactivate the particle system if it's already active.

Change 4123406 by [REDACTED_PERSON]

	Sequencer: Intersect the section range with the visit range to ensure keys that fall outside the section range aren't snapped against.

Change 4124874 by [REDACTED_PERSON]

	Sequencer: Added ability to manually invalidate eveluation field segments through the debug visualizer
	   - Also fixed layout issues with DPI scaling turned on

Change 4127408 by [REDACTED_PERSON]

	Submitting multiple test assets for verifying functionality around Sequencer Spawnables, located at Game/Developers/JostinBilyeu/Spawnables_Coverage

Change 4133839 by JohnHenry.Carawon

	QAGame: Updates for Playback Controls test content

Change 4143540 by [REDACTED_PERSON]

	QAGame: Updating sequences for Spawnables coverage

Change 4147037 by [REDACTED_PERSON]

	Removes "Show Transport Controls" from the Level Editor Viewport UI as the controls have been completely surpassed by the Cinematic Viewport Preview and you had to have Sequenecer open for the controls to be visible (so they were just a duplicate of the Sequencer window!)

Change 4148126 by [REDACTED_PERSON]

	Renames "Cinematic Preview" to "Allow Cinematic Control" to remove confusion with the difference between Cinematic Preview and Cinematic Viewport.

Change 4148265 by [REDACTED_PERSON]

	Added tooltips to the Cinematic Viewport UI to help clarify what some of the numbers being shown are.

Change 4148582 by [REDACTED_PERSON]

	Renaming Object Bindings in Sequencer now marks the package as modified (so they prompt for save requests) and pass the rename onto the Spawnable so that pasting object bindings doesn't revert to the original name when they were created.

Change 4148682 by [REDACTED_PERSON]

	Fixes various issues copy and pasting Spawnables from one sequence to another and within the same sequence. No longer prints superfluous warnings on paste and properly initializes the new object template so that it does not get garbage collected/cause issues on save.

Change 4170341 by [REDACTED_PERSON]

	Sequencer: Resolve merge conflict with 4.20

Change 4170503 by [REDACTED_PERSON]

	Media Framework: Fix merge conflict

Change 4181586 by [REDACTED_PERSON]

	Cine Camera: Change ActorToTrack to a TSoftObjectPtr<AActor> to support cross level references

Change 4185157 by [REDACTED_PERSON]

	Checking in changes to Spawnables sequencer content, moving files to correct location, altering auto activate flag for cascade emitter inside of sequences Sequencer_Spawnables_Spawning and Sequencer_Spawnables_Despawning

Change 4185605 by [REDACTED_PERSON]

	Checking in M_BaseColor with flag bUsedWithNiagaraSprites set to True

Change 4187236 by [REDACTED_PERSON]

	Sequencer now initiates a slow task when opening via the Cinematic's dropdown in the UI. This helps with first time opening of large cinematics that require compressing audio and animation data.

Change 4187237 by [REDACTED_PERSON]

	Sequencer: Fix GoToEndAndStop() and other jump state so that when calling stop, sequencer gets torn down properly.

	This fixes a bug where if you just call GoToEndAndStop(), the sequence won't actually stop and tear down.

Change 4187291 by [REDACTED_PERSON]

	Sequencer: Fix file extensions for movie files on mac

Change 4189511 by [REDACTED_PERSON]

	Sequencer Marks Implementation

	Feature Documenation: https://docs.google.com/document/d/1dLcasYUWWQzQ_r-6lq90LiNmzT96EaoWuxq-3gHv4-Y/edit?usp=sharing

Change 4189728 by [REDACTED_PERSON]

	Slight updates to Sequencer_Spawnables testing sequences

Change 4189780 by [REDACTED_PERSON]

	Fix Bad Header for previous changelist 4189511

Change 4205029 by [REDACTED_PERSON]

	Check to see if level sequence has itself as a tick priority actor.

Change 4206593 by [REDACTED_PERSON]

	Updating cascade particle assets within sequencer/spawnables test conent to utilize particle toggle track instead of auto activate flag. Deleting deprecated asset Sequencer_Spawnables_DirtyAsset_Possessable.

Change 4208128 by [REDACTED_PERSON]

	Spawnable Object Bindings now show up Gray in Sequencer instead of Red when they have not been spawned yet due to the evaluation of the Spawn track.

Change 4209438 by [REDACTED_PERSON]

	Blueprints: Formalized the concept of a blueprint embedded within an asset for nativization and "Find in Blueprints" purposes

	  - Added IBlueprintAssetHandler to define how blueprints should be retrieved from particular assets, registered by class name
	  - Replaced the majority of UWorld specific logic from the find in blueprint manager with calls to this new handler
	  - All search entries are stored as the full Blueprint path, all unindexed (uncached) entries are stored as the full asset path.
	  - Added a function (SupportsNativization) to UBlueprint to enable specific blueprints to explicitly disallow nativization

Change 4209439 by [REDACTED_PERSON]

	Sequencer: Fixed additional edit widgets not appearing on top-level key area nodes

Change 4209456 by [REDACTED_PERSON]

	Sequencer: Sequencer Event Improvements
	  - New events mandate a Blueprint funtion entry node as an endpoint.
	  - 2 new event sections exist: Triggers and Repeaters.
	  - New director blueprint exists as a child of level sequences.
	  - Added a new concept for editor-only customization of sequence types - IMovieSceneSequenceEditor, analogous to IMovieSceneTrackEditor
	  - Outstanding tasks: FiB integration, Nativization
	  - Potential future work:
	     - Expose the director blueprint more prominently
	     - Usability and UX polish

Change 4209458 by [REDACTED_PERSON]

	EngineTest: Added automation test for new sequencer events

Change 4209921 by [REDACTED_PERSON]

	Fixed missing include

Change 4214704 by [REDACTED_PERSON]

	Sequencer: Reset the last view target when stopping the sequence so that it gets recached when evaluating the sequence again.

Change 4215485 by [REDACTED_PERSON]

	Changed action checked for frame rate radio button to check against value attribute instead of current frame rate because value attribute was the one being set when a frame rate was selected.

Change 4216652 by [REDACTED_PERSON]

	Trivial changes to appease static analysis

Change 4220571 by [REDACTED_PERSON]

	Sequencer: Compensate the added transform when an attach track is added. If there's no existing transform track,  create one and add the compensation to that transform section. If there's an existing transform section without keys, compensate to that transform section. Otherwise, create an additive transform section and compensate to that transform section.

Change 4222198 by [REDACTED_PERSON]

	Sequencer: Geometry cache track needs to update the Elapsed Time in order for the chunk time to update. This fixes an issue where the geometry cache track only evalutes for 5 seconds.

	Copy from Release-4.20

Change 4222199 by [REDACTED_PERSON]

	Sequencer: Forcibly tick skeletal mesh component in case it has already been ticked.

	This fixes first frame T pose issues where skeletal component ticks first, then sequencer resolves bindings which leads to adding tick prerequisites for the next frame.

	Copy from Release-4.20

Change 4222202 by [REDACTED_PERSON]

	Sequencer: Fix CurrentShot LocalTime computation by using sequence time in playback resolution to compute the local shot time.

	Also, fixed the burnin asset so that CurrentShotLocalTime is hooked up to ShotFrame instead of MasterTime.

	This fixes a bug where the burnin's {ShotFrame} is not reporting the local shot frame number.

	Copy from Release-4.20

Change 4223278 by [REDACTED_PERSON]

	EngineTest - Fixing expected results in Bone Positions with Blend Multi Node

	Copy from Release-4.20

Change 4223549 by [REDACTED_PERSON]

	EngineTest: Missed checking in a file from my previous change; should fix Bone Positions with Blend Multi Node automation test

	Copy from Release-4.20

Change 4227337 by [REDACTED_PERSON]

	Sequencer: Add default sections for all tracks (not just property tracks)

Change 4229050 by [REDACTED_PERSON]

	Sequencer: Fix crash in step to next/previous shot.

	Added check for null subsequence data for the root sequence id.

	Copy from Release-4.20

Change 4230668 by [REDACTED_PERSON]

	Update bone position content from Release-4.20

Change 4230676 by [REDACTED_PERSON]

	Sequencer: Pause before begin scrubbing.

	This fixes an issue where if you play the sequence with audio and scrub and release, the audio keeps playing. This is because there's no explicit evaluation in the stopped state OnEndScrubbing.

Change 4230677 by [REDACTED_PERSON]

	Sequence Recorder: Auto size sections after recording to expand to time range of recorded data.

Change 4230678 by [REDACTED_PERSON]

	Sequence Recorder: Fix updating the target level sequence when duplicating before recording into the level sequence.

Change 4230679 by [REDACTED_PERSON]

	Sequencer: Don't attempt to select actors/components on locked levels.

Change 4231440 by [REDACTED_PERSON]

	UE-56235 - Sequencer Possessable references no longer lose their references when saving a map for the first time.

Change 4234819 by [REDACTED_PERSON]

	Sequencer: Fix crash on importing fbx cameras with identical names to existing objects in the level. Add cameras if bound objects don't exist. Changed logs to popup toasts.

Change 4240610 by [REDACTED_PERSON]

	Back out changelist 4227337

Change 4240747 by [REDACTED_PERSON]

	Sequencer: Add default sections for all tracks (not just property tracks)

Change 4242826 by [REDACTED_PERSON]

	Sequencer: Fix duplicated spawnble tracks not getting saved. Clear transient flags from the tracks/sections coming from the FMovieSceneCopyableBinding.

Change 4244413 by [REDACTED_PERSON]

	Sequence Recorder: Fix typo

Change 4247089 by [REDACTED_PERSON]

	Sequencer: Overlay the custom widget (ie. +Section) on the key editors

Change 4247165 by [REDACTED_PERSON]

	Sequencer: Clarified tooltip for export fbx

Change 4247509 by JohnHenry.Carawon

	QAGame: Updates for TM-PlaybackControls test content

Change 4247964 by [REDACTED_PERSON]

	Sequencer: Fix crash splitting a sub section without a sequence.

Change 4255119 by [REDACTED_PERSON]

	Sequencer: Clarified allow bindings from asset tooltips

Change 4257152 by JohnHenry.Carawon

	QAGame: Updates for TM-PlaybackControls

	- Added TextBox and CheckBox UI to set all of the Create Level Sequence Player inputs
	- Added Reset functionality for the Create Level Sequence Player inputs
	- Added Restart functionality to return to the beginning of the test content
	- Adjusted UI layout to fit mobile devices by adjusting Anchors, moving around widget elements and adding Scale Boxes
	- Added Set User Focus nodes to allow controller use for consoles
	- Setup content to test Override Bindings nodes
	- Added toggle to show/hide buttons that activate the sequence control nodes (play,stop,jump..)
	- Separated sequence control buttons into categories and created a menu system using the categories
	- Added button Highlight functionality so we can tell which button has focus when using controllers
	- Added text and functions for the animation widget blueprint nodes to display current time, if it?s playing and if it?s playing forwards
	- Adjusted blueprint layout and comments
	- Moved Scrub node to activate when we activate the Scrub To nodes. The Scrub node just sets a state so all we see is the sequence paused.
	- Fixed a key in the sequence that was causing content to be invisible
	- Tested the level on an android device
	- Tested the level in PIE using a controller

Change 4259663 by JohnHenry.Carawon

	QAGame: Updates to TM-PlaybackControls

	- Added more UMG animations and setup the UMG widget to play all of the animations
	- Adjusted the layout of the examples to be more horizontal. This uses the horizontal space better on mobile devices.
	- Adjusted the UMG widgets so that they match the other widget buttons. They are a row of buttons on the bottom of the screen
	- Changed the Add Binding nodes to account for Allow Bindings from Actor
	- Added Set Input Mode Game Only toggle functionality so that we can test Disable Movement functionality then return to UI control
	- Removed unnecessary sublevel from the sequence

Change 4265883 by [REDACTED_PERSON]

	Fixing scaling of audio tracks in sequencer. Changed to multiplying duration by 1/pitch instead of pitch.

Change 4269110 by [REDACTED_PERSON]

	Sequencer: When re-keying, don't set interpolation on existing keys

Change 4269121 by [REDACTED_PERSON]

	Sequencer: Fix restore state for attach tracks so that the original parent is re-attached

Change 4269165 by [REDACTED_PERSON]

	Fixed FrameNumber numeric input fields not understanding "*" which prevented pressing Ctrl + G and then Enter (to escape focus) when evaluating a sub-frame. Fixed Timecode not displaying Subframe markers.

Change 4271024 by [REDACTED_PERSON]

	Sequencer: Remove ctrl-w as a hotkey to toggle weighted/non-weighted because it conflicts with the more generic binding to duplicate objects/tracks.

Change 4271063 by [REDACTED_PERSON]

	Sequencer: Fix some newly added objects not getting added to selected folders.

	Call OnAddTrack for some newly added tracks so that they automatically get added to the selected folders.
	Add spawnables to selected folders.

Change 4271147 by [REDACTED_PERSON]

	Sequencer: FMovieSceneAccumulatedBlendState no longer consolidates blendable tokens if there are none to blend

Change 4271823 by [REDACTED_PERSON]

	Sequencer: If playback is at the end, calling play again should reset to the beginning.

Change 4272085 by [REDACTED_PERSON]

	Sequencer: Revert update on bind because it will be updated in the sequencer template if the pose hasn't been ticked yet. See 4229000

Change 4275112 by [REDACTED_PERSON]

	Sequencer: Disable audio when playing backwards

Change 4275517 by [REDACTED_PERSON]

	Sequencer: Add property name in ensure when trying to bind to a track with a null object.

Change 4276060 by [REDACTED_PERSON]

	Sequencer: Event tracks now always fill empty space to ensure that they always get evaluated when playing off a finite section

Change 4276063 by [REDACTED_PERSON]

	Sequencer: Removed persistent FMovieSceneEvaluationTemplateInstance and replaced with a completely transient ptr cache

	  - With the introduction of DirectorInstances (4.21), there became a need for a persistent store of data so directors could be kept alive for each sub-sequence
	  - Making the instance map persistent seemed like the best way to achieve this, however, that resulted in pointers to data structures being held when they could be reallocated
	  - Rather than revert back to 4.19's idea of a TransientInstances member that is repopulated every evaluation, the concept of an instance-per-sequence-id has been completely removed and changed for a stack allocated cache of the ptrs required for evaluation. Director instances are now stored in a separate map
	  - This makes it much clearer that the ptrs are not to be held persistently, and reduces the probability of this confusion happening again

Change 4276095 by andrew.porter

	QA-Game: Adding test content for event track in dev folder

Change 4276122 by [REDACTED_PERSON]

	Correct initialization order

Change 4276326 by [REDACTED_PERSON]

	Sequencer: Event repeater sections now only trigger if the play head is over the section

	Changed FMovieSceneEvalTemplate::EvaluateSwept to include the swept range rather than providing a cut-and-sliced context

Change 4276346 by [REDACTED_PERSON]

	Sequencer: Added new (disabled) perf test that can be used for performance profiling of compilation and evaluation

Change 4276754 by andrew.porter

	QAGame: updating event track test content starting to put together BP Interface

Change 4276888 by [REDACTED_PERSON]

	Sequencer: Force tick if a new instance is bound

	This fixes another first frame issue with warmup frames when movie rendering. If we change to tick sequencer first always in its own tick group, this can go away.

Change 4277176 by andrew.porter

	QAGame: Adding blueprint to dev folder

Change 4278894 by [REDACTED_PERSON]

	Appease static analysis

Change 4279867 by [REDACTED_PERSON]

	Fixed GUIDs not being created for call function nodes created from a movie scene event quick binding

Change 4279957 by [REDACTED_PERSON]

	Sequencer: Fix movie scene player with start offset, which was getting nullified because stop() was getting called last.

Change 4279960 by [REDACTED_PERSON]

	Rich Curve: Don't expose RCIM_None and RCTM_None as DisplayName="None". It doesn't make sense to have no interpolation

Change 4279975 by [REDACTED_PERSON]

	Sequencer: Make unique spawnable name to avoid it being renamed in SetActorLabel.

Change 4280019 by andrew.porter

	QAGame: Adding UMG Animation to test level in dev folder

Change 4280319 by [REDACTED_PERSON]

	Sequencer: When evaluating a range of a sequence, all overlapping segments of that range are now guaranteed to be up-to-date and compiled.

Change 4280368 by [REDACTED_PERSON]

	Added extra test cases to the event endpoints tests
	  - Added test cases for event tracks that only contain repeaters, only contain triggers, and a combination of the two.

Change 4282906 by [REDACTED_PERSON]

	Actor Sequence: Check for nullptr when attempting to create director instances

Change 4283773 by andrew.porter

	QAGame: Updating test content with flip flop scenario

Change 4283846 by [REDACTED_PERSON]

	Sequencer: Add double click to open the shot asset directly

Change 4288441 by [REDACTED_PERSON]

	Sequencer: Fix event repeater sections not displaying on multiple rows since it doesn't have an explicit layout. Removed restriction on using sub tracks if there's no section layout.

	Throb and select the created event section.

Change 4288463 by [REDACTED_PERSON]

	Sequencer: Throb and select newly created sections so that it's clear if there's an overlapping section that gets created.

Change 4306485 by [REDACTED_PERSON]

	Sequencer: Increase curve color contrast.

Change 4307888 by [REDACTED_PERSON]

	Sequencer: Change so adding a track through the Add Track menu doesn't automatically create a key

Change 4316090 by laz.matech

	QAGame: Moving content out of developer folder into Sequencer folders. Renamed the FBX_Export_Tests to FBX_Tests

Change 4316331 by [REDACTED_PERSON]

	QAGame: Deleting duplicate files after moving/renaming them  in the content browser

Change 4317353 by [REDACTED_PERSON]

	Back out changelist 4280319

Change 4318859 by [REDACTED_PERSON]

	Sequencer: Fix logic issue in newly added key not expanding section range.

Change 4319666 by [REDACTED_PERSON]

	Sequencer: Rewind to beginning if playing forwards, as defined as positive playrate rather than just the bReversePlayback flag (and vice versa)

Change 4321128 by [REDACTED_PERSON]

	Sequencer: Don't pass auto tangent mode to FBX since it has a different implementation. Default to user as before.

Change 4321389 by [REDACTED_PERSON]

	Sequencer: Added ability to turn off saving animations sequences when exporting level meshes so we don't save the animation sequence twice when saving FBX from the sequencer, which was causes a bug in this file since the times weren't in sync.

Change 4322081 by JohnHenry.Carawon

	QAGame: TM-PlaybackControls

	- Added text for Hide Hud, Hide Player and Pause At End
	- Rearranged text layout
	- Moved the test content out of my developer folder

diff --git a/Engine/Source/Editor/MovieSceneTools/Private/MatineeImportTools.cpp b/Engine/Source/Editor/MovieSceneTools/Private/MatineeImportTools.cpp
--- a/Engine/Source/Editor/MovieSceneTools/Private/MatineeImportTools.cpp
+++ b/Engine/Source/Editor/MovieSceneTools/Private/MatineeImportTools.cpp
@@ -114,7 +114,8 @@ bool FMatineeImportTools::TryConvertMatineeToggleToOutParticleKey( ETrackToggleA
 }
 
 
-void FMatineeImportTools::SetOrAddKey(TMovieSceneChannelData<FMovieSceneFloatValue>& ChannelData, FFrameNumber Time, float Value, float ArriveTangent, float LeaveTangent, EInterpCurveMode MatineeInterpMode )
+void FMatineeImportTools::SetOrAddKey(TMovieSceneChannelData<FMovieSceneFloatValue>& ChannelData, FFrameNumber Time, float Value, float ArriveTangent, float LeaveTangent, EInterpCurveMode MatineeInterpMode
+	, ERichCurveTangentWeightMode WeightedMode, float ArriveTangentWeight, float LeaveTangentWeight )
 {
 	if (ChannelData.FindKey(Time) == INDEX_NONE)
 	{
@@ -124,6 +125,9 @@ void FMatineeImportTools::SetOrAddKey(TMovieSceneChannelData<FMovieSceneFloatVal
 		NewKey.TangentMode = MatineeInterpolationToRichCurveTangent( MatineeInterpMode );
 		NewKey.Tangent.ArriveTangent = ArriveTangent;
 		NewKey.Tangent.LeaveTangent = LeaveTangent;
+		NewKey.Tangent.TangentWeightMode = WeightedMode;
+		NewKey.Tangent.ArriveTangentWeight = ArriveTangentWeight;
+		NewKey.Tangent.LeaveTangentWeight = LeaveTangentWeight;
 		ChannelData.AddKey( Time, NewKey );
 	}
 }
diff --git a/Engine/Source/Editor/MovieSceneTools/Public/MatineeImportTools.h b/Engine/Source/Editor/MovieSceneTools/Public/MatineeImportTools.h
--- a/Engine/Source/Editor/MovieSceneTools/Public/MatineeImportTools.h
+++ b/Engine/Source/Editor/MovieSceneTools/Public/MatineeImportTools.h
@@ -47,7 +47,8 @@ public:
 	static bool TryConvertMatineeToggleToOutParticleKey( ETrackToggleAction ToggleAction, EParticleKey& OutParticleKey );
 
 	/** Adds a key to a rich curve based on matinee curve key data. */
-	static void SetOrAddKey( TMovieSceneChannelData<FMovieSceneFloatValue>& Curve, FFrameNumber Time, float Value, float ArriveTangent, float LeaveTangent, EInterpCurveMode MatineeInterpMode );
+	static void SetOrAddKey( TMovieSceneChannelData<FMovieSceneFloatValue>& Curve, FFrameNumber Time, float Value, float ArriveTangent, float LeaveTangent, EInterpCurveMode MatineeInterpMode
+	, ERichCurveTangentWeightMode WeightedMode = RCTWM_WeightedNone, float ArriveTangentWeight = 0.0f, float LeaveTangentWeight = 0.0f);
 
 	/** Copies keys from a matinee bool track to a sequencer bool track. */
 	static bool CopyInterpBoolTrack( UInterpTrackBoolProp* MatineeBoolTrack, UMovieSceneBoolTrack* BoolTrack );
diff --git a/Engine/Source/Editor/Sequencer/Private/SequencerEdMode.cpp b/Engine/Source/Editor/Sequencer/Private/SequencerEdMode.cpp
--- a/Engine/Source/Editor/Sequencer/Private/SequencerEdMode.cpp
+++ b/Engine/Source/Editor/Sequencer/Private/SequencerEdMode.cpp
@@ -107,7 +107,7 @@ void FSequencerEdMode::DrawHUD(FEditorViewportClient* ViewportClient,FViewport*
 {
 	FEdMode::DrawHUD(ViewportClient,Viewport,View,Canvas);
 
-	if( ViewportClient->AllowsCinematicPreview() )
+	if( ViewportClient->AllowsCinematicControl() )
 	{
 		// Get the size of the viewport
 		const int32 SizeX = Viewport->GetSizeXY().X;
@@ -142,6 +142,8 @@ void FSequencerEdMode::OnKeySelected(FViewport* Viewport, HMovieSceneKeyProxy* K
 
 	for (TWeakPtr<FSequencer> WeakSequencer : Sequencers)
 	{
+		bool bChangedSelection = false;
+
 		TSharedPtr<FSequencer> Sequencer = WeakSequencer.Pin();
 		if (Sequencer.IsValid())
 		{
@@ -150,6 +152,12 @@ void FSequencerEdMode::OnKeySelected(FViewport* Viewport, HMovieSceneKeyProxy* K
 			FSequencerSelection& Selection = Sequencer->GetSelection();
 			if (!bAddToSelection && !bToggleSelection)
 			{
+				if (!bChangedSelection)
+				{
+					Sequencer->GetSelection().SuspendBroadcast();
+					bChangedSelection = true;
+				}
+
 				Sequencer->GetSelection().EmptySelectedKeys();
 			}
 
@@ -174,6 +182,12 @@ void FSequencerEdMode::OnKeySelected(FViewport* Viewport, HMovieSceneKeyProxy* K
 							TSharedPtr<IKeyArea> KeyArea = KeyAreaNode->GetKeyArea(Section);
 							if (KeyArea.IsValid() && KeyArea->GetName() == KeyData.ChannelName)
 							{
+								if (!bChangedSelection)
+								{
+									Sequencer->GetSelection().SuspendBroadcast();
+									bChangedSelection = true;
+								}
+
 								Sequencer->SelectKey(Section, KeyArea, KeyData.KeyHandle.GetValue(), bToggleSelection);
 								goto next_key;
 							}
@@ -183,6 +197,12 @@ void FSequencerEdMode::OnKeySelected(FViewport* Viewport, HMovieSceneKeyProxy* K
 next_key:
 				continue;
 			}
+			if (bChangedSelection)
+			{
+				Sequencer->GetSelection().ResumeBroadcast();
+				Sequencer->GetSelection().GetOnKeySelectionChanged().Broadcast();
+				Sequencer->GetSelection().GetOnOutlinerNodeSelectionChangedObjectGuids().Broadcast();
+			}
 		}
 	}
 }
@@ -309,7 +329,7 @@ bool FSequencerEdMode::GetParentTM(FTransform& CurrentRefTM, const TSharedPtr<FS
 			}
 		}
 	}
-
+	
 	return false;
 }
 
@@ -354,7 +374,6 @@ FTransform FSequencerEdMode::GetRefFrame(const TSharedPtr<FSequencer>& Sequencer
 FTransform FSequencerEdMode::GetRefFrame(const TSharedPtr<FSequencer>& Sequencer, const UObject* InObject, FFrameTime KeyTime)
 {
 	FTransform RefTM = FTransform::Identity;
-
 	const AActor* Actor = Cast<AActor>(InObject);
 	if (Actor != nullptr)
 	{
@@ -477,14 +496,14 @@ void FSequencerEdMode::GetLocationAtTime(FMovieSceneEvaluationTrack* Track, UObj
 }
 
 void FSequencerEdMode::DrawTransformTrack(const TSharedPtr<FSequencer>& Sequencer, FPrimitiveDrawInterface* PDI,
-	UMovieScene3DTransformTrack* TransformTrack, const TArray<TWeakObjectPtr<UObject>>& BoundObjects, const bool bIsSelected)
+											UMovieScene3DTransformTrack* TransformTrack, const TArray<TWeakObjectPtr<UObject>>& BoundObjects, const bool bIsSelected)
 {
 	bool bHitTesting = true;
-	if (PDI != nullptr)
+	if( PDI != nullptr )
 	{
 		bHitTesting = PDI->IsHitTesting();
 	}
-
+	
 	ASequencerMeshTrail* TrailActor = nullptr;
 	// Get the Trail Actor associated with this track if we are drawing mesh trails
 	if (bDrawMeshTrails)
@@ -501,24 +520,24 @@ void FSequencerEdMode::DrawTransformTrack(const TSharedPtr<FSequencer>& Sequence
 
 	bool bShowTrajectory = TransformTrack->GetAllSections().ContainsByPredicate(
 		[bIsSelected](UMovieSceneSection* Section)
-	{
-		UMovieScene3DTransformSection* TransformSection = Cast<UMovieScene3DTransformSection>(Section);
-		if (TransformSection)
 		{
-			switch (TransformSection->GetShow3DTrajectory())
+			UMovieScene3DTransformSection* TransformSection = Cast<UMovieScene3DTransformSection>(Section);
+			if (TransformSection)
 			{
-			case EShow3DTrajectory::EST_Always:				return true;
-			case EShow3DTrajectory::EST_Never:				return false;
-			case EShow3DTrajectory::EST_OnlyWhenSelected:	return bIsSelected;
+				switch (TransformSection->GetShow3DTrajectory())
+				{
+				case EShow3DTrajectory::EST_Always:				return true;
+				case EShow3DTrajectory::EST_Never:				return false;
+				case EShow3DTrajectory::EST_OnlyWhenSelected:	return bIsSelected;
+				}
 			}
+			return false;
 		}
-		return false;
-	}
 	);
-
+	
 	FFrameRate TickResolution = Sequencer->GetFocusedTickResolution();
 	FMovieSceneEvaluationTemplate* Template = Sequencer->GetEvaluationTemplate().FindTemplate(Sequencer->GetFocusedTemplateID());
-	if (!bShowTrajectory || !Template || !TransformTrack->GetAllSections().ContainsByPredicate([](UMovieSceneSection* In) { return In->IsActive(); }))
+	if (!bShowTrajectory || !Template || !TransformTrack->GetAllSections().ContainsByPredicate([](UMovieSceneSection* In){ return In->IsActive(); }))
 	{
 		return;
 	}
@@ -539,7 +558,7 @@ void FSequencerEdMode::DrawTransformTrack(const TSharedPtr<FSequencer>& Sequence
 			GetParents(Parents, BoundObject);
 
 			FVector OldKeyPos(0);
-			FFrameTime OldKeyTime = 0;
+ 			FFrameTime OldKeyTime = 0;
 			int KeyTimeIndex = 0;
 			FTransform OldPosRefTM;;
 			FVector OldPos_G;
@@ -553,7 +572,7 @@ void FSequencerEdMode::DrawTransformTrack(const TSharedPtr<FSequencer>& Sequence
 				FVector WorldPosition;
 				FTrajectoryKey TrajectoryKey;
 
-				FKeyPositionRotation(const FTrajectoryKey &InTrajKey, const FVector &InPos, const FRotator &InRot, const FVector &InWorldPos) :
+				FKeyPositionRotation(const FTrajectoryKey &InTrajKey, const FVector &InPos, const FRotator &InRot, const FVector &InWorldPos) : 
 					 Position(InPos), Rotation(InRot), WorldPosition(InWorldPos), TrajectoryKey(InTrajKey) {}
 			};
 			TArray<FKeyPositionRotation> KeyPosRots;
@@ -563,31 +582,31 @@ void FSequencerEdMode::DrawTransformTrack(const TSharedPtr<FSequencer>& Sequence
 				FFrameTime NewKeyTime = NewTrajectoryKey.Time;
 
 				FVector NewKeyPos(0);
-				FRotator NewKeyRot(0, 0, 0);
+				FRotator NewKeyRot(0,0,0);
 
 				GetLocationAtTime(EvalTrack, BoundObject, NewKeyTime, NewKeyPos, NewKeyRot, Sequencer);
 				FTransform NewPosRefTM = GetRefFrame(Sequencer, Parents, NewKeyTime);
 				FVector NewKeyPos_G = NewPosRefTM.TransformPosition(NewKeyPos);
-				FKeyPositionRotation KeyPosRot(NewTrajectoryKey, NewKeyPos, NewKeyRot, NewKeyPos_G);
+				FKeyPositionRotation KeyPosRot(NewTrajectoryKey,NewKeyPos, NewKeyRot, NewKeyPos_G);
 				KeyPosRots.Push(KeyPosRot);
 				// If not the first keypoint, draw a line to the last keypoint.
-				if (KeyTimeIndex > 0)
+				if(KeyTimeIndex > 0)
 				{
 					int32 NumSteps = FMath::CeilToInt((TickResolution.AsSeconds(NewKeyTime) - TickResolution.AsSeconds(OldKeyTime)) / SequencerEdMode_Draw3D::DrawTrackTimeRes);
 					// Limit the number of steps to prevent a rendering performance hit
-					NumSteps = FMath::Min(100, NumSteps);
-					FFrameTime DrawSubstep = NumSteps == 0 ? 0 : (NewKeyTime - OldKeyTime)*(1.f / NumSteps);
+					NumSteps = FMath::Min( 100, NumSteps );
+					FFrameTime DrawSubstep = NumSteps == 0 ? 0 : (NewKeyTime - OldKeyTime)*(1.f/NumSteps);
 					// Find position on first keyframe.
 					FFrameTime OldTime = OldKeyTime;
 					FVector OldPos(0);
-					FRotator OldRot(0, 0, 0);
+					FRotator OldRot(0,0,0);
 					GetLocationAtTime(EvalTrack, BoundObject, OldKeyTime, OldPos, OldRot, Sequencer);
 
 					const bool bIsConstantKey = NewTrajectoryKey.Is(ERichCurveInterpMode::RCIM_Constant);
 					// For constant interpolation - don't draw ticks - just draw dotted line.
 					if (bIsConstantKey)
 					{
-						if (PDI != nullptr)
+						if(PDI != nullptr)
 						{
 							DrawDashedLine(PDI, OldPos_G, NewKeyPos_G, TrackColor, 20, SDPG_Foreground);
 						}
@@ -595,12 +614,12 @@ void FSequencerEdMode::DrawTransformTrack(const TSharedPtr<FSequencer>& Sequence
 					else
 					{
 						// Then draw a line for each substep.
-						for (int32 j = 1; j<NumSteps + 1; j++)
+						for (int32 j=1; j<NumSteps+1; j++)
 						{
 							FFrameTime NewTime = OldKeyTime + DrawSubstep*j;
 
 							FVector NewPos(0);
-							FRotator NewRot(0, 0, 0);
+							FRotator NewRot(0,0,0);
 							GetLocationAtTime(EvalTrack, BoundObject, NewTime, NewPos, NewRot, Sequencer);
 
 							FTransform RefTM = GetRefFrame(Sequencer, Parents, NewTime);
@@ -637,9 +656,9 @@ void FSequencerEdMode::DrawTransformTrack(const TSharedPtr<FSequencer>& Sequence
 
 			// Draw keypoints on top of curve
 			FColor KeyColor = TrackColor.ToFColor(true);
-			for (const FKeyPositionRotation& KeyPosRot : KeyPosRots)
+			for (const FKeyPositionRotation& KeyPosRot: KeyPosRots)
 			{
-				if (bHitTesting && PDI)
+				if (bHitTesting && PDI) 
 				{
 					PDI->SetHitProxy(new HMovieSceneKeyProxy(TransformTrack, KeyPosRot.TrajectoryKey));
 				}
@@ -666,7 +685,7 @@ void FSequencerEdMode::DrawTransformTrack(const TSharedPtr<FSequencer>& Sequence
 					}
 				}
 
-				if (bHitTesting && PDI)
+				if (bHitTesting && PDI) 
 				{
 					PDI->SetHitProxy(nullptr);
 				}
@@ -691,7 +710,7 @@ void FSequencerEdMode::DrawTracks3D(FPrimitiveDrawInterface* PDI)
 		// Map between object binding nodes and selection
 		TMap<TSharedRef<FSequencerDisplayNode>, bool > ObjectBindingNodesSelectionMap;
 
-		for (auto ObjectBinding : Sequencer->GetNodeTree()->GetObjectBindingMap())
+		for (auto ObjectBinding : Sequencer->GetNodeTree()->GetObjectBindingMap() )
 		{
 			if (!ObjectBinding.Value.IsValid())
 			{
@@ -730,8 +749,8 @@ void FSequencerEdMode::DrawTracks3D(FPrimitiveDrawInterface* PDI)
 							});
 							if (TrailPtr == nullptr)
 							{
-								UViewportWorldInteraction* WorldInteraction = Cast<UViewportWorldInteraction>(GEditor->GetEditorWorldExtensionsManager()->GetEditorWorldExtensions(GetWorld())->FindExtension(UViewportWorldInteraction::StaticClass()));
-								if (WorldInteraction != nullptr)
+								UViewportWorldInteraction* WorldInteraction = Cast<UViewportWorldInteraction>( GEditor->GetEditorWorldExtensionsManager()->GetEditorWorldExtensions( GetWorld() )->FindExtension( UViewportWorldInteraction::StaticClass() ) );
+								if( WorldInteraction != nullptr )
 								{
 									ASequencerMeshTrail* TrailActor = WorldInteraction->SpawnTransientSceneActor<ASequencerMeshTrail>(TEXT("SequencerMeshTrail"), true);
 									FMeshTrailData MeshTrail = FMeshTrailData(TransformTrack, TrailActor);
@@ -748,7 +767,6 @@ void FSequencerEdMode::DrawTracks3D(FPrimitiveDrawInterface* PDI)
 	}
 }
 
-
 FSequencerEdModeTool::FSequencerEdModeTool(FSequencerEdMode* InSequencerEdMode) :
 	SequencerEdMode(InSequencerEdMode)
 {
