Need to extend section before keying so interrogation happens.

diff --git a/Engine/Source/Editor/Sequencer/Private/SKeyNavigationButtons.h b/Engine/Source/Editor/Sequencer/Private/SKeyNavigationButtons.h
--- a/Engine/Source/Editor/Sequencer/Private/SKeyNavigationButtons.h
+++ b/Engine/Source/Editor/Sequencer/Private/SKeyNavigationButtons.h
@@ -285,6 +285,9 @@ public:
 		}
 		DisplayNode->GetChildKeyAreaNodesRecursively(KeyAreaNodes);
 
+		//Need to key first since we may need to interrogate the section
+		SectionToKey->ExpandToFrame(CurrentTime.FrameNumber);
+
 		for (TSharedRef<FSequencerSectionKeyAreaNode> KeyAreaNode : KeyAreaNodes)
 		{
 			TSharedPtr<IKeyArea> KeyArea =  KeyAreaNode->GetKeyArea(SectionToKey);
@@ -294,8 +297,6 @@ public:
 			}
 		}
 
-		SectionToKey->ExpandToFrame(CurrentTime.FrameNumber);
-
 		Sequencer.NotifyMovieSceneDataChanged(EMovieSceneDataChangeType::TrackValueChanged);
 		Sequencer.UpdatePlaybackRange();
 
