Sequencer: Snap the local time after changing fps

diff --git a/Engine/Source/Editor/Sequencer/Private/SSequencerPlayRateCombo.cpp b/Engine/Source/Editor/Sequencer/Private/SSequencerPlayRateCombo.cpp
--- a/Engine/Source/Editor/Sequencer/Private/SSequencerPlayRateCombo.cpp
+++ b/Engine/Source/Editor/Sequencer/Private/SSequencerPlayRateCombo.cpp
@@ -359,6 +359,9 @@ void SSequencerPlayRateCombo::SetDisplayRate(FFrameRate InFrameRate)
 			}
 		}
 	}
+
+	// Snap the local time to the new display rate
+	Sequencer->SetLocalTime(Sequencer->GetLocalTime().Time, ESnapTimeMode::STM_Interval);
 }
 
 FFrameRate SSequencerPlayRateCombo::GetDisplayRate() const
