Sequencer: Restore the focused widget that was cleared when setting global time

diff --git a/Engine/Source/Editor/Sequencer/Private/Sequencer.cpp b/Engine/Source/Editor/Sequencer/Private/Sequencer.cpp
--- a/Engine/Source/Editor/Sequencer/Private/Sequencer.cpp
+++ b/Engine/Source/Editor/Sequencer/Private/Sequencer.cpp
@@ -2299,6 +2299,8 @@ void FSequencer::SetLocalTimeDirectly(FFrameTime NewTime)
 
 void FSequencer::SetGlobalTime(FFrameTime NewTime)
 {
+	TSharedPtr<SWidget> PreviousFocusedWidget = FSlateApplication::Get().GetKeyboardFocusedWidget();
+
 	// Clear focus before setting time in case there's a key editor value selected that gets committed to a newly selected key on UserMovedFocus
 	FSlateApplication::Get().ClearKeyboardFocus(EFocusCause::Cleared);
 
@@ -2321,6 +2323,11 @@ void FSequencer::SetGlobalTime(FFrameTime NewTime)
 		SetPlaybackStatus(EMovieScenePlayerStatus::Stopped);
 		AutoScrubTarget.Reset();
 	}
+
+	if (PreviousFocusedWidget.IsValid())
+	{
+		FSlateApplication::Get().SetKeyboardFocus(PreviousFocusedWidget);
+	}
 }
 
 void FSequencer::ForceEvaluate()
