Sequencer: Clear focus only when stopped

diff --git a/Engine/Plugins/FX/Niagara/Source/NiagaraEditor/Private/ViewModels/NiagaraSystemViewModel.cpp b/Engine/Plugins/FX/Niagara/Source/NiagaraEditor/Private/ViewModels/NiagaraSystemViewModel.cpp
--- a/Engine/Plugins/FX/Niagara/Source/NiagaraEditor/Private/ViewModels/NiagaraSystemViewModel.cpp
+++ b/Engine/Plugins/FX/Niagara/Source/NiagaraEditor/Private/ViewModels/NiagaraSystemViewModel.cpp
@@ -935,7 +935,7 @@ void FNiagaraSystemViewModel::ResetSystemInternal(bool bCanResetTime)
 		TGuardValue<bool> Guard(bSettingSequencerTimeDirectly, true);
 		if (Sequencer->GetPlaybackStatus() == EMovieScenePlayerStatus::Playing)
 		{
-			Sequencer->SetPlaybackStatus(EMovieScenePlayerStatus::Stopped);
+			Sequencer->SetPlaybackStatus(EMovieScenePlayerStatus::Paused);
 			Sequencer->SetGlobalTime(0);
 			Sequencer->SetPlaybackStatus(EMovieScenePlayerStatus::Playing);
 		}
diff --git a/Engine/Source/Editor/Sequencer/Private/Sequencer.cpp b/Engine/Source/Editor/Sequencer/Private/Sequencer.cpp
--- a/Engine/Source/Editor/Sequencer/Private/Sequencer.cpp
+++ b/Engine/Source/Editor/Sequencer/Private/Sequencer.cpp
@@ -2302,7 +2302,10 @@ void FSequencer::SetGlobalTime(FFrameTime NewTime)
 	TSharedPtr<SWidget> PreviousFocusedWidget = FSlateApplication::Get().GetKeyboardFocusedWidget();
 
 	// Clear focus before setting time in case there's a key editor value selected that gets committed to a newly selected key on UserMovedFocus
-	FSlateApplication::Get().ClearKeyboardFocus(EFocusCause::Cleared);
+	if (GetPlaybackStatus() == EMovieScenePlayerStatus::Stopped)
+	{
+		FSlateApplication::Get().ClearKeyboardFocus(EFocusCause::Cleared);
+	}
 
 	NewTime = ConvertFrameTime(NewTime, GetRootTickResolution(), PlayPosition.GetInputRate());
 	if (PlayPosition.GetEvaluationType() == EMovieSceneEvaluationType::FrameLocked)
