Sequencer: Check for null movie scene folder

diff --git a/Engine/Source/Editor/Sequencer/Private/SequencerNodeTree.cpp b/Engine/Source/Editor/Sequencer/Private/SequencerNodeTree.cpp
--- a/Engine/Source/Editor/Sequencer/Private/SequencerNodeTree.cpp
+++ b/Engine/Source/Editor/Sequencer/Private/SequencerNodeTree.cpp
@@ -408,9 +408,12 @@ void FSequencerNodeTree::CreateAndPopulateFolderNodes(
 		ObjectGuidToDisplayNodeMap.Add( ObjectBindingNode->GetObjectBinding(), ObjectBindingNode );
 	}
 
-	for ( UMovieSceneFolder* MovieSceneFolder : MovieSceneFolders )
+	for (UMovieSceneFolder* MovieSceneFolder : MovieSceneFolders)
 	{
-		FolderAndObjectNodes.Add( CreateFolderNode( *MovieSceneFolder, *this, MasterTrackToDisplayNodeMap, ObjectGuidToDisplayNodeMap ) );	
+		if (MovieSceneFolder != nullptr)
+		{
+			FolderAndObjectNodes.Add(CreateFolderNode(*MovieSceneFolder, *this, MasterTrackToDisplayNodeMap, ObjectGuidToDisplayNodeMap));
+		}
 	}
 
 	TArray<TSharedRef<FSequencerTrackNode>> NonFolderTrackNodes;
