Sequencer: Check for null movie scene folder

diff --git a/Engine/Source/Editor/Sequencer/Private/SequencerNodeTree.cpp b/Engine/Source/Editor/Sequencer/Private/SequencerNodeTree.cpp
--- a/Engine/Source/Editor/Sequencer/Private/SequencerNodeTree.cpp
+++ b/Engine/Source/Editor/Sequencer/Private/SequencerNodeTree.cpp
@@ -361,9 +361,12 @@ TSharedRef<FSequencerDisplayNode> CreateFolderNode(
 {
 	TSharedRef<FSequencerFolderNode> FolderNode( new FSequencerFolderNode( MovieSceneFolder, TSharedPtr<FSequencerDisplayNode>(), NodeTree ) );
 
-	for ( UMovieSceneFolder* ChildFolder : MovieSceneFolder.GetChildFolders() )
+	for (UMovieSceneFolder* ChildFolder : MovieSceneFolder.GetChildFolders())
 	{
-		FolderNode->AddChildNode( CreateFolderNode( *ChildFolder, NodeTree, MasterTrackToDisplayNodeMap, ObjectGuidToDisplayNodeMap ) );
+		if (ChildFolder != nullptr)
+		{
+			FolderNode->AddChildNode(CreateFolderNode(*ChildFolder, NodeTree, MasterTrackToDisplayNodeMap, ObjectGuidToDisplayNodeMap));
+		}
 	}
 
 	for ( UMovieSceneTrack* MasterTrack : MovieSceneFolder.GetChildMasterTracks() )
