Sequencer: Add hotkey (control A) to add selected actors to sequencer

diff --git a/Engine/Source/Editor/Sequencer/Private/Sequencer.cpp b/Engine/Source/Editor/Sequencer/Private/Sequencer.cpp
--- a/Engine/Source/Editor/Sequencer/Private/Sequencer.cpp
+++ b/Engine/Source/Editor/Sequencer/Private/Sequencer.cpp
@@ -7867,6 +7867,22 @@ void FSequencer::ToggleExpandCollapseNodesAndDescendants()
 }
 
 
+void FSequencer::AddSelectedActors()
+{
+	USelection* ActorSelection = GEditor->GetSelectedActors();
+	TArray<TWeakObjectPtr<AActor> > SelectedActors;
+	for (FSelectionIterator Iter(*ActorSelection); Iter; ++Iter)
+	{
+		AActor* Actor = Cast<AActor>(*Iter);
+		if (Actor)
+		{
+			SelectedActors.Add(Actor);
+		}
+	}
+
+	AddActors(SelectedActors);
+}
+
 void FSequencer::SetKey()
 {
 	FScopedTransaction SetKeyTransaction( NSLOCTEXT("Sequencer", "SetKey_Transaction", "Set Key") );
@@ -9122,6 +9138,10 @@ void FSequencer::BindCommands()
 		Commands.ToggleExpandCollapseNodesAndDescendants,
 		FExecuteAction::CreateSP(this, &FSequencer::ToggleExpandCollapseNodesAndDescendants));
 
+	SequencerCommandBindings->MapAction(
+		Commands.AddActorsToSequencer,
+		FExecuteAction::CreateSP( this, &FSequencer::AddSelectedActors));
+
 	SequencerCommandBindings->MapAction(
 		Commands.SetKey,
 		FExecuteAction::CreateSP( this, &FSequencer::SetKey ) );
diff --git a/Engine/Source/Editor/Sequencer/Private/Sequencer.h b/Engine/Source/Editor/Sequencer/Private/Sequencer.h
--- a/Engine/Source/Editor/Sequencer/Private/Sequencer.h
+++ b/Engine/Source/Editor/Sequencer/Private/Sequencer.h
@@ -896,6 +896,9 @@ protected:
 	/** Sort all nodes and their descendants by category then alphabetically */
 	void SortAllNodesAndDescendants();
 
+	/** Add selected actors to sequencer */
+	void AddSelectedActors();
+
 	/** Manually sets a key for the selected objects at the current time */
 	void SetKey();
 
diff --git a/Engine/Source/Editor/Sequencer/Private/SequencerCommands.cpp b/Engine/Source/Editor/Sequencer/Private/SequencerCommands.cpp
--- a/Engine/Source/Editor/Sequencer/Private/SequencerCommands.cpp
+++ b/Engine/Source/Editor/Sequencer/Private/SequencerCommands.cpp
@@ -51,6 +51,7 @@ void FSequencerCommands::RegisterCommands()
 	UI_COMMAND( SelectSectionsInSelectionRange, "Select Sections in Selection Range", "Select all sections that fall into the selection range", EUserInterfaceActionType::Button, FInputChord() );
 	UI_COMMAND( SelectAllInSelectionRange, "Select All in Selection Range", "Select all keys and section that fall into the selection range", EUserInterfaceActionType::Button, FInputChord() );
 
+	UI_COMMAND( AddActorsToSequencer, "Add Actors", "Add actors to sequencer", EUserInterfaceActionType::Button, FInputChord(EModifierKey::Control, EKeys::A) );
 	UI_COMMAND( SetKey, "Set Key", "Sets a key on the selected tracks", EUserInterfaceActionType::Button, FInputChord(EKeys::Enter) );
 	UI_COMMAND( SetInterpolationCubicAuto, "Set Key Auto", "Cubic interpolation - Automatic tangents", EUserInterfaceActionType::Button, FInputChord(EKeys::One));
 	UI_COMMAND( SetInterpolationCubicUser, "Set Key User", "Cubic interpolation - User flat tangents", EUserInterfaceActionType::Button, FInputChord(EKeys::Two));
diff --git a/Engine/Source/Editor/Sequencer/Public/SequencerCommands.h b/Engine/Source/Editor/Sequencer/Public/SequencerCommands.h
--- a/Engine/Source/Editor/Sequencer/Public/SequencerCommands.h
+++ b/Engine/Source/Editor/Sequencer/Public/SequencerCommands.h
@@ -136,6 +136,9 @@ public:
 	/** Select all keys and sections that fall into the selection range*/
 	TSharedPtr< FUICommandInfo > SelectAllInSelectionRange;
 
+	/** Add selected actors to sequencer */
+	TSharedPtr< FUICommandInfo > AddActorsToSequencer;
+
 	/** Sets a key at the current time for the selected actor */
 	TSharedPtr< FUICommandInfo > SetKey;
 
