Sequencer: Expand parents when scrolling into view. A la world outliner.

diff --git a/Engine/Source/Editor/Sequencer/Private/SSequencer.cpp b/Engine/Source/Editor/Sequencer/Private/SSequencer.cpp
--- a/Engine/Source/Editor/Sequencer/Private/SSequencer.cpp
+++ b/Engine/Source/Editor/Sequencer/Private/SSequencer.cpp
@@ -722,6 +722,13 @@ void SSequencer::HandleOutlinerNodeSelectionChanged()
 	{
 		for ( auto& Node : OutlinerSelection )
 		{
+			auto Parent = Node->GetParent();
+			while (Parent.IsValid())
+			{
+				TreeView->SetItemExpansion(Parent->AsShared(), true);
+				Parent = Parent->GetParent();
+			}
+
 			TreeView->RequestScrollIntoView( Node );
 			break;
 		}
@@ -1659,6 +1666,13 @@ void SSequencer::UpdateLayoutTree()
 			{
 				if (Node->GetPathName() == NodePath)
 				{
+					auto Parent = Node->GetParent();
+					while (Parent.IsValid())
+					{
+						TreeView->SetItemExpansion(Parent->AsShared(), true);
+						Parent = Parent->GetParent();
+					}
+
 					TreeView->RequestScrollIntoView(Node);
 					break;
 				}
