Sequencer: Always render from the root (transform the range)

diff --git a/Engine/Source/Editor/Sequencer/Private/Sequencer.cpp b/Engine/Source/Editor/Sequencer/Private/Sequencer.cpp
--- a/Engine/Source/Editor/Sequencer/Private/Sequencer.cpp
+++ b/Engine/Source/Editor/Sequencer/Private/Sequencer.cpp
@@ -2534,13 +2534,25 @@ void FSequencer::RenderMovieInternal(TRange<FFrameNumber> Range, bool bSetFrameO
 		Range = TRange<FFrameNumber>::Hull(Range, GetPlaybackRange());
 	}
 
+	// If focused on a subsequence, transform the playback range to the root in order to always render from the root
+	if (GetRootMovieSceneSequence() != GetFocusedMovieSceneSequence())
+	{
+		bSetFrameOverrides = true;
+
+		if (const FMovieSceneSubSequenceData* SubSequenceData = RootTemplateInstance.GetHierarchy().FindSubData(GetFocusedTemplateID()))
+		{
+			Range = Range * SubSequenceData->RootToSequenceTransform.Inverse();
+		}
+	}
+
 	FLevelEditorModule& LevelEditorModule = FModuleManager::GetModuleChecked<FLevelEditorModule>(TEXT("LevelEditor"));
 
 	// Create a new movie scene capture object for an automated level sequence, and open the tab
 	UAutomatedLevelSequenceCapture* MovieSceneCapture = NewObject<UAutomatedLevelSequenceCapture>(GetTransientPackage(), UAutomatedLevelSequenceCapture::StaticClass(), UMovieSceneCapture::MovieSceneCaptureUIName, RF_Transient);
 	MovieSceneCapture->LoadFromConfig();
 
-	MovieSceneCapture->LevelSequenceAsset = GetCurrentAsset()->GetPathName();
+	// Always render from the root
+	MovieSceneCapture->LevelSequenceAsset = GetRootMovieSceneSequence()->GetMovieScene()->GetOuter()->GetPathName();
 
 	FFrameRate DisplayRate = GetFocusedDisplayRate();
 	FFrameRate TickResolution = GetFocusedTickResolution();
