Sequencer: Fixed off-by-one error that led to invalid sections not being removed from their track nodes correctly

diff --git a/Engine/Source/Editor/Sequencer/Private/DisplayNodes/SequencerTrackNode.cpp b/Engine/Source/Editor/Sequencer/Private/DisplayNodes/SequencerTrackNode.cpp
--- a/Engine/Source/Editor/Sequencer/Private/DisplayNodes/SequencerTrackNode.cpp
+++ b/Engine/Source/Editor/Sequencer/Private/DisplayNodes/SequencerTrackNode.cpp
@@ -234,7 +234,7 @@ void FSequencerTrackNode::UpdateSections()
 	const int32 NumToRemove = Sections.Num() - CurrentSections.Num();
 	if (NumToRemove > 0)
 	{
-		Sections.RemoveAt(Sections.Num()-1-NumToRemove, NumToRemove, true);
+		Sections.RemoveAt(Sections.Num()-NumToRemove, NumToRemove, true);
 	}
 }
 
