Fixed Sequencer folders labels becoming grey when they had items inside of themselves.

diff --git a/Engine/Source/Editor/Sequencer/Private/DisplayNodes/SequencerFolderNode.cpp b/Engine/Source/Editor/Sequencer/Private/DisplayNodes/SequencerFolderNode.cpp
--- a/Engine/Source/Editor/Sequencer/Private/DisplayNodes/SequencerFolderNode.cpp
+++ b/Engine/Source/Editor/Sequencer/Private/DisplayNodes/SequencerFolderNode.cpp
@@ -115,6 +115,12 @@ FSlateColor FSequencerFolderNode::GetIconColor() const
 	return FSlateColor(MovieSceneFolder.GetFolderColor());
 }
 
+FLinearColor FSequencerFolderNode::GetDisplayNameColor() const
+{
+	return FLinearColor::White;
+}
+
+
 bool FSequencerFolderNode::CanDrag() const
 {
 	return true;
diff --git a/Engine/Source/Editor/Sequencer/Private/DisplayNodes/SequencerFolderNode.h b/Engine/Source/Editor/Sequencer/Private/DisplayNodes/SequencerFolderNode.h
--- a/Engine/Source/Editor/Sequencer/Private/DisplayNodes/SequencerFolderNode.h
+++ b/Engine/Source/Editor/Sequencer/Private/DisplayNodes/SequencerFolderNode.h
@@ -27,6 +27,7 @@ public:
 	virtual void SetDisplayName( const FText& NewDisplayName ) override;
 	virtual const FSlateBrush* GetIconBrush() const override;
 	virtual FSlateColor GetIconColor() const override;
+	virtual FLinearColor GetDisplayNameColor() const override;
 	virtual bool CanDrag() const override;
 	virtual TOptional<EItemDropZone> CanDrop( FSequencerDisplayNodeDragDropOp& DragDropOp, EItemDropZone ItemDropZone ) const override;
 	virtual void Drop( const TArray<TSharedRef<FSequencerDisplayNode>>& DraggedNodes, EItemDropZone ItemDropZone ) override;
