Sequencer: Fixed deleted categories and channels not being removed from their parent track nodes correctly

diff --git a/Engine/Source/Editor/Sequencer/Private/DisplayNodes/SequencerTrackNode.cpp b/Engine/Source/Editor/Sequencer/Private/DisplayNodes/SequencerTrackNode.cpp
--- a/Engine/Source/Editor/Sequencer/Private/DisplayNodes/SequencerTrackNode.cpp
+++ b/Engine/Source/Editor/Sequencer/Private/DisplayNodes/SequencerTrackNode.cpp
@@ -236,6 +236,8 @@ void FSequencerTrackNode::UpdateSections()
 	{
 		Sections.RemoveAt(Sections.Num()-NumToRemove, NumToRemove, true);
 	}
+
+	RemoveStaleChildren();
 }
 
 void FSequencerTrackNode::ClearChildren()
@@ -249,6 +251,33 @@ void FSequencerTrackNode::ClearChildren()
 	}
 }
 
+void FSequencerTrackNode::RemoveStaleChildren()
+{
+	// Gather stale nodes into a separate array
+	TArray<TSharedRef<FSequencerDisplayNode>> StaleNodes;
+
+	TArray<TSharedRef<FSequencerDisplayNode>> NodesToCheck = ChildNodes;
+	for (int32 Index = 0; Index < NodesToCheck.Num(); ++Index)
+	{
+		TSharedRef<FSequencerDisplayNode> Child = NodesToCheck[Index];
+		if (Child->TreeSerialNumber != TreeSerialNumber)
+		{
+			// This node is stale - remove it
+			StaleNodes.Add(Child);
+		}
+		else
+		{
+			// This node is still relevant, but its children may not be - recurse into those
+			NodesToCheck.Append(Child->GetChildNodes());
+		}
+	}
+
+	for (TSharedRef<FSequencerDisplayNode> StaleNode : StaleNodes)
+	{
+		StaleNode->SetParent(nullptr);
+	}
+}
+
 FSequencerTrackNode::ESubTrackMode FSequencerTrackNode::GetSubTrackMode() const
 {
 	return SubTrackMode;
diff --git a/Engine/Source/Editor/Sequencer/Private/DisplayNodes/SequencerTrackNode.h b/Engine/Source/Editor/Sequencer/Private/DisplayNodes/SequencerTrackNode.h
--- a/Engine/Source/Editor/Sequencer/Private/DisplayNodes/SequencerTrackNode.h
+++ b/Engine/Source/Editor/Sequencer/Private/DisplayNodes/SequencerTrackNode.h
@@ -153,6 +153,8 @@ private:
 
 	void ClearChildren();
 
+	void RemoveStaleChildren();
+
 private:
 
 	/** The track editor for the track associated with this node. */
