Fixed Event Tracks belonging to both an Actor Binding and a Folder when the event track was added to an actor binding inside of a folder. This solves an issue where deleting an Actor Binding would cause the Event Track to appear to move/duplicate itself to the folder that contained that binding. Also cleans up post-track add behavior to make new selection/scene rebuilding more consistent.

diff --git a/Engine/Plugins/Animation/LiveLink/Source/LiveLinkSequencer/Private/Sequencer/LiveLinkPropertyTrackEditor.cpp b/Engine/Plugins/Animation/LiveLink/Source/LiveLinkSequencer/Private/Sequencer/LiveLinkPropertyTrackEditor.cpp
--- a/Engine/Plugins/Animation/LiveLink/Source/LiveLinkSequencer/Private/Sequencer/LiveLinkPropertyTrackEditor.cpp
+++ b/Engine/Plugins/Animation/LiveLink/Source/LiveLinkSequencer/Private/Sequencer/LiveLinkPropertyTrackEditor.cpp
@@ -250,11 +250,8 @@ void FLiveLinkPropertyTrackEditor::HandleAddLiveLinkTrackMenuEntryExecute()
 
 		if (GetSequencer().IsValid())
 		{
-			GetSequencer()->OnAddTrack(NewTrack);
+			GetSequencer()->OnAddTrack(NewTrack, FGuid());
 		}
-
-
-		GetSequencer()->NotifyMovieSceneDataChanged(EMovieSceneDataChangeType::MovieSceneStructureItemAdded);
 	}
 }
 
diff --git a/Engine/Plugins/Media/MediaCompositing/Source/MediaCompositingEditor/Private/Sequencer/MediaTrackEditor.cpp b/Engine/Plugins/Media/MediaCompositing/Source/MediaCompositingEditor/Private/Sequencer/MediaTrackEditor.cpp
--- a/Engine/Plugins/Media/MediaCompositing/Source/MediaCompositingEditor/Private/Sequencer/MediaTrackEditor.cpp
+++ b/Engine/Plugins/Media/MediaCompositing/Source/MediaCompositingEditor/Private/Sequencer/MediaTrackEditor.cpp
@@ -304,10 +304,8 @@ void FMediaTrackEditor::HandleAddMediaTrackMenuEntryExecute()
 
 	if (GetSequencer().IsValid())
 	{
-		GetSequencer()->OnAddTrack(NewTrack);
+		GetSequencer()->OnAddTrack(NewTrack, FGuid());
 	}
-
-	GetSequencer()->NotifyMovieSceneDataChanged(EMovieSceneDataChangeType::MovieSceneStructureItemAdded);
 }
 
 
diff --git a/Engine/Source/Editor/MovieSceneTools/Private/TrackEditors/AudioTrackEditor.cpp b/Engine/Source/Editor/MovieSceneTools/Private/TrackEditors/AudioTrackEditor.cpp
--- a/Engine/Source/Editor/MovieSceneTools/Private/TrackEditors/AudioTrackEditor.cpp
+++ b/Engine/Source/Editor/MovieSceneTools/Private/TrackEditors/AudioTrackEditor.cpp
@@ -1042,16 +1042,11 @@ FKeyPropertyResult FAudioTrackEditor::AddNewMasterSound( FFrameNumber KeyTime, U
 
 		if (GetSequencer().IsValid())
 		{
-			GetSequencer()->OnAddTrack(AudioTrack);
+			GetSequencer()->OnAddTrack(AudioTrack, FGuid());
 		}
 	}
 
 	KeyPropertyResult.bTrackModified = true;
-
-	GetSequencer()->EmptySelection();
-	GetSequencer()->SelectSection(NewSection);
-	GetSequencer()->ThrobSectionSelection();
-
 	return KeyPropertyResult;
 }
 
@@ -1120,10 +1115,8 @@ void FAudioTrackEditor::HandleAddAudioTrackMenuEntryExecute()
 
 	if (GetSequencer().IsValid())
 	{
-		GetSequencer()->OnAddTrack(NewTrack);
+		GetSequencer()->OnAddTrack(NewTrack, FGuid());
 	}
-
-	GetSequencer()->NotifyMovieSceneDataChanged( EMovieSceneDataChangeType::MovieSceneStructureItemAdded );
 }
 
 TSharedRef<SWidget> FAudioTrackEditor::BuildAudioSubMenu(UMovieSceneTrack* Track)
diff --git a/Engine/Source/Editor/MovieSceneTools/Private/TrackEditors/CameraCutTrackEditor.cpp b/Engine/Source/Editor/MovieSceneTools/Private/TrackEditors/CameraCutTrackEditor.cpp
--- a/Engine/Source/Editor/MovieSceneTools/Private/TrackEditors/CameraCutTrackEditor.cpp
+++ b/Engine/Source/Editor/MovieSceneTools/Private/TrackEditors/CameraCutTrackEditor.cpp
@@ -336,9 +336,8 @@ void FCameraCutTrackEditor::HandleAddCameraCutTrackMenuEntryExecute()
 	{
 		if (GetSequencer().IsValid())
 		{
-			GetSequencer()->OnAddTrack(CameraCutTrack);
+			GetSequencer()->OnAddTrack(CameraCutTrack, FGuid());
 		}
-		GetSequencer()->NotifyMovieSceneDataChanged(EMovieSceneDataChangeType::MovieSceneStructureItemAdded);
 	}
 }
 
diff --git a/Engine/Source/Editor/MovieSceneTools/Private/TrackEditors/CinematicShotTrackEditor.cpp b/Engine/Source/Editor/MovieSceneTools/Private/TrackEditors/CinematicShotTrackEditor.cpp
--- a/Engine/Source/Editor/MovieSceneTools/Private/TrackEditors/CinematicShotTrackEditor.cpp
+++ b/Engine/Source/Editor/MovieSceneTools/Private/TrackEditors/CinematicShotTrackEditor.cpp
@@ -578,9 +578,9 @@ void FCinematicShotTrackEditor::HandleAddCinematicShotTrackMenuEntryExecute()
 	{
 		if (GetSequencer().IsValid())
 		{
-			GetSequencer()->OnAddTrack(ShotTrack);
+			// Cinematic Shot Tracks can't be placed in folders, they're only allowed in the root.
+			GetSequencer()->OnAddTrack(ShotTrack, FGuid());
 		}
-		GetSequencer()->NotifyMovieSceneDataChanged(EMovieSceneDataChangeType::MovieSceneStructureItemAdded);
 	}
 }
 
diff --git a/Engine/Source/Editor/MovieSceneTools/Private/TrackEditors/EventTrackEditor.cpp b/Engine/Source/Editor/MovieSceneTools/Private/TrackEditors/EventTrackEditor.cpp
--- a/Engine/Source/Editor/MovieSceneTools/Private/TrackEditors/EventTrackEditor.cpp
+++ b/Engine/Source/Editor/MovieSceneTools/Private/TrackEditors/EventTrackEditor.cpp
@@ -269,13 +269,24 @@ void FEventTrackEditor::HandleAddEventTrackMenuEntryExecute(TArray<FGuid> InObje
 	{
 		if (InObjectBindingID.IsValid())
 		{
-			NewTracks.Add(FocusedMovieScene->AddTrack<UMovieSceneEventTrack>(InObjectBindingID));
+			UMovieSceneEventTrack* NewObjectTrack = FocusedMovieScene->AddTrack<UMovieSceneEventTrack>(InObjectBindingID);
+			NewTracks.Add(NewObjectTrack);
+		
+			if (GetSequencer().IsValid())
+			{
+				GetSequencer()->OnAddTrack(NewObjectTrack, InObjectBindingID);
+			}
 		}
 	}
 
 	if (!NewTracks.Num())
 	{
-		NewTracks.Add(FocusedMovieScene->AddMasterTrack<UMovieSceneEventTrack>());
+		UMovieSceneEventTrack* NewMasterTrack = FocusedMovieScene->AddMasterTrack<UMovieSceneEventTrack>();
+		NewTracks.Add(NewMasterTrack);
+		if (GetSequencer().IsValid())
+		{
+			GetSequencer()->OnAddTrack(NewMasterTrack, FGuid());
+		}
 	}
 
 	check(NewTracks.Num() != 0);
@@ -286,13 +297,8 @@ void FEventTrackEditor::HandleAddEventTrackMenuEntryExecute(TArray<FGuid> InObje
 
 		NewTrack->SetDisplayName(LOCTEXT("TrackName", "Events"));
 
-		if (GetSequencer().IsValid())
-		{
-			GetSequencer()->OnAddTrack(NewTrack);
-		}
+		
 	}
-
-	GetSequencer()->NotifyMovieSceneDataChanged( EMovieSceneDataChangeType::MovieSceneStructureItemAdded );
 }
 
 void FEventTrackEditor::CreateNewSection(UMovieSceneTrack* Track, int32 RowIndex, UClass* SectionType, bool bSelect)
diff --git a/Engine/Source/Editor/MovieSceneTools/Private/TrackEditors/FadeTrackEditor.cpp b/Engine/Source/Editor/MovieSceneTools/Private/TrackEditors/FadeTrackEditor.cpp
--- a/Engine/Source/Editor/MovieSceneTools/Private/TrackEditors/FadeTrackEditor.cpp
+++ b/Engine/Source/Editor/MovieSceneTools/Private/TrackEditors/FadeTrackEditor.cpp
@@ -162,9 +162,8 @@ void FFadeTrackEditor::HandleAddFadeTrackMenuEntryExecute()
 	FadeTrack->AddSection(*NewSection);
 	if (GetSequencer().IsValid())
 	{
-		GetSequencer()->OnAddTrack(FadeTrack);
+		GetSequencer()->OnAddTrack(FadeTrack, FGuid());
 	}
-	GetSequencer()->NotifyMovieSceneDataChanged( EMovieSceneDataChangeType::MovieSceneStructureItemAdded );
 }
 
 bool FFadeTrackEditor::HandleAddFadeTrackMenuEntryCanExecute() const
diff --git a/Engine/Source/Editor/MovieSceneTools/Private/TrackEditors/LevelVisibilityTrackEditor.cpp b/Engine/Source/Editor/MovieSceneTools/Private/TrackEditors/LevelVisibilityTrackEditor.cpp
--- a/Engine/Source/Editor/MovieSceneTools/Private/TrackEditors/LevelVisibilityTrackEditor.cpp
+++ b/Engine/Source/Editor/MovieSceneTools/Private/TrackEditors/LevelVisibilityTrackEditor.cpp
@@ -111,12 +111,8 @@ void FLevelVisibilityTrackEditor::OnAddTrack()
 	UMovieSceneLevelVisibilitySection* NewSection = AddNewSection( FocusedMovieScene, NewTrack, ELevelVisibility::Visible );
 	if (GetSequencer().IsValid())
 	{
-		GetSequencer()->OnAddTrack(NewTrack);
+		GetSequencer()->OnAddTrack(NewTrack, FGuid());
 	}
-	GetSequencer()->NotifyMovieSceneDataChanged( EMovieSceneDataChangeType::MovieSceneStructureItemAdded );
-	GetSequencer()->EmptySelection();
-	GetSequencer()->SelectSection(NewSection);
-	GetSequencer()->ThrobSectionSelection();
 }
 
 
diff --git a/Engine/Source/Editor/MovieSceneTools/Private/TrackEditors/MaterialParameterCollectionTrackEditor.cpp b/Engine/Source/Editor/MovieSceneTools/Private/TrackEditors/MaterialParameterCollectionTrackEditor.cpp
--- a/Engine/Source/Editor/MovieSceneTools/Private/TrackEditors/MaterialParameterCollectionTrackEditor.cpp
+++ b/Engine/Source/Editor/MovieSceneTools/Private/TrackEditors/MaterialParameterCollectionTrackEditor.cpp
@@ -156,9 +156,8 @@ void FMaterialParameterCollectionTrackEditor::AddTrackToSequence(const FAssetDat
 
 	if (GetSequencer().IsValid())
 	{
-		GetSequencer()->OnAddTrack(Track);
+		GetSequencer()->OnAddTrack(Track, FGuid());
 	}
-	GetSequencer()->NotifyMovieSceneDataChanged(EMovieSceneDataChangeType::MovieSceneStructureItemAdded);
 }
 
 void FMaterialParameterCollectionTrackEditor::AddTrackToSequenceEnterPressed(const TArray<FAssetData>& InAssetData)
diff --git a/Engine/Source/Editor/MovieSceneTools/Private/TrackEditors/SlomoTrackEditor.cpp b/Engine/Source/Editor/MovieSceneTools/Private/TrackEditors/SlomoTrackEditor.cpp
--- a/Engine/Source/Editor/MovieSceneTools/Private/TrackEditors/SlomoTrackEditor.cpp
+++ b/Engine/Source/Editor/MovieSceneTools/Private/TrackEditors/SlomoTrackEditor.cpp
@@ -95,10 +95,8 @@ void FSlomoTrackEditor::HandleAddSlomoTrackMenuEntryExecute()
 
 	if (GetSequencer().IsValid())
 	{
-		GetSequencer()->OnAddTrack(SlomoTrack);
+		GetSequencer()->OnAddTrack(SlomoTrack, FGuid());
 	}
-
-	GetSequencer()->NotifyMovieSceneDataChanged( EMovieSceneDataChangeType::MovieSceneStructureItemAdded );
 }
 
 bool FSlomoTrackEditor::HandleAddSlomoTrackMenuEntryCanExecute() const
diff --git a/Engine/Source/Editor/MovieSceneTools/Private/TrackEditors/SubTrackEditor.cpp b/Engine/Source/Editor/MovieSceneTools/Private/TrackEditors/SubTrackEditor.cpp
--- a/Engine/Source/Editor/MovieSceneTools/Private/TrackEditors/SubTrackEditor.cpp
+++ b/Engine/Source/Editor/MovieSceneTools/Private/TrackEditors/SubTrackEditor.cpp
@@ -653,9 +653,8 @@ void FSubTrackEditor::HandleAddSubTrackMenuEntryExecute()
 
 	if (GetSequencer().IsValid())
 	{
-		GetSequencer()->OnAddTrack(NewTrack);
+		GetSequencer()->OnAddTrack(NewTrack, FGuid());
 	}
-	GetSequencer()->NotifyMovieSceneDataChanged( EMovieSceneDataChangeType::MovieSceneStructureItemAdded );
 }
 
 /** Helper function - get the first PIE world (or first PIE client world if there is more than one) */
diff --git a/Engine/Source/Editor/Sequencer/Private/Sequencer.cpp b/Engine/Source/Editor/Sequencer/Private/Sequencer.cpp
--- a/Engine/Source/Editor/Sequencer/Private/Sequencer.cpp
+++ b/Engine/Source/Editor/Sequencer/Private/Sequencer.cpp
@@ -8083,13 +8083,18 @@ void FSequencer::OnAddFolder()
 	NotifyMovieSceneDataChanged( EMovieSceneDataChangeType::MovieSceneStructureItemAdded );
 }
 
-void FSequencer::OnAddTrack(const TWeakObjectPtr<UMovieSceneTrack>& InTrack)
+void FSequencer::OnAddTrack(const TWeakObjectPtr<UMovieSceneTrack>& InTrack, const FGuid& ObjectBinding)
 {
+	if (!ensureAlwaysMsgf(InTrack.IsValid(), TEXT("Attempted to add a null UMovieSceneTrack to Sequencer. This should never happen.")))
+	{
+		return;
+	}
+
 	FString NewNodePath;
 
-	// Cinematic Shot Tracks and Camera Cut Tracks are always in the root and ignore sorting/folders, so we don't give them a chance to be placed into a folder.
-	bool bIsValidTrack = !(InTrack->IsA(UMovieSceneCinematicShotTrack::StaticClass()) || InTrack->IsA(UMovieSceneCameraCutTrack::StaticClass()));
-	if(bIsValidTrack)
+	// If they specified an object binding it's being added to, we don't add it to a folder since we can't have it existing
+	// as a children of two places at once.
+	if(!GetFocusedMovieSceneSequence()->GetMovieScene()->FindBinding(ObjectBinding))
 	{
 		TArray<UMovieSceneFolder*> SelectedParentFolders;
 		CalculateSelectedFolderAndPath(SelectedParentFolders, NewNodePath);
@@ -8100,13 +8105,20 @@ void FSequencer::OnAddTrack(const TWeakObjectPtr<UMovieSceneTrack>& InTrack)
 			SelectedParentFolders[0]->AddChildMasterTrack(InTrack.Get());
 		}
 	}
-	Selection.Empty();
 
 	// We can't add the newly created folder to the selection set as the nodes for it don't actually exist yet.
 	// However, we can calculate the resulting path that the node will end up at and add that to the selection
 	// set, which will cause the newly created node to be selected when the selection is restored post-refresh.
 	NewNodePath += InTrack->GetFName().ToString();
 	SequencerWidget->AddAdditionalPathToSelectionSet(NewNodePath);
+
+	NotifyMovieSceneDataChanged(EMovieSceneDataChangeType::MovieSceneStructureItemAdded);
+	EmptySelection();
+	if (InTrack->GetAllSections().Num() > 0)
+	{
+		SelectSection(InTrack->GetAllSections()[0]);
+	}
+	ThrobSectionSelection();
 }
 
 
diff --git a/Engine/Source/Editor/Sequencer/Private/Sequencer.h b/Engine/Source/Editor/Sequencer/Private/Sequencer.h
--- a/Engine/Source/Editor/Sequencer/Private/Sequencer.h
+++ b/Engine/Source/Editor/Sequencer/Private/Sequencer.h
@@ -974,8 +974,8 @@ protected:
 	/** Handles loading in previously recorded data. */
 	void OnLoadRecordedData();
 
-	/** Handles adding a newly created track to the outliner tree by assigning it into a folder and selecting it. */
-	void OnAddTrack(const TWeakObjectPtr<UMovieSceneTrack>& InTrack);
+	/** Adds the track to the selected folder (if FGuid is invalid) and selects the track, throbs it, and notifies the sequence to rebuild any necessary data. */
+	void OnAddTrack(const TWeakObjectPtr<UMovieSceneTrack>& InTrack, const FGuid& ObjectBinding);
 
 	/** Determines the selected parent folders and returns the node path to the first folder. Also expands the first folder. */
 	void CalculateSelectedFolderAndPath(TArray<UMovieSceneFolder*>& OutSelectedParentFolders, FString& OutNewNodePath);
diff --git a/Engine/Source/Editor/Sequencer/Public/ISequencer.h b/Engine/Source/Editor/Sequencer/Public/ISequencer.h
--- a/Engine/Source/Editor/Sequencer/Public/ISequencer.h
+++ b/Engine/Source/Editor/Sequencer/Public/ISequencer.h
@@ -239,10 +239,11 @@ public:
 	virtual TArray<FGuid> AddActors(const TArray<TWeakObjectPtr<AActor> >& InActors, bool bSelectActors = true) = 0;
 
 	/**
-	 * Calling this function will add the specified track to the currently selected folder
-	 * if there is one, and will set the newly created track as the current selection.
+	 * Should be called after adding a track to the MovieScene. This will set the specified track as your current selection
+	 * cause it to throb, notify the sequence to rebuild any data required. The track will be added to the selected folder
+	 * unless ObjectBinding points to a valid FGuid.
 	 */
-	virtual void OnAddTrack(const TWeakObjectPtr<UMovieSceneTrack>& InTrack) = 0;
+	virtual void OnAddTrack(const TWeakObjectPtr<UMovieSceneTrack>& InTrack, const FGuid& ObjectBinding) = 0;
 
 	/**
 	 * Adds a movie scene as a section inside the current movie scene
