Sequencer: Tidied up marquee selection logic

  - This has actually never worked well from what I can tell from the code history, and was actually highlighted more with the removal of spacer nodes, since marqueeing off the end of the last node will no longer hit a spacer node and select all the collapsed nodes implicitly.
  - Removed logic that was dealing with non-subtrack tracks that had multiple rows since this can never happen any more.
  - Any collapsed key areas are now always visited along with their parent if it does not have keys of its own. Previously they would only match if they intersected the virtual marquee space

diff --git a/Engine/Source/Editor/Sequencer/Private/Tools/SequencerEntityVisitor.cpp b/Engine/Source/Editor/Sequencer/Private/Tools/SequencerEntityVisitor.cpp
--- a/Engine/Source/Editor/Sequencer/Private/Tools/SequencerEntityVisitor.cpp
+++ b/Engine/Source/Editor/Sequencer/Private/Tools/SequencerEntityVisitor.cpp
@@ -19,28 +19,10 @@ FSequencerEntityRange::FSequencerEntityRange(FVector2D TopLeft, FVector2D Bottom
 {
 }
 
-bool FSequencerEntityRange::IntersectSection(const UMovieSceneSection* InSection, const TSharedRef<FSequencerTrackNode>& InTrackNode, int32 MaxRowIndex) const
+bool FSequencerEntityRange::IntersectSection(const UMovieSceneSection* InSection) const
 {
 	// Test horizontal bounds
-	if (!(InSection->GetRange() / TickResolution).Overlaps(Range))
-	{
-		return false;
-	}
-	// Test vertical bounds
-	else if (MaxRowIndex > 0 && VerticalTop.IsSet())
-	{
-		const float NodeTop = InTrackNode->GetVirtualTop();
-		const float NodeBottom = InTrackNode->GetVirtualBottom();
-
-		const float RowHeight = (NodeBottom - NodeTop) / (MaxRowIndex + 1);
-		const float RowTop = NodeTop + InSection->GetRowIndex() * RowHeight;
-
-		return RowTop <= VerticalBottom.GetValue() && RowTop + RowHeight >= VerticalTop.GetValue();
-	}
-	else
-	{
-		return true;
-	}
+	return (InSection->GetRange() / TickResolution).Overlaps(Range);
 }
 
 bool FSequencerEntityRange::IntersectNode(TSharedRef<FSequencerDisplayNode> InNode) const
@@ -69,164 +51,94 @@ FSequencerEntityWalker::FSequencerEntityWalker(const FSequencerEntityRange& InRa
 /* @todo: Could probably optimize this by not walking every single node, and binary searching the begin<->end ranges instead */
 void FSequencerEntityWalker::Traverse(const ISequencerEntityVisitor& Visitor, const TArray< TSharedRef<FSequencerDisplayNode> >& Nodes)
 {
-	for (auto& Child : Nodes)
+	for (const TSharedRef<FSequencerDisplayNode>& Child : Nodes)
 	{
 		if (!Child->IsHidden())
 		{
-			HandleNode(Visitor, Child);
+			ConditionallyIntersectNode(Visitor, Child);
 		}
 	}
 }
 
-void FSequencerEntityWalker::HandleNode(const ISequencerEntityVisitor& Visitor, const TSharedRef<FSequencerDisplayNode>& InNode)
+void FSequencerEntityWalker::ConditionallyIntersectNode(const ISequencerEntityVisitor& Visitor, const TSharedRef<FSequencerDisplayNode>& InNode)
 {
-	if (InNode->GetType() == ESequencerNode::Track)
-	{
-		HandleTrackNode(Visitor, StaticCastSharedRef<FSequencerTrackNode>(InNode));
-	}
-
-	if (InNode->IsExpanded())
+	if (Range.IntersectNode(InNode))
 	{
-		for (auto& Child : InNode->GetChildNodes())
+		// Visit sections within this track
+		if (InNode->GetType() == ESequencerNode::Track && Visitor.CheckEntityMask(ESequencerEntity::Section))
 		{
-			if (!Child->IsHidden())
+			TSharedRef<FSequencerTrackNode> TrackNode = StaticCastSharedRef<FSequencerTrackNode>(InNode);
+
+			for (TSharedRef<ISequencerSection> SectionInterface : TrackNode->GetSections())
 			{
-				HandleNode(Visitor, Child);
+				UMovieSceneSection* Section = SectionInterface->GetSectionObject();
+				if (Range.IntersectSection(Section))
+				{
+					Visitor.VisitSection(Section, InNode);
+				}
 			}
 		}
-	}
-}
-
-void FSequencerEntityWalker::HandleTrackNode(const ISequencerEntityVisitor& Visitor, const TSharedRef<FSequencerTrackNode>& InTrackNode)
-{
-	TArray<TSharedRef<ISequencerSection>> Sections = InTrackNode->GetSections();
-
-	if (Range.IntersectNode(InTrackNode))
-	{
-
-		int32 MaxRowIndex;
-		if (InTrackNode->GetSubTrackMode() == FSequencerTrackNode::ESubTrackMode::None)
-		{
-			MaxRowIndex = InTrackNode->GetTrack()->GetMaxRowIndex();
-		}
-		else
-		{
-			// When using sub-tracks each section index gets it's own track so the effective max index
-			// within the track will always be 0.
-			MaxRowIndex = 0;
-		}
 
-		// Prune the selections to anything that is in the range, visiting if necessary
-		for (int32 SectionIndex = 0; SectionIndex < Sections.Num();)
+		if (Range.IntersectKeyArea(InNode, VirtualKeySize.Y))
 		{
-			UMovieSceneSection* Section = Sections[SectionIndex]->GetSectionObject();
-			if (!Range.IntersectSection(Section, InTrackNode, MaxRowIndex))
-			{
-				Sections.RemoveAtSwap(SectionIndex, 1, false);
-				continue;
-			}
-
-			if (Visitor.CheckEntityMask(ESequencerEntity::Section))
-			{
-				Visitor.VisitSection(Section, InTrackNode);
-			}
-
-			++SectionIndex;
+			VisitKeyAnyAreas(Visitor, InNode);
 		}
-
-		HandleSingleNode(Visitor, InTrackNode, Sections);
 	}
 
-	if (InTrackNode->IsExpanded())
+	// Iterate into expanded nodes
+	if (InNode->IsExpanded())
 	{
-		// Handle Children
-		for (auto& Child : InTrackNode->GetChildNodes())
+		for (const TSharedRef<FSequencerDisplayNode>& Child : InNode->GetChildNodes())
 		{
+			// Do not visit nodes that are currently filtered out
 			if (!Child->IsHidden())
 			{
-				HandleChildNode(Visitor, Child, Sections);
+				ConditionallyIntersectNode(Visitor, Child);
 			}
 		}
 	}
 }
 
-void FSequencerEntityWalker::HandleChildNode(const ISequencerEntityVisitor& Visitor, const TSharedRef<FSequencerDisplayNode>& InNode, const TArray<TSharedRef<ISequencerSection>>& InSections)
+void FSequencerEntityWalker::VisitKeyAnyAreas(const ISequencerEntityVisitor& Visitor, const TSharedRef<FSequencerDisplayNode>& InNode)
 {
-	if (Range.IntersectNode(InNode))
+	if (!Visitor.CheckEntityMask(ESequencerEntity::Key))
 	{
-		HandleSingleNode(Visitor, InNode, InSections);
+		return;
 	}
 
-	if (InNode->IsExpanded())
-	{
-		// Handle Children
-		for (auto& Child : InNode->GetChildNodes())
-		{
-			if (!Child->IsHidden())
-			{
-				HandleChildNode(Visitor, Child, InSections);
-			}
-		}
-	}
-}
-
-void FSequencerEntityWalker::HandleSingleNode(const ISequencerEntityVisitor& Visitor, const TSharedRef<FSequencerDisplayNode>& InNode, const TArray<TSharedRef<ISequencerSection>>& InSections)
-{
-	bool bNodeHasKeyArea = false;
+	TSharedPtr<FSequencerSectionKeyAreaNode> KeyAreaNode;
 	if (InNode->GetType() == ESequencerNode::KeyArea)
 	{
-		HandleKeyAreaNode(Visitor, StaticCastSharedRef<FSequencerSectionKeyAreaNode>(InNode), InNode, InSections);
-		bNodeHasKeyArea = true;
+		KeyAreaNode = StaticCastSharedRef<FSequencerSectionKeyAreaNode>(InNode);
 	}
 	else if (InNode->GetType() == ESequencerNode::Track)
 	{
-		TSharedPtr<FSequencerSectionKeyAreaNode> SectionKeyNode = StaticCastSharedRef<FSequencerTrackNode>(InNode)->GetTopLevelKeyNode();
-		if (SectionKeyNode.IsValid())
-		{
-			HandleKeyAreaNode(Visitor, SectionKeyNode.ToSharedRef(), InNode, InSections);
-			bNodeHasKeyArea = true;
-		}
+		KeyAreaNode = StaticCastSharedRef<FSequencerTrackNode>(InNode)->GetTopLevelKeyNode();
 	}
 
-	// As a fallback, we need to handle:
-	//  - Key groupings on collapsed parents
-	//  - Sections that have no key areas
-	const bool bIterateKeyGroupings = Visitor.CheckEntityMask(ESequencerEntity::Key) &&
-		!bNodeHasKeyArea &&
-		(!InNode->IsExpanded() || InNode->GetChildNodes().Num() != 0);
-
-	if (bIterateKeyGroupings)
+	// If this node has or is a key area, visit all the keys on the track
+	if (KeyAreaNode)
 	{
-		if (InNode->GetChildNodes().Num() != 0 && Range.IntersectKeyArea(InNode, VirtualKeySize.X))
+		for (TSharedRef<IKeyArea> KeyArea : KeyAreaNode->GetAllKeyAreas())
 		{
-			for (TSharedRef<FSequencerDisplayNode> ChildNode : InNode->GetChildNodes())
+			UMovieSceneSection* Section = KeyArea->GetOwningSection();
+			if (Section)
 			{
-				HandleSingleNode(Visitor, ChildNode, InSections);
+				VisitKeyArea(Visitor, KeyArea, Section, InNode);
 			}
 		}
 	}
-}
-
-void FSequencerEntityWalker::HandleKeyAreaNode(const ISequencerEntityVisitor& Visitor, const TSharedRef<FSequencerSectionKeyAreaNode>& InKeyAreaNode, const TSharedRef<FSequencerDisplayNode>& InOwnerNode, const TArray<TSharedRef<ISequencerSection>>& InSections)
-{
-	for (TSharedRef<ISequencerSection> SectionInterface : InSections)
+	// Otherwise it might be a collapsed node that contains key areas as children. If so we visit them as if they were a part of this track so that key groupings are visited properly.
+	else if (!InNode->IsExpanded())
 	{
-		UMovieSceneSection* Section = SectionInterface->GetSectionObject();
-		if (Visitor.CheckEntityMask(ESequencerEntity::Key))
+		for (TSharedRef<FSequencerDisplayNode> ChildNode : InNode->GetChildNodes())
 		{
-			if (Range.IntersectKeyArea(InOwnerNode, VirtualKeySize.X))
-			{
-				TSharedPtr<IKeyArea> KeyArea = InKeyAreaNode->GetKeyArea(Section);
-				if (KeyArea.IsValid())
-				{
-					HandleKeyArea(Visitor, KeyArea.ToSharedRef(), Section, InOwnerNode);
-				}
-			}
+			VisitKeyAnyAreas(Visitor, ChildNode);
 		}
 	}
 }
 
-void FSequencerEntityWalker::HandleKeyArea(const ISequencerEntityVisitor& Visitor, const TSharedRef<IKeyArea>& KeyArea, UMovieSceneSection* Section, const TSharedRef<FSequencerDisplayNode>& InNode)
+void FSequencerEntityWalker::VisitKeyArea(const ISequencerEntityVisitor& Visitor, const TSharedRef<IKeyArea>& KeyArea, UMovieSceneSection* Section, const TSharedRef<FSequencerDisplayNode>& InNode)
 {
 	TArray<FKeyHandle> Handles;
 	TArray<FFrameNumber> Times;
diff --git a/Engine/Source/Editor/Sequencer/Private/Tools/SequencerEntityVisitor.h b/Engine/Source/Editor/Sequencer/Private/Tools/SequencerEntityVisitor.h
--- a/Engine/Source/Editor/Sequencer/Private/Tools/SequencerEntityVisitor.h
+++ b/Engine/Source/Editor/Sequencer/Private/Tools/SequencerEntityVisitor.h
@@ -47,10 +47,10 @@ struct FSequencerEntityRange
 	FSequencerEntityRange(const TRange<double>& InRange, FFrameRate InTickResolution);
 	FSequencerEntityRange(FVector2D TopLeft, FVector2D BottomRight, FFrameRate InTickResolution);
 
-	/** Check whether the specified section intersects this range */
-	bool IntersectSection(const UMovieSceneSection* InSection, const TSharedRef<FSequencerTrackNode>& InTrackNode, int32 MaxRowIndex) const;
+	/** Check whether the specified section intersects the horizontal range */
+	bool IntersectSection(const UMovieSceneSection* InSection) const;
 
-	/** Check whether the specified node intersects this range */
+	/** Check whether the specified node intersects the vertical range */
 	bool IntersectNode(TSharedRef<FSequencerDisplayNode> InNode) const;
 
 	/** Check whether the specified node's key area intersects this range */
@@ -77,18 +77,14 @@ struct FSequencerEntityWalker
 
 private:
 
-	/** Handle visitation of a particular node */
-	void HandleNode(const ISequencerEntityVisitor& Visitor, const TSharedRef<FSequencerDisplayNode>& InNode);
-	/** Handle visitation of a track node, along with a set of sections */
-	void HandleTrackNode(const ISequencerEntityVisitor& Visitor, const TSharedRef<FSequencerTrackNode>& InNode);
-	/** Handle visitation of an arbitrary child node, with a filtered set of sections */
-	void HandleChildNode(const ISequencerEntityVisitor& Visitor, const TSharedRef<FSequencerDisplayNode>& InNode, const TArray<TSharedRef<ISequencerSection>>& InSections);
-	/** Handle a single node that is known to be within the range */
-	void HandleSingleNode(const ISequencerEntityVisitor& Visitor, const TSharedRef<FSequencerDisplayNode>& InNode, const TArray<TSharedRef<ISequencerSection>>& InSections);
-	/** Handle visitation of a key area node */
-	void HandleKeyAreaNode(const ISequencerEntityVisitor& Visitor, const TSharedRef<FSequencerSectionKeyAreaNode>& InKeyAreaNode, const TSharedRef<FSequencerDisplayNode>& InOwnerNode, const TArray<TSharedRef<ISequencerSection>>& InSections);
-	/** Handle visitation of a key area */
-	void HandleKeyArea(const ISequencerEntityVisitor& Visitor, const TSharedRef<IKeyArea>& KeyArea, UMovieSceneSection* Section, const TSharedRef<FSequencerDisplayNode>& InNode);
+	/** Check whether the specified node intersects the range's vertical space, and visit any keys within it if so */
+	void ConditionallyIntersectNode(const ISequencerEntityVisitor& Visitor, const TSharedRef<FSequencerDisplayNode>& InNode);
+	/** Visit any keys within the specified track node that overlap the range's horizontal space */
+	void VisitTrackNode(const ISequencerEntityVisitor& Visitor, const TSharedRef<FSequencerTrackNode>& InNode);
+	/** Visit any keys within any key area nodes that belong to the specified node that overlap the range's horizontal space */
+	void VisitKeyAnyAreas(const ISequencerEntityVisitor& Visitor, const TSharedRef<FSequencerDisplayNode>& InNode);
+	/** Visit any keys within the specified key area that overlap the range's horizontal space */
+	void VisitKeyArea(const ISequencerEntityVisitor& Visitor, const TSharedRef<IKeyArea>& KeyArea, UMovieSceneSection* Section, const TSharedRef<FSequencerDisplayNode>& InNode);
 
 	/** The bounds of the range */
 	FSequencerEntityRange Range;
