Added the Sequencer Playback controls to the Sequencer Curve Editor. It does this by creating a second copy of the Sequencer Transport Controls widget (which is already bound to Sequencer) and inserting it into the treeview UI. The keyboard shortcuts are also supported via placing the relevant commands into a new Command List and merging that with the Curve Editor instance.

diff --git a/Engine/Source/Editor/CurveEditor/Private/SCurveEditorPanel.cpp b/Engine/Source/Editor/CurveEditor/Private/SCurveEditorPanel.cpp
--- a/Engine/Source/Editor/CurveEditor/Private/SCurveEditorPanel.cpp
+++ b/Engine/Source/Editor/CurveEditor/Private/SCurveEditorPanel.cpp
@@ -218,13 +218,13 @@ void SCurveEditorPanel::Construct(const FArguments& InArgs, TSharedRef<FCurveEdi
 			.Orientation(Orient_Horizontal)
 
 			+ SSplitter::Slot()
-			.Value(0.25f)
+			.Value(0.30f)
 			[
 				InArgs._TreeContent.Widget
 			]
 
 			+ SSplitter::Slot()
-			.Value(0.75f)
+			.Value(0.7f)
 			[
 				MainContent
 			]
diff --git a/Engine/Source/Editor/Sequencer/Private/SSequencer.cpp b/Engine/Source/Editor/Sequencer/Private/SSequencer.cpp
--- a/Engine/Source/Editor/Sequencer/Private/SSequencer.cpp
+++ b/Engine/Source/Editor/Sequencer/Private/SSequencer.cpp
@@ -434,6 +434,13 @@ void SSequencer::Construct(const FArguments& InArgs, TSharedRef<FSequencer> InSe
 				[
 					SNew(SCurveEditorTree, InSequencer->GetCurveEditor())
 				]
+
+				+ SVerticalBox::Slot()
+				.AutoHeight()
+				.HAlign(HAlign_Center)
+				[
+					InSequencer->MakeTransportControls(true)
+				]
 			];
 
 			// Register an instanced custom property type layout to handle converting FFrameNumber from Tick Resolution to Display Rate.
diff --git a/Engine/Source/Editor/Sequencer/Private/Sequencer.cpp b/Engine/Source/Editor/Sequencer/Private/Sequencer.cpp
--- a/Engine/Source/Editor/Sequencer/Private/Sequencer.cpp
+++ b/Engine/Source/Editor/Sequencer/Private/Sequencer.cpp
@@ -455,7 +455,8 @@ void FSequencer::InitSequencer(const FSequencerInitParams& InitParams, const TSh
 
 FSequencer::FSequencer()
 	: SequencerCommandBindings( new FUICommandList )
-	, SequencerSharedBindings( new FUICommandList )
+	, SequencerSharedBindings(new FUICommandList)
+	, CurveEditorSharedBindings(new FUICommandList)
 	, TargetViewRange(0.f, 5.f)
 	, LastViewRange(0.f, 5.f)
 	, ViewRangeBeforeZoom(TRange<double>::Empty())
@@ -10298,6 +10299,23 @@ void FSequencer::BindCommands()
 		FExecuteAction::CreateSP( this, &FSequencer::SetPlaybackRangeToAllShots ),
 		FCanExecuteAction::CreateSP( this, &FSequencer::IsViewingMasterSequence ) );
 
+	// We want a subset of the commands to work in the Curve Editor too, but bound to our functions. This minimizes code duplication
+	// while also freeing us up from issues that result from Sequencer already using two lists (for which our commands might be spread
+	// across both lists which makes a direct copy like it already uses difficult).
+	CurveEditorSharedBindings->MapAction(Commands.TogglePlay,			*SequencerCommandBindings->GetActionForCommand(Commands.TogglePlay));
+	CurveEditorSharedBindings->MapAction(Commands.PlayForward,			*SequencerCommandBindings->GetActionForCommand(Commands.PlayForward));
+	CurveEditorSharedBindings->MapAction(Commands.JumpToStart,			*SequencerCommandBindings->GetActionForCommand(Commands.JumpToStart));
+	CurveEditorSharedBindings->MapAction(Commands.JumpToEnd,			*SequencerCommandBindings->GetActionForCommand(Commands.JumpToEnd));
+	CurveEditorSharedBindings->MapAction(Commands.ShuttleBackward,		*SequencerCommandBindings->GetActionForCommand(Commands.ShuttleBackward));
+	CurveEditorSharedBindings->MapAction(Commands.ShuttleForward,		*SequencerCommandBindings->GetActionForCommand(Commands.ShuttleForward));
+	CurveEditorSharedBindings->MapAction(Commands.Pause,				*SequencerCommandBindings->GetActionForCommand(Commands.Pause));
+	CurveEditorSharedBindings->MapAction(Commands.StepForward,			*SequencerCommandBindings->GetActionForCommand(Commands.StepForward));
+	CurveEditorSharedBindings->MapAction(Commands.StepBackward,			*SequencerCommandBindings->GetActionForCommand(Commands.StepBackward));
+	CurveEditorSharedBindings->MapAction(Commands.StepToNextKey,		*SequencerCommandBindings->GetActionForCommand(Commands.StepToNextKey));
+	CurveEditorSharedBindings->MapAction(Commands.StepToPreviousKey, 	*SequencerCommandBindings->GetActionForCommand(Commands.StepToPreviousKey));
+
+	GetCurveEditor()->GetCommands()->Append(CurveEditorSharedBindings);
+
 	// bind widget specific commands
 	SequencerWidget->BindCommands(SequencerCommandBindings);
 }
diff --git a/Engine/Source/Editor/Sequencer/Private/Sequencer.h b/Engine/Source/Editor/Sequencer/Private/Sequencer.h
--- a/Engine/Source/Editor/Sequencer/Private/Sequencer.h
+++ b/Engine/Source/Editor/Sequencer/Private/Sequencer.h
@@ -1051,6 +1051,9 @@ private:
 	/** Command list for sequencer commands (shared by non-Sequencer). */
 	TSharedRef<FUICommandList> SequencerSharedBindings;
 
+	/** Command list privately shared with the Curve Editor to allow a subset of keybinds to have matching behavior there. */
+	TSharedRef<FUICommandList> CurveEditorSharedBindings;
+
 	/** List of tools we own */
 	TArray<TSharedPtr<ISequencerTrackEditor>> TrackEditors;
 
