Sequencer: Section layout builder now guarantees insertion order based on the order of operation, regardless of the existing state of the layout

diff --git a/Engine/Source/Editor/Sequencer/Private/DisplayNodes/SequencerDisplayNode.cpp b/Engine/Source/Editor/Sequencer/Private/DisplayNodes/SequencerDisplayNode.cpp
--- a/Engine/Source/Editor/Sequencer/Private/DisplayNodes/SequencerDisplayNode.cpp
+++ b/Engine/Source/Editor/Sequencer/Private/DisplayNodes/SequencerDisplayNode.cpp
@@ -383,7 +383,7 @@ void FSequencerDisplayNode::SetParentDirectly(TSharedPtr<FSequencerDisplayNode>
 	ParentNode = InParent;
 }
 
-void FSequencerDisplayNode::SetParent(TSharedPtr<FSequencerDisplayNode> InParent)
+void FSequencerDisplayNode::SetParent(TSharedPtr<FSequencerDisplayNode> InParent, int32 DesiredChildIndex)
 {
 	TSharedPtr<FSequencerDisplayNode> CurrentParent = ParentNode.Pin();
 	if (CurrentParent != InParent)
@@ -398,7 +398,15 @@ void FSequencerDisplayNode::SetParent(TSharedPtr<FSequencerDisplayNode> InParent
 		if (InParent)
 		{
 			// Add to new parent
-			InParent->ChildNodes.Add(ThisNode);
+			if (DesiredChildIndex != INDEX_NONE && ensureMsgf(DesiredChildIndex <= InParent->ChildNodes.Num(), TEXT("Invalid insert index specified")))
+			{
+				InParent->ChildNodes.Insert(ThisNode, DesiredChildIndex);
+			}
+			else
+			{
+				InParent->ChildNodes.Add(ThisNode);
+			}
+
 			bExpanded = ParentTree.GetSavedExpansionState( *this );
 		}
 	}
@@ -406,6 +414,22 @@ void FSequencerDisplayNode::SetParent(TSharedPtr<FSequencerDisplayNode> InParent
 	ParentNode = InParent;
 }
 
+void FSequencerDisplayNode::MoveChild(int32 InChildIndex, int32 InDesiredNewIndex)
+{
+	check(ChildNodes.IsValidIndex(InChildIndex) && InDesiredNewIndex >= 0 && InDesiredNewIndex <= ChildNodes.Num());
+
+	TSharedRef<FSequencerDisplayNode> Child = ChildNodes[InChildIndex];
+	ChildNodes.RemoveAt(InChildIndex, 1, false);
+
+	if (InDesiredNewIndex > InChildIndex)
+	{
+		// Decrement the desired index to account for the removal
+		--InDesiredNewIndex;
+	}
+
+	ChildNodes.Insert(Child, InDesiredNewIndex);
+}
+
 FSequencer& FSequencerDisplayNode::GetSequencer() const
 {
 	return ParentTree.GetSequencer();
diff --git a/Engine/Source/Editor/Sequencer/Private/DisplayNodes/SequencerDisplayNode.h b/Engine/Source/Editor/Sequencer/Private/DisplayNodes/SequencerDisplayNode.h
--- a/Engine/Source/Editor/Sequencer/Private/DisplayNodes/SequencerDisplayNode.h
+++ b/Engine/Source/Editor/Sequencer/Private/DisplayNodes/SequencerDisplayNode.h
@@ -475,12 +475,25 @@ public:
 
 public:
 
-	/** Assigns the parent of this node and adds it the the parent's child node list, removing it from its current parent's children if necessary. */
-	void SetParent(TSharedPtr<FSequencerDisplayNode> InParent);
+	/**
+	 * Assigns the parent of this node and adds it the the parent's child node list, removing it from its current parent's children if necessary.
+	 * 
+	 * @param InParent           This node's new parent node
+	 * @param DesiredChildIndex  (optional) An optional index at which this node should be inserted to its new parent node's children array or INDEX_NONE to add it to the end
+	 */
+	void SetParent(TSharedPtr<FSequencerDisplayNode> InParent, int32 DesiredChildIndex = INDEX_NONE);
 
 	/** Directly assigns the parent of this node without performing any other operation. Should only be used with care when child/parent relationships can be guaranteed. */
 	void SetParentDirectly(TSharedPtr<FSequencerDisplayNode> InParent);
 
+	/**
+	 * Move a child of this node from one index to another - does not re-arrange any other children
+	 * 
+	 * @param InChildIndex       The index of the child node to move from
+	 * @param InDesiredNewIndex  The index to move to - must be a valid index within this node's child array.
+	 */
+	void MoveChild(int32 InChildIndex, int32 InDesiredNewIndex);
+
 	/** Request that this node be reinitialized when the tree next refreshes. Currently Initialization only affects default expansion states. */
 	void RequestReinitialize()
 	{
diff --git a/Engine/Source/Editor/Sequencer/Private/SequencerSectionLayoutBuilder.cpp b/Engine/Source/Editor/Sequencer/Private/SequencerSectionLayoutBuilder.cpp
--- a/Engine/Source/Editor/Sequencer/Private/SequencerSectionLayoutBuilder.cpp
+++ b/Engine/Source/Editor/Sequencer/Private/SequencerSectionLayoutBuilder.cpp
@@ -11,33 +11,69 @@ FSequencerSectionLayoutBuilder::FSequencerSectionLayoutBuilder(TSharedRef<FSeque
 	, CurrentNode(InRootTrackNode)
 	, Section(InSection)
 	, bHasAnyLayout(false)
-{}
+{
+	InsertIndexStack.Add(0);
+}
 
-void FSequencerSectionLayoutBuilder::PushCategory( FName CategoryName, const FText& DisplayLabel )
+TSharedPtr<FSequencerDisplayNode> FindAndRelocateExistingNode(TSharedRef<FSequencerDisplayNode> ParentNode, int32 ExpectedIndex, ESequencerNode::Type NodeType, FName NodeName)
 {
-	bHasAnyLayout = true;
+	auto MatchNode = [NodeType, NodeName](const TSharedRef<FSequencerDisplayNode>& InNode)
+	{
+		return InNode->GetType() == NodeType && InNode->GetNodeName() == NodeName;
+	};
+
+	const TArray<TSharedRef<FSequencerDisplayNode>>& CurrentChildren = ParentNode->GetChildNodes();
 
-	TSharedPtr<FSequencerSectionCategoryNode> CategoryNode;
+	check(ExpectedIndex <= CurrentChildren.Num());
 
-	// Attempt to re-use an existing key area of the same name
-	for (const TSharedRef<FSequencerDisplayNode>& Child : CurrentNode->GetChildNodes())
+	// Common up-to-date case: check the desired insert index for an existing node that matches
+	if (CurrentChildren.IsValidIndex(ExpectedIndex) && MatchNode(CurrentChildren[ExpectedIndex]))
 	{
-		if (Child->GetType() == ESequencerNode::Category && Child->GetNodeName() == CategoryName)
+		// Node already exists and is in the correct position
+		return CurrentChildren[ExpectedIndex];
+	}
+	// Rare-case: find an existing match at the wrong index, and move it to the correct index
+	else for (int32 Index = 0; Index < CurrentChildren.Num(); ++Index)
+	{
+		if (MatchNode(CurrentChildren[Index]))
 		{
-			// Ensure its name is up to date
-			CategoryNode = StaticCastSharedRef<FSequencerSectionCategoryNode>(Child);
+			// *Important - we copy the child here so that it's still valid after it gets moved
+			TSharedRef<FSequencerDisplayNode> Child = CurrentChildren[Index];
+
+			// Node already exists but is at the wrong index - needs moving
+			ParentNode->MoveChild(Index, ExpectedIndex);
+
+			return Child;
 		}
 	}
 
-	if (!CategoryNode)
+	return nullptr;
+}
+
+void FSequencerSectionLayoutBuilder::PushCategory( FName CategoryName, const FText& DisplayLabel )
+{
+	bHasAnyLayout = true;
+
+	const int32 DesiredInsertIndex = InsertIndexStack.Last();
+
+	TSharedPtr<FSequencerDisplayNode>         ExistingNode = FindAndRelocateExistingNode(CurrentNode, DesiredInsertIndex, ESequencerNode::Category, CategoryName);
+	TSharedPtr<FSequencerSectionCategoryNode> CategoryNode = StaticCastSharedPtr<FSequencerSectionCategoryNode>(ExistingNode);
+
+	if (!ExistingNode)
 	{
 		CategoryNode = MakeShared<FSequencerSectionCategoryNode>(CategoryName, RootNode->GetParentTree());
-		CategoryNode->SetParent(CurrentNode);
+		CategoryNode->SetParent(CurrentNode, DesiredInsertIndex);
 	}
 
 	CategoryNode->DisplayName = DisplayLabel;
 	CategoryNode->TreeSerialNumber = RootNode->TreeSerialNumber;
 	CurrentNode = CategoryNode.ToSharedRef();
+	
+	// Move onto the next index at this level
+	++InsertIndexStack.Last();
+
+	// Add a new index to add to inside the new current node
+	InsertIndexStack.Add(0);
 }
 
 void FSequencerSectionLayoutBuilder::PopCategory()
@@ -48,6 +84,7 @@ void FSequencerSectionLayoutBuilder::PopCategory()
 		if (CurrentNode->GetType() == ESequencerNode::Category)
 		{
 			CurrentNode = Parent.ToSharedRef();
+			InsertIndexStack.Pop();
 		}
 	}
 }
@@ -83,26 +120,23 @@ void FSequencerSectionLayoutBuilder::AddChannel( const FMovieSceneChannelHandle&
 
 	bHasAnyLayout = true;
 
-	TSharedPtr<FSequencerSectionKeyAreaNode> KeyAreaNode;
+	const int32 DesiredInsertIndex = InsertIndexStack.Last();
 
-	// Attempt to re-use an existing key area of the same name
-	for (const TSharedRef<FSequencerDisplayNode>& Child : CurrentNode->GetChildNodes())
-	{
-		if (Child->GetType() == ESequencerNode::KeyArea && Child->GetNodeName() == MetaData->Name)
-		{
-			KeyAreaNode = StaticCastSharedRef<FSequencerSectionKeyAreaNode>(Child);
-		}
-	}
+	TSharedPtr<FSequencerDisplayNode>        ExistingNode = FindAndRelocateExistingNode(CurrentNode, DesiredInsertIndex, ESequencerNode::KeyArea, MetaData->Name);
+	TSharedPtr<FSequencerSectionKeyAreaNode> KeyAreaNode  = StaticCastSharedPtr<FSequencerSectionKeyAreaNode>(ExistingNode);
 
 	if (!KeyAreaNode.IsValid())
 	{
 		// No existing node found make a new one
 		KeyAreaNode = MakeShared<FSequencerSectionKeyAreaNode>(MetaData->Name, CurrentNode->GetParentTree());
 		KeyAreaNode->DisplayName = MetaData->DisplayText;
-		KeyAreaNode->SetParent(CurrentNode);
+		KeyAreaNode->SetParent(CurrentNode, DesiredInsertIndex);
 	}
 
 	AddOrUpdateChannel(KeyAreaNode.ToSharedRef(), Channel);
+
+	// Move onto the next index at this level
+	++InsertIndexStack.Last();
 }
 
 void FSequencerSectionLayoutBuilder::AddOrUpdateChannel(TSharedRef<FSequencerSectionKeyAreaNode> KeyAreaNode, const FMovieSceneChannelHandle& Channel)
diff --git a/Engine/Source/Editor/Sequencer/Private/SequencerSectionLayoutBuilder.h b/Engine/Source/Editor/Sequencer/Private/SequencerSectionLayoutBuilder.h
--- a/Engine/Source/Editor/Sequencer/Private/SequencerSectionLayoutBuilder.h
+++ b/Engine/Source/Editor/Sequencer/Private/SequencerSectionLayoutBuilder.h
@@ -43,4 +43,7 @@ private:
 
 	/** Boolean indicating whether this section layout builder has been given any layout or not */
 	bool bHasAnyLayout;
+
+	/** Stack of insertion indices for the current category level that define what child index the next node should be added as. */
+	TArray<int32, TInlineAllocator<1>> InsertIndexStack;
 };
