Sequencer: Fixed a flicker caused by uninitialized current position when initiating a marquee selection

diff --git a/Engine/Source/Editor/Sequencer/Private/Tools/SequencerEditTool_Selection.cpp b/Engine/Source/Editor/Sequencer/Private/Tools/SequencerEditTool_Selection.cpp
--- a/Engine/Source/Editor/Sequencer/Private/Tools/SequencerEditTool_Selection.cpp
+++ b/Engine/Source/Editor/Sequencer/Private/Tools/SequencerEditTool_Selection.cpp
@@ -98,7 +98,7 @@ public:
 	virtual void OnBeginDrag(const FPointerEvent& MouseEvent, FVector2D LocalMousePos, const FVirtualTrackArea& VirtualTrackArea) override
 	{
 		// Start a new marquee selection
-		InitialPosition =  VirtualTrackArea.PhysicalToVirtual(LocalMousePos);
+		InitialPosition = CurrentPosition = VirtualTrackArea.PhysicalToVirtual(LocalMousePos);
 		CurrentMousePos = LocalMousePos;
 
 		if (MouseEvent.IsShiftDown())
