Sequencer: Refocus on the previously focused when setting the current time with the play time widget. This allows the user to do: spacebar play, ctrl-g to set a time, spacebar play without needing to click to focus after setting the time.

diff --git a/Engine/Source/Editor/Sequencer/Private/SSequencer.cpp b/Engine/Source/Editor/Sequencer/Private/SSequencer.cpp
--- a/Engine/Source/Editor/Sequencer/Private/SSequencer.cpp
+++ b/Engine/Source/Editor/Sequencer/Private/SSequencer.cpp
@@ -612,7 +612,7 @@ void SSequencer::Construct(const FArguments& InArgs, TSharedRef<FSequencer> InSe
 								.BorderImage(nullptr)
 								[
 									// Current Play Time 
-									SAssignNew(PlayTimeDisplay, SSpinBox<double>)
+									SAssignNew(PlayTimeDisplay, STemporarilyFocusedSpinBox<double>)
 									.Style(&FEditorStyle::GetWidgetStyle<FSpinBoxStyle>("Sequencer.PlayTimeSpinBox"))
 									.Value_Lambda([this]() -> double {
 										return SequencerPtr.Pin()->GetLocalTime().Time.GetFrame().Value;
@@ -620,6 +620,9 @@ void SSequencer::Construct(const FArguments& InArgs, TSharedRef<FSequencer> InSe
 									.OnValueChanged(this, &SSequencer::SetPlayTimeClampedByWorkingRange)
 									.OnValueCommitted_Lambda([this](double InFrame, ETextCommit::Type) {
 										SetPlayTimeClampedByWorkingRange(InFrame);
+
+										// Refocus on the previously focused widget so that user can continue on after setting a time
+										PlayTimeDisplay->Refocus();
 									})
 									.MinValue(TOptional<double>())
 									.MaxValue(TOptional<double>())
@@ -859,7 +862,7 @@ void SSequencer::BindCommands(TSharedRef<FUICommandList> SequencerCommandBinding
 
 	SequencerCommandBindings->MapAction(
 		FSequencerCommands::Get().ToggleShowGotoBox,
-		FExecuteAction::CreateLambda([this] { FSlateApplication::Get().SetKeyboardFocus(PlayTimeDisplay, EFocusCause::SetDirectly); })
+		FExecuteAction::CreateLambda([this] { PlayTimeDisplay->Setup();  FSlateApplication::Get().SetKeyboardFocus(PlayTimeDisplay, EFocusCause::SetDirectly); })
 	);
 
 	SequencerCommandBindings->MapAction(
diff --git a/Engine/Source/Editor/Sequencer/Private/SSequencer.h b/Engine/Source/Editor/Sequencer/Private/SSequencer.h
--- a/Engine/Source/Editor/Sequencer/Private/SSequencer.h
+++ b/Engine/Source/Editor/Sequencer/Private/SSequencer.h
@@ -86,6 +86,30 @@ struct FSequencerBreadcrumb
 };
 
 
+/**
+ * A widget that holds a widget that is to be refocused on completion
+ */
+template<typename T>
+struct STemporarilyFocusedSpinBox : SSpinBox<T>
+{
+public:
+	void Setup()
+	{
+		PreviousFocusedWidget = FSlateApplication::Get().GetKeyboardFocusedWidget();
+	}
+
+	void Refocus()
+	{
+		if (PreviousFocusedWidget.IsValid())
+		{
+			FSlateApplication::Get().SetKeyboardFocus(PreviousFocusedWidget.Pin());
+		}
+	}
+
+private:
+	TWeakPtr<SWidget> PreviousFocusedWidget;
+};
+
 /**
  * Main sequencer UI widget
  */
@@ -510,7 +534,7 @@ private:
 	TSharedPtr<SSearchBox> SearchBox;
 
 	/** The current playback time display.*/
-	TSharedPtr<SSpinBox<double>> PlayTimeDisplay;
+	TSharedPtr<STemporarilyFocusedSpinBox<double>> PlayTimeDisplay;
 
 	/** The sequencer tree view responsible for the outliner and track areas */
 	TSharedPtr<SSequencerTreeView> TreeView;
