Sequencer: Check for null section when gathering keys for the cinematic viewport

diff --git a/Engine/Source/Editor/Sequencer/Private/SequencerKeyCollection.cpp b/Engine/Source/Editor/Sequencer/Private/SequencerKeyCollection.cpp
--- a/Engine/Source/Editor/Sequencer/Private/SequencerKeyCollection.cpp
+++ b/Engine/Source/Editor/Sequencer/Private/SequencerKeyCollection.cpp
@@ -161,7 +161,10 @@ bool FSequencerKeyCollection::Update(const FSequencerKeyCollectionSignature& InS
 	// Get all the key times for the key areas
 	for (auto& Pair : InSignature.GetKeyAreas())
 	{
-		Pair.Key->GetKeyTimes(AllTimes, Pair.Key->GetOwningSection()->GetRange());
+		if (Pair.Key->GetOwningSection())
+		{
+			Pair.Key->GetKeyTimes(AllTimes, Pair.Key->GetOwningSection()->GetRange());
+		}
 	}
 
 	AllTimes.Sort();
