Sequencer: Fixed key editors breaking when their channel handles became invalid

  - The key editor switcher now detects when any change of the key area array has been made and rebuilds its widget switcher in that case
  - This change also adds protection against stale key areas remaining inside a key area node by adding a tree serial number to the key areas in the same way we do for key area nodes

diff --git a/Engine/Source/Editor/Sequencer/Private/DisplayNodes/SequencerSectionKeyAreaNode.cpp b/Engine/Source/Editor/Sequencer/Private/DisplayNodes/SequencerSectionKeyAreaNode.cpp
--- a/Engine/Source/Editor/Sequencer/Private/DisplayNodes/SequencerSectionKeyAreaNode.cpp
+++ b/Engine/Source/Editor/Sequencer/Private/DisplayNodes/SequencerSectionKeyAreaNode.cpp
@@ -21,11 +21,18 @@ FSequencerSectionKeyAreaNode::FSequencerSectionKeyAreaNode(FName NodeName, FSequ
 
 void FSequencerSectionKeyAreaNode::AddKeyArea(TSharedRef<IKeyArea> KeyArea)
 {
+	KeyArea->TreeSerialNumber = TreeSerialNumber;
 	KeyAreas.Add(KeyArea);
+}
 
-	if (KeyEditorSwitcher.IsValid())
+void FSequencerSectionKeyAreaNode::RemoveStaleKeyAreas()
+{
+	for (int32 Index = KeyAreas.Num() - 1; Index >= 0; --Index)
 	{
-		KeyEditorSwitcher->Rebuild();
+		if (KeyAreas[Index]->TreeSerialNumber != TreeSerialNumber)
+		{
+			KeyAreas.RemoveAt(Index, 1, false);
+		}
 	}
 }
 
@@ -49,26 +56,22 @@ bool FSequencerSectionKeyAreaNode::CanRenameNode() const
 	return false;
 }
 
-TSharedRef<SWidget> FSequencerSectionKeyAreaNode::GetOrCreateKeyAreaEditorSwitcher()
+EVisibility FSequencerSectionKeyAreaNode::GetKeyEditorVisibility() const
 {
-	if (!KeyEditorSwitcher.IsValid())
-	{
-		KeyEditorSwitcher = SNew(SKeyAreaEditorSwitcher, SharedThis(this));
-	}
-	return KeyEditorSwitcher.ToSharedRef();
+	return KeyAreas.Num() == 0 ? EVisibility::Collapsed : EVisibility::Visible;
 }
 
 TSharedRef<SWidget> FSequencerSectionKeyAreaNode::GetCustomOutlinerContent()
 {
-	if (GetAllKeyAreas().Num() > 0)
-	{
-		return SNew(SHorizontalBox)
+	// Even if this key area node doesn't have any key areas right now, it may in the future, so we always create the switcher, and just hide it if it is not relevant
+	return SNew(SHorizontalBox)
+		.Visibility(this, &FSequencerSectionKeyAreaNode::GetKeyEditorVisibility)
 
 		+ SHorizontalBox::Slot()
 		.HAlign(HAlign_Right)
 		.VAlign(VAlign_Center)
 		[
-			GetOrCreateKeyAreaEditorSwitcher()
+			SNew(SKeyAreaEditorSwitcher, SharedThis(this))
 		]
 
 		+ SHorizontalBox::Slot()
@@ -77,9 +80,6 @@ TSharedRef<SWidget> FSequencerSectionKeyAreaNode::GetCustomOutlinerContent()
 		[
 			SNew(SKeyNavigationButtons, SharedThis(this))
 		];
-	}
-
-	return FSequencerDisplayNode::GetCustomOutlinerContent();
 }
 
 
diff --git a/Engine/Source/Editor/Sequencer/Private/DisplayNodes/SequencerSectionKeyAreaNode.h b/Engine/Source/Editor/Sequencer/Private/DisplayNodes/SequencerSectionKeyAreaNode.h
--- a/Engine/Source/Editor/Sequencer/Private/DisplayNodes/SequencerSectionKeyAreaNode.h
+++ b/Engine/Source/Editor/Sequencer/Private/DisplayNodes/SequencerSectionKeyAreaNode.h
@@ -58,10 +58,10 @@ public:
 		return KeyAreas;
 	}
 
-	void ClearKeyAreas();
-
-	/** Retrieve the key area editor switcher widget, creating it if it doesn't yet exist */
-	TSharedRef<SWidget> GetOrCreateKeyAreaEditorSwitcher();
+	/**
+	 * Remove any key areas that do not correspond to the current tree serial number of this node
+	 */
+	void RemoveStaleKeyAreas();
 
 public:
 
@@ -80,6 +80,8 @@ public:
 
 private:
 
+	EVisibility GetKeyEditorVisibility() const;
+
 	/** All key areas on this node (one per section). */
 	TArray<TSharedRef<IKeyArea>> KeyAreas;
 
diff --git a/Engine/Source/Editor/Sequencer/Private/DisplayNodes/SequencerTrackNode.cpp b/Engine/Source/Editor/Sequencer/Private/DisplayNodes/SequencerTrackNode.cpp
--- a/Engine/Source/Editor/Sequencer/Private/DisplayNodes/SequencerTrackNode.cpp
+++ b/Engine/Source/Editor/Sequencer/Private/DisplayNodes/SequencerTrackNode.cpp
@@ -29,6 +29,7 @@
 #include "MovieScene.h"
 #include "Tracks/MovieSceneCinematicShotTrack.h"
 #include "Tracks/MovieSceneCameraCutTrack.h"
+#include "SKeyAreaEditorSwitcher.h"
 #include "MovieSceneFolder.h"
 #include "Tracks/MovieScenePropertyTrack.h"
 #include "Tracks/MovieScene3DTransformTrack.h"
@@ -266,14 +267,18 @@ void FSequencerTrackNode::RemoveStaleChildren()
 		TSharedRef<FSequencerDisplayNode> Child = NodesToCheck[Index];
 		if (Child->TreeSerialNumber != TreeSerialNumber)
 		{
-			// This node is stale - remove it
+			// This node is stale - remove it and orphan any children
 			StaleNodes.Add(Child);
+			continue;
 		}
-		else
+
+		if (Child->GetType() == ESequencerNode::KeyArea)
 		{
-			// This node is still relevant, but its children may not be - recurse into those
-			NodesToCheck.Append(Child->GetChildNodes());
+			StaticCastSharedRef<FSequencerSectionKeyAreaNode>(Child)->RemoveStaleKeyAreas();
 		}
+
+		// This node is still relevant, but its children may not be - recurse into those
+		NodesToCheck.Append(Child->GetChildNodes());
 	}
 
 	for (TSharedRef<FSequencerDisplayNode> StaleNode : StaleNodes)
@@ -381,7 +386,7 @@ TSharedRef<SWidget> FSequencerTrackNode::GetCustomOutlinerContent()
 	TSharedPtr<SWidget> KeyEditorWidget;
 	if (KeyAreaNode.IsValid())
 	{
-		KeyEditorWidget = KeyAreaNode->GetOrCreateKeyAreaEditorSwitcher();
+		KeyEditorWidget = SNew(SKeyAreaEditorSwitcher, KeyAreaNode.ToSharedRef());
 	}
 
 	TAttribute<bool> NodeIsHovered = TAttribute<bool>::Create(TAttribute<bool>::FGetter::CreateSP(this, &FSequencerDisplayNode::IsHovered));
diff --git a/Engine/Source/Editor/Sequencer/Private/IKeyArea.cpp b/Engine/Source/Editor/Sequencer/Private/IKeyArea.cpp
--- a/Engine/Source/Editor/Sequencer/Private/IKeyArea.cpp
+++ b/Engine/Source/Editor/Sequencer/Private/IKeyArea.cpp
@@ -14,7 +14,8 @@
 #include "MovieSceneSequence.h"
 
 IKeyArea::IKeyArea(UMovieSceneSection* InSection, FMovieSceneChannelHandle InChannel)
-	: ChannelHandle(InChannel)
+	: TreeSerialNumber(0)
+	, ChannelHandle(InChannel)
 {
 	Reinitialize(InSection, InChannel);
 }
@@ -220,7 +221,7 @@ TUniquePtr<FCurveModel> IKeyArea::CreateCurveEditorModel(TSharedRef<ISequencer>
 {
 	ISequencerChannelInterface* EditorInterface = FindChannelEditorInterface();
 	UMovieSceneSection* OwningSection = GetOwningSection();
-	if (EditorInterface && OwningSection)
+	if (EditorInterface && OwningSection && ChannelHandle.Get() != nullptr)
 	{
 		TUniquePtr<FCurveModel> CurveModel = EditorInterface->CreateCurveEditorModel_Raw(ChannelHandle, OwningSection, InSequencer);
 		if (CurveModel.IsValid())
diff --git a/Engine/Source/Editor/Sequencer/Private/SKeyAreaEditorSwitcher.cpp b/Engine/Source/Editor/Sequencer/Private/SKeyAreaEditorSwitcher.cpp
--- a/Engine/Source/Editor/Sequencer/Private/SKeyAreaEditorSwitcher.cpp
+++ b/Engine/Source/Editor/Sequencer/Private/SKeyAreaEditorSwitcher.cpp
@@ -1,7 +1,7 @@
 // Copyright 1998-2019 Epic Games, Inc. All Rights Reserved.
 
 #include "SKeyAreaEditorSwitcher.h"
-#include "Widgets/SOverlay.h"
+#include "Widgets/Layout/SWidgetSwitcher.h"
 #include "Widgets/Layout/SBox.h"
 #include "Widgets/SBoxPanel.h"
 #include "Sequencer.h"
@@ -12,70 +12,117 @@
 void SKeyAreaEditorSwitcher::Construct(const FArguments& InArgs, TSharedRef<FSequencerSectionKeyAreaNode> InKeyAreaNode)
 {
 	WeakKeyAreaNode = InKeyAreaNode;
-	ChildSlot
-	[
-		SAssignNew(Overlay, SOverlay)
-		.IsEnabled(!InKeyAreaNode->GetSequencer().IsReadOnly())
-	];
 
 	Rebuild();
 }
 
-void SKeyAreaEditorSwitcher::Rebuild()
+int32 SKeyAreaEditorSwitcher::GetWidgetIndex() const
 {
-	Overlay->ClearChildren();
-	VisibleIndex = -1;
+	return VisibleIndex;
+}
 
+void SKeyAreaEditorSwitcher::Rebuild()
+{
 	TSharedPtr<FSequencerSectionKeyAreaNode> KeyAreaNode = WeakKeyAreaNode.Pin();
 	if (!KeyAreaNode.IsValid())
 	{
+		// Empty our cache so we don't persistently rebuild
+		CachedKeyAreas.Empty();
+
+		// Node is no longer valid so just nuke everything and make this a null widget
+		ChildSlot
+		[
+			SNullWidget::NullWidget
+		];
 		return;
 	}
 
+	const bool bIsEnabled = !KeyAreaNode->GetSequencer().IsReadOnly();
+
+	// Update the cached list so we know when to rebuild next
+	CachedKeyAreas = KeyAreaNode->GetAllKeyAreas();
+
+	// Index 0 is always the spacer node
+	VisibleIndex = 0;
+
+	TSharedRef<SWidgetSwitcher> Switcher = SNew(SWidgetSwitcher)
+		.IsEnabled(bIsEnabled)
+		.WidgetIndex(this, &SKeyAreaEditorSwitcher::GetWidgetIndex)
+
+		+ SWidgetSwitcher::Slot()
+		[
+			SNullWidget::NullWidget
+		];
+
 	TSharedPtr<FSequencerObjectBindingNode> ParentObjectBinding = KeyAreaNode->FindParentObjectBindingNode();
 	FGuid ObjectBindingID = ParentObjectBinding.IsValid() ? ParentObjectBinding->GetObjectBinding() : FGuid();
 
-	for (int32 Index = 0; Index < KeyAreaNode->GetAllKeyAreas().Num(); ++Index)
+	for (TSharedRef<IKeyArea> KeyArea : CachedKeyAreas)
 	{
-		TSharedRef<IKeyArea> KeyArea = KeyAreaNode->GetAllKeyAreas()[Index];
-
-		if (KeyArea->CanCreateKeyEditor())
+		if (!KeyArea->CanCreateKeyEditor())
+		{
+			// Always generate a slot so that indices line up correctly
+			Switcher->AddSlot()
+			[
+				SNullWidget::NullWidget
+			];
+		}
+		else
 		{
-			Overlay->AddSlot()
+			Switcher->AddSlot()
 			[
 				SNew(SBox)
-				.IsEnabled(!KeyAreaNode->GetSequencer().IsReadOnly())
-				.HAlign(HAlign_Left)
+				.IsEnabled(bIsEnabled)
 				.MinDesiredWidth(100)
-				.Visibility(this, &SKeyAreaEditorSwitcher::GetWidgetVisibility, Index)
+				.HAlign(HAlign_Left)
 				[
 					KeyArea->CreateKeyEditor(KeyAreaNode->GetSequencer().AsShared(), ObjectBindingID)
 				]
 			];
 		}
 	}
-}
 
-EVisibility SKeyAreaEditorSwitcher::GetWidgetVisibility(int32 Index) const
-{
-	return Index == VisibleIndex ? EVisibility::Visible : EVisibility::Collapsed;
+	ChildSlot
+	[
+		Switcher
+	];
 }
 
 void SKeyAreaEditorSwitcher::Tick( const FGeometry& AllottedGeometry, const double InCurrentTime, const float InDeltaTime )
 {
-	VisibleIndex = -1;
-
 	TSharedPtr<FSequencerSectionKeyAreaNode> KeyAreaNode = WeakKeyAreaNode.Pin();
-	if (!KeyAreaNode.IsValid() || KeyAreaNode->GetAllKeyAreas().Num() == 0)
+	if (!KeyAreaNode.IsValid())
 	{
+		if (CachedKeyAreas.Num() != 0)
+		{
+			// Node is not valid but we have a valid cache - we need to rebuild the switcher now
+			Rebuild();
+		}
 		return;
 	}
-
-	TArray<UMovieSceneSection*> AllSections;
-	for (const TSharedRef<IKeyArea>& KeyArea : KeyAreaNode->GetAllKeyAreas())
+	else
 	{
-		AllSections.Add(KeyArea->GetOwningSection());
-	}
+		if (CachedKeyAreas != KeyAreaNode->GetAllKeyAreas())
+		{
+			// Node is valid but now has a different set of key areas. Must rebuild the widgets.
+			Rebuild();
+		}
+
+		TArray<UMovieSceneSection*> AllSections;
+		for (const TSharedRef<IKeyArea>& KeyArea : CachedKeyAreas)
+		{
+			AllSections.Add(KeyArea->GetOwningSection());
+		}
 
-	VisibleIndex = SequencerHelpers::GetSectionFromTime(AllSections, KeyAreaNode->GetSequencer().GetLocalTime().Time.FrameNumber);
+		const int32 ActiveKeyArea = SequencerHelpers::GetSectionFromTime(AllSections, KeyAreaNode->GetSequencer().GetLocalTime().Time.FrameNumber);
+		if (ActiveKeyArea != INDEX_NONE)
+		{
+			// Index 0 is the spacer node, so add 1 to the key area index to get the widget index
+			VisibleIndex = 1 + ActiveKeyArea;
+		}
+		else
+		{
+			VisibleIndex = 0;
+		}
+	}
 }
diff --git a/Engine/Source/Editor/Sequencer/Private/SKeyAreaEditorSwitcher.h b/Engine/Source/Editor/Sequencer/Private/SKeyAreaEditorSwitcher.h
--- a/Engine/Source/Editor/Sequencer/Private/SKeyAreaEditorSwitcher.h
+++ b/Engine/Source/Editor/Sequencer/Private/SKeyAreaEditorSwitcher.h
@@ -7,6 +7,7 @@
 
 class FSequencerSectionKeyAreaNode;
 class SOverlay;
+class IKeyArea;
 
 class SKeyAreaEditorSwitcher : public SCompoundWidget
 {
@@ -22,12 +23,11 @@ public:
 
 private:
 
-	/** Get the widget's visibility that corresponds to the specified key area index */
-	EVisibility GetWidgetVisibility(int32 Index) const;
-
 	/** Tick this widget. Updates the currently visible key editor */
 	virtual void Tick( const FGeometry& AllottedGeometry, const double InCurrentTime, const float InDeltaTime ) override;
 
+	int32 GetWidgetIndex() const;
+
 private:
 
 	/** Our overlay widget */
@@ -36,4 +36,6 @@ private:
 	int32 VisibleIndex;
 	/** The key area to which we relate */
 	TWeakPtr<FSequencerSectionKeyAreaNode> WeakKeyAreaNode;
+	/** Key areas cached from the node */
+	TArray<TSharedRef<IKeyArea>> CachedKeyAreas;
 };
diff --git a/Engine/Source/Editor/Sequencer/Private/SequencerSectionLayoutBuilder.cpp b/Engine/Source/Editor/Sequencer/Private/SequencerSectionLayoutBuilder.cpp
--- a/Engine/Source/Editor/Sequencer/Private/SequencerSectionLayoutBuilder.cpp
+++ b/Engine/Source/Editor/Sequencer/Private/SequencerSectionLayoutBuilder.cpp
@@ -141,6 +141,12 @@ void FSequencerSectionLayoutBuilder::AddChannel( const FMovieSceneChannelHandle&
 
 void FSequencerSectionLayoutBuilder::AddOrUpdateChannel(TSharedRef<FSequencerSectionKeyAreaNode> KeyAreaNode, const FMovieSceneChannelHandle& Channel)
 {
+	const FMovieSceneChannelMetaData* MetaData = Channel.GetMetaData();
+	if (!ensureAlwaysMsgf(MetaData, TEXT("Attempting to update an expired channel handle to the node tree")))
+	{
+		return;
+	}
+
 	KeyAreaNode->TreeSerialNumber = RootNode->TreeSerialNumber;
 
 	TSharedPtr<IKeyArea> KeyArea = KeyAreaNode->GetKeyArea(Section);
@@ -149,8 +155,11 @@ void FSequencerSectionLayoutBuilder::AddOrUpdateChannel(TSharedRef<FSequencerSec
 		// No key area for this section exists - create a new one
 		TSharedRef<IKeyArea> NewKeyArea = MakeShared<IKeyArea>(Section, Channel);
 		KeyAreaNode->AddKeyArea(NewKeyArea);
+		return;
 	}
-	else if (KeyArea->GetChannel() != Channel)
+
+	KeyArea->TreeSerialNumber = RootNode->TreeSerialNumber;
+	if (KeyArea->GetChannel() != Channel)
 	{
 		// A key area exists but for a different channel handle so this needs re-creating
 		KeyArea->Reinitialize(Section, Channel);
@@ -158,12 +167,6 @@ void FSequencerSectionLayoutBuilder::AddOrUpdateChannel(TSharedRef<FSequencerSec
 	else
 	{
 		// Just ensure the name is up to date
-		const FMovieSceneChannelMetaData* MetaData = Channel.GetMetaData();
-		if (!ensureAlwaysMsgf(MetaData, TEXT("Attempting to update an expired channel handle to the node tree")))
-		{
-			return;
-		}
-
 		KeyArea->SetName(MetaData->Name);
 	}
 }
diff --git a/Engine/Source/Editor/Sequencer/Public/IKeyArea.h b/Engine/Source/Editor/Sequencer/Public/IKeyArea.h
--- a/Engine/Source/Editor/Sequencer/Public/IKeyArea.h
+++ b/Engine/Source/Editor/Sequencer/Public/IKeyArea.h
@@ -35,6 +35,9 @@ class SEQUENCER_API IKeyArea : public TSharedFromThis<IKeyArea>
 {
 public:
 
+	/** Serial number of the node tree from the last time this key area was seen */
+	uint32 TreeSerialNumber;
+
 	/**
 	 * Constructor
 	 *
