Sequencer: Force an evaluation in the stopped state when scrubbing. Calling SetLocalTimeDirectly will not evaluate if the time has not changed.

diff --git a/Engine/Source/Editor/Sequencer/Private/Sequencer.cpp b/Engine/Source/Editor/Sequencer/Private/Sequencer.cpp
--- a/Engine/Source/Editor/Sequencer/Private/Sequencer.cpp
+++ b/Engine/Source/Editor/Sequencer/Private/Sequencer.cpp
@@ -4219,6 +4219,9 @@ void FSequencer::OnEndScrubbing()
 	AutoscrubOffset.Reset();
 	StopAutoscroll();
 
+	// Force an evaluation in the stopped state
+	ForceEvaluate();
+
 	OnEndScrubbingDelegate.Broadcast();
 }
 
