Sequencer: Restore pre animated state on PreCompileBlueprint since current cached state will no longer be valid. Force evaluation when the blueprint is compiled to put objects back at their evaluated state.

diff --git a/Engine/Source/Editor/Sequencer/Private/Sequencer.cpp b/Engine/Source/Editor/Sequencer/Private/Sequencer.cpp
--- a/Engine/Source/Editor/Sequencer/Private/Sequencer.cpp
+++ b/Engine/Source/Editor/Sequencer/Private/Sequencer.cpp
@@ -335,8 +335,21 @@ void FSequencer::InitSequencer(const FSequencerInitParams& InitParams, const TSh
 	}
 
 	{
-		FDelegateHandle OnBlueprintCompiledHandle = GEditor->OnBlueprintCompiled().AddLambda([&]{ State.InvalidateExpiredObjects(); });
-		AcquiredResources.Add([=]{ GEditor->OnBlueprintCompiled().Remove(OnBlueprintCompiledHandle); });
+		FDelegateHandle OnBlueprintPreCompileHandle = GEditor->OnBlueprintPreCompile().AddLambda([&](UBlueprint* InBlueprint)
+		{
+			// Restore pre animate state since objects will be reinstanced and current cached state will no longer be valid.
+			RestorePreAnimatedState();
+		});
+		AcquiredResources.Add([=] { GEditor->OnBlueprintPreCompile().Remove(OnBlueprintPreCompileHandle); });
+
+		FDelegateHandle OnBlueprintCompiledHandle = GEditor->OnBlueprintCompiled().AddLambda([&]
+		{
+			State.InvalidateExpiredObjects();
+
+			// Force re-evaluation since animated state was restored in PreCompile
+			ForceEvaluate();
+		});
+		AcquiredResources.Add([=] { GEditor->OnBlueprintCompiled().Remove(OnBlueprintCompiledHandle); });
 	}
 
 	{
