Sequencer: Check for null section when gathering keys

diff --git a/Engine/Source/Editor/Sequencer/Private/DisplayNodes/SequencerDisplayNode.cpp b/Engine/Source/Editor/Sequencer/Private/DisplayNodes/SequencerDisplayNode.cpp
--- a/Engine/Source/Editor/Sequencer/Private/DisplayNodes/SequencerDisplayNode.cpp
+++ b/Engine/Source/Editor/Sequencer/Private/DisplayNodes/SequencerDisplayNode.cpp
@@ -258,14 +258,17 @@ void SSequencerCombinedKeysTrack::GenerateCachedKeyPositions(const FGeometry& Al
 
 		UMovieSceneSection* Section = KeyArea->GetOwningSection();
 
-		if (Section->HasStartFrame())
+		if (Section)
 		{
-			SectionBoundTimes.Add(Section->GetInclusiveStartFrame() / CachedTickResolution);
-		}
+			if (Section->HasStartFrame())
+			{
+				SectionBoundTimes.Add(Section->GetInclusiveStartFrame() / CachedTickResolution);
+			}
 
-		if (Section->HasEndFrame())
-		{
-			SectionBoundTimes.Add(Section->GetExclusiveEndFrame() / CachedTickResolution);
+			if (Section->HasEndFrame())
+			{
+				SectionBoundTimes.Add(Section->GetExclusiveEndFrame() / CachedTickResolution);
+			}
 		}
 
 		FNameAndSignature CacheKey{ CachePair.Value, KeyArea->GetName() };
