Sequencer: Ensure bMouseDownInRegion is cleared on mouse up

diff --git a/Engine/Source/Editor/Sequencer/Private/SequencerTimeSliderController.cpp b/Engine/Source/Editor/Sequencer/Private/SequencerTimeSliderController.cpp
--- a/Engine/Source/Editor/Sequencer/Private/SequencerTimeSliderController.cpp
+++ b/Engine/Source/Editor/Sequencer/Private/SequencerTimeSliderController.cpp
@@ -774,6 +774,7 @@ FReply FSequencerTimeSliderController::OnMouseButtonUp( SWidget& WidgetOwner, co
 		}
 		
 		bPanning = false;
+		bMouseDownInRegion = false;
 		
 		return FReply::Handled().ReleaseMouseCapture();
 	}
@@ -856,6 +857,7 @@ FReply FSequencerTimeSliderController::OnMouseButtonUp( SWidget& WidgetOwner, co
 		return FReply::Handled().ReleaseMouseCapture();
 	}
 
+	bMouseDownInRegion = false;
 	return FReply::Unhandled();
 }
 
