Sequencer: Defer force evaluation until next tick to avoid reentrancy when compiling director blueprints

diff --git a/Engine/Source/Editor/Sequencer/Private/Sequencer.cpp b/Engine/Source/Editor/Sequencer/Private/Sequencer.cpp
--- a/Engine/Source/Editor/Sequencer/Private/Sequencer.cpp
+++ b/Engine/Source/Editor/Sequencer/Private/Sequencer.cpp
@@ -347,7 +347,7 @@ void FSequencer::InitSequencer(const FSequencerInitParams& InitParams, const TSh
 			State.InvalidateExpiredObjects();
 
 			// Force re-evaluation since animated state was restored in PreCompile
-			ForceEvaluate();
+			bNeedsEvaluate = true;
 		});
 		AcquiredResources.Add([=] { GEditor->OnBlueprintCompiled().Remove(OnBlueprintCompiledHandle); });
 	}
