Sequencer: Fixed crash when synchronizing the curve editor tree when a curve editor did not exist

diff --git a/Engine/Source/Editor/Sequencer/Private/SSequencer.cpp b/Engine/Source/Editor/Sequencer/Private/SSequencer.cpp
--- a/Engine/Source/Editor/Sequencer/Private/SSequencer.cpp
+++ b/Engine/Source/Editor/Sequencer/Private/SSequencer.cpp
@@ -1002,35 +1002,39 @@ void SSequencer::HandleOutlinerNodeSelectionChanged()
 
 	const TSet<TSharedRef<FSequencerDisplayNode>>& SelectedDisplayNodes = Sequencer->GetSelection().GetSelectedOutlinerNodes();
 
-	// If we're isolating to the selection and there is one, add the filter
-	if (Settings->ShouldIsolateToCurveEditorSelection() && SelectedDisplayNodes.Num() != 0)
+	TSharedPtr<FCurveEditor> CurveEditor = Sequencer->GetCurveEditor();
+	if (CurveEditor && CurveEditorTree)
 	{
-		if (!SequencerSelectionCurveEditorFilter)
+		// If we're isolating to the selection and there is one, add the filter
+		if (Settings->ShouldIsolateToCurveEditorSelection() && SelectedDisplayNodes.Num() != 0)
 		{
-			static const int32 FilterPass = -1000;
-			SequencerSelectionCurveEditorFilter = MakeShared<FCurveEditorTreeFilter>(ISequencerModule::GetSequencerSelectionFilterType(), FilterPass);
+			if (!SequencerSelectionCurveEditorFilter)
+			{
+				static const int32 FilterPass = -1000;
+				SequencerSelectionCurveEditorFilter = MakeShared<FCurveEditorTreeFilter>(ISequencerModule::GetSequencerSelectionFilterType(), FilterPass);
+			}
+			CurveEditor->GetTree()->AddFilter(SequencerSelectionCurveEditorFilter);
+		}
+		// If we're not isolating to the selection (or there is no selection) remove the filter
+		else if (SequencerSelectionCurveEditorFilter)
+		{
+			CurveEditor->GetTree()->RemoveFilter(SequencerSelectionCurveEditorFilter);
+			SequencerSelectionCurveEditorFilter = nullptr;
 		}
-		Sequencer->GetCurveEditor()->GetTree()->AddFilter(SequencerSelectionCurveEditorFilter);
-	}
-	// If we're not isolating to the selection (or there is no selection) remove the filter
-	else if (SequencerSelectionCurveEditorFilter)
-	{
-		Sequencer->GetCurveEditor()->GetTree()->RemoveFilter(SequencerSelectionCurveEditorFilter);
-		SequencerSelectionCurveEditorFilter = nullptr;
-	}
-
-	if (Settings->ShouldSyncCurveEditorSelection())
-	{
-		TSharedRef<FSequencerNodeTree> NodeTree = Sequencer->GetNodeTree();
 
-		// Clear the tree selection
-		CurveEditorTree->ClearSelection();
-		for (TSharedRef<FSequencerDisplayNode> Node : SelectedDisplayNodes)
+		if (Settings->ShouldSyncCurveEditorSelection())
 		{
-			FCurveEditorTreeItemID CurveEditorTreeItem = NodeTree->FindCurveEditorTreeItem(Node);
-			if (CurveEditorTreeItem != FCurveEditorTreeItemID::Invalid())
+			TSharedRef<FSequencerNodeTree> NodeTree = Sequencer->GetNodeTree();
+
+			// Clear the tree selection
+			CurveEditorTree->ClearSelection();
+			for (TSharedRef<FSequencerDisplayNode> Node : SelectedDisplayNodes)
 			{
-				CurveEditorTree->SetItemSelection(CurveEditorTreeItem, true);
+				FCurveEditorTreeItemID CurveEditorTreeItem = NodeTree->FindCurveEditorTreeItem(Node);
+				if (CurveEditorTreeItem != FCurveEditorTreeItemID::Invalid())
+				{
+					CurveEditorTree->SetItemSelection(CurveEditorTreeItem, true);
+				}
 			}
 		}
 	}
