Sequencer: Fix CombinedKeysTrack getting in the way of resizing shot tracks.

diff --git a/Engine/Source/Editor/Sequencer/Private/DisplayNodes/SequencerDisplayNode.cpp b/Engine/Source/Editor/Sequencer/Private/DisplayNodes/SequencerDisplayNode.cpp
--- a/Engine/Source/Editor/Sequencer/Private/DisplayNodes/SequencerDisplayNode.cpp
+++ b/Engine/Source/Editor/Sequencer/Private/DisplayNodes/SequencerDisplayNode.cpp
@@ -238,14 +238,6 @@ void SSequencerCombinedKeysTrack::Tick( const FGeometry& AllottedGeometry, const
 
 int32 SSequencerCombinedKeysTrack::OnPaint(const FPaintArgs& Args, const FGeometry& AllottedGeometry, const FSlateRect& MyCullingRect, FSlateWindowElementList& OutDrawElements, int32 LayerId, const FWidgetStyle& InWidgetStyle, bool bParentEnabled) const
 {
-	if (RootNode->GetType() == ESequencerNode::Track)
-	{
-		if (static_cast<FSequencerTrackNode&>(*RootNode).GetSubTrackMode() == FSequencerTrackNode::ESubTrackMode::ParentTrack && RootNode->IsExpanded())
-		{
-			return LayerId;
-		}
-	}
-
 	if (RootNode->GetSequencer().GetSequencerSettings()->GetShowCombinedKeyframes())
 	{
 		for (float KeyPosition : KeyDrawPositions)
@@ -771,6 +763,17 @@ TSharedRef<SWidget> FSequencerDisplayNode::GenerateWidgetForSectionArea(const TA
 	return SNew(SSequencerCombinedKeysTrack, SharedThis(this))
 		.ViewRange(ViewRange)
 		.IsEnabled(!GetSequencer().IsReadOnly())
+		.Visibility_Lambda([this]()
+		{
+			if (GetType() == ESequencerNode::Track)
+			{
+				if (static_cast<FSequencerTrackNode&>(*this).GetSubTrackMode() == FSequencerTrackNode::ESubTrackMode::ParentTrack && IsExpanded())
+				{
+					return EVisibility::Hidden;
+				}
+			}
+			return EVisibility::Visible;
+		})
 		.TickResolution(this, &FSequencerDisplayNode::GetTickResolution);
 }
 
