Pressing Ctrl F in the Sequencer or Sequencer Curve Editor will now jump keyboard focus to the search widget allowing you to quickly search the tree.

diff --git a/Engine/Source/Editor/CurveEditor/Private/Tree/SCurveEditorTreeTextFilter.cpp b/Engine/Source/Editor/CurveEditor/Private/Tree/SCurveEditorTreeTextFilter.cpp
--- a/Engine/Source/Editor/CurveEditor/Private/Tree/SCurveEditorTreeTextFilter.cpp
+++ b/Engine/Source/Editor/CurveEditor/Private/Tree/SCurveEditorTreeTextFilter.cpp
@@ -6,6 +6,7 @@
 #include "CurveEditor.h"
 
 #include "Widgets/Input/SSearchBox.h"
+#include "Framework/Application/SlateApplication.h"
 
 
 void SCurveEditorTreeTextFilter::Construct(const FArguments& InArgs, TSharedPtr<FCurveEditor> CurveEditor)
@@ -67,4 +68,12 @@ void SCurveEditorTreeTextFilter::OnFilterTextChanged( const FText& FilterText )
 		CurveEditor->GetTree()->Events.OnFiltersChanged.AddSP(this, &SCurveEditorTreeTextFilter::OnTreeFilterListChanged);
 	}
 }
-
+FReply SCurveEditorTreeTextFilter::OnFocusReceived(const FGeometry& MyGeometry, const FFocusEvent& InFocusEvent) 
+{ 
+	if(SearchBox)
+	{
+		FSlateApplication::Get().SetKeyboardFocus(SearchBox, InFocusEvent.GetCause());
+	}
+	
+	return FReply::Handled();
+}
diff --git a/Engine/Source/Editor/CurveEditor/Public/Tree/SCurveEditorTreeTextFilter.h b/Engine/Source/Editor/CurveEditor/Public/Tree/SCurveEditorTreeTextFilter.h
--- a/Engine/Source/Editor/CurveEditor/Public/Tree/SCurveEditorTreeTextFilter.h
+++ b/Engine/Source/Editor/CurveEditor/Public/Tree/SCurveEditorTreeTextFilter.h
@@ -17,7 +17,11 @@ public:
 	SLATE_END_ARGS()
 
 	void Construct(const FArguments& InArgs, TSharedPtr<FCurveEditor> CurveEditor);
-
+	
+	// SWidget Interface
+	virtual FReply OnFocusReceived(const FGeometry& MyGeometry, const FFocusEvent& InFocusEvent) override;
+	virtual bool SupportsKeyboardFocus() const override { return true; }
+	// ~SWidget Interface
 private:
 
 	void CreateSearchBox();
diff --git a/Engine/Source/Editor/Sequencer/Private/SSequencer.cpp b/Engine/Source/Editor/Sequencer/Private/SSequencer.cpp
--- a/Engine/Source/Editor/Sequencer/Private/SSequencer.cpp
+++ b/Engine/Source/Editor/Sequencer/Private/SSequencer.cpp
@@ -430,7 +430,7 @@ void SSequencer::Construct(const FArguments& InArgs, TSharedRef<FSequencer> InSe
 				+ SVerticalBox::Slot()
 				.AutoHeight()
 				[
-					SNew(SCurveEditorTreeTextFilter, InSequencer->GetCurveEditor())
+					SAssignNew(CurveEditorSearchBox, SCurveEditorTreeTextFilter, InSequencer->GetCurveEditor())
 				]
 
 				+ SVerticalBox::Slot()
@@ -883,6 +883,18 @@ void SSequencer::BindCommands(TSharedRef<FUICommandList> SequencerCommandBinding
 		FExecuteAction::CreateLambda([this] { TransformBox->ToggleVisibility(); })
 	);
 
+	// Allow jumping to the Sequencer tree search if you have Sequencer focused
+	SequencerCommandBindings->MapAction(
+		FSequencerCommands::Get().QuickTreeSearch,
+		FExecuteAction::CreateLambda([this] { FSlateApplication::Get().SetKeyboardFocus(SearchBox, EFocusCause::SetDirectly); })
+	);
+
+	// And jump to the Curve Editor tree search if you have the Curve Editor focused
+	SequencerPtr.Pin()->GetCurveEditor()->GetCommands()->MapAction(
+		FSequencerCommands::Get().QuickTreeSearch,
+		FExecuteAction::CreateLambda([this] { FSlateApplication::Get().SetKeyboardFocus(CurveEditorSearchBox, EFocusCause::SetDirectly); })
+	);
+	
 	SequencerCommandBindings->MapAction(
 		FSequencerCommands::Get().ToggleShowStretchBox,
 		FExecuteAction::CreateLambda([this] { StretchBox->ToggleVisibility(); })
diff --git a/Engine/Source/Editor/Sequencer/Private/SSequencer.h b/Engine/Source/Editor/Sequencer/Private/SSequencer.h
--- a/Engine/Source/Editor/Sequencer/Private/SSequencer.h
+++ b/Engine/Source/Editor/Sequencer/Private/SSequencer.h
@@ -543,6 +543,9 @@ private:
 	/** The search box for filtering tracks. */
 	TSharedPtr<SSearchBox> SearchBox;
 
+	/** The search widget for filtering curves in the Curve Editor tree. */
+	TSharedPtr<SWidget> CurveEditorSearchBox;
+
 	/** The current playback time display.*/
 	TSharedPtr<STemporarilyFocusedSpinBox<double>> PlayTimeDisplay;
 
diff --git a/Engine/Source/Editor/Sequencer/Private/SequencerCommands.cpp b/Engine/Source/Editor/Sequencer/Private/SequencerCommands.cpp
--- a/Engine/Source/Editor/Sequencer/Private/SequencerCommands.cpp
+++ b/Engine/Source/Editor/Sequencer/Private/SequencerCommands.cpp
@@ -147,6 +147,9 @@ void FSequencerCommands::RegisterCommands()
 	UI_COMMAND( ExportToCameraAnim, "Export to Camera Anim...", "Exports the animation to a camera anim asset.", EUserInterfaceActionType::Button, FInputChord() );
 
 	UI_COMMAND( ToggleEvaluateSubSequencesInIsolation, "Evaluate Sub Sequences In Isolation", "When enabled, will only evaluate the currently focused sequence; otherwise evaluate from the master sequence.", EUserInterfaceActionType::ToggleButton, FInputChord() );
+
+	UI_COMMAND( QuickTreeSearch, "Quick Tree Search", "Jumps keyboard focus to the tree searchbox to allow searching for tracks in the current Sequence.", EUserInterfaceActionType::Button, FInputChord(EModifierKey::Control, EKeys::F));
+	
 }
 
 #undef LOCTEXT_NAMESPACE
diff --git a/Engine/Source/Editor/Sequencer/Public/SequencerCommands.h b/Engine/Source/Editor/Sequencer/Public/SequencerCommands.h
--- a/Engine/Source/Editor/Sequencer/Public/SequencerCommands.h
+++ b/Engine/Source/Editor/Sequencer/Public/SequencerCommands.h
@@ -228,6 +228,9 @@ public:
 
 	/** Toggle the visibility of the stretch box. */
 	TSharedPtr< FUICommandInfo > ToggleShowStretchBox;
+	
+	/** Sets the tree search widget as the focused widget in Slate for easy typing. */
+	TSharedPtr< FUICommandInfo > QuickTreeSearch;
 
 	/** Bake transform. */
 	TSharedPtr< FUICommandInfo > BakeTransform;
