Additional performance improvements on the Sequencer viewport interactions.

diff --git a/Engine/Source/Editor/Cascade/Private/CascadeEmitterCanvasClient.h b/Engine/Source/Editor/Cascade/Private/CascadeEmitterCanvasClient.h
--- a/Engine/Source/Editor/Cascade/Private/CascadeEmitterCanvasClient.h
+++ b/Engine/Source/Editor/Cascade/Private/CascadeEmitterCanvasClient.h
@@ -29,11 +29,6 @@ public:
 	virtual void Draw(FViewport* Viewport, FCanvas* Canvas) override;
 	virtual bool InputKey(FViewport* Viewport, int32 ControllerId, FKey Key, EInputEvent Event, float AmountDepressed = 1.0f, bool bGamepad = false) override;
 	virtual void CapturedMouseMove(FViewport* Viewport, int32 X, int32 Y) override;
-	virtual void RequestRealTimeFrames(uint64 NumRealTimeFrames = 1) override
-	{
-		// Intentionally do nothing. We don't want to render this viewport except under our own very
-		// specific rules (this isn't a real editor viewport and doesn't call the super on draw anyways)
-	}
 
 	/** Returns the ratio of the size of the particle emitters to the size of the viewport */
 	float GetViewportVerticalScrollBarRatio() const;
diff --git a/Engine/Source/Editor/Sequencer/Private/LevelEditorSequencerIntegration.cpp b/Engine/Source/Editor/Sequencer/Private/LevelEditorSequencerIntegration.cpp
--- a/Engine/Source/Editor/Sequencer/Private/LevelEditorSequencerIntegration.cpp
+++ b/Engine/Source/Editor/Sequencer/Private/LevelEditorSequencerIntegration.cpp
@@ -411,10 +411,9 @@ void FLevelEditorSequencerIntegration::OnSequencerEvaluated()
 	}
 
 	// Request a single real-time frame to be rendered to ensure that we tick the world and update the viewport
-	// ToDo: This attempts to redraw all viewports (blueprint windows, cascade particles, etc.) so it's not the most efficient
-	// but we can't just use GetLevelViewportClients() because of Actor Sequences. For now, we've disabled the realtime frames
-	// on the most expensive viewport (cascade).
- 	for (FEditorViewportClient* LevelVC : GEditor->GetAllViewportClients())
+	// We only do this on level viewports instead of GetAllViewportClients to avoid needlessly redrawing Cascade,
+	// Blueprint, and other editors that have a 3d viewport.
+	for (FEditorViewportClient* LevelVC : GEditor->GetLevelViewportClients())
 	{
 		if (LevelVC && !LevelVC->IsRealtime())
 		{
