Sequencer: Null section check

diff --git a/Engine/Source/Editor/Sequencer/Private/Tools/EditToolDragOperations.cpp b/Engine/Source/Editor/Sequencer/Private/Tools/EditToolDragOperations.cpp
--- a/Engine/Source/Editor/Sequencer/Private/Tools/EditToolDragOperations.cpp
+++ b/Engine/Source/Editor/Sequencer/Private/Tools/EditToolDragOperations.cpp
@@ -1218,6 +1218,9 @@ void FMoveKeysAndSections::HandleKeyMovement(TOptional<FFrameNumber> MaxDeltaX,
 
 	for (UMovieSceneSection* Section : ModifiedNonSelectedSections)
 	{
-		Section->MarkAsChanged();
+		if (Section)
+		{
+			Section->MarkAsChanged();
+		}
 	}
 }
