Fix for trajectories with multiple parents.
Need to always get relative transforms and need to also always evaluate if there is a track in order to get correct pre/post values.  Now works with multiple parents, some in sequencer, or some not.

Also cleaned out some dead code, was confusing.

diff --git a/Engine/Source/Editor/Sequencer/Private/SequencerEdMode.cpp b/Engine/Source/Editor/Sequencer/Private/SequencerEdMode.cpp
--- a/Engine/Source/Editor/Sequencer/Private/SequencerEdMode.cpp
+++ b/Engine/Source/Editor/Sequencer/Private/SequencerEdMode.cpp
@@ -284,7 +284,7 @@ bool FSequencerEdMode::GetParentTM(FTransform& CurrentRefTM, const TSharedPtr<FS
 	{
 		return false;
 	}
-
+	//TODO this doesn't handle blended sections at all
 	for (const UMovieSceneTrack* Track : Binding->GetTracks())
 	{
 		const UMovieScene3DTransformTrack* TransformTrack = Cast<UMovieScene3DTransformTrack>(Track);
@@ -293,29 +293,25 @@ bool FSequencerEdMode::GetParentTM(FTransform& CurrentRefTM, const TSharedPtr<FS
 			continue;
 		}
 
-		for (const UMovieSceneSection* Section : TransformTrack->GetAllSections())
-		{
-			if (Section->IsTimeWithinSection(KeyTime.FrameNumber))
-			{
-				continue;
-			}
+		//we used to loop between sections here and only evaluate if we are in a section, this will give us wrong transfroms though
+		//when in between or outside of the section range. We still want to evaluate, though it is heavy.
 
-			FMovieSceneEvaluationTrack* EvalTrack = MovieSceneToolHelpers::GetEvaluationTrack(Sequencer.Get(), TransformTrack->GetSignature());
-			if (EvalTrack)
-			{
-				FVector ParentKeyPos;
-				FRotator ParentKeyRot;
-				GetLocationAtTime(EvalTrack, ParentObject, KeyTime, ParentKeyPos, ParentKeyRot, Sequencer);
-				CurrentRefTM = FTransform(ParentKeyRot, ParentKeyPos);
-				return true;
-			}
+		FMovieSceneEvaluationTrack* EvalTrack = MovieSceneToolHelpers::GetEvaluationTrack(Sequencer.Get(), TransformTrack->GetSignature());
+		if (EvalTrack)
+		{
+			FVector ParentKeyPos;
+			FRotator ParentKeyRot;
+			GetLocationAtTime(EvalTrack, ParentObject, KeyTime, ParentKeyPos, ParentKeyRot, Sequencer);
+			CurrentRefTM = FTransform(ParentKeyRot, ParentKeyPos);
+			return true;
 		}
+		
 	}
 
 	return false;
 }
 
-FTransform FSequencerEdMode::GetRefFrame(const TSharedPtr<FSequencer>& Sequencer, const TArray<const UObject *>& Parents, FFrameTime KeyTime)
+FTransform FSequencerEdMode::GetRefFrameFromParents(const TSharedPtr<FSequencer>& Sequencer, const TArray<const UObject *>& Parents, FFrameTime KeyTime)
 {
 	FTransform RefTM = FTransform::Identity;
 	FTransform ParentRefTM = FTransform::Identity;
@@ -327,9 +323,18 @@ FTransform FSequencerEdMode::GetRefFrame(const TSharedPtr<FSequencer>& Sequencer
 		{
 			if (Actor->GetRootComponent() != nullptr && Actor->GetRootComponent()->GetAttachParent() != nullptr)
 			{
+				//Always get local ref tm since we don't know which parent is in the sequencer or not.
 				if (!GetParentTM(ParentRefTM, Sequencer, Actor->GetRootComponent()->GetAttachParent()->GetOwner(), KeyTime))
 				{
-					RefTM = Actor->GetRootComponent()->GetAttachParent()->GetSocketTransform(Actor->GetRootComponent()->GetAttachSocketName());
+					AActor *Parent = Actor->GetRootComponent()->GetAttachParent()->GetOwner();
+					if (Parent && Parent->GetRootComponent())
+					{
+						ParentRefTM = Parent->GetRootComponent()->GetRelativeTransform();
+					}
+					else
+					{
+						continue;
+					}
 				}
 				RefTM = ParentRefTM * RefTM;
 			}
@@ -353,107 +358,6 @@ FTransform FSequencerEdMode::GetRefFrame(const TSharedPtr<FSequencer>& Sequencer
 	return RefTM;
 }
 
-FTransform FSequencerEdMode::GetRefFrame(const TSharedPtr<FSequencer>& Sequencer, const UObject* InObject, FFrameTime KeyTime)
-{
-	FTransform RefTM = FTransform::Identity;
-	const AActor* Actor = Cast<AActor>(InObject);
-	if (Actor != nullptr)
-	{
-		RefTM = GetRefFrame(Sequencer, Actor, KeyTime);
-	}
-	else
-	{
-		const USceneComponent* SceneComponent = Cast<USceneComponent>(InObject);
-
-		if (SceneComponent != nullptr)
-		{
-			RefTM = GetRefFrame(Sequencer, SceneComponent, KeyTime);
-		}
-	}
-
-	return RefTM;
-}
-
-FTransform FSequencerEdMode::GetRefFrame(const TSharedPtr<FSequencer>& Sequencer, const AActor* Actor, FFrameTime KeyTime)
-{
-	FTransform RefTM = FTransform::Identity;
-
-	if (Actor != nullptr && Actor->GetRootComponent() != nullptr && Actor->GetRootComponent()->GetAttachParent() != nullptr)
-	{
-		RefTM = Actor->GetRootComponent()->GetAttachParent()->GetSocketTransform(Actor->GetRootComponent()->GetAttachSocketName());
-	}
-
-	return RefTM;
-}
-
-FTransform FSequencerEdMode::GetRefFrame(const TSharedPtr<FSequencer>& Sequencer, const USceneComponent* SceneComponent, FFrameTime KeyTime)
-{
-	UMovieSceneSequence* Sequence = Sequencer->GetFocusedMovieSceneSequence();
-	if (!Sequence)
-	{
-		return FTransform::Identity;
-	}
-
-	if (SceneComponent == nullptr || SceneComponent->GetAttachParent() == nullptr)
-	{
-		return FTransform::Identity;
-	}
-
-	FTransform ParentRefTM;
-
-	// If our parent is the root component, get the RefFrame from the Actor
-	if (SceneComponent->GetAttachParent() == SceneComponent->GetOwner()->GetRootComponent())
-	{
-		ParentRefTM = GetRefFrame(Sequencer, SceneComponent->GetAttachParent()->GetOwner(), KeyTime);
-	}
-	else
-	{
-		ParentRefTM = GetRefFrame(Sequencer, SceneComponent->GetAttachParent(), KeyTime);
-	}
-
-
-	FTransform CurrentRefTM = SceneComponent->GetAttachParent()->GetRelativeTransform();
-
-	// Check if our parent is animated in this Sequencer
-	UObject* ParentObject = SceneComponent->GetAttachParent() == SceneComponent->GetOwner()->GetRootComponent() ? static_cast<UObject*>(SceneComponent->GetOwner()) : SceneComponent->GetAttachParent();
-	FGuid ObjectBinding = Sequencer->FindCachedObjectId(*ParentObject, Sequencer->GetFocusedTemplateID());
-
-	const FMovieSceneBinding* Binding = ObjectBinding.IsValid() ? Sequence->GetMovieScene()->FindBinding(ObjectBinding) : nullptr;
-	if (Binding)
-	{
-		for (const UMovieSceneTrack* Track : Binding->GetTracks())
-		{
-			const UMovieScene3DTransformTrack* TransformTrack = Cast<UMovieScene3DTransformTrack>(Track);
-			if (!TransformTrack)
-			{
-				continue;
-			}
-
-			for (const UMovieSceneSection* Section : TransformTrack->GetAllSections())
-			{
-				if (Section->IsTimeWithinSection(KeyTime.FrameNumber))
-				{
-					continue;
-				}
-
-				FMovieSceneEvaluationTrack* EvalTrack = MovieSceneToolHelpers::GetEvaluationTrack(Sequencer.Get(), TransformTrack->GetSignature());
-				if (EvalTrack)
-				{
-					FVector ParentKeyPos;
-					FRotator ParentKeyRot;
-					GetLocationAtTime(EvalTrack, ParentObject, KeyTime, ParentKeyPos, ParentKeyRot, Sequencer);
-
-					CurrentRefTM = FTransform(ParentKeyRot, ParentKeyPos);
-
-					return CurrentRefTM * ParentRefTM;
-				}
-			}
-		}
-	}
-
-	return CurrentRefTM * ParentRefTM;
-}
-
 void FSequencerEdMode::GetLocationAtTime(FMovieSceneEvaluationTrack* Track, UObject* Object, FFrameTime KeyTime, FVector& KeyPos, FRotator& KeyRot, const TSharedPtr<FSequencer>& Sequencer)
 {
 	FMovieSceneInterrogationData InterrogationData;
@@ -560,7 +464,7 @@ void FSequencerEdMode::DrawTransformTrack(const TSharedPtr<FSequencer>& Sequence
 				FRotator NewKeyRot(0,0,0);
 
 				GetLocationAtTime(EvalTrack, BoundObject, NewKeyTime, NewKeyPos, NewKeyRot, Sequencer);
-				FTransform NewPosRefTM = GetRefFrame(Sequencer, Parents, NewKeyTime);
+				FTransform NewPosRefTM = GetRefFrameFromParents(Sequencer, Parents, NewKeyTime);
 				FVector NewKeyPos_G = NewPosRefTM.TransformPosition(NewKeyPos);
 				FKeyPositionRotation KeyPosRot(NewTrajectoryKey,NewKeyPos, NewKeyRot, NewKeyPos_G);
 				KeyPosRots.Push(KeyPosRot);
@@ -597,7 +501,7 @@ void FSequencerEdMode::DrawTransformTrack(const TSharedPtr<FSequencer>& Sequence
 							FRotator NewRot(0,0,0);
 							GetLocationAtTime(EvalTrack, BoundObject, NewTime, NewPos, NewRot, Sequencer);
 
-							FTransform RefTM = GetRefFrame(Sequencer, Parents, NewTime);
+							FTransform RefTM = GetRefFrameFromParents(Sequencer, Parents, NewTime);
 
 							FVector NewPos_G = RefTM.TransformPosition(NewPos);
 							if (PDI != nullptr)
diff --git a/Engine/Source/Editor/Sequencer/Private/SequencerEdMode.h b/Engine/Source/Editor/Sequencer/Private/SequencerEdMode.h
--- a/Engine/Source/Editor/Sequencer/Private/SequencerEdMode.h
+++ b/Engine/Source/Editor/Sequencer/Private/SequencerEdMode.h
@@ -75,11 +75,8 @@ protected:
 
 protected:
 	static void GetLocationAtTime(FMovieSceneEvaluationTrack* Track, UObject* BoundObject, FFrameTime KeyTime, FVector& KeyPos, FRotator& KeyRot, const TSharedPtr<FSequencer>& Sequencer);
-	static FTransform GetRefFrame(const TSharedPtr<FSequencer>& Sequencer, const UObject* InObject, FFrameTime KeyTime);
-	static FTransform GetRefFrame(const TSharedPtr<FSequencer>& Sequencer, const AActor* Actor, FFrameTime KeyTime);
-	static FTransform GetRefFrame(const TSharedPtr<FSequencer>& Sequencer, const USceneComponent* SceneComponent, FFrameTime KeyTime);
 	static void GetParents(TArray<const UObject *>& Parents, const UObject* InObject);
-	static FTransform GetRefFrame(const TSharedPtr<FSequencer>& Sequencer, const TArray<const UObject *>& Parents, FFrameTime KeyTime);
+	static FTransform GetRefFrameFromParents(const TSharedPtr<FSequencer>& Sequencer, const TArray<const UObject *>& Parents, FFrameTime KeyTime);
 	static bool GetParentTM(FTransform& CurrentRefTM, const TSharedPtr<FSequencer>& Sequencer, UObject* ParentObject, FFrameTime KeyTime);
 private:
 	TArray<TWeakPtr<FSequencer>> Sequencers;
