Sequencer: Check for valid generated class before restoring pre animated state

diff --git a/Engine/Source/Editor/Sequencer/Private/Sequencer.cpp b/Engine/Source/Editor/Sequencer/Private/Sequencer.cpp
--- a/Engine/Source/Editor/Sequencer/Private/Sequencer.cpp
+++ b/Engine/Source/Editor/Sequencer/Private/Sequencer.cpp
@@ -338,7 +338,7 @@ void FSequencer::InitSequencer(const FSequencerInitParams& InitParams, const TSh
 		FDelegateHandle OnBlueprintPreCompileHandle = GEditor->OnBlueprintPreCompile().AddLambda([&](UBlueprint* InBlueprint)
 		{
 			// Restore pre animate state since objects will be reinstanced and current cached state will no longer be valid.
-			if (InBlueprint)
+			if (InBlueprint && InBlueprint->GeneratedClass.Get())
 			{
 				RestorePreAnimatedState(InBlueprint->GeneratedClass.Get());
 			}
