Sequencer: Restore drawing paths for selected nodes and channels

diff --git a/Engine/Source/Editor/Sequencer/Private/SequencerEdMode.cpp b/Engine/Source/Editor/Sequencer/Private/SequencerEdMode.cpp
--- a/Engine/Source/Editor/Sequencer/Private/SequencerEdMode.cpp
+++ b/Engine/Source/Editor/Sequencer/Private/SequencerEdMode.cpp
@@ -718,6 +718,38 @@ void FSequencerEdMode::DrawTracks3D(FPrimitiveDrawInterface* PDI)
 			}
 			TSharedRef<FSequencerObjectBindingNode> ObjectBindingNode = ObjectBinding.Value.ToSharedRef();
 			bool bSelected = Sequencer->GetSelection().IsSelected(ObjectBindingNode);
+
+			if (!bSelected)
+			{
+				TSet<TSharedRef<FSequencerDisplayNode> > DescendantNodes;
+				SequencerHelpers::GetDescendantNodes(ObjectBindingNode, DescendantNodes);
+
+				// If one of our child is selected, we're considered selected
+				for (auto& DescendantNode : DescendantNodes)
+				{
+					if (Sequencer->GetSelection().IsSelected(DescendantNode) ||
+						Sequencer->GetSelection().NodeHasSelectedKeysOrSections(DescendantNode))
+					{
+						bSelected = true;
+						break;
+					}
+				}
+			}
+
+			// If one of our parent is selected, we're considered selected
+			TSharedPtr<FSequencerDisplayNode> ParentNode = ObjectBindingNode->GetParent();
+
+			while (!bSelected && ParentNode.IsValid())
+			{
+				if (Sequencer->GetSelection().IsSelected(ParentNode.ToSharedRef()) ||
+					Sequencer->GetSelection().NodeHasSelectedKeysOrSections(ParentNode.ToSharedRef()))
+				{
+					bSelected = true;
+				}
+
+				ParentNode = ParentNode->GetParent();
+			}
+
 			ObjectBindingNodesSelectionMap.Add(ObjectBindingNode, bSelected);
 		}
 
