Sequencer: Changed HandleSingleNode fallback check for collapsed parents.

diff --git a/Engine/Source/Editor/Sequencer/Private/Tools/SequencerEntityVisitor.cpp b/Engine/Source/Editor/Sequencer/Private/Tools/SequencerEntityVisitor.cpp
--- a/Engine/Source/Editor/Sequencer/Private/Tools/SequencerEntityVisitor.cpp
+++ b/Engine/Source/Editor/Sequencer/Private/Tools/SequencerEntityVisitor.cpp
@@ -193,7 +193,7 @@ void FSequencerEntityWalker::HandleSingleNode(const ISequencerEntityVisitor& Vis
 	//  - Sections that have no key areas
 	const bool bIterateKeyGroupings = Visitor.CheckEntityMask(ESequencerEntity::Key) &&
 		!bNodeHasKeyArea &&
-		(!InNode->IsExpanded() || InNode->GetChildNodes().Num() == 0);
+		(!InNode->IsExpanded() || InNode->GetChildNodes().Num() != 0);
 
 	if (bIterateKeyGroupings)
 	{
