Sequencer: Step to next/prev key no longer uses ExclusiveEndFrame - 1. This fixes cases where you want to step to the last key of a section and add another animation clip that doesn't overlap.

diff --git a/Engine/Source/Editor/Sequencer/Private/SSequencer.cpp b/Engine/Source/Editor/Sequencer/Private/SSequencer.cpp
--- a/Engine/Source/Editor/Sequencer/Private/SSequencer.cpp
+++ b/Engine/Source/Editor/Sequencer/Private/SSequencer.cpp
@@ -2094,7 +2094,7 @@ void SSequencer::StepToKey(bool bStepToNextKey, bool bCameraOnly)
 						}
 						if (Section->HasEndFrame())
 						{
-							AllTimes.Add(Section->GetExclusiveEndFrame()-1);
+							AllTimes.Add(Section->GetExclusiveEndFrame());
 						}
 					}
 				}
