Sequencer: Fix duplicating hovered keys and sections so that they are selected before duplicating if they're not already selected.

diff --git a/Engine/Source/Editor/Sequencer/Private/Tools/EditToolDragOperations.cpp b/Engine/Source/Editor/Sequencer/Private/Tools/EditToolDragOperations.cpp
--- a/Engine/Source/Editor/Sequencer/Private/Tools/EditToolDragOperations.cpp
+++ b/Engine/Source/Editor/Sequencer/Private/Tools/EditToolDragOperations.cpp
@@ -478,10 +478,17 @@ void FDuplicateKeysAndSections::OnBeginDrag(const FPointerEvent& MouseEvent, FVe
 
 	// Duplicate our selections as well.
 	bool bDelayedStructureRebuild = false;
+
+	TArray<UMovieSceneSection*> SectionsToDuplicate;
 	for (const FSectionHandle& SectionHandle : Sections)
 	{
-		UMovieSceneSection* DuplicatedSection = DuplicateObject<UMovieSceneSection>(SectionHandle.GetSectionObject(), SectionHandle.GetSectionObject()->GetOuter());
-		UMovieSceneTrack* OwningTrack = SectionHandle.GetSectionObject()->GetTypedOuter<UMovieSceneTrack>();
+		SectionsToDuplicate.Add(SectionHandle.GetSectionObject());
+	}
+
+	for (UMovieSceneSection* SectionToDuplicate : SectionsToDuplicate)
+	{
+		UMovieSceneSection* DuplicatedSection = DuplicateObject<UMovieSceneSection>(SectionToDuplicate, SectionToDuplicate->GetOuter());
+		UMovieSceneTrack* OwningTrack = SectionToDuplicate->GetTypedOuter<UMovieSceneTrack>();
 		OwningTrack->Modify();
 		OwningTrack->AddSection(*DuplicatedSection);
 
diff --git a/Engine/Source/Editor/Sequencer/Private/Tools/SequencerEditTool_Movement.cpp b/Engine/Source/Editor/Sequencer/Private/Tools/SequencerEditTool_Movement.cpp
--- a/Engine/Source/Editor/Sequencer/Private/Tools/SequencerEditTool_Movement.cpp
+++ b/Engine/Source/Editor/Sequencer/Private/Tools/SequencerEditTool_Movement.cpp
@@ -163,6 +163,42 @@ TSharedPtr<ISequencerEditToolDragOperation> FSequencerEditTool_Movement::CreateD
 		}
 		else if (bIsDuplicateEvent)
 		{
+			if (HotspotType == ESequencerHotspot::Key)
+			{
+				TArrayView<const FSequencerSelectedKey> HoveredKeys = StaticCastSharedPtr<FKeyHotspot>(DelayedDrag->Hotspot)->Keys;
+
+				auto AnyUnselectedKey = [&Selection](const FSequencerSelectedKey& InKey)
+				{
+					return !Selection.IsSelected(InKey);
+				};
+
+				if (HoveredKeys.ContainsByPredicate(AnyUnselectedKey))
+				{
+					// If any are not selected, we'll treat this as a unique drag
+					Selection.EmptySelectedKeys();
+					Selection.EmptySelectedSections();
+					Selection.EmptyNodesWithSelectedKeysOrSections();
+					for (const FSequencerSelectedKey& Key : HoveredKeys)
+					{
+						Selection.AddToSelection(Key);
+					}
+					SequencerHelpers::UpdateHoveredNodeFromSelectedKeys(Sequencer);
+				}
+			}
+			else if (HotspotType == ESequencerHotspot::Section)
+			{
+				FSectionHotspot HoveredSection = StaticCastSharedPtr<FSectionHotspot>(DelayedDrag->Hotspot)->Section;
+
+				if (!Selection.IsSelected(HoveredSection.Section.GetSectionObject()))
+				{
+					Selection.EmptySelectedKeys();
+					Selection.EmptySelectedSections();
+					Selection.EmptyNodesWithSelectedKeysOrSections();
+					Selection.AddToSelection(HoveredSection.Section.GetSectionObject());
+					SequencerHelpers::UpdateHoveredNodeFromSelectedSections(Sequencer);
+				}
+			}
+
 			return MakeShareable(new FDuplicateKeysAndSections(Sequencer, Selection.GetSelectedKeys(), SequencerWidget->GetSectionHandles(Selection.GetSelectedSections()), bHotspotIsSection));
 		}
 
