Sequencer: Fix GetParents not recursing properly to get parent actors. This fixes trajectory drawing when there's more than 1 parent in the chain.

diff --git a/Engine/Source/Editor/Sequencer/Private/SequencerEdMode.cpp b/Engine/Source/Editor/Sequencer/Private/SequencerEdMode.cpp
--- a/Engine/Source/Editor/Sequencer/Private/SequencerEdMode.cpp
+++ b/Engine/Source/Editor/Sequencer/Private/SequencerEdMode.cpp
@@ -261,25 +261,13 @@ static const float	CurveHandleScale = 0.5f;
 void FSequencerEdMode::GetParents(TArray<const UObject *>& Parents, const UObject* InObject)
 {
 	const AActor* Actor = Cast<AActor>(InObject);
-	if (Actor != nullptr)
-	{
-		Parents.Emplace(InObject);
-	}
-	else
+	if (Actor)
 	{
-		const USceneComponent* SceneComponent = Cast<USceneComponent>(InObject);
-
-		if (SceneComponent != nullptr)
+		Parents.Emplace(Actor);
+		const AActor* ParentActor = Actor->GetAttachParentActor();
+		if (ParentActor)
 		{
-			Parents.Emplace(InObject);
-			if (SceneComponent->GetAttachParent() == SceneComponent->GetOwner()->GetRootComponent())
-			{
-				GetParents(Parents, SceneComponent->GetAttachParent()->GetOwner());
-			}
-			else
-			{
-				GetParents(Parents, SceneComponent->GetAttachParent());
-			}
+			GetParents(Parents, ParentActor);
 		}
 	}
 }
