Sequencer: Limit key collection to section range. This fixes a bug where keys outside of the section range can be selected/manipulated.

diff --git a/Engine/Source/Editor/Sequencer/Private/SSequencer.cpp b/Engine/Source/Editor/Sequencer/Private/SSequencer.cpp
--- a/Engine/Source/Editor/Sequencer/Private/SSequencer.cpp
+++ b/Engine/Source/Editor/Sequencer/Private/SSequencer.cpp
@@ -2077,7 +2077,7 @@ void SSequencer::StepToKey(bool bStepToNextKey, bool bCameraOnly)
 				SequencerHelpers::GetAllKeyAreas( Node, KeyAreas );
 				for ( TSharedPtr<IKeyArea> KeyArea : KeyAreas )
 				{
-					KeyArea->GetKeyTimes(AllTimes);
+					KeyArea->GetKeyTimes(AllTimes, KeyArea->GetOwningSection()->GetRange());
 				}
 
 				TSet<TWeakObjectPtr<UMovieSceneSection> > Sections;
diff --git a/Engine/Source/Editor/Sequencer/Private/SequencerKeyCollection.cpp b/Engine/Source/Editor/Sequencer/Private/SequencerKeyCollection.cpp
--- a/Engine/Source/Editor/Sequencer/Private/SequencerKeyCollection.cpp
+++ b/Engine/Source/Editor/Sequencer/Private/SequencerKeyCollection.cpp
@@ -1,6 +1,7 @@
 // Copyright 1998-2018 Epic Games, Inc. All Rights Reserved.
 
 #include "SequencerKeyCollection.h"
+#include "MovieSceneSection.h"
 #include "IKeyArea.h"
 #include "DisplayNodes/SequencerDisplayNode.h"
 #include "DisplayNodes/SequencerTrackNode.h"
@@ -160,7 +161,7 @@ bool FSequencerKeyCollection::Update(const FSequencerKeyCollectionSignature& InS
 	// Get all the key times for the key areas
 	for (auto& Pair : InSignature.GetKeyAreas())
 	{
-		Pair.Key->GetKeyTimes(AllTimes);
+		Pair.Key->GetKeyTimes(AllTimes, Pair.Key->GetOwningSection()->GetRange());
 	}
 
 	AllTimes.Sort();
