Sequencer: Autoscrub to the desired time when holding down shift and clicking in the sequencer timeline. Autoscrub is also enabled when stepping to the next/previous mark.

diff --git a/Engine/Source/Editor/Sequencer/Private/Sequencer.cpp b/Engine/Source/Editor/Sequencer/Private/Sequencer.cpp
--- a/Engine/Source/Editor/Sequencer/Private/Sequencer.cpp
+++ b/Engine/Source/Editor/Sequencer/Private/Sequencer.cpp
@@ -131,6 +131,16 @@
 
 DEFINE_LOG_CATEGORY(LogSequencer);
 
+static TAutoConsoleVariable<float> CVarAutoScrubSpeed(
+	TEXT("Sequencer.AutoScrubSpeed"),
+	6.0f,
+	TEXT("How fast to scrub forward/backward when auto-scrubbing"));
+
+static TAutoConsoleVariable<float> CVarAutoScrubCurveExponent(
+	TEXT("Sequencer.AutoScrubCurveExponent"),
+	2.0f,
+	TEXT("How much to ramp in and out the scrub speed when auto-scrubbing"));
+
 struct FSequencerTemplateStore : IMovieSceneSequenceTemplateStore
 {
 	void Reset()
@@ -600,6 +610,35 @@ void FSequencer::Tick(float InDeltaTime)
 		PlayPosition.Reset(GlobalTime.ConvertTo(PlayPosition.GetInputRate()));
 	}
 
+	if (AutoScrubTarget.IsSet())
+	{
+		const double ScrubSpeed = CVarAutoScrubSpeed->GetFloat();		// How fast to scrub at peak curve speed
+		const double AutoScrubExp = CVarAutoScrubCurveExponent->GetFloat();	// How long to ease in and out.  Bigger numbers allow for longer easing.
+
+		const double SecondsPerFrame = GetFocusedTickResolution().AsInterval() / ScrubSpeed;
+		const int32 TotalFrames = FMath::Abs(AutoScrubTarget.GetValue().DestinationTime.GetFrame().Value - AutoScrubTarget.GetValue().SourceTime.GetFrame().Value);
+		const double TargetSeconds = (double)TotalFrames * SecondsPerFrame;
+
+		double ElapsedSeconds = FPlatformTime::Seconds() - AutoScrubTarget.GetValue().StartTime;
+		float Alpha = ElapsedSeconds / TargetSeconds;
+		Alpha = FMath::Clamp(Alpha, 0.f, 1.f);
+		int32 NewFrameNumber = FMath::InterpEaseInOut(AutoScrubTarget.GetValue().SourceTime.GetFrame().Value, AutoScrubTarget.GetValue().DestinationTime.GetFrame().Value, Alpha, AutoScrubExp);
+
+		FAutoScrubTarget CachedTarget = AutoScrubTarget.GetValue();
+
+		SetPlaybackStatus(EMovieScenePlayerStatus::Scrubbing);
+		PlayPosition.SetTimeBase(GetRootTickResolution(), GetRootTickResolution(), EMovieSceneEvaluationType::WithSubFrames);
+		SetLocalTimeDirectly(FFrameNumber(NewFrameNumber));
+
+		AutoScrubTarget = CachedTarget;
+
+		if (FMath::IsNearlyEqual(Alpha, 1.f, KINDA_SMALL_NUMBER))
+		{
+			SetPlaybackStatus(EMovieScenePlayerStatus::Stopped);
+			AutoScrubTarget.Reset();
+		}
+	}
+
 	UpdateSubSequenceData();
 
 	// Tick all the tools we own as well
@@ -2207,7 +2246,7 @@ void FSequencer::SetLocalTimeDirectly(FFrameTime NewTime)
 }
 
 
-void FSequencer::SetGlobalTime( FFrameTime NewTime )
+void FSequencer::SetGlobalTime(FFrameTime NewTime)
 {
 	NewTime = ConvertFrameTime(NewTime, GetRootTickResolution(), PlayPosition.GetInputRate());
 	if (PlayPosition.GetEvaluationType() == EMovieSceneEvaluationType::FrameLocked)
@@ -2222,6 +2261,12 @@ void FSequencer::SetGlobalTime( FFrameTime NewTime )
 	{
 		EvaluateInternal(PlayPosition.JumpTo(NewTime));
 	}
+
+	if (AutoScrubTarget.IsSet())
+	{
+		SetPlaybackStatus(EMovieScenePlayerStatus::Stopped);
+		AutoScrubTarget.Reset();
+	}
 }
 
 void FSequencer::ForceEvaluate()
@@ -2372,6 +2417,11 @@ TOptional<float> FSequencer::CalculateAutoscrollEncroachment(double NewTime, flo
 }
 
 
+void FSequencer::AutoScrubToTime(FFrameTime DestinationTime)
+{
+	AutoScrubTarget = FAutoScrubTarget(DestinationTime, GetLocalTime().Time, FPlatformTime::Seconds());
+}
+
 void FSequencer::SetPerspectiveViewportPossessionEnabled(bool bEnabled)
 {
 	bPerspectiveViewportPossessionEnabled = bEnabled;
@@ -3817,7 +3867,14 @@ void FSequencer::OnScrubPositionChanged( FFrameTime NewScrubPosition, bool bScru
 		NewScrubPosition = FMath::Clamp(NewScrubPosition, LowerBound, UpperBound);
 	}
 
-	SetLocalTimeDirectly( NewScrubPosition );
+	if (!bScrubbing && FSlateApplication::Get().GetModifierKeys().IsShiftDown())
+	{
+		AutoScrubToTime(NewScrubPosition);
+	}
+	else
+	{
+		SetLocalTimeDirectly(NewScrubPosition);
+	}
 }
 
 
@@ -7233,7 +7290,7 @@ void FSequencer::StepToNextMark()
 				{
 					if (TickFrame > CurrentTickFrameNumber)
 					{
-						SetLocalTime(TickFrame.Value);
+						AutoScrubToTime(TickFrame.Value);
 						break;
 					}
 				}
@@ -7264,7 +7321,7 @@ void FSequencer::StepToPreviousMark()
 				{
 					if (TickFrame < CurrentTickFrameNumber)
 					{
-						SetLocalTime(TickFrame.Value);
+						AutoScrubToTime(TickFrame.Value);
 						break;
 					}
 				}
diff --git a/Engine/Source/Editor/Sequencer/Private/Sequencer.h b/Engine/Source/Editor/Sequencer/Private/Sequencer.h
--- a/Engine/Source/Editor/Sequencer/Private/Sequencer.h
+++ b/Engine/Source/Editor/Sequencer/Private/Sequencer.h
@@ -600,6 +600,9 @@ public:
 	 */
 	void UpdateAutoScroll(double NewTime);
 
+	/** Autoscrub to destination time */
+	void AutoScrubToTime(FFrameTime DestinationTime);
+
 public:
 
 	//~ FGCObject Interface
@@ -1029,6 +1032,19 @@ private:
 	/** The amount of autoscrub offset that is currently being applied */
 	TOptional<float> AutoscrubOffset;
 
+	struct FAutoScrubTarget
+	{
+		FAutoScrubTarget(FFrameTime InDestinationTime, FFrameTime InSourceTime, double InStartTime)
+			: DestinationTime(InDestinationTime)
+			, SourceTime(InSourceTime)
+			, StartTime(InStartTime) {}
+		FFrameTime DestinationTime;
+		FFrameTime SourceTime;
+		double StartTime;
+	};
+
+	TOptional<FAutoScrubTarget> AutoScrubTarget;
+
 	/** Zoom smoothing curves */
 	FCurveSequence ZoomAnimation;
 	FCurveHandle ZoomCurve;
