Sequencer: Refresh sequencer on map changed. This fixes a crash bug when unloading a map with actors attached to sublevels.

diff --git a/Engine/Source/Editor/Sequencer/Private/LevelEditorSequencerIntegration.cpp b/Engine/Source/Editor/Sequencer/Private/LevelEditorSequencerIntegration.cpp
--- a/Engine/Source/Editor/Sequencer/Private/LevelEditorSequencerIntegration.cpp
+++ b/Engine/Source/Editor/Sequencer/Private/LevelEditorSequencerIntegration.cpp
@@ -41,6 +41,7 @@
 #include "Widgets/Text/STextBlock.h"
 #include "UObject/ObjectKey.h"
 #include "UnrealEdGlobals.h"
+#include "UnrealEdMisc.h"
 #include "Editor/UnrealEdEngine.h"
 #include "EditorSupportDelegates.h"
 
@@ -238,13 +239,15 @@ void FLevelEditorSequencerIntegration::Initialize()
 	{
 		FLevelEditorModule& LevelEditorModule = FModuleManager::Get().GetModuleChecked<FLevelEditorModule>("LevelEditor");
 
-		FDelegateHandle Handle = LevelEditorModule.OnTabContentChanged().AddRaw(this, &FLevelEditorSequencerIntegration::OnTabContentChanged);
+		FDelegateHandle TabContentChanged = LevelEditorModule.OnTabContentChanged().AddRaw(this, &FLevelEditorSequencerIntegration::OnTabContentChanged);
+		FDelegateHandle MapChanged = LevelEditorModule.OnMapChanged().AddRaw(this, &FLevelEditorSequencerIntegration::OnMapChanged);
 		AcquiredResources.Add(
 			[=]{
 				FLevelEditorModule* LevelEditorModulePtr = FModuleManager::Get().GetModulePtr<FLevelEditorModule>("LevelEditor");
 				if (LevelEditorModulePtr)
 				{
-					LevelEditorModulePtr->OnTabContentChanged().Remove(Handle);
+					LevelEditorModulePtr->OnTabContentChanged().Remove(TabContentChanged);
+					LevelEditorModulePtr->OnMapChanged().Remove(MapChanged);
 				}
 			}
 		);
@@ -963,6 +966,29 @@ void FLevelEditorSequencerIntegration::OnTabContentChanged()
 
 }
 
+void FLevelEditorSequencerIntegration::OnMapChanged(UWorld* World, EMapChangeType MapChangeType)
+{
+	if (MapChangeType == EMapChangeType::TearDownWorld)
+	{
+		IterateAllSequencers(
+			[=](FSequencer& In, const FLevelEditorSequencerIntegrationOptions& Options)
+		{
+			if (Options.bRequiresLevelEvents)
+			{
+				In.GetEvaluationTemplate().ResetDirectorInstances();
+				In.RestorePreAnimatedState();
+				In.State.ClearObjectCaches(In);
+
+				// Notify data changed to enqueue an evaluate
+				In.NotifyMovieSceneDataChanged(EMovieSceneDataChangeType::Unknown);
+			}
+		}
+		);
+	}
+}
+
+
+
 void FLevelEditorSequencerIntegration::AddSequencer(TSharedRef<ISequencer> InSequencer, const FLevelEditorSequencerIntegrationOptions& Options)
 {
 	if (!BoundSequencers.Num())
diff --git a/Engine/Source/Editor/Sequencer/Public/LevelEditorSequencerIntegration.h b/Engine/Source/Editor/Sequencer/Public/LevelEditorSequencerIntegration.h
--- a/Engine/Source/Editor/Sequencer/Public/LevelEditorSequencerIntegration.h
+++ b/Engine/Source/Editor/Sequencer/Public/LevelEditorSequencerIntegration.h
@@ -107,6 +107,9 @@ private:
 	/** Called when viewport tab content changes. */
 	void OnTabContentChanged();
 
+	/** Called when the map is changed. */
+	void OnMapChanged(UWorld* World, EMapChangeType MapChangeType);
+
 	/** Called before a PIE session begins. */
 	void OnPreBeginPIE(bool bIsSimulating);
 
