Sequencer: Fix components for spawnables so that they also display in gray when despawned

diff --git a/Engine/Source/Editor/Sequencer/Private/DisplayNodes/SequencerObjectBindingNode.cpp b/Engine/Source/Editor/Sequencer/Private/DisplayNodes/SequencerObjectBindingNode.cpp
--- a/Engine/Source/Editor/Sequencer/Private/DisplayNodes/SequencerObjectBindingNode.cpp
+++ b/Engine/Source/Editor/Sequencer/Private/DisplayNodes/SequencerObjectBindingNode.cpp
@@ -419,11 +419,23 @@ FLinearColor FSequencerObjectBindingNode::GetDisplayNameColor() const
 
 	// Spawnables don't have valid object bindings when their track hasn't spawned them yet,
 	// so we override the default behavior of red with a gray so that users don't think there is something wrong.
-	if (GetBindingType() == EObjectBindingType::Spawnable)
+	
+	TSharedPtr<FSequencerDisplayNode> CurrentNode = SharedThis((FSequencerDisplayNode*)this);
+
+	while (CurrentNode.IsValid())
 	{
-		return FLinearColor::Gray;
+		if (CurrentNode->GetType() == ESequencerNode::Object)
+		{
+			if (StaticCastSharedPtr<FSequencerObjectBindingNode>(CurrentNode)->GetBindingType() == EObjectBindingType::Spawnable)
+			{
+				return FLinearColor::Gray;
+			}
+		}
+
+		CurrentNode = CurrentNode->GetParent();
 	}
 
+
 	return FLinearColor::Red;
 }
 
