Sequencer: When auto set track defaults is off, make sure a key is always added when adding a track so that the value is saved.

diff --git a/Engine/Source/Editor/Sequencer/Private/DisplayNodes/SequencerObjectBindingNode.cpp b/Engine/Source/Editor/Sequencer/Private/DisplayNodes/SequencerObjectBindingNode.cpp
--- a/Engine/Source/Editor/Sequencer/Private/DisplayNodes/SequencerObjectBindingNode.cpp
+++ b/Engine/Source/Editor/Sequencer/Private/DisplayNodes/SequencerObjectBindingNode.cpp
@@ -1011,7 +1011,9 @@ void FSequencerObjectBindingNode::HandlePropertyMenuItemExecute(FPropertyPath Pr
 		}
 	}
 
-	FKeyPropertyParams KeyPropertyParams(KeyableBoundObjects, PropertyPath, ESequencerKeyMode::ManualKey);
+	// When auto setting track defaults are disabled, force add a key so that the changed
+	// value is saved and is propagated to the property.
+	FKeyPropertyParams KeyPropertyParams(KeyableBoundObjects, PropertyPath, Sequencer.GetAutoSetTrackDefaults() == false ? ESequencerKeyMode::ManualKeyForced : ESequencerKeyMode::ManualKey);
 
 	Sequencer.KeyProperty(KeyPropertyParams);
 }
