Sequencer: Fix adding a binding so that default tracks are added as well. For example, if you add a camera cut for a camera that is not yet in Sequencer, default tracks should be added for that camera.

diff --git a/Engine/Source/Editor/Sequencer/Private/Sequencer.cpp b/Engine/Source/Editor/Sequencer/Private/Sequencer.cpp
--- a/Engine/Source/Editor/Sequencer/Private/Sequencer.cpp
+++ b/Engine/Source/Editor/Sequencer/Private/Sequencer.cpp
@@ -2663,6 +2663,16 @@ FGuid FSequencer::GetHandleToObject( UObject* Object, bool bCreateHandleIfMissin
 
 		ObjectGuid = CreateBinding(*Object, PossessedActor != nullptr ? PossessedActor->GetActorLabel() : Object->GetName());
 
+		AActor* ActorAdded = PossessedActor ? PossessedActor : Object->GetTypedOuter<AActor>();
+		if (ActorAdded)
+		{
+			FGuid ActorAddedGuid = GetHandleToObject(ActorAdded);
+			if (ActorAddedGuid.IsValid())
+			{
+				OnActorAddedToSequencerEvent.Broadcast(ActorAdded, ActorAddedGuid);
+			}
+		}
+
 		NotifyMovieSceneDataChanged( EMovieSceneDataChangeType::MovieSceneStructureItemAdded );
 	}
 	
