Sequencer: Fix updating radio buttons

diff --git a/Engine/Source/Editor/Sequencer/Private/SSequencer.cpp b/Engine/Source/Editor/Sequencer/Private/SSequencer.cpp
--- a/Engine/Source/Editor/Sequencer/Private/SSequencer.cpp
+++ b/Engine/Source/Editor/Sequencer/Private/SSequencer.cpp
@@ -1098,8 +1098,6 @@ void SSequencer::FillTimeDisplayFormatMenu(FMenuBuilder& MenuBuilder)
 
 	if (Settings)
 	{
-		EFrameNumberDisplayFormats CurrentDisplay = Settings->GetTimeDisplayFormat();
-
 		for (int32 Index = 0; Index < FrameNumberDisplayEnum->NumEnums() - 1; Index++)
 		{
 			if (!FrameNumberDisplayEnum->HasMetaData(TEXT("Hidden"), Index))
@@ -1110,7 +1108,6 @@ void SSequencer::FillTimeDisplayFormatMenu(FMenuBuilder& MenuBuilder)
 				if (Value == EFrameNumberDisplayFormats::DropFrameTimecode && !bSupportsDropFormatDisplay)
 					continue;
 
-				bool bIsSet = CurrentDisplay == Value;
 				MenuBuilder.AddMenuEntry(
 					FrameNumberDisplayEnum->GetDisplayNameTextByIndex(Index),
 					FrameNumberDisplayEnum->GetToolTipTextByIndex(Index),
@@ -1118,7 +1115,7 @@ void SSequencer::FillTimeDisplayFormatMenu(FMenuBuilder& MenuBuilder)
 					FUIAction(
 						FExecuteAction::CreateUObject(Settings, &USequencerSettings::SetTimeDisplayFormat, Value),
 						FCanExecuteAction(),
-						FIsActionChecked::CreateLambda([=] { return CurrentDisplay == Value; })
+						FIsActionChecked::CreateLambda([=] { return Settings->GetTimeDisplayFormat() == Value; })
 					),
 					NAME_None,
 					EUserInterfaceActionType::RadioButton
diff --git a/Engine/Source/Editor/Sequencer/Private/SSequencerPlayRateCombo.cpp b/Engine/Source/Editor/Sequencer/Private/SSequencerPlayRateCombo.cpp
--- a/Engine/Source/Editor/Sequencer/Private/SSequencerPlayRateCombo.cpp
+++ b/Engine/Source/Editor/Sequencer/Private/SSequencerPlayRateCombo.cpp
@@ -265,15 +265,12 @@ void SSequencerPlayRateCombo::PopulateClockSourceMenu(FMenuBuilder& MenuBuilder)
 
 	if (RootSequence)
 	{
-		EUpdateClockSource CurrentDisplay = RootSequence->GetMovieScene()->GetClockSource();
-
 		for (int32 Index = 0; Index < ClockSourceEnum->NumEnums() - 1; Index++)
 		{
 			if (!ClockSourceEnum->HasMetaData(TEXT("Hidden"), Index))
 			{
 				EUpdateClockSource Value = (EUpdateClockSource)ClockSourceEnum->GetValueByIndex(Index);
 
-				bool bIsSet = CurrentDisplay == Value;
 				MenuBuilder.AddMenuEntry(
 					ClockSourceEnum->GetDisplayNameTextByIndex(Index),
 					ClockSourceEnum->GetToolTipTextByIndex(Index),
@@ -281,7 +278,7 @@ void SSequencerPlayRateCombo::PopulateClockSourceMenu(FMenuBuilder& MenuBuilder)
 					FUIAction(
 						FExecuteAction::CreateSP(this, &SSequencerPlayRateCombo::SetClockSource, Value),
 						FCanExecuteAction(),
-						FIsActionChecked::CreateLambda([=]{ return CurrentDisplay == Value; })
+						FIsActionChecked::CreateLambda([=]{ return RootSequence->GetMovieScene()->GetClockSource() == Value; })
 					),
 					NAME_None,
 					EUserInterfaceActionType::RadioButton
