When converting bound ojbects in sequencer preserve sorting order and folder hierarchy.

diff --git a/Engine/Source/Editor/Sequencer/Private/Sequencer.cpp b/Engine/Source/Editor/Sequencer/Private/Sequencer.cpp
--- a/Engine/Source/Editor/Sequencer/Private/Sequencer.cpp
+++ b/Engine/Source/Editor/Sequencer/Private/Sequencer.cpp
@@ -6533,9 +6533,28 @@ FMovieSceneSpawnable* FSequencer::ConvertToSpawnableInternal(FGuid PossessableGu
 		// Remap all the spawnable's tracks and child bindings onto the new possessable
 		MovieScene->MoveBindingContents(PossessableGuid, SpawnableGuid);
 
+		FMovieSceneBinding* PossessableBinding = (FMovieSceneBinding*)MovieScene->GetBindings().FindByPredicate([&](FMovieSceneBinding& Binding) { return Binding.GetObjectGuid() == PossessableGuid; });
+		check(PossessableBinding);
+
+		for (UMovieSceneFolder* Folder : MovieScene->GetRootFolders())
+		{
+			if (ReplaceFolderBindingGUID(Folder, PossessableGuid, SpawnableGuid))
+			{
+				break;
+			}
+		}
+
+		int32 SortingOrder = PossessableBinding->GetSortingOrder();
+
 		if (MovieScene->RemovePossessable(PossessableGuid))
 		{
 			Sequence->UnbindPossessableObjects(PossessableGuid);
+
+			FMovieSceneBinding* SpawnableBinding = (FMovieSceneBinding*)MovieScene->GetBindings().FindByPredicate([&](FMovieSceneBinding& Binding) { return Binding.GetObjectGuid() == SpawnableGuid; });
+			check(SpawnableBinding);
+			
+			SpawnableBinding->SetSortingOrder(SortingOrder);
+
 		}
 
 		TOptional<FTransformData> TransformData;
@@ -6711,12 +6730,30 @@ FMovieScenePossessable* FSequencer::ConvertToPossessableInternal(FGuid Spawnable
 		// Remap all the spawnable's tracks and child bindings onto the new possessable
 		MovieScene->MoveBindingContents(OldSpawnableGuid, NewPossessableGuid);
 
+		FMovieSceneBinding* SpawnableBinding = (FMovieSceneBinding*)MovieScene->GetBindings().FindByPredicate([&](FMovieSceneBinding& Binding) { return Binding.GetObjectGuid() == OldSpawnableGuid; });
+		check(SpawnableBinding);
+
+		for (UMovieSceneFolder* Folder : MovieScene->GetRootFolders())
+		{
+			if (ReplaceFolderBindingGUID(Folder, Spawnable->GetGuid(), Possessable->GetGuid()))
+			{
+				break;
+			}
+		}
+
+		int32 SortingOrder = SpawnableBinding->GetSortingOrder();
+
 		// Remove the spawnable and all it's sub tracks
 		if (MovieScene->RemoveSpawnable(OldSpawnableGuid))
 		{
 			SpawnRegister->DestroySpawnedObject(OldSpawnableGuid, ActiveTemplateIDs.Top(), *this);
+
+			FMovieSceneBinding* PossessableBinding = (FMovieSceneBinding*)MovieScene->GetBindings().FindByPredicate([&](FMovieSceneBinding& Binding) { return Binding.GetObjectGuid() == NewPossessableGuid; });
+			check(PossessableBinding);
+			
+			PossessableBinding->SetSortingOrder(SortingOrder);
 		}
-	
+
 		static const FName SequencerActorTag(TEXT("SequencerActor"));
 		PossessedActor->Tags.Remove(SequencerActorTag);
 
@@ -6728,6 +6765,36 @@ FMovieScenePossessable* FSequencer::ConvertToPossessableInternal(FGuid Spawnable
 	return Possessable;
 }
 
+bool FSequencer::ReplaceFolderBindingGUID(UMovieSceneFolder* Folder, FGuid Original, FGuid Converted)
+{
+	UMovieScene* MovieScene = GetFocusedMovieSceneSequence()->GetMovieScene();
+
+	if (MovieScene->IsReadOnly())
+	{
+		return true;
+	}
+
+	for (FGuid ChildGuid : Folder->GetChildObjectBindings())
+	{
+		if (ChildGuid == Original)
+		{
+			Folder->AddChildObjectBinding(Converted);
+			Folder->RemoveChildObjectBinding(Original);
+			return true;
+		}
+	}
+
+	for (UMovieSceneFolder* ChildFolder : Folder->GetChildFolders())
+	{
+		if (ReplaceFolderBindingGUID(ChildFolder, Original, Converted))
+		{
+			return true;
+		}
+	}
+
+	return false;
+}
+
 void FSequencer::OnAddFolder()
 {
 	UMovieScene* FocusedMovieScene = GetFocusedMovieSceneSequence()->GetMovieScene();
diff --git a/Engine/Source/Editor/Sequencer/Private/Sequencer.h b/Engine/Source/Editor/Sequencer/Private/Sequencer.h
--- a/Engine/Source/Editor/Sequencer/Private/Sequencer.h
+++ b/Engine/Source/Editor/Sequencer/Private/Sequencer.h
@@ -905,6 +905,9 @@ protected:
 	/** Internal conversion function that doesn't perform expensive reset/update tasks */
 	FMovieScenePossessable* ConvertToPossessableInternal(FGuid SpawnableGuid);
 
+	/** Recurses through a folder to replace converted GUID with new GUID */
+	bool ReplaceFolderBindingGUID(UMovieSceneFolder *Folder, FGuid Original, FGuid Converted);
+
 	/** Internal function to render movie for a given start/end time */
 	void RenderMovieInternal(TRange<FFrameNumber> Range, bool bSetFrameOverrides = false) const;
 
