Sequencer: Remove "Keep Playback Range in Section Bounds" from level sequencer UI

diff --git a/Engine/Source/Editor/Sequencer/Private/SSequencer.cpp b/Engine/Source/Editor/Sequencer/Private/SSequencer.cpp
--- a/Engine/Source/Editor/Sequencer/Private/SSequencer.cpp
+++ b/Engine/Source/Editor/Sequencer/Private/SSequencer.cpp
@@ -1239,8 +1239,13 @@ TSharedRef<SWidget> SSequencer::MakePlaybackMenu()
 
 		MenuBuilder.AddMenuEntry( FSequencerCommands::Get().ToggleKeepCursorInPlaybackRangeWhileScrubbing );
 		MenuBuilder.AddMenuEntry( FSequencerCommands::Get().ToggleKeepCursorInPlaybackRange );
-		MenuBuilder.AddMenuEntry( FSequencerCommands::Get().ToggleKeepPlaybackRangeInSectionBounds );
-		MenuBuilder.AddMenuEntry( FSequencerCommands::Get().ToggleLinkCurveEditorTimeRange );
+
+		if (!SequencerPtr.Pin()->IsLevelEditorSequencer())
+		{
+			MenuBuilder.AddMenuEntry(FSequencerCommands::Get().ToggleKeepPlaybackRangeInSectionBounds);
+		}
+		
+		MenuBuilder.AddMenuEntry(FSequencerCommands::Get().ToggleLinkCurveEditorTimeRange);
 
 		// Menu entry for zero padding
 		auto OnZeroPadChanged = [=](uint8 NewValue){
