Sequencer: Pasting a spawnable doesn't set up defaults since tracks will be included in the copy

diff --git a/Engine/Source/Editor/Sequencer/Private/Sequencer.cpp b/Engine/Source/Editor/Sequencer/Private/Sequencer.cpp
--- a/Engine/Source/Editor/Sequencer/Private/Sequencer.cpp
+++ b/Engine/Source/Editor/Sequencer/Private/Sequencer.cpp
@@ -4039,7 +4039,7 @@ FGuid FSequencer::AddSpawnable(UObject& Object, UActorFactory* ActorFactory)
 	return NewGuid;
 }
 
-FGuid FSequencer::MakeNewSpawnable( UObject& Object, UActorFactory* ActorFactory )
+FGuid FSequencer::MakeNewSpawnable( UObject& Object, UActorFactory* ActorFactory, bool bSetupDefaults )
 {
 	UMovieSceneSequence* Sequence = GetFocusedMovieSceneSequence();
 	UMovieScene* MovieScene = Sequence->GetMovieScene();
@@ -4080,8 +4080,11 @@ FGuid FSequencer::MakeNewSpawnable( UObject& Object, UActorFactory* ActorFactory
 	// Spawn the object so we can position it correctly, it's going to get spawned anyway since things default to spawned.
 	UObject* SpawnedObject = SpawnRegister->SpawnObject(NewGuid, *MovieScene, ActiveTemplateIDs.Top(), *this);
 
-	FTransformData TransformData;
-	SpawnRegister->SetupDefaultsForSpawnable(SpawnedObject, Spawnable->GetGuid(), TransformData, AsShared(), Settings);
+	if (bSetupDefaults)
+	{
+		FTransformData TransformData;
+		SpawnRegister->SetupDefaultsForSpawnable(SpawnedObject, Spawnable->GetGuid(), TransformData, AsShared(), Settings);
+	}
 
 	Spawnable->SetSpawnOwnership(SavedOwnership);
 
@@ -6251,7 +6254,7 @@ bool FSequencer::PasteObjectBindings(const FString& TextToImport)
 			// because we need to use our binding (which has tracks associated with it). To solve this, we let it create
 			// an object template based off of our (transient package owned) template, then find the newly created binding
 			// and update it.
-			FGuid NewGuid = MakeNewSpawnable(*CopyableBinding->SpawnableObjectTemplate);
+			FGuid NewGuid = MakeNewSpawnable(*CopyableBinding->SpawnableObjectTemplate, nullptr, false);
 			FMovieSceneBinding NewBinding(NewGuid, CopyableBinding->Binding.GetName(), CopyableBinding->Tracks);
 			FMovieSceneSpawnable* Spawnable = MovieScene->FindSpawnable(NewGuid);
 
diff --git a/Engine/Source/Editor/Sequencer/Private/Sequencer.h b/Engine/Source/Editor/Sequencer/Private/Sequencer.h
--- a/Engine/Source/Editor/Sequencer/Private/Sequencer.h
+++ b/Engine/Source/Editor/Sequencer/Private/Sequencer.h
@@ -726,7 +726,7 @@ public:
 	virtual void Pause() override;
 	virtual TSharedRef<SWidget> MakeTimeRange(const TSharedRef<SWidget>& InnerContent, bool bShowWorkingRange, bool bShowViewRange, bool bShowPlaybackRange) override;
 	virtual UObject* FindSpawnedObjectOrTemplate(const FGuid& BindingId) override;
-	virtual FGuid MakeNewSpawnable(UObject& SourceObject, UActorFactory* ActorFactory = nullptr) override;
+	virtual FGuid MakeNewSpawnable(UObject& SourceObject, UActorFactory* ActorFactory = nullptr, bool bSetupDefaults = true) override;
 	virtual bool IsReadOnly() const override;
 	virtual void ExternalSelectionHasChanged() override { SynchronizeSequencerSelectionWithExternalSelection(); }
 	/** Access the user-supplied settings object */
diff --git a/Engine/Source/Editor/Sequencer/Public/ISequencer.h b/Engine/Source/Editor/Sequencer/Public/ISequencer.h
--- a/Engine/Source/Editor/Sequencer/Public/ISequencer.h
+++ b/Engine/Source/Editor/Sequencer/Public/ISequencer.h
@@ -226,7 +226,7 @@ public:
 	 * @param	ActorFactory	Optional actor factory to use to create spawnable type
 	 * @return	The spawnable guid for the spawnable, or an invalid Guid if we were not able to create a spawnable
 	 */
-	virtual FGuid MakeNewSpawnable(UObject& SourceObject, UActorFactory* ActorFactory = nullptr) = 0;
+	virtual FGuid MakeNewSpawnable(UObject& SourceObject, UActorFactory* ActorFactory = nullptr, bool bSetupDefaults = false) = 0;
 
 	/**
 	 * Add actors as possessable objects to sequencer.
