Deprecate SEditableLabel and remove use cases. It is functionally identical to SInlineEditableTextBlock

diff --git a/Engine/Source/Editor/Sequencer/Private/SAnimationOutlinerTreeNode.cpp b/Engine/Source/Editor/Sequencer/Private/SAnimationOutlinerTreeNode.cpp
--- a/Engine/Source/Editor/Sequencer/Private/SAnimationOutlinerTreeNode.cpp
+++ b/Engine/Source/Editor/Sequencer/Private/SAnimationOutlinerTreeNode.cpp
@@ -16,10 +16,10 @@
 #include "DisplayNodes/SequencerTrackNode.h"
 #include "SSequencer.h"
 #include "ScopedTransaction.h"
-#include "Widgets/Input/SEditableLabel.h"
 #include "SSequencerTreeView.h"
 #include "Widgets/Colors/SColorPicker.h"
 #include "SequencerSectionPainter.h"
+#include "Widgets/Text/SInlineEditableTextBlock.h"
 
 #define LOCTEXT_NAMESPACE "AnimationOutliner"
 
@@ -123,11 +123,11 @@ void SAnimationOutlinerTreeNode::Construct( const FArguments& InArgs, TSharedRef
 
 	FSlateFontInfo NodeFont = FEditorStyle::GetFontStyle("Sequencer.AnimationOutliner.RegularFont");
 
-	EditableLabel = SNew(SEditableLabel)
-		.CanEdit(this, &SAnimationOutlinerTreeNode::HandleNodeLabelCanEdit)
+	EditableLabel = SNew(SInlineEditableTextBlock)
+		.IsReadOnly(this, &SAnimationOutlinerTreeNode::IsNodeLabelReadOnly)
 		.Font(NodeFont)
 		.ColorAndOpacity(this, &SAnimationOutlinerTreeNode::GetDisplayNameColor)
-		.OnTextChanged(this, &SAnimationOutlinerTreeNode::HandleNodeLabelTextChanged)
+		.OnTextCommitted(this, &SAnimationOutlinerTreeNode::HandleNodeLabelTextCommitted)
 		.Text(this, &SAnimationOutlinerTreeNode::GetDisplayName)
 		.ToolTipText(this, &SAnimationOutlinerTreeNode::GetDisplayNameToolTipText)
 		.Clipping(EWidgetClipping::ClipToBounds);
@@ -268,7 +268,7 @@ void SAnimationOutlinerTreeNode::Construct( const FArguments& InArgs, TSharedRef
 
 void SAnimationOutlinerTreeNode::EnterRenameMode()
 {
-	EditableLabel->EnterTextMode();
+	EditableLabel->EnterEditingMode();
 }
 
 
@@ -438,13 +438,13 @@ FText SAnimationOutlinerTreeNode::GetDisplayName() const
 }
 
 
-bool SAnimationOutlinerTreeNode::HandleNodeLabelCanEdit() const
+bool SAnimationOutlinerTreeNode::IsNodeLabelReadOnly() const
 {
-	return !DisplayNode->GetSequencer().IsReadOnly() && DisplayNode->CanRenameNode();
+	return DisplayNode->GetSequencer().IsReadOnly() || !DisplayNode->CanRenameNode();
 }
 
 
-void SAnimationOutlinerTreeNode::HandleNodeLabelTextChanged(const FText& NewLabel)
+void SAnimationOutlinerTreeNode::HandleNodeLabelTextCommitted(const FText& NewLabel, ETextCommit::Type CommitType)
 {
 	DisplayNode->SetDisplayName(NewLabel);
 }
diff --git a/Engine/Source/Editor/Sequencer/Private/SAnimationOutlinerTreeNode.h b/Engine/Source/Editor/Sequencer/Private/SAnimationOutlinerTreeNode.h
--- a/Engine/Source/Editor/Sequencer/Private/SAnimationOutlinerTreeNode.h
+++ b/Engine/Source/Editor/Sequencer/Private/SAnimationOutlinerTreeNode.h
@@ -10,7 +10,7 @@
 #include "Widgets/DeclarativeSyntaxSupport.h"
 #include "Widgets/SCompoundWidget.h"
 
-class SEditableLabel;
+class SInlineEditableTextBlock;
 class SSequencerTreeViewRow;
 struct FSlateBrush;
 struct FTableRowStyle;
@@ -97,10 +97,10 @@ private:
 	FText GetDisplayName() const;
 
 	/** Callback for checking whether the node label can be edited. */
-	bool HandleNodeLabelCanEdit() const;
+	bool IsNodeLabelReadOnly() const;
 
 	/** Callback for when the node label text has changed. */
-	void HandleNodeLabelTextChanged(const FText& NewLabel);
+	void HandleNodeLabelTextCommitted(const FText& NewLabel, ETextCommit::Type CommitType);
 
 	/** Get all descendant nodes from the given root node. */
 	void GetAllDescendantNodes(TSharedPtr<FSequencerDisplayNode> RootNode, TArray<TSharedRef<FSequencerDisplayNode> >& AllNodes);
@@ -114,7 +114,7 @@ private:
 	TSharedPtr<FSequencerDisplayNode> DisplayNode;
 
 	/** Holds the editable text label widget. */
-	TSharedPtr<SEditableLabel> EditableLabel;
+	TSharedPtr<SInlineEditableTextBlock> EditableLabel;
 
 	/** True if this node is a top level node, at the root of the tree, false otherwise */
 	bool bIsOuterTopLevelNode;
diff --git a/Engine/Source/Editor/Sequencer/Private/SSequencerLabelBrowser.cpp b/Engine/Source/Editor/Sequencer/Private/SSequencerLabelBrowser.cpp
--- a/Engine/Source/Editor/Sequencer/Private/SSequencerLabelBrowser.cpp
+++ b/Engine/Source/Editor/Sequencer/Private/SSequencerLabelBrowser.cpp
@@ -287,7 +287,7 @@ bool SSequencerLabelBrowser::HandleRenameLabelMenuEntryCanExecute() const
 
 		if (ListRow.IsValid())
 		{
-			return ListRow->CanEnterRenameMode();
+			return !ListRow->IsReadOnly();
 		}
 	}
 	return false;
diff --git a/Engine/Source/Editor/Sequencer/Private/SSequencerLabelListRow.h b/Engine/Source/Editor/Sequencer/Private/SSequencerLabelListRow.h
--- a/Engine/Source/Editor/Sequencer/Private/SSequencerLabelListRow.h
+++ b/Engine/Source/Editor/Sequencer/Private/SSequencerLabelListRow.h
@@ -13,7 +13,7 @@
 #include "Widgets/Views/SListView.h"
 #include "Widgets/Images/SImage.h"
 #include "EditorStyleSet.h"
-#include "Widgets/Input/SEditableLabel.h"
+#include "Widgets/Text/SInlineEditableTextBlock.h"
 
 struct FSequencerLabelTreeNode;
 
@@ -105,9 +105,9 @@ public:
 						.Padding(0.0f, 2.0f)
 						.VAlign(VAlign_Center)
 						[
-							SAssignNew(EditableLabel, SEditableLabel)
-								.CanEdit(this, &SSequencerLabelListRow::CanEnterRenameMode)
-								.OnTextChanged(this, &SSequencerLabelListRow::HandleFolderNameTextChanged)
+							SAssignNew(EditableLabel, SInlineEditableTextBlock)
+								.IsReadOnly(this, &SSequencerLabelListRow::IsReadOnly)
+								.OnTextCommitted(this, &SSequencerLabelListRow::HandleFolderNameTextChanged)
 								.Text(
 									Node->Label.IsEmpty()
 										? LOCTEXT("AllTracksLabel", "All Tracks")
@@ -122,13 +122,13 @@ public:
 	/** Change the label text to edit mode. */
 	void EnterRenameMode()
 	{
-		EditableLabel->EnterTextMode();
+		EditableLabel->EnterEditingMode();
 	}
 
 	/** Can this label be edited? */
-	bool CanEnterRenameMode() const
+	bool IsReadOnly() const
 	{
-		return !Node->Label.IsEmpty();
+		return Node->Label.IsEmpty();
 	}
 
 private:
@@ -156,7 +156,7 @@ private:
 		return FLinearColor::Gray;
 	}
 
-	void HandleFolderNameTextChanged(const FText& NewLabel)
+	void HandleFolderNameTextChanged(const FText& NewLabel, ETextCommit::Type TextCommitType)
 	{
 		FString NewLabelString = NewLabel.ToString();
 
@@ -169,7 +169,7 @@ private:
 private:
 
 	/** Holds the editable text label widget. */
-	TSharedPtr<SEditableLabel> EditableLabel;
+	TSharedPtr<SInlineEditableTextBlock> EditableLabel;
 
 	/** Holds the label node. */
 	TSharedPtr<FSequencerLabelTreeNode> Node;
diff --git a/Engine/Source/Runtime/Slate/Public/Widgets/Input/SEditableLabel.h b/Engine/Source/Runtime/Slate/Public/Widgets/Input/SEditableLabel.h
--- a/Engine/Source/Runtime/Slate/Public/Widgets/Input/SEditableLabel.h
+++ b/Engine/Source/Runtime/Slate/Public/Widgets/Input/SEditableLabel.h
@@ -17,7 +17,8 @@
 class SEditableText;
 class STextBlock;
 
-class SLATE_API SEditableLabel
+
+class DEPRECATED(4.21, "SEditableLabel is deprecated, please use SInlineEditableTextBlock instead") SLATE_API SEditableLabel
 	: public SCompoundWidget
 {
 public:
