Sequencer: Add transaction to track renaming

diff --git a/Engine/Source/Editor/Sequencer/Private/DisplayNodes/SequencerTrackNode.cpp b/Engine/Source/Editor/Sequencer/Private/DisplayNodes/SequencerTrackNode.cpp
--- a/Engine/Source/Editor/Sequencer/Private/DisplayNodes/SequencerTrackNode.cpp
+++ b/Engine/Source/Editor/Sequencer/Private/DisplayNodes/SequencerTrackNode.cpp
@@ -510,8 +510,10 @@ void FSequencerTrackNode::SetDisplayName(const FText& NewDisplayName)
 {
 	auto NameableTrack = Cast<UMovieSceneNameableTrack>(AssociatedTrack.Get());
 
-	if (NameableTrack != nullptr)
+	if (NameableTrack != nullptr && !NameableTrack->GetDisplayName().EqualTo(NewDisplayName))
 	{
+		const FScopedTransaction Transaction(NSLOCTEXT("SequencerTrackNode", "RenameTrack", "Rename Track"));
+
 		NameableTrack->SetDisplayName(NewDisplayName);
 		GetSequencer().NotifyMovieSceneDataChanged(EMovieSceneDataChangeType::TrackValueChanged);
 	}
