Sequencer: Undo/redo support for label browser

diff --git a/Engine/Source/Editor/Sequencer/Private/SSequencerLabelBrowser.cpp b/Engine/Source/Editor/Sequencer/Private/SSequencerLabelBrowser.cpp
--- a/Engine/Source/Editor/Sequencer/Private/SSequencerLabelBrowser.cpp
+++ b/Engine/Source/Editor/Sequencer/Private/SSequencerLabelBrowser.cpp
@@ -4,7 +4,7 @@
 #include "SlateOptMacros.h"
 #include "Framework/MultiBox/MultiBoxBuilder.h"
 #include "SSequencerLabelListRow.h"
-
+#include "ScopedTransaction.h"
 
 #define LOCTEXT_NAMESPACE "SSequencerLabelBrowser"
 
@@ -147,6 +147,8 @@ void SSequencerLabelBrowser::ReloadLabelList(bool FullyReload)
 
 void SSequencerLabelBrowser::HandleLabelListRowLabelRenamed(TSharedPtr<FSequencerLabelTreeNode> Node, const FString& NewLabel)
 {
+	const FScopedTransaction Transaction(LOCTEXT("RenameLabel", "Rename Label"));
+
 	if (Sequencer.IsValid() && Sequencer.Pin()->GetLabelManager().RenameLabel(Node->Label, NewLabel))
 	{
 		ReloadLabelList(true);
@@ -244,6 +246,8 @@ void SSequencerLabelBrowser::HandleRemoveLabelMenuEntryExecute()
 
 	if (Sequencer.IsValid() && LabelTreeView->GetSelectedItems(SelectedItems) > 0)
 	{
+		const FScopedTransaction Transaction(LOCTEXT("RemoveLabel", "Remove Label"));
+
 		Sequencer.Pin()->GetLabelManager().RemoveObjectLabel(FGuid(), SelectedItems[0]->Label);
 	}
 }
diff --git a/Engine/Source/Editor/Sequencer/Private/SSequencerLabelEditor.cpp b/Engine/Source/Editor/Sequencer/Private/SSequencerLabelEditor.cpp
--- a/Engine/Source/Editor/Sequencer/Private/SSequencerLabelEditor.cpp
+++ b/Engine/Source/Editor/Sequencer/Private/SSequencerLabelEditor.cpp
@@ -8,7 +8,7 @@
 #include "Sequencer.h"
 #include "Widgets/Views/SListView.h"
 #include "SSequencerLabelEditorListRow.h"
-
+#include "ScopedTransaction.h"
 
 #define LOCTEXT_NAMESPACE "SSequencerLabelEditor"
 
@@ -86,6 +86,8 @@ void SSequencerLabelEditor::Construct(const FArguments& InArgs, FSequencer& InSe
 
 void SSequencerLabelEditor::CreateLabelFromFilterText()
 {
+	const FScopedTransaction Transaction(LOCTEXT("AddLabel", "Add Label"));
+
 	for (auto ObjectId : ObjectIds)
 	{
 		Sequencer->GetLabelManager().AddObjectLabel(ObjectId, FilterBox->GetText().ToString());
diff --git a/Engine/Source/Editor/Sequencer/Private/SequencerLabelManager.cpp b/Engine/Source/Editor/Sequencer/Private/SequencerLabelManager.cpp
--- a/Engine/Source/Editor/Sequencer/Private/SequencerLabelManager.cpp
+++ b/Engine/Source/Editor/Sequencer/Private/SequencerLabelManager.cpp
@@ -9,8 +9,10 @@ void FSequencerLabelManager::SetMovieScene(UMovieScene* InMovieScene)
 	if (MovieScene != InMovieScene)
 	{
 		MovieScene = InMovieScene;
-		LabelsChangedEvent.Broadcast();
 	}
+
+	// Broadcast the labels changed regardless of the movie scene changed since the label data could have changed
+	LabelsChangedEvent.Broadcast();
 }
 
 
@@ -18,6 +20,7 @@ void FSequencerLabelManager::AddObjectLabel(const FGuid& ObjectId, const FString
 {
 	if (MovieScene.IsValid())
 	{
+		MovieScene->Modify();
 		MovieScene->GetObjectsToLabels().FindOrAdd(ObjectId.ToString()).Strings.AddUnique(Label);
 		MovieScene->MarkPackageDirty();
 		LabelsChangedEvent.Broadcast();
@@ -45,6 +48,8 @@ void FSequencerLabelManager::RemoveObjectLabel(const FGuid& ObjectId, const FStr
 {
 	if (MovieScene.IsValid())
 	{
+		MovieScene->Modify();
+
 		TMap<FString, FMovieSceneTrackLabels>& ObjectsToLabels = MovieScene->GetObjectsToLabels();
 
 		if (ObjectId.IsValid())
@@ -115,6 +120,8 @@ bool FSequencerLabelManager::RenameLabel(const FString& OldLabel, const FString&
 
 	bool FoundOldLabel = false;
 
+	MovieScene->Modify();
+
 	for (auto& LabelsPair : MovieScene->GetObjectsToLabels())
 	{
 		TArray<FString>& Strings = LabelsPair.Value.Strings;
