Sequencer: Fix selection in/out undo

diff --git a/Engine/Source/Editor/Sequencer/Private/Sequencer.cpp b/Engine/Source/Editor/Sequencer/Private/Sequencer.cpp
--- a/Engine/Source/Editor/Sequencer/Private/Sequencer.cpp
+++ b/Engine/Source/Editor/Sequencer/Private/Sequencer.cpp
@@ -1775,9 +1775,10 @@ TRange<FFrameNumber> FSequencer::GetSelectionRange() const
 void FSequencer::SetSelectionRange(TRange<FFrameNumber> Range)
 {
 	const FScopedTransaction Transaction(LOCTEXT("SetSelectionRange_Transaction", "Set Selection Range"));
-	UMovieScene* FocussedMovieScene = GetFocusedMovieSceneSequence()->GetMovieScene();
+	UMovieScene* FocusedMovieScene = GetFocusedMovieSceneSequence()->GetMovieScene();
 
-	FocussedMovieScene->SetSelectionRange(Range);
+	FocusedMovieScene->Modify();
+	FocusedMovieScene->SetSelectionRange(Range);
 }
 
 
@@ -6271,8 +6272,6 @@ bool FSequencer::PasteObjectBindings(const FString& TextToImport)
 				}
 			}
 
-				// Replace the auto-generated binding with our deserialized bindings (which has our tracks)
-				MovieScene->ReplaceBinding(NewGuid, NewBinding);
 
 			OldToNewGuidMap.Add(CopyableBinding->Spawnable.GetGuid(), NewGuid);
 
