Sequencer: Remove child labels

diff --git a/Engine/Source/Editor/Sequencer/Private/SSequencerLabelBrowser.cpp b/Engine/Source/Editor/Sequencer/Private/SSequencerLabelBrowser.cpp
--- a/Engine/Source/Editor/Sequencer/Private/SSequencerLabelBrowser.cpp
+++ b/Engine/Source/Editor/Sequencer/Private/SSequencerLabelBrowser.cpp
@@ -239,6 +239,19 @@ void SSequencerLabelBrowser::HandleLabelTreeViewSelectionChanged(TSharedPtr<FSeq
 	);
 }
 
+void SSequencerLabelBrowser::RemoveLabelItem(TSharedPtr<FSequencerLabelTreeNode> InItem)
+{
+	if (InItem.IsValid())
+	{
+		for (auto Child : InItem->Children)
+		{
+			RemoveLabelItem(Child);
+		}
+
+		Sequencer.Pin()->GetLabelManager().RemoveObjectLabel(FGuid(), InItem->Label);
+	}
+}
+
 
 void SSequencerLabelBrowser::HandleRemoveLabelMenuEntryExecute()
 {
@@ -248,7 +261,7 @@ void SSequencerLabelBrowser::HandleRemoveLabelMenuEntryExecute()
 	{
 		const FScopedTransaction Transaction(LOCTEXT("RemoveLabel", "Remove Label"));
 
-		Sequencer.Pin()->GetLabelManager().RemoveObjectLabel(FGuid(), SelectedItems[0]->Label);
+		RemoveLabelItem(SelectedItems[0]);
 	}
 }
 
diff --git a/Engine/Source/Editor/Sequencer/Private/SSequencerLabelBrowser.h b/Engine/Source/Editor/Sequencer/Private/SSequencerLabelBrowser.h
--- a/Engine/Source/Editor/Sequencer/Private/SSequencerLabelBrowser.h
+++ b/Engine/Source/Editor/Sequencer/Private/SSequencerLabelBrowser.h
@@ -89,6 +89,9 @@ private:
 	/** Callback for checking whether 'Rename' context menu entry can execute. */
 	bool HandleRenameLabelMenuEntryCanExecute() const;
 
+	/** Remove label item and its children */
+	void RemoveLabelItem(TSharedPtr<FSequencerLabelTreeNode> InItem);
+
 private:
 
 	/** Holds the collection of root labels to be displayed in the tree view. */
