Sequencer: Add support for filtering child labels from the parent

diff --git a/Engine/Source/Editor/Sequencer/Private/SSequencerLabelBrowser.cpp b/Engine/Source/Editor/Sequencer/Private/SSequencerLabelBrowser.cpp
--- a/Engine/Source/Editor/Sequencer/Private/SSequencerLabelBrowser.cpp
+++ b/Engine/Source/Editor/Sequencer/Private/SSequencerLabelBrowser.cpp
@@ -211,6 +211,26 @@ void SSequencerLabelBrowser::HandleLabelTreeViewGetChildren(TSharedPtr<FSequence
 	}
 }
 
+void SSequencerLabelBrowser::AddLabelNode(TSharedPtr<FSequencerLabelTreeNode> InItem, FString& NewLabel)
+{
+	if (InItem.IsValid())
+	{
+		if (!NewLabel.IsEmpty())
+		{
+			NewLabel += TEXT(" ");
+		}
+
+		if (!InItem->Label.IsEmpty())
+		{
+			NewLabel += FString(TEXT("label:")) + InItem->Label;
+		}
+
+		for (auto Child : InItem->Children)
+		{
+			AddLabelNode(Child, NewLabel);
+		}
+	}
+}
 
 void SSequencerLabelBrowser::HandleLabelTreeViewSelectionChanged(TSharedPtr<FSequencerLabelTreeNode> InItem, ESelectInfo::Type SelectInfo)
 {
@@ -220,15 +240,7 @@ void SSequencerLabelBrowser::HandleLabelTreeViewSelectionChanged(TSharedPtr<FSeq
 	LabelTreeView->GetSelectedItems(SelectedItems);
 	for (TSharedPtr<FSequencerLabelTreeNode> Item : SelectedItems)
 	{
-		if (!NewLabel.IsEmpty())
-		{
-			NewLabel += TEXT(" ");
-		}
-		
-		if (!Item->Label.IsEmpty())
-		{
-			NewLabel += FString(TEXT("label:")) + Item->Label;
-		}
+		AddLabelNode(Item, NewLabel);
 	}
 
 	OnSelectionChanged.ExecuteIfBound(
diff --git a/Engine/Source/Editor/Sequencer/Private/SSequencerLabelBrowser.h b/Engine/Source/Editor/Sequencer/Private/SSequencerLabelBrowser.h
--- a/Engine/Source/Editor/Sequencer/Private/SSequencerLabelBrowser.h
+++ b/Engine/Source/Editor/Sequencer/Private/SSequencerLabelBrowser.h
@@ -92,6 +92,9 @@ private:
 	/** Remove label item and its children */
 	void RemoveLabelItem(TSharedPtr<FSequencerLabelTreeNode> InItem);
 
+	/** Add node's label and its children to the new label */
+	void AddLabelNode(TSharedPtr<FSequencerLabelTreeNode> InItem, FString& NewLabel);
+
 private:
 
 	/** Holds the collection of root labels to be displayed in the tree view. */
