Sequencer: Force refresh when sequencer is opened so that keyframe buttons appear.

diff --git a/Engine/Source/Editor/Sequencer/Private/LevelEditorSequencerIntegration.cpp b/Engine/Source/Editor/Sequencer/Private/LevelEditorSequencerIntegration.cpp
--- a/Engine/Source/Editor/Sequencer/Private/LevelEditorSequencerIntegration.cpp
+++ b/Engine/Source/Editor/Sequencer/Private/LevelEditorSequencerIntegration.cpp
@@ -725,6 +725,9 @@ void FLevelEditorSequencerIntegration::ActivateDetailHandler()
 
 	AcquiredResources.Add(DetailHandlerName, DeactivateDetailKeyframeHandler);
 	AcquiredResources.Add(DetailHandlerRefreshName, RefreshDetailHandler);
+
+	const bool bForceRefresh = true;
+	UpdateDetails(bForceRefresh);
 }
 
 void FLevelEditorSequencerIntegration::OnPropertyEditorOpened()
