Sequencer: Add error checking for invalid labels with spaces

diff --git a/Engine/Source/Editor/Sequencer/Private/SSequencerLabelEditor.cpp b/Engine/Source/Editor/Sequencer/Private/SSequencerLabelEditor.cpp
--- a/Engine/Source/Editor/Sequencer/Private/SSequencerLabelEditor.cpp
+++ b/Engine/Source/Editor/Sequencer/Private/SSequencerLabelEditor.cpp
@@ -159,8 +159,11 @@ FReply SSequencerLabelEditor::HandleFilterBoxKeyDown(const FGeometry& /*Geometry
 {
 	if (KeyEvent.GetKey() == EKeys::Enter)
 	{
-		CreateLabelFromFilterText();
-		return FReply::Handled();
+		if (HandleCreateNewLabelButtonIsEnabled())
+		{
+			CreateLabelFromFilterText();
+			return FReply::Handled();
+		}
 	}
 
 	return FReply::Unhandled();
diff --git a/Engine/Source/Editor/Sequencer/Private/SSequencerLabelListRow.h b/Engine/Source/Editor/Sequencer/Private/SSequencerLabelListRow.h
--- a/Engine/Source/Editor/Sequencer/Private/SSequencerLabelListRow.h
+++ b/Engine/Source/Editor/Sequencer/Private/SSequencerLabelListRow.h
@@ -107,6 +107,7 @@ public:
 						[
 							SAssignNew(EditableLabel, SInlineEditableTextBlock)
 								.IsReadOnly(this, &SSequencerLabelListRow::IsReadOnly)
+								.OnVerifyTextChanged(this, &SSequencerLabelListRow::HandleVerifyNameChanged)
 								.OnTextCommitted(this, &SSequencerLabelListRow::HandleFolderNameTextChanged)
 								.Text(
 									Node->Label.IsEmpty()
@@ -156,6 +157,19 @@ private:
 		return FLinearColor::Gray;
 	}
 
+
+	bool HandleVerifyNameChanged(const FText& NewText, FText& OutErrorMessage)
+	{
+		if (NewText.ToString().Contains(TEXT(" ")))	
+		{
+			OutErrorMessage = LOCTEXT("InvalidLabel", "Label cannot contains spaces");
+			return false;
+		}
+
+		return true;
+	}
+
+
 	void HandleFolderNameTextChanged(const FText& NewLabel, ETextCommit::Type TextCommitType)
 	{
 		FString NewLabelString = NewLabel.ToString();
