Sequencer: When inserting a section into a particular row, move existing sections down by one row.

diff --git a/Engine/Source/Editor/Sequencer/Private/SequencerUtilities.cpp b/Engine/Source/Editor/Sequencer/Private/SequencerUtilities.cpp
--- a/Engine/Source/Editor/Sequencer/Private/SequencerUtilities.cpp
+++ b/Engine/Source/Editor/Sequencer/Private/SequencerUtilities.cpp
@@ -108,6 +108,12 @@ void FSequencerUtilities::PopulateMenu_CreateNewSection(FMenuBuilder& MenuBuilde
 			for (UMovieSceneSection* Section : Track->GetAllSections())
 			{
 				OverlapPriority = FMath::Max(Section->GetOverlapPriority() + 1, OverlapPriority);
+
+				// Move existing sections on the same row or beyond so that they don't overlap with the new section
+				if (Section != NewSection && Section->GetRowIndex() >= RowIndex)
+				{
+					Section->SetRowIndex(Section->GetRowIndex() + 1);
+				}
 			}
 
 			Track->Modify();
