Sequencer: Allow setting/scrubbing the playback range beyond the view range. When out of range, also extend the view range by 10 percent.

diff --git a/Engine/Source/Editor/Sequencer/Private/SSequencer.cpp b/Engine/Source/Editor/Sequencer/Private/SSequencer.cpp
--- a/Engine/Source/Editor/Sequencer/Private/SSequencer.cpp
+++ b/Engine/Source/Editor/Sequencer/Private/SSequencer.cpp
@@ -1136,14 +1136,34 @@ TSharedRef<SWidget> SSequencer::MakePlaybackMenu()
 		auto OnStartChanged = [=](double NewValue){
 
 			// We clamp the new value when the value is set. We can't clamp in the UI because we need an unset Min/Max for linear scaling to work.
-			double Min = (SequencerPtr.Pin()->GetClampRange().GetLowerBoundValue() * SequencerPtr.Pin()->GetFocusedTickResolution()).GetFrame().Value;
-			double Max = SequencerPtr.Pin()->GetPlaybackRange().GetUpperBoundValue().Value; 
+			double Min = -FLT_MAX;
+			double Max = SequencerPtr.Pin()->GetPlaybackRange().GetUpperBoundValue().Value;
 
 			NewValue = FMath::Clamp(NewValue, Min, Max);
 			FFrameNumber ValueAsFrame = FFrameTime::FromDecimal(NewValue).GetFrame();
 
 			FFrameNumber Upper = MovieScene::DiscreteExclusiveUpper(SequencerPtr.Pin()->GetPlaybackRange());
-			SequencerPtr.Pin()->SetPlaybackRange(TRange<FFrameNumber>(FMath::Min(ValueAsFrame, Upper-1), Upper));
+
+			TRange<FFrameNumber> NewRange = TRange<FFrameNumber>(FMath::Min(ValueAsFrame, Upper - 1), Upper);
+
+			SequencerPtr.Pin()->SetPlaybackRange(NewRange);
+
+			TRange<double> PlayRangeSeconds = SequencerPtr.Pin()->GetPlaybackRange() / SequencerPtr.Pin()->GetFocusedTickResolution();
+			const double AdditionalRange = (PlayRangeSeconds.GetUpperBoundValue() - PlayRangeSeconds.GetLowerBoundValue()) * 0.1;
+
+			TRange<double> NewClampRange = SequencerPtr.Pin()->GetClampRange();
+			NewClampRange.SetLowerBoundValue(SequencerPtr.Pin()->GetPlaybackRange().GetLowerBoundValue() / SequencerPtr.Pin()->GetFocusedTickResolution() - AdditionalRange);
+			if (SequencerPtr.Pin()->GetClampRange().GetLowerBoundValue() > NewClampRange.GetLowerBoundValue())
+			{
+				SequencerPtr.Pin()->SetClampRange(NewClampRange);
+			}
+
+			TRange<double> NewViewRange = SequencerPtr.Pin()->GetViewRange();
+			NewViewRange.SetLowerBoundValue(SequencerPtr.Pin()->GetPlaybackRange().GetLowerBoundValue() / SequencerPtr.Pin()->GetFocusedTickResolution() - AdditionalRange);
+			if (SequencerPtr.Pin()->GetViewRange().GetLowerBoundValue() > NewViewRange.GetLowerBoundValue())
+			{
+				SequencerPtr.Pin()->SetViewRange(NewViewRange);
+			}
 		};
 
 		MenuBuilder.AddWidget(
@@ -1180,13 +1200,30 @@ TSharedRef<SWidget> SSequencer::MakePlaybackMenu()
 
 			// We clamp the new value when the value is set. We can't clamp in the UI because we need an unset Min/Max for linear scaling to work.
 			double Min = SequencerPtr.Pin()->GetPlaybackRange().GetLowerBoundValue().Value;
-			double Max = (SequencerPtr.Pin()->GetClampRange().GetUpperBoundValue() * SequencerPtr.Pin()->GetFocusedTickResolution()).GetFrame().Value;
+			double Max = FLT_MAX;
 
 			NewValue = FMath::Clamp(NewValue, Min, Max);
 			FFrameNumber ValueAsFrame = FFrameTime::FromDecimal(NewValue).GetFrame();
 
 			FFrameNumber Lower = MovieScene::DiscreteInclusiveLower(SequencerPtr.Pin()->GetPlaybackRange());
 			SequencerPtr.Pin()->SetPlaybackRange(TRange<FFrameNumber>(Lower, FMath::Max(ValueAsFrame, Lower)));
+
+			TRange<double> PlayRangeSeconds = SequencerPtr.Pin()->GetPlaybackRange() / SequencerPtr.Pin()->GetFocusedTickResolution();
+			const double AdditionalRange = (PlayRangeSeconds.GetUpperBoundValue() - PlayRangeSeconds.GetLowerBoundValue()) * 0.1;
+
+			TRange<double> NewClampRange = SequencerPtr.Pin()->GetClampRange();
+			NewClampRange.SetUpperBoundValue(SequencerPtr.Pin()->GetPlaybackRange().GetUpperBoundValue() / SequencerPtr.Pin()->GetFocusedTickResolution() + AdditionalRange);
+			if (SequencerPtr.Pin()->GetClampRange().GetUpperBoundValue() < NewClampRange.GetUpperBoundValue())
+			{
+				SequencerPtr.Pin()->SetClampRange(NewClampRange);
+			}
+
+			TRange<double> NewViewRange = SequencerPtr.Pin()->GetViewRange();
+			NewViewRange.SetUpperBoundValue(SequencerPtr.Pin()->GetPlaybackRange().GetUpperBoundValue() / SequencerPtr.Pin()->GetFocusedTickResolution() + AdditionalRange);
+			if (SequencerPtr.Pin()->GetViewRange().GetUpperBoundValue() < NewViewRange.GetUpperBoundValue())
+			{
+				SequencerPtr.Pin()->SetViewRange(NewViewRange);
+			}
 		};
 
 		MenuBuilder.AddWidget(
