Sequencer: Skip selected track if its ancestor is also selected

diff --git a/Engine/Source/Editor/Sequencer/Private/Sequencer.cpp b/Engine/Source/Editor/Sequencer/Private/Sequencer.cpp
--- a/Engine/Source/Editor/Sequencer/Private/Sequencer.cpp
+++ b/Engine/Source/Editor/Sequencer/Private/Sequencer.cpp
@@ -6258,7 +6258,22 @@ void FSequencer::CopySelectedTracks(TArray<TSharedPtr<FSequencerTrackNode>>& Tra
 	TArray<UObject*> TracksToCopy;
 	for (TSharedPtr<FSequencerTrackNode> TrackNode : TrackNodes)
 	{
-		TracksToCopy.Add(TrackNode->GetTrack());
+		bool bIsParentSelected = false;
+		TSharedPtr<FSequencerDisplayNode> ParentNode = TrackNode->GetParent();
+		while (ParentNode.IsValid())
+		{
+			if (Selection.GetSelectedOutlinerNodes().Contains(ParentNode.ToSharedRef()))
+			{
+				bIsParentSelected = true;
+				break;
+			}
+			ParentNode = ParentNode->GetParent();
+		}
+
+		if (!bIsParentSelected)
+		{
+			TracksToCopy.Add(TrackNode->GetTrack());
+		}
 	}
 
 	FString ExportedText;
