Sequencer: Paste to selected folder

diff --git a/Engine/Source/Editor/Sequencer/Private/Sequencer.cpp b/Engine/Source/Editor/Sequencer/Private/Sequencer.cpp
--- a/Engine/Source/Editor/Sequencer/Private/Sequencer.cpp
+++ b/Engine/Source/Editor/Sequencer/Private/Sequencer.cpp
@@ -6340,6 +6340,10 @@ bool FSequencer::PasteObjectBindings(const FString& TextToImport)
 		return false;
 	}
 	
+	TArray<UMovieSceneFolder*> SelectedParentFolders;
+	FString NewNodePath;
+	CalculateSelectedFolderAndPath(SelectedParentFolders, NewNodePath);
+
 	FScopedTransaction Transaction(FGenericCommands::Get().Paste->GetDescription());
 	
 	UMovieSceneSequence* OwnerSequence = GetFocusedMovieSceneSequence();
@@ -6426,9 +6430,9 @@ bool FSequencer::PasteObjectBindings(const FString& TextToImport)
 	}
 
 	// Fix possessable actor bindings
-	for (auto PossessableGuid : PossessableGuids)
+	for (int32 PossessableGuidIndex = 0; PossessableGuidIndex < PossessableGuids.Num(); ++PossessableGuidIndex)
 	{
-		FMovieScenePossessable* Possessable = MovieScene->FindPossessable(PossessableGuid);
+		FMovieScenePossessable* Possessable = MovieScene->FindPossessable(PossessableGuids[PossessableGuidIndex]);
 		UWorld* PlaybackContext = Cast<UWorld>(GetPlaybackContext());
 		if (Possessable && PlaybackContext)
 		{
@@ -6441,13 +6445,35 @@ bool FSequencer::PasteObjectBindings(const FString& TextToImport)
 
 					if (!ExistingGuid.IsValid())
 					{
-						DoAssignActor(&Actor, 1, Possessable->GetGuid());
+						FGuid NewGuid = DoAssignActor(&Actor, 1, Possessable->GetGuid());
+
+						// If assigning produces a new guid, update the possesable guids and the bindings pasted data
+						if (NewGuid.IsValid())
+						{
+							for (auto BindingPasted : BindingsPasted)
+							{
+								if (BindingPasted.GetObjectGuid() == PossessableGuids[PossessableGuidIndex])
+								{
+									BindingPasted.SetObjectGuid(NewGuid);
+								}
+							}
+
+							PossessableGuids[PossessableGuidIndex] = NewGuid;
+						}
 					}
 				}
 			}
 		}
 	}
 
+	if (SelectedParentFolders.Num() > 0)
+	{
+		for (auto PossessableGuid : PossessableGuids)
+		{
+			SelectedParentFolders[0]->AddChildObjectBinding(PossessableGuid);
+		}
+	}
+
 	OnMovieSceneBindingsPastedDelegate.Broadcast(BindingsPasted);
 
 	// Refresh all immediately so that spawned actors will be generated immediately
@@ -6494,6 +6520,10 @@ bool FSequencer::PasteTracks(const FString& TextToImport)
 		return false;
 	}
 	
+	TArray<UMovieSceneFolder*> SelectedParentFolders;
+	FString NewNodePath;
+	CalculateSelectedFolderAndPath(SelectedParentFolders, NewNodePath);
+
 	FScopedTransaction Transaction(FGenericCommands::Get().Paste->GetDescription());
 
 	UMovieSceneSequence* OwnerSequence = GetFocusedMovieSceneSequence();
@@ -6556,12 +6586,22 @@ bool FSequencer::PasteTracks(const FString& TextToImport)
 		if (NewTrack->IsA(UMovieSceneCameraCutTrack::StaticClass()))
 		{
 			GetFocusedMovieSceneSequence()->GetMovieScene()->SetCameraCutTrack(NewTrack);
+			if (SelectedParentFolders.Num() > 0)
+			{
+				SelectedParentFolders[0]->AddChildMasterTrack(NewTrack);
+			}
+
 			NotifyMovieSceneDataChanged(EMovieSceneDataChangeType::MovieSceneStructureItemsChanged);
 		}
 		else
 		{
 			if (GetFocusedMovieSceneSequence()->GetMovieScene()->AddGivenMasterTrack(NewTrack))
 			{
+				if (SelectedParentFolders.Num() > 0)
+				{
+					SelectedParentFolders[0]->AddChildMasterTrack(NewTrack);
+				}
+
 				NotifyMovieSceneDataChanged(EMovieSceneDataChangeType::MovieSceneStructureItemsChanged);
 			}
 		}
