Sequencer: Copy object bindings and tracks together

diff --git a/Engine/Source/Editor/Sequencer/Private/Sequencer.cpp b/Engine/Source/Editor/Sequencer/Private/Sequencer.cpp
--- a/Engine/Source/Editor/Sequencer/Private/Sequencer.cpp
+++ b/Engine/Source/Editor/Sequencer/Private/Sequencer.cpp
@@ -6183,7 +6183,7 @@ void FSequencer::ImportObjectBindingsFromText(const FString& TextToImport, /*out
 }
 
 
-void FSequencer::CopySelectedObjects(TArray<TSharedPtr<FSequencerObjectBindingNode>>& ObjectNodes)
+void FSequencer::CopySelectedObjects(TArray<TSharedPtr<FSequencerObjectBindingNode>>& ObjectNodes, FString& ExportedText)
 {
 	// Gather guids for the object nodes and any child object nodes
 	TSet<FGuid> GuidsToCopy;
@@ -6247,13 +6247,11 @@ void FSequencer::CopySelectedObjects(TArray<TSharedPtr<FSequencerObjectBindingNo
 		}
 	}
 
-	FString ExportedText;
 	ExportObjectBindingsToText(CopyableBindings, /*out*/ ExportedText);
-	FPlatformApplicationMisc::ClipboardCopy(*ExportedText);
 }
 
 
-void FSequencer::CopySelectedTracks(TArray<TSharedPtr<FSequencerTrackNode>>& TrackNodes)
+void FSequencer::CopySelectedTracks(TArray<TSharedPtr<FSequencerTrackNode>>& TrackNodes, FString& ExportedText)
 {
 	TArray<UObject*> TracksToCopy;
 	for (TSharedPtr<FSequencerTrackNode> TrackNode : TrackNodes)
@@ -6276,9 +6274,7 @@ void FSequencer::CopySelectedTracks(TArray<TSharedPtr<FSequencerTrackNode>>& Tra
 		}
 	}
 
-	FString ExportedText;
 	FSequencer::ExportObjectsToText(TracksToCopy, /*out*/ ExportedText);
-	FPlatformApplicationMisc::ClipboardCopy(*ExportedText);
 }
 
 
@@ -6316,18 +6312,9 @@ void FSequencer::DoPaste()
 	FString TextToImport;
 	FPlatformApplicationMisc::ClipboardPaste(TextToImport);
 
-	if (PasteObjectBindings(TextToImport))
-	{
-		return;
-	}
-	else if (PasteTracks(TextToImport))
-	{
-		return;
-	}
-	else if (PasteSections(TextToImport))
-	{
-		return;
-	}
+	PasteObjectBindings(TextToImport);
+	PasteTracks(TextToImport);
+	PasteSections(TextToImport);
 }
 
 bool FSequencer::PasteObjectBindings(const FString& TextToImport)
@@ -8181,14 +8168,24 @@ void FSequencer::CopySelection()
 			GClipboardStack.Push(NullClipboard);
 		}
 
+		FString ObjectsExportedText;
+		FString TracksExportedText;
+
 		if (ObjectsToCopy.Num())
 		{
-			CopySelectedObjects(ObjectsToCopy);
+			CopySelectedObjects(ObjectsToCopy, ObjectsExportedText);
 		}
-		else
+
+		if (TracksToCopy.Num())
 		{
-			CopySelectedTracks(TracksToCopy);
+			CopySelectedTracks(TracksToCopy, TracksExportedText);
 		}
+
+		FString ExportedText;
+		ExportedText += ObjectsExportedText;
+		ExportedText += TracksExportedText;
+
+		FPlatformApplicationMisc::ClipboardCopy(*ExportedText);
 	}
 }
 
diff --git a/Engine/Source/Editor/Sequencer/Private/Sequencer.h b/Engine/Source/Editor/Sequencer/Private/Sequencer.h
--- a/Engine/Source/Editor/Sequencer/Private/Sequencer.h
+++ b/Engine/Source/Editor/Sequencer/Private/Sequencer.h
@@ -473,8 +473,8 @@ public:
 	void DeleteSelectedNodes();
 
 	/** Called when a user executes the copy track menu item */
-	void CopySelectedObjects(TArray<TSharedPtr<FSequencerObjectBindingNode>>& ObjectNodes);
-	void CopySelectedTracks(TArray<TSharedPtr<FSequencerTrackNode>>& TrackNodes);
+	void CopySelectedObjects(TArray<TSharedPtr<FSequencerObjectBindingNode>>& ObjectNodes, /*out*/ FString& ExportedText);
+	void CopySelectedTracks(TArray<TSharedPtr<FSequencerTrackNode>>& TrackNodes, /*out*/ FString& ExportedText);
 	void ExportObjectsToText(TArray<UObject*> ObjectsToExport, /*out*/ FString& ExportedText);
 
 	/** Called when a user executes the paste track menu item */
