Sequencer: FBX Export now exports selected objects if there are tracks selected (not just when there are object bindings selected)

diff --git a/Engine/Source/Editor/Sequencer/Private/Sequencer.cpp b/Engine/Source/Editor/Sequencer/Private/Sequencer.cpp
--- a/Engine/Source/Editor/Sequencer/Private/Sequencer.cpp
+++ b/Engine/Source/Editor/Sequencer/Private/Sequencer.cpp
@@ -8855,7 +8855,7 @@ void FSequencer::ExportFBXInternal(const FString& ExportFilename, TArray<FGuid>&
 			Exporter->SetTrasformBaking(false);
 			Exporter->SetKeepHierarchy(true);
 
-			const bool bSelectedOnly = Bindings.Num() != 0;
+			const bool bSelectedOnly = Selection.GetSelectedTracks().Num() != 0;
 
 			UnFbx::FFbxExporter::FLevelSequenceNodeNameAdapter NodeNameAdapter(MovieScene, this, GetFocusedTemplateID());
 
