Sequencer: Fix copy/pasting UMovieSceneTrack so that master tracks are only pasted to the root.

diff --git a/Engine/Source/Editor/Sequencer/Private/MovieSceneCopyableTrack.h b/Engine/Source/Editor/Sequencer/Private/MovieSceneCopyableTrack.h
new file mode 100644
--- /dev/null
+++ b/Engine/Source/Editor/Sequencer/Private/MovieSceneCopyableTrack.h
@@ -0,0 +1,22 @@
+// Copyright 1998-2019 Epic Games, Inc. All Rights Reserved.
+
+#pragma once
+
+#include "CoreMinimal.h"
+#include "UObject/ObjectMacros.h"
+#include "MovieSceneCopyableTrack.generated.h"
+
+class UMovieSceneTrack;
+
+UCLASS(Transient)
+class UMovieSceneCopyableTrack : public UObject
+{
+public:
+	GENERATED_BODY()
+
+	UPROPERTY()
+	UMovieSceneTrack* Track;
+
+	UPROPERTY()
+	bool bIsAMasterTrack;
+};
diff --git a/Engine/Source/Editor/Sequencer/Private/SSequencer.cpp b/Engine/Source/Editor/Sequencer/Private/SSequencer.cpp
--- a/Engine/Source/Editor/Sequencer/Private/SSequencer.cpp
+++ b/Engine/Source/Editor/Sequencer/Private/SSequencer.cpp
@@ -75,6 +75,7 @@
 #include "LevelSequence.h"
 #include "SequencerLog.h"
 #include "MovieSceneCopyableBinding.h"
+#include "MovieSceneCopyableTrack.h"
 
 #define LOCTEXT_NAMESPACE "Sequencer"
 
@@ -2381,7 +2382,7 @@ bool SSequencer::CanPaste()
 	// Attempts to deserialize the text into object bindings/tracks that Sequencer understands.
 	if (Sequencer->CanPaste(TextToImport))
 	{
-		TArray<UMovieSceneTrack*> ImportedTracks;
+		TArray<UMovieSceneCopyableTrack*> ImportedTracks;
 		TArray<UMovieSceneSection*> ImportedSections;
 		TArray<UMovieSceneCopyableBinding*> ImportedObjects;
 		Sequencer->ImportTracksFromText(TextToImport, ImportedTracks);
diff --git a/Engine/Source/Editor/Sequencer/Private/Sequencer.cpp b/Engine/Source/Editor/Sequencer/Private/Sequencer.cpp
--- a/Engine/Source/Editor/Sequencer/Private/Sequencer.cpp
+++ b/Engine/Source/Editor/Sequencer/Private/Sequencer.cpp
@@ -112,6 +112,7 @@
 #include "Compilation/MovieSceneCompiler.h"
 #include "SequencerKeyActor.h"
 #include "MovieSceneCopyableBinding.h"
+#include "MovieSceneCopyableTrack.h"
 #include "ISequencerChannelInterface.h"
 #include "SequencerKeyCollection.h"
 #include "CurveEditor.h"
@@ -6253,7 +6254,9 @@ void FSequencer::CopySelectedObjects(TArray<TSharedPtr<FSequencerObjectBindingNo
 
 void FSequencer::CopySelectedTracks(TArray<TSharedPtr<FSequencerTrackNode>>& TrackNodes, FString& ExportedText)
 {
-	TArray<UObject*> TracksToCopy;
+	UMovieScene* MovieScene = GetFocusedMovieSceneSequence()->GetMovieScene();
+
+	TArray<UObject*> CopyableTracks;
 	for (TSharedPtr<FSequencerTrackNode> TrackNode : TrackNodes)
 	{
 		bool bIsParentSelected = false;
@@ -6270,11 +6273,16 @@ void FSequencer::CopySelectedTracks(TArray<TSharedPtr<FSequencerTrackNode>>& Tra
 
 		if (!bIsParentSelected)
 		{
-			TracksToCopy.Add(TrackNode->GetTrack());
+			UMovieSceneCopyableTrack *CopyableTrack = NewObject<UMovieSceneCopyableTrack>(GetTransientPackage(), UMovieSceneCopyableTrack::StaticClass(), NAME_None, RF_Transient);
+			CopyableTracks.Add(CopyableTrack);
+
+			UMovieSceneTrack* DuplicatedTrack = Cast<UMovieSceneTrack>(StaticDuplicateObject(TrackNode->GetTrack(), CopyableTrack));
+			CopyableTrack->Track = DuplicatedTrack;
+			CopyableTrack->bIsAMasterTrack = MovieScene->IsAMasterTrack(*TrackNode->GetTrack());
 		}
 	}
 
-	FSequencer::ExportObjectsToText(TracksToCopy, /*out*/ ExportedText);
+	ExportObjectsToText(CopyableTracks, /*out*/ ExportedText);
 }
 
 
@@ -6312,12 +6320,23 @@ void FSequencer::DoPaste()
 	FString TextToImport;
 	FPlatformApplicationMisc::ClipboardPaste(TextToImport);
 
-	PasteObjectBindings(TextToImport);
-	PasteTracks(TextToImport);
-	PasteSections(TextToImport);
+	TArray<FNotificationInfo> PasteErrors;
+	bool bAnythingPasted = false;
+	bAnythingPasted |= PasteObjectBindings(TextToImport, PasteErrors);
+	bAnythingPasted |= PasteTracks(TextToImport, PasteErrors);
+	bAnythingPasted |= PasteSections(TextToImport, PasteErrors);
+
+	if (!bAnythingPasted)
+	{
+		for (auto NotificationInfo : PasteErrors)
+		{
+			NotificationInfo.bUseLargeFont = false;
+			FSlateNotificationManager::Get().AddNotification(NotificationInfo);
+		}
+	}
 }
 
-bool FSequencer::PasteObjectBindings(const FString& TextToImport)
+bool FSequencer::PasteObjectBindings(const FString& TextToImport, TArray<FNotificationInfo>& PasteErrors)
 {
 	TArray<UMovieSceneCopyableBinding*> ImportedBindings;
 	ImportObjectBindingsFromText(TextToImport, ImportedBindings);
@@ -6497,9 +6516,9 @@ bool FSequencer::PasteObjectBindings(const FString& TextToImport)
 	return true;
 }
 
-bool FSequencer::PasteTracks(const FString& TextToImport)
+bool FSequencer::PasteTracks(const FString& TextToImport, TArray<FNotificationInfo>& PasteErrors)
 {
-	TArray<UMovieSceneTrack*> ImportedTracks;
+	TArray<UMovieSceneCopyableTrack*> ImportedTracks;
 	FSequencer::ImportTracksFromText(TextToImport, ImportedTracks);
 
 	if (ImportedTracks.Num() == 0)
@@ -6533,72 +6552,80 @@ bool FSequencer::PasteTracks(const FString& TextToImport)
 		}
 	}
 
+	bool bAnythingPasted = false;
 	if (ObjectNodes.Num())
 	{
 		for (TSharedPtr<FSequencerObjectBindingNode> ObjectNode : ObjectNodes)
 		{
 			FGuid ObjectGuid = ObjectNode->GetObjectBinding();
 
-			TArray<UMovieSceneTrack*> NewTracks;
+			TArray<UMovieSceneCopyableTrack*> NewTracks;
 			FSequencer::ImportTracksFromText(TextToImport, NewTracks);
 
-			for (UMovieSceneTrack* NewTrack : NewTracks)
+			for (UMovieSceneCopyableTrack* CopyableTrack : NewTracks)
 			{
-				if (!GetFocusedMovieSceneSequence()->GetMovieScene()->AddGivenTrack(NewTrack, ObjectGuid))
-				{
-					FNotificationInfo Info(LOCTEXT("TrackAlreadyBound", "Can't Paste: Binding doesn't exist"));
-					Info.FadeInDuration = 0.1f;
-					Info.FadeOutDuration = 0.5f;
-					Info.ExpireDuration = 2.5f;
-					auto NotificationItem = FSlateNotificationManager::Get().AddNotification(Info);
-
-					NotificationItem->SetCompletionState(SNotificationItem::CS_Success);
-					NotificationItem->ExpireAndFadeout();
-
-					continue;
-				}
-				else
+				if (!CopyableTrack->bIsAMasterTrack)
 				{
-					NotifyMovieSceneDataChanged(EMovieSceneDataChangeType::MovieSceneStructureItemsChanged);
+					UMovieSceneTrack* NewTrack = CopyableTrack->Track;
+					NewTrack->ClearFlags(RF_Transient);
+					if (!GetFocusedMovieSceneSequence()->GetMovieScene()->AddGivenTrack(NewTrack, ObjectGuid))
+					{
+						FNotificationInfo Info(LOCTEXT("TrackAlreadyBound", "Can't Paste: Binding doesn't exist"));
+						PasteErrors.Add(Info);
+						continue;
+					}
+					else
+					{
+						bAnythingPasted = true;
+					}
 				}
 			}
 		}
-
-		return true;
 	}
 
 	// Add as master track or set camera cut track
-	for (UMovieSceneTrack* NewTrack : ImportedTracks)
+	for (UMovieSceneCopyableTrack* CopyableTrack : ImportedTracks)
 	{
-		if (NewTrack->IsA(UMovieSceneCameraCutTrack::StaticClass()))
+		if (CopyableTrack->bIsAMasterTrack)
 		{
-			GetFocusedMovieSceneSequence()->GetMovieScene()->SetCameraCutTrack(NewTrack);
-			if (SelectedParentFolders.Num() > 0)
-			{
-				SelectedParentFolders[0]->AddChildMasterTrack(NewTrack);
-			}
+			UMovieSceneTrack* NewTrack = CopyableTrack->Track;
+			NewTrack->ClearFlags(RF_Transient);
 
-			NotifyMovieSceneDataChanged(EMovieSceneDataChangeType::MovieSceneStructureItemsChanged);
-		}
-		else
-		{
-			if (GetFocusedMovieSceneSequence()->GetMovieScene()->AddGivenMasterTrack(NewTrack))
+			if (NewTrack->IsA(UMovieSceneCameraCutTrack::StaticClass()))
 			{
+				GetFocusedMovieSceneSequence()->GetMovieScene()->SetCameraCutTrack(NewTrack);
 				if (SelectedParentFolders.Num() > 0)
 				{
 					SelectedParentFolders[0]->AddChildMasterTrack(NewTrack);
 				}
 
-				NotifyMovieSceneDataChanged(EMovieSceneDataChangeType::MovieSceneStructureItemsChanged);
+				bAnythingPasted = true;
+			}
+			else
+			{
+				if (GetFocusedMovieSceneSequence()->GetMovieScene()->AddGivenMasterTrack(NewTrack))
+				{
+					if (SelectedParentFolders.Num() > 0)
+					{
+						SelectedParentFolders[0]->AddChildMasterTrack(NewTrack);
+					}
+
+					bAnythingPasted = true;
+				}
 			}
 		}
 	}
 
-	return true;
+	if (bAnythingPasted)
+	{
+		NotifyMovieSceneDataChanged(EMovieSceneDataChangeType::MovieSceneStructureItemsChanged);
+	}
+
+	return bAnythingPasted;
 }
 
 
-bool FSequencer::PasteSections(const FString& TextToImport)
+bool FSequencer::PasteSections(const FString& TextToImport, TArray<FNotificationInfo>& PasteErrors)
 {
 	TArray<UMovieSceneSection*> ImportedSections;
 	FSequencer::ImportSectionsFromText(TextToImport, ImportedSections);
@@ -6613,8 +6640,7 @@ bool FSequencer::PasteSections(const FString& TextToImport)
 	if (SelectedNodes.Num() == 0)
 	{
 		FNotificationInfo Info(LOCTEXT("PasteSections_NoSelectedTracks", "No selected tracks to paste sections onto"));
-		Info.bUseLargeFont = false;
-		FSlateNotificationManager::Get().AddNotification(Info);
+		PasteErrors.Add(Info);
 		return false;
 	}
 
@@ -6699,8 +6725,7 @@ bool FSequencer::PasteSections(const FString& TextToImport)
 		Transaction.Cancel();
 
 		FNotificationInfo Info(LOCTEXT("PasteSections_NothingPasted", "Nothing pasted. No matching section types found."));
-		Info.bUseLargeFont = false;
-		FSlateNotificationManager::Get().AddNotification(Info);
+		PasteErrors.Add(Info);
 		return false;
 	}
 
@@ -6726,7 +6751,7 @@ public:
 	// FCustomizableTextObjectFactory implementation
 	virtual bool CanCreateClass(UClass* InObjectClass, bool& bOmitSubObjs) const override
 	{
-		if (InObjectClass->IsChildOf(UMovieSceneTrack::StaticClass()))
+		if (InObjectClass->IsChildOf(UMovieSceneCopyableTrack::StaticClass()))
 		{
 			return true;
 		}
@@ -6738,11 +6763,11 @@ public:
 	{
 		check(NewObject);
 
-		NewTracks.Add(Cast<UMovieSceneTrack>(NewObject));
+		NewTracks.Add(Cast<UMovieSceneCopyableTrack>(NewObject));
 	}
 
 public:
-	TArray<UMovieSceneTrack*> NewTracks;
+	TArray<UMovieSceneCopyableTrack*> NewTracks;
 };
 
 
@@ -6800,7 +6825,7 @@ bool FSequencer::CanPaste(const FString& TextToImport)
 	return false;
 }
 
-void FSequencer::ImportTracksFromText(const FString& TextToImport, /*out*/ TArray<UMovieSceneTrack*>& ImportedTracks)
+void FSequencer::ImportTracksFromText(const FString& TextToImport, /*out*/ TArray<UMovieSceneCopyableTrack*>& ImportedTracks)
 {
 	UPackage* TempPackage = NewObject<UPackage>(nullptr, TEXT("/Engine/Sequencer/Editor/Transient"), RF_Transient);
 	TempPackage->AddToRoot();
diff --git a/Engine/Source/Editor/Sequencer/Private/Sequencer.h b/Engine/Source/Editor/Sequencer/Private/Sequencer.h
--- a/Engine/Source/Editor/Sequencer/Private/Sequencer.h
+++ b/Engine/Source/Editor/Sequencer/Private/Sequencer.h
@@ -61,12 +61,14 @@ class UMovieSceneSequence;
 class UMovieSceneSubSection;
 class USequencerSettings;
 class UMovieSceneCopyableBinding;
+class UMovieSceneCopyableTrack;
 struct FMovieSceneTimeController;
 struct FMovieScenePossessable;
 struct FSequencerTemplateStore;
 struct FTransformData;
 struct ISequencerHotspot;
 struct FKeyAttributes;
+struct FNotificationInfo;
 
 enum class EMapChangeType : uint8;
 
@@ -480,11 +482,11 @@ public:
 	/** Called when a user executes the paste track menu item */
 	bool CanPaste(const FString& TextToImport);
 	void DoPaste();
-	bool PasteTracks(const FString& TextToImport);
-	bool PasteSections(const FString& TextToImport);
-	bool PasteObjectBindings(const FString& TextToImport);
-	
-	void ImportTracksFromText(const FString& TextToImport, /*out*/ TArray<UMovieSceneTrack*>& ImportedTracks);
+	bool PasteTracks(const FString& TextToImport, TArray<FNotificationInfo>& PasteErrors);
+	bool PasteSections(const FString& TextToImport, TArray<FNotificationInfo>& PasteErrors);
+	bool PasteObjectBindings(const FString& TextToImport, TArray<FNotificationInfo>& PasteErrors);
+
+	void ImportTracksFromText(const FString& TextToImport, /*out*/ TArray<UMovieSceneCopyableTrack*>& ImportedTracks);
 	void ImportSectionsFromText(const FString& TextToImport, /*out*/ TArray<UMovieSceneSection*>& ImportedSections);
 	void ImportObjectBindingsFromText(const FString& TextToImport, /*out*/ TArray<UMovieSceneCopyableBinding*>& ImportedObjects);
 
