Sequencer: Better error reporting for pasting tracks

diff --git a/Engine/Source/Editor/Sequencer/Private/Sequencer.cpp b/Engine/Source/Editor/Sequencer/Private/Sequencer.cpp
--- a/Engine/Source/Editor/Sequencer/Private/Sequencer.cpp
+++ b/Engine/Source/Editor/Sequencer/Private/Sequencer.cpp
@@ -6324,15 +6324,16 @@ void FSequencer::DoPaste()
 	bool bAnythingPasted = false;
 	bAnythingPasted |= PasteObjectBindings(TextToImport, PasteErrors);
 	bAnythingPasted |= PasteTracks(TextToImport, PasteErrors);
-	bAnythingPasted |= PasteSections(TextToImport, PasteErrors);
-
+	
 	if (!bAnythingPasted)
 	{
-		for (auto NotificationInfo : PasteErrors)
-		{
-			NotificationInfo.bUseLargeFont = false;
-			FSlateNotificationManager::Get().AddNotification(NotificationInfo);
-		}
+		bAnythingPasted |= PasteSections(TextToImport, PasteErrors);
+	}
+
+	for (auto NotificationInfo : PasteErrors)
+	{
+		NotificationInfo.bUseLargeFont = false;
+		FSlateNotificationManager::Get().AddNotification(NotificationInfo);
 	}
 }
 
@@ -6552,7 +6553,23 @@ bool FSequencer::PasteTracks(const FString& TextToImport, TArray<FNotificationIn
 		}
 	}
 
-	bool bAnythingPasted = false;
+	int32 NumMasterTracks = 0;
+	int32 NumTracks = 0;
+
+	for (UMovieSceneCopyableTrack* CopyableTrack : ImportedTracks)
+	{
+		if (CopyableTrack->bIsAMasterTrack)
+		{
+			++NumMasterTracks;
+		}
+		else
+		{
+			++NumTracks;
+		}
+	}
+
+	int32 NumMasterTracksPasted = 0;
+	int32 NumTracksPasted = 0;
 	if (ObjectNodes.Num())
 	{
 		for (TSharedPtr<FSequencerObjectBindingNode> ObjectNode : ObjectNodes)
@@ -6570,13 +6587,11 @@ bool FSequencer::PasteTracks(const FString& TextToImport, TArray<FNotificationIn
 					NewTrack->ClearFlags(RF_Transient);
 					if (!GetFocusedMovieSceneSequence()->GetMovieScene()->AddGivenTrack(NewTrack, ObjectGuid))
 					{
-						FNotificationInfo Info(LOCTEXT("TrackAlreadyBound", "Can't Paste: Binding doesn't exist"));
-						PasteErrors.Add(Info);
 						continue;
 					}
 					else
 					{
-						bAnythingPasted = true;
+						++NumTracksPasted;
 					}
 				}
 			}
@@ -6599,7 +6614,7 @@ bool FSequencer::PasteTracks(const FString& TextToImport, TArray<FNotificationIn
 					SelectedParentFolders[0]->AddChildMasterTrack(NewTrack);
 				}
 
-				bAnythingPasted = true;
+				++NumMasterTracksPasted;
 			}
 			else
 			{
@@ -6610,18 +6625,31 @@ bool FSequencer::PasteTracks(const FString& TextToImport, TArray<FNotificationIn
 						SelectedParentFolders[0]->AddChildMasterTrack(NewTrack);
 					}
 
-					bAnythingPasted = true;
 				}
+
+				++NumMasterTracksPasted;
 			}
 		}
 	}
 
-	if (bAnythingPasted)
+	if (NumMasterTracksPasted < NumMasterTracks)
+	{
+		FNotificationInfo Info(LOCTEXT("PasteTracks_NoMasterTracks", "Can't paste track. Master track could not be pasted"));
+		PasteErrors.Add(Info);
+	}
+
+	if (NumTracksPasted < NumTracks)
+	{
+		FNotificationInfo Info(LOCTEXT("PasteSections_NoSelectedObjects", "Can't paste track. No selected objects to paste tracks onto"));
+		PasteErrors.Add(Info);
+	}
+
+	if ((NumMasterTracksPasted + NumTracksPasted) > 0)
 	{
 		NotifyMovieSceneDataChanged(EMovieSceneDataChangeType::MovieSceneStructureItemsChanged);
 	}
 
-	return bAnythingPasted;
+	return true;
 }
 
 
@@ -6639,7 +6667,7 @@ bool FSequencer::PasteSections(const FString& TextToImport, TArray<FNotification
 
 	if (SelectedNodes.Num() == 0)
 	{
-		FNotificationInfo Info(LOCTEXT("PasteSections_NoSelectedTracks", "No selected tracks to paste sections onto"));
+		FNotificationInfo Info(LOCTEXT("PasteSections_NoSelectedTracks", "Can't paste section. No selected tracks to paste sections onto"));
 		PasteErrors.Add(Info);
 		return false;
 	}
@@ -6724,7 +6752,7 @@ bool FSequencer::PasteSections(const FString& TextToImport, TArray<FNotification
 	{
 		Transaction.Cancel();
 
-		FNotificationInfo Info(LOCTEXT("PasteSections_NothingPasted", "Nothing pasted. No matching section types found."));
+		FNotificationInfo Info(LOCTEXT("PasteSections_NothingPasted", "Can't paste section. No matching section types found."));
 		PasteErrors.Add(Info);
 		return false;
 	}
