Sequencer: Fix to disallow keyframing when allow edits is allow level edits only.

diff --git a/Engine/Source/Editor/Sequencer/Public/ISequencer.h b/Engine/Source/Editor/Sequencer/Public/ISequencer.h
--- a/Engine/Source/Editor/Sequencer/Public/ISequencer.h
+++ b/Engine/Source/Editor/Sequencer/Public/ISequencer.h
@@ -293,7 +293,7 @@ public:
 	virtual bool GetAutoSetTrackDefaults() const = 0;
 
 	/** @return Returns whether sequencer will respond to changes and possibly create a key or track */
-	virtual bool IsAllowedToChange() const { return GetAllowEditsMode() != EAllowEditsMode::AllowLevelEditsOnly || GetAutoChangeMode() != EAutoChangeMode::None; }
+	virtual bool IsAllowedToChange() const { return GetAllowEditsMode() != EAllowEditsMode::AllowLevelEditsOnly && GetAutoChangeMode() != EAutoChangeMode::None; }
 
 	/**
 	 * Gets the current time of the time slider relative to the currently focused movie scene
