Copying //UE4/Dev-Sequencer @ 4773300 to Dev-Main (//UE4/Dev-Main)

diff --git a/Engine/Config/BaseEditorPerProjectUserSettings.ini b/Engine/Config/BaseEditorPerProjectUserSettings.ini
--- a/Engine/Config/BaseEditorPerProjectUserSettings.ini
+++ b/Engine/Config/BaseEditorPerProjectUserSettings.ini
@@ -833,6 +833,7 @@ LoopMode=SLM_Loop
 +TrackSettings=(MatchingActorClass=/Script/Engine.Light,DefaultTracks=(None),DefaultPropertyTracks=((ComponentPath="LightComponent0",PropertyPath="Intensity"),(ComponentPath="LightComponent0",PropertyPath="LightColor")))
 
 [/Script/MovieSceneTools.MovieSceneToolsProjectSettings]
++FbxSettings=(FbxPropertyName="FieldOfView", PropertyPath=(ComponentName="CameraComponent",PropertyName="FieldOfView"))
 +FbxSettings=(FbxPropertyName="FocalLength", PropertyPath=(ComponentName="CameraComponent",PropertyName="CurrentFocalLength"))
 +FbxSettings=(FbxPropertyName="FocusDistance", PropertyPath=(ComponentName="CameraComponent",PropertyName="FocusSettings.ManualFocusDistance"))
 
diff --git a/Engine/Config/BaseInput.ini b/Engine/Config/BaseInput.ini
--- a/Engine/Config/BaseInput.ini
+++ b/Engine/Config/BaseInput.ini
@@ -186,6 +186,8 @@ MaxScrollbackSize=1024
 +ManualAutoCompleteList=(Command="StopRecordingAnimation All",Desc="Stop all recording animations")
 +ManualAutoCompleteList=(Command="RecordSequence Filter ActorOrClassName",Desc="Record a level sequence from gameplay. Filter=<All|Actor|Class>")
 +ManualAutoCompleteList=(Command="StopRecordingSequence",Desc="Stop recording the current sequence. Only one sequence recording can be active at one time.")
++ManualAutoCompleteList=(Command="RecordTake Filter ActorOrClassName",Desc="Record a level sequence from gameplay. Filter=<All|Actor|Class>")
++ManualAutoCompleteList=(Command="StopRecordingTake",Desc="Stop recording the current sequence. Only one sequence recording can be active at one time.")
 +ManualAutoCompleteList=(Command="FreezeRendering",Desc="Toggle freezing of most aspects of rendering (such as visibility calculations), useful in conjunction with ToggleDebugCamera to fly around and see how frustum and occlusion culling is working")
 +ManualAutoCompleteList=(Command="ProfileGPU",Desc="Profile one frame of rendering commands sent to the GPU")
 +ManualAutoCompleteList=(Command="ProfileGPUHitches",Desc="Toggle profiling of GPU hitches.")
diff --git a/Engine/Plugins/Compositing/Composure/Source/Composure/Private/MovieScene/MovieSceneComposurePostMoveSettingsSection.cpp b/Engine/Plugins/Compositing/Composure/Source/Composure/Private/MovieScene/MovieSceneComposurePostMoveSettingsSection.cpp
--- a/Engine/Plugins/Compositing/Composure/Source/Composure/Private/MovieScene/MovieSceneComposurePostMoveSettingsSection.cpp
+++ b/Engine/Plugins/Compositing/Composure/Source/Composure/Private/MovieScene/MovieSceneComposurePostMoveSettingsSection.cpp
@@ -17,24 +17,30 @@ struct FPostMoveSettingsChannelEditorData
 		MetaData[0].SetIdentifiers("Pivot.X", FCommonChannelData::ChannelX, PivotGroup);
 		MetaData[0].SortOrder = 0;
 		MetaData[0].Color = FCommonChannelData::RedChannelColor;
+		MetaData[0].bCanCollapseToTrack = false;
 
 		MetaData[1].SetIdentifiers("Pivot.Y", FCommonChannelData::ChannelY, PivotGroup);
 		MetaData[1].SortOrder = 1;
 		MetaData[1].Color = FCommonChannelData::GreenChannelColor;
+		MetaData[1].bCanCollapseToTrack = false;
 
 		MetaData[2].SetIdentifiers("Translation.X", FCommonChannelData::ChannelX, TranslationGroup);
 		MetaData[2].SortOrder = 2;
 		MetaData[2].Color = FCommonChannelData::RedChannelColor;
+		MetaData[2].bCanCollapseToTrack = false;
 
 		MetaData[3].SetIdentifiers("Translation.Y", FCommonChannelData::ChannelY, TranslationGroup);
 		MetaData[3].SortOrder = 3;
 		MetaData[3].Color = FCommonChannelData::GreenChannelColor;
+		MetaData[3].bCanCollapseToTrack = false;
 
 		MetaData[4].SetIdentifiers("Rotation", NSLOCTEXT("PostMoves", "Rotation", "Rotation"));
 		MetaData[4].SortOrder = 4;
+		MetaData[4].bCanCollapseToTrack = false;
 
 		MetaData[5].SetIdentifiers("Scale", NSLOCTEXT("PostMoves", "Scale", "Scale"));
 		MetaData[5].SortOrder = 5;
+		MetaData[5].bCanCollapseToTrack = false;
 
 		ExternalValues[0].OnGetExternalValue = ExtractPivotX;
 		ExternalValues[1].OnGetExternalValue = ExtractPivotY;
diff --git a/Engine/Plugins/Media/MediaCompositing/Source/MediaCompositing/Private/MovieScene/MovieSceneMediaTemplate.cpp b/Engine/Plugins/Media/MediaCompositing/Source/MediaCompositing/Private/MovieScene/MovieSceneMediaTemplate.cpp
--- a/Engine/Plugins/Media/MediaCompositing/Source/MediaCompositing/Private/MovieScene/MovieSceneMediaTemplate.cpp
+++ b/Engine/Plugins/Media/MediaCompositing/Source/MediaCompositing/Private/MovieScene/MovieSceneMediaTemplate.cpp
@@ -114,7 +114,7 @@ struct FMediaSectionExecutionToken
 				MediaPlayer->Seek(MediaTime);
 				MediaPlayer->SetRate(1.0f);
 			}
-			else 
+			else if (Context.HasJumped())
 			{
 				MediaPlayer->Seek(MediaTime);
 			}
diff --git a/Engine/Plugins/MovieScene/LevelSequenceEditor/Source/LevelSequenceEditor/Private/LevelSequenceEditorModule.cpp b/Engine/Plugins/MovieScene/LevelSequenceEditor/Source/LevelSequenceEditor/Private/LevelSequenceEditorModule.cpp
--- a/Engine/Plugins/MovieScene/LevelSequenceEditor/Source/LevelSequenceEditor/Private/LevelSequenceEditorModule.cpp
+++ b/Engine/Plugins/MovieScene/LevelSequenceEditor/Source/LevelSequenceEditor/Private/LevelSequenceEditorModule.cpp
@@ -76,7 +76,7 @@ public:
 		{
 			virtual UBlueprint* RetrieveBlueprint(UObject* InObject) const override
 			{
-				return Cast<UBlueprint>(CastChecked<ULevelSequence>(InObject)->DirectorBlueprint);
+				return CastChecked<ULevelSequence>(InObject)->GetDirectorBlueprint();
 			}
 
 			virtual bool AssetContainsBlueprint(const FAssetData& InAssetData) const
diff --git a/Engine/Plugins/MovieScene/LevelSequenceEditor/Source/LevelSequenceEditor/Private/Misc/MovieSceneSequenceEditor_LevelSequence.h b/Engine/Plugins/MovieScene/LevelSequenceEditor/Source/LevelSequenceEditor/Private/Misc/MovieSceneSequenceEditor_LevelSequence.h
--- a/Engine/Plugins/MovieScene/LevelSequenceEditor/Source/LevelSequenceEditor/Private/Misc/MovieSceneSequenceEditor_LevelSequence.h
+++ b/Engine/Plugins/MovieScene/LevelSequenceEditor/Source/LevelSequenceEditor/Private/Misc/MovieSceneSequenceEditor_LevelSequence.h
@@ -19,7 +19,7 @@ struct FMovieSceneSequenceEditor_LevelSequence : FMovieSceneSequenceEditor
 	virtual UBlueprint* GetBlueprintForSequence(UMovieSceneSequence* InSequence) const override
 	{
 		ULevelSequence* LevelSequence = CastChecked<ULevelSequence>(InSequence);
-		return CastChecked<UBlueprint>(LevelSequence->DirectorBlueprint, ECastCheckedType::NullAllowed);
+		return LevelSequence->GetDirectorBlueprint();
 	}
 
 	virtual UBlueprint* CreateBlueprintForSequence(UMovieSceneSequence* InSequence) const override
@@ -35,8 +35,7 @@ struct FMovieSceneSequenceEditor_LevelSequence : FMovieSceneSequenceEditor
 		FName BlueprintName = "SequenceDirector";
 		Blueprint = FKismetEditorUtilities::CreateBlueprint(ULevelSequenceDirector::StaticClass(), InSequence, BlueprintName, BPTYPE_Normal, UBlueprint::StaticClass(), UBlueprintGeneratedClass::StaticClass());
 
-		LevelSequence->DirectorBlueprint = Blueprint;
-		LevelSequence->DirectorClass = CastChecked<UBlueprintGeneratedClass>(Blueprint->GeneratedClass);
+		LevelSequence->SetDirectorBlueprint(Blueprint);
 		return Blueprint;
 	}
 
diff --git a/Engine/Plugins/VirtualProduction/Takes/Source/TakeRecorder/Private/Widgets/STakeRecorderCockpit.cpp b/Engine/Plugins/VirtualProduction/Takes/Source/TakeRecorder/Private/Widgets/STakeRecorderCockpit.cpp
--- a/Engine/Plugins/VirtualProduction/Takes/Source/TakeRecorder/Private/Widgets/STakeRecorderCockpit.cpp
+++ b/Engine/Plugins/VirtualProduction/Takes/Source/TakeRecorder/Private/Widgets/STakeRecorderCockpit.cpp
@@ -94,6 +94,15 @@ void STakeRecorderCockpit::Construct(const FArguments& InArgs)
 
 	CacheMetaData();
 
+	if (TakeMetaData && !TakeMetaData->IsLocked())
+	{
+		int32 NextTakeNumber = UTakesCoreBlueprintLibrary::ComputeNextTakeNumber(TakeMetaData->GetSlate());
+		if (NextTakeNumber != TakeMetaData->GetTakeNumber())
+		{
+			TakeMetaData->SetTakeNumber(NextTakeNumber);
+		}
+	}
+
 	UpdateTakeError();
 	UpdateRecordError();
 
diff --git a/Engine/Source/Editor/LevelEditor/Private/LevelViewportActions.cpp b/Engine/Source/Editor/LevelEditor/Private/LevelViewportActions.cpp
--- a/Engine/Source/Editor/LevelEditor/Private/LevelViewportActions.cpp
+++ b/Engine/Source/Editor/LevelEditor/Private/LevelViewportActions.cpp
@@ -33,7 +33,6 @@ void FLevelViewportCommands::RegisterCommands()
 	UI_COMMAND( ToggleGameView, "Game View", "Toggles game view.  Game view shows the scene as it appears in game", EUserInterfaceActionType::ToggleButton, FInputChord( EKeys::G ) );
 	UI_COMMAND( ToggleImmersive, "Immersive Mode", "Switches this viewport between immersive mode and regular mode", EUserInterfaceActionType::ToggleButton, PLATFORM_MAC ? FInputChord( EModifierKey::Control, EKeys::F11 ) : FInputChord( EKeys::F11 ) );
 
-	UI_COMMAND( CreateCamera, "Create Camera Here", "Creates a new camera actor at the current location of this viewport's camera", EUserInterfaceActionType::Button, FInputChord() );
 	UI_COMMAND( HighResScreenshot, "High Resolution Screenshot...", "Opens the control panel for high resolution screenshots", EUserInterfaceActionType::Button, FInputChord() );
 	
 	UI_COMMAND( UseDefaultShowFlags, "Use Defaults", "Resets all show flags to default", EUserInterfaceActionType::Button, FInputChord() );
diff --git a/Engine/Source/Editor/LevelEditor/Private/SLevelViewport.cpp b/Engine/Source/Editor/LevelEditor/Private/SLevelViewport.cpp
--- a/Engine/Source/Editor/LevelEditor/Private/SLevelViewport.cpp
+++ b/Engine/Source/Editor/LevelEditor/Private/SLevelViewport.cpp
@@ -67,6 +67,8 @@
 #include "FoliageType.h"
 #include "IVREditorModule.h"
 #include "ShowFlagMenuCommands.h"
+#include "AssetRegistryModule.h"
+#include "IAssetRegistry.h"
 #include "BufferVisualizationMenuCommands.h"
 
 static const FName LevelEditorName("LevelEditor");
@@ -1188,13 +1190,51 @@ void SLevelViewport::BindOptionCommands( FUICommandList& OutCommandList )
 		FIsActionChecked::CreateSP( this, &SLevelViewport::AllowsCinematicPreview )
 		);
 
+	IAssetRegistry & AssetRegistry = FModuleManager::LoadModuleChecked<FAssetRegistryModule>("AssetRegistry").Get();
+	TArray<FName> ClassNames;
+	TSet<FName> DerivedClassNames;
+	ClassNames.Add(ACameraActor::StaticClass()->GetFName());
+	AssetRegistry.GetDerivedClassNames(ClassNames, TSet<FName>(), DerivedClassNames);
 
-	OutCommandList.MapAction(
-		ViewportActions.CreateCamera,
-		FExecuteAction::CreateSP( this, &SLevelViewport::OnCreateCameraActor ),
-		FCanExecuteAction(),
-		FCanExecuteAction::CreateSP( this, &SLevelViewport::IsPerspectiveViewport ) 
-		);
+	for (FName Name : DerivedClassNames)
+	{
+		// Ignore generated types that cannot be spawned
+		if (Name.ToString().StartsWith("SKEL_") || Name.ToString().StartsWith("REINST_"))
+		{
+			continue;
+		}
+		
+		// Remove _C from display names for blueprint classes
+		FName DisplayName = Name;
+		if (Name.ToString().EndsWith("_C"))
+		{
+			DisplayName = FName(*DisplayName.ToString().LeftChop(2));
+		}
+		
+		// Look for existing UI Command info so one isn't created for every viewport
+		TSharedPtr<FUICommandInfo> * FoundCamera = FLevelViewportCommands::Get().CreateCameras.FindByPredicate([Name](TSharedPtr<FUICommandInfo> Camera) { return Camera->GetCommandName() == Name; });
+		
+		if (FoundCamera)
+		{
+			OutCommandList.MapAction(
+				*FoundCamera,
+				FExecuteAction::CreateSP(this, &SLevelViewport::OnCreateCameraActor, FindObject<UClass>(ANY_PACKAGE, *Name.ToString()))
+			);
+		}
+		else
+		{
+			// If command info isn't found, create a new one
+			TSharedRef<FUICommandInfo> NewCamera = FUICommandInfoDecl(FLevelViewportCommands::Get().AsShared(), Name, FText::FromName(DisplayName), FText::Format(LOCTEXT("SpawnCamerasTooltip", "Spawn Camera here of type {0}"), FText::FromName(Name))).UserInterfaceType(EUserInterfaceActionType::Button).DefaultChord(FInputChord());
+			
+			OutCommandList.MapAction(
+				NewCamera,
+				FExecuteAction::CreateSP(this, &SLevelViewport::OnCreateCameraActor, FindObject<UClass>(ANY_PACKAGE, *Name.ToString()))
+			);
+			
+			FLevelViewportCommands::Get().CreateCameras.Add(NewCamera);
+		}
+
+	}
 
 	OutCommandList.MapAction(
 		ViewportActions.HighResScreenshot,
@@ -1680,7 +1720,7 @@ bool SLevelViewport::IsImmersive() const
 	return false;
 }
 
-void SLevelViewport::OnCreateCameraActor()
+void SLevelViewport::OnCreateCameraActor(UClass* InClass)
 {		
 	// Find the perspective viewport we were using
 	FViewport* pViewPort = GEditor->GetActiveViewport();
@@ -1704,10 +1744,10 @@ void SLevelViewport::OnCreateCameraActor()
 	const FScopedTransaction Transaction(NSLOCTEXT("LevelViewport", "CreateCameraHere", "Create Camera Here"));
 
 	// Set new camera to match viewport
-	ACameraActor* pNewCamera = ViewportClient->GetWorld()->SpawnActor<ACameraActor>();
+	ACameraActor* pNewCamera = Cast<ACameraActor>(ViewportClient->GetWorld()->SpawnActor(InClass));
 	pNewCamera->SetActorLocation( ViewportClient->GetViewLocation(), false );
 	pNewCamera->SetActorRotation( ViewportClient->GetViewRotation() );
-	pNewCamera->GetCameraComponent()->FieldOfView = ViewportClient->ViewFOV;
+	pNewCamera->GetCameraComponent()->SetFieldOfView( ViewportClient->ViewFOV );
 
 	// Deselect any currently selected actors
 	GUnrealEd->SelectNone( true, true );
diff --git a/Engine/Source/Editor/LevelEditor/Private/SLevelViewportToolBar.cpp b/Engine/Source/Editor/LevelEditor/Private/SLevelViewportToolBar.cpp
--- a/Engine/Source/Editor/LevelEditor/Private/SLevelViewportToolBar.cpp
+++ b/Engine/Source/Editor/LevelEditor/Private/SLevelViewportToolBar.cpp
@@ -1,4 +1,4 @@
-﻿// Copyright 1998-2019 Epic Games, Inc. All Rights Reserved.
+// Copyright 1998-2019 Epic Games, Inc. All Rights Reserved.
 
 #include "SLevelViewportToolBar.h"
 #include "Framework/Commands/UIAction.h"
@@ -582,8 +582,13 @@ TSharedRef<SWidget> SLevelViewportToolBar::GenerateOptionsMenu() const
 					);
 			}
 			OptionsMenuBuilder.EndSection();
-
-			OptionsMenuBuilder.AddMenuEntry( LevelViewportActions.CreateCamera );
+		
+			OptionsMenuBuilder.AddSubMenu(
+				LOCTEXT("CameraSubMeun", "Create Camera Here"),
+				LOCTEXT("CameraSubMenu_ToolTip", "Select a camera type to create at current viewport's location"),
+				FNewMenuDelegate::CreateSP(this, &SLevelViewportToolBar::GenerateCameraSpawnMenu)
+			);
+			
 		}
 
 		OptionsMenuBuilder.AddMenuEntry( LevelViewportActions.HighResScreenshot );
@@ -736,6 +741,16 @@ void SLevelViewportToolBar::GenerateViewportTypeMenu( FMenuBuilder& Builder ) co
 	});
 }
 
+void SLevelViewportToolBar::GenerateCameraSpawnMenu(FMenuBuilder & Builder) const
+{
+	const FLevelViewportCommands& Actions = FLevelViewportCommands::Get();
+
+	for (TSharedPtr<FUICommandInfo> Camera : Actions.CreateCameras)
+	{
+		Builder.AddMenuEntry(Camera);
+	}
+}
+
 TSharedRef<SWidget> SLevelViewportToolBar::GenerateCameraMenu() const
 {
 	Viewport.Pin()->OnFloatingButtonClicked();
diff --git a/Engine/Source/Editor/LevelEditor/Private/SLevelViewportToolBar.h b/Engine/Source/Editor/LevelEditor/Private/SLevelViewportToolBar.h
--- a/Engine/Source/Editor/LevelEditor/Private/SLevelViewportToolBar.h
+++ b/Engine/Source/Editor/LevelEditor/Private/SLevelViewportToolBar.h
@@ -128,6 +128,13 @@ private:
 	 */
 	void GenerateViewportTypeMenu( FMenuBuilder& Builder ) const;
 
+	/**
+	 * Generates menu entries for spawning cameras at the current viewport
+	 *
+	 * @param Builder	The menu builder to add menu entries to
+	 */
+	void GenerateCameraSpawnMenu(FMenuBuilder& Builder) const;
+
 	/**
 	 * Generates the toolbar view menu content 
 	 *
diff --git a/Engine/Source/Editor/LevelEditor/Public/LevelViewportActions.h b/Engine/Source/Editor/LevelEditor/Public/LevelViewportActions.h
--- a/Engine/Source/Editor/LevelEditor/Public/LevelViewportActions.h
+++ b/Engine/Source/Editor/LevelEditor/Public/LevelViewportActions.h
@@ -93,8 +93,8 @@ PRAGMA_ENABLE_DEPRECATION_WARNINGS
 	/** Toggles maximize mode in the viewport */
 	TSharedPtr< FUICommandInfo > ToggleMaximize;
 
-	/** Creates a cameraActor at current perspective viewport position */
-	TSharedPtr< FUICommandInfo > CreateCamera;
+	/** Creates a subclass of ACameraActor at current perspective viewport position */
+	TArray< TSharedPtr< FUICommandInfo > > CreateCameras;
 
 	/** Opens the control panel for high resolution screenshots */
 	TSharedPtr< FUICommandInfo > HighResScreenshot;
diff --git a/Engine/Source/Editor/LevelEditor/Public/SLevelViewport.h b/Engine/Source/Editor/LevelEditor/Public/SLevelViewport.h
--- a/Engine/Source/Editor/LevelEditor/Public/SLevelViewport.h
+++ b/Engine/Source/Editor/LevelEditor/Public/SLevelViewport.h
@@ -437,7 +437,7 @@ private:
 	/**
 	 * Called to create a cameraActor in the currently selected perspective viewport
 	 */
-	void OnCreateCameraActor();
+	void OnCreateCameraActor(UClass *InClass);
 
 	/**
 	 * Called to bring up the screenshot UI
diff --git a/Engine/Source/Editor/MovieSceneTools/Private/AutomatedLevelSequenceCapture.cpp b/Engine/Source/Editor/MovieSceneTools/Private/AutomatedLevelSequenceCapture.cpp
--- a/Engine/Source/Editor/MovieSceneTools/Private/AutomatedLevelSequenceCapture.cpp
+++ b/Engine/Source/Editor/MovieSceneTools/Private/AutomatedLevelSequenceCapture.cpp
@@ -202,6 +202,9 @@ void UAutomatedLevelSequenceCapture::Initialize(TSharedPtr<FSceneViewport> InVie
 		}
 	}
 
+	ExportEDL();
+	ExportFCPXML();
+
 	if (Actor)
 	{
 		// Ensure it doesn't loop (-1 is indefinite)
@@ -243,9 +246,6 @@ void UAutomatedLevelSequenceCapture::Initialize(TSharedPtr<FSceneViewport> InVie
 		UE_LOG(LogMovieSceneCapture, Error, TEXT("Could not find or create a Level Sequence Actor for this capture. Capturing will fail."));
 	}
 
-	ExportEDL();
-	ExportFCPXML();
-
 	CaptureState = ELevelSequenceCaptureState::Setup;
 	CaptureStrategy = MakeShareable(new FFixedTimeStepCaptureStrategy(Settings.FrameRate));
 	CaptureStrategy->OnInitialize();
@@ -903,8 +903,9 @@ void UAutomatedLevelSequenceCapture::ExportEDL()
 
 	FString SaveFilename = 	Settings.OutputDirectory.Path / MovieScene->GetOuter()->GetName();
 	int32 HandleFrames = Settings.HandleFrames;
+	FString MovieExtension = Settings.MovieExtension;
 
-	MovieSceneTranslatorEDL::ExportEDL(MovieScene, Settings.FrameRate, SaveFilename, HandleFrames);
+	MovieSceneTranslatorEDL::ExportEDL(MovieScene, Settings.FrameRate, SaveFilename, HandleFrames, MovieExtension);
 }
 
 double UAutomatedLevelSequenceCapture::GetEstimatedCaptureDurationSeconds() const
@@ -947,13 +948,14 @@ void UAutomatedLevelSequenceCapture::ExportFCPXML()
 	FFrameRate FrameRate = Settings.FrameRate;
 	uint32 ResX = Settings.Resolution.ResX;
 	uint32 ResY = Settings.Resolution.ResY;
+	FString MovieExtension = Settings.MovieExtension;
 
 	FFCPXMLExporter *Exporter = new FFCPXMLExporter;
 
 	TSharedRef<FMovieSceneTranslatorContext> ExportContext(new FMovieSceneTranslatorContext);
 	ExportContext->Init();
 
-	bool bSuccess = Exporter->Export(MovieScene, FilenameFormat, FrameRate, ResX, ResY, HandleFrames, SaveFilename, ExportContext);
+	bool bSuccess = Exporter->Export(MovieScene, FilenameFormat, FrameRate, ResX, ResY, HandleFrames, SaveFilename, ExportContext, MovieExtension);
 
 	// Log any messages in context
 	MovieSceneToolHelpers::MovieSceneTranslatorLogMessages(Exporter, ExportContext, false);
diff --git a/Engine/Source/Editor/MovieSceneTools/Private/FCPXML/FCPXMLExport.cpp b/Engine/Source/Editor/MovieSceneTools/Private/FCPXML/FCPXMLExport.cpp
--- a/Engine/Source/Editor/MovieSceneTools/Private/FCPXML/FCPXMLExport.cpp
+++ b/Engine/Source/Editor/MovieSceneTools/Private/FCPXML/FCPXMLExport.cpp
@@ -540,9 +540,10 @@ bool FFCPXMLExportVisitor::ConstructVideoNode(TSharedRef<FFCPXMLNode> InParentNo
 		return false;
 	}
  
-	for (TSharedPtr<FMovieSceneExportCinematicTrackData> CinematicTrack : CinematicMasterTrackData->CinematicTracks)
+	// Add in reverse order
+	for (int32 RowIndex = CinematicMasterTrackData->CinematicTracks.Num()-1; RowIndex >= 0; --RowIndex)
 	{
-		if (!ConstructVideoTrackNode(VideoNode, CinematicTrack, CinematicMasterTrackData)) 
+		if (!ConstructVideoTrackNode(VideoNode, CinematicMasterTrackData->CinematicTracks[RowIndex], CinematicMasterTrackData)) 
 		{ 
 			return false;
 		}
@@ -600,17 +601,18 @@ bool FFCPXMLExportVisitor::ConstructAudioNode(TSharedRef<FFCPXMLNode> InParentNo
 
 	uint32 TrackIndex = 1;
 
-	for (TSharedPtr<FMovieSceneExportAudioMasterTrackData> AudioMasterTrack : ExportData->MovieSceneData->AudioMasterTracks)
+	// Add in reverse order
+	for (int32 RowIndex = ExportData->MovieSceneData->AudioMasterTracks.Num() - 1; RowIndex >= 0; --RowIndex)
 	{
-		if (!AudioMasterTrack.IsValid())
+		if (!ExportData->MovieSceneData->AudioMasterTracks[RowIndex].IsValid())
 		{
 			return false;
 		}
 
-		for (TSharedPtr<FMovieSceneExportAudioTrackData> AudioTrack : AudioMasterTrack->AudioTracks)
+		for (TSharedPtr<FMovieSceneExportAudioTrackData> AudioTrack : ExportData->MovieSceneData->AudioMasterTracks[RowIndex]->AudioTracks)
 		{
 			uint32 OutNumTracks{ 0 };
-			if (!ConstructAudioTrackNode(AudioNode, AudioTrack, AudioMasterTrack, TrackIndex, OutNumTracks))
+			if (!ConstructAudioTrackNode(AudioNode, AudioTrack, ExportData->MovieSceneData->AudioMasterTracks[RowIndex], TrackIndex, OutNumTracks))
 			{
 				return false;
 			}
@@ -1142,12 +1144,12 @@ bool FFCPXMLExportVisitor::GetCinematicSectionFrames(const TSharedPtr<FMovieScen
 		return false;
 	}
 
-	int32 Handles = ExportData->GetHandleFrames().FloorToFrame().Value;
+	int32 HandleFrames = ExportData->GetHandleFrames();
 	OutStartFrame = InCinematicSectionData->StartFrame.Value;
 	OutEndFrame = InCinematicSectionData->EndFrame.Value;
 	OutDuration = OutEndFrame - OutStartFrame;
-	OutInFrame = Handles;
-	OutOutFrame = Handles + OutDuration;
+	OutInFrame = HandleFrames + 1;
+	OutOutFrame = HandleFrames + OutDuration;
 
 	return true;
 }
@@ -1281,7 +1283,16 @@ bool FFCPXMLExportVisitor::CreateCinematicSectionMetadata(const UMovieSceneCinem
 	{
 		return false;
 	}
+
 	OutMetadata = TEXT("[UE4ShotSection=") + InSection->GetPathName() + TEXT("]");
+
+	// Store the start offset and the handle frames for round-tripping to compute the new start offset
+	int32 HandleFrames = ExportData->GetHandleFrames();
+	FFrameRate TickResolution = InSection->GetTypedOuter<UMovieScene>()->GetTickResolution();
+	int32 StartFrameOffset = ConvertFrameTime(InSection->Parameters.StartFrameOffset, TickResolution, ExportData->GetFrameRate()).CeilToFrame().Value;
+
+	OutMetadata += TEXT("[UE4ShotStartOffset=") + FString::FromInt(StartFrameOffset) + TEXT("]");
+	OutMetadata += TEXT("[UE4ShotHandleFrames=") + FString::FromInt(HandleFrames) + TEXT("]");
 	return true;
 }
 
diff --git a/Engine/Source/Editor/MovieSceneTools/Private/FCPXML/FCPXMLImport.cpp b/Engine/Source/Editor/MovieSceneTools/Private/FCPXML/FCPXMLImport.cpp
--- a/Engine/Source/Editor/MovieSceneTools/Private/FCPXML/FCPXMLImport.cpp
+++ b/Engine/Source/Editor/MovieSceneTools/Private/FCPXML/FCPXMLImport.cpp
@@ -62,6 +62,8 @@ FFCPXMLImportVisitor::FFCPXMLImportVisitor(TSharedRef<FMovieSceneImportData> InI
 , CurrAudioMasterTrack(nullptr)
 , CurrAudioTrackRowIndex(0)
 , bCurrImportAudioTrackIsStereoChannel(false)
+, MaxVideoTrackRowIndex(0)
+, MaxAudioTrackRowIndex(0)
 {
 	ConstructAudioTrackList();
 	if (AudioTrackList.Num() > 0)
@@ -149,6 +151,10 @@ bool FFCPXMLImportVisitor::VisitNode(TSharedRef<FFCPXMLTrackNode> InTrackNode)
 	{
 		bInVideoTrackNode = bPrevInTrackNode;
 		CurrVideoTrackRowIndex++;
+		if (MaxVideoTrackRowIndex < CurrVideoTrackRowIndex)
+		{
+			MaxVideoTrackRowIndex = CurrVideoTrackRowIndex;
+		}
 	}
 	else if (bInSequenceNode && bInAudioNode)
 	{
@@ -170,6 +176,10 @@ bool FFCPXMLImportVisitor::VisitNode(TSharedRef<FFCPXMLTrackNode> InTrackNode)
 				CurrAudioTrackRowIndex++;
 			}
 		}
+		if (MaxAudioTrackRowIndex < CurrAudioTrackRowIndex)
+		{
+			MaxAudioTrackRowIndex = CurrAudioTrackRowIndex;
+		}
 	}
 
 	return bSuccess;
@@ -261,6 +271,14 @@ bool FFCPXMLImportVisitor::VisitVideoClipItemNode(TSharedRef<FFCPXMLClipItemNode
 
 	FString SectionPathName = GetCinematicSectionPathName(LogNote, MasterClipId);
 
+	int32 HandleFrames = 0;
+	int32 OriginalStartOffset = 0;
+	TOptional<FFrameNumber> NewStartOffset;
+	if (GetCinematicSectionHandleFramesFromMetadata(LogNote, HandleFrames) && GetCinematicSectionStartOffsetFromMetadata(LogNote, OriginalStartOffset))
+	{
+		NewStartOffset = OriginalStartOffset - ((1 + HandleFrames) - StartOffset);
+	}
+
 	// Find actual section
 	TSharedPtr<FMovieSceneImportCinematicSectionData> SectionData = nullptr;
 	if (!SectionPathName.IsEmpty())
@@ -271,7 +289,7 @@ bool FFCPXMLImportVisitor::VisitVideoClipItemNode(TSharedRef<FFCPXMLClipItemNode
 	if (SectionData.IsValid())
 	{
 		// Update existing cinematic section
-		if (!ImportData->SetCinematicSection(SectionData, CurrVideoTrackRowIndex, FrameRate, Start, End, StartOffset))
+		if (!ImportData->SetCinematicSection(SectionData, CurrVideoTrackRowIndex, FrameRate, Start, End, NewStartOffset))
 		{
 			return false;
 		}
@@ -279,7 +297,7 @@ bool FFCPXMLImportVisitor::VisitVideoClipItemNode(TSharedRef<FFCPXMLClipItemNode
 	else
 	{
 		// Add new cinematic section
-		SectionData = ImportData->CreateCinematicSection(ClipItemName, CurrVideoTrackRowIndex, FrameRate, Start, End, StartOffset);
+		SectionData = ImportData->CreateCinematicSection(ClipItemName, CurrVideoTrackRowIndex, FrameRate, Start, End, NewStartOffset.IsSet() ? NewStartOffset.GetValue() : 0);
 		if (!SectionData.IsValid())
 		{
 			return false;
@@ -314,7 +332,7 @@ bool FFCPXMLImportVisitor::VisitAudioClipItemNode(TSharedRef<FFCPXMLClipItemNode
 	FFrameNumber StartOffset;
 	FFrameNumber Start;
 	FFrameNumber End;
-	
+
 	// Get relevant data from this clip item node
 	if (!GetClipItemNodeData(InClipItemNode, ClipItemName, ClipItemId, MasterClipId, LoggingInfoNode, LogNote, Filename, FrameRate, StartOffset, Start, End))
 	{
@@ -768,6 +786,30 @@ bool FFCPXMLImportVisitor::GetCinematicSectionPathNameFromMetadata(const FString
 	return ParseMetadata(InMetadata, TEXT("UE4ShotSection"), OutSectionObjectName);
 }
 
+/** Get sequencer shot handle frames from section metadata. Format is "[UE4ShotHandleFrames=handleframes]", whitespace ok. */
+bool FFCPXMLImportVisitor::GetCinematicSectionHandleFramesFromMetadata(const FString& InMetadata, int32& OutHandleFrames) const
+{
+	FString HandleFrameData;
+	bool bSuccess = ParseMetadata(InMetadata, TEXT("UE4ShotHandleFrames"), HandleFrameData);
+	if (bSuccess)
+	{
+		OutHandleFrames = FCString::Atoi(*HandleFrameData);
+	}
+	return bSuccess;
+}
+
+/** Get sequencer shot start offset frame from section metadata. Format is "[UE4ShotStartOffset=startoffset]", whitespace ok. */
+bool FFCPXMLImportVisitor::GetCinematicSectionStartOffsetFromMetadata(const FString& InMetadata, int32& OutStartOffset) const
+{
+	FString StartOffsetData;
+	bool bSuccess = ParseMetadata(InMetadata, TEXT("UE4ShotStartOffset"), StartOffsetData);
+	if (bSuccess)
+	{
+		OutStartOffset = FCString::Atoi(*StartOffsetData);
+	}
+	return bSuccess;
+}
+
 /** Get sequencer track name from track metadata. Format is "[UE4Track=trackobjectname][UE4Row=rowindex]", whitespace ok. */
 bool FFCPXMLImportVisitor::GetAudioFromMetadata(const FString& InMetadata, TSharedPtr<FFCPXMLImportAudioMetadata>& OutAudioMetadata) const
 {
diff --git a/Engine/Source/Editor/MovieSceneTools/Private/FCPXML/FCPXMLImport.h b/Engine/Source/Editor/MovieSceneTools/Private/FCPXML/FCPXMLImport.h
--- a/Engine/Source/Editor/MovieSceneTools/Private/FCPXML/FCPXMLImport.h
+++ b/Engine/Source/Editor/MovieSceneTools/Private/FCPXML/FCPXMLImport.h
@@ -110,9 +110,18 @@ private:
 	bool ParseMetadata(const FString& InMetadata, const FString& InKey, FString& OutValue) const;
 	/** Get sequencer section id from section metadata. Format is "[UE4ShotSection=sectionobjectname]", whitespace ok. */
 	bool GetCinematicSectionPathNameFromMetadata(const FString& InMetadata, FString& OutSectionObjectName) const;
+	/** Get sequencer shot handle frames from section metadata. Format is "[UE4ShotHandleFrames=handleframes]", whitespace ok. */
+	bool GetCinematicSectionHandleFramesFromMetadata(const FString& InMetadata, int32& OutHandleFrames) const;
+	/** Get sequencer shot start offset frame from section metadata. Format is "[UE4ShotStartOffset=startoffset]", whitespace ok. */
+	bool GetCinematicSectionStartOffsetFromMetadata(const FString& InMetadata, int32& OutStartOffset) const;
 	/** Get sequencer sound wave id and audio section ids from metadata. Format is "[UE4SoundWave=trackobjectname][UE4AudioSectionTopLevel=toplevelobjectname][UE4AudioSection=audiosectionobjectname]", whitespace ok. */
 	bool GetAudioFromMetadata(const FString& InMetadata, TSharedPtr<FFCPXMLImportAudioMetadata>& OutAudioMetadata) const;
 	
+public:
+
+	int32 GetMaxVideoTrackRowIndex() const { return MaxVideoTrackRowIndex; }
+	int32 GetMaxAudioTrackRowIndex() const { return MaxAudioTrackRowIndex; }
+
 private:
 	TSharedRef<FMovieSceneImportData> ImportData;
 	TSharedRef<FMovieSceneTranslatorContext> ImportContext;
@@ -141,5 +150,8 @@ private:
 	TSharedPtr<FMovieSceneImportAudioMasterTrackData> CurrAudioMasterTrack;
 	int32 CurrAudioTrackRowIndex;
 	bool bCurrImportAudioTrackIsStereoChannel;
+
+	int32 MaxVideoTrackRowIndex;
+	int32 MaxAudioTrackRowIndex;
 };
 
diff --git a/Engine/Source/Editor/MovieSceneTools/Private/FCPXML/FCPXMLMovieSceneTranslator.cpp b/Engine/Source/Editor/MovieSceneTools/Private/FCPXML/FCPXMLMovieSceneTranslator.cpp
--- a/Engine/Source/Editor/MovieSceneTools/Private/FCPXML/FCPXMLMovieSceneTranslator.cpp
+++ b/Engine/Source/Editor/MovieSceneTools/Private/FCPXML/FCPXMLMovieSceneTranslator.cpp
@@ -81,6 +81,38 @@ bool FFCPXMLImporter::Import(UMovieScene* InMovieScene, FFrameRate InFrameRate,
 		// Import the loaded Xml structure into the Sequencer movie scene
 		FFCPXMLImportVisitor ImportVisitor(ImportData, InContext);
 		bSuccess = FCPXMLFile->Accept(ImportVisitor);
+
+		// Flip rows
+		if (ImportData->GetCinematicMasterTrackData(false).IsValid())
+		{
+			int32 MaxVideoTrackRowIndex = ImportVisitor.GetMaxVideoTrackRowIndex();
+			for (TSharedPtr<FMovieSceneImportCinematicSectionData> SectionData : ImportData->GetCinematicMasterTrackData(false)->CinematicSections)
+			{
+				if (SectionData.IsValid())
+				{
+					UMovieSceneCinematicShotSection* CinematicSection = SectionData.Get()->CinematicSection;
+					if (CinematicSection)
+					{
+						CinematicSection->SetRowIndex(MaxVideoTrackRowIndex - CinematicSection->GetRowIndex());
+					}
+				}
+			}
+		}
+		if (ImportData->GetAudioMasterTrackData().IsValid())
+		{
+			int32 MaxAudioTrackRowIndex = ImportVisitor.GetMaxAudioTrackRowIndex();
+			for (TSharedPtr<FMovieSceneImportAudioSectionData> SectionData : ImportData->GetAudioMasterTrackData()->AudioSections)
+			{
+				if (SectionData.IsValid())
+				{
+					UMovieSceneAudioSection* AudioSection = SectionData.Get()->AudioSection;
+					if (AudioSection)
+					{
+						AudioSection->SetRowIndex(MaxAudioTrackRowIndex - AudioSection->GetRowIndex());
+					}
+				}
+			}
+		}
 	}
 
 	// add error message if one does not exist in the context
@@ -144,7 +176,7 @@ FText FFCPXMLExporter::GetMessageLogLabel() const
 	return LOCTEXT("FCPXMLExportLogLabel", "FCP 7 XML Export Log");
 }
 
-bool FFCPXMLExporter::Export(const UMovieScene* InMovieScene, FString InFilenameFormat, FFrameRate InFrameRate, uint32 InResX, uint32 InResY, int32 InHandleFrames, FString InSaveFilename, TSharedRef<FMovieSceneTranslatorContext> InContext)
+bool FFCPXMLExporter::Export(const UMovieScene* InMovieScene, FString InFilenameFormat, FFrameRate InFrameRate, uint32 InResX, uint32 InResY, int32 InHandleFrames, FString InSaveFilename, TSharedRef<FMovieSceneTranslatorContext> InContext, FString InMovieExtension)
 {
 	// add warning message if filename format is not "{shot}"
 	FString AcceptedFormat = TEXT("{shot}");
@@ -160,7 +192,7 @@ bool FFCPXMLExporter::Export(const UMovieScene* InMovieScene, FString InFilename
 	TSharedRef<FFCPXMLFile> FCPXMLFile = MakeShared<FFCPXMLFile>();
 	FCPXMLFile->ConstructFile(FPaths::GetBaseFilename(InSaveFilename, true));
 
-	TSharedRef<FMovieSceneExportData> ExportData = MakeShared<FMovieSceneExportData>(InMovieScene, InFrameRate, InResX, InResY, InHandleFrames, InSaveFilename, InContext);
+	TSharedRef<FMovieSceneExportData> ExportData = MakeShared<FMovieSceneExportData>(InMovieScene, InFrameRate, InResX, InResY, InHandleFrames, InSaveFilename, InContext, InMovieExtension);
 
 	// Export sequencer movie scene, merging with existing Xml structure.
 	FFCPXMLExportVisitor ExportVisitor(InSaveFilename, ExportData, InContext);
diff --git a/Engine/Source/Editor/MovieSceneTools/Private/MovieSceneToolHelpers.cpp b/Engine/Source/Editor/MovieSceneTools/Private/MovieSceneToolHelpers.cpp
--- a/Engine/Source/Editor/MovieSceneTools/Private/MovieSceneToolHelpers.cpp
+++ b/Engine/Source/Editor/MovieSceneTools/Private/MovieSceneToolHelpers.cpp
@@ -643,7 +643,7 @@ bool MovieSceneToolHelpers::ShowImportEDLDialog(UMovieScene* InMovieScene, FFram
 	return MovieSceneTranslatorEDL::ImportEDL(InMovieScene, InFrameRate, OpenFilenames[0]);
 }
 
-bool MovieSceneToolHelpers::ShowExportEDLDialog(const UMovieScene* InMovieScene, FFrameRate InFrameRate, FString InSaveDirectory, int32 InHandleFrames)
+bool MovieSceneToolHelpers::ShowExportEDLDialog(const UMovieScene* InMovieScene, FFrameRate InFrameRate, FString InSaveDirectory, int32 InHandleFrames, FString InMovieExtension)
 {
 	TArray<FString> SaveFilenames;
 	FString SequenceName = InMovieScene->GetOuter()->GetName();
@@ -677,7 +677,7 @@ bool MovieSceneToolHelpers::ShowExportEDLDialog(const UMovieScene* InMovieScene,
 		return false;
 	}
 
-	if (MovieSceneTranslatorEDL::ExportEDL(InMovieScene, InFrameRate, SaveFilenames[0], InHandleFrames))
+	if (MovieSceneTranslatorEDL::ExportEDL(InMovieScene, InFrameRate, SaveFilenames[0], InHandleFrames, InMovieExtension))
 	{
 		const FString AbsoluteFilename = FPaths::ConvertRelativePathToFull(SaveFilenames[0]);
 		const FString SaveDirectory = FPaths::GetPath(AbsoluteFilename);
@@ -753,6 +753,7 @@ bool MovieSceneToolHelpers::MovieSceneTranslatorExport(FMovieSceneExporter* InEx
 	FFrameRate FrameRate = Settings.FrameRate;
 	uint32 ResX = Settings.Resolution.ResX;
 	uint32 ResY = Settings.Resolution.ResY;
+	FString MovieExtension = Settings.MovieExtension;
 
 	TArray<FString> SaveFilenames;
 	FString SequenceName = InMovieScene->GetOuter()->GetName();
@@ -785,7 +786,7 @@ bool MovieSceneToolHelpers::MovieSceneTranslatorExport(FMovieSceneExporter* InEx
 	TSharedRef<FMovieSceneTranslatorContext> ExportContext(new FMovieSceneTranslatorContext);
 	ExportContext->Init();
 
-	bool bSuccess = InExporter->Export(InMovieScene, FilenameFormat, FrameRate, ResX, ResY, HandleFrames, SaveFilenames[0], ExportContext);
+	bool bSuccess = InExporter->Export(InMovieScene, FilenameFormat, FrameRate, ResX, ResY, HandleFrames, SaveFilenames[0], ExportContext, MovieExtension);
 	
 	// Display any messages in context
 	MovieSceneTranslatorLogMessages(InExporter, ExportContext, true);
@@ -1173,7 +1174,7 @@ FbxNode* RetrieveObjectFromName(const TCHAR* ObjectName, FbxNode* Root)
 	return nullptr;
 }
 
-void CopyCameraProperties(FbxCamera* CameraNode, ACineCameraActor* CameraActor)
+void CopyCameraProperties(FbxCamera* CameraNode, AActor* InCameraActor)
 {
 	float FieldOfView;
 	float FocalLength;
@@ -1192,26 +1193,53 @@ void CopyCameraProperties(FbxCamera* CameraNode, ACineCameraActor* CameraActor)
 	float ApertureWidth = CameraNode->GetApertureWidth();
 	float ApertureHeight = CameraNode->GetApertureHeight();
 
-	UCineCameraComponent* CineCameraComponent = CameraActor->GetCineCameraComponent();
+	UCameraComponent* CameraComponent = nullptr;
 
-	CineCameraComponent->SetProjectionMode(CameraNode->ProjectionType.Get() == FbxCamera::ePerspective ? ECameraProjectionMode::Perspective : ECameraProjectionMode::Orthographic);
-	CineCameraComponent->SetAspectRatio(CameraNode->AspectWidth.Get() / CameraNode->AspectHeight.Get());
-	CineCameraComponent->SetOrthoNearClipPlane(CameraNode->NearPlane.Get());
-	CineCameraComponent->SetOrthoFarClipPlane(CameraNode->FarPlane.Get());
-	CineCameraComponent->SetOrthoWidth(CameraNode->OrthoZoom.Get());
-	CineCameraComponent->SetFieldOfView(FieldOfView);
-	CineCameraComponent->FilmbackSettings.SensorWidth = FUnitConversion::Convert(ApertureWidth, EUnit::Inches, EUnit::Millimeters);
-	CineCameraComponent->FilmbackSettings.SensorHeight = FUnitConversion::Convert(ApertureHeight, EUnit::Inches, EUnit::Millimeters);
-	CineCameraComponent->FocusSettings.ManualFocusDistance = CameraNode->FocusDistance;
-	if (FocalLength < CineCameraComponent->LensSettings.MinFocalLength)
+	if (ACineCameraActor* CineCameraActor = Cast<ACineCameraActor>(InCameraActor))
 	{
-		CineCameraComponent->LensSettings.MinFocalLength = FocalLength;
+		CameraComponent = CineCameraActor->GetCineCameraComponent();
+
+		UCineCameraComponent* CineCameraComponent = CineCameraActor->GetCineCameraComponent();
+		CineCameraComponent->FilmbackSettings.SensorWidth = FUnitConversion::Convert(ApertureWidth, EUnit::Inches, EUnit::Millimeters);
+		CineCameraComponent->FilmbackSettings.SensorHeight = FUnitConversion::Convert(ApertureHeight, EUnit::Inches, EUnit::Millimeters);
+		CineCameraComponent->FocusSettings.ManualFocusDistance = CameraNode->FocusDistance;
+		if (FocalLength < CineCameraComponent->LensSettings.MinFocalLength)
+		{
+			CineCameraComponent->LensSettings.MinFocalLength = FocalLength;
+		}
+		if (FocalLength > CineCameraComponent->LensSettings.MaxFocalLength)
+		{
+			CineCameraComponent->LensSettings.MaxFocalLength = FocalLength;
+		}
+		CineCameraComponent->CurrentFocalLength = FocalLength;
+	}
+	else if (ACameraActor* CameraActor = Cast<ACameraActor>(InCameraActor))
+	{
+		CameraComponent = CameraActor->GetCameraComponent();
+	}
+
+	if (!CameraComponent)
+	{
+		return;
 	}
-	if (FocalLength > CineCameraComponent->LensSettings.MaxFocalLength)
+
+	CameraComponent->SetProjectionMode(CameraNode->ProjectionType.Get() == FbxCamera::ePerspective ? ECameraProjectionMode::Perspective : ECameraProjectionMode::Orthographic);
+	CameraComponent->SetAspectRatio(CameraNode->AspectWidth.Get() / CameraNode->AspectHeight.Get());
+	CameraComponent->SetOrthoNearClipPlane(CameraNode->NearPlane.Get());
+	CameraComponent->SetOrthoFarClipPlane(CameraNode->FarPlane.Get());
+	CameraComponent->SetOrthoWidth(CameraNode->OrthoZoom.Get());
+	CameraComponent->SetFieldOfView(FieldOfView);
+}
+
+FString GetCameraName(FbxCamera* InCamera)
+{
+	FbxNode* CameraNode = InCamera->GetNode();
+	if (CameraNode)
 	{
-		CineCameraComponent->LensSettings.MaxFocalLength = FocalLength;
+		return CameraNode->GetName();
 	}
-	CineCameraComponent->CurrentFocalLength = FocalLength;
+
+	return InCamera->GetName();
 }
 
 void ImportFBXCamera(UnFbx::FFbxImporter* FbxImporter, UMovieScene* InMovieScene, ISequencer& InSequencer, TMap<FGuid, FString>& InObjectBindingMap, bool bMatchByNameOnly, bool bCreateCameras)
@@ -1227,13 +1255,13 @@ void ImportFBXCamera(UnFbx::FFbxImporter* FbxImporter, UMovieScene* InMovieScene
 		TArray<FbxCamera*> UnmatchedCameras;
 		for (auto Camera : AllCameras)
 		{
-			FString NodeName = FString(Camera->GetName());
+			FString NodeName = GetCameraName(Camera);
 
 			bool bMatched = false;
 			for (auto InObjectBinding : InObjectBindingMap)
 			{		
 				FString ObjectName = InObjectBinding.Value;
-				if (!FCString::Strcmp(*ObjectName, UTF8_TO_TCHAR(Camera->GetName())))
+				if (ObjectName == NodeName)
 				{
 					// Look for a valid bound object, otherwise need to create a new camera and assign this binding to it
 					bool bFoundBoundObject = false;
@@ -1272,11 +1300,21 @@ void ImportFBXCamera(UnFbx::FFbxImporter* FbxImporter, UMovieScene* InMovieScene
 		// Add any unmatched cameras
 		for (auto UnmatchedCamera : UnmatchedCameras)
 		{
-			FString CameraName = FString(ANSI_TO_TCHAR(UnmatchedCamera->GetName()));
+			FString CameraName = GetCameraName(UnmatchedCamera);
 
-			FActorSpawnParameters SpawnParams;
-			ACineCameraActor* NewCamera = World->SpawnActor<ACineCameraActor>(SpawnParams);
-			NewCamera->SetActorLabel(*CameraName);
+			AActor* NewCamera = nullptr;
+			if (UnmatchedCamera->GetApertureMode() == FbxCamera::eFocalLength)
+			{
+				FActorSpawnParameters SpawnParams;
+				NewCamera = World->SpawnActor<ACineCameraActor>(SpawnParams);
+				NewCamera->SetActorLabel(*CameraName);
+			}
+			else
+			{
+				FActorSpawnParameters SpawnParams;
+				NewCamera = World->SpawnActor<ACameraActor>(SpawnParams);
+				NewCamera->SetActorLabel(*CameraName);
+			}
 
 			// Copy camera properties before adding default tracks so that initial camera properties match and can be restored after sequencer finishes
 			CopyCameraProperties(UnmatchedCamera, NewCamera);
@@ -1320,7 +1358,7 @@ void ImportFBXCamera(UnFbx::FFbxImporter* FbxImporter, UMovieScene* InMovieScene
 			CameraNode = FindCamera(FbxImporter->Scene->GetRootNode());
 			if (CameraNode)
 			{
-				FString CameraName = FString(ANSI_TO_TCHAR(CameraNode->GetName()));
+				FString CameraName = GetCameraName(CameraNode);
 				FNotificationInfo Info(FText::Format(NSLOCTEXT("MovieSceneTools", "NoMatchingCameraWarning", "Failed to find any matching camera for {0}. Importing onto first camera from fbx {1}"), FText::FromString(ObjectName), FText::FromString(CameraName)));
 				Info.ExpireDuration = 5.0f;
 				FSlateNotificationManager::Get().AddNotification(Info)->SetCompletionState(SNotificationItem::CS_Fail);
@@ -1349,20 +1387,39 @@ void ImportFBXCamera(UnFbx::FFbxImporter* FbxImporter, UMovieScene* InMovieScene
 		for (TWeakObjectPtr<>& WeakObject : BoundObjects)
 		{
 			UObject* FoundObject = WeakObject.Get();
-			if (FoundObject && FoundObject->IsA(ACineCameraActor::StaticClass()))
+			if (FoundObject && FoundObject->GetClass()->IsChildOf(ACameraActor::StaticClass()))
 			{
-				ACineCameraActor* CineCameraActor = Cast<ACineCameraActor>(FoundObject);
-				UCineCameraComponent* CineCameraComponent = CineCameraActor->GetCineCameraComponent();
-				CopyCameraProperties(CameraNode, CineCameraActor);
+				CopyCameraProperties(CameraNode, Cast<AActor>(FoundObject));
+
+				UCameraComponent* CameraComponent = nullptr;
+				FName TrackName;
+				float TrackValue;
+				
+				if (ACineCameraActor* CineCameraActor = Cast<ACineCameraActor>(FoundObject))
+				{
+					CameraComponent = CineCameraActor->GetCineCameraComponent();
+					TrackName = TEXT("CurrentFocalLength");
+					TrackValue = FocalLength;
+				}
+				else if (ACameraActor* CameraActor = Cast<ACameraActor>(FoundObject))
+				{
+					CameraComponent = CameraActor->GetCameraComponent();
+					TrackName = TEXT("FieldOfView");
+					TrackValue = FieldOfView;
+				}
+				else
+				{
+					continue;
+				}
 
-				// Set the default value of the current focal length section
-				FGuid PropertyOwnerGuid = InSequencer.GetHandleToObject(CineCameraComponent);
+				// Set the default value of the current focal length or field of view section
+				FGuid PropertyOwnerGuid = InSequencer.GetHandleToObject(CameraComponent);
 				if (!PropertyOwnerGuid.IsValid())
 				{
 					continue;
 				}
 
-				UMovieSceneFloatTrack* FloatTrack = InMovieScene->FindTrack<UMovieSceneFloatTrack>(PropertyOwnerGuid, TEXT("CurrentFocalLength"));
+				UMovieSceneFloatTrack* FloatTrack = InMovieScene->FindTrack<UMovieSceneFloatTrack>(PropertyOwnerGuid, TrackName);
 				if (FloatTrack)
 				{
 					FloatTrack->Modify();
@@ -1382,7 +1439,7 @@ void ImportFBXCamera(UnFbx::FFbxImporter* FbxImporter, UMovieScene* InMovieScene
 						FloatSection->SetRange(TRange<FFrameNumber>::All());
 					}
 
-					FloatSection->GetChannelProxy().GetChannel<FMovieSceneFloatChannel>(0)->SetDefault(FocalLength);
+					FloatSection->GetChannelProxy().GetChannel<FMovieSceneFloatChannel>(0)->SetDefault(TrackValue);
 				}
 			}
 		}
@@ -1391,9 +1448,11 @@ void ImportFBXCamera(UnFbx::FFbxImporter* FbxImporter, UMovieScene* InMovieScene
 
 FGuid FindCameraGuid(FbxCamera* Camera, TMap<FGuid, FString>& InObjectBindingMap)
 {
+	FString CameraName = GetCameraName(Camera);
+
 	for (auto& Pair : InObjectBindingMap)
 	{
-		if (FCString::Strcmp(*Pair.Value, UTF8_TO_TCHAR(Camera->GetName())) == 0)
+		if (Pair.Value == CameraName)
 		{
 			return Pair.Key;
 		}
@@ -1575,10 +1634,18 @@ private:
 		TArray<FString> AllNodeNames;
 		CurveAPI.GetAllNodeNameArray(AllNodeNames);
 
+		FString RootNodeName = FbxImporter->Scene->GetRootNode()->GetName();
+
 		// First try matching by name
 		for (int32 NodeIndex = 0; NodeIndex < AllNodeNames.Num(); )
 		{
 			FString NodeName = AllNodeNames[NodeIndex];
+			if (RootNodeName == NodeName)
+			{
+				++NodeIndex;
+				continue;
+			}
+
 			bool bFoundMatch = false;
 			for (auto It = ObjectBindingMap.CreateConstIterator(); It; ++It)
 			{
@@ -1608,6 +1675,12 @@ private:
 			for (int32 NodeIndex = 0; NodeIndex < AllNodeNames.Num(); )
 			{
 				FString NodeName = AllNodeNames[NodeIndex];
+				if (RootNodeName == NodeName)
+				{
+					++NodeIndex;
+					continue;
+				}
+
 				auto It = ObjectBindingMap.CreateConstIterator();
 				if (It)
 				{
diff --git a/Engine/Source/Editor/MovieSceneTools/Private/MovieSceneTranslator.cpp b/Engine/Source/Editor/MovieSceneTools/Private/MovieSceneTranslator.cpp
--- a/Engine/Source/Editor/MovieSceneTools/Private/MovieSceneTranslator.cpp
+++ b/Engine/Source/Editor/MovieSceneTools/Private/MovieSceneTranslator.cpp
@@ -38,7 +38,7 @@ namespace
 	}
 }
 
-FMovieSceneExportData::FMovieSceneExportData(const UMovieScene* InMovieScene, FFrameRate InFrameRate, uint32 InResX, uint32 InResY, int32 InHandleFrames, FString InSaveFilename, TSharedPtr<FMovieSceneTranslatorContext> InContext)
+FMovieSceneExportData::FMovieSceneExportData(const UMovieScene* InMovieScene, FFrameRate InFrameRate, uint32 InResX, uint32 InResY, int32 InHandleFrames, FString InSaveFilename, TSharedPtr<FMovieSceneTranslatorContext> InContext, FString InMovieExtension)
 {
 	if (InMovieScene == nullptr)
 	{
@@ -50,8 +50,9 @@ FMovieSceneExportData::FMovieSceneExportData(const UMovieScene* InMovieScene, FF
 	FrameRate = InFrameRate;
 	ResX = InResX;
 	ResY = InResY;
-	HandleFrames = ConvertFrameTime(FFrameTime(InHandleFrames), InMovieScene->GetTickResolution(), FrameRate);
+	HandleFrames = InHandleFrames;
 	SaveFilename = InSaveFilename;
+	MovieExtension = InMovieExtension;
 
 	// preferred sample rate in UE4
 	DefaultAudioSampleRate = 44100;
@@ -344,16 +345,7 @@ bool FMovieSceneExportData::ConstructCinematicSectionData(const UMovieScene* InM
 
 	SectionData->DisplayName = InCinematicSection->GetShotDisplayName();
 
-#if PLATFORM_MAC
-	static const TCHAR* Extension = TEXT(".mov");
-#elif PLATFORM_UNIX
-	static const TCHAR* Extension = TEXT(".unsupp");
-	return false;
-#else
-	static const TCHAR* Extension = TEXT(".avi");
-#endif
-
-	SectionData->SourceFilename = SectionData->DisplayName + Extension;
+	SectionData->SourceFilename = SectionData->DisplayName + MovieExtension;
 	SectionData->SourceFilePath = TEXT("");
 
 	ConstructSectionData(InMovieScene, SectionData, InCinematicSection, EMovieSceneTranslatorSectionType::Cinematic, SectionData->DisplayName);
@@ -524,6 +516,11 @@ FString FMovieSceneExportData::GetFilenamePath() const
 	return SaveFilenamePath;
 }
 
+FString FMovieSceneExportData::GetMovieExtension() const
+{
+	return MovieExtension;
+}
+
 FFrameRate FMovieSceneExportData::GetFrameRate() const
 {
 	return FrameRate;
@@ -555,7 +552,7 @@ bool FMovieSceneExportData::GetFrameRateIsNTSC() const
 	return (!FMath::IsNearlyZero(FractionalPart));
 }
 
-FFrameTime FMovieSceneExportData::GetHandleFrames() const
+int32 FMovieSceneExportData::GetHandleFrames() const
 {
 	return HandleFrames;
 }
@@ -692,7 +689,7 @@ TSharedPtr<FMovieSceneImportCinematicSectionData> FMovieSceneImportData::CreateC
 	return SectionData;
 }
 
-bool FMovieSceneImportData::SetCinematicSection(TSharedPtr<FMovieSceneImportCinematicSectionData> InSection, int32 InRow, FFrameRate InFrameRate, FFrameNumber InStartFrame, FFrameNumber InEndFrame, FFrameNumber InStartOffsetFrame)
+bool FMovieSceneImportData::SetCinematicSection(TSharedPtr<FMovieSceneImportCinematicSectionData> InSection, int32 InRow, FFrameRate InFrameRate, FFrameNumber InStartFrame, FFrameNumber InEndFrame, TOptional<FFrameNumber> InStartOffsetFrame)
 {
 	if (!InSection.IsValid() || InSection->CinematicSection == nullptr)
 	{
@@ -701,11 +698,14 @@ bool FMovieSceneImportData::SetCinematicSection(TSharedPtr<FMovieSceneImportCine
 
 	FFrameRate TickResolution = MovieSceneData->MovieScene->GetTickResolution();
 	FFrameNumber StartFrame = ConvertFrameTime(InStartFrame, InFrameRate, TickResolution).GetFrame();
-	FFrameNumber StartOffsetFrame = ConvertFrameTime(InStartOffsetFrame, InFrameRate, TickResolution).GetFrame();
 	FFrameNumber EndFrame = ConvertFrameTime(InEndFrame, InFrameRate, TickResolution).GetFrame();
-
+		
 	InSection->CinematicSection->Modify();
-	InSection->CinematicSection->Parameters.StartFrameOffset = StartOffsetFrame;
+	if (InStartOffsetFrame.IsSet())
+	{
+		FFrameNumber StartOffsetFrame = ConvertFrameTime(InStartOffsetFrame.GetValue(), InFrameRate, TickResolution).GetFrame();
+		InSection->CinematicSection->Parameters.StartFrameOffset = StartOffsetFrame;
+	}
 	InSection->CinematicSection->SetRange(TRange<FFrameNumber>(StartFrame, EndFrame));
 	if (InRow != InSection->CinematicSection->GetRowIndex())
 	{
diff --git a/Engine/Source/Editor/MovieSceneTools/Private/MovieSceneTranslatorEDL.cpp b/Engine/Source/Editor/MovieSceneTools/Private/MovieSceneTranslatorEDL.cpp
--- a/Engine/Source/Editor/MovieSceneTools/Private/MovieSceneTranslatorEDL.cpp
+++ b/Engine/Source/Editor/MovieSceneTools/Private/MovieSceneTranslatorEDL.cpp
@@ -432,7 +432,7 @@ bool MovieSceneTranslatorEDL::ImportEDL(UMovieScene* InMovieScene, FFrameRate In
 	return true;
 }
 
-bool MovieSceneTranslatorEDL::ExportEDL(const UMovieScene* InMovieScene, FFrameRate InFrameRate, FString InSaveFilename, const int32 InHandleFrames)
+bool MovieSceneTranslatorEDL::ExportEDL(const UMovieScene* InMovieScene, FFrameRate InFrameRate, FString InSaveFilename, const int32 InHandleFrames, FString InMovieExtension)
 {
 	FString SequenceName = InMovieScene->GetOuter()->GetName();
 	FString SaveBasename = FPaths::GetPath(InSaveFilename) / FPaths::GetBaseFilename(InSaveFilename);
@@ -459,31 +459,25 @@ bool MovieSceneTranslatorEDL::ExportEDL(const UMovieScene* InMovieScene, FFrameR
 			{
 				UMovieSceneCinematicShotSection* CinematicShotSection = Cast<UMovieSceneCinematicShotSection>(ShotSection);
 
-				if ( CinematicShotSection->GetSequence() == nullptr || !CinematicShotSection->HasStartFrame() || !CinematicShotSection->HasEndFrame() )
+				if ( CinematicShotSection->GetSequence() == nullptr || !CinematicShotSection->HasStartFrame() || !CinematicShotSection->HasEndFrame() || !CinematicShotSection->IsActive())
 				{
 					// If the shot doesn't have a valid sequence skip it.  This is currently the case for filler sections.
 					// TODO: Handle this properly in the edl output.
 					continue;
 				}
 
+				FFrameRate TickResolution = CinematicShotSection->GetSequence()->GetMovieScene()->GetTickResolution();
+
 				FString ShotName = CinematicShotSection->GetShotDisplayName();
 				FString ShotPath = CinematicShotSection->GetSequence()->GetMovieScene()->GetOuter()->GetPathName();
 
-				FFrameNumber SourceInFrame  = InHandleFrames;
-				FFrameNumber SourceOutFrame = InHandleFrames + MovieScene::DiscreteSize(CinematicShotSection->GetRange());
+				FFrameNumber SourceInFrame = ConvertFrameTime(InHandleFrames + 1, InFrameRate, TickResolution).FrameNumber;
+				FFrameNumber SourceOutFrame = ConvertFrameTime(InHandleFrames, InFrameRate, TickResolution).FrameNumber + MovieScene::DiscreteSize(CinematicShotSection->GetRange());
 
 				FFrameNumber EditInFrame    = CinematicShotSection->GetInclusiveStartFrame();
 				FFrameNumber EditOutFrame   = CinematicShotSection->GetExclusiveEndFrame();
 
-#if PLATFORM_MAC
-				static const TCHAR* Extension = TEXT(".mov");
-#elif PLATFORM_UNIX
-				static const TCHAR* Extension = TEXT(".unsupp");
-				continue;
-#else
-				static const TCHAR* Extension = TEXT(".avi");
-#endif
-				ShotName += Extension;
+				ShotName += InMovieExtension;
 
 				//@todo shotpath should really be moviefile path
 				ShotPath = ShotName;
diff --git a/Engine/Source/Editor/MovieSceneTools/Private/TrackEditors/CinematicShotTrackEditor.cpp b/Engine/Source/Editor/MovieSceneTools/Private/TrackEditors/CinematicShotTrackEditor.cpp
--- a/Engine/Source/Editor/MovieSceneTools/Private/TrackEditors/CinematicShotTrackEditor.cpp
+++ b/Engine/Source/Editor/MovieSceneTools/Private/TrackEditors/CinematicShotTrackEditor.cpp
@@ -525,8 +525,8 @@ void FCinematicShotTrackEditor::SwitchTake(uint32 TakeNumber)
 		if (TakeObject && TakeObject->IsA(UMovieSceneSequence::StaticClass()))
 		{
 			UMovieSceneSequence* MovieSceneSequence = CastChecked<UMovieSceneSequence>(TakeObject);
-
-			UMovieSceneCinematicShotTrack* CinematicShotTrack = FindOrCreateCinematicShotTrack();
+			
+			UMovieSceneCinematicShotTrack* CinematicShotTrack = CastChecked<UMovieSceneCinematicShotTrack>(Section->GetOuter());
 
 			TRange<FFrameNumber> NewShotRange         = Section->GetRange();
 			FFrameNumber         NewShotStartOffset   = Section->Parameters.StartFrameOffset;
@@ -658,9 +658,10 @@ FKeyPropertyResult FCinematicShotTrackEditor::AddKeyInternal(FFrameNumber KeyTim
 	{
 		UMovieSceneCinematicShotTrack* CinematicShotTrack = FindOrCreateCinematicShotTrack();
 		
+		const FFrameRate TickResolution = InMovieSceneSequence->GetMovieScene()->GetTickResolution();
 		const FQualifiedFrameTime InnerDuration = FQualifiedFrameTime(
 			MovieScene::DiscreteSize(InMovieSceneSequence->GetMovieScene()->GetPlaybackRange()),
-			InMovieSceneSequence->GetMovieScene()->GetTickResolution());
+			TickResolution);
 
 		const FFrameRate OuterFrameRate = CinematicShotTrack->GetTypedOuter<UMovieScene>()->GetTickResolution();
 		const int32      OuterDuration  = InnerDuration.ConvertTo(OuterFrameRate).FrameNumber.Value;
@@ -672,6 +673,13 @@ FKeyPropertyResult FCinematicShotTrackEditor::AddKeyInternal(FFrameNumber KeyTim
 		GetSequencer()->SelectSection(NewSection);
 		GetSequencer()->ThrobSectionSelection();
 
+		if (TickResolution != OuterFrameRate)
+		{
+			FNotificationInfo Info(FText::Format(LOCTEXT("TickResolutionMismatch", "The parent sequence has a different tick resolution {0} than the newly added sequence {1}"), OuterFrameRate.ToPrettyText(), TickResolution.ToPrettyText()));
+			Info.bUseLargeFont = false;
+			FSlateNotificationManager::Get().AddNotification(Info);
+		}
+
 		return KeyPropertyResult;
 	}
 
@@ -811,9 +819,10 @@ FKeyPropertyResult FCinematicShotTrackEditor::HandleSequenceAdded(FFrameNumber K
 
 	auto CinematicShotTrack = FindOrCreateCinematicShotTrack();
 
+	const FFrameRate TickResolution = Sequence->GetMovieScene()->GetTickResolution();
 	const FQualifiedFrameTime InnerDuration = FQualifiedFrameTime(
 		MovieScene::DiscreteSize(Sequence->GetMovieScene()->GetPlaybackRange()),
-		Sequence->GetMovieScene()->GetTickResolution());
+		TickResolution);
 
 	const FFrameRate OuterFrameRate = CinematicShotTrack->GetTypedOuter<UMovieScene>()->GetTickResolution();
 	const int32      OuterDuration  = InnerDuration.ConvertTo(OuterFrameRate).FrameNumber.Value;
@@ -825,6 +834,13 @@ FKeyPropertyResult FCinematicShotTrackEditor::HandleSequenceAdded(FFrameNumber K
 	GetSequencer()->SelectSection(NewSection);
 	GetSequencer()->ThrobSectionSelection();
 
+	if (TickResolution != OuterFrameRate)
+	{
+		FNotificationInfo Info(FText::Format(LOCTEXT("TickResolutionMismatch", "The parent sequence has a different tick resolution {0} than the newly added sequence {1}"), OuterFrameRate.ToPrettyText(), TickResolution.ToPrettyText()));
+		Info.bUseLargeFont = false;
+		FSlateNotificationManager::Get().AddNotification(Info);
+	}
+
 	return KeyPropertyResult;
 }
 
@@ -838,7 +854,7 @@ UAutomatedLevelSequenceCapture* GetMovieSceneCapture()
 	
 	if (!MovieSceneCapture)
 	{
-		MovieSceneCapture = NewObject<UAutomatedLevelSequenceCapture>(GetTransientPackage(), UAutomatedLevelSequenceCapture::StaticClass(), UAutomatedLevelSequenceCapture::AutomatedLevelSequenceCaptureUIName, RF_Transient);
+		MovieSceneCapture = NewObject<UAutomatedLevelSequenceCapture>(GetTransientPackage(), UAutomatedLevelSequenceCapture::StaticClass(), UMovieSceneCapture::MovieSceneCaptureUIName, RF_Transient);
 		MovieSceneCapture->LoadFromConfig();
 	}
 
@@ -898,8 +914,9 @@ void FCinematicShotTrackEditor::ExportEDL()
 	const FMovieSceneCaptureSettings& Settings = MovieSceneCapture->GetSettings();
 	FString SaveDirectory = FPaths::ConvertRelativePathToFull(Settings.OutputDirectory.Path);
 	int32 HandleFrames = Settings.HandleFrames;
+	FString MovieExtension = Settings.MovieExtension;
 
-	MovieSceneToolHelpers::ShowExportEDLDialog(MovieScene, MovieScene->GetDisplayRate(), SaveDirectory, HandleFrames);
+	MovieSceneToolHelpers::ShowExportEDLDialog(MovieScene, MovieScene->GetDisplayRate(), SaveDirectory, HandleFrames, MovieExtension);
 }
 
 
diff --git a/Engine/Source/Editor/MovieSceneTools/Private/TrackEditors/MaterialParameterCollectionTrackEditor.cpp b/Engine/Source/Editor/MovieSceneTools/Private/TrackEditors/MaterialParameterCollectionTrackEditor.cpp
--- a/Engine/Source/Editor/MovieSceneTools/Private/TrackEditors/MaterialParameterCollectionTrackEditor.cpp
+++ b/Engine/Source/Editor/MovieSceneTools/Private/TrackEditors/MaterialParameterCollectionTrackEditor.cpp
@@ -69,8 +69,11 @@ void FMaterialParameterCollectionTrackEditor::BuildTrackContextMenu(FMenuBuilder
 		{
 			FScopedTransaction Transaction(LOCTEXT("SetAssetTransaction", "Assign Material Parameter Collection"));
 			MPCTrack->Modify();
+			MPCTrack->SetDisplayName(FText::FromString(MPC->GetName()));
 			MPCTrack->MPC = MPC;
 		}
+
+		FSlateApplication::Get().DismissAllMenus();
 	};
 
 	auto AssignAssetEnterPressed = [AssignAsset](const TArray<FAssetData>& InAssetData)
diff --git a/Engine/Source/Editor/MovieSceneTools/Private/TrackEditors/SubTrackEditor.cpp b/Engine/Source/Editor/MovieSceneTools/Private/TrackEditors/SubTrackEditor.cpp
--- a/Engine/Source/Editor/MovieSceneTools/Private/TrackEditors/SubTrackEditor.cpp
+++ b/Engine/Source/Editor/MovieSceneTools/Private/TrackEditors/SubTrackEditor.cpp
@@ -730,9 +730,10 @@ FKeyPropertyResult FSubTrackEditor::AddKeyInternal(FFrameNumber KeyTime, UMovieS
 	{
 		UMovieSceneSubTrack* SubTrack = Cast<UMovieSceneSubTrack>(InTrack);
 
+		const FFrameRate TickResolution = InMovieSceneSequence->GetMovieScene()->GetTickResolution();
 		const FQualifiedFrameTime InnerDuration = FQualifiedFrameTime(
 			MovieScene::DiscreteSize(InMovieSceneSequence->GetMovieScene()->GetPlaybackRange()),
-			InMovieSceneSequence->GetMovieScene()->GetTickResolution());
+			TickResolution);
 
 		const FFrameRate OuterFrameRate = SubTrack->GetTypedOuter<UMovieScene>()->GetTickResolution();
 		const int32      OuterDuration  = InnerDuration.ConvertTo(OuterFrameRate).FrameNumber.Value;
@@ -744,6 +745,13 @@ FKeyPropertyResult FSubTrackEditor::AddKeyInternal(FFrameNumber KeyTime, UMovieS
 		GetSequencer()->SelectSection(NewSection);
 		GetSequencer()->ThrobSectionSelection();
 
+		if (TickResolution != OuterFrameRate)
+		{
+			FNotificationInfo Info(FText::Format(LOCTEXT("TickResolutionMismatch", "The parent sequence has a different tick resolution {0} than the newly added sequence {1}"), OuterFrameRate.ToPrettyText(), TickResolution.ToPrettyText()));
+			Info.bUseLargeFont = false;
+			FSlateNotificationManager::Get().AddNotification(Info);
+		}
+
 		return KeyPropertyResult;
 	}
 
@@ -760,9 +768,10 @@ FKeyPropertyResult FSubTrackEditor::HandleSequenceAdded(FFrameNumber KeyTime, UM
 
 	auto SubTrack = FindOrCreateMasterTrack<UMovieSceneSubTrack>().Track;
 
+	const FFrameRate TickResolution = Sequence->GetMovieScene()->GetTickResolution();
 	const FQualifiedFrameTime InnerDuration = FQualifiedFrameTime(
 		MovieScene::DiscreteSize(Sequence->GetMovieScene()->GetPlaybackRange()),
-		Sequence->GetMovieScene()->GetTickResolution());
+		TickResolution);
 
 	const FFrameRate OuterFrameRate = SubTrack->GetTypedOuter<UMovieScene>()->GetTickResolution();
 	const int32      OuterDuration  = InnerDuration.ConvertTo(OuterFrameRate).FrameNumber.Value;
@@ -774,6 +783,13 @@ FKeyPropertyResult FSubTrackEditor::HandleSequenceAdded(FFrameNumber KeyTime, UM
 	GetSequencer()->SelectSection(NewSection);
 	GetSequencer()->ThrobSectionSelection();
 
+	if (TickResolution != OuterFrameRate)
+	{
+		FNotificationInfo Info(FText::Format(LOCTEXT("TickResolutionMismatch", "The parent sequence has a different tick resolution {0} than the newly added sequence {1}"), OuterFrameRate.ToPrettyText(), TickResolution.ToPrettyText()));
+		Info.bUseLargeFont = false;
+		FSlateNotificationManager::Get().AddNotification(Info);
+	}
+
 	return KeyPropertyResult;
 }
 
@@ -833,8 +849,8 @@ void FSubTrackEditor::SwitchTake(uint32 TakeNumber)
 		{
 			UMovieSceneSequence* MovieSceneSequence = CastChecked<UMovieSceneSequence>(TakeObject);
 
-			UMovieSceneSubTrack* SubTrack = FindOrCreateMasterTrack<UMovieSceneSubTrack>().Track;
-		
+			UMovieSceneSubTrack* SubTrack = CastChecked<UMovieSceneSubTrack>(Section->GetOuter());
+
 			TRange<FFrameNumber> NewShotRange         = Section->GetRange();
 			FFrameNumber		 NewShotStartOffset   = Section->Parameters.StartFrameOffset;
 			float                NewShotTimeScale     = Section->Parameters.TimeScale;
diff --git a/Engine/Source/Editor/MovieSceneTools/Public/FCPXML/FCPXMLMovieSceneTranslator.h b/Engine/Source/Editor/MovieSceneTools/Public/FCPXML/FCPXMLMovieSceneTranslator.h
--- a/Engine/Source/Editor/MovieSceneTools/Public/FCPXML/FCPXMLMovieSceneTranslator.h
+++ b/Engine/Source/Editor/MovieSceneTools/Public/FCPXML/FCPXMLMovieSceneTranslator.h
@@ -82,7 +82,8 @@ public:
 	 * @param InHandleFrames The number of handle frames to include for each shot.
 	 * @param InSaveFilename The file path to save to.
 	 * @param OutError The return error message
+	 * @param MovieExtension The movie extension for the shot filenames (ie. .avi, .mov, .mp4)
 	 * @return Whether the export was successful
 	 */
-	virtual bool Export(const UMovieScene* InMovieScene, FString InFilenameFormat, FFrameRate InFrameRate, uint32 InResX, uint32 InResY, int32 InHandleFrames, FString InSaveFilename, TSharedRef<FMovieSceneTranslatorContext> InContext);
+	virtual bool Export(const UMovieScene* InMovieScene, FString InFilenameFormat, FFrameRate InFrameRate, uint32 InResX, uint32 InResY, int32 InHandleFrames, FString InSaveFilename, TSharedRef<FMovieSceneTranslatorContext> InContext, FString InMovieExtension);
 };
diff --git a/Engine/Source/Editor/MovieSceneTools/Public/MovieSceneToolHelpers.h b/Engine/Source/Editor/MovieSceneTools/Public/MovieSceneToolHelpers.h
--- a/Engine/Source/Editor/MovieSceneTools/Public/MovieSceneToolHelpers.h
+++ b/Engine/Source/Editor/MovieSceneTools/Public/MovieSceneToolHelpers.h
@@ -143,9 +143,10 @@ public:
 	 * @param InFrameRate The frame rate to export the EDL at
 	 * @param InSaveDirectory Optional directory path to save to. If none given, a dialog will pop up to prompt the user
 	 * @param InHandleFrames The number of handle frames to include for each shot.
+	 * @param MovieExtension The movie extension for the shot filenames (ie. .avi, .mov, .mp4)
 	 * @return Whether the export was successful
 	 */
-	static bool ShowExportEDLDialog(const UMovieScene* InMovieScene, FFrameRate InFrameRate, FString InSaveDirectory = TEXT(""), int32 InHandleFrames = 8);
+	static bool ShowExportEDLDialog(const UMovieScene* InMovieScene, FFrameRate InFrameRate, FString InSaveDirectory = TEXT(""), int32 InHandleFrames = 8, FString InMovieExtension = TEXT(".avi"));
 
 	/**
 	* Import movie scene formats
diff --git a/Engine/Source/Editor/MovieSceneTools/Public/MovieSceneTranslator.h b/Engine/Source/Editor/MovieSceneTools/Public/MovieSceneTranslator.h
--- a/Engine/Source/Editor/MovieSceneTools/Public/MovieSceneTranslator.h
+++ b/Engine/Source/Editor/MovieSceneTools/Public/MovieSceneTranslator.h
@@ -124,6 +124,7 @@ struct FMovieSceneExportMovieSceneData
 	FFrameNumber PlaybackRangeEndFrame;
 	TSharedPtr<FMovieSceneExportCinematicMasterTrackData> CinematicMasterTrack;
 	TArray< TSharedPtr<FMovieSceneExportAudioMasterTrackData> > AudioMasterTracks;
+	FString MovieExtension;
 };
 
 enum class EMovieSceneTranslatorSectionType : int32
@@ -139,7 +140,7 @@ class MOVIESCENETOOLS_API FMovieSceneExportData : public TSharedFromThis<FMovieS
 {
 public:
 	/** Constructor */
-	FMovieSceneExportData(const UMovieScene* InMovieScene, FFrameRate InFrameRate, uint32 InResX, uint32 InResY, int32 InHandleFrames, FString InSaveFilename, TSharedPtr<FMovieSceneTranslatorContext> InContext);
+	FMovieSceneExportData(const UMovieScene* InMovieScene, FFrameRate InFrameRate, uint32 InResX, uint32 InResY, int32 InHandleFrames, FString InSaveFilename, TSharedPtr<FMovieSceneTranslatorContext> InContext, FString InMovieExtension);
 	/** Default constructor, necessary for shared ref - should not be used */
 	FMovieSceneExportData();
 	/** Destructor */
@@ -149,6 +150,8 @@ public:
 	FString GetFilename() const;
 	/** Gets export filename with full path */
 	FString GetFilenamePath() const;
+	/** Gets the shot movie extension */
+	FString GetMovieExtension() const;
 
 	/** Gets export frame rate */
 	FFrameRate GetFrameRate() const;
@@ -161,8 +164,8 @@ public:
 	/** Returns the nearest integral frame rate */
 	uint32 GetNearestWholeFrameRate() const;
 
-	/** Gets the frame handle which has been converted to the export frame rate */
-	FFrameTime GetHandleFrames() const;
+	/** Gets the frame handle */
+	int32 GetHandleFrames() const;
 	/** Gets default audio sample rate */
 	int32 GetDefaultAudioSampleRate() const;
 	/** Gets default audio depth */
@@ -208,12 +211,13 @@ private:
 	FFrameRate FrameRate;
 	uint32 ResX;
 	uint32 ResY;
-	FFrameTime HandleFrames;
+	int32 HandleFrames;
 	FString SaveFilename;
 	FString SaveFilenamePath;
 	bool bExportDataIsValid;
 	int32 DefaultAudioSampleRate;
 	int32 DefaultAudioDepth;
+	FString MovieExtension;
 };
 
 struct FMovieSceneImportCinematicSectionData 
@@ -296,7 +300,7 @@ public:
 	/** Create cinematic section */
 	TSharedPtr<FMovieSceneImportCinematicSectionData> CreateCinematicSection(FString InName, int32 InRow, FFrameRate InFrameRate, FFrameNumber InStartFrame, FFrameNumber InEndFrame, FFrameNumber InStartOffsetFrame);
 	/** Set cinematic section */
-	bool SetCinematicSection(TSharedPtr<FMovieSceneImportCinematicSectionData> InSection, int32 InRow, FFrameRate InFrameRate, FFrameNumber InStartFrame, FFrameNumber InEndFrame, FFrameNumber InStartOffsetFrame);
+	bool SetCinematicSection(TSharedPtr<FMovieSceneImportCinematicSectionData> InSection, int32 InRow, FFrameRate InFrameRate, FFrameNumber InStartFrame, FFrameNumber InEndFrame, TOptional<FFrameNumber> InStartOffsetFrame);
 
 	/** Returns the audio master track data pointer or nullptr if one does not exist */
 	TSharedPtr<FMovieSceneImportAudioMasterTrackData> GetAudioMasterTrackData();
@@ -408,8 +412,9 @@ public:
 	* @param InHandleFrames The number of handle frames to include for each shot.
 	* @param InSaveFilename The file path to save to.
 	* @param OutError The return error message
+	* @param MovieExtension The movie extension for the shot filenames (ie. .avi, .mov, .mp4)
 	* @return Whether the export was successful
 	*/
-	virtual bool Export(const UMovieScene* InMovieScene, FString InFilenameFormat, FFrameRate InFrameRate, uint32 InResX, uint32 InResY, int32 InHandleFrames, FString InSaveFilename, TSharedRef<FMovieSceneTranslatorContext> InContext) = 0;
+	virtual bool Export(const UMovieScene* InMovieScene, FString InFilenameFormat, FFrameRate InFrameRate, uint32 InResX, uint32 InResY, int32 InHandleFrames, FString InSaveFilename, TSharedRef<FMovieSceneTranslatorContext> InContext, FString InMovieExtension) = 0;
 };
 
diff --git a/Engine/Source/Editor/MovieSceneTools/Public/MovieSceneTranslatorEDL.h b/Engine/Source/Editor/MovieSceneTools/Public/MovieSceneTranslatorEDL.h
--- a/Engine/Source/Editor/MovieSceneTools/Public/MovieSceneTranslatorEDL.h
+++ b/Engine/Source/Editor/MovieSceneTools/Public/MovieSceneTranslatorEDL.h
@@ -29,7 +29,8 @@ public:
 	 * @param InFrameRate The frame rate to export the EDL at
 	 * @param InSaveFilename The file path to save to.
 	 * @param InHandleFrames The number of handle frames to include for each shot.
+	 * @param MovieExtension The movie extension for the shot filenames (ie. .avi, .mov, .mp4)
 	 * @return Whether the export was successful
 	 */
-	static bool ExportEDL(const UMovieScene* InMovieScene, FFrameRate InFrameRate, FString InSaveFilename, int32 InHandleFrames);
+	static bool ExportEDL(const UMovieScene* InMovieScene, FFrameRate InFrameRate, FString InSaveFilename, int32 InHandleFrames, FString InMovieExtension);
 };
diff --git a/Engine/Source/Editor/Sequencer/Private/DisplayNodes/SequencerTrackNode.cpp b/Engine/Source/Editor/Sequencer/Private/DisplayNodes/SequencerTrackNode.cpp
--- a/Engine/Source/Editor/Sequencer/Private/DisplayNodes/SequencerTrackNode.cpp
+++ b/Engine/Source/Editor/Sequencer/Private/DisplayNodes/SequencerTrackNode.cpp
@@ -510,8 +510,10 @@ void FSequencerTrackNode::SetDisplayName(const FText& NewDisplayName)
 {
 	auto NameableTrack = Cast<UMovieSceneNameableTrack>(AssociatedTrack.Get());
 
-	if (NameableTrack != nullptr)
+	if (NameableTrack != nullptr && !NameableTrack->GetDisplayName().EqualTo(NewDisplayName))
 	{
+		const FScopedTransaction Transaction(NSLOCTEXT("SequencerTrackNode", "RenameTrack", "Rename Track"));
+
 		NameableTrack->SetDisplayName(NewDisplayName);
 		GetSequencer().NotifyMovieSceneDataChanged(EMovieSceneDataChangeType::TrackValueChanged);
 	}
diff --git a/Engine/Source/Editor/Sequencer/Private/LevelEditorSequencerIntegration.cpp b/Engine/Source/Editor/Sequencer/Private/LevelEditorSequencerIntegration.cpp
--- a/Engine/Source/Editor/Sequencer/Private/LevelEditorSequencerIntegration.cpp
+++ b/Engine/Source/Editor/Sequencer/Private/LevelEditorSequencerIntegration.cpp
@@ -252,6 +252,9 @@ void FLevelEditorSequencerIntegration::Initialize()
 			}
 		);
 	}
+
+	const bool bForceRefresh = true;
+	UpdateDetails(bForceRefresh);
 }
 
 void RenameSpawnableRecursive(FSequencer* Sequencer, FMovieSceneSequenceIDRef SequenceID, AActor* ChangedActor)
diff --git a/Engine/Source/Editor/Sequencer/Private/MovieSceneCopyableTrack.h b/Engine/Source/Editor/Sequencer/Private/MovieSceneCopyableTrack.h
new file mode 100644
--- /dev/null
+++ b/Engine/Source/Editor/Sequencer/Private/MovieSceneCopyableTrack.h
@@ -0,0 +1,22 @@
+// Copyright 1998-2019 Epic Games, Inc. All Rights Reserved.
+
+#pragma once
+
+#include "CoreMinimal.h"
+#include "UObject/ObjectMacros.h"
+#include "MovieSceneCopyableTrack.generated.h"
+
+class UMovieSceneTrack;
+
+UCLASS(Transient)
+class UMovieSceneCopyableTrack : public UObject
+{
+public:
+	GENERATED_BODY()
+
+	UPROPERTY()
+	UMovieSceneTrack* Track;
+
+	UPROPERTY()
+	bool bIsAMasterTrack;
+};
diff --git a/Engine/Source/Editor/Sequencer/Private/SSequencer.cpp b/Engine/Source/Editor/Sequencer/Private/SSequencer.cpp
--- a/Engine/Source/Editor/Sequencer/Private/SSequencer.cpp
+++ b/Engine/Source/Editor/Sequencer/Private/SSequencer.cpp
@@ -75,6 +75,7 @@
 #include "LevelSequence.h"
 #include "SequencerLog.h"
 #include "MovieSceneCopyableBinding.h"
+#include "MovieSceneCopyableTrack.h"
 
 #define LOCTEXT_NAMESPACE "Sequencer"
 
@@ -1145,14 +1146,34 @@ TSharedRef<SWidget> SSequencer::MakePlaybackMenu()
 		auto OnStartChanged = [=](double NewValue){
 
 			// We clamp the new value when the value is set. We can't clamp in the UI because we need an unset Min/Max for linear scaling to work.
-			double Min = (SequencerPtr.Pin()->GetClampRange().GetLowerBoundValue() * SequencerPtr.Pin()->GetFocusedTickResolution()).GetFrame().Value;
-			double Max = SequencerPtr.Pin()->GetPlaybackRange().GetUpperBoundValue().Value; 
+			double Min = -FLT_MAX;
+			double Max = SequencerPtr.Pin()->GetPlaybackRange().GetUpperBoundValue().Value;
 
 			NewValue = FMath::Clamp(NewValue, Min, Max);
 			FFrameNumber ValueAsFrame = FFrameTime::FromDecimal(NewValue).GetFrame();
 
 			FFrameNumber Upper = MovieScene::DiscreteExclusiveUpper(SequencerPtr.Pin()->GetPlaybackRange());
-			SequencerPtr.Pin()->SetPlaybackRange(TRange<FFrameNumber>(FMath::Min(ValueAsFrame, Upper-1), Upper));
+
+			TRange<FFrameNumber> NewRange = TRange<FFrameNumber>(FMath::Min(ValueAsFrame, Upper - 1), Upper);
+
+			SequencerPtr.Pin()->SetPlaybackRange(NewRange);
+
+			TRange<double> PlayRangeSeconds = SequencerPtr.Pin()->GetPlaybackRange() / SequencerPtr.Pin()->GetFocusedTickResolution();
+			const double AdditionalRange = (PlayRangeSeconds.GetUpperBoundValue() - PlayRangeSeconds.GetLowerBoundValue()) * 0.1;
+
+			TRange<double> NewClampRange = SequencerPtr.Pin()->GetClampRange();
+			NewClampRange.SetLowerBoundValue(SequencerPtr.Pin()->GetPlaybackRange().GetLowerBoundValue() / SequencerPtr.Pin()->GetFocusedTickResolution() - AdditionalRange);
+			if (SequencerPtr.Pin()->GetClampRange().GetLowerBoundValue() > NewClampRange.GetLowerBoundValue())
+			{
+				SequencerPtr.Pin()->SetClampRange(NewClampRange);
+			}
+
+			TRange<double> NewViewRange = SequencerPtr.Pin()->GetViewRange();
+			NewViewRange.SetLowerBoundValue(SequencerPtr.Pin()->GetPlaybackRange().GetLowerBoundValue() / SequencerPtr.Pin()->GetFocusedTickResolution() - AdditionalRange);
+			if (SequencerPtr.Pin()->GetViewRange().GetLowerBoundValue() > NewViewRange.GetLowerBoundValue())
+			{
+				SequencerPtr.Pin()->SetViewRange(NewViewRange);
+			}
 		};
 
 		MenuBuilder.AddWidget(
@@ -1189,13 +1210,30 @@ TSharedRef<SWidget> SSequencer::MakePlaybackMenu()
 
 			// We clamp the new value when the value is set. We can't clamp in the UI because we need an unset Min/Max for linear scaling to work.
 			double Min = SequencerPtr.Pin()->GetPlaybackRange().GetLowerBoundValue().Value;
-			double Max = (SequencerPtr.Pin()->GetClampRange().GetUpperBoundValue() * SequencerPtr.Pin()->GetFocusedTickResolution()).GetFrame().Value;
+			double Max = FLT_MAX;
 
 			NewValue = FMath::Clamp(NewValue, Min, Max);
 			FFrameNumber ValueAsFrame = FFrameTime::FromDecimal(NewValue).GetFrame();
 
 			FFrameNumber Lower = MovieScene::DiscreteInclusiveLower(SequencerPtr.Pin()->GetPlaybackRange());
 			SequencerPtr.Pin()->SetPlaybackRange(TRange<FFrameNumber>(Lower, FMath::Max(ValueAsFrame, Lower)));
+
+			TRange<double> PlayRangeSeconds = SequencerPtr.Pin()->GetPlaybackRange() / SequencerPtr.Pin()->GetFocusedTickResolution();
+			const double AdditionalRange = (PlayRangeSeconds.GetUpperBoundValue() - PlayRangeSeconds.GetLowerBoundValue()) * 0.1;
+
+			TRange<double> NewClampRange = SequencerPtr.Pin()->GetClampRange();
+			NewClampRange.SetUpperBoundValue(SequencerPtr.Pin()->GetPlaybackRange().GetUpperBoundValue() / SequencerPtr.Pin()->GetFocusedTickResolution() + AdditionalRange);
+			if (SequencerPtr.Pin()->GetClampRange().GetUpperBoundValue() < NewClampRange.GetUpperBoundValue())
+			{
+				SequencerPtr.Pin()->SetClampRange(NewClampRange);
+			}
+
+			TRange<double> NewViewRange = SequencerPtr.Pin()->GetViewRange();
+			NewViewRange.SetUpperBoundValue(SequencerPtr.Pin()->GetPlaybackRange().GetUpperBoundValue() / SequencerPtr.Pin()->GetFocusedTickResolution() + AdditionalRange);
+			if (SequencerPtr.Pin()->GetViewRange().GetUpperBoundValue() < NewViewRange.GetUpperBoundValue())
+			{
+				SequencerPtr.Pin()->SetViewRange(NewViewRange);
+			}
 		};
 
 		MenuBuilder.AddWidget(
@@ -2344,7 +2382,7 @@ bool SSequencer::CanPaste()
 	// Attempts to deserialize the text into object bindings/tracks that Sequencer understands.
 	if (Sequencer->CanPaste(TextToImport))
 	{
-		TArray<UMovieSceneTrack*> ImportedTracks;
+		TArray<UMovieSceneCopyableTrack*> ImportedTracks;
 		TArray<UMovieSceneSection*> ImportedSections;
 		TArray<UMovieSceneCopyableBinding*> ImportedObjects;
 		Sequencer->ImportTracksFromText(TextToImport, ImportedTracks);
diff --git a/Engine/Source/Editor/Sequencer/Private/SSequencerLabelBrowser.cpp b/Engine/Source/Editor/Sequencer/Private/SSequencerLabelBrowser.cpp
--- a/Engine/Source/Editor/Sequencer/Private/SSequencerLabelBrowser.cpp
+++ b/Engine/Source/Editor/Sequencer/Private/SSequencerLabelBrowser.cpp
@@ -4,7 +4,7 @@
 #include "SlateOptMacros.h"
 #include "Framework/MultiBox/MultiBoxBuilder.h"
 #include "SSequencerLabelListRow.h"
-
+#include "ScopedTransaction.h"
 
 #define LOCTEXT_NAMESPACE "SSequencerLabelBrowser"
 
@@ -147,6 +147,8 @@ void SSequencerLabelBrowser::ReloadLabelList(bool FullyReload)
 
 void SSequencerLabelBrowser::HandleLabelListRowLabelRenamed(TSharedPtr<FSequencerLabelTreeNode> Node, const FString& NewLabel)
 {
+	const FScopedTransaction Transaction(LOCTEXT("RenameLabel", "Rename Label"));
+
 	if (Sequencer.IsValid() && Sequencer.Pin()->GetLabelManager().RenameLabel(Node->Label, NewLabel))
 	{
 		ReloadLabelList(true);
@@ -209,6 +211,26 @@ void SSequencerLabelBrowser::HandleLabelTreeViewGetChildren(TSharedPtr<FSequence
 	}
 }
 
+void SSequencerLabelBrowser::AddLabelNode(TSharedPtr<FSequencerLabelTreeNode> InItem, FString& NewLabel)
+{
+	if (InItem.IsValid())
+	{
+		if (!NewLabel.IsEmpty())
+		{
+			NewLabel += TEXT(" ");
+		}
+
+		if (!InItem->Label.IsEmpty())
+		{
+			NewLabel += FString(TEXT("label:")) + InItem->Label;
+		}
+
+		for (auto Child : InItem->Children)
+		{
+			AddLabelNode(Child, NewLabel);
+		}
+	}
+}
 
 void SSequencerLabelBrowser::HandleLabelTreeViewSelectionChanged(TSharedPtr<FSequencerLabelTreeNode> InItem, ESelectInfo::Type SelectInfo)
 {
@@ -218,15 +240,7 @@ void SSequencerLabelBrowser::HandleLabelTreeViewSelectionChanged(TSharedPtr<FSeq
 	LabelTreeView->GetSelectedItems(SelectedItems);
 	for (TSharedPtr<FSequencerLabelTreeNode> Item : SelectedItems)
 	{
-		if (!NewLabel.IsEmpty())
-		{
-			NewLabel += TEXT(" ");
-		}
-		
-		if (!Item->Label.IsEmpty())
-		{
-			NewLabel += FString(TEXT("label:")) + Item->Label;
-		}
+		AddLabelNode(Item, NewLabel);
 	}
 
 	OnSelectionChanged.ExecuteIfBound(
@@ -237,6 +251,19 @@ void SSequencerLabelBrowser::HandleLabelTreeViewSelectionChanged(TSharedPtr<FSeq
 	);
 }
 
+void SSequencerLabelBrowser::RemoveLabelItem(TSharedPtr<FSequencerLabelTreeNode> InItem)
+{
+	if (InItem.IsValid())
+	{
+		for (auto Child : InItem->Children)
+		{
+			RemoveLabelItem(Child);
+		}
+
+		Sequencer.Pin()->GetLabelManager().RemoveObjectLabel(FGuid(), InItem->Label);
+	}
+}
+
 
 void SSequencerLabelBrowser::HandleRemoveLabelMenuEntryExecute()
 {
@@ -244,7 +271,9 @@ void SSequencerLabelBrowser::HandleRemoveLabelMenuEntryExecute()
 
 	if (Sequencer.IsValid() && LabelTreeView->GetSelectedItems(SelectedItems) > 0)
 	{
-		Sequencer.Pin()->GetLabelManager().RemoveObjectLabel(FGuid(), SelectedItems[0]->Label);
+		const FScopedTransaction Transaction(LOCTEXT("RemoveLabel", "Remove Label"));
+
+		RemoveLabelItem(SelectedItems[0]);
 	}
 }
 
diff --git a/Engine/Source/Editor/Sequencer/Private/SSequencerLabelBrowser.h b/Engine/Source/Editor/Sequencer/Private/SSequencerLabelBrowser.h
--- a/Engine/Source/Editor/Sequencer/Private/SSequencerLabelBrowser.h
+++ b/Engine/Source/Editor/Sequencer/Private/SSequencerLabelBrowser.h
@@ -89,6 +89,12 @@ private:
 	/** Callback for checking whether 'Rename' context menu entry can execute. */
 	bool HandleRenameLabelMenuEntryCanExecute() const;
 
+	/** Remove label item and its children */
+	void RemoveLabelItem(TSharedPtr<FSequencerLabelTreeNode> InItem);
+
+	/** Add node's label and its children to the new label */
+	void AddLabelNode(TSharedPtr<FSequencerLabelTreeNode> InItem, FString& NewLabel);
+
 private:
 
 	/** Holds the collection of root labels to be displayed in the tree view. */
diff --git a/Engine/Source/Editor/Sequencer/Private/SSequencerLabelEditor.cpp b/Engine/Source/Editor/Sequencer/Private/SSequencerLabelEditor.cpp
--- a/Engine/Source/Editor/Sequencer/Private/SSequencerLabelEditor.cpp
+++ b/Engine/Source/Editor/Sequencer/Private/SSequencerLabelEditor.cpp
@@ -8,7 +8,7 @@
 #include "Sequencer.h"
 #include "Widgets/Views/SListView.h"
 #include "SSequencerLabelEditorListRow.h"
-
+#include "ScopedTransaction.h"
 
 #define LOCTEXT_NAMESPACE "SSequencerLabelEditor"
 
@@ -86,6 +86,8 @@ void SSequencerLabelEditor::Construct(const FArguments& InArgs, FSequencer& InSe
 
 void SSequencerLabelEditor::CreateLabelFromFilterText()
 {
+	const FScopedTransaction Transaction(LOCTEXT("AddLabel", "Add Label"));
+
 	for (auto ObjectId : ObjectIds)
 	{
 		Sequencer->GetLabelManager().AddObjectLabel(ObjectId, FilterBox->GetText().ToString());
@@ -157,8 +159,11 @@ FReply SSequencerLabelEditor::HandleFilterBoxKeyDown(const FGeometry& /*Geometry
 {
 	if (KeyEvent.GetKey() == EKeys::Enter)
 	{
-		CreateLabelFromFilterText();
-		return FReply::Handled();
+		if (HandleCreateNewLabelButtonIsEnabled())
+		{
+			CreateLabelFromFilterText();
+			return FReply::Handled();
+		}
 	}
 
 	return FReply::Unhandled();
diff --git a/Engine/Source/Editor/Sequencer/Private/SSequencerLabelListRow.h b/Engine/Source/Editor/Sequencer/Private/SSequencerLabelListRow.h
--- a/Engine/Source/Editor/Sequencer/Private/SSequencerLabelListRow.h
+++ b/Engine/Source/Editor/Sequencer/Private/SSequencerLabelListRow.h
@@ -107,6 +107,7 @@ public:
 						[
 							SAssignNew(EditableLabel, SInlineEditableTextBlock)
 								.IsReadOnly(this, &SSequencerLabelListRow::IsReadOnly)
+								.OnVerifyTextChanged(this, &SSequencerLabelListRow::HandleVerifyNameChanged)
 								.OnTextCommitted(this, &SSequencerLabelListRow::HandleFolderNameTextChanged)
 								.Text(
 									Node->Label.IsEmpty()
@@ -156,6 +157,19 @@ private:
 		return FLinearColor::Gray;
 	}
 
+
+	bool HandleVerifyNameChanged(const FText& NewText, FText& OutErrorMessage)
+	{
+		if (NewText.ToString().Contains(TEXT(" ")))	
+		{
+			OutErrorMessage = LOCTEXT("InvalidLabel", "Label cannot contains spaces");
+			return false;
+		}
+
+		return true;
+	}
+
+
 	void HandleFolderNameTextChanged(const FText& NewLabel, ETextCommit::Type TextCommitType)
 	{
 		FString NewLabelString = NewLabel.ToString();
diff --git a/Engine/Source/Editor/Sequencer/Private/Sequencer.cpp b/Engine/Source/Editor/Sequencer/Private/Sequencer.cpp
--- a/Engine/Source/Editor/Sequencer/Private/Sequencer.cpp
+++ b/Engine/Source/Editor/Sequencer/Private/Sequencer.cpp
@@ -112,6 +112,7 @@
 #include "Compilation/MovieSceneCompiler.h"
 #include "SequencerKeyActor.h"
 #include "MovieSceneCopyableBinding.h"
+#include "MovieSceneCopyableTrack.h"
 #include "ISequencerChannelInterface.h"
 #include "SequencerKeyCollection.h"
 #include "CurveEditor.h"
@@ -1810,9 +1811,10 @@ TRange<FFrameNumber> FSequencer::GetSelectionRange() const
 void FSequencer::SetSelectionRange(TRange<FFrameNumber> Range)
 {
 	const FScopedTransaction Transaction(LOCTEXT("SetSelectionRange_Transaction", "Set Selection Range"));
-	UMovieScene* FocussedMovieScene = GetFocusedMovieSceneSequence()->GetMovieScene();
+	UMovieScene* FocusedMovieScene = GetFocusedMovieSceneSequence()->GetMovieScene();
 
-	FocussedMovieScene->SetSelectionRange(Range);
+	FocusedMovieScene->Modify();
+	FocusedMovieScene->SetSelectionRange(Range);
 }
 
 
@@ -6182,7 +6184,7 @@ void FSequencer::ImportObjectBindingsFromText(const FString& TextToImport, /*out
 }
 
 
-void FSequencer::CopySelectedObjects(TArray<TSharedPtr<FSequencerObjectBindingNode>>& ObjectNodes)
+void FSequencer::CopySelectedObjects(TArray<TSharedPtr<FSequencerObjectBindingNode>>& ObjectNodes, FString& ExportedText)
 {
 	// Gather guids for the object nodes and any child object nodes
 	TSet<FGuid> GuidsToCopy;
@@ -6246,23 +6248,41 @@ void FSequencer::CopySelectedObjects(TArray<TSharedPtr<FSequencerObjectBindingNo
 		}
 	}
 
-	FString ExportedText;
 	ExportObjectBindingsToText(CopyableBindings, /*out*/ ExportedText);
-	FPlatformApplicationMisc::ClipboardCopy(*ExportedText);
 }
 
 
-void FSequencer::CopySelectedTracks(TArray<TSharedPtr<FSequencerTrackNode>>& TrackNodes)
+void FSequencer::CopySelectedTracks(TArray<TSharedPtr<FSequencerTrackNode>>& TrackNodes, FString& ExportedText)
 {
-	TArray<UObject*> TracksToCopy;
+	UMovieScene* MovieScene = GetFocusedMovieSceneSequence()->GetMovieScene();
+
+	TArray<UObject*> CopyableTracks;
 	for (TSharedPtr<FSequencerTrackNode> TrackNode : TrackNodes)
 	{
-		TracksToCopy.Add(TrackNode->GetTrack());
+		bool bIsParentSelected = false;
+		TSharedPtr<FSequencerDisplayNode> ParentNode = TrackNode->GetParent();
+		while (ParentNode.IsValid())
+		{
+			if (Selection.GetSelectedOutlinerNodes().Contains(ParentNode.ToSharedRef()))
+			{
+				bIsParentSelected = true;
+				break;
+			}
+			ParentNode = ParentNode->GetParent();
+		}
+
+		if (!bIsParentSelected)
+		{
+			UMovieSceneCopyableTrack *CopyableTrack = NewObject<UMovieSceneCopyableTrack>(GetTransientPackage(), UMovieSceneCopyableTrack::StaticClass(), NAME_None, RF_Transient);
+			CopyableTracks.Add(CopyableTrack);
+
+			UMovieSceneTrack* DuplicatedTrack = Cast<UMovieSceneTrack>(StaticDuplicateObject(TrackNode->GetTrack(), CopyableTrack));
+			CopyableTrack->Track = DuplicatedTrack;
+			CopyableTrack->bIsAMasterTrack = MovieScene->IsAMasterTrack(*TrackNode->GetTrack());
+		}
 	}
 
-	FString ExportedText;
-	FSequencer::ExportObjectsToText(TracksToCopy, /*out*/ ExportedText);
-	FPlatformApplicationMisc::ClipboardCopy(*ExportedText);
+	ExportObjectsToText(CopyableTracks, /*out*/ ExportedText);
 }
 
 
@@ -6300,21 +6320,24 @@ void FSequencer::DoPaste()
 	FString TextToImport;
 	FPlatformApplicationMisc::ClipboardPaste(TextToImport);
 
-	if (PasteObjectBindings(TextToImport))
+	TArray<FNotificationInfo> PasteErrors;
+	bool bAnythingPasted = false;
+	bAnythingPasted |= PasteObjectBindings(TextToImport, PasteErrors);
+	bAnythingPasted |= PasteTracks(TextToImport, PasteErrors);
+	
+	if (!bAnythingPasted)
 	{
-		return;
+		bAnythingPasted |= PasteSections(TextToImport, PasteErrors);
 	}
-	else if (PasteTracks(TextToImport))
-	{
-		return;
-	}
-	else if (PasteSections(TextToImport))
+
+	for (auto NotificationInfo : PasteErrors)
 	{
-		return;
+		NotificationInfo.bUseLargeFont = false;
+		FSlateNotificationManager::Get().AddNotification(NotificationInfo);
 	}
 }
 
-bool FSequencer::PasteObjectBindings(const FString& TextToImport)
+bool FSequencer::PasteObjectBindings(const FString& TextToImport, TArray<FNotificationInfo>& PasteErrors)
 {
 	TArray<UMovieSceneCopyableBinding*> ImportedBindings;
 	ImportObjectBindingsFromText(TextToImport, ImportedBindings);
@@ -6324,6 +6347,10 @@ bool FSequencer::PasteObjectBindings(const FString& TextToImport)
 		return false;
 	}
 	
+	TArray<UMovieSceneFolder*> SelectedParentFolders;
+	FString NewNodePath;
+	CalculateSelectedFolderAndPath(SelectedParentFolders, NewNodePath);
+
 	FScopedTransaction Transaction(FGenericCommands::Get().Paste->GetDescription());
 	
 	UMovieSceneSequence* OwnerSequence = GetFocusedMovieSceneSequence();
@@ -6410,9 +6437,9 @@ bool FSequencer::PasteObjectBindings(const FString& TextToImport)
 	}
 
 	// Fix possessable actor bindings
-	for (auto PossessableGuid : PossessableGuids)
+	for (int32 PossessableGuidIndex = 0; PossessableGuidIndex < PossessableGuids.Num(); ++PossessableGuidIndex)
 	{
-		FMovieScenePossessable* Possessable = MovieScene->FindPossessable(PossessableGuid);
+		FMovieScenePossessable* Possessable = MovieScene->FindPossessable(PossessableGuids[PossessableGuidIndex]);
 		UWorld* PlaybackContext = Cast<UWorld>(GetPlaybackContext());
 		if (Possessable && PlaybackContext)
 		{
@@ -6425,13 +6452,35 @@ bool FSequencer::PasteObjectBindings(const FString& TextToImport)
 
 					if (!ExistingGuid.IsValid())
 					{
-						DoAssignActor(&Actor, 1, Possessable->GetGuid());
+						FGuid NewGuid = DoAssignActor(&Actor, 1, Possessable->GetGuid());
+
+						// If assigning produces a new guid, update the possesable guids and the bindings pasted data
+						if (NewGuid.IsValid())
+						{
+							for (auto BindingPasted : BindingsPasted)
+							{
+								if (BindingPasted.GetObjectGuid() == PossessableGuids[PossessableGuidIndex])
+								{
+									BindingPasted.SetObjectGuid(NewGuid);
+								}
+							}
+
+							PossessableGuids[PossessableGuidIndex] = NewGuid;
+						}
 					}
 				}
 			}
 		}
 	}
 
+	if (SelectedParentFolders.Num() > 0)
+	{
+		for (auto PossessableGuid : PossessableGuids)
+		{
+			SelectedParentFolders[0]->AddChildObjectBinding(PossessableGuid);
+		}
+	}
+
 	OnMovieSceneBindingsPastedDelegate.Broadcast(BindingsPasted);
 
 	// Refresh all immediately so that spawned actors will be generated immediately
@@ -6468,9 +6517,9 @@ bool FSequencer::PasteObjectBindings(const FString& TextToImport)
 	return true;
 }
 
-bool FSequencer::PasteTracks(const FString& TextToImport)
+bool FSequencer::PasteTracks(const FString& TextToImport, TArray<FNotificationInfo>& PasteErrors)
 {
-	TArray<UMovieSceneTrack*> ImportedTracks;
+	TArray<UMovieSceneCopyableTrack*> ImportedTracks;
 	FSequencer::ImportTracksFromText(TextToImport, ImportedTracks);
 
 	if (ImportedTracks.Num() == 0)
@@ -6478,6 +6527,10 @@ bool FSequencer::PasteTracks(const FString& TextToImport)
 		return false;
 	}
 	
+	TArray<UMovieSceneFolder*> SelectedParentFolders;
+	FString NewNodePath;
+	CalculateSelectedFolderAndPath(SelectedParentFolders, NewNodePath);
+
 	FScopedTransaction Transaction(FGenericCommands::Get().Paste->GetDescription());
 
 	UMovieSceneSequence* OwnerSequence = GetFocusedMovieSceneSequence();
@@ -6500,62 +6553,107 @@ bool FSequencer::PasteTracks(const FString& TextToImport)
 		}
 	}
 
+	int32 NumMasterTracks = 0;
+	int32 NumTracks = 0;
+
+	for (UMovieSceneCopyableTrack* CopyableTrack : ImportedTracks)
+	{
+		if (CopyableTrack->bIsAMasterTrack)
+		{
+			++NumMasterTracks;
+		}
+		else
+		{
+			++NumTracks;
+		}
+	}
+
+	int32 NumMasterTracksPasted = 0;
+	int32 NumTracksPasted = 0;
 	if (ObjectNodes.Num())
 	{
 		for (TSharedPtr<FSequencerObjectBindingNode> ObjectNode : ObjectNodes)
 		{
 			FGuid ObjectGuid = ObjectNode->GetObjectBinding();
 
-			TArray<UMovieSceneTrack*> NewTracks;
+			TArray<UMovieSceneCopyableTrack*> NewTracks;
 			FSequencer::ImportTracksFromText(TextToImport, NewTracks);
 
-			for (UMovieSceneTrack* NewTrack : NewTracks)
+			for (UMovieSceneCopyableTrack* CopyableTrack : NewTracks)
 			{
-				if (!GetFocusedMovieSceneSequence()->GetMovieScene()->AddGivenTrack(NewTrack, ObjectGuid))
+				if (!CopyableTrack->bIsAMasterTrack)
 				{
-					FNotificationInfo Info(LOCTEXT("TrackAlreadyBound", "Can't Paste: Binding doesn't exist"));
-					Info.FadeInDuration = 0.1f;
-					Info.FadeOutDuration = 0.5f;
-					Info.ExpireDuration = 2.5f;
-					auto NotificationItem = FSlateNotificationManager::Get().AddNotification(Info);
-
-					NotificationItem->SetCompletionState(SNotificationItem::CS_Success);
-					NotificationItem->ExpireAndFadeout();
-
-					continue;
-				}
-				else
-				{
-					NotifyMovieSceneDataChanged(EMovieSceneDataChangeType::MovieSceneStructureItemsChanged);
+					UMovieSceneTrack* NewTrack = CopyableTrack->Track;
+					NewTrack->ClearFlags(RF_Transient);
+					if (!GetFocusedMovieSceneSequence()->GetMovieScene()->AddGivenTrack(NewTrack, ObjectGuid))
+					{
+						continue;
+					}
+					else
+					{
+						++NumTracksPasted;
+					}
 				}
 			}
 		}
-
-		return true;
 	}
 
 	// Add as master track or set camera cut track
-	for (UMovieSceneTrack* NewTrack : ImportedTracks)
+	for (UMovieSceneCopyableTrack* CopyableTrack : ImportedTracks)
 	{
-		if (NewTrack->IsA(UMovieSceneCameraCutTrack::StaticClass()))
-		{
-			GetFocusedMovieSceneSequence()->GetMovieScene()->SetCameraCutTrack(NewTrack);
-			NotifyMovieSceneDataChanged(EMovieSceneDataChangeType::MovieSceneStructureItemsChanged);
-		}
-		else
+		if (CopyableTrack->bIsAMasterTrack)
 		{
-			if (GetFocusedMovieSceneSequence()->GetMovieScene()->AddGivenMasterTrack(NewTrack))
+			UMovieSceneTrack* NewTrack = CopyableTrack->Track;
+			NewTrack->ClearFlags(RF_Transient);
+
+			if (NewTrack->IsA(UMovieSceneCameraCutTrack::StaticClass()))
+			{
+				GetFocusedMovieSceneSequence()->GetMovieScene()->SetCameraCutTrack(NewTrack);
+				if (SelectedParentFolders.Num() > 0)
+				{
+					SelectedParentFolders[0]->AddChildMasterTrack(NewTrack);
+				}
+
+				++NumMasterTracksPasted;
+			}
+			else
 			{
-				NotifyMovieSceneDataChanged(EMovieSceneDataChangeType::MovieSceneStructureItemsChanged);
+				if (GetFocusedMovieSceneSequence()->GetMovieScene()->AddGivenMasterTrack(NewTrack))
+				{
+					if (SelectedParentFolders.Num() > 0)
+					{
+						SelectedParentFolders[0]->AddChildMasterTrack(NewTrack);
+					}
+
+				}
+
+				++NumMasterTracksPasted;
 			}
 		}
 	}
 
+	if (NumMasterTracksPasted < NumMasterTracks)
+	{
+		FNotificationInfo Info(LOCTEXT("PasteTracks_NoMasterTracks", "Can't paste track. Master track could not be pasted"));
+		PasteErrors.Add(Info);
+	}
+
+	if (NumTracksPasted < NumTracks)
+	{
+		FNotificationInfo Info(LOCTEXT("PasteSections_NoSelectedObjects", "Can't paste track. No selected objects to paste tracks onto"));
+		PasteErrors.Add(Info);
+	}
+
+	if ((NumMasterTracksPasted + NumTracksPasted) > 0)
+	{
+		NotifyMovieSceneDataChanged(EMovieSceneDataChangeType::MovieSceneStructureItemsChanged);
+	}
+
 	return true;
 }
 
 
-bool FSequencer::PasteSections(const FString& TextToImport)
+bool FSequencer::PasteSections(const FString& TextToImport, TArray<FNotificationInfo>& PasteErrors)
 {
 	TArray<UMovieSceneSection*> ImportedSections;
 	FSequencer::ImportSectionsFromText(TextToImport, ImportedSections);
@@ -6569,9 +6667,8 @@ bool FSequencer::PasteSections(const FString& TextToImport)
 
 	if (SelectedNodes.Num() == 0)
 	{
-		FNotificationInfo Info(LOCTEXT("PasteSections_NoSelectedTracks", "No selected tracks to paste sections onto"));
-		Info.bUseLargeFont = false;
-		FSlateNotificationManager::Get().AddNotification(Info);
+		FNotificationInfo Info(LOCTEXT("PasteSections_NoSelectedTracks", "Can't paste section. No selected tracks to paste sections onto"));
+		PasteErrors.Add(Info);
 		return false;
 	}
 
@@ -6655,9 +6752,8 @@ bool FSequencer::PasteSections(const FString& TextToImport)
 	{
 		Transaction.Cancel();
 
-		FNotificationInfo Info(LOCTEXT("PasteSections_NothingPasted", "Nothing pasted. No matching section types found."));
-		Info.bUseLargeFont = false;
-		FSlateNotificationManager::Get().AddNotification(Info);
+		FNotificationInfo Info(LOCTEXT("PasteSections_NothingPasted", "Can't paste section. No matching section types found."));
+		PasteErrors.Add(Info);
 		return false;
 	}
 
@@ -6683,7 +6779,7 @@ public:
 	// FCustomizableTextObjectFactory implementation
 	virtual bool CanCreateClass(UClass* InObjectClass, bool& bOmitSubObjs) const override
 	{
-		if (InObjectClass->IsChildOf(UMovieSceneTrack::StaticClass()))
+		if (InObjectClass->IsChildOf(UMovieSceneCopyableTrack::StaticClass()))
 		{
 			return true;
 		}
@@ -6695,11 +6791,11 @@ public:
 	{
 		check(NewObject);
 
-		NewTracks.Add(Cast<UMovieSceneTrack>(NewObject));
+		NewTracks.Add(Cast<UMovieSceneCopyableTrack>(NewObject));
 	}
 
 public:
-	TArray<UMovieSceneTrack*> NewTracks;
+	TArray<UMovieSceneCopyableTrack*> NewTracks;
 };
 
 
@@ -6757,7 +6853,7 @@ bool FSequencer::CanPaste(const FString& TextToImport)
 	return false;
 }
 
-void FSequencer::ImportTracksFromText(const FString& TextToImport, /*out*/ TArray<UMovieSceneTrack*>& ImportedTracks)
+void FSequencer::ImportTracksFromText(const FString& TextToImport, /*out*/ TArray<UMovieSceneCopyableTrack*>& ImportedTracks)
 {
 	UPackage* TempPackage = NewObject<UPackage>(nullptr, TEXT("/Engine/Sequencer/Editor/Transient"), RF_Transient);
 	TempPackage->AddToRoot();
@@ -8125,14 +8221,24 @@ void FSequencer::CopySelection()
 			GClipboardStack.Push(NullClipboard);
 		}
 
+		FString ObjectsExportedText;
+		FString TracksExportedText;
+
 		if (ObjectsToCopy.Num())
 		{
-			CopySelectedObjects(ObjectsToCopy);
+			CopySelectedObjects(ObjectsToCopy, ObjectsExportedText);
 		}
-		else
+
+		if (TracksToCopy.Num())
 		{
-			CopySelectedTracks(TracksToCopy);
+			CopySelectedTracks(TracksToCopy, TracksExportedText);
 		}
+
+		FString ExportedText;
+		ExportedText += ObjectsExportedText;
+		ExportedText += TracksExportedText;
+
+		FPlatformApplicationMisc::ClipboardCopy(*ExportedText);
 	}
 }
 
@@ -8802,7 +8908,7 @@ void FSequencer::ExportFBXInternal(const FString& ExportFilename, TArray<FGuid>&
 			Exporter->SetTrasformBaking(false);
 			Exporter->SetKeepHierarchy(true);
 
-			const bool bSelectedOnly = Bindings.Num() != 0;
+			const bool bSelectedOnly = (Selection.GetSelectedTracks().Num() + Bindings.Num()) != 0;
 
 			UnFbx::FFbxExporter::FLevelSequenceNodeNameAdapter NodeNameAdapter(MovieScene, this, GetFocusedTemplateID());
 
diff --git a/Engine/Source/Editor/Sequencer/Private/Sequencer.h b/Engine/Source/Editor/Sequencer/Private/Sequencer.h
--- a/Engine/Source/Editor/Sequencer/Private/Sequencer.h
+++ b/Engine/Source/Editor/Sequencer/Private/Sequencer.h
@@ -61,12 +61,14 @@ class UMovieSceneSequence;
 class UMovieSceneSubSection;
 class USequencerSettings;
 class UMovieSceneCopyableBinding;
+class UMovieSceneCopyableTrack;
 struct FMovieSceneTimeController;
 struct FMovieScenePossessable;
 struct FSequencerTemplateStore;
 struct FTransformData;
 struct ISequencerHotspot;
 struct FKeyAttributes;
+struct FNotificationInfo;
 
 enum class EMapChangeType : uint8;
 
@@ -473,18 +475,18 @@ public:
 	void DeleteSelectedNodes();
 
 	/** Called when a user executes the copy track menu item */
-	void CopySelectedObjects(TArray<TSharedPtr<FSequencerObjectBindingNode>>& ObjectNodes);
-	void CopySelectedTracks(TArray<TSharedPtr<FSequencerTrackNode>>& TrackNodes);
+	void CopySelectedObjects(TArray<TSharedPtr<FSequencerObjectBindingNode>>& ObjectNodes, /*out*/ FString& ExportedText);
+	void CopySelectedTracks(TArray<TSharedPtr<FSequencerTrackNode>>& TrackNodes, /*out*/ FString& ExportedText);
 	void ExportObjectsToText(TArray<UObject*> ObjectsToExport, /*out*/ FString& ExportedText);
 
 	/** Called when a user executes the paste track menu item */
 	bool CanPaste(const FString& TextToImport);
 	void DoPaste();
-	bool PasteTracks(const FString& TextToImport);
-	bool PasteSections(const FString& TextToImport);
-	bool PasteObjectBindings(const FString& TextToImport);
-	
-	void ImportTracksFromText(const FString& TextToImport, /*out*/ TArray<UMovieSceneTrack*>& ImportedTracks);
+	bool PasteTracks(const FString& TextToImport, TArray<FNotificationInfo>& PasteErrors);
+	bool PasteSections(const FString& TextToImport, TArray<FNotificationInfo>& PasteErrors);
+	bool PasteObjectBindings(const FString& TextToImport, TArray<FNotificationInfo>& PasteErrors);
+
+	void ImportTracksFromText(const FString& TextToImport, /*out*/ TArray<UMovieSceneCopyableTrack*>& ImportedTracks);
 	void ImportSectionsFromText(const FString& TextToImport, /*out*/ TArray<UMovieSceneSection*>& ImportedSections);
 	void ImportObjectBindingsFromText(const FString& TextToImport, /*out*/ TArray<UMovieSceneCopyableBinding*>& ImportedObjects);
 
diff --git a/Engine/Source/Editor/Sequencer/Private/SequencerLabelManager.cpp b/Engine/Source/Editor/Sequencer/Private/SequencerLabelManager.cpp
--- a/Engine/Source/Editor/Sequencer/Private/SequencerLabelManager.cpp
+++ b/Engine/Source/Editor/Sequencer/Private/SequencerLabelManager.cpp
@@ -9,8 +9,10 @@ void FSequencerLabelManager::SetMovieScene(UMovieScene* InMovieScene)
 	if (MovieScene != InMovieScene)
 	{
 		MovieScene = InMovieScene;
-		LabelsChangedEvent.Broadcast();
 	}
+
+	// Broadcast the labels changed regardless of the movie scene changed since the label data could have changed
+	LabelsChangedEvent.Broadcast();
 }
 
 
@@ -18,6 +20,7 @@ void FSequencerLabelManager::AddObjectLabel(const FGuid& ObjectId, const FString
 {
 	if (MovieScene.IsValid())
 	{
+		MovieScene->Modify();
 		MovieScene->GetObjectsToLabels().FindOrAdd(ObjectId.ToString()).Strings.AddUnique(Label);
 		MovieScene->MarkPackageDirty();
 		LabelsChangedEvent.Broadcast();
@@ -45,6 +48,8 @@ void FSequencerLabelManager::RemoveObjectLabel(const FGuid& ObjectId, const FStr
 {
 	if (MovieScene.IsValid())
 	{
+		MovieScene->Modify();
+
 		TMap<FString, FMovieSceneTrackLabels>& ObjectsToLabels = MovieScene->GetObjectsToLabels();
 
 		if (ObjectId.IsValid())
@@ -115,6 +120,8 @@ bool FSequencerLabelManager::RenameLabel(const FString& OldLabel, const FString&
 
 	bool FoundOldLabel = false;
 
+	MovieScene->Modify();
+
 	for (auto& LabelsPair : MovieScene->GetObjectsToLabels())
 	{
 		TArray<FString>& Strings = LabelsPair.Value.Strings;
diff --git a/Engine/Source/Editor/Sequencer/Private/SequencerUtilities.cpp b/Engine/Source/Editor/Sequencer/Private/SequencerUtilities.cpp
--- a/Engine/Source/Editor/Sequencer/Private/SequencerUtilities.cpp
+++ b/Engine/Source/Editor/Sequencer/Private/SequencerUtilities.cpp
@@ -109,6 +109,12 @@ void FSequencerUtilities::PopulateMenu_CreateNewSection(FMenuBuilder& MenuBuilde
 			for (UMovieSceneSection* Section : Track->GetAllSections())
 			{
 				OverlapPriority = FMath::Max(Section->GetOverlapPriority() + 1, OverlapPriority);
+
+				// Move existing sections on the same row or beyond so that they don't overlap with the new section
+				if (Section != NewSection && Section->GetRowIndex() >= RowIndex)
+				{
+					Section->SetRowIndex(Section->GetRowIndex() + 1);
+				}
 			}
 
 			Track->Modify();
diff --git a/Engine/Source/Editor/UMGEditor/Private/TabFactory/AnimationTabSummoner.cpp b/Engine/Source/Editor/UMGEditor/Private/TabFactory/AnimationTabSummoner.cpp
--- a/Engine/Source/Editor/UMGEditor/Private/TabFactory/AnimationTabSummoner.cpp
+++ b/Engine/Source/Editor/UMGEditor/Private/TabFactory/AnimationTabSummoner.cpp
@@ -48,6 +48,27 @@ FAnimationTabSummoner::FAnimationTabSummoner(TSharedPtr<class FWidgetBlueprintEd
 
 bool VerifyAnimationRename( FWidgetBlueprintEditor& BlueprintEditor, UWidgetAnimation* Animation, FString NewAnimationName, FText& OutErrorMessage )
 {
+	if (NewAnimationName != SlugStringForValidName(NewAnimationName))
+	{
+		FString InvalidCharacters = INVALID_OBJECTNAME_CHARACTERS;
+		FString CurrentInvalidCharacter;
+		FString FoundInvalidCharacters;
+
+		// Create a string with all invalid characters found to output with the error message.
+		for (int32 StringIndex = 0; StringIndex < InvalidCharacters.Len(); ++StringIndex)
+		{
+			CurrentInvalidCharacter = InvalidCharacters.Mid(StringIndex, 1);
+
+			if (NewAnimationName.Contains(CurrentInvalidCharacter))
+			{
+				FoundInvalidCharacters += CurrentInvalidCharacter;
+			}
+		}
+		
+		OutErrorMessage = FText::Format(LOCTEXT("NameContainsInvalidCharacters", "The object name may not contain the following characters:  {0}"), FText::FromString(FoundInvalidCharacters));
+		return false;
+	}
+
 	UWidgetBlueprint* Blueprint = BlueprintEditor.GetWidgetBlueprintObj();
 	if (Blueprint && FindObject<UWidgetAnimation>( Blueprint, *NewAnimationName, true ) )
 	{
@@ -157,7 +178,7 @@ private:
 	{
 		UWidgetAnimation* Animation = ListItem.Pin()->Animation;
 
-		const FName NewName = MakeObjectNameFromDisplayLabel(InText.ToString(), Animation->GetFName());
+		const FName NewName = *InText.ToString();
 
 		if ( Animation->GetFName() != NewName )
 		{
@@ -173,8 +194,8 @@ private:
 		UWidgetAnimation* WidgetAnimation = ListItem.Pin()->Animation;
 		UWidgetBlueprint* Blueprint = BlueprintEditor.Pin()->GetWidgetBlueprintObj();
 
-		// Get the new FName slug from the given display name
-		const FName NewFName = MakeObjectNameFromDisplayLabel(InText.ToString(), WidgetAnimation->GetFName());
+		// Name has already been checked in VerifyAnimationRename
+		const FName NewFName = *InText.ToString();
 		const FName OldFName = WidgetAnimation->GetFName();
 
 		UObjectPropertyBase* ExistingProperty = Cast<UObjectPropertyBase>(Blueprint->ParentClass->FindPropertyByName(NewFName));
@@ -409,7 +430,7 @@ private:
 		UWidgetBlueprint* WidgetBlueprint = Editor->GetWidgetBlueprintObj();
 
 		FString BaseName = "NewAnimation";
-		UWidgetAnimation* NewAnimation = NewObject<UWidgetAnimation>(WidgetBlueprint, *BaseName, RF_Transactional);
+		UWidgetAnimation* NewAnimation = NewObject<UWidgetAnimation>(WidgetBlueprint, FName(), RF_Transactional);
 
 		FString UniqueName = BaseName;
 		int32 NameIndex = 1;
@@ -420,7 +441,7 @@ private:
 			NameIndex++;
 		}
 
-		const FName NewFName = MakeObjectNameFromDisplayLabel(UniqueName, NewAnimation->GetFName());
+		const FName NewFName = FName(*UniqueName);
 		NewAnimation->SetDisplayLabel( UniqueName );
 		NewAnimation->Rename(*UniqueName);
 
@@ -572,7 +593,7 @@ private:
 			);
 	
 		NewAnimation->MovieScene->Rename(*NewAnimation->GetName(), nullptr, REN_DontCreateRedirectors | REN_ForceNoResetLoaders);
-
+		NewAnimation->SetDisplayLabel(NewAnimation->GetName());
 
 		bool bRenameRequestPending = true;
 		bool bNewAnimation = true;
diff --git a/Engine/Source/Editor/UMGEditor/Private/WidgetBlueprintEditor.cpp b/Engine/Source/Editor/UMGEditor/Private/WidgetBlueprintEditor.cpp
--- a/Engine/Source/Editor/UMGEditor/Private/WidgetBlueprintEditor.cpp
+++ b/Engine/Source/Editor/UMGEditor/Private/WidgetBlueprintEditor.cpp
@@ -21,6 +21,7 @@
 
 #include "Tracks/MovieScenePropertyTrack.h"
 #include "ISequencerModule.h"
+#include "SequencerSettings.h"
 #include "ObjectEditorUtils.h"
 
 #include "PropertyCustomizationHelpers.h"
@@ -588,7 +589,6 @@ static bool MigratePropertyValue(UObject* SourceObject, UObject* DestinationObje
 		{
 			if (DestinationObject)
 			{
-				DestinationObject->SetFlags(RF_Transactional);
 				DestinationObject->Modify();
 			}
 			return true;
@@ -666,7 +666,6 @@ void FWidgetBlueprintEditor::PostUndo(bool bSuccessful)
 	FBlueprintEditor::PostUndo(bSuccessful);
 
 	OnWidgetBlueprintTransaction.Broadcast();
-	RefreshPreview();
 }
 
 void FWidgetBlueprintEditor::PostRedo(bool bSuccessful)
@@ -674,7 +673,6 @@ void FWidgetBlueprintEditor::PostRedo(bool bSuccessful)
 	FBlueprintEditor::PostRedo(bSuccessful);
 
 	OnWidgetBlueprintTransaction.Broadcast();
-	RefreshPreview();
 }
 
 TSharedRef<SWidget> FWidgetBlueprintEditor::CreateSequencerWidget()
@@ -793,6 +791,8 @@ TSharedPtr<ISequencer>& FWidgetBlueprintEditor::GetSequencer()
 		};
 
 		Sequencer = FModuleManager::LoadModuleChecked<ISequencerModule>("Sequencer").CreateSequencer(SequencerInitParams);
+		// Never recompile the blueprint on evaluate as this can create an insidious loop
+		Sequencer->GetSequencerSettings()->SetCompileDirectorOnEvaluate(false);
 		Sequencer->OnMovieSceneDataChanged().AddSP( this, &FWidgetBlueprintEditor::OnMovieSceneDataChanged );
 		Sequencer->OnMovieSceneBindingsPasted().AddSP( this, &FWidgetBlueprintEditor::OnMovieSceneBindingsPasted );
 		// Change selected widgets in the sequencer tree view
@@ -1378,6 +1378,18 @@ void FWidgetBlueprintEditor::ReplaceTrackWithSelectedWidget(FWidgetReference Sel
 			}
 		);
 
+		if (!SourceBinding)
+		{
+			WidgetAnimation->BindPossessableObject(ObjectId, *PreviewWidget, GetAnimationPlaybackContext());
+
+			SourceBinding = WidgetAnimation->AnimationBindings.FindByPredicate(
+				[&](FWidgetAnimationBinding& In)
+			{
+				return In.AnimationGuid == ObjectId;
+			}
+			);
+		}
+
 		check(SourceBinding);
 
 		// Set binding names used for lookup
diff --git a/Engine/Source/Editor/UnrealEd/Private/Fbx/FbxCurveDataImport.cpp b/Engine/Source/Editor/UnrealEd/Private/Fbx/FbxCurveDataImport.cpp
--- a/Engine/Source/Editor/UnrealEd/Private/Fbx/FbxCurveDataImport.cpp
+++ b/Engine/Source/Editor/UnrealEd/Private/Fbx/FbxCurveDataImport.cpp
@@ -332,14 +332,16 @@ namespace UnFbx {
 				//so we just need to scale it by the normalized weight value.
 				if (!FMath::IsNearlyZero(ArriveTangentWeight))
 				{
-					const float Y = ArriveTangent * ArriveTangentWeight;
-					ArriveTangentWeight = FMath::Sqrt(Y*Y + ArriveTimeDiff * ArriveTimeDiff);
+					const float X = ArriveTangentWeight * ArriveTimeDiff;
+					const float Y = ArriveTangent * ArriveTangentWeight;  //timediff already there
+					ArriveTangentWeight = FMath::Sqrt(Y*Y + X * X);
 				}
 				NewKey.ArriveTangentWeight = ArriveTangentWeight;
 				if (!FMath::IsNearlyZero(LeaveTangentWeight))
 				{
-					const float Y = LeaveTangent * LeaveTangentWeight;
-					LeaveTangentWeight = FMath::Sqrt(Y*Y + LeaveTimeDiff * LeaveTimeDiff);
+					const float X = LeaveTangentWeight * LeaveTimeDiff;
+					const float Y = LeaveTangent * LeaveTangentWeight; //timediff already there
+					LeaveTangentWeight = FMath::Sqrt(Y*Y + X * X);
 				}
 				NewKey.LeaveTangentWeight = LeaveTangentWeight;
 			}
diff --git a/Engine/Source/Runtime/CinematicCamera/Private/CineCameraComponent.cpp b/Engine/Source/Runtime/CinematicCamera/Private/CineCameraComponent.cpp
--- a/Engine/Source/Runtime/CinematicCamera/Private/CineCameraComponent.cpp
+++ b/Engine/Source/Runtime/CinematicCamera/Private/CineCameraComponent.cpp
@@ -184,6 +184,13 @@ void UCineCameraComponent::ResetProxyMeshTransform()
 
 #endif	// WITH_EDITORONLY_DATA
 
+void UCineCameraComponent::SetFieldOfView(float InFieldOfView)
+{
+	Super::SetFieldOfView(InFieldOfView);
+
+	CurrentFocalLength = (FilmbackSettings.SensorWidth / 2.f) / FMath::Tan(FMath::DegreesToRadians(InFieldOfView / 2.f));
+}
+
 float UCineCameraComponent::GetHorizontalFieldOfView() const
 {
 	return (CurrentFocalLength > 0.f)
diff --git a/Engine/Source/Runtime/CinematicCamera/Public/CineCameraComponent.h b/Engine/Source/Runtime/CinematicCamera/Public/CineCameraComponent.h
--- a/Engine/Source/Runtime/CinematicCamera/Public/CineCameraComponent.h
+++ b/Engine/Source/Runtime/CinematicCamera/Public/CineCameraComponent.h
@@ -254,6 +254,9 @@ public:
 	UPROPERTY(VisibleAnywhere, Category = "Current Camera Settings")
 	float CurrentHorizontalFOV;
 #endif
+
+	/** Override setting FOV to manipulate Focal Length. */
+	virtual void SetFieldOfView(float InFieldOfView) override;
 	
 	/** Returns the horizonal FOV of the camera with current settings. */
 	UFUNCTION(BlueprintCallable, Category = "Cine Camera")
diff --git a/Engine/Source/Runtime/Core/Public/UObject/SequencerObjectVersion.h b/Engine/Source/Runtime/Core/Public/UObject/SequencerObjectVersion.h
--- a/Engine/Source/Runtime/Core/Public/UObject/SequencerObjectVersion.h
+++ b/Engine/Source/Runtime/Core/Public/UObject/SequencerObjectVersion.h
@@ -42,6 +42,9 @@ struct CORE_API FSequencerObjectVersion
 		// Manual serialization of float channel
 		SerializeFloatChannel,
 
+		// Change the linear keys so they act the old way and interpolate always.
+		ModifyLinearKeysForOldInterp,
+
 		// -----<new versions can be added above this line>-------------------------------------------------
 		VersionPlusOne,
 		LatestVersion = VersionPlusOne - 1
diff --git a/Engine/Source/Runtime/Engine/Classes/Camera/CameraComponent.h b/Engine/Source/Runtime/Engine/Classes/Camera/CameraComponent.h
--- a/Engine/Source/Runtime/Engine/Classes/Camera/CameraComponent.h
+++ b/Engine/Source/Runtime/Engine/Classes/Camera/CameraComponent.h
@@ -26,7 +26,7 @@ class ENGINE_API UCameraComponent : public USceneComponent
 	UPROPERTY(Interp, EditAnywhere, BlueprintReadWrite, Category = CameraSettings, meta = (UIMin = "5.0", UIMax = "170", ClampMin = "0.001", ClampMax = "360.0", Units = deg))
 	float FieldOfView;
 	UFUNCTION(BlueprintCallable, Category = Camera)
-	void SetFieldOfView(float InFieldOfView) { FieldOfView = InFieldOfView; }
+	virtual void SetFieldOfView(float InFieldOfView) { FieldOfView = InFieldOfView; }
 
 	/** The desired width (in world units) of the orthographic view (ignored in Perspective mode) */
 	UPROPERTY(Interp, EditAnywhere, BlueprintReadWrite, Category = CameraSettings)
diff --git a/Engine/Source/Runtime/LevelSequence/Private/LevelSequence.cpp b/Engine/Source/Runtime/LevelSequence/Private/LevelSequence.cpp
--- a/Engine/Source/Runtime/LevelSequence/Private/LevelSequence.cpp
+++ b/Engine/Source/Runtime/LevelSequence/Private/LevelSequence.cpp
@@ -186,11 +186,38 @@ void PurgeLegacyBlueprints(UObject* InObject, UPackage* Package)
 }
 #endif
 
+void ULevelSequence::PostDuplicate(bool bDuplicateForPIE)
+{
+	Super::PostDuplicate(bDuplicateForPIE);
+
+#if WITH_EDITORONLY_DATA
+	if (DirectorBlueprint)
+	{
+		DirectorClass = DirectorBlueprint->GeneratedClass.Get();
+
+		// Remove the binding for the director blueprint recompilation and re-add it to be sure there is only one entry in the list
+		DirectorBlueprint->OnCompiled().RemoveAll(this);
+		DirectorBlueprint->OnCompiled().AddUObject(this, &ULevelSequence::OnDirectorRecompiled);
+	}
+	else
+	{
+		DirectorClass = nullptr;
+	}
+#endif
+}
+
 void ULevelSequence::PostLoad()
 {
 	Super::PostLoad();
 
 #if WITH_EDITOR
+	if (DirectorBlueprint)
+	{
+		// Remove the binding for the director blueprint recompilation and re-add it to be sure there is only one entry in the list
+		DirectorBlueprint->OnCompiled().RemoveAll(this);
+		DirectorBlueprint->OnCompiled().AddUObject(this, &ULevelSequence::OnDirectorRecompiled);
+	}
+
 	TSet<FGuid> InvalidSpawnables;
 
 	for (int32 Index = 0; Index < MovieScene->GetSpawnableCount(); ++Index)
@@ -371,6 +398,37 @@ void ULevelSequence::UnbindInvalidObjects(const FGuid& ObjectId, UObject* InCont
 
 #if WITH_EDITOR
 
+UBlueprint* ULevelSequence::GetDirectorBlueprint() const
+{
+	return DirectorBlueprint;
+}
+
+void ULevelSequence::SetDirectorBlueprint(UBlueprint* NewDirectorBlueprint)
+{
+	if (DirectorBlueprint)
+	{
+		DirectorBlueprint->OnCompiled().RemoveAll(this);
+	}
+
+	DirectorBlueprint = NewDirectorBlueprint;
+
+	if (DirectorBlueprint)
+	{
+		DirectorClass = NewDirectorBlueprint->GeneratedClass.Get();
+		DirectorBlueprint->OnCompiled().AddUObject(this, &ULevelSequence::OnDirectorRecompiled);
+	}
+	else
+	{
+		DirectorClass = nullptr;
+	}
+}
+
+void ULevelSequence::OnDirectorRecompiled(UBlueprint* InCompiledBlueprint)
+{
+	ensure(InCompiledBlueprint == DirectorBlueprint);
+	DirectorClass = DirectorBlueprint->GeneratedClass.Get();
+}
+
 FGuid ULevelSequence::FindOrAddBinding(UObject* InObject)
 {
 	UObject* PlaybackContext = InObject ? InObject->GetWorld() : nullptr;
diff --git a/Engine/Source/Runtime/LevelSequence/Public/LevelSequence.h b/Engine/Source/Runtime/LevelSequence/Public/LevelSequence.h
--- a/Engine/Source/Runtime/LevelSequence/Public/LevelSequence.h
+++ b/Engine/Source/Runtime/LevelSequence/Public/LevelSequence.h
@@ -11,6 +11,7 @@
 #include "Templates/SubclassOf.h"
 #include "LevelSequence.generated.h"
 
+class UBlueprint;
 class UMovieScene;
 class UBlueprintGeneratedClass;
 
@@ -60,33 +61,35 @@ public:
 	virtual void GetAssetRegistryTagMetadata(TMap<FName, FAssetRegistryTagMetadata>& OutMetadata) const override;
 	virtual void GetAssetRegistryTags(TArray<FAssetRegistryTag>& OutTags) const override;
 #endif
+	virtual void PostDuplicate(bool bDuplicateForPIE) override;
 
 	void LocateBoundObjects(const FGuid& ObjectId, UObject* Context, FName StreamedLevelAssetPath, TArray<UObject*, TInlineAllocator<1>>& OutObjects) const;
+#if WITH_EDITOR
 
-#if WITH_EDITORONLY_DATA
 
-	/** A pointer to the director blueprint that generates this sequence's DirectorClass. */
-	UPROPERTY()
-	UObject* DirectorBlueprint;
+public:
 
-#endif
+	/**
+	 * Assign a new director blueprint to this level sequence. The specified blueprint *must* be contained within this object.?	 */
+	void SetDirectorBlueprint(UBlueprint* NewDirectorBlueprint);
 
 	/**
-	 * The class that is used to spawn this level sequence's director instance.
-	 * Director instances are allocated on-demand one per sequence during evaluation and are used by event tracks for triggering events.
+	 * Retrieve the currently assigned director blueprint for this level sequence
 	 */
-	UPROPERTY()
-	UBlueprintGeneratedClass* DirectorClass;
+	UBlueprint* GetDirectorBlueprint() const;
 
 protected:
 
-#if WITH_EDITOR
-
 	virtual FGuid CreatePossessable(UObject* ObjectToPossess) override;
 	virtual FGuid CreateSpawnable(UObject* ObjectToSpawn) override;
 
 	FGuid FindOrAddBinding(UObject* ObjectToPossess);
 
+	/**
+	 * Invoked when this level sequence's director blueprint has been recompiled
+	 */
+	void OnDirectorRecompiled(UBlueprint*);
+
 #endif // WITH_EDITOR
 
 protected:
@@ -103,6 +106,21 @@ protected:
 	UPROPERTY()
 	TMap<FString, FLevelSequenceObject> PossessedObjects_DEPRECATED;
 
+#if WITH_EDITORONLY_DATA
+
+	/** A pointer to the director blueprint that generates this sequence's DirectorClass. */
+	UPROPERTY()
+	UBlueprint* DirectorBlueprint;
+
+#endif
+
+	/**
+	 * The class that is used to spawn this level sequence's director instance.
+	 * Director instances are allocated on-demand one per sequence during evaluation and are used by event tracks for triggering events.
+	 */
+	UPROPERTY()
+	UClass* DirectorClass;
+
 #if WITH_EDITOR
 public:
 	/**
diff --git a/Engine/Source/Runtime/MovieScene/Private/Channels/MovieSceneFloatChannel.cpp b/Engine/Source/Runtime/MovieScene/Private/Channels/MovieSceneFloatChannel.cpp
--- a/Engine/Source/Runtime/MovieScene/Private/Channels/MovieSceneFloatChannel.cpp
+++ b/Engine/Source/Runtime/MovieScene/Private/Channels/MovieSceneFloatChannel.cpp
@@ -5,6 +5,13 @@
 #include "Channels/MovieSceneChannelProxy.h"
 #include "MovieSceneFrameMigration.h"
 #include "UObject/SequencerObjectVersion.h"
+#include "HAL/ConsoleManager.h"
+
+static TAutoConsoleVariable<int32> CVarSequencerLinearCubicInterpolation(
+	TEXT("Sequencer.LinearCubicInterpolation"),
+	1,
+	TEXT("If 1 Linear Keys Act As Cubic Interpolation with Linear Tangents, if 0 Linear Key Forces Linear Interpolation to Next Key."),
+	ECVF_Default);
 
 bool FMovieSceneTangentData::Serialize(FArchive& Ar)
 {
@@ -183,14 +190,15 @@ static float EvalForTwoKeys(const FMovieSceneFloatValue& Key1, FFrameNumber Key1
 
 	float Diff = (float)(Key2Time - Key1Time).Value;
 	Diff /= DecimalRate;
-	
+	const int CheckBothLinear = CVarSequencerLinearCubicInterpolation->GetInt();
+
 	if (Diff > 0 && Key1.InterpMode != RCIM_Constant)
 	{
 		const float Alpha = ((float)(InTime - Key1Time).Value / DecimalRate) / Diff;
 		const float P0 = Key1.Value;
 		const float P3 = Key2.Value;
 
-		if (Key1.InterpMode == RCIM_Linear)
+		if (Key1.InterpMode == RCIM_Linear && (!CheckBothLinear || Key2.InterpMode != RCIM_Cubic))
 		{
 			return FMath::Lerp(P0, P3, Alpha);
 		}
@@ -521,6 +529,7 @@ bool FMovieSceneFloatChannel::Evaluate(FFrameTime InTime,  float& OutValue) cons
 	float Interp = 0.f;
 	int32 Index1 = INDEX_NONE, Index2 = INDEX_NONE;
 	MovieScene::EvaluateTime(Times, Params.Time, Index1, Index2, Interp);
+	const int CheckBothLinear = CVarSequencerLinearCubicInterpolation->GetInt();
 
 	if (Index1 == INDEX_NONE)
 	{
@@ -534,8 +543,13 @@ bool FMovieSceneFloatChannel::Evaluate(FFrameTime InTime,  float& OutValue) cons
 	{
 		FMovieSceneFloatValue Key1 = Values[Index1];
 		FMovieSceneFloatValue Key2 = Values[Index2];
-
-		switch (Key1.InterpMode)
+		TEnumAsByte<ERichCurveInterpMode> InterpMode = Key1.InterpMode;
+	    if(InterpMode == RCIM_Linear && (CheckBothLinear  && Key2.InterpMode == RCIM_Cubic))
+		{
+			InterpMode = RCIM_Cubic;
+		}
+		
+		switch (InterpMode)
 		{
 		case RCIM_Cubic:
 		{
@@ -667,7 +681,15 @@ void FMovieSceneFloatChannel::AutoSetTangents(float Tension)
 
 	{
 		FMovieSceneFloatValue& FirstValue = Values[0];
-		if (FirstValue.InterpMode == RCIM_Cubic && FirstValue.TangentMode == RCTM_Auto)
+		if (FirstValue.InterpMode == RCIM_Linear)
+		{
+			FirstValue.Tangent.TangentWeightMode = RCTWM_WeightedNone;
+			FMovieSceneFloatValue& NextKey = Values[1];
+			const float NextTimeDiff = FMath::Max<double>(KINDA_SMALL_NUMBER, Times[1].Value - Times[0].Value);
+			float NewTangent = (NextKey.Value - FirstValue.Value) / NextTimeDiff;
+			FirstValue.Tangent.LeaveTangent = NewTangent;
+		}
+		else if (FirstValue.InterpMode == RCIM_Cubic && FirstValue.TangentMode == RCTM_Auto)
 		{
 			FirstValue.Tangent.LeaveTangent = FirstValue.Tangent.ArriveTangent = 0.0f;
 			FirstValue.Tangent.TangentWeightMode = RCTWM_WeightedNone; 
@@ -676,7 +698,16 @@ void FMovieSceneFloatChannel::AutoSetTangents(float Tension)
 
 	{
 		FMovieSceneFloatValue& LastValue = Values.Last();
-		if (LastValue.InterpMode == RCIM_Cubic && LastValue.TangentMode == RCTM_Auto)
+		if (LastValue.InterpMode == RCIM_Linear)
+		{
+			LastValue.Tangent.TangentWeightMode = RCTWM_WeightedNone;
+			int32 Index = Values.Num() - 1;
+			FMovieSceneFloatValue& PrevKey = Values[Index-1];
+			const float PrevTimeDiff = FMath::Max<double>(KINDA_SMALL_NUMBER, Times[Index].Value - Times[Index - 1].Value);
+			float NewTangent = (LastValue.Value - PrevKey.Value) / PrevTimeDiff;
+			LastValue.Tangent.ArriveTangent = NewTangent;
+		}
+		else if (LastValue.InterpMode == RCIM_Cubic && LastValue.TangentMode == RCTM_Auto)
 		{
 			LastValue.Tangent.LeaveTangent = LastValue.Tangent.ArriveTangent = 0.0f;
 			LastValue.Tangent.TangentWeightMode = RCTWM_WeightedNone;
@@ -691,7 +722,6 @@ void FMovieSceneFloatChannel::AutoSetTangents(float Tension)
 		if (ThisKey.InterpMode == RCIM_Cubic && ThisKey.TangentMode == RCTM_Auto)
 		{
 			FMovieSceneFloatValue& NextKey = Values[Index+1];
-
 			const float PrevToNextTimeDiff = FMath::Max<double>(KINDA_SMALL_NUMBER, Times[Index+1].Value - Times[Index-1].Value);
 
 			float NewTangent = 0.f;
@@ -702,6 +732,19 @@ void FMovieSceneFloatChannel::AutoSetTangents(float Tension)
 			ThisKey.Tangent.LeaveTangent = ThisKey.Tangent.ArriveTangent = NewTangent;
 			ThisKey.Tangent.TangentWeightMode = RCTWM_WeightedNone;
 		}
+		else if (ThisKey.InterpMode == RCIM_Linear)
+		{
+			ThisKey.Tangent.TangentWeightMode = RCTWM_WeightedNone; 
+			FMovieSceneFloatValue& NextKey = Values[Index + 1];
+
+			const float PrevTimeDiff = FMath::Max<double>(KINDA_SMALL_NUMBER, Times[Index].Value - Times[Index - 1].Value);
+			float NewTangent  = (ThisKey.Value - PrevKey.Value) / PrevTimeDiff;
+			ThisKey.Tangent.ArriveTangent = NewTangent;
+			
+			const float NextTimeDiff = FMath::Max<double>(KINDA_SMALL_NUMBER, Times[Index + 1].Value - Times[Index].Value);
+			NewTangent = (NextKey.Value - ThisKey.Value) / NextTimeDiff;
+			ThisKey.Tangent.LeaveTangent = NewTangent;
+		}
 		else if (PrevKey.InterpMode == RCIM_Constant || ThisKey.InterpMode == RCIM_Constant)
 		{
 			if (PrevKey.InterpMode != RCIM_Cubic)
@@ -711,6 +754,7 @@ void FMovieSceneFloatChannel::AutoSetTangents(float Tension)
 
 			ThisKey.Tangent.LeaveTangent  = 0.0f;
 		}
+		
 	}
 }
 
@@ -851,6 +895,11 @@ void FMovieSceneFloatChannel::Reset()
 	bHasDefaultValue = false;
 }
 
+void FMovieSceneFloatChannel::PostEditChange()
+{
+	AutoSetTangents();
+}
+
 void FMovieSceneFloatChannel::Offset(FFrameNumber DeltaPosition)
 {
 	GetData().Offset(DeltaPosition);
@@ -946,4 +995,42 @@ void FMovieSceneFloatChannel::AddKeys(const TArray<FFrameNumber>& InTimes, const
 	{
 		KeyHandles.AllocateHandle(Index);
 	}
-}
\ No newline at end of file
+}
+
+bool FMovieSceneFloatChannel::Serialize(FArchive& Ar)
+{
+	Ar.UsingCustomVersion(FSequencerObjectVersion::GUID);
+	return false;
+}
+
+void FMovieSceneFloatChannel::PostSerialize(const FArchive& Ar)
+{
+	if (Ar.CustomVer(FSequencerObjectVersion::GUID) < FSequencerObjectVersion::ModifyLinearKeysForOldInterp)
+	{
+		bool bNeedAutoSetAlso = false;
+		//we need to possibly modify cuvic tangents if we get a set of linear..cubic tangents so it works like it used to
+		if (Values.Num() >= 2)
+		{
+			for (int32 Index = 1; Index < Values.Num(); ++Index)
+			{
+				FMovieSceneFloatValue  PrevKey = Values[Index - 1];
+				FMovieSceneFloatValue& ThisKey = Values[Index];
+
+				if (ThisKey.InterpMode == RCIM_Cubic && PrevKey.InterpMode == RCIM_Linear)
+				{
+					ThisKey.Tangent.TangentWeightMode = RCTWM_WeightedNone;
+					ThisKey.TangentMode = RCTM_Break;
+					//leave next tangent will be set up if auto or user, just need to modify prev.
+					const float PrevTimeDiff = FMath::Max<double>(KINDA_SMALL_NUMBER, Times[Index].Value - Times[Index - 1].Value);
+					float NewTangent = (ThisKey.Value - PrevKey.Value) / PrevTimeDiff;
+					ThisKey.Tangent.ArriveTangent = NewTangent;
+					bNeedAutoSetAlso = true;
+				}
+			}
+		}
+		if (bNeedAutoSetAlso)
+		{
+			AutoSetTangents();
+		}
+	}
+}
diff --git a/Engine/Source/Runtime/MovieScene/Private/MovieScene.cpp b/Engine/Source/Runtime/MovieScene/Private/MovieScene.cpp
--- a/Engine/Source/Runtime/MovieScene/Private/MovieScene.cpp
+++ b/Engine/Source/Runtime/MovieScene/Private/MovieScene.cpp
@@ -104,8 +104,8 @@ void UMovieScene::Serialize( FArchive& Ar )
 		else if (FixedFrameInterval_DEPRECATED == 1 / 120.0f) { DisplayRate = FCommonFrameRates::FPS_120(); }
 		else if (FixedFrameInterval_DEPRECATED != 0.f)
 		{
-			uint32 Numerator = FMath::RoundToInt(1.f / FixedFrameInterval_DEPRECATED);
-			DisplayRate = FFrameRate(Numerator, 1);
+			uint32 Numerator = FMath::RoundToInt(1000.f / FixedFrameInterval_DEPRECATED);
+			DisplayRate = FFrameRate(Numerator, 1000);
 		}
 		else
 		{
diff --git a/Engine/Source/Runtime/MovieScene/Private/MovieSceneSequencePlayer.cpp b/Engine/Source/Runtime/MovieScene/Private/MovieSceneSequencePlayer.cpp
--- a/Engine/Source/Runtime/MovieScene/Private/MovieSceneSequencePlayer.cpp
+++ b/Engine/Source/Runtime/MovieScene/Private/MovieSceneSequencePlayer.cpp
@@ -290,6 +290,9 @@ void UMovieSceneSequencePlayer::StopInternal(FFrameTime TimeToResetTo)
 
 		CurrentNumLoops = 0;
 
+		// Reset loop count on stop so that it doesn't persist to the next call to play
+		PlaybackSettings.LoopCount.Value = 0;
+
 		if (PlaybackSettings.bRestoreState)
 		{
 			RestorePreAnimatedState();
diff --git a/Engine/Source/Runtime/MovieScene/Public/Channels/MovieSceneChannel.h b/Engine/Source/Runtime/MovieScene/Public/Channels/MovieSceneChannel.h
--- a/Engine/Source/Runtime/MovieScene/Public/Channels/MovieSceneChannel.h
+++ b/Engine/Source/Runtime/MovieScene/Public/Channels/MovieSceneChannel.h
@@ -140,4 +140,11 @@ struct MOVIESCENE_API FMovieSceneChannel
 	 */
 	virtual void ClearDefault()
 	{}
+
+	/**
+	 * Perfor a possibly heavy operation after an edit change 
+	 *
+	 */
+	virtual void PostEditChange() {}
+
 };
\ No newline at end of file
diff --git a/Engine/Source/Runtime/MovieScene/Public/Channels/MovieSceneFloatChannel.h b/Engine/Source/Runtime/MovieScene/Public/Channels/MovieSceneFloatChannel.h
--- a/Engine/Source/Runtime/MovieScene/Public/Channels/MovieSceneFloatChannel.h
+++ b/Engine/Source/Runtime/MovieScene/Public/Channels/MovieSceneFloatChannel.h
@@ -222,6 +222,7 @@ public:
 	virtual void Offset(FFrameNumber DeltaPosition) override;
 	virtual void Optimize(const FKeyDataOptimizationParams& InParameters) override;
 	virtual void ClearDefault() override;
+	virtual void PostEditChange() override;
 
 public:
 
@@ -256,6 +257,15 @@ public:
 
 public:
 
+	bool Serialize(FArchive& Ar);
+	void PostSerialize(const FArchive& Ar);
+	friend FArchive& operator<<(FArchive& Ar, FMovieSceneFloatChannel& Me)
+	{
+		Me.Serialize(Ar);
+		return Ar;
+	}
+
+
 	/** Serialize this float function from a mismatching property tag (FRichCurve) */
 	bool SerializeFromMismatchedTag(const FPropertyTag& Tag, FStructuredArchive::FSlot Slot);
 
@@ -349,7 +359,12 @@ public:
 template<>
 struct TStructOpsTypeTraits<FMovieSceneFloatChannel> : public TStructOpsTypeTraitsBase2<FMovieSceneFloatChannel>
 {
-	enum { WithStructuredSerializeFromMismatchedTag = true };
+	enum 
+	{ 
+		WithStructuredSerializeFromMismatchedTag = true, 
+	    WithSerializer = true,
+		WithPostSerialize = true,
+    };
 };
 
 template<>
diff --git a/Engine/Source/Runtime/MovieScene/Public/MovieSceneKeyProxy.h b/Engine/Source/Runtime/MovieScene/Public/MovieSceneKeyProxy.h
--- a/Engine/Source/Runtime/MovieScene/Public/MovieSceneKeyProxy.h
+++ b/Engine/Source/Runtime/MovieScene/Public/MovieSceneKeyProxy.h
@@ -89,6 +89,7 @@ void IMovieSceneKeyProxy::OnProxyValueChanged(TMovieSceneChannelHandle<ChannelTy
 			ChannelData.MoveKey(KeyIndex, InOutTime);
 			InSection->ExpandToFrame(InOutTime);
 		}
+		Channel->PostEditChange();
 	}
 }
 
diff --git a/Engine/Source/Runtime/MovieSceneCapture/Private/MovieSceneCapture.cpp b/Engine/Source/Runtime/MovieSceneCapture/Private/MovieSceneCapture.cpp
--- a/Engine/Source/Runtime/MovieSceneCapture/Private/MovieSceneCapture.cpp
+++ b/Engine/Source/Runtime/MovieSceneCapture/Private/MovieSceneCapture.cpp
@@ -95,6 +95,14 @@ FMovieSceneCaptureSettings::FMovieSceneCaptureSettings()
 	bAllowTurning = false;
 	bShowPlayer = false;
 	bShowHUD = false;
+
+#if PLATFORM_MAC
+	MovieExtension = TEXT(".mov");
+#elif PLATFORM_UNIX
+	MovieExtension = TEXT(".unsupp");
+#else
+	MovieExtension = TEXT(".avi");
+#endif
 }
 
 UMovieSceneCapture::UMovieSceneCapture(const FObjectInitializer& Initializer)
diff --git a/Engine/Source/Runtime/MovieSceneCapture/Private/UserDefinedCaptureProtocol.cpp b/Engine/Source/Runtime/MovieSceneCapture/Private/UserDefinedCaptureProtocol.cpp
--- a/Engine/Source/Runtime/MovieSceneCapture/Private/UserDefinedCaptureProtocol.cpp
+++ b/Engine/Source/Runtime/MovieSceneCapture/Private/UserDefinedCaptureProtocol.cpp
@@ -11,6 +11,9 @@
 #include "Slate/SceneViewport.h"
 #include "ImagePixelData.h"
 #include "Engine/Engine.h"
+#include "Logging/MessageLog.h"
+
+#define LOCTEXT_NAMESPACE "UserDefinedImageCaptureProtocol"
 
 struct FCaptureProtocolFrameData : IFramePayload
 {
@@ -65,6 +68,26 @@ private:
 	TWeakObjectPtr<UUserDefinedCaptureProtocol> WeakProtocol;
 };
 
+FString FCapturedPixelsID::ToString() const
+{
+	FString Name;
+	for (const TTuple<FName, FName>& Pair : Identifiers)
+	{
+		if (Name.Len() > 0)
+		{
+			Name += TEXT(",");
+		}
+
+		Name += Pair.Key.ToString();
+		if (Pair.Value != NAME_None)
+		{
+			Name += TEXT(":");
+			Name += Pair.Value.ToString();
+		}
+	}
+	return Name.Len() > 0 ? Name : TEXT("<none>");
+}
+
 
 UUserDefinedCaptureProtocol::UUserDefinedCaptureProtocol(const FObjectInitializer& ObjInit)
 	: Super(ObjInit)
@@ -371,15 +394,26 @@ void UUserDefinedImageCaptureProtocol::WriteImageToDisk(const FCapturedPixels& P
 	{
 		return;
 	}
-
-	TUniquePtr<FImageWriteTask> ImageTask = MakeUnique<FImageWriteTask>();
-
-	// If the pixels are FColors, and this is the final pixels buffer, and we're writing PNG, always write out full alpha
-	if (PixelData.ImageData->GetType() == EImagePixelType::Color && ImageTask->Format == EImageFormat::PNG && StreamID.Identifiers.OrderIndependentCompareEqual(FinalPixelsID.Identifiers))
+	else if (PixelData.ImageData->GetBitDepth() != 8)
 	{
-		ImageTask->PixelPreProcessors.Add(TAsyncAlphaWrite<FColor>(255));
+		if (Format == EDesiredImageFormat::BMP)
+		{
+			FMessageLog("PIE")
+			.Warning(FText::Format(LOCTEXT("InvalidBMPExport", "Unable to write the specified render target (stream '{0}' is {1}bit) as BMP. BMPs must be supplied 8bit render targets."),
+				FText::FromString(StreamID.ToString()), FText::AsNumber(PixelData.ImageData->GetBitDepth())));
+			return;
+		}
+		else if (Format == EDesiredImageFormat::JPG)
+		{
+			FMessageLog("PIE")
+				.Warning(FText::Format(LOCTEXT("InvalidJPGExport", "Unable to write the specified render target (stream '{0}' is {1}bit) as JPG. JPGs must be supplied 8bit render targets."),
+					FText::FromString(StreamID.ToString()), FText::AsNumber(PixelData.ImageData->GetBitDepth())));
+			return;
+		}
 	}
 
+	TUniquePtr<FImageWriteTask> ImageTask = MakeUnique<FImageWriteTask>();
+
 	// Cache the buffer ID so we generate the correct filename
 	CurrentStreamID = &StreamID;
 
@@ -390,6 +424,12 @@ void UUserDefinedImageCaptureProtocol::WriteImageToDisk(const FCapturedPixels& P
 
 	CurrentStreamID = nullptr;
 
+	// If the pixels are FColors, and this is the final pixels buffer, and we're writing PNG, always write out full alpha
+	if (PixelData.ImageData->GetType() == EImagePixelType::Color && ImageTask->Format == EImageFormat::PNG && StreamID.Identifiers.OrderIndependentCompareEqual(FinalPixelsID.Identifiers))
+	{
+		ImageTask->PixelPreProcessors.Add(TAsyncAlphaWrite<FColor>(255));
+	}
+
 	if (Format == EDesiredImageFormat::EXR)
 	{
 		ImageTask->CompressionQuality = bEnableCompression ? (int32)EImageCompressionQuality::Default : (int32)EImageCompressionQuality::Uncompressed;
@@ -424,4 +464,6 @@ void UUserDefinedImageCaptureProtocol::WriteImageToDisk(const FCapturedPixels& P
 void UUserDefinedImageCaptureProtocol::OnFileWritten()
 {
 	--NumOutstandingOperations;
-}
\ No newline at end of file
+}
+
+#undef LOCTEXT_NAMESPACE
\ No newline at end of file
diff --git a/Engine/Source/Runtime/MovieSceneCapture/Public/MovieSceneCaptureSettings.h b/Engine/Source/Runtime/MovieSceneCapture/Public/MovieSceneCaptureSettings.h
--- a/Engine/Source/Runtime/MovieSceneCapture/Public/MovieSceneCaptureSettings.h
+++ b/Engine/Source/Runtime/MovieSceneCapture/Public/MovieSceneCaptureSettings.h
@@ -67,6 +67,10 @@ struct MOVIESCENECAPTURE_API FMovieSceneCaptureSettings
 	UPROPERTY(config, EditAnywhere, BlueprintReadWrite, Category=Sequence, AdvancedDisplay, meta=(ClampMin=0, UIMin=0))
 	int32 HandleFrames;
 
+	/** Filename extension for movies referenced in the XMLs/EDLs */
+	UPROPERTY(config, EditAnywhere, BlueprintReadWrite, Category=Sequence, AdvancedDisplay)
+	FString MovieExtension;
+
 	/** How much to zero-pad frame numbers on filenames */
 	UPROPERTY(config, EditAnywhere, BlueprintReadWrite, Category=General, AdvancedDisplay, meta=(ClampMin=0, UIMin=0))
 	uint8 ZeroPadFrameNumbers;
diff --git a/Engine/Source/Runtime/MovieSceneCapture/Public/Protocols/UserDefinedCaptureProtocol.h b/Engine/Source/Runtime/MovieSceneCapture/Public/Protocols/UserDefinedCaptureProtocol.h
--- a/Engine/Source/Runtime/MovieSceneCapture/Public/Protocols/UserDefinedCaptureProtocol.h
+++ b/Engine/Source/Runtime/MovieSceneCapture/Public/Protocols/UserDefinedCaptureProtocol.h
@@ -20,6 +20,9 @@ struct FCapturedPixelsID
 {
 	GENERATED_BODY()
 
+	/** Convert this strream ID to a string */
+	FString ToString() const;
+
 	/** Map of identifiers to their values for this ID. */
 	UPROPERTY(EditAnywhere, BlueprintReadWrite, Category=Capture)
 	TMap<FName, FName> Identifiers;
diff --git a/Engine/Source/Runtime/MovieSceneTracks/Private/Evaluation/MovieSceneObjectPropertyTemplate.cpp b/Engine/Source/Runtime/MovieSceneTracks/Private/Evaluation/MovieSceneObjectPropertyTemplate.cpp
--- a/Engine/Source/Runtime/MovieSceneTracks/Private/Evaluation/MovieSceneObjectPropertyTemplate.cpp
+++ b/Engine/Source/Runtime/MovieSceneTracks/Private/Evaluation/MovieSceneObjectPropertyTemplate.cpp
@@ -54,10 +54,11 @@ struct FObjectPropertyExecToken : IMovieSceneExecutionToken
 		{
 			return !TargetProperty->HasAnyPropertyFlags(CPF_NoClear);
 		}
-		else
+		else if (DesiredValue->GetClass() != nullptr)
 		{
 			return DesiredValue->GetClass()->IsChildOf(TargetProperty->PropertyClass);
 		}
+		return false;
 	}
 
 	UObject* NewObjectValue;
diff --git a/Engine/Source/Runtime/MovieSceneTracks/Private/Evaluation/MovieSceneSkeletalAnimationTemplate.cpp b/Engine/Source/Runtime/MovieSceneTracks/Private/Evaluation/MovieSceneSkeletalAnimationTemplate.cpp
--- a/Engine/Source/Runtime/MovieSceneTracks/Private/Evaluation/MovieSceneSkeletalAnimationTemplate.cpp
+++ b/Engine/Source/Runtime/MovieSceneTracks/Private/Evaluation/MovieSceneSkeletalAnimationTemplate.cpp
@@ -242,7 +242,7 @@ namespace MovieScene
 			// If the skeletal component has already ticked this frame because tick prerequisites weren't set up yet or a new binding was created, forcibly tick this component to update.
 			// This resolves first frame issues where the skeletal component ticks first, then the sequencer binding is resolved which sets up tick prerequisites
 			// for the next frame.
-			if (SkeletalMeshComponent->PoseTickedThisFrame() || SequencerInstance != ExistingAnimInstance)
+			if (SkeletalMeshComponent->PoseTickedThisFrame() || (SequencerInstance && SequencerInstance != ExistingAnimInstance))
 			{
 				SkeletalMeshComponent->TickAnimation(0.f, false);
 
diff --git a/Engine/Source/Runtime/MovieSceneTracks/Private/Sections/MovieScene3DTransformSection.cpp b/Engine/Source/Runtime/MovieSceneTracks/Private/Sections/MovieScene3DTransformSection.cpp
--- a/Engine/Source/Runtime/MovieSceneTracks/Private/Sections/MovieScene3DTransformSection.cpp
+++ b/Engine/Source/Runtime/MovieSceneTracks/Private/Sections/MovieScene3DTransformSection.cpp
@@ -23,48 +23,57 @@ struct F3DTransformChannelEditorData
 			MetaData[0].bEnabled = EnumHasAllFlags(Mask, EMovieSceneTransformChannel::TranslationX);
 			MetaData[0].Color = FCommonChannelData::RedChannelColor;
 			MetaData[0].SortOrder = 0;
+			MetaData[0].bCanCollapseToTrack = false;
 
 			MetaData[1].SetIdentifiers("Location.Y", FCommonChannelData::ChannelY, LocationGroup);
 			MetaData[1].bEnabled = EnumHasAllFlags(Mask, EMovieSceneTransformChannel::TranslationY);
 			MetaData[1].Color = FCommonChannelData::GreenChannelColor;
 			MetaData[1].SortOrder = 1;
+			MetaData[1].bCanCollapseToTrack = false;
 
 			MetaData[2].SetIdentifiers("Location.Z", FCommonChannelData::ChannelZ, LocationGroup);
 			MetaData[2].bEnabled = EnumHasAllFlags(Mask, EMovieSceneTransformChannel::TranslationZ);
 			MetaData[2].Color = FCommonChannelData::BlueChannelColor;
 			MetaData[2].SortOrder = 2;
+			MetaData[2].bCanCollapseToTrack = false;
 		}
 		{
 			MetaData[3].SetIdentifiers("Rotation.X", NSLOCTEXT("MovieSceneTransformSection", "RotationX", "Roll"), RotationGroup);
 			MetaData[3].bEnabled = EnumHasAllFlags(Mask, EMovieSceneTransformChannel::RotationX);
 			MetaData[3].Color = FCommonChannelData::RedChannelColor;
 			MetaData[3].SortOrder = 3;
+			MetaData[3].bCanCollapseToTrack = false;
 
 			MetaData[4].SetIdentifiers("Rotation.Y", NSLOCTEXT("MovieSceneTransformSection", "RotationY", "Pitch"), RotationGroup);
 			MetaData[4].bEnabled = EnumHasAllFlags(Mask, EMovieSceneTransformChannel::RotationY);
 			MetaData[4].Color = FCommonChannelData::GreenChannelColor;
 			MetaData[4].SortOrder = 4;
+			MetaData[4].bCanCollapseToTrack = false;
 
 			MetaData[5].SetIdentifiers("Rotation.Z", NSLOCTEXT("MovieSceneTransformSection", "RotationZ", "Yaw"), RotationGroup);
 			MetaData[5].bEnabled = EnumHasAllFlags(Mask, EMovieSceneTransformChannel::RotationZ);
 			MetaData[5].Color = FCommonChannelData::BlueChannelColor;
 			MetaData[5].SortOrder = 5;
+			MetaData[5].bCanCollapseToTrack = false;
 		}
 		{
 			MetaData[6].SetIdentifiers("Scale.X", FCommonChannelData::ChannelX, ScaleGroup);
 			MetaData[6].bEnabled = EnumHasAllFlags(Mask, EMovieSceneTransformChannel::ScaleX);
 			MetaData[6].Color = FCommonChannelData::RedChannelColor;
 			MetaData[6].SortOrder = 6;
+			MetaData[6].bCanCollapseToTrack = false;
 
 			MetaData[7].SetIdentifiers("Scale.Y", FCommonChannelData::ChannelY, ScaleGroup);
 			MetaData[7].bEnabled = EnumHasAllFlags(Mask, EMovieSceneTransformChannel::ScaleY);
 			MetaData[7].Color = FCommonChannelData::GreenChannelColor;
 			MetaData[7].SortOrder = 7;
+			MetaData[7].bCanCollapseToTrack = false;
 
 			MetaData[8].SetIdentifiers("Scale.Z", FCommonChannelData::ChannelZ, ScaleGroup);
 			MetaData[8].bEnabled = EnumHasAllFlags(Mask, EMovieSceneTransformChannel::ScaleZ);
 			MetaData[8].Color = FCommonChannelData::BlueChannelColor;
 			MetaData[8].SortOrder = 8;
+			MetaData[8].bCanCollapseToTrack = false;
 		}
 		{
 			MetaData[9].SetIdentifiers("Weight", NSLOCTEXT("MovieSceneTransformSection", "Weight", "Weight"));
diff --git a/Engine/Source/Runtime/MovieSceneTracks/Private/Sections/MovieSceneColorSection.cpp b/Engine/Source/Runtime/MovieSceneTracks/Private/Sections/MovieSceneColorSection.cpp
--- a/Engine/Source/Runtime/MovieSceneTracks/Private/Sections/MovieSceneColorSection.cpp
+++ b/Engine/Source/Runtime/MovieSceneTracks/Private/Sections/MovieSceneColorSection.cpp
@@ -14,17 +14,21 @@ struct FColorSectionEditorData
 		MetaData[0].SetIdentifiers("Color.R", FCommonChannelData::ChannelR);
 		MetaData[0].SortOrder = 0;
 		MetaData[0].Color = FCommonChannelData::RedChannelColor;
+		MetaData[0].bCanCollapseToTrack = false;
 
 		MetaData[1].SetIdentifiers("Color.G", FCommonChannelData::ChannelG);
 		MetaData[1].SortOrder = 1;
 		MetaData[1].Color = FCommonChannelData::GreenChannelColor;
+		MetaData[1].bCanCollapseToTrack = false;
 
 		MetaData[2].SetIdentifiers("Color.B", FCommonChannelData::ChannelB);
 		MetaData[2].SortOrder = 2;
 		MetaData[2].Color = FCommonChannelData::BlueChannelColor;
+		MetaData[2].bCanCollapseToTrack = false;
 
 		MetaData[3].SetIdentifiers("Color.A", FCommonChannelData::ChannelA);
 		MetaData[3].SortOrder = 3;
+		MetaData[3].bCanCollapseToTrack = false;
 
 		ExternalValues[0].OnGetExternalValue = ExtractChannelR;
 		ExternalValues[1].OnGetExternalValue = ExtractChannelG;
diff --git a/Engine/Source/Runtime/MovieSceneTracks/Private/Sections/MovieSceneVectorSection.cpp b/Engine/Source/Runtime/MovieSceneTracks/Private/Sections/MovieSceneVectorSection.cpp
--- a/Engine/Source/Runtime/MovieSceneTracks/Private/Sections/MovieSceneVectorSection.cpp
+++ b/Engine/Source/Runtime/MovieSceneTracks/Private/Sections/MovieSceneVectorSection.cpp
@@ -16,17 +16,21 @@ struct FVectorSectionEditorData
 		MetaData[0].SetIdentifiers("Vector.X", FCommonChannelData::ChannelX);
 		MetaData[0].SortOrder = 0;
 		MetaData[0].Color = FCommonChannelData::RedChannelColor;
+		MetaData[0].bCanCollapseToTrack = false;
 
 		MetaData[1].SetIdentifiers("Vector.Y", FCommonChannelData::ChannelY);
 		MetaData[1].SortOrder = 1;
 		MetaData[1].Color = FCommonChannelData::GreenChannelColor;
+		MetaData[1].bCanCollapseToTrack = false;
 
 		MetaData[2].SetIdentifiers("Vector.Z", FCommonChannelData::ChannelZ);
 		MetaData[2].SortOrder = 2;
 		MetaData[2].Color = FCommonChannelData::BlueChannelColor;
+		MetaData[2].bCanCollapseToTrack = false;
 
 		MetaData[3].SetIdentifiers("Vector.W", FCommonChannelData::ChannelW);
 		MetaData[3].SortOrder = 3;
+		MetaData[3].bCanCollapseToTrack = false;
 
 		ExternalValues[0].OnGetExternalValue = [NumChannels](UObject& InObject, FTrackInstancePropertyBindings* Bindings) { return ExtractChannelX(InObject, Bindings, NumChannels); };
 		ExternalValues[1].OnGetExternalValue = [NumChannels](UObject& InObject, FTrackInstancePropertyBindings* Bindings) { return ExtractChannelY(InObject, Bindings, NumChannels); };
diff --git a/Engine/Source/Runtime/MovieSceneTracks/Public/Tracks/MovieSceneSkeletalAnimationTrack.h b/Engine/Source/Runtime/MovieSceneTracks/Public/Tracks/MovieSceneSkeletalAnimationTrack.h
--- a/Engine/Source/Runtime/MovieSceneTracks/Public/Tracks/MovieSceneSkeletalAnimationTrack.h
+++ b/Engine/Source/Runtime/MovieSceneTracks/Public/Tracks/MovieSceneSkeletalAnimationTrack.h
@@ -19,13 +19,13 @@ class UMovieSceneSkeletalAnimationTrack
 public:
 
 	/** Adds a new animation to this track */
-	virtual UMovieSceneSection* AddNewAnimationOnRow(FFrameNumber KeyTime, class UAnimSequenceBase* AnimSequence, int32 RowIndex);
+	MOVIESCENETRACKS_API virtual UMovieSceneSection* AddNewAnimationOnRow(FFrameNumber KeyTime, class UAnimSequenceBase* AnimSequence, int32 RowIndex);
 
 	/** Adds a new animation to this track on the next available/non-overlapping row */
-	virtual UMovieSceneSection* AddNewAnimation(FFrameNumber KeyTime, class UAnimSequenceBase* AnimSequence) { return AddNewAnimationOnRow(KeyTime, AnimSequence, INDEX_NONE); }
+	MOVIESCENETRACKS_API virtual UMovieSceneSection* AddNewAnimation(FFrameNumber KeyTime, class UAnimSequenceBase* AnimSequence) { return AddNewAnimationOnRow(KeyTime, AnimSequence, INDEX_NONE); }
 
 	/** Gets the animation sections at a certain time */
-	TArray<UMovieSceneSection*> GetAnimSectionsAtTime(FFrameNumber Time);
+	MOVIESCENETRACKS_API TArray<UMovieSceneSection*> GetAnimSectionsAtTime(FFrameNumber Time);
 
 public:
 
diff --git a/Engine/Source/Runtime/UMG/Private/Animation/MovieScene2DTransformSection.cpp b/Engine/Source/Runtime/UMG/Private/Animation/MovieScene2DTransformSection.cpp
--- a/Engine/Source/Runtime/UMG/Private/Animation/MovieScene2DTransformSection.cpp
+++ b/Engine/Source/Runtime/UMG/Private/Animation/MovieScene2DTransformSection.cpp
@@ -20,30 +20,37 @@ struct F2DTransformSectionEditorData
 		MetaData[0].SetIdentifiers("Translation.X", FCommonChannelData::ChannelX, TranslationGroup);
 		MetaData[0].bEnabled = EnumHasAllFlags(Mask, EMovieScene2DTransformChannel::TranslationX);
 		MetaData[0].SortOrder = 0;
+		MetaData[0].bCanCollapseToTrack = false;
 
 		MetaData[1].SetIdentifiers("Translation.Y", FCommonChannelData::ChannelY, TranslationGroup);
 		MetaData[1].bEnabled = EnumHasAllFlags(Mask, EMovieScene2DTransformChannel::TranslationY);
 		MetaData[1].SortOrder = 1;
+		MetaData[1].bCanCollapseToTrack = false;
 
 		MetaData[2].SetIdentifiers("Angle", NSLOCTEXT("MovieScene2DTransformSection", "AngleText", "Angle"), RotationGroup);
 		MetaData[2].bEnabled = EnumHasAllFlags(Mask, EMovieScene2DTransformChannel::Rotation);
 		MetaData[2].SortOrder = 2;
+		MetaData[2].bCanCollapseToTrack = false;
 
 		MetaData[3].SetIdentifiers("Scale.X", FCommonChannelData::ChannelX, ScaleGroup);
 		MetaData[3].bEnabled = EnumHasAllFlags(Mask, EMovieScene2DTransformChannel::ScaleX);
 		MetaData[3].SortOrder = 3;
-		
+		MetaData[3].bCanCollapseToTrack = false;
+
 		MetaData[4].SetIdentifiers("Scale.Y", FCommonChannelData::ChannelY, ScaleGroup);
 		MetaData[4].bEnabled = EnumHasAllFlags(Mask, EMovieScene2DTransformChannel::ScaleY);
 		MetaData[4].SortOrder = 4;
+		MetaData[4].bCanCollapseToTrack = false;
 
 		MetaData[5].SetIdentifiers("Shear.X", FCommonChannelData::ChannelX, ShearGroup);
 		MetaData[5].bEnabled = EnumHasAllFlags(Mask, EMovieScene2DTransformChannel::ShearX);
 		MetaData[5].SortOrder = 5;
-		
+		MetaData[5].bCanCollapseToTrack = false;
+
 		MetaData[6].SetIdentifiers("Shear.Y", FCommonChannelData::ChannelY, ShearGroup);
 		MetaData[6].bEnabled = EnumHasAllFlags(Mask, EMovieScene2DTransformChannel::ShearY);
 		MetaData[6].SortOrder = 6;
+		MetaData[6].bCanCollapseToTrack = false;
 
 		ExternalValues[0].OnGetExternalValue = ExtractTranslationX;
 		ExternalValues[1].OnGetExternalValue = ExtractTranslationY;
@@ -191,7 +198,7 @@ FMovieScene2DTransformMask UMovieScene2DTransformSection::GetMaskByName(const FN
 	{
 		return EMovieScene2DTransformChannel::TranslationY;
 	}
-	else if (InName == TEXT("Angle"))
+	else if (InName == TEXT("Angle") || InName == TEXT("Rotation"))
 	{
 		return EMovieScene2DTransformChannel::Rotation;
 	}
diff --git a/Engine/Source/Runtime/UMG/Private/Animation/MovieSceneMarginSection.cpp b/Engine/Source/Runtime/UMG/Private/Animation/MovieSceneMarginSection.cpp
--- a/Engine/Source/Runtime/UMG/Private/Animation/MovieSceneMarginSection.cpp
+++ b/Engine/Source/Runtime/UMG/Private/Animation/MovieSceneMarginSection.cpp
@@ -11,15 +11,19 @@ struct FMarginSectionEditorData
 	{
 		MetaData[0].SetIdentifiers("Left", NSLOCTEXT("MovieSceneMarginSection", "LeftText", "Left"));
 		MetaData[0].SortOrder = 0;
+		MetaData[0].bCanCollapseToTrack = false;
 
 		MetaData[1].SetIdentifiers("Top", NSLOCTEXT("MovieSceneMarginSection", "TopText", "Top"));
 		MetaData[1].SortOrder = 1;
+		MetaData[1].bCanCollapseToTrack = false;
 
 		MetaData[2].SetIdentifiers("Right", NSLOCTEXT("MovieSceneMarginSection", "RightText", "Right"));
 		MetaData[2].SortOrder = 2;
+		MetaData[2].bCanCollapseToTrack = false;
 
 		MetaData[3].SetIdentifiers("Bottom", NSLOCTEXT("MovieSceneMarginSection", "BottomText", "Bottom"));
 		MetaData[3].SortOrder = 3;
+		MetaData[3].bCanCollapseToTrack = false;
 
 		ExternalValues[0].OnGetExternalValue = ExtractLeftChannel;
 		ExternalValues[1].OnGetExternalValue = ExtractTopChannel;
