Fixed localization gather warnings

diff --git a/Engine/Source/Editor/BlueprintGraph/Private/K2Node_GetSubsystem.cpp b/Engine/Source/Editor/BlueprintGraph/Private/K2Node_GetSubsystem.cpp
--- a/Engine/Source/Editor/BlueprintGraph/Private/K2Node_GetSubsystem.cpp
+++ b/Engine/Source/Editor/BlueprintGraph/Private/K2Node_GetSubsystem.cpp
@@ -702,7 +702,7 @@ void UK2Node_GetEditorSubsystem::ValidateNodeDuringCompilation(class FCompilerRe
 
 	if (!FBlueprintEditorUtils::IsBlutility(BP))
 	{
-		const FText ErrorText = NSLOCTEXT("K2Node", "GetSubsystem_Error", "Editor Subsystems can only be used in Editor Utilities / Blutilities");
+		const FText ErrorText = NSLOCTEXT("K2Node", "GetEditorSubsystem_Error", "Editor Subsystems can only be used in Editor Utilities / Blutilities");
 		MessageLog.Error(*ErrorText.ToString(), this);
 	}
 }
\ No newline at end of file
diff --git a/Engine/Source/Editor/Sequencer/Private/SequencerContextMenus.cpp b/Engine/Source/Editor/Sequencer/Private/SequencerContextMenus.cpp
--- a/Engine/Source/Editor/Sequencer/Private/SequencerContextMenus.cpp
+++ b/Engine/Source/Editor/Sequencer/Private/SequencerContextMenus.cpp
@@ -397,8 +397,8 @@ void FSectionContextMenu::PopulateMenu(FMenuBuilder& MenuBuilder)
 		if (CanSetSectionToKey())
 		{
 			MenuBuilder.AddMenuEntry(
-				LOCTEXT("SetSectionToKey", "Key This Section"),
-				LOCTEXT("SetSectionToKey_ToolTip", "This section will get changed when we modify the property externally"),
+				LOCTEXT("KeySection", "Key This Section"),
+				LOCTEXT("KeySection_ToolTip", "This section will get changed when we modify the property externally"),
 				FSlateIcon(),
 				FUIAction(FExecuteAction::CreateLambda([=] { return Shared->SetSectionToKey(); }))
 			);
