Sequencer: Add transaction for renaming track nodes

diff --git a/Engine/Source/Editor/Sequencer/Private/DisplayNodes/SequencerObjectBindingNode.cpp b/Engine/Source/Editor/Sequencer/Private/DisplayNodes/SequencerObjectBindingNode.cpp
--- a/Engine/Source/Editor/Sequencer/Private/DisplayNodes/SequencerObjectBindingNode.cpp
+++ b/Engine/Source/Editor/Sequencer/Private/DisplayNodes/SequencerObjectBindingNode.cpp
@@ -625,6 +625,8 @@ void FSequencerObjectBindingNode::SetDisplayName(const FText& NewDisplayName)
 
 	if (MovieScene != nullptr)
 	{
+		FScopedTransaction Transaction(LOCTEXT("SetTrackName", "Set Track Name"));
+
 		// Modify the movie scene so that it gets marked dirty and renames are saved consistently.
 		MovieScene->Modify();
 		MovieScene->SetObjectDisplayName(ObjectBinding, NewDisplayName);
