Sequencer: Clear flags on subobjects on paste track

diff --git a/Engine/Source/Editor/Sequencer/Private/Sequencer.cpp b/Engine/Source/Editor/Sequencer/Private/Sequencer.cpp
--- a/Engine/Source/Editor/Sequencer/Private/Sequencer.cpp
+++ b/Engine/Source/Editor/Sequencer/Private/Sequencer.cpp
@@ -6585,6 +6585,13 @@ bool FSequencer::PasteTracks(const FString& TextToImport, TArray<FNotificationIn
 				{
 					UMovieSceneTrack* NewTrack = CopyableTrack->Track;
 					NewTrack->ClearFlags(RF_Transient);
+					TArray<UObject*> Subobjects;
+					GetObjectsWithOuter(NewTrack, Subobjects);
+					for (UObject* Subobject : Subobjects)
+					{
+						Subobject->ClearFlags(RF_Transient);
+					}
+
 					if (!GetFocusedMovieSceneSequence()->GetMovieScene()->AddGivenTrack(NewTrack, ObjectGuid))
 					{
 						continue;
@@ -6605,6 +6612,12 @@ bool FSequencer::PasteTracks(const FString& TextToImport, TArray<FNotificationIn
 		{
 			UMovieSceneTrack* NewTrack = CopyableTrack->Track;
 			NewTrack->ClearFlags(RF_Transient);
+			TArray<UObject*> Subobjects;
+			GetObjectsWithOuter(NewTrack, Subobjects);
+			for (UObject* Subobject : Subobjects)
+			{
+				Subobject->ClearFlags(RF_Transient);
+			}
 
 			if (NewTrack->IsA(UMovieSceneCameraCutTrack::StaticClass()))
 			{
