Removed GetCachedGeometry from SSequencerTrackArea. Redundant with SWidget::GetCachedGeometry which does the same thing.

diff --git a/Engine/Source/Editor/Sequencer/Private/SSequencerTrackArea.cpp b/Engine/Source/Editor/Sequencer/Private/SSequencerTrackArea.cpp
--- a/Engine/Source/Editor/Sequencer/Private/SSequencerTrackArea.cpp
+++ b/Engine/Source/Editor/Sequencer/Private/SSequencerTrackArea.cpp
@@ -434,8 +434,6 @@ FCursorReply SSequencerTrackArea::OnCursorQuery( const FGeometry& MyGeometry, co
 
 void SSequencerTrackArea::Tick( const FGeometry& AllottedGeometry, const double InCurrentTime, const float InDeltaTime )
 {
-	CachedGeometry = AllottedGeometry;
-
 	FVector2D Size = AllottedGeometry.GetLocalSize();
 
 	if (SizeLastFrame.IsSet() && Size.X != SizeLastFrame->X)
diff --git a/Engine/Source/Editor/Sequencer/Private/SSequencerTrackArea.h b/Engine/Source/Editor/Sequencer/Private/SSequencerTrackArea.h
--- a/Engine/Source/Editor/Sequencer/Private/SSequencerTrackArea.h
+++ b/Engine/Source/Editor/Sequencer/Private/SSequencerTrackArea.h
@@ -77,12 +77,6 @@ public:
 	/** Attempt to find an existing slot relating to the given node */
 	TSharedPtr<SSequencerTrackLane> FindTrackSlot(const TSharedRef<FSequencerDisplayNode>& InNode);
 
-	/** Access the cached geometry for this track area */
-	const FGeometry& GetCachedGeometry() const
-	{
-		return CachedGeometry;
-	}
-
 	/** Assign a tree view to this track area. */
 	void SetTreeView(const TSharedPtr<SSequencerTreeView>& InTreeView);
 
@@ -129,9 +123,6 @@ private:
 
 private:
 
-	/** Cached geometry. */
-	FGeometry CachedGeometry;
-
 	/** A map of child slot content that exist in our view. */
 	TMap<TSharedPtr<FSequencerDisplayNode>, TWeakPtr<SSequencerTrackLane>> TrackSlots;
 
