Sequencer: Fix autoscroll when in a subsequence

diff --git a/Engine/Source/Editor/Sequencer/Private/Sequencer.cpp b/Engine/Source/Editor/Sequencer/Private/Sequencer.cpp
--- a/Engine/Source/Editor/Sequencer/Private/Sequencer.cpp
+++ b/Engine/Source/Editor/Sequencer/Private/Sequencer.cpp
@@ -612,7 +612,7 @@ void FSequencer::Tick(float InDeltaTime)
 		if (IsAutoScrollEnabled() && GetPlaybackStatus() == EMovieScenePlayerStatus::Playing)
 		{
 			const float ThresholdPercentage = 0.15f;
-			UpdateAutoScroll(GetGlobalTime().Time / GetFocusedTickResolution(), ThresholdPercentage);
+			UpdateAutoScroll(GetLocalTime().Time / GetFocusedTickResolution(), ThresholdPercentage);
 		}
 	}
 	else
