Sequencer: Fix for Debug exception due to TArray::Reserve(-1)

diff --git a/Engine/Source/Editor/Sequencer/Private/SequencerNodeTree.cpp b/Engine/Source/Editor/Sequencer/Private/SequencerNodeTree.cpp
--- a/Engine/Source/Editor/Sequencer/Private/SequencerNodeTree.cpp
+++ b/Engine/Source/Editor/Sequencer/Private/SequencerNodeTree.cpp
@@ -170,7 +170,7 @@ void FSequencerNodeTree::Update()
 	}
 
 	RootNodes.Append(FolderAndObjectAndTrackNodes);
-	RootNodes.Reserve((RootNodes.Num()-1)*2);
+	RootNodes.Reserve(FMath::Max(1, RootNodes.Num() - 1) * 2);
 	for (int32 Index = 1; Index < RootNodes.Num(); Index += 2)
 	{
 		RootNodes.Insert(MakeShareable(new FSequencerSpacerNode(1.f, nullptr, *this, false)), Index);
