Sequencer: Draw section title at section clipping rect bottom left so that it's always on screen

diff --git a/Engine/Source/Editor/MovieSceneTools/Private/TrackEditors/SubTrackEditor.cpp b/Engine/Source/Editor/MovieSceneTools/Private/TrackEditors/SubTrackEditor.cpp
--- a/Engine/Source/Editor/MovieSceneTools/Private/TrackEditors/SubTrackEditor.cpp
+++ b/Engine/Source/Editor/MovieSceneTools/Private/TrackEditors/SubTrackEditor.cpp
@@ -204,12 +204,31 @@ public:
 
 			int32 NumTracks = MovieScene->GetPossessableCount() + MovieScene->GetSpawnableCount() + MovieScene->GetMasterTracks().Num();
 
+			FVector2D TopLeft = InPainter.SectionGeometry.AbsoluteToLocal(InPainter.SectionClippingRect.GetTopLeft()) + FVector2D(1.f, -1.f);
+
+			FSlateFontInfo FontInfo = FEditorStyle::GetFontStyle("NormalFont");
+
+			TSharedRef<FSlateFontCache> FontCache = FSlateApplication::Get().GetRenderer()->GetFontCache();
+
+			auto GetFontHeight = [&]
+			{
+				return FontCache->GetMaxCharacterHeight(FontInfo, 1.f) + FontCache->GetBaseline(FontInfo, 1.f);
+			};
+			while (GetFontHeight() > InPainter.SectionGeometry.Size.Y && FontInfo.Size > 11)
+			{
+				FontInfo.Size = FMath::Max(FMath::FloorToInt(FontInfo.Size - 6.f), 11);
+			}
+
 			FSlateDrawElement::MakeText(
 				InPainter.DrawElements,
 				++LayerId,
-				InPainter.SectionGeometry.ToOffsetPaintGeometry(FVector2D(11.0f, 32.0f)),
+				InPainter.SectionGeometry.MakeChild(
+					FVector2D(InPainter.SectionGeometry.Size.X, GetFontHeight()),
+					FSlateLayoutTransform(TopLeft + FVector2D(0.f, GetFontHeight()) + FVector2D(11.f, GetFontHeight()*2.f))
+				).ToPaintGeometry(),
+				//InPainter.SectionGeometry.ToOffsetPaintGeometry(FVector2D(11.0f, 32.0f)),
 				FText::Format(LOCTEXT("NumTracksFormat", "{0} track(s)"), FText::AsNumber(NumTracks)),
-				FEditorStyle::GetFontStyle("NormalFont"),
+				FontInfo,
 				DrawEffects,
 				FColor(200, 200, 200)
 			);
diff --git a/Engine/Source/Editor/Sequencer/Private/SSequencerSection.cpp b/Engine/Source/Editor/Sequencer/Private/SSequencerSection.cpp
--- a/Engine/Source/Editor/Sequencer/Private/SSequencerSection.cpp
+++ b/Engine/Source/Editor/Sequencer/Private/SSequencerSection.cpp
@@ -1123,12 +1123,31 @@ int32 SSequencerSection::OnPaint( const FPaintArgs& Args, const FGeometry& Allot
 
 	if (!SectionTitle.IsEmpty())
 	{
+		FVector2D TopLeft = SectionGeometry.AbsoluteToLocal(Painter.SectionClippingRect.GetTopLeft()) + FVector2D(1.f, -1.f);
+
+		FSlateFontInfo FontInfo = FEditorStyle::GetFontStyle("NormalFont");
+
+		TSharedRef<FSlateFontCache> FontCache = FSlateApplication::Get().GetRenderer()->GetFontCache();
+
+		auto GetFontHeight = [&]
+		{
+			return FontCache->GetMaxCharacterHeight(FontInfo, 1.f) + FontCache->GetBaseline(FontInfo, 1.f);
+		};
+		while (GetFontHeight() > SectionGeometry.Size.Y && FontInfo.Size > 11)
+		{
+			FontInfo.Size = FMath::Max(FMath::FloorToInt(FontInfo.Size - 6.f), 11);
+		}
+
+		// Drop shadow
 		FSlateDrawElement::MakeText(
 			OutDrawElements,
 			LayerId,
-			Painter.SectionGeometry.ToOffsetPaintGeometry(FVector2D(ContentPadding.Left + 1, ContentPadding.Top + 1)),
+			SectionGeometry.MakeChild(
+				FVector2D(SectionGeometry.Size.X, GetFontHeight()),
+				FSlateLayoutTransform(TopLeft + FVector2D(0.f, GetFontHeight()) + FVector2D(1.f, 1.f))
+			).ToPaintGeometry(),
 			SectionTitle,
-			FEditorStyle::GetFontStyle("NormalFont"),
+			FontInfo,
 			DrawEffects,
 			FLinearColor(0,0,0,.5f)
 		);
@@ -1136,9 +1155,12 @@ int32 SSequencerSection::OnPaint( const FPaintArgs& Args, const FGeometry& Allot
 		FSlateDrawElement::MakeText(
 			OutDrawElements,
 			LayerId,
-			Painter.SectionGeometry.ToOffsetPaintGeometry(FVector2D(ContentPadding.Left, ContentPadding.Top)),
+			SectionGeometry.MakeChild(
+				FVector2D(SectionGeometry.Size.X, GetFontHeight()),
+				FSlateLayoutTransform(TopLeft + FVector2D(0.f, GetFontHeight()))
+			).ToPaintGeometry(),
 			SectionTitle,
-			FEditorStyle::GetFontStyle("NormalFont"),
+			FontInfo,
 			DrawEffects,
 			FColor(200, 200, 200)
 		);
