Sequencer: Fix crash removing a marked frame.

P4 conflict auto resolved incorrectly.

diff --git a/Engine/Source/Editor/Sequencer/Private/Sequencer.cpp b/Engine/Source/Editor/Sequencer/Private/Sequencer.cpp
--- a/Engine/Source/Editor/Sequencer/Private/Sequencer.cpp
+++ b/Engine/Source/Editor/Sequencer/Private/Sequencer.cpp
@@ -8133,7 +8133,6 @@ void FSequencer::SetMarkedFrame(FFrameNumber FrameNumber, bool bSetMark)
 			{
 				int32 MarkedFrameIndex = FocusedMovieScene->FindMarkedFrameByFrameNumber(FrameNumber);
 				if (MarkedFrameIndex != INDEX_NONE)
-				FocusedMovieScene->RemoveMarkedFrame(MarkedFrameIndex);
 				{
 					FocusedMovieScene->RemoveMarkedFrame(MarkedFrameIndex);
 				}
