Sequencer: Make sure external selection is up to date with sequencer selection when exporting selected tracks. The selection can be out of sync since nodes are not select on mouse down.

diff --git a/Engine/Source/Editor/Sequencer/Private/Sequencer.cpp b/Engine/Source/Editor/Sequencer/Private/Sequencer.cpp
--- a/Engine/Source/Editor/Sequencer/Private/Sequencer.cpp
+++ b/Engine/Source/Editor/Sequencer/Private/Sequencer.cpp
@@ -8928,6 +8928,12 @@ void FSequencer::ExportFBXInternal(const FString& ExportFilename, TArray<FGuid>&
 
 			const bool bSelectedOnly = (Selection.GetSelectedTracks().Num() + Bindings.Num()) != 0;
 
+			// Make sure external selection is up to date since export could happen on tracks that have been right clicked but not have their underlying bound objects selected yet since that happens on mouse up.
+			if (bSelectedOnly)
+			{
+				SynchronizeExternalSelectionWithSequencerSelection();
+			}
+
 			UnFbx::FFbxExporter::FLevelSequenceNodeNameAdapter NodeNameAdapter(MovieScene, this, GetFocusedTemplateID());
 
 			// Export the persistent level and all of it's actors
