Sequencer: Fix regression caused by not setting up default tracks for a spawnable

diff --git a/Engine/Source/Editor/Sequencer/Public/ISequencer.h b/Engine/Source/Editor/Sequencer/Public/ISequencer.h
--- a/Engine/Source/Editor/Sequencer/Public/ISequencer.h
+++ b/Engine/Source/Editor/Sequencer/Public/ISequencer.h
@@ -224,9 +224,10 @@ public:
 	 *
 	 * @param	Object	The asset, class, or actor to add a spawnable for
 	 * @param	ActorFactory	Optional actor factory to use to create spawnable type
+	 * @param   bSetupDefaults Setup default tracks for this spawnable
 	 * @return	The spawnable guid for the spawnable, or an invalid Guid if we were not able to create a spawnable
 	 */
-	virtual FGuid MakeNewSpawnable(UObject& SourceObject, UActorFactory* ActorFactory = nullptr, bool bSetupDefaults = false) = 0;
+	virtual FGuid MakeNewSpawnable(UObject& SourceObject, UActorFactory* ActorFactory = nullptr, bool bSetupDefaults = true) = 0;
 
 	/**
 	 * Add actors as possessable objects to sequencer.
