Sequencer: Clear focus before setting time in case there's a key editor value selected that gets committed to a newly selected key on UserMovedFocus

diff --git a/Engine/Source/Editor/Sequencer/Private/Sequencer.cpp b/Engine/Source/Editor/Sequencer/Private/Sequencer.cpp
--- a/Engine/Source/Editor/Sequencer/Private/Sequencer.cpp
+++ b/Engine/Source/Editor/Sequencer/Private/Sequencer.cpp
@@ -2299,6 +2299,9 @@ void FSequencer::SetLocalTimeDirectly(FFrameTime NewTime)
 
 void FSequencer::SetGlobalTime(FFrameTime NewTime)
 {
+	// Clear focus before setting time in case there's a key editor value selected that gets committed to a newly selected key on UserMovedFocus
+	FSlateApplication::Get().ClearKeyboardFocus(EFocusCause::Cleared);
+
 	NewTime = ConvertFrameTime(NewTime, GetRootTickResolution(), PlayPosition.GetInputRate());
 	if (PlayPosition.GetEvaluationType() == EMovieSceneEvaluationType::FrameLocked)
 	{
