Sequencer: Add "Animatable Property Changed" transaction to SequencerObjectChangeListener before it broadcasts to the track editors because there can be more than 1 track editor per property type, so multiple transactions were being put on the stack for changing properties like visibility.

diff --git a/Engine/Source/Editor/Sequencer/Private/SequencerObjectChangeListener.cpp b/Engine/Source/Editor/Sequencer/Private/SequencerObjectChangeListener.cpp
--- a/Engine/Source/Editor/Sequencer/Private/SequencerObjectChangeListener.cpp
+++ b/Engine/Source/Editor/Sequencer/Private/SequencerObjectChangeListener.cpp
@@ -97,6 +97,11 @@ void FSequencerObjectChangeListener::BroadcastPropertyChanged( FKeyPropertyParam
 		{
 			StructPathToKey = *KeyPropertyParams.PropertyPath.TrimRoot(PropertyPath.GetNumProperties());
 		}
+
+		// Create a transaction record because we are about to add keys/tracks
+		const bool bShouldActuallyTransact = !GIsTransacting;		// Don't transact if we're recording in a PIE world.  That type of keyframe capture cannot be undone.
+		FScopedTransaction PropertyChangedTransaction(NSLOCTEXT("Sequencer", "PropertyChanged", "Animatable Property Changed"), bShouldActuallyTransact);
+
 		FPropertyChangedParams Params(KeyableObjects, PropertyPath, StructPathToKey, KeyPropertyParams.KeyMode);
 		Delegate.Broadcast(Params);
 	}
