Sequencer: Fix autotracking.

On pre edit change, the changed property is checked against CanKey. This didn't work for structs, where the property chain would be something like DepthOfField->Focal Distance. A fix was made in 3729216 to support structs, but this broke generic properties because it would check CanKey with the property chain, LightComponent->Intensity.

The fix here is to find the keyable property and then use the remaining property path. For example, for intensity, the property path will be Intensity. For Focal Distance, the property path will be DepthOfField->Focal Distance.

diff --git a/Engine/Source/Editor/Sequencer/Private/SequencerObjectChangeListener.cpp b/Engine/Source/Editor/Sequencer/Private/SequencerObjectChangeListener.cpp
--- a/Engine/Source/Editor/Sequencer/Private/SequencerObjectChangeListener.cpp
+++ b/Engine/Source/Editor/Sequencer/Private/SequencerObjectChangeListener.cpp
@@ -8,6 +8,7 @@
 #include "PropertyEditorModule.h"
 #include "IPropertyChangeListener.h"
 #include "MovieSceneSequence.h"
+#include "ScopedTransaction.h"
 
 DEFINE_LOG_CATEGORY_STATIC(LogSequencerTools, Log, All);
 
@@ -71,12 +72,24 @@ void FSequencerObjectChangeListener::BroadcastPropertyChanged( FKeyPropertyParam
 			for (TFieldIterator<UProperty> PropertyIterator(ObjectToKey->GetClass()); PropertyIterator; ++PropertyIterator)
 			{
 				UProperty* CheckProperty = *PropertyIterator;
-				if (CheckProperty == KeyPropertyParams.PropertyPath.GetRootProperty().Property.Get())
+
+				for (int32 Index = 0; Index < KeyPropertyParams.PropertyPath.GetNumProperties(); ++Index)
 				{
-					if (CanKeyProperty_Internal(FCanKeyPropertyParams(ObjectToKey->GetClass(), KeyPropertyParams.PropertyPath), Delegate, Property, PropertyPath))
+					const FPropertyInfo& PropertyInfo = KeyPropertyParams.PropertyPath.GetPropertyInfo(Index);
+
+					if (CheckProperty == PropertyInfo.Property.Get())
 					{
-						KeyableObjects.Add(ObjectToKey);
-						break;
+						FPropertyPath TrimmedPropertyPath;
+						for (int32 TrimmedIndex = Index; TrimmedIndex < KeyPropertyParams.PropertyPath.GetNumProperties(); ++TrimmedIndex)
+						{
+							TrimmedPropertyPath.AddProperty(KeyPropertyParams.PropertyPath.GetPropertyInfo(TrimmedIndex));
+						}
+
+						if (CanKeyProperty_Internal(FCanKeyPropertyParams(ObjectToKey->GetClass(), TrimmedPropertyPath), Delegate, Property, PropertyPath))
+						{
+							KeyableObjects.Add(ObjectToKey);
+							break;
+						}
 					}
 				}
 			}
