Sequencer: Fix multiple components being added

This was a result of adding the event callback OnActorsAddedToSequencer in GetHandleToObject. This would cause issues because the callstack is:

GetHandleToObject -> LightComponent
    GetHandleToObject -> PointLight
        AddDefaultTracks
           GetHandleToObject -> LightComponent (this fails to find the first guid because the parenting structure isn't set up yet)

The fix is to move OnActorsAddedToSequencer to CreateBinding

diff --git a/Engine/Source/Editor/Sequencer/Private/Sequencer.cpp b/Engine/Source/Editor/Sequencer/Private/Sequencer.cpp
--- a/Engine/Source/Editor/Sequencer/Private/Sequencer.cpp
+++ b/Engine/Source/Editor/Sequencer/Private/Sequencer.cpp
@@ -840,11 +840,19 @@ FGuid FSequencer::CreateBinding(UObject& InObject, const FString& InName)
 	UObject* ParentObject = OwnerSequence->GetParentObject(&InObject);
 	UObject* BindingContext = GetPlaybackContext();
 
+	AActor* ParentActorAdded = nullptr;
+	FGuid ParentGuid;
+
 	if (ParentObject)
 	{
 		// Ensure we have possessed the outer object, if necessary
-		FGuid ParentGuid = GetHandleToObject(ParentObject);
-		
+		ParentGuid = GetHandleToObject(ParentObject, false);
+		if (!ParentGuid.IsValid())
+		{
+			ParentGuid = GetHandleToObject(ParentObject);
+			ParentActorAdded = Cast<AActor>(ParentObject);
+		}
+
 		if (OwnerSequence->AreParentContextsSignificant())
 		{
 			BindingContext = ParentObject;
@@ -868,6 +876,12 @@ FGuid FSequencer::CreateBinding(UObject& InObject, const FString& InName)
 	}
 
 	OwnerSequence->BindPossessableObject(PossessableGuid, InObject, BindingContext);
+	
+	// Broadcast if a parent actor was added as a result of adding this object
+	if (ParentActorAdded && ParentGuid.IsValid())
+	{
+		OnActorAddedToSequencerEvent.Broadcast(ParentActorAdded, ParentGuid);
+	}
 
 	return PossessableGuid;
 }
@@ -2793,11 +2807,7 @@ FGuid FSequencer::GetHandleToObject( UObject* Object, bool bCreateHandleIfMissin
 
 		if (OwningActor)
 		{
-			FGuid ActorAddedGuid = GetHandleToObject(OwningActor);
-			if (ActorAddedGuid.IsValid())
-			{
-				OnActorAddedToSequencerEvent.Broadcast(OwningActor, ActorAddedGuid);
-			}
+			GetHandleToObject(OwningActor);
 		}
 
 		// Some sources that create object bindings may want to group all of these objects together for organizations sake.
