Sequencer: Change save icon and add * to breadcrumb to indicate when the asset is dirty and needs to be saved.

diff --git a/Engine/Source/Editor/EditorStyle/Private/SlateEditorStyle.cpp b/Engine/Source/Editor/EditorStyle/Private/SlateEditorStyle.cpp
--- a/Engine/Source/Editor/EditorStyle/Private/SlateEditorStyle.cpp
+++ b/Engine/Source/Editor/EditorStyle/Private/SlateEditorStyle.cpp
@@ -1845,10 +1845,12 @@ void FSlateEditorStyle::FStyle::SetupGeneralStyles()
 
 		Set( "Sequencer.NotificationImage_AddedPlayMovieSceneEvent", new IMAGE_BRUSH( "Old/Checkbox_checked", Icon16x16 ) );
 
-		Set( "Sequencer.Save", new IMAGE_BRUSH( "Sequencer/Main_Icons/Icon_Sequencer_Save_24x", Icon48x48 ) );
-		Set( "Sequencer.Save.Small", new IMAGE_BRUSH( "Sequencer/Main_Icons/Icon_Sequencer_Save_24x", Icon24x24 ) );
-		Set( "Sequencer.SaveAs", new IMAGE_BRUSH( "Sequencer/Main_Icons/Icon_Sequencer_Save_As_24x", Icon48x48 ) );
-		Set( "Sequencer.SaveAs.Small", new IMAGE_BRUSH( "Sequencer/Main_Icons/Icon_Sequencer_Save_As_24x", Icon24x24 ) );
+		Set( "Sequencer.Save", new IMAGE_BRUSH("Sequencer/Main_Icons/Icon_Sequencer_Save_48x", Icon48x48) );
+		Set( "Sequencer.Save.Small", new IMAGE_BRUSH("Sequencer/Main_Icons/Icon_Sequencer_Save_48x", Icon24x24) );
+		Set( "Sequencer.SaveAsterisk", new IMAGE_BRUSH("Sequencer/Main_Icons/Icon_Sequencer_SaveAsterisk_48x", Icon48x48) );
+		Set( "Sequencer.SaveAsterisk.Small", new IMAGE_BRUSH("Sequencer/Main_Icons/Icon_Sequencer_SaveAsterisk_48x", Icon24x24) );
+		Set( "Sequencer.SaveAs", new IMAGE_BRUSH( "Sequencer/Main_Icons/Icon_Sequencer_SaveAs_48x", Icon48x48 ) );
+		Set( "Sequencer.SaveAs.Small", new IMAGE_BRUSH( "Sequencer/Main_Icons/Icon_Sequencer_SaveAs_48x", Icon24x24 ) );
 		Set( "Sequencer.DiscardChanges", new IMAGE_BRUSH( "Sequencer/Main_Icons/Icon_Sequencer_Revert_24x", Icon48x48 ) );
 		Set( "Sequencer.DiscardChanges.Small", new IMAGE_BRUSH( "Sequencer/Main_Icons/Icon_Sequencer_Revert_24x", Icon24x24 ) );
 		Set( "Sequencer.RestoreAnimatedState", new IMAGE_BRUSH("Sequencer/Main_Icons/Icon_Sequencer_RestoreAnimatedState_24x", Icon48x48 ) );
diff --git a/Engine/Source/Editor/Sequencer/Private/SSequencer.cpp b/Engine/Source/Editor/Sequencer/Private/SSequencer.cpp
--- a/Engine/Source/Editor/Sequencer/Private/SSequencer.cpp
+++ b/Engine/Source/Editor/Sequencer/Private/SSequencer.cpp
@@ -797,12 +797,39 @@ TSharedRef<SWidget> SSequencer::MakeToolBar()
 		// General 
 		if (SequencerPtr.Pin()->IsLevelEditorSequencer())
 		{
+			TAttribute<FSlateIcon> SaveIcon;
+			SaveIcon.Bind(TAttribute<FSlateIcon>::FGetter::CreateLambda([&] {
+
+				bool bAnyMovieSceneDirty = false;
+
+				TArray<UMovieScene*> MovieScenesToSave;
+				MovieSceneHelpers::GetDescendantMovieScenes(SequencerPtr.Pin()->GetRootMovieSceneSequence(), MovieScenesToSave);
+				for (auto MovieSceneToSave : MovieScenesToSave)
+				{
+					UPackage* MovieScenePackageToSave = MovieSceneToSave->GetOuter()->GetOutermost();
+					if (MovieScenePackageToSave->IsDirty())
+					{
+						bAnyMovieSceneDirty = true;
+						break;
+					}
+				}
+
+				if (bAnyMovieSceneDirty)
+				{
+					return FSlateIcon(FEditorStyle::GetStyleSetName(), "Sequencer.SaveAsterisk");
+				}
+				else
+				{
+					return FSlateIcon(FEditorStyle::GetStyleSetName(), "Sequencer.Save");
+				}
+			}));
+
 			ToolBarBuilder.AddToolBarButton(
 				FUIAction(FExecuteAction::CreateSP(this, &SSequencer::OnSaveMovieSceneClicked)),
 				NAME_None,
 				LOCTEXT("SaveDirtyPackages", "Save"),
-				LOCTEXT("SaveDirtyPackagesTooltip", "Saves the current sequence"),
-				FSlateIcon(FEditorStyle::GetStyleSetName(), "Sequencer.Save")
+				LOCTEXT("SaveDirtyPackagesTooltip", "Saves the current sequence and any subsequences"),
+				SaveIcon
 			);
 
 			ToolBarBuilder.AddToolBarButton(
@@ -2201,15 +2228,35 @@ FText SSequencer::GetBreadcrumbTextForSection(TWeakObjectPtr<UMovieSceneSubSecti
 FText SSequencer::GetBreadcrumbTextForSequence(TWeakObjectPtr<UMovieSceneSequence> Sequence, bool bIsActive) const
 {
 	UMovieSceneSequence* SequencePtr = Sequence.Get();
+
+	bool bIsDirty = SequencePtr->GetMovieScene()->GetOuter()->GetOutermost()->IsDirty();
+
 	if (bIsActive)
 	{
-		return SequencePtr->GetDisplayName();
+		if (bIsDirty)
+		{
+			return FText::Format(LOCTEXT("DirtySequenceBreadcrumbFormat", "{0}*"), SequencePtr->GetDisplayName());
+		}
+		else
+		{
+			return SequencePtr->GetDisplayName();
+		}
 	}
 	else
 	{
-		return FText::Format(LOCTEXT("InactiveSequenceBreadcrumbFormat", "{0} [{1}]"),
-			SequencePtr->GetDisplayName(),
-			LOCTEXT("InactiveSequenceBreadcrumb", "Inactive"));
+		if (bIsDirty)
+		{
+			return FText::Format(LOCTEXT("DirtyInactiveSequenceBreadcrumbFormat", "{0}* [{1}]"),
+				SequencePtr->GetDisplayName(),
+				LOCTEXT("InactiveSequenceBreadcrumb", "Inactive"));
+
+		}
+		else
+		{
+			return FText::Format(LOCTEXT("InactiveSequenceBreadcrumbFormat", "{0} [{1}]"),
+				SequencePtr->GetDisplayName(),
+				LOCTEXT("InactiveSequenceBreadcrumb", "Inactive"));
+		}
 	}
 }
 
