This is a significant overhaul to the Curve Editor used by Sequencer which adds a plugin-based architecture and extensibility. New tools and toolbar buttons can be added to all usages of the curve edi

diff --git a/Engine/Source/Editor/Sequencer/Private/DisplayNodes/SequencerDisplayNode.cpp b/Engine/Source/Editor/Sequencer/Private/DisplayNodes/SequencerDisplayNode.cpp
--- a/Engine/Source/Editor/Sequencer/Private/DisplayNodes/SequencerDisplayNode.cpp
+++ b/Engine/Source/Editor/Sequencer/Private/DisplayNodes/SequencerDisplayNode.cpp
@@ -8,7 +8,10 @@
 #include "Widgets/SLeafWidget.h"
 #include "Textures/SlateIcon.h"
 #include "Framework/Commands/UIAction.h"
+#include "Widgets/SBoxPanel.h"
+#include "Widgets/SOverlay.h"
 #include "Widgets/Layout/SSpacer.h"
+#include "Widgets/Images/SImage.h"
 #include "Framework/MultiBox/MultiBoxBuilder.h"
 #include "EditorStyleSet.h"
 #include "DisplayNodes/SequencerObjectBindingNode.h"
@@ -22,8 +25,10 @@
 #include "SSequencerSectionAreaView.h"
 #include "CommonMovieSceneTools.h"
 #include "Framework/Commands/GenericCommands.h"
+#include "Tree/SCurveEditorTreePin.h"
 #include "ScopedTransaction.h"
 #include "SequencerKeyTimeCache.h"
+#include "SequencerNodeSortingMethods.h"
 #include "SequencerKeyCollection.h"
 
 #define LOCTEXT_NAMESPACE "SequencerDisplayNode"
@@ -35,6 +40,61 @@ namespace SequencerNodeConstants
 	const float CommonPadding = 4.f;
 
 	static const FVector2D KeyMarkSize = FVector2D(3.f, 21.f);
+
+	static const uint8 DefaultSortBias[(__underlying_type(EDisplayNodeSortType))EDisplayNodeSortType::NUM] = {
+		2, // Folders
+		3, // Tracks
+		4, // ObjectBindings
+		1, // CameraCuts
+		0, // Shots
+		5, // Anything else
+	};
+
+	static const uint8 ObjectBindingSortBias[(__underlying_type(EDisplayNodeSortType))EDisplayNodeSortType::NUM] = {
+		2, // Folders            - shouldn't exist inside object bindings
+		1, // Tracks
+		0, // ObjectBindings
+		3, // CameraCuts         - shouldn't exist inside object bindings
+		4, // Shots              - shouldn't exist inside object bindings
+		5, // Anything else
+	};
+
+	static_assert(ARRAY_COUNT(DefaultSortBias) == (__underlying_type(EDisplayNodeSortType))EDisplayNodeSortType::NUM, "Mismatched type/bias count");
+	static_assert(ARRAY_COUNT(ObjectBindingSortBias) == (__underlying_type(EDisplayNodeSortType))EDisplayNodeSortType::NUM, "Mismatched type/bias count");
+
+	inline bool SortChildrenWithBias(const TSharedRef<FSequencerDisplayNode>& A, const TSharedRef<FSequencerDisplayNode>& B, const uint8* SortBias)
+	{
+		const uint8 BiasA = SortBias[(__underlying_type(EDisplayNodeSortType))A->GetSortType()];
+		const uint8 BiasB = SortBias[(__underlying_type(EDisplayNodeSortType))B->GetSortType()];
+
+		// For nodes of the same bias, sort by name
+		if (BiasA == BiasB)
+		{
+			return A->GetDisplayName().CompareToCaseIgnored(B->GetDisplayName()) < 0;
+		}
+		return BiasA < BiasB;
+	}
+
+	inline bool SortObjectBindingChildren(const TSharedRef<FSequencerDisplayNode>& A, const TSharedRef<FSequencerDisplayNode>& B)
+	{
+		return SortChildrenWithBias(A, B, SequencerNodeConstants::ObjectBindingSortBias);
+	}
+
+	static bool SortChildrenDefault(const TSharedRef<FSequencerDisplayNode>& A, const TSharedRef<FSequencerDisplayNode>& B)
+	{
+		const int32 SortA = A->GetSortingOrder();
+		const int32 SortB = B->GetSortingOrder();
+
+		if (SortA >= 0 && SortB >= 0)
+		{
+			// Both nodes have persistent sort orders, use those
+			return SortA < SortB;
+		}
+
+		// When either or neither node has a persistent sort order, we use the default ordering between the two nodes to ensure that 
+		// New nodes get added to the correctly sorted position by default
+		return SortChildrenWithBias(A, B, SequencerNodeConstants::DefaultSortBias);
+	}
 }
 
 struct FNameAndSignature
@@ -274,31 +334,119 @@ void SSequencerObjectTrack::GenerateCachedKeyPositions(const FGeometry& Allotted
 }
 
 
-FSequencerDisplayNode::FSequencerDisplayNode( FName InNodeName, TSharedPtr<FSequencerDisplayNode> InParentNode, FSequencerNodeTree& InParentTree )
-	: VirtualTop( 0.f )
+FSequencerDisplayNode::FSequencerDisplayNode( FName InNodeName, FSequencerNodeTree& InParentTree )
+	: TreeSerialNumber(0)
+	, VirtualTop( 0.f )
 	, VirtualBottom( 0.f )
-	, ParentNode( InParentNode )
 	, ParentTree( InParentTree )
 	, NodeName( InNodeName )
 	, bExpanded( false )
+	, bHasBeenInitialized( false )
+{
+	SortType = EDisplayNodeSortType::Undefined;
+}
+
+bool FSequencerDisplayNode::IsParentStillRelevant(uint32 SerialNumber) const
+{
+	TSharedPtr<FSequencerDisplayNode> ExistingParent = GetParent();
+	return ExistingParent.IsValid() && ExistingParent->TreeSerialNumber == SerialNumber;
+}
+
+bool FSequencerDisplayNode::IsRootNode() const
+{
+	return ParentNode == ParentTree.GetRootNode();
+}
+
+void FSequencerDisplayNode::SetParentDirectly(TSharedPtr<FSequencerDisplayNode> InParent)
 {
+	ParentNode = InParent;
+}
+
+void FSequencerDisplayNode::SetParent(TSharedPtr<FSequencerDisplayNode> InParent)
+{
+	TSharedPtr<FSequencerDisplayNode> CurrentParent = ParentNode.Pin();
+	if (CurrentParent != InParent)
+	{
+		TSharedRef<FSequencerDisplayNode> ThisNode = AsShared();
+		if (CurrentParent)
+		{
+			// Remove from parent
+			CurrentParent->ChildNodes.Remove(ThisNode);
+		}
+
+		if (InParent)
+		{
+			// Add to new parent
+			InParent->ChildNodes.Add(ThisNode);
+			bExpanded = ParentTree.GetSavedExpansionState( *this );
+		}
+	}
+
+	ParentNode = InParent;
 }
 
+FSequencer& FSequencerDisplayNode::GetSequencer() const
+{
+	return ParentTree.GetSequencer();
+}
 
-void FSequencerDisplayNode::Initialize(float InVirtualTop, float InVirtualBottom)
+void FSequencerDisplayNode::OnTreeRefreshed(float InVirtualTop, float InVirtualBottom)
 {
-	bExpanded = ParentTree.GetSavedExpansionState( *this );
+	if (!bHasBeenInitialized)
+	{
+		// Assign the saved expansion state when this node is initialized for the first time
+		bExpanded = ParentTree.GetSavedExpansionState( *this );
+	}
 
 	VirtualTop = InVirtualTop;
 	VirtualBottom = InVirtualBottom;
-}
 
+	SortImmediateChildren();
 
-void FSequencerDisplayNode::AddObjectBindingNode(TSharedRef<FSequencerObjectBindingNode> ObjectBindingNode)
+	bHasBeenInitialized = true;
+}
+
+void FSequencerDisplayNode::SortImmediateChildren()
 {
-	AddChildAndSetParent( ObjectBindingNode );
+	const ESequencerNode::Type NodeType = GetType();
+	if (ChildNodes.Num() == 0 || NodeType == ESequencerNode::Category || NodeType == ESequencerNode::Track)
+	{
+		return;
+	}
+
+	if (NodeType == ESequencerNode::Object)
+	{
+		// Objects never use their serialized sort order
+		Algo::Sort(ChildNodes, SequencerNodeConstants::SortObjectBindingChildren);
+	}
+	else
+	{
+		Algo::Sort(ChildNodes, SequencerNodeConstants::SortChildrenDefault);
+	}
+
+	if (NodeType != ESequencerNode::Track || static_cast<FSequencerTrackNode*>(this)->GetSubTrackMode() == FSequencerTrackNode::ESubTrackMode::None)
+	{
+		// Set persistent sort orders
+		for (int32 Index = 0; Index < ChildNodes.Num(); ++Index)
+		{
+			ChildNodes[Index]->SetSortingOrder(Index);
+		}
+	}
 }
 
+void FSequencerDisplayNode::ResortImmediateChildren()
+{
+	if (ChildNodes.Num() > 0)
+	{
+		// Unset persistent sort orders
+		for (TSharedRef<FSequencerDisplayNode> Child : ChildNodes)
+		{
+			Child->SetSortingOrder(-1);
+		}
+
+		SortImmediateChildren();
+	}
+}
 
 TSharedPtr<FSequencerObjectBindingNode> FSequencerDisplayNode::FindParentObjectBindingNode() const
 {
@@ -306,13 +454,9 @@ TSharedPtr<FSequencerObjectBindingNode> FSequencerDisplayNode::FindParentObjectB
 
 	while (CurrentParentNode.IsValid())
 	{
-		if (CurrentParentNode.Get()->GetType() == ESequencerNode::Object)
+		if (CurrentParentNode->GetType() == ESequencerNode::Object)
 		{
-			TSharedPtr<FSequencerObjectBindingNode> ObjectNode = StaticCastSharedPtr<FSequencerObjectBindingNode>(CurrentParentNode);
-			if (ObjectNode.IsValid())
-			{
-				return ObjectNode;
-			}
+			return StaticCastSharedPtr<FSequencerObjectBindingNode>(CurrentParentNode);
 		}
 		CurrentParentNode = CurrentParentNode->GetParent();
 	}
@@ -320,6 +464,11 @@ TSharedPtr<FSequencerObjectBindingNode> FSequencerDisplayNode::FindParentObjectB
 	return nullptr;
 }
 
+FGuid FSequencerDisplayNode::GetObjectGuid() const
+{
+	TSharedPtr<FSequencerObjectBindingNode> ObjectBindingNode = FindParentObjectBindingNode();
+	return ObjectBindingNode ? ObjectBindingNode->GetObjectGuid() : FGuid();
+}
 
 bool FSequencerDisplayNode::Traverse_ChildFirst(const TFunctionRef<bool(FSequencerDisplayNode&)>& InPredicate, bool bIncludeThisNode)
 {
@@ -400,53 +549,6 @@ bool FSequencerDisplayNode::TraverseVisible_ParentFirst(const TFunctionRef<bool(
 	return true;
 }
 
-TSharedRef<FSequencerSectionCategoryNode> FSequencerDisplayNode::AddCategoryNode( FName CategoryName, const FText& DisplayLabel )
-{
-	TSharedPtr<FSequencerSectionCategoryNode> CategoryNode;
-
-	// See if there is an already existing category node to use
-	for (const TSharedRef<FSequencerDisplayNode>& Node : ChildNodes)
-	{
-		if ((Node->GetNodeName() == CategoryName) && (Node->GetType() == ESequencerNode::Category))
-		{
-			CategoryNode = StaticCastSharedRef<FSequencerSectionCategoryNode>(Node);
-		}
-	}
-
-	if (!CategoryNode.IsValid())
-	{
-		// No existing category found, make a new one
-		CategoryNode = MakeShareable(new FSequencerSectionCategoryNode(CategoryName, DisplayLabel, SharedThis(this), ParentTree));
-		ChildNodes.Add(CategoryNode.ToSharedRef());
-	}
-
-	return CategoryNode.ToSharedRef();
-}
-
-
-void FSequencerDisplayNode::AddKeyAreaNode(FName KeyAreaName, const FText& DisplayName, TSharedRef<IKeyArea> KeyArea)
-{
-	TSharedPtr<FSequencerSectionKeyAreaNode> KeyAreaNode;
-
-	// See if there is an already existing key area node to use
-	for (const TSharedRef<FSequencerDisplayNode>& Node : ChildNodes)
-	{
-		if ((Node->GetNodeName() == KeyAreaName) && (Node->GetType() == ESequencerNode::KeyArea))
-		{
-			KeyAreaNode = StaticCastSharedRef<FSequencerSectionKeyAreaNode>(Node);
-		}
-	}
-
-	if (!KeyAreaNode.IsValid())
-	{
-		// No existing node found make a new one
-		KeyAreaNode = MakeShareable( new FSequencerSectionKeyAreaNode( KeyAreaName, DisplayName, SharedThis( this ), ParentTree ) );
-		ChildNodes.Add(KeyAreaNode.ToSharedRef());
-	}
-
-	KeyAreaNode->AddKeyArea(KeyArea);
-}
-
 FLinearColor FSequencerDisplayNode::GetDisplayNameColor() const
 {
 	auto FindInActiveSection = [](FSequencerDisplayNode& InNode, bool EmptyNotActive = true)
@@ -619,10 +721,11 @@ FString FSequencerDisplayNode::GetPathName() const
 	// First get our parent's path
 	FString PathName;
 
-	if (ParentNode.IsValid())
+	TSharedPtr<FSequencerDisplayNode> Parent = GetParent();
+	if (Parent.IsValid())
 	{
-		ensure(ParentNode != SharedThis(this));
-		PathName = ParentNode.Pin()->GetPathName() + TEXT(".");
+		ensure(Parent != SharedThis(this));
+		PathName = Parent->GetPathName() + TEXT(".");
 	}
 
 	//then append our path
@@ -871,6 +974,12 @@ bool FSequencerDisplayNode::IsHidden() const
 }
 
 
+bool FSequencerDisplayNode::IsVisible() const
+{
+	return !ParentTree.HasActiveFilter() || ParentTree.IsNodeFiltered(AsShared());
+}
+
+
 bool FSequencerDisplayNode::IsHovered() const
 {
 	return ParentTree.GetHoveredNode().Get() == this;
@@ -889,10 +998,60 @@ bool FSequencerDisplayNode::HandleContextMenuRenameNodeCanExecute() const
 }
 
 
-void FSequencerDisplayNode::AddChildAndSetParent( TSharedRef<FSequencerDisplayNode> InChild )
+TSharedPtr<SWidget> FSequencerDisplayNode::GenerateCurveEditorTreeWidget(const FName& InColumnName, TWeakPtr<FCurveEditor> InCurveEditor, FCurveEditorTreeItemID InTreeItemID)
+{
+	if (InColumnName == ColumnNames.Label)
+	{
+		return SNew(SHorizontalBox)
+
+			+ SHorizontalBox::Slot()
+			.Padding(FMargin(0.f, 0.f, 4.f, 0.f))
+			.VAlign(VAlign_Center)
+			.AutoWidth()
+			[
+				SNew(SOverlay)
+
+				+ SOverlay::Slot()
+				[
+					SNew(SImage)
+					.Image(this, &FSequencerDisplayNode::GetIconBrush)
+					.ColorAndOpacity(this, &FSequencerDisplayNode::GetIconColor)
+				]
+
+				+ SOverlay::Slot()
+				.VAlign(VAlign_Top)
+				.HAlign(HAlign_Right)
+				[
+					SNew(SImage)
+					.Image(this, &FSequencerDisplayNode::GetIconOverlayBrush)
+				]
+
+				+ SOverlay::Slot()
+				[
+					SNew(SSpacer)
+					.Visibility(EVisibility::Visible)
+					.ToolTipText(this, &FSequencerDisplayNode::GetIconToolTipText)
+				]
+			]
+
+			+ SHorizontalBox::Slot()
+			.VAlign(VAlign_Center)
+			[
+				SNew(STextBlock)
+				.Text(this, &FSequencerDisplayNode::GetDisplayName)
+				.ToolTipText(this, &FSequencerDisplayNode::GetDisplayNameToolTipText)
+			];
+	}
+	else if (InColumnName == ColumnNames.PinHeader)
+	{
+		return SNew(SCurveEditorTreePin, InCurveEditor, InTreeItemID);
+	}
+
+	return nullptr;
+}
+
+void FSequencerDisplayNode::CreateCurveModels(TArray<TUniquePtr<FCurveModel>>& OutCurveModels)
 {
-	ChildNodes.Add( InChild );
-	InChild->ParentNode = SharedThis( this );
 }
 
 
diff --git a/Engine/Source/Editor/Sequencer/Private/DisplayNodes/SequencerDisplayNode.h b/Engine/Source/Editor/Sequencer/Private/DisplayNodes/SequencerDisplayNode.h
--- a/Engine/Source/Editor/Sequencer/Private/DisplayNodes/SequencerDisplayNode.h
+++ b/Engine/Source/Editor/Sequencer/Private/DisplayNodes/SequencerDisplayNode.h
@@ -6,11 +6,13 @@
 #include "Misc/Attribute.h"
 #include "Styling/SlateColor.h"
 #include "Widgets/SWidget.h"
-#include "SequencerNodeTree.h"
 #include "Misc/FrameRate.h"
+#include "Tree/ICurveEditorTreeItem.h"
+#include "CurveEditorTypes.h"
 
 class FMenuBuilder;
 class FSequencer;
+class FSequencerNodeTree;
 class FSequencerDisplayNodeDragDropOp;
 class FSequencerObjectBindingNode;
 class IKeyArea;
@@ -54,43 +56,51 @@ namespace ESequencerNode
 		KeyArea,
 		/* Displays a category */
 		Category,
-		/* Benign spacer node */
-		Spacer,
+		/* Symbolic root node */
+		Root,
 		/* Folder node */
 		Folder
 	};
 }
 
+enum class EDisplayNodeSortType : uint8
+{
+	Folders,
+	Tracks,
+	ObjectBindings,
+	CameraCuts,
+	Shots,
+	Undefined,
+
+	NUM,
+};
 
 /**
  * Base Sequencer layout node.
  */
 class FSequencerDisplayNode
 	: public TSharedFromThis<FSequencerDisplayNode>
+	, public ICurveEditorTreeItem
 {
 public:
 
+	/** The serial number taken from the tree last time this node was encountered during a refresh.
+	 * When != FSequencerNodeTree::SerialNumber this node should be removed from all relevant structures */
+	uint32 TreeSerialNumber;
+
 	/**
 	 * Create and initialize a new instance.
 	 * 
 	 * @param InNodeName	The name identifier of then node
-	 * @param InParentNode	The parent of this node or nullptr if this is a root node
 	 * @param InParentTree	The tree this node is in
 	 */
-	FSequencerDisplayNode( FName InNodeName, TSharedPtr<FSequencerDisplayNode> InParentNode, FSequencerNodeTree& InParentTree );
+	FSequencerDisplayNode( FName InNodeName, FSequencerNodeTree& InParentTree);
 
 	/** Virtual destructor. */
 	virtual ~FSequencerDisplayNode(){}
 
 public:
 
-	/**
-	* Adds an object binding node to this node.
-	*
-	* @param ObjectBindingNode The node to add
-	*/
-	void AddObjectBindingNode( TSharedRef<FSequencerObjectBindingNode> ObjectBindingNode );
-
 	/** 
 	 * Finds any parent object binding node above this node in the hierarchy
 	 *
@@ -98,22 +108,29 @@ public:
 	 */
 	TSharedPtr<FSequencerObjectBindingNode> FindParentObjectBindingNode() const;
 
+	/** 
+	 * Finds this display node's closest parent object binding GUID, or an empty FGuid if it there is none
+	 */
+	FGuid GetObjectGuid() const;
+
 	/**
-	 * Adds a category to this node
-	 * 
-	 * @param CategoryName	Name of the category
-	 * @param DisplayName	Localized label to display in the animation outliner
+	 * Check whether this node is displayed at the root of the tree on the UI (ie, is part of the FSequencerNodeTree::GetRootNodes array)
+	 * @note: this is not the same as the node being the symbolic root node (FSequencerNodeTree::GetRootNode), of which there is only one.
 	 */
-	TSharedRef<class FSequencerSectionCategoryNode> AddCategoryNode( FName CategoryName, const FText& DisplayLabel );
+	bool IsRootNode() const;
 
 	/**
-	 * Adds a key area to this node
-	 *
-	 * @param KeyAreaName	Name of the key area
-	 * @param DisplayLabel	Localized label to display in the animation outliner
-	 * @param KeyArea		Key area interface for drawing and interaction with keys
+	 * Checks whether this node's parent is still relevant to the specified serial number
 	 */
-	void AddKeyAreaNode( FName KeyAreaName, const FText& DisplayLabel, TSharedRef<IKeyArea> KeyArea );
+	bool IsParentStillRelevant(uint32 SerialNumber) const;
+
+	/**
+	 * Retrieve the sort type of this node
+	 */
+	EDisplayNodeSortType GetSortType() const
+	{
+		return SortType;
+	}
 
 	/**
 	 * @return The type of node this is
@@ -256,6 +273,10 @@ public:
 	 */
 	virtual FText GetIconToolTipText() const;
 
+	// ICurveEditorTreeItem interface
+	virtual TSharedPtr<SWidget> GenerateCurveEditorTreeWidget(const FName& InColumnName, TWeakPtr<FCurveEditor> InCurveEditor, FCurveEditorTreeItemID InTreeItemID) override;
+	virtual void CreateCurveModels(TArray<TUniquePtr<FCurveModel>>& OutCurveModels) override;
+
 	/**
 	 * Get the display node that is ultimately responsible for constructing a section area widget for this node.
 	 * Could return this node itself, or a parent node
@@ -336,41 +357,43 @@ public:
 	bool TraverseVisible_ParentFirst(const TFunctionRef<bool(FSequencerDisplayNode&)>& InPredicate, bool bIncludeThisNode = true);
 
 	/**
-	 * Sorts the child nodes with the supplied predicate
+	 * Sort this node's immediate children using the persistent user-specifid reordering if possible
 	 */
-	template <class PREDICATE_CLASS>
-	void SortChildNodes(const PREDICATE_CLASS& Predicate)
-	{
-		ChildNodes.StableSort(Predicate);
+	void SortImmediateChildren();
 
-		for (TSharedRef<FSequencerDisplayNode> ChildNode : ChildNodes)
-		{
-			ChildNode->SortChildNodes(Predicate);
-		}
-	}
+	/**
+	 * Resort this node's immediate children, resetting any persistent user-specified reordering
+	 */
+	void ResortImmediateChildren();
 
 	/**
-	 * @return The parent of this node
+	 * @return The parent of this node. Will return null if this node is part of the FSequencerNodeTree::GetRootNodes array.
 	 */
 	TSharedPtr<FSequencerDisplayNode> GetParent() const
+	{
+		TSharedPtr<FSequencerDisplayNode> Pinned = ParentNode.Pin();
+		return (Pinned && Pinned->GetType() != ESequencerNode::Root) ? Pinned : nullptr;
+	}
+
+	/**
+	 * @return The parent of this node, or the symbolic root node if this node is part of the FSequencerNodeTree::GetRootNodes array.
+	 */
+	TSharedPtr<FSequencerDisplayNode> GetParentOrRoot() const
 	{
 		return ParentNode.Pin();
 	}
 
 	/**
-	 * @return The outermost parent of this node
+	 * @return The outermost parent of this node, ignoring the symbolic root node (ie, will never return FSequencerNodeTree::GetRootNode()
 	 */
 	TSharedRef<FSequencerDisplayNode> GetOutermostParent()
 	{
-		TSharedPtr<FSequencerDisplayNode> Parent = ParentNode.Pin();
+		TSharedPtr<FSequencerDisplayNode> Parent = GetParent();
 		return Parent.IsValid() ? Parent->GetOutermostParent() : AsShared();
 	}
 	
 	/** Gets the sequencer that owns this node */
-	FSequencer& GetSequencer() const
-	{
-		return ParentTree.GetSequencer();
-	}
+	FSequencer& GetSequencer() const;
 	
 	/** Gets the parent tree that this node is in */
 	FSequencerNodeTree& GetParentTree() const
@@ -396,14 +419,21 @@ public:
 	 */
 	bool IsHidden() const;
 
+	/**
+	 * @return Whether this node should be displayed on the tree view
+	 */
+	bool IsVisible() const;
+
 	/**
 	 * Check whether the node's tree view or track area widgets are hovered by the user's mouse.
 	 *
 	 * @return true if hovered, false otherwise. */
 	bool IsHovered() const;
 
-	/** Initialize this node with expansion states and virtual offsets */
-	void Initialize(float InVirtualTop, float InVirtualBottom);
+	/**
+	 * Called when the tree has been refreshed and this node has its new position in the hierarchy assigned
+	 */
+	void OnTreeRefreshed(float InVirtualTop, float InVirtualBottom);
 
 	/** @return this node's virtual offset from the top of the tree, irrespective of expansion states */
 	float GetVirtualTop() const
@@ -435,13 +465,20 @@ public:
 	 */
 	virtual void Drop( const TArray<TSharedRef<FSequencerDisplayNode>>& DraggedNodes, EItemDropZone DropZone ) { }
 
-protected:
 
-	/** Adds a child to this node, and sets it's parent to this node. */
-	void AddChildAndSetParent( TSharedRef<FSequencerDisplayNode> InChild );
 public:
-	/** Clears the parent of this node. */
-	void ClearParent() { ParentNode = nullptr; }
+
+	/** Assigns the parent of this node and adds it the the parent's child node list, removing it from its current parent's children if necessary. */
+	void SetParent(TSharedPtr<FSequencerDisplayNode> InParent);
+
+	/** Directly assigns the parent of this node without performing any other operation. Should only be used with care when child/parent relationships can be guaranteed. */
+	void SetParentDirectly(TSharedPtr<FSequencerDisplayNode> InParent);
+
+	/** Request that this node be reinitialized when the tree next refreshes. Currently Initialization only affects default expansion states. */
+	void RequestReinitialize()
+	{
+		bHasBeenInitialized = false;
+	}
 
 private:
 
@@ -462,7 +499,7 @@ protected:
 protected:
 
 	/** The parent of this node*/
-	TWeakPtr<FSequencerDisplayNode > ParentNode;
+	TWeakPtr<FSequencerDisplayNode> ParentNode;
 
 	/** List of children belonging to this node */
 	TArray<TSharedRef<FSequencerDisplayNode>> ChildNodes;
@@ -478,4 +515,12 @@ protected:
 
 	/** Event that is triggered when rename is requested */
 	FRequestRenameEvent RenameRequestedEvent;
+
+	/** The kind of thing that this node represents for sorting purposes */
+	EDisplayNodeSortType SortType;
+
+private:
+
+	/** Set to true when this node has been completely initialized for the first time */
+	bool bHasBeenInitialized;
 };
diff --git a/Engine/Source/Editor/Sequencer/Private/DisplayNodes/SequencerFolderNode.cpp b/Engine/Source/Editor/Sequencer/Private/DisplayNodes/SequencerFolderNode.cpp
--- a/Engine/Source/Editor/Sequencer/Private/DisplayNodes/SequencerFolderNode.cpp
+++ b/Engine/Source/Editor/Sequencer/Private/DisplayNodes/SequencerFolderNode.cpp
@@ -24,12 +24,14 @@
 
 #define LOCTEXT_NAMESPACE "SequencerFolderNode"
 
-FSequencerFolderNode::FSequencerFolderNode( UMovieSceneFolder& InMovieSceneFolder, TSharedPtr<FSequencerDisplayNode> InParentNode, FSequencerNodeTree& InParentTree )
-	: FSequencerDisplayNode( InMovieSceneFolder.GetFolderName(), InParentNode, InParentTree )
+FSequencerFolderNode::FSequencerFolderNode( UMovieSceneFolder& InMovieSceneFolder, FSequencerNodeTree& InParentTree )
+	: FSequencerDisplayNode( InMovieSceneFolder.GetFolderName(), InParentTree )
 	, MovieSceneFolder(InMovieSceneFolder)
 {
 	FolderOpenBrush = FEditorStyle::GetBrush( "ContentBrowser.AssetTreeFolderOpen" );
 	FolderClosedBrush = FEditorStyle::GetBrush( "ContentBrowser.AssetTreeFolderClosed" );
+
+	SortType = EDisplayNodeSortType::Folders;
 }
 
 
@@ -145,12 +147,12 @@ TOptional<EItemDropZone> FSequencerFolderNode::CanDrop( FSequencerDisplayNodeDra
 	// If they're trying to move an item above/below or into us. This item may or may not be a sibling; If it is not already a sibling then
 	// we need to check if it has a conflicting name with a sibling that we already have (for folders).
 	TArray<UMovieSceneFolder*> AdjacentFolders;
-	TSharedPtr<FSequencerDisplayNode> ChildParent;
+	TSharedPtr<const FSequencerDisplayNode> ChildParent;
 
 	if (ItemDropZone == EItemDropZone::OntoItem)
 	{
 		// If the item is being dropped onto us, we check our own children for name conflicts
-		ChildParent = SharedThis((FSequencerDisplayNode*)this);
+		ChildParent = AsShared();
 	}
 	else
 	{
@@ -281,7 +283,7 @@ void FSequencerFolderNode::MoveDisplayNodeToFolder(TSharedRef<FSequencerDisplayN
 
 	// Update the node's parent so that request for the NodePath reflect the new path
 	// instead of waiting until all nodes are regenerated by the subsequent Refresh call.
-	AddChildAndSetParent(Node);
+	Node->SetParent(AsShared());
 
 	// Update the expansion state using our new path.
 	Node->Traverse_ParentFirst([](FSequencerDisplayNode& TraversalNode)
@@ -419,11 +421,6 @@ void FSequencerFolderNode::OnColorPickerCancelled(FLinearColor NewFolderColor)
 	MovieSceneFolder.SetFolderColor(InitialFolderColor);
 }
 
-void FSequencerFolderNode::AddChildNode( TSharedRef<FSequencerDisplayNode> ChildNode )
-{
-	AddChildAndSetParent( ChildNode );
-}
-
 
 UMovieSceneFolder& FSequencerFolderNode::GetFolder() const
 {
@@ -445,4 +442,5 @@ void FSequencerFolderNode::ModifyAndSetSortingOrder(const int32 InSortingOrder)
 	MovieSceneFolder.Modify();
 	SetSortingOrder(InSortingOrder);
 }
+
 #undef LOCTEXT_NAMESPACE
diff --git a/Engine/Source/Editor/Sequencer/Private/DisplayNodes/SequencerFolderNode.h b/Engine/Source/Editor/Sequencer/Private/DisplayNodes/SequencerFolderNode.h
--- a/Engine/Source/Editor/Sequencer/Private/DisplayNodes/SequencerFolderNode.h
+++ b/Engine/Source/Editor/Sequencer/Private/DisplayNodes/SequencerFolderNode.h
@@ -16,7 +16,7 @@ enum class EItemDropZone;
 class FSequencerFolderNode : public FSequencerDisplayNode
 {
 public:
-	FSequencerFolderNode( UMovieSceneFolder& InMovieSceneFolder, TSharedPtr<FSequencerDisplayNode> InParentNode, FSequencerNodeTree& InParentTree );
+	FSequencerFolderNode( UMovieSceneFolder& InMovieSceneFolder, FSequencerNodeTree& InParentTree );
 
 	// FSequencerDisplayNode interface
 	virtual ESequencerNode::Type GetType() const override;
@@ -38,9 +38,6 @@ public:
 	/** Removes a node from it's old parent and makes it a child of this node. */
 	void MoveDisplayNodeToFolder(TSharedRef<FSequencerDisplayNode>& Node);
 
-	/** Adds a child node to this folder node. */
-	void AddChildNode( TSharedRef<FSequencerDisplayNode> ChildNode );
-
 	/** Gets the folder data for this display node. */
 	UMovieSceneFolder& GetFolder() const;
 
diff --git a/Engine/Source/Editor/Sequencer/Private/DisplayNodes/SequencerObjectBindingNode.cpp b/Engine/Source/Editor/Sequencer/Private/DisplayNodes/SequencerObjectBindingNode.cpp
--- a/Engine/Source/Editor/Sequencer/Private/DisplayNodes/SequencerObjectBindingNode.cpp
+++ b/Engine/Source/Editor/Sequencer/Private/DisplayNodes/SequencerObjectBindingNode.cpp
@@ -118,10 +118,9 @@ struct PropertyMenuData
 
 
 
-FSequencerObjectBindingNode::FSequencerObjectBindingNode(FName NodeName, const FText& InDisplayName, const FGuid& InObjectBinding, TSharedPtr<FSequencerDisplayNode> InParentNode, FSequencerNodeTree& InParentTree)
-	: FSequencerDisplayNode(NodeName, InParentNode, InParentTree)
+FSequencerObjectBindingNode::FSequencerObjectBindingNode(FName NodeName, const FGuid& InObjectBinding, FSequencerNodeTree& InParentTree)
+	: FSequencerDisplayNode(NodeName, InParentTree)
 	, ObjectBinding(InObjectBinding)
-	, DefaultDisplayName(InDisplayName)
 {
 	UMovieScene* MovieScene = GetSequencer().GetFocusedMovieSceneSequence()->GetMovieScene();
 
@@ -137,11 +136,8 @@ FSequencerObjectBindingNode::FSequencerObjectBindingNode(FName NodeName, const F
 	{
 		BindingType = EObjectBindingType::Unknown;
 	}
-}
 
-void FSequencerObjectBindingNode::AddTrackNode( TSharedRef<FSequencerTrackNode> NewChild )
-{
-	AddChildAndSetParent( NewChild );
+	SortType = EDisplayNodeSortType::ObjectBindings;
 }
 
 /* FSequencerDisplayNode interface
@@ -453,7 +449,7 @@ FText FSequencerObjectBindingNode::GetDisplayName() const
 		return MovieScene->GetObjectDisplayName(ObjectBinding);
 	}
 
-	return DefaultDisplayName;
+	return FText();
 }
 
 FLinearColor FSequencerObjectBindingNode::GetDisplayNameColor() const
diff --git a/Engine/Source/Editor/Sequencer/Private/DisplayNodes/SequencerObjectBindingNode.h b/Engine/Source/Editor/Sequencer/Private/DisplayNodes/SequencerObjectBindingNode.h
--- a/Engine/Source/Editor/Sequencer/Private/DisplayNodes/SequencerObjectBindingNode.h
+++ b/Engine/Source/Editor/Sequencer/Private/DisplayNodes/SequencerObjectBindingNode.h
@@ -30,12 +30,10 @@ public:
 	 * Create and initialize a new instance.
 	 * 
 	 * @param InNodeName The name identifier of then node.
-	 * @param InObjectName The name of the object we're binding to.
 	 * @param InObjectBinding Object binding guid for associating with live objects.
-	 * @param InParentNode The parent of this node, or nullptr if this is a root node.
 	 * @param InParentTree The tree this node is in.
 	 */
-	FSequencerObjectBindingNode(FName NodeName, const FText& InDisplayName, const FGuid& InObjectBinding, TSharedPtr<FSequencerDisplayNode> InParentNode, FSequencerNodeTree& InParentTree);
+	FSequencerObjectBindingNode(FName InNodeName, const FGuid& InObjectBinding, FSequencerNodeTree& InParentTree);
 
 public:
 
@@ -51,13 +49,6 @@ public:
 		return BindingType;
 	}
 
-	/**
-	 * Adds a new externally created node to this display node
-	 *
-	 * @param NewChild		The child node to add
-	 */
-	void AddTrackNode( TSharedRef<FSequencerTrackNode> NewChild );
-
 public:
 
 	// FSequencerDisplayNode interface
@@ -108,9 +99,6 @@ private:
 	/** The binding to live objects */
 	FGuid ObjectBinding;
 
-	/** The default display name of the object which is used if the binding manager doesn't provide one. */
-	FText DefaultDisplayName;
-
 	/** Enumeration specifying what kind of object binding this is */
 	EObjectBindingType BindingType;
 };
diff --git a/Engine/Source/Editor/Sequencer/Private/DisplayNodes/SequencerRootNode.h b/Engine/Source/Editor/Sequencer/Private/DisplayNodes/SequencerRootNode.h
new file mode 100644
--- /dev/null
+++ b/Engine/Source/Editor/Sequencer/Private/DisplayNodes/SequencerRootNode.h
@@ -0,0 +1,37 @@
+// Copyright 1998-2019 Epic Games, Inc. All Rights Reserved.
+
+#pragma once
+
+#include "DisplayNodes/SequencerDisplayNode.h"
+#include "Widgets/SNullWidget.h"
+
+/**
+ * A single symbolic root node for the whole sequencer tree
+ */
+class FSequencerRootNode
+	: public FSequencerDisplayNode
+{
+public:
+
+	/**
+	 * Create and initialize a new instance.
+	 * 
+	 * @param InParentTree The tree this node is in.
+	 */
+	explicit FSequencerRootNode(FSequencerNodeTree& InParentTree)
+		: FSequencerDisplayNode(NAME_None, InParentTree)
+	{}
+
+public:
+
+	// FSequencerDisplayNode interface
+
+	virtual bool CanRenameNode() const override { return false; }
+	virtual FText GetDisplayName() const override { return FText(); }
+	virtual float GetNodeHeight() const override { return 0.f; }
+	virtual FNodePadding GetNodePadding() const override { return FNodePadding(0.f); }
+	virtual ESequencerNode::Type GetType() const override { return ESequencerNode::Root; }
+	virtual void SetDisplayName(const FText& NewDisplayName) override { }
+	virtual TSharedRef<SWidget> GenerateContainerWidgetForOutliner(const TSharedRef<SSequencerTreeViewRow>& InRow) override { return SNullWidget::NullWidget; }
+	virtual bool IsSelectable() const override { return false; }
+};
diff --git a/Engine/Source/Editor/Sequencer/Private/DisplayNodes/SequencerSectionCategoryNode.h b/Engine/Source/Editor/Sequencer/Private/DisplayNodes/SequencerSectionCategoryNode.h
--- a/Engine/Source/Editor/Sequencer/Private/DisplayNodes/SequencerSectionCategoryNode.h
+++ b/Engine/Source/Editor/Sequencer/Private/DisplayNodes/SequencerSectionCategoryNode.h
@@ -14,17 +14,17 @@ class FSequencerSectionCategoryNode
 {
 public:
 
+	/** The display name of the category */
+	FText DisplayName;
+
 	/**
 	 * Create and initialize a new instance.
 	 * 
 	 * @param InNodeName The name identifier of then node.
-	 * @param InDisplayName Display name of the category.
-	 * @param InParentNode The parent of this node, or nullptr if this is a root node.
 	 * @param InParentTree The tree this node is in.
 	 */
-	FSequencerSectionCategoryNode(FName NodeName, const FText& InDisplayName, TSharedPtr<FSequencerDisplayNode> InParentNode, FSequencerNodeTree& InParentTree)
-		: FSequencerDisplayNode(NodeName, InParentNode, InParentTree)
-		, DisplayName(InDisplayName)
+	FSequencerSectionCategoryNode(FName NodeName, FSequencerNodeTree& InParentTree)
+		: FSequencerDisplayNode(NodeName, InParentTree)
 	{ }
 
 public:
@@ -39,8 +39,4 @@ public:
 	virtual ESequencerNode::Type GetType() const override;
 	virtual void SetDisplayName(const FText& NewDisplayName) override;
 
-private:
-
-	/** The display name of the category */
-	FText DisplayName;
 };
diff --git a/Engine/Source/Editor/Sequencer/Private/DisplayNodes/SequencerSectionKeyAreaNode.cpp b/Engine/Source/Editor/Sequencer/Private/DisplayNodes/SequencerSectionKeyAreaNode.cpp
--- a/Engine/Source/Editor/Sequencer/Private/DisplayNodes/SequencerSectionKeyAreaNode.cpp
+++ b/Engine/Source/Editor/Sequencer/Private/DisplayNodes/SequencerSectionKeyAreaNode.cpp
@@ -8,15 +8,14 @@
 #include "IKeyArea.h"
 #include "SKeyNavigationButtons.h"
 #include "SKeyAreaEditorSwitcher.h"
+#include "CurveModel.h"
 
 
 /* FSectionKeyAreaNode interface
  *****************************************************************************/
 
-FSequencerSectionKeyAreaNode::FSequencerSectionKeyAreaNode(FName NodeName, const FText& InDisplayName, TSharedPtr<FSequencerDisplayNode> InParentNode, FSequencerNodeTree& InParentTree, bool bInTopLevel)
-	: FSequencerDisplayNode(NodeName, InParentNode, InParentTree)
-	, DisplayName(InDisplayName)
-	, bTopLevel(bInTopLevel)
+FSequencerSectionKeyAreaNode::FSequencerSectionKeyAreaNode(FName NodeName, FSequencerNodeTree& InParentTree)
+	: FSequencerDisplayNode(NodeName, InParentTree)
 {
 }
 
@@ -113,3 +112,17 @@ void FSequencerSectionKeyAreaNode::SetDisplayName(const FText& NewDisplayName)
 {
 	check(false);
 }
+
+void FSequencerSectionKeyAreaNode::CreateCurveModels(TArray<TUniquePtr<FCurveModel>>& OutCurveModels)
+{
+	TSharedRef<ISequencer> Sequencer = GetSequencer().AsShared();
+
+	for (const TSharedRef<IKeyArea>& KeyArea : KeyAreas)
+	{
+		TUniquePtr<FCurveModel> NewCurve = KeyArea->CreateCurveEditorModel(Sequencer);
+		if (NewCurve.IsValid())
+		{
+			OutCurveModels.Add(MoveTemp(NewCurve));
+		}
+	}
+}
diff --git a/Engine/Source/Editor/Sequencer/Private/DisplayNodes/SequencerSectionKeyAreaNode.h b/Engine/Source/Editor/Sequencer/Private/DisplayNodes/SequencerSectionKeyAreaNode.h
--- a/Engine/Source/Editor/Sequencer/Private/DisplayNodes/SequencerSectionKeyAreaNode.h
+++ b/Engine/Source/Editor/Sequencer/Private/DisplayNodes/SequencerSectionKeyAreaNode.h
@@ -20,16 +20,16 @@ class FSequencerSectionKeyAreaNode
 {
 public:
 
+	/** The display name of the key area. */
+	FText DisplayName;
+
 	/**
 	 * Create and initialize a new instance.
 	 * 
 	 * @param InNodeName The name identifier of then node.
-	 * @param InDisplayName Display name of the category.
-	 * @param InParentNode The parent of this node, or nullptr if this is a root node.
 	 * @param InParentTree The tree this node is in.
-	 * @param bInTopLevel If true the node is part of the section itself instead of taking up extra height in the section.
 	 */
-	FSequencerSectionKeyAreaNode(FName NodeName, const FText& InDisplayName, TSharedPtr<FSequencerDisplayNode> InParentNode, FSequencerNodeTree& InParentTree, bool bInTopLevel = false);
+	FSequencerSectionKeyAreaNode(FName NodeName, FSequencerNodeTree& InParentTree);
 
 public:
 
@@ -58,11 +58,7 @@ public:
 		return KeyAreas;
 	}
 
-	/** @return Whether the node is top level.  (I.E., is part of the section itself instead of taking up extra height in the section) */
-	bool IsTopLevel() const
-	{
-		return bTopLevel;
-	}
+	void ClearKeyAreas();
 
 	/** Retrieve the key area editor switcher widget, creating it if it doesn't yet exist */
 	TSharedRef<SWidget> GetOrCreateKeyAreaEditorSwitcher();
@@ -79,17 +75,14 @@ public:
 	virtual ESequencerNode::Type GetType() const override;
 	virtual void SetDisplayName(const FText& NewDisplayName) override;
 
-private:
+	// ICurveEditorTreeItem interface
+	virtual void CreateCurveModels(TArray<TUniquePtr<FCurveModel>>& OutCurveModels) override;
 
-	/** The display name of the key area. */
-	FText DisplayName;
+private:
 
 	/** All key areas on this node (one per section). */
 	TArray<TSharedRef<IKeyArea>> KeyAreas;
 
 	/** The outliner key editor switcher widget. */
 	TSharedPtr<SKeyAreaEditorSwitcher> KeyEditorSwitcher;
-
-	/** If true the node is part of the section itself instead of taking up extra height in the section. */
-	bool bTopLevel;
 };
diff --git a/Engine/Source/Editor/Sequencer/Private/DisplayNodes/SequencerTrackNode.cpp b/Engine/Source/Editor/Sequencer/Private/DisplayNodes/SequencerTrackNode.cpp
--- a/Engine/Source/Editor/Sequencer/Private/DisplayNodes/SequencerTrackNode.cpp
+++ b/Engine/Source/Editor/Sequencer/Private/DisplayNodes/SequencerTrackNode.cpp
@@ -1,6 +1,7 @@
 // Copyright 1998-2019 Epic Games, Inc. All Rights Reserved.
 
 #include "DisplayNodes/SequencerTrackNode.h"
+#include "Algo/Copy.h"
 #include "Widgets/DeclarativeSyntaxSupport.h"
 #include "Widgets/SBoxPanel.h"
 #include "Widgets/Layout/SSpacer.h"
@@ -23,8 +24,11 @@
 #include "SequencerNodeSortingMethods.h"
 #include "SequencerNodeTree.h"
 #include "DisplayNodes/SequencerFolderNode.h"
+#include "SequencerSectionLayoutBuilder.h"
 #include "MovieSceneSequence.h"
 #include "MovieScene.h"
+#include "Tracks/MovieSceneCinematicShotTrack.h"
+#include "Tracks/MovieSceneCameraCutTrack.h"
 #include "MovieSceneFolder.h"
 
 #define LOCTEXT_NAMESPACE "SequencerTrackNode"
@@ -46,35 +50,203 @@ bool ContainsKeyableArea(TSharedRef<FSequencerSectionKeyAreaNode> InKeyAreaNode)
 	return false;
 }
 
-/* FTrackNode structors
- *****************************************************************************/
-
-FSequencerTrackNode::FSequencerTrackNode(UMovieSceneTrack& InAssociatedTrack, ISequencerTrackEditor& InAssociatedEditor, bool bInCanBeDragged, TSharedPtr<FSequencerDisplayNode> InParentNode, FSequencerNodeTree& InParentTree)
-	: FSequencerDisplayNode(InAssociatedTrack.GetFName(), InParentNode, InParentTree)
+FSequencerTrackNode::FSequencerTrackNode(UMovieSceneTrack* InAssociatedTrack, ISequencerTrackEditor& InAssociatedEditor, bool bInCanBeDragged, FSequencerNodeTree& InParentTree)
+	: FSequencerDisplayNode(InAssociatedTrack->GetFName(), InParentTree)
 	, AssociatedEditor(InAssociatedEditor)
-	, AssociatedTrack(&InAssociatedTrack)
+	, AssociatedTrack(InAssociatedTrack)
 	, bCanBeDragged(bInCanBeDragged)
 	, SubTrackMode(ESubTrackMode::None)
-{ }
+{
+	if (Cast<UMovieSceneCinematicShotTrack>(InAssociatedTrack))
+	{
+		SortType = EDisplayNodeSortType::Shots;
+	}
+	else if (Cast<UMovieSceneCameraCutTrack>(InAssociatedTrack))
+	{
+		SortType = EDisplayNodeSortType::CameraCuts;
+	}
+	else
+	{
+		SortType = EDisplayNodeSortType::Tracks;
+	}
+}
+
+void FSequencerTrackNode::UpdateInnerHierarchy()
+{
+	UMovieSceneTrack* Track = GetTrack();
+	if (!Track)
+	{
+		ClearChildren();
+		Sections.Empty();
+		return;
+	}
 
+	const int32 MaxRowIndex = Track->GetMaxRowIndex();
+	if (MaxRowIndex == 0)
+	{
+		// Single row, perhaps with a top level key area
+		if (SubTrackMode != ESubTrackMode::None)
+		{
+			RequestReinitialize();
+			SubTrackMode = ESubTrackMode::None;
+			ClearChildren();
+		}
 
-/* FTrackNode interface
- *****************************************************************************/
+		UpdateSections();
+	}
+	else
+	{
+		if (SubTrackMode != ESubTrackMode::ParentTrack)
+		{
+			// Change of type, so clear everything
+			TopLevelKeyNode = nullptr;
+			Sections.Empty();
+			ClearChildren();
+			SubTrackMode = ESubTrackMode::ParentTrack;
+			RequestReinitialize();
+		}
+
+		// Set bits for any row index that has a section on it. We then unset bits for existing tracks with that row index
+		TBitArray<> RowsWithSections;
+		RowsWithSections.Add(false, MaxRowIndex+1);
+		for (UMovieSceneSection* Section : Track->GetAllSections())
+		{
+			RowsWithSections[Section->GetRowIndex()] = true;
+		}
+
+		// Ensure we have one sub track node for each row index with at least one section on it
+		TArray<TSharedRef<FSequencerTrackNode>, TInlineAllocator<8>> SubTracks;
+
+		for (int32 Index = ChildNodes.Num()-1; Index >= 0; --Index)
+		{
+			TSharedRef<FSequencerDisplayNode> Child = ChildNodes[Index];
+			if (Child->GetType() == ESequencerNode::Track)
+			{
+				TSharedRef<FSequencerTrackNode> SubTrack = StaticCastSharedRef<FSequencerTrackNode>(Child);
+				ensure(SubTrack->GetSubTrackMode() == ESubTrackMode::SubTrack);
+
+				const int32 ThisTrackRow = SubTrack->GetRowIndex();
+				const bool bIsRelevant = RowsWithSections.IsValidIndex(ThisTrackRow) && RowsWithSections[ThisTrackRow] == true;
+
+				if (bIsRelevant)
+				{
+					// Keep this track
+					SubTrack->TreeSerialNumber = TreeSerialNumber;
+					SubTrack->UpdateSections();
+
+					// Unset the bit to indicate that we now have a track for this row
+					RowsWithSections[ThisTrackRow] = false;
+				}
+				else
+				{
+					// Remove this track node since it is no longer relevant
+					// Use SetParentDirectly to ensure that we do not modify ChildNodes while iterating
+					Child->SetParentDirectly(nullptr);
+					ChildNodes.RemoveAt(Index, 1, false);
+				}
+			}
+		}
+
+		bool bRequiresSort = false;
+
+		// Add new sub tracks for any remaining relevant tracks
+		for (TConstSetBitIterator<> It(RowsWithSections); It; ++It)
+		{
+			TSharedRef<FSequencerTrackNode> NewSubTrack = MakeShared<FSequencerTrackNode>(Track, AssociatedEditor, false, GetParentTree());
+			NewSubTrack->SetSubTrackMode(ESubTrackMode::SubTrack);
+			NewSubTrack->SetRowIndex(It.GetIndex());
+			// SetParent adds the track to our ChildNodes
+			NewSubTrack->SetParent(AsShared());
+
+			NewSubTrack->SetExpansionState(IsExpanded());
+
+			NewSubTrack->TreeSerialNumber = TreeSerialNumber;
+			NewSubTrack->UpdateSections();
 
-void FSequencerTrackNode::SetSectionAsKeyArea(TSharedRef<IKeyArea>& KeyArea)
+			bRequiresSort = true;
+		}
+
+		if (bRequiresSort)
+		{
+			auto SortByRowIndex = [](TSharedRef<FSequencerDisplayNode> A, TSharedRef<FSequencerDisplayNode> B)
+			{
+				const bool bBothTracks = (A->GetType() == ESequencerNode::Track) && (B->GetType() == ESequencerNode::Track);
+				return bBothTracks && StaticCastSharedRef<FSequencerTrackNode>(A)->GetRowIndex() < StaticCastSharedRef<FSequencerTrackNode>(B)->GetRowIndex();
+			};
+
+			Algo::Sort(ChildNodes, SortByRowIndex);
+		}
+	}
+}
+
+void FSequencerTrackNode::UpdateSections()
 {
-	if( !TopLevelKeyNode.IsValid() )
+	UMovieSceneTrack* Track = AssociatedTrack.Get();
+	if (!Track)
+	{
+		Sections.Empty();
+		TreeSerialNumber = 0;
+		ClearChildren();
+		return;
+	}
+
+	FGuid ObjectBinding;
+	if (TSharedPtr<FSequencerObjectBindingNode> ObjectBindingNode = FindParentObjectBindingNode())
+	{
+		ObjectBinding = ObjectBindingNode->GetObjectBinding();
+	}
+
+	TArray<UMovieSceneSection*, TInlineAllocator<4>> CurrentSections;
+
+	// ParentTracks never contain sections
+	if (SubTrackMode == ESubTrackMode::SubTrack)
+	{
+		Algo::CopyIf(Track->GetAllSections(), CurrentSections, [this](UMovieSceneSection* In) { return In->GetRowIndex() == this->RowIndex; });
+	}
+	else if (SubTrackMode == ESubTrackMode::None)
+	{
+		CurrentSections = Track->GetAllSections();
+	}
+
+	if (Sections.Num() != CurrentSections.Num())
+	{
+		Sections.Empty();
+	}
+
+	for (int32 Index = 0; Index < CurrentSections.Num(); ++Index)
 	{
-		bool bTopLevel = true;
-		TopLevelKeyNode = MakeShareable( new FSequencerSectionKeyAreaNode( GetNodeName(), FText::GetEmpty(), SharedThis(this), ParentTree, bTopLevel ) );
+		UMovieSceneSection* ThisSection     = CurrentSections[Index];
+		UMovieSceneSection* ExistingSection = Index < Sections.Num() ? Sections[Index]->GetSectionObject() : nullptr;
+
+		// Add a new section interface if there isn't one, or it doesn't correspond to the same section
+		if ( !ExistingSection || ExistingSection != ThisSection )
+		{
+			TSharedRef<ISequencerSection> SectionInterface = AssociatedEditor.MakeSectionInterface(*ThisSection, *Track, ObjectBinding);
+			Sections.Insert(SectionInterface, Index);
+		}
+
+		// Ask the section to generate its inner layout
+		FSequencerSectionLayoutBuilder LayoutBuilder(SharedThis(this), ThisSection);
+		Sections[Index]->GenerateSectionLayout(LayoutBuilder);
 	}
 
-	TopLevelKeyNode->AddKeyArea( KeyArea );
+	// Crop the section array at the new length
+	const int32 NumToRemove = Sections.Num() - CurrentSections.Num();
+	if (NumToRemove > 0)
+	{
+		Sections.RemoveAt(Sections.Num()-1-NumToRemove, NumToRemove, true);
+	}
 }
 
-void FSequencerTrackNode::AddKey(const FGuid& ObjectGuid)
+void FSequencerTrackNode::ClearChildren()
 {
-	AssociatedEditor.AddKey(ObjectGuid);
+	TArray<TSharedRef<FSequencerDisplayNode>> OldChildren;
+	Swap(ChildNodes, OldChildren);
+
+	for (TSharedRef<FSequencerDisplayNode> Child : OldChildren)
+	{
+		Child->SetParent(nullptr);
+	}
 }
 
 FSequencerTrackNode::ESubTrackMode FSequencerTrackNode::GetSubTrackMode() const
@@ -89,6 +261,7 @@ void FSequencerTrackNode::SetSubTrackMode(FSequencerTrackNode::ESubTrackMode InS
 
 int32 FSequencerTrackNode::GetRowIndex() const
 {
+	check(SubTrackMode == ESubTrackMode::SubTrack);
 	return RowIndex;
 }
 
@@ -661,4 +834,15 @@ void FSequencerTrackNode::ModifyAndSetSortingOrder(const int32 InSortingOrder)
 		SetSortingOrder(InSortingOrder);
 	}
 }
+
+void FSequencerTrackNode::CreateCurveModels(TArray<TUniquePtr<FCurveModel>>& OutCurveModels)
+{
+	TSharedPtr<FSequencerSectionKeyAreaNode> KeyAreaNode = GetTopLevelKeyNode();
+	if (KeyAreaNode.IsValid())
+	{
+		KeyAreaNode->CreateCurveModels(OutCurveModels);
+	}
+}
+
+
 #undef LOCTEXT_NAMESPACE
\ No newline at end of file
diff --git a/Engine/Source/Editor/Sequencer/Private/DisplayNodes/SequencerTrackNode.h b/Engine/Source/Editor/Sequencer/Private/DisplayNodes/SequencerTrackNode.h
--- a/Engine/Source/Editor/Sequencer/Private/DisplayNodes/SequencerTrackNode.h
+++ b/Engine/Source/Editor/Sequencer/Private/DisplayNodes/SequencerTrackNode.h
@@ -10,6 +10,7 @@
 #include "ISequencerSection.h"
 #include "IKeyArea.h"
 #include "SequencerHotspots.h"
+#include "SectionHandle.h"
 
 class FMenuBuilder;
 class ISequencerTrackEditor;
@@ -38,10 +39,9 @@ public:
 	 * @param InAssociatedType The track that this node represents.
 	 * @param InAssociatedEditor The track editor for the track that this node represents.
 	 * @param bInCanBeDragged Whether or not this node can be dragged and dropped.
-	 * @param InParentNode The parent of this node, or nullptr if this is a root node.
 	 * @param InParentTree The tree this node is in.
 	 */
-	FSequencerTrackNode(UMovieSceneTrack& InAssociatedTrack, ISequencerTrackEditor& InAssociatedEditor, bool bInCanBeDragged, TSharedPtr<FSequencerDisplayNode> InParentNode, FSequencerNodeTree& InParentTree);
+	FSequencerTrackNode(UMovieSceneTrack* InAssociatedTrack, ISequencerTrackEditor& InAssociatedEditor, bool bInCanBeDragged, FSequencerNodeTree& InParentTree);
 
 public:
 	/** Defines interaction modes when using sub-tracks for sections on multiple rows. */
@@ -58,32 +58,14 @@ public:
 public:
 
 	/**
-	 * Adds a section to this node
-	 *
-	 * @param SequencerSection	The section to add
+	 * Ensure this track's inner hierarchy is up to date, and that this track has the correct sub track mode initialized
 	 */
-	void AddSection(TSharedRef<ISequencerSection>& SequencerSection)
-	{
-		Sections.Add(SequencerSection);
-	}
-
-	void AddChildTrack(TSharedRef<FSequencerTrackNode> TrackNode)
-	{
-		AddChildAndSetParent(TrackNode);
-	}
+	void UpdateInnerHierarchy();
 
 	/**
-	 * Makes the section itself a key area without taking up extra space
-	 *
-	 * @param KeyArea	Interface for the key area
-	 */
-	void SetSectionAsKeyArea(TSharedRef<IKeyArea>& KeyArea);
-	
-	/**
-	 * Adds a key to the track
-	 *
+	 * Ensure that the section pointers for this track node are all correct based on its sub track mode and row index
 	 */
-	void AddKey(const FGuid& ObjectGuid);
+	void UpdateSections();
 
 	/**
 	 * @return All sections in this node
@@ -98,6 +80,11 @@ public:
 		return Sections;
 	}
 
+	void SetTopLevelKeyNode(TSharedPtr<FSequencerSectionKeyAreaNode> InTopLevelKeyNode)
+	{
+		TopLevelKeyNode = InTopLevelKeyNode;
+	}
+
 	/** @return Returns the top level key node for the section area if it exists */
 	TSharedPtr<FSequencerSectionKeyAreaNode> GetTopLevelKeyNode() const
 	{
@@ -157,10 +144,15 @@ public:
 	virtual void SetSortingOrder(const int32 InSortingOrder) override;
 	virtual void ModifyAndSetSortingOrder(const int32 InSortingOrder) override;
 
+	// ICurveEditorTreeItem interface
+	virtual void CreateCurveModels(TArray<TUniquePtr<FCurveModel>>& OutCurveModels) override;
+
 private:
 
 	FReply CreateNewSection() const;
 
+	void ClearChildren();
+
 private:
 
 	/** The track editor for the track associated with this node. */
diff --git a/Engine/Source/Editor/Sequencer/Private/IKeyArea.cpp b/Engine/Source/Editor/Sequencer/Private/IKeyArea.cpp
--- a/Engine/Source/Editor/Sequencer/Private/IKeyArea.cpp
+++ b/Engine/Source/Editor/Sequencer/Private/IKeyArea.cpp
@@ -10,12 +10,21 @@
 #include "Channels/MovieSceneChannel.h"
 #include "Channels/MovieSceneChannelProxy.h"
 #include "CurveModel.h"
+#include "ISequencer.h"
+#include "MovieSceneSequence.h"
 
-IKeyArea::IKeyArea(UMovieSceneSection& InSection, FMovieSceneChannelHandle InChannel)
-	: WeakOwningSection(&InSection)
-	, ChannelHandle(InChannel)
-	, Color(FLinearColor::White)
+IKeyArea::IKeyArea(UMovieSceneSection* InSection, FMovieSceneChannelHandle InChannel)
+	: ChannelHandle(InChannel)
 {
+	Reinitialize(InSection, InChannel);
+}
+
+void IKeyArea::Reinitialize(UMovieSceneSection* InSection, FMovieSceneChannelHandle InChannel)
+{
+	WeakOwningSection = InSection;
+	ChannelHandle = InChannel;
+	Color = FLinearColor::White;
+
 	if (const FMovieSceneChannelMetaData* MetaData = ChannelHandle.GetMetaData())
 	{
 		Color = MetaData->Color;
@@ -23,7 +32,7 @@ IKeyArea::IKeyArea(UMovieSceneSection& InSection, FMovieSceneChannelHandle InCha
 		DisplayText = MetaData->DisplayText;
 	}
 
-	UMovieScenePropertyTrack* PropertyTrack = InSection.GetTypedOuter<UMovieScenePropertyTrack>();
+	UMovieScenePropertyTrack* PropertyTrack = InSection->GetTypedOuter<UMovieScenePropertyTrack>();
 	if (PropertyTrack && !PropertyTrack->GetPropertyPath().IsEmpty())
 	{
 		PropertyBindings = FTrackInstancePropertyBindings(PropertyTrack->GetPropertyName(), PropertyTrack->GetPropertyPath());
@@ -190,6 +199,23 @@ void IKeyArea::PasteKeys(const FMovieSceneClipboardKeyTrack& KeyTrack, const FMo
 	}
 }
 
+FText GetOwningObjectBindingName(UMovieSceneTrack* InTrack, ISequencer& InSequencer)
+{
+	check(InTrack);
+
+	UMovieSceneSequence* FocusedSequence = InSequencer.GetFocusedMovieSceneSequence();
+	UMovieScene* MovieScene = FocusedSequence->GetMovieScene();
+
+	FGuid PossessableGuid;
+	if (MovieScene->FindTrackBinding(*InTrack, PossessableGuid))
+	{
+		return MovieScene->GetObjectDisplayName(PossessableGuid);
+	}
+
+	// Couldn't find an owning track, so must not be nestled inside of something!
+	return FText();
+}
+
 TUniquePtr<FCurveModel> IKeyArea::CreateCurveEditorModel(TSharedRef<ISequencer> InSequencer) const
 {
 	ISequencerChannelInterface* EditorInterface = FindChannelEditorInterface();
@@ -199,7 +225,62 @@ TUniquePtr<FCurveModel> IKeyArea::CreateCurveEditorModel(TSharedRef<ISequencer>
 		TUniquePtr<FCurveModel> CurveModel = EditorInterface->CreateCurveEditorModel_Raw(ChannelHandle, OwningSection, InSequencer);
 		if (CurveModel.IsValid())
 		{
-			CurveModel->SetDisplayName(DisplayText);
+			// Build long, short and context names for this curve to maximize information shown in the Curve Editor UI.
+			UMovieSceneTrack* OwningTrack = OwningSection->GetTypedOuter<UMovieSceneTrack>();
+			FText OwningTrackName;
+			FText ObjectBindingName;
+
+			if (OwningTrack)
+			{
+				OwningTrackName = OwningTrack->GetDisplayName();
+				
+				// This track might be inside an object binding and we'd like to prepend the object binding's name for more context.
+				ObjectBindingName = GetOwningObjectBindingName(OwningTrack, *InSequencer);
+			}
+
+			// Not all tracks have all the information so we need to format it differently depending on how many are valid.
+			TArray<FText> ValidNames;
+
+			if (!ObjectBindingName.IsEmptyOrWhitespace())
+			{
+				ValidNames.Add(ObjectBindingName);
+			}
+			if (!OwningTrackName.IsEmptyOrWhitespace())
+			{
+				ValidNames.Add(OwningTrackName);
+			}
+			if (!ChannelHandle.GetMetaData()->Group.IsEmptyOrWhitespace())
+			{
+				ValidNames.Add(ChannelHandle.GetMetaData()->Group);
+			}
+			if (!DisplayText.IsEmptyOrWhitespace())
+			{
+				ValidNames.Add(DisplayText);
+			}
+
+			// Now we loop through and string them together into one big format string.
+			FText LongDisplayNameFormatString;
+			for (int32 NameIndex = 0; NameIndex < ValidNames.Num(); NameIndex++)
+			{
+				const bool bLastEntry = NameIndex == ValidNames.Num() - 1;
+				if (!bLastEntry)
+				{
+					LongDisplayNameFormatString = FText::Format(NSLOCTEXT("SequencerIKeyArea", "CurveLongDisplayNameFormat", "{0}`{{1}`}."), LongDisplayNameFormatString, NameIndex);
+				}
+				else
+				{
+					LongDisplayNameFormatString = FText::Format(NSLOCTEXT("SequencerIKeyArea", "CurveLongDisplayNameFormatEnd", "{0}`{{1}`}"), LongDisplayNameFormatString, NameIndex);
+				}
+			}
+
+			FText LongDisplayName = FText::Format(LongDisplayNameFormatString, FFormatOrderedArguments(ValidNames));
+			const FText ShortDisplayName = DisplayText;
+			const FString IntentName = ChannelHandle.GetMetaData()->Group.IsEmptyOrWhitespace() ? DisplayText.ToString() : FString::Format(TEXT("{0}.{1}"), { ChannelHandle.GetMetaData()->Group.ToString(), DisplayText.ToString() });
+
+			CurveModel->SetShortDisplayName(DisplayText);
+			CurveModel->SetLongDisplayName(LongDisplayName);
+			CurveModel->SetIntentionName(IntentName);
+
 			if (Color.IsSet())
 			{
 				CurveModel->SetColor(Color.GetValue());
diff --git a/Engine/Source/Editor/Sequencer/Private/ISequencerSection.cpp b/Engine/Source/Editor/Sequencer/Private/ISequencerSection.cpp
--- a/Engine/Source/Editor/Sequencer/Private/ISequencerSection.cpp
+++ b/Engine/Source/Editor/Sequencer/Private/ISequencerSection.cpp
@@ -15,14 +15,6 @@ struct FChannelData
 
 	/** The channel's editor meta data */
 	const FMovieSceneChannelMetaData& MetaData;
-
-	/**
-	 * Make a key area out of this data
-	 */
-	TSharedRef<IKeyArea> MakeKeyArea(UMovieSceneSection& InSection) const
-	{
-		return MakeShared<IKeyArea>(InSection, Channel);
-	}
 };
 
 /** Data pertaining to a group of channels */
@@ -104,7 +96,7 @@ void ISequencerSection::GenerateSectionLayout( ISectionLayoutBuilder& LayoutBuil
 		const TTuple<FName, FGroupData>& Pair = *GroupToChannelsMap.CreateIterator();
 		if (Pair.Value.Channels.Num() == 1 && Pair.Value.Channels[0].MetaData.bCanCollapseToTrack)
 		{
-			LayoutBuilder.SetSectionAsKeyArea(Pair.Value.Channels[0].MakeKeyArea(*Section));
+			LayoutBuilder.SetTopLevelChannel(Pair.Value.Channels[0].Channel);
 			return;
 		}
 	}
@@ -155,8 +147,7 @@ void ISequencerSection::GenerateSectionLayout( ISectionLayoutBuilder& LayoutBuil
 
 		for (const FChannelData& ChannelAndData : ChannelData.Channels)
 		{
-			TSharedRef<IKeyArea> KeyArea = ChannelAndData.MakeKeyArea(*Section);
-			LayoutBuilder.AddKeyArea(ChannelAndData.MetaData.Name, ChannelAndData.MetaData.DisplayText, KeyArea);
+			LayoutBuilder.AddChannel(ChannelAndData.Channel);
 		}
 
 		if (!GroupName.IsNone())
diff --git a/Engine/Source/Editor/Sequencer/Private/SSequencer.cpp b/Engine/Source/Editor/Sequencer/Private/SSequencer.cpp
--- a/Engine/Source/Editor/Sequencer/Private/SSequencer.cpp
+++ b/Engine/Source/Editor/Sequencer/Private/SSequencer.cpp
@@ -28,6 +28,7 @@
 #include "Widgets/Layout/SSplitter.h"
 #include "Widgets/Input/SSpinBox.h"
 #include "Widgets/Input/SCheckBox.h"
+#include "Widgets/Docking/SDockTab.h"
 #include "EditorStyleSet.h"
 #include "Engine/Selection.h"
 #include "LevelEditorViewport.h"
@@ -37,7 +38,6 @@
 #include "DisplayNodes/SequencerTrackNode.h"
 #include "Widgets/Input/SNumericDropDown.h"
 #include "SequencerCommonHelpers.h"
-#include "SSequencerCurveEditorToolBar.h"
 #include "SSequencerLabelBrowser.h"
 #include "ISequencerWidgetsModule.h"
 #include "ScopedTransaction.h"
@@ -70,21 +70,192 @@
 #include "SSequencerPlayRateCombo.h"
 #include "Camera/CameraActor.h"
 #include "SCurveEditorPanel.h"
+#include "Tree/SCurveEditorTree.h"
+#include "SCurveKeyDetailPanel.h"
 #include "MovieSceneTimeHelpers.h"
 #include "FrameNumberNumericInterface.h"
 #include "LevelSequence.h"
 #include "SequencerLog.h"
 #include "MovieSceneCopyableBinding.h"
 #include "MovieSceneCopyableTrack.h"
+#include "IPropertyRowGenerator.h"
 
 #define LOCTEXT_NAMESPACE "Sequencer"
 
+const FName SSequencer::CurveEditorTabName = FName(TEXT("SequencerGraphEditor"));
+
 TSharedRef<IPropertyTypeCustomization> CreateFrameNumberCustomization(TWeakPtr<FSequencer> WeakSequencer)
 {
 	TSharedPtr<ISequencer> SequencerPtr = WeakSequencer.Pin();
 	return MakeShared<FFrameNumberDetailsCustomization>(SequencerPtr->GetNumericTypeInterface());
 }
 
+class SSequencerCurveEditor : public SCompoundWidget
+{
+	SLATE_BEGIN_ARGS(SSequencerCurveEditor)
+	{}
+	SLATE_END_ARGS()
+public:
+	void Construct(const FArguments& InArgs, TSharedRef<SCurveEditorPanel> InEditorPanel)
+	{
+		ChildSlot
+		[
+			SNew(SVerticalBox)
+
+			+ SVerticalBox::Slot()
+			.AutoHeight()
+			[
+				MakeToolbar(InEditorPanel)
+			]
+			+ SVerticalBox::Slot()
+			.FillHeight(1.0f)
+			[
+				InEditorPanel
+			]
+		];
+	}
+
+	TSharedRef<SWidget> MakeToolbar(TSharedRef<SCurveEditorPanel> InEditorPanel)
+	{
+		FToolBarBuilder ToolBarBuilder(InEditorPanel->GetCommands(), FMultiBoxCustomization::None, InEditorPanel->GetToolbarExtender(), EOrientation::Orient_Horizontal, true);
+		ToolBarBuilder.BeginSection("Asset");
+		ToolBarBuilder.EndSection();
+		// We just use all of the extenders as our toolbar, we don't have a need to create a separate toolbar.
+		return ToolBarBuilder.MakeWidget();
+	}
+};
+
+class FSequencerCurveEditorTimeSliderController : public FSequencerTimeSliderController
+{
+public:
+
+	FSequencerCurveEditorTimeSliderController(const FTimeSliderArgs& InArgs, TWeakPtr<FSequencer> InWeakSequencer, TSharedRef<FCurveEditor> InCurveEditor)
+		: FSequencerTimeSliderController(InArgs, InWeakSequencer)
+	{
+		WeakSequencer = InWeakSequencer;
+		WeakCurveEditor = InCurveEditor;
+	}
+
+	virtual void ClampViewRange(double& NewRangeMin, double& NewRangeMax) override
+	{
+		// Since the CurveEditor uses a different view range (potentially) we have to be careful about which one we clamp.
+		TSharedPtr<ISequencer> Sequencer = WeakSequencer.Pin();
+		if (!Sequencer.IsValid())
+		{
+			return;
+		}
+
+		const bool bLinkedTimeRange = Sequencer->GetSequencerSettings()->GetLinkCurveEditorTimeRange();
+		if (bLinkedTimeRange)
+		{
+			return FSequencerTimeSliderController::ClampViewRange(NewRangeMin, NewRangeMax);
+		}
+		else
+		{
+			TSharedPtr<FCurveEditor> CurveEditor = WeakCurveEditor.Pin();
+			if (CurveEditor.IsValid())
+			{
+				double InputMin, InputMax;
+				CurveEditor->GetBounds().GetInputBounds(InputMin, InputMax);
+
+				bool bNeedsClampSet = false;
+				double NewClampRangeMin = InputMin;
+				if (NewRangeMin < InputMin)
+				{
+					NewClampRangeMin = NewRangeMin;
+					bNeedsClampSet = true;
+				}
+
+				double NewClampRangeMax = InputMax;
+				if (NewRangeMax > InputMax)
+				{
+					NewClampRangeMax = NewRangeMax;
+					bNeedsClampSet = true;
+				}
+
+				if (bNeedsClampSet)
+				{
+					CurveEditor->GetBounds().SetInputBounds(NewClampRangeMin, NewClampRangeMax);
+				}
+
+			}
+		}
+	}
+
+	virtual void SetViewRange(double NewRangeMin, double NewRangeMax, EViewRangeInterpolation Interpolation) override
+	{
+		TSharedPtr<ISequencer> Sequencer = WeakSequencer.Pin();
+		if (!Sequencer.IsValid())
+		{
+			return;
+		}
+
+		const bool bLinkedTimeRange = Sequencer->GetSequencerSettings()->GetLinkCurveEditorTimeRange();
+		if (bLinkedTimeRange)
+		{
+			return FSequencerTimeSliderController::SetViewRange(NewRangeMin, NewRangeMax, Interpolation);
+		}
+		else
+		{
+			// Clamp to a minimum size to avoid zero-sized or negative visible ranges
+			double MinVisibleTimeRange = FFrameNumber(1) / GetTickResolution();
+			TRange<double> ExistingViewRange = GetViewRange();
+
+			if (NewRangeMax == ExistingViewRange.GetUpperBoundValue())
+			{
+				if (NewRangeMin > NewRangeMax - MinVisibleTimeRange)
+				{
+					NewRangeMin = NewRangeMax - MinVisibleTimeRange;
+				}
+			}
+			else if (NewRangeMax < NewRangeMin + MinVisibleTimeRange)
+			{
+				NewRangeMax = NewRangeMin + MinVisibleTimeRange;
+			}
+
+			TSharedPtr<FCurveEditor> CurveEditor = WeakCurveEditor.Pin();
+			if (CurveEditor.IsValid())
+			{
+				CurveEditor->GetBounds().SetInputBounds(NewRangeMin, NewRangeMax);
+			}
+		}
+	}
+
+
+
+	virtual FAnimatedRange GetViewRange() const override
+	{ 
+		// If they've linked the Sequencer timerange we can return the internal controller's view range, otherwise we return the bounds (which internally does the same check)
+		TSharedPtr<ISequencer> Sequencer = WeakSequencer.Pin();
+		if (!Sequencer.IsValid())
+		{
+			return FAnimatedRange();
+		}
+		const bool bLinkedTimeRange = Sequencer->GetSequencerSettings()->GetLinkCurveEditorTimeRange();
+		if (bLinkedTimeRange)
+		{
+			return FSequencerTimeSliderController::GetViewRange();
+		}
+		else
+		{
+			TSharedPtr<FCurveEditor> CurveEditor = WeakCurveEditor.Pin();
+			if (CurveEditor.IsValid())
+			{
+				double InputMin, InputMax;
+				CurveEditor->GetBounds().GetInputBounds(InputMin, InputMax);
+
+				return FAnimatedRange(InputMin, InputMax);
+			}
+		}
+
+		return FAnimatedRange();
+	}
+
+private:
+	TWeakPtr<ISequencer> WeakSequencer;
+	TWeakPtr<FCurveEditor> WeakCurveEditor;
+};
+
 /* SSequencer interface
  *****************************************************************************/
 PRAGMA_DISABLE_OPTIMIZATION
@@ -221,13 +392,43 @@ void SSequencer::Construct(const FArguments& InArgs, TSharedRef<FSequencer> InSe
 	TrackArea->SetTreeView(TreeView);
 
 	TAttribute<FAnimatedRange> ViewRangeAttribute = InArgs._ViewRange;
-	TSharedRef<SCurveEditorPanel> CurveEditorPanel = 
-		SNew(SCurveEditorPanel, InSequencer->GetCurveEditor().ToSharedRef())
-		.Visibility(this, &SSequencer::GetCurveEditorVisibility)
-		// Grid lines match the color specified in FSequencerTimeSliderController::OnPaintSectionView
-		.GridLineTint(FLinearColor(0.f, 0.f, 0.f, 0.3f));
+	
+	// We create a custom Time Slider Controller which is just a wrapper around the actual one, but is aware of our custom bounds logic. Currently the range the
+	// bar displays is tied to Sequencer timeline and not the Bounds, so we need a way of changing it to look at the Bounds but only for the Curve Editor time
+	// slider controller. We want everything else to just pass through though.
+	TSharedRef<ITimeSliderController> CurveEditorTimeSliderController = MakeShared<FSequencerCurveEditorTimeSliderController>(TimeSliderArgs, SequencerPtr, InSequencer->GetCurveEditor().ToSharedRef());
+
+	// Initialize the Curve Editor Widget if there is a tab manager to spawn our extra tab in.
+	// Some areas that use Sequencer don't use our curve editor. In this case no button is shown on the UI.
+	if (InSequencer->GetToolkitHost().IsValid())
+	{
+		TSharedRef<SCurveEditorPanel> CurveEditorWidget = SNew(SCurveEditorPanel, InSequencer->GetCurveEditor().ToSharedRef())
+			// Grid lines match the color specified in FSequencerTimeSliderController::OnPaintViewArea
+			.GridLineTint(FLinearColor(0.f, 0.f, 0.f, 0.3f))
+			.ExternalTimeSliderController(CurveEditorTimeSliderController)
+			.TabManager(InSequencer->GetToolkitHost()->GetTabManager())
+			.DisabledTimeSnapTooltip(LOCTEXT("CurveEditorTimeSnapDisabledTooltip", "Time Snapping is currently driven by Sequencer."))
+			.TreeContent()
+			[
+				SNew(SCurveEditorTree, InSequencer->GetCurveEditor())
+			];
+
+			// Register an instanced custom property type layout to handle converting FFrameNumber from Tick Resolution to Display Rate.
+			CurveEditorWidget->GetKeyDetailsView()->GetPropertyRowGenerator()->RegisterInstancedCustomPropertyTypeLayout("FrameNumber", FOnGetPropertyTypeCustomizationInstance::CreateStatic(CreateFrameNumberCustomization, SequencerPtr));
+			TAttribute<bool> IsEnabledAttribute = TAttribute<bool>::Create(TAttribute<bool>::FGetter::CreateSP(this, &SSequencer::GetIsCurveEditorEnabled));
+	
+			CurveEditorPanel = SNew(SSequencerCurveEditor, CurveEditorWidget);
+			CurveEditorPanel->SetEnabled(IsEnabledAttribute);
+			CurveEditorWidget->SetEnabled(IsEnabledAttribute);
+
+		// Check to see if the tab is already opened due to the saved window layout.
+		TSharedPtr<SDockTab> ExistingCurveEditorTab = InSequencer->GetToolkitHost()->GetTabManager()->FindExistingLiveTab(FTabId(SSequencer::CurveEditorTabName));
+		if (ExistingCurveEditorTab)
+		{
+			ExistingCurveEditorTab->SetContent(CurveEditorPanel.ToSharedRef());
+		}
+	}
 
-	CurveEditorPanel->GetKeyDetailsView()->RegisterInstancedCustomPropertyTypeLayout("FrameNumber", FOnGetPropertyTypeCustomizationInstance::CreateStatic(CreateFrameNumberCustomization, SequencerPtr));
 
 	const int32 Column0 = 0, Column1 = 1;
 	const int32 Row0 = 0, Row1 = 1, Row2 = 2, Row3 = 3, Row4 = 4;
@@ -235,10 +436,6 @@ void SSequencer::Construct(const FArguments& InArgs, TSharedRef<FSequencer> InSe
 	const float CommonPadding = 3.f;
 	const FMargin ResizeBarPadding(4.f, 0, 0, 0);
 
-	TSharedRef<FUICommandList> CurveEditorAndSequencerCommands = MakeShared<FUICommandList>();
-	CurveEditorAndSequencerCommands->Append(CurveEditorPanel->GetCommands().ToSharedRef());
-	CurveEditorAndSequencerCommands->Append(InSequencer->GetCommandBindings().ToSharedRef());
-
 	ChildSlot
 	[
 		SNew(SVerticalBox)
@@ -284,19 +481,14 @@ void SSequencer::Construct(const FArguments& InArgs, TSharedRef<FSequencer> InSe
 						[
 							SNew(SHorizontalBox)
 
+							// Left aligned Toolbar Icons
 							+SHorizontalBox::Slot()
 							.AutoWidth()
 							[
 								MakeToolBar()
 							]
 
-							+SHorizontalBox::Slot()
-							.AutoWidth()
-							[
-								SNew(SSequencerCurveEditorToolBar, InSequencer, CurveEditorAndSequencerCommands)
-								.Visibility(this, &SSequencer::GetCurveEditorToolBarVisibility)
-							]
-
+							// Right Aligned Breadcrumbs
 							+ SHorizontalBox::Slot()
 							.HAlign(HAlign_Right)
 							.VAlign(VAlign_Center)
@@ -316,6 +508,7 @@ void SSequencer::Construct(const FArguments& InArgs, TSharedRef<FSequencer> InSe
 								.TextStyle(FEditorStyle::Get(), "Sequencer.BreadcrumbText")
 							]
 
+							// Sequence Locking symbol
 							+SHorizontalBox::Slot()
 							.HAlign(HAlign_Right)
 							.VAlign(VAlign_Center)
@@ -432,7 +625,6 @@ void SSequencer::Construct(const FArguments& InArgs, TSharedRef<FSequencer> InSe
 									[
 										SNew(SBox)
 										.Padding(ResizeBarPadding)
-										.Visibility(this, &SSequencer::GetTrackAreaVisibility )
 										.Clipping(EWidgetClipping::ClipToBounds)
 										[
 											TrackArea.ToSharedRef()
@@ -447,12 +639,6 @@ void SSequencer::Construct(const FArguments& InArgs, TSharedRef<FSequencer> InSe
 								ScrollBar
 							]
 						]
-
-						+ SHorizontalBox::Slot()
-						.FillWidth( TAttribute<float>( this, &SSequencer::GetOutlinerSpacerFill ) )
-						[
-							SNew(SSpacer)
-						]
 					]
 
 					// playback buttons
@@ -504,15 +690,8 @@ void SSequencer::Construct(const FArguments& InArgs, TSharedRef<FSequencer> InSe
 						.Clipping(EWidgetClipping::ClipToBounds)
 					]
 
-					// Curve editor
-					+ SGridPanel::Slot( Column1, Row2, SGridPanel::Layer(20) )
-					.Padding(ResizeBarPadding)
-					[
-						CurveEditorPanel
-					]
-
 					// Overlay that draws the scrub position
-					+ SGridPanel::Slot( Column1, Row2, SGridPanel::Layer(30) )
+					+ SGridPanel::Slot( Column1, Row2, SGridPanel::Layer(20) )
 					.Padding(ResizeBarPadding)
 					[
 						SNew( SSequencerSectionOverlay, TimeSliderControllerRef )
@@ -524,7 +703,7 @@ void SSequencer::Construct(const FArguments& InArgs, TSharedRef<FSequencer> InSe
 						.Clipping(EWidgetClipping::ClipToBounds)
 					]
 
-					+ SGridPanel::Slot(Column1, Row2, SGridPanel::Layer(40))
+					+ SGridPanel::Slot(Column1, Row2, SGridPanel::Layer(30))
 						.Padding(ResizeBarPadding)
 						.HAlign(HAlign_Left)
 						.VAlign(VAlign_Top)
@@ -533,7 +712,7 @@ void SSequencer::Construct(const FArguments& InArgs, TSharedRef<FSequencer> InSe
 						SAssignNew(TransformBox, SSequencerTransformBox, SequencerPtr.Pin().ToSharedRef(), *Settings, NumericTypeInterface.ToSharedRef())
 					]
 
-					+ SGridPanel::Slot(Column1, Row2, SGridPanel::Layer(50))
+					+ SGridPanel::Slot(Column1, Row2, SGridPanel::Layer(40))
 					.Padding(ResizeBarPadding)
 					[
 						SAssignNew(TickResolutionOverlay, SSequencerTimePanel, SequencerPtr)
@@ -1022,7 +1201,11 @@ TSharedRef<SWidget> SSequencer::MakeToolBar()
 
 	ToolBarBuilder.BeginSection("Curve Editor");
 	{
-		ToolBarBuilder.AddToolBarButton( FSequencerCommands::Get().ToggleShowCurveEditor );
+		// Only add the button if we have a toolkit host to spawn tabs in
+		if (SequencerPtr.Pin()->GetToolkitHost().IsValid())
+		{
+			ToolBarBuilder.AddToolBarButton( FSequencerCommands::Get().ToggleShowCurveEditor );
+		}
 	}
 	ToolBarBuilder.EndSection();
 
@@ -1520,6 +1703,18 @@ TSharedPtr<ITimeSlider> SSequencer::GetTopTimeSliderWidget() const
 SSequencer::~SSequencer()
 {
 	USelection::SelectionChangedEvent.RemoveAll(this);
+
+
+	TSharedPtr<FSequencer> Sequencer = SequencerPtr.Pin();
+	if(Sequencer)
+	{
+		FTabId TabId = FTabId(SSequencer::CurveEditorTabName);
+		TSharedPtr<SDockTab> CurveEditorTab = Sequencer->GetToolkitHost()->GetTabManager()->FindExistingLiveTab(TabId);
+		if (CurveEditorTab)
+		{
+			CurveEditorTab->RequestCloseTab();
+		}
+	}
 }
 
 
@@ -2001,11 +2196,6 @@ void SSequencer::OnCrumbClicked(const FSequencerBreadcrumb& Item)
 		}
 		else
 		{
-			if (SequencerPtr.Pin()->GetShowCurveEditor())
-			{
-				SequencerPtr.Pin()->SetShowCurveEditor(false);
-			}
-
 			SequencerPtr.Pin()->PopToSequenceInstance( Item.SequenceID );
 		}
 	}
@@ -2024,39 +2214,6 @@ TSharedPtr<SSequencerTreeView> SSequencer::GetTreeView() const
 }
 
 
-TArray<FSectionHandle> SSequencer::GetSectionHandles(const TSet<TWeakObjectPtr<UMovieSceneSection>>& DesiredSections) const
-{
-	TArray<FSectionHandle> SectionHandles;
-
-	TSharedPtr<FSequencer> Sequencer = SequencerPtr.Pin();
-	if (Sequencer.IsValid())
-	{
-		// @todo sequencer: this is potentially slow as it traverses the entire tree - there's scope for optimization here
-		for (const FDisplayNodeRef& Node : Sequencer->GetNodeTree()->GetRootNodes())
-		{
-			Node->Traverse_ParentFirst([&](FSequencerDisplayNode& InNode) {
-				if (InNode.GetType() == ESequencerNode::Track)
-				{
-					FSequencerTrackNode& TrackNode = static_cast<FSequencerTrackNode&>(InNode);
-
-					const TArray<TSharedRef<ISequencerSection>>& AllSections = TrackNode.GetSections();
-					for (int32 Index = 0; Index < AllSections.Num(); ++Index)
-					{
-						if (DesiredSections.Contains(MakeWeakObjectPtr(AllSections[Index]->GetSectionObject())))
-						{
-							SectionHandles.Emplace(StaticCastSharedRef<FSequencerTrackNode>(TrackNode.AsShared()), Index);
-						}
-					}
-				}
-				return true;
-			});
-		}
-	}
-
-	return SectionHandles;
-}
-
-
 void SSequencer::OnSaveMovieSceneClicked()
 {
 	SequencerPtr.Pin()->SaveCurrentMovieScene();
@@ -2267,12 +2424,6 @@ EVisibility SSequencer::GetBreadcrumbTrailVisibility() const
 }
 
 
-EVisibility SSequencer::GetCurveEditorToolBarVisibility() const
-{
-	return SequencerPtr.Pin()->GetShowCurveEditor() ? EVisibility::Visible : EVisibility::Collapsed;
-}
-
-
 EVisibility SSequencer::GetBottomTimeSliderVisibility() const
 {
 	return Settings->GetShowRangeSlider() ? EVisibility::Hidden : EVisibility::Visible;
@@ -2289,55 +2440,34 @@ EFrameNumberDisplayFormats SSequencer::GetTimeDisplayFormat() const
 	return Settings->GetTimeDisplayFormat();
 }
 
-float SSequencer::GetOutlinerSpacerFill() const
-{
-	const float Column1Coeff = GetColumnFillCoefficient(1);
-	return SequencerPtr.Pin()->GetShowCurveEditor() ? Column1Coeff / (1 - Column1Coeff) : 0.f;
-}
-
 
 void SSequencer::OnColumnFillCoefficientChanged(float FillCoefficient, int32 ColumnIndex)
 {
 	ColumnFillCoefficients[ColumnIndex] = FillCoefficient;
 }
 
-
-EVisibility SSequencer::GetTrackAreaVisibility() const
-{
-	return SequencerPtr.Pin()->GetShowCurveEditor() ? EVisibility::Collapsed : EVisibility::Visible;
-}
-
-
-EVisibility SSequencer::GetCurveEditorVisibility() const
+void SSequencer::OnCurveEditorVisibilityChanged(bool bShouldBeVisible)
 {
-	return SequencerPtr.Pin()->GetShowCurveEditor() ? EVisibility::Visible : EVisibility::Collapsed;
-}
-
+	TSharedPtr<FSequencer> Sequencer = SequencerPtr.Pin();
+	FTabId TabId = FTabId(SSequencer::CurveEditorTabName);
 
-void SSequencer::OnCurveEditorVisibilityChanged()
-{
-	if (!Settings->GetLinkCurveEditorTimeRange())
+	if (bShouldBeVisible)
 	{
-		TRange<double> ClampRange = SequencerPtr.Pin()->GetClampRange();
-		if (CachedClampRange.IsEmpty())
-		{
-			CachedClampRange = ClampRange;
-		}
-		SequencerPtr.Pin()->SetClampRange(CachedClampRange);
-		CachedClampRange = ClampRange;
-
-		TRange<double> ViewRange = SequencerPtr.Pin()->GetViewRange();
-		if (CachedViewRange.IsEmpty())
+		// Request the Tab Manager invoke the tab. This will spawn the tab if needed, otherwise pull it to focus. This assumes
+		// that the Toolkit Host's Tab Manager has already registered a tab with a NullWidget for content.
+		TSharedRef<SDockTab> CurveEditorTab = Sequencer->GetToolkitHost()->GetTabManager()->InvokeTab(TabId);
+		CurveEditorTab->SetContent(CurveEditorPanel.ToSharedRef());
+		SequencerPtr.Pin()->GetCurveEditor()->ZoomToFit();
+	}
+	else
+	{
+		TSharedPtr<SDockTab> ExistingTab = Sequencer->GetToolkitHost()->GetTabManager()->FindExistingLiveTab(TabId);
+		if (ExistingTab)
 		{
-			CachedViewRange = ViewRange;
+			ExistingTab->RequestCloseTab();
 		}
-		SequencerPtr.Pin()->SetViewRange(CachedViewRange);
-		CachedViewRange = ViewRange;
 	}
 
-	SequencerPtr.Pin()->SyncCurveEditorToSelection(false);
-	SequencerPtr.Pin()->GetCurveEditor()->ZoomToFit();
-
 	TreeView->UpdateTrackArea();
 }
 
@@ -2388,8 +2518,11 @@ FPasteContextMenuArgs SSequencer::GeneratePasteArgs(FFrameNumber PasteAtTime, TS
 	TSharedPtr<FSequencer> Sequencer = SequencerPtr.Pin();
 	if (Settings->GetIsSnapEnabled())
 	{
-		// @todo: sequencer-timecode: play rate override
-		//PasteAtTime = FQualifiedFrameTime(PasteAtTime, GetTickResolution()).ToFrameRate();
+		FFrameRate TickResolution = Sequencer->GetFocusedTickResolution();
+		FFrameRate DisplayRate    = Sequencer->GetFocusedDisplayRate();
+
+		PasteAtTime = ConvertFrameTime(PasteAtTime, TickResolution, DisplayRate).RoundToFrame();
+		PasteAtTime = ConvertFrameTime(PasteAtTime, DisplayRate, TickResolution).FrameNumber;
 	}
 
 	// Open a paste menu at the current mouse position
@@ -2407,9 +2540,12 @@ FPasteContextMenuArgs SSequencer::GeneratePasteArgs(FFrameNumber PasteAtTime, TS
 			Sections.Add(Key.Section);
 		}
 
-		for (const FSectionHandle& Handle : GetSectionHandles(Sections))
+		for (TWeakObjectPtr<UMovieSceneSection> WeakSection : Sections)
 		{
-			PasteIntoNodes.Add(Handle.TrackNode.ToSharedRef());
+			if (TOptional<FSectionHandle> Handle = Sequencer->GetNodeTree()->GetSectionHandle(WeakSection.Get()))
+			{
+				PasteIntoNodes.Add(Handle->GetTrackNode());
+			}
 		}
 	}
 
diff --git a/Engine/Source/Editor/Sequencer/Private/SSequencer.h b/Engine/Source/Editor/Sequencer/Private/SSequencer.h
--- a/Engine/Source/Editor/Sequencer/Private/SSequencer.h
+++ b/Engine/Source/Editor/Sequencer/Private/SSequencer.h
@@ -29,9 +29,10 @@ class SSequencerTrackArea;
 class SSequencerTrackOutliner;
 class SSequencerTransformBox;
 class SSequencerTreeView;
+class SCurveEditorPanel;
+class SDockTab;
 class USequencerSettings;
 struct FPaintPlaybackRangeArgs;
-struct FSectionHandle;
 
 namespace SequencerLayoutConstants
 {
@@ -229,7 +230,7 @@ public:
 	void OnSaveMovieSceneAsClicked();
 
 	/** Called when the curve editor is shown or hidden */
-	void OnCurveEditorVisibilityChanged();
+	void OnCurveEditorVisibilityChanged(bool bShouldBeVisible);
 
 	/** Access the tree view for this sequencer */
 	TSharedPtr<SSequencerTreeView> GetTreeView() const;
@@ -237,9 +238,6 @@ public:
 	/** Generate a helper structure that can be used to transform between phsyical space and virtual space in the track area */
 	FVirtualTrackArea GetVirtualTrackArea() const;
 
-	/** Get an array of section handles for the given set of movie scene sections */
-	TArray<FSectionHandle> GetSectionHandles(const TSet<TWeakObjectPtr<UMovieSceneSection>>& DesiredSections) const;
-
 	/** @return a numeric type interface that will parse and display numbers as frames and times correctly */
 	TSharedRef<INumericTypeInterface<double>> GetNumericTypeInterface() const;
 	
@@ -333,7 +331,6 @@ public:
 	void FillTimeDisplayFormatMenu(FMenuBuilder& MenuBuilder);
 
 public:	
-
 	/** Makes a time range widget with the specified inner content */
 	TSharedRef<SWidget> MakeTimeRange(const TSharedRef<SWidget>& InnerContent, bool bShowWorkingRange, bool bShowViewRange, bool bShowPlaybackRange);
 
@@ -374,15 +371,6 @@ private:
 		return ColumnFillCoefficients[ColumnIndex];
 	}
 
-	/** Get the amount of space that the outliner spacer should fill */
-	float GetOutlinerSpacerFill() const;
-
-	/** Get the visibility of the curve area */
-	EVisibility GetCurveEditorVisibility() const;
-
-	/** Get the visibility of the track area */
-	EVisibility GetTrackAreaVisibility() const;
-
 	/**
 	 * Called when one or more assets are dropped into the widget
 	 *
@@ -416,8 +404,7 @@ private:
 	/** Gets whether or not the breadcrumb trail should be visible. */
 	EVisibility GetBreadcrumbTrailVisibility() const;
 
-	/** Gets whether or not the curve editor toolbar should be visible. */
-	EVisibility GetCurveEditorToolBarVisibility() const;
+
 
 	/** Gets whether or not the bottom time slider should be visible. */
 	EVisibility GetBottomTimeSliderVisibility() const;
@@ -445,6 +432,8 @@ private:
 	bool GetIsSequenceReadOnly() const;
 	void OnSetSequenceReadOnly(ECheckBoxState CheckBoxState);
 
+	/** Returns whether or not the Curve Editor is enabled. Allows us to bind to the Slate Enabled attribute. */
+	bool GetIsCurveEditorEnabled() const { return !GetIsSequenceReadOnly(); }
 public:
 	/** On Paste Command */
 	void OnPaste();
@@ -498,6 +487,9 @@ private:
 	/** The top time slider widget */
 	TSharedPtr<ITimeSlider> TopTimeSlider;
 
+	/** The curve editor panel. This is created and updated even if it is not currently visible. */
+	TSharedPtr<SWidget> CurveEditorPanel;
+
 	/** Cached settings provided to the sequencer itself on creation */
 	USequencerSettings* Settings;
 
@@ -554,4 +546,7 @@ private:
 	TArray<FString> AdditionalSelectionsToAdd;
 
 	TSharedPtr<SWidget> TickResolutionOverlay;
+
+public:
+	static const FName CurveEditorTabName;
 };
diff --git a/Engine/Source/Editor/Sequencer/Private/SSequencerSection.cpp b/Engine/Source/Editor/Sequencer/Private/SSequencerSection.cpp
--- a/Engine/Source/Editor/Sequencer/Private/SSequencerSection.cpp
+++ b/Engine/Source/Editor/Sequencer/Private/SSequencerSection.cpp
@@ -361,7 +361,7 @@ struct FSequencerSectionPainterImpl : FSequencerSectionPainter
 		float EaseInScale = 0.f, EaseOutScale = 0.f;
 		if (bEaseInHandle || bEaseOutHandle)
 		{
-			if (static_cast<const FSectionEasingHandleHotspot*>(Hotspot)->Section == Handle)
+			if (static_cast<const FSectionEasingHandleHotspot*>(Hotspot)->WeakSection.Get() == Handle.GetSectionObject())
 			{
 				if (bEaseInHandle)
 				{
@@ -377,7 +377,7 @@ struct FSequencerSectionPainterImpl : FSequencerSectionPainter
 		{
 			for (const FEasingAreaHandle& Easing : static_cast<const FSectionEasingAreaHotspot*>(Hotspot)->Easings)
 			{
-				if (Easing.Section == Handle)
+				if (Easing.WeakSection.Get() == Handle.GetSectionObject())
 				{
 					if (Easing.EasingType == ESequencerEasingType::In)
 					{
@@ -865,7 +865,7 @@ bool SSequencerSection::CheckForEasingHandleInteraction( const FPointerEvent& Mo
 
 	// Gather all underlapping sections
 	TArray<FSectionHandle> AllUnderlappingSections;
-	AllUnderlappingSections.Add(FSectionHandle(ParentSectionArea, SectionIndex));
+	AllUnderlappingSections.Add(FSectionHandle(ParentSectionArea.ToSharedRef(), SectionIndex));
 	for (const FSequencerOverlapRange& Segment : UnderlappingSegments)
 	{
 		for (FSectionHandle Section : Segment.Sections)
@@ -876,7 +876,7 @@ bool SSequencerSection::CheckForEasingHandleInteraction( const FPointerEvent& Mo
 
 	for (FSectionHandle Handle : AllUnderlappingSections)
 	{
-		TSharedRef<ISequencerSection> EasingSection    =  Handle.TrackNode->GetSections()[Handle.SectionIndex];
+		TSharedRef<ISequencerSection> EasingSection    =  Handle.GetSectionInterface();
 		UMovieSceneSection*           EasingSectionObj = EasingSection->GetSectionObject();
 
 		if (EasingSectionObj->HasStartFrame())
@@ -886,7 +886,7 @@ bool SSequencerSection::CheckForEasingHandleInteraction( const FPointerEvent& Mo
 
 			if (FMath::IsNearlyEqual(MouseTime, HandlePositionIn, HalfHandleSizeX))
 			{
-				GetSequencer().SetHotspot(MakeShared<FSectionEasingHandleHotspot>(ESequencerEasingType::In, Handle));
+				GetSequencer().SetHotspot(MakeShared<FSectionEasingHandleHotspot>(ESequencerEasingType::In, EasingSectionObj));
 				return true;
 			}
 		}
@@ -898,7 +898,7 @@ bool SSequencerSection::CheckForEasingHandleInteraction( const FPointerEvent& Mo
 
 			if (FMath::IsNearlyEqual(MouseTime, HandlePositionOut, HalfHandleSizeX))
 			{
-				GetSequencer().SetHotspot(MakeShared<FSectionEasingHandleHotspot>(ESequencerEasingType::Out, Handle));
+				GetSequencer().SetHotspot(MakeShared<FSectionEasingHandleHotspot>(ESequencerEasingType::Out, EasingSectionObj));
 				return true;
 			}
 		}
@@ -917,7 +917,7 @@ bool SSequencerSection::CheckForEdgeInteraction( const FPointerEvent& MouseEvent
 	}
 
 	TArray<FSectionHandle> AllUnderlappingSections;
-	AllUnderlappingSections.Add(FSectionHandle(ParentSectionArea, SectionIndex));
+	AllUnderlappingSections.Add(FSectionHandle(ParentSectionArea.ToSharedRef(), SectionIndex));
 	for (const FSequencerOverlapRange& Segment : UnderlappingSegments)
 	{
 		for (FSectionHandle Section : Segment.Sections)
@@ -931,7 +931,7 @@ bool SSequencerSection::CheckForEdgeInteraction( const FPointerEvent& MouseEvent
 
 	for (FSectionHandle Handle : AllUnderlappingSections)
 	{
-		TSharedRef<ISequencerSection> UnderlappingSection =  Handle.TrackNode->GetSections()[Handle.SectionIndex];
+		TSharedRef<ISequencerSection> UnderlappingSection =  Handle.GetSectionInterface();
 		UMovieSceneSection* UnderlappingSectionObj = UnderlappingSection->GetSectionObject();
 		if (!UnderlappingSection->SectionIsResizable())
 		{
@@ -951,7 +951,7 @@ bool SSequencerSection::CheckForEdgeInteraction( const FPointerEvent& MouseEvent
 
 			if( SectionRectLeft.IsUnderLocation( MouseEvent.GetScreenSpacePosition() ) )
 			{
-				GetSequencer().SetHotspot(MakeShareable( new FSectionResizeHotspot(FSectionResizeHotspot::Left, Handle)) );
+				GetSequencer().SetHotspot(MakeShareable( new FSectionResizeHotspot(FSectionResizeHotspot::Left, UnderlappingSectionObj)) );
 				return true;
 			}
 		}
@@ -965,7 +965,7 @@ bool SSequencerSection::CheckForEdgeInteraction( const FPointerEvent& MouseEvent
 
 			if( SectionRectRight.IsUnderLocation( MouseEvent.GetScreenSpacePosition() ) )
 			{
-				GetSequencer().SetHotspot(MakeShareable( new FSectionResizeHotspot(FSectionResizeHotspot::Right, Handle)) );
+				GetSequencer().SetHotspot(MakeShareable( new FSectionResizeHotspot(FSectionResizeHotspot::Right, UnderlappingSectionObj)) );
 				return true;
 			}
 		}
@@ -998,17 +998,17 @@ bool SSequencerSection::CheckForEasingAreaInteraction( const FPointerEvent& Mous
 			UMovieSceneSection* Section = Handle.GetSectionObject();
 			if (Section->GetEaseInRange().Contains(MouseTime))
 			{
-				EasingAreas.Add(FEasingAreaHandle{ Handle, ESequencerEasingType::In });
+				EasingAreas.Add(FEasingAreaHandle{ Section, ESequencerEasingType::In });
 			}
 			if (Section->GetEaseOutRange().Contains(MouseTime))
 			{
-				EasingAreas.Add(FEasingAreaHandle{ Handle, ESequencerEasingType::Out });
+				EasingAreas.Add(FEasingAreaHandle{ Section, ESequencerEasingType::Out });
 			}
 		}
 
 		if (EasingAreas.Num())
 		{
-			GetSequencer().SetHotspot(MakeShared<FSectionEasingAreaHotspot>(EasingAreas, FSectionHandle(ParentSectionArea, SectionIndex)));
+			GetSequencer().SetHotspot(MakeShared<FSectionEasingAreaHotspot>(EasingAreas, ThisSection));
 			return true;
 		}
 	}
@@ -1023,7 +1023,7 @@ FSequencer& SSequencerSection::GetSequencer() const
 
 int32 SSequencerSection::OnPaint( const FPaintArgs& Args, const FGeometry& AllottedGeometry, const FSlateRect& MyCullingRect, FSlateWindowElementList& OutDrawElements, int32 LayerId, const FWidgetStyle& InWidgetStyle, bool bParentEnabled ) const
 {
-	UMovieSceneSection& SectionObject = *SectionInterface->GetSectionObject();
+	UMovieSceneSection* SectionObject = SectionInterface->GetSectionObject();
 
 	const ISequencerEditTool* EditTool = GetSequencer().GetEditTool();
 	const ISequencerHotspot* Hotspot = EditTool ? EditTool->GetDragHotspot() : nullptr;
@@ -1032,13 +1032,13 @@ int32 SSequencerSection::OnPaint( const FPaintArgs& Args, const FGeometry& Allot
 		Hotspot = GetSequencer().GetHotspot().Get();
 	}
 
-	const bool bEnabled = bParentEnabled && SectionObject.IsActive();
-	const bool bLocked = SectionObject.IsLocked();
-	UMovieScenePropertyTrack* Track = SectionObject.GetTypedOuter<UMovieScenePropertyTrack>();
+	const bool bEnabled = bParentEnabled && SectionObject->IsActive();
+	const bool bLocked = SectionObject->IsLocked();
+	UMovieScenePropertyTrack* Track = SectionObject->GetTypedOuter<UMovieScenePropertyTrack>();
 	bool bSetSectionToKey = false;
 	if (Track)
 	{
-		if (Track->GetSectionToKey() == &SectionObject)
+		if (Track->GetSectionToKey() == SectionObject)
 		{
 			bSetSectionToKey = true;
 		}
@@ -1049,7 +1049,7 @@ int32 SSequencerSection::OnPaint( const FPaintArgs& Args, const FGeometry& Allot
 
 	FGeometry SectionGeometry = MakeSectionGeometryWithoutHandles( AllottedGeometry, SectionInterface );
 
-	FSequencerSectionPainterImpl Painter(ParentSectionArea->GetSequencer(), SectionObject, OutDrawElements, SectionGeometry, *this);
+	FSequencerSectionPainterImpl Painter(ParentSectionArea->GetSequencer(), *SectionObject, OutDrawElements, SectionGeometry, *this);
 
 	FGeometry PaintSpaceParentGeometry = ParentGeometry;
 	PaintSpaceParentGeometry.AppendTransform(FSlateLayoutTransform(Inverse(Args.GetWindowToDesktopTransform())));
@@ -1064,9 +1064,9 @@ int32 SSequencerSection::OnPaint( const FPaintArgs& Args, const FGeometry& Allot
 
 	Painter.LayerId = LayerId;
 	Painter.bParentEnabled = bEnabled;
-	Painter.bIsHighlighted = IsSectionHighlighted(FSectionHandle(ParentSectionArea, SectionIndex), Hotspot);
+	Painter.bIsHighlighted = IsSectionHighlighted(SectionObject, Hotspot);
 	auto& Selection = ParentSectionArea->GetSequencer().GetSelection();
-	Painter.bIsSelected = Selection.IsSelected(&SectionObject);
+	Painter.bIsSelected = Selection.IsSelected(SectionObject);
 
 	FSlateClippingZone ClippingZone(Painter.SectionClippingRect);
 	OutDrawElements.PushClip(ClippingZone);
@@ -1115,7 +1115,7 @@ int32 SSequencerSection::OnPaint( const FPaintArgs& Args, const FGeometry& Allot
 	FText SectionTitle = SectionInterface->GetSectionTitle();
 	FMargin ContentPadding = SectionInterface->GetContentPadding();
 
-	const int32 EaseInAmount = SectionObject.Easing.GetEaseInDuration();
+	const int32 EaseInAmount = SectionObject->Easing.GetEaseInDuration();
 	if (EaseInAmount > 0)
 	{
 		ContentPadding.Left += Painter.GetTimeConverter().FrameToPixel(EaseInAmount) - Painter.GetTimeConverter().FrameToPixel(0);
@@ -1594,16 +1594,16 @@ void SSequencerSection::PaintEasingHandles( FSequencerSectionPainter& InPainter,
 	}
 
 	TArray<FSectionHandle> AllUnderlappingSections;
-	if (IsSectionHighlighted(FSectionHandle(ParentSectionArea, SectionIndex), Hotspot))
+	if (IsSectionHighlighted(SectionInterface->GetSectionObject(), Hotspot))
 	{
-		AllUnderlappingSections.Add(FSectionHandle(ParentSectionArea, SectionIndex));
+		AllUnderlappingSections.Add(FSectionHandle(ParentSectionArea.ToSharedRef(), SectionIndex));
 	}
 
 	for (const FSequencerOverlapRange& Segment : UnderlappingSegments)
 	{
 		for (FSectionHandle Section : Segment.Sections)
 		{
-			if (IsSectionHighlighted(Section, Hotspot) && !AllUnderlappingSections.Contains(Section))
+			if (IsSectionHighlighted(Section.GetSectionObject(), Hotspot) && !AllUnderlappingSections.Contains(Section))
 			{
 				AllUnderlappingSections.Add(Section);
 			}
@@ -1630,7 +1630,7 @@ void SSequencerSection::PaintEasingHandles( FSequencerSectionPainter& InPainter,
 			{
 				const FSectionEasingHandleHotspot* EasingHotspot = static_cast<const FSectionEasingHandleHotspot*>(Hotspot);
 
-				bDrawThisSectionsHandles = EasingHotspot->Section == Handle;
+				bDrawThisSectionsHandles = (EasingHotspot->WeakSection == Handle.GetSectionObject());
 				bLeftHandleActive = Hotspot->GetType() == ESequencerHotspot::EaseInHandle;
 				bRightHandleActive = Hotspot->GetType() == ESequencerHotspot::EaseOutHandle;
 			}
@@ -1639,7 +1639,7 @@ void SSequencerSection::PaintEasingHandles( FSequencerSectionPainter& InPainter,
 				const FSectionEasingAreaHotspot* EasingAreaHotspot = static_cast<const FSectionEasingAreaHotspot*>(Hotspot);
 				for (const FEasingAreaHandle& Easing : EasingAreaHotspot->Easings)
 				{
-					if (Easing.Section == Handle)
+					if (Easing.WeakSection == Handle.GetSectionObject())
 					{
 						if (Easing.EasingType == ESequencerEasingType::In)
 						{
@@ -1730,7 +1730,7 @@ void SSequencerSection::DrawSectionHandles( const FGeometry& AllottedGeometry, F
 	OutDrawElements.PushClip(FSlateClippingZone(AllottedGeometry.GetLayoutBoundingRect()));
 
 	TArray<FSectionHandle> AllUnderlappingSections;
-	AllUnderlappingSections.Add(FSectionHandle(ParentSectionArea, SectionIndex));
+	AllUnderlappingSections.Add(FSectionHandle(ParentSectionArea.ToSharedRef(), SectionIndex));
 	for (const FSequencerOverlapRange& Segment : UnderlappingSegments)
 	{
 		for (FSectionHandle Section : Segment.Sections)
@@ -1744,14 +1744,14 @@ void SSequencerSection::DrawSectionHandles( const FGeometry& AllottedGeometry, F
 
 	for (FSectionHandle Handle : AllUnderlappingSections)
 	{
-		TSharedRef<ISequencerSection> UnderlappingSection =  Handle.TrackNode->GetSections()[Handle.SectionIndex];
+		TSharedRef<ISequencerSection> UnderlappingSection =  Handle.GetSectionInterface();
 		UMovieSceneSection* UnderlappingSectionObj = UnderlappingSection->GetSectionObject();
 		if (!UnderlappingSection->SectionIsResizable() || UnderlappingSectionObj->GetRange() == TRange<FFrameNumber>::All())
 		{
 			continue;
 		}
 
-		bool bDrawThisSectionsHandles = (UnderlappingSectionObj == ThisSection && HandleOffsetPx != 0) || IsSectionHighlighted(Handle, Hotspot);
+		bool bDrawThisSectionsHandles = (UnderlappingSectionObj == ThisSection && HandleOffsetPx != 0) || IsSectionHighlighted(UnderlappingSectionObj, Hotspot);
 		bool bLeftHandleActive = false;
 		bool bRightHandleActive = false;
 
@@ -1761,7 +1761,7 @@ void SSequencerSection::DrawSectionHandles( const FGeometry& AllottedGeometry, F
 			Hotspot->GetType() == ESequencerHotspot::SectionResize_R))
 		{
 			const FSectionResizeHotspot* ResizeHotspot = static_cast<const FSectionResizeHotspot*>(Hotspot);
-			if (ResizeHotspot->Section == Handle)
+			if (ResizeHotspot->WeakSection == Handle.GetSectionObject())
 			{
 				bDrawThisSectionsHandles = true;
 				bLeftHandleActive = Hotspot->GetType() == ESequencerHotspot::SectionResize_L;
@@ -2021,8 +2021,9 @@ FReply SSequencerSection::OnMouseMove( const FGeometry& MyGeometry, const FPoint
 		!CheckForEdgeInteraction(MouseEvent, MyGeometry) &&
 		!CheckForEasingAreaInteraction(MouseEvent, MyGeometry))
 	{
+		UMovieSceneSection* ThisSection = SectionInterface->GetSectionObject();
 		// If nothing was hit, we just hit the section
-		GetSequencer().SetHotspot( MakeShareable( new FSectionHotspot(FSectionHandle(ParentSectionArea, SectionIndex))) );
+		GetSequencer().SetHotspot( MakeShareable( new FSectionHotspot(ThisSection)) );
 	}
 
 	return FReply::Unhandled();
@@ -2088,7 +2089,7 @@ float SSequencerSection::GetKeySelectionThrobValue()
 	return 0.f;
 }
 
-bool SSequencerSection::IsSectionHighlighted(FSectionHandle InSectionHandle, const ISequencerHotspot* Hotspot)
+bool SSequencerSection::IsSectionHighlighted(UMovieSceneSection* InSection, const ISequencerHotspot* Hotspot)
 {
 	if (!Hotspot)
 	{
@@ -2098,21 +2099,21 @@ bool SSequencerSection::IsSectionHighlighted(FSectionHandle InSectionHandle, con
 	switch(Hotspot->GetType())
 	{
 	case ESequencerHotspot::Key:
-		return static_cast<const FKeyHotspot*>(Hotspot)->Keys.ContainsByPredicate([InSectionHandle](const FSequencerSelectedKey& Key){ return Key.Section == InSectionHandle.GetSectionObject(); });
+		return static_cast<const FKeyHotspot*>(Hotspot)->Keys.ContainsByPredicate([InSection](const FSequencerSelectedKey& Key){ return Key.Section == InSection; });
 
 	case ESequencerHotspot::Section:
-		return static_cast<const FSectionHotspot*>(Hotspot)->Section == InSectionHandle;
+		return static_cast<const FSectionHotspot*>(Hotspot)->WeakSection == InSection;
 
 	case ESequencerHotspot::SectionResize_L:
 	case ESequencerHotspot::SectionResize_R:
-		return static_cast<const FSectionResizeHotspot*>(Hotspot)->Section == InSectionHandle;
+		return static_cast<const FSectionResizeHotspot*>(Hotspot)->WeakSection == InSection;
 
 	case ESequencerHotspot::EaseInHandle:
 	case ESequencerHotspot::EaseOutHandle:
-		return static_cast<const FSectionEasingHandleHotspot*>(Hotspot)->Section == InSectionHandle;
+		return static_cast<const FSectionEasingHandleHotspot*>(Hotspot)->WeakSection == InSection;
 
 	case ESequencerHotspot::EasingArea:
-		return static_cast<const FSectionEasingAreaHotspot*>(Hotspot)->Contains(InSectionHandle);
+		return static_cast<const FSectionEasingAreaHotspot*>(Hotspot)->Contains(InSection);
 
 	default:
 		return false;
diff --git a/Engine/Source/Editor/Sequencer/Private/SSequencerSection.h b/Engine/Source/Editor/Sequencer/Private/SSequencerSection.h
--- a/Engine/Source/Editor/Sequencer/Private/SSequencerSection.h
+++ b/Engine/Source/Editor/Sequencer/Private/SSequencerSection.h
@@ -134,7 +134,7 @@ public:
 	/**
 	 * Check to see whether the specified section is highlighted
 	 */
-	static bool IsSectionHighlighted(FSectionHandle InSectionHandle, const ISequencerHotspot* Hotspot);
+	static bool IsSectionHighlighted(UMovieSceneSection* InSection, const ISequencerHotspot* Hotspot);
 
 private:
 	/** Interface to section data */
diff --git a/Engine/Source/Editor/Sequencer/Private/SSequencerSectionOverlay.cpp b/Engine/Source/Editor/Sequencer/Private/SSequencerSectionOverlay.cpp
--- a/Engine/Source/Editor/Sequencer/Private/SSequencerSectionOverlay.cpp
+++ b/Engine/Source/Editor/Sequencer/Private/SSequencerSectionOverlay.cpp
@@ -4,7 +4,7 @@
 
 int32 SSequencerSectionOverlay::OnPaint( const FPaintArgs& Args, const FGeometry& AllottedGeometry, const FSlateRect& MyCullingRect, FSlateWindowElementList& OutDrawElements, int32 LayerId, const FWidgetStyle& InWidgetStyle, bool bParentEnabled ) const
 {
-	FPaintSectionAreaViewArgs PaintArgs;
+	FPaintViewAreaArgs PaintArgs;
 	PaintArgs.bDisplayTickLines = bDisplayTickLines.Get();
 	PaintArgs.bDisplayScrubPosition = bDisplayScrubPosition.Get();
 	PaintArgs.bDisplayMarkedFrames = bDisplayMarkedFrames.Get();
@@ -14,7 +14,7 @@ int32 SSequencerSectionOverlay::OnPaint( const FPaintArgs& Args, const FGeometry
 		PaintArgs.PlaybackRangeArgs = PaintPlaybackRangeArgs.Get();
 	}
 
-	TimeSliderController->OnPaintSectionView( AllottedGeometry, MyCullingRect, OutDrawElements, LayerId, ShouldBeEnabled( bParentEnabled ), PaintArgs );
+	TimeSliderController->OnPaintViewArea( AllottedGeometry, MyCullingRect, OutDrawElements, LayerId, ShouldBeEnabled( bParentEnabled ), PaintArgs );
 
 	return SCompoundWidget::OnPaint( Args, AllottedGeometry, MyCullingRect, OutDrawElements, LayerId, InWidgetStyle, bParentEnabled  );
 }
