This is a significant overhaul to the Curve Editor used by Sequencer which adds a plugin-based architecture and extensibility. New tools and toolbar buttons can be added to all usages of the curve edi

diff --git a/Engine/Source/Editor/Sequencer/Private/SSequencerTreeView.cpp b/Engine/Source/Editor/Sequencer/Private/SSequencerTreeView.cpp
--- a/Engine/Source/Editor/Sequencer/Private/SSequencerTreeView.cpp
+++ b/Engine/Source/Editor/Sequencer/Private/SSequencerTreeView.cpp
@@ -3,6 +3,8 @@
 #include "SSequencerTreeView.h"
 #include "SSequencerTrackLane.h"
 #include "EditorStyleSet.h"
+#include "Algo/BinarySearch.h"
+#include "Algo/Copy.h"
 #include "SequencerDisplayNodeDragDropOp.h"
 #include "Framework/Commands/GenericCommands.h"
 #include "Framework/MultiBox/MultiBoxBuilder.h"
@@ -10,40 +12,10 @@
 
 static FName TrackAreaName = "TrackArea";
 
-namespace Utils
-{
-	enum class ESearchState { Before, After, Found };
-
-	template<typename T, typename F>
-	const T* BinarySearch(const TArray<T>& InContainer, const F& InPredicate)
-	{
-		int32 Min = 0;
-		int32 Max = InContainer.Num();
-
-		ESearchState State = ESearchState::Before;
-
-		for ( ; Min != Max ; )
-		{
-			int32 SearchIndex = Min + (Max - Min) / 2;
-
-			auto& Item = InContainer[SearchIndex];
-			State = InPredicate(Item);
-			switch (State)
-			{
-				case ESearchState::Before:	Max = SearchIndex; break;
-				case ESearchState::After:	Min = SearchIndex + 1; break;
-				case ESearchState::Found: 	return &Item;
-			}
-		}
-
-		return nullptr;
-	}
-}
-
 SSequencerTreeViewRow::~SSequencerTreeViewRow()
 {
-	auto TreeView = StaticCastSharedPtr<SSequencerTreeView>(OwnerTablePtr.Pin());
-	auto PinnedNode = Node.Pin();
+	const TSharedPtr<SSequencerTreeView>& TreeView = StaticCastSharedPtr<SSequencerTreeView>(OwnerTablePtr.Pin());
+	TSharedPtr<FSequencerDisplayNode> PinnedNode = Node.Pin();
 	if (TreeView.IsValid() && PinnedNode.IsValid())
 	{
 		TreeView->OnChildRowRemoved(PinnedNode.ToSharedRef());
@@ -62,13 +34,26 @@ void SSequencerTreeViewRow::Construct(const FArguments& InArgs, const TSharedRef
 			.OnDragDetected(this, &SSequencerTreeViewRow::OnDragDetected)
 			.OnCanAcceptDrop(this, &SSequencerTreeViewRow::OnCanAcceptDrop)
 			.OnAcceptDrop(this, &SSequencerTreeViewRow::OnAcceptDrop)
-			.ShowSelection(bIsSelectable),
+			.ShowSelection(bIsSelectable)
+			.Padding(this, &SSequencerTreeViewRow::GetRowPadding),
 		OwnerTableView);
 }
 
+FMargin SSequencerTreeViewRow::GetRowPadding() const
+{
+	TSharedPtr<FSequencerDisplayNode> PinnedNode = Node.Pin();
+	TSharedPtr<FSequencerDisplayNode> ParentNode = PinnedNode ? PinnedNode->GetParentOrRoot() : nullptr;
+
+	if (ParentNode->GetType() == ESequencerNode::Root && ParentNode->GetChildNodes()[0] != PinnedNode)
+	{
+		return FMargin(0.f, 1.f, 0.f, 0.f);
+	}
+	return FMargin(0.f, 0.f, 0.f, 0.f);
+}
+
 TSharedRef<SWidget> SSequencerTreeViewRow::GenerateWidgetForColumn(const FName& ColumnId)
 {
-	auto PinnedNode = Node.Pin();
+	TSharedPtr<FSequencerDisplayNode> PinnedNode = Node.Pin();
 	if (PinnedNode.IsValid())
 	{
 		return OnGenerateWidgetForColumn.Execute(PinnedNode.ToSharedRef(), ColumnId, SharedThis(this));
@@ -184,6 +169,7 @@ void SSequencerTreeView::Construct(const FArguments& InArgs, const TSharedRef<FS
 		.AllowOverscroll(EAllowOverscroll::No)
 		.OnContextMenuOpening( this, &SSequencerTreeView::OnContextMenuOpening )
 		.OnSetExpansionRecursive(this, &SSequencerTreeView::SetItemExpansionRecursive)
+		.HighlightParentNodesForSelection(true)
 	);
 }
 
@@ -218,10 +204,6 @@ void SSequencerTreeView::Tick(const FGeometry& AllottedGeometry, const double In
 	if (SequencerNodeTree->GetHoveredNode().IsValid())
 	{
 		TSharedRef<FSequencerDisplayNode> OutermostParent = SequencerNodeTree->GetHoveredNode()->GetOutermostParent();
-		if (OutermostParent->GetType() == ESequencerNode::Spacer)
-		{
-			return;
-		}
 
 		TOptional<float> PhysicalTop = ComputeNodePosition(OutermostParent);
 
@@ -298,9 +280,9 @@ TOptional<float> SSequencerTreeView::ComputeNodePosition(const FDisplayNodeRef&
 	TOptional<float> Top;
 	
 	// Iterate parent first until we find a tree view row we can use for the offset height
-	auto Iter = [&](FSequencerDisplayNode& InDisplayNode){
-		
-		TOptional<FCachedGeometry> ChildRowGeometry = GetPhysicalGeometryForNode(InDisplayNode.AsShared());
+	auto Iter = [this, &NegativeOffset, &Top](FSequencerDisplayNode& InDisplayNode)
+	{
+		TOptional<FCachedGeometry> ChildRowGeometry = this->GetPhysicalGeometryForNode(InDisplayNode.AsShared());
 		if (ChildRowGeometry.IsSet())
 		{
 			Top = ChildRowGeometry->PhysicalTop;
@@ -344,58 +326,29 @@ void SSequencerTreeView::OnChildRowRemoved(const FDisplayNodeRef& InNode)
 
 TSharedPtr<FSequencerDisplayNode> SSequencerTreeView::HitTestNode(float InPhysical) const
 {
-	auto* Found = Utils::BinarySearch<FCachedGeometry>(PhysicalNodes, [&](const FCachedGeometry& In){
-
-		if (InPhysical < In.PhysicalTop)
-		{
-			return Utils::ESearchState::Before;
-		}
-		else if (InPhysical > In.PhysicalTop + In.PhysicalHeight)
-		{
-			return Utils::ESearchState::After;
-		}
-
-		return Utils::ESearchState::Found;
-
-	});
-
-	if (Found)
+	// Find the first node with a top after the specified value - the hit node must be the one preceeding this
+	const int32 FoundIndex = Algo::UpperBoundBy(PhysicalNodes, InPhysical, &FCachedGeometry::PhysicalTop) - 1;
+	if (FoundIndex >= 0)
 	{
-		return Found->Node;
+		return PhysicalNodes[FoundIndex].Node;
 	}
-	
 	return nullptr;
 }
 
 float SSequencerTreeView::PhysicalToVirtual(float InPhysical) const
 {
-	int32 SearchIndex = PhysicalNodes.Num() / 2;
-
-	auto* Found = Utils::BinarySearch<FCachedGeometry>(PhysicalNodes, [&](const FCachedGeometry& In){
-
-		if (InPhysical < In.PhysicalTop)
-		{
-			return Utils::ESearchState::Before;
-		}
-		else if (InPhysical > In.PhysicalTop + In.PhysicalHeight)
-		{
-			return Utils::ESearchState::After;
-		}
-
-		return Utils::ESearchState::Found;
-
-	});
-
-
-	if (Found)
+	// Find the first node with a top after the specified value - the hit node must be the one preceeding this
+	const int32 FoundIndex = Algo::UpperBoundBy(PhysicalNodes, InPhysical, &FCachedGeometry::PhysicalTop) - 1;
+	if (FoundIndex >= 0)
 	{
-		const float FractionalHeight = (InPhysical - Found->PhysicalTop) / Found->PhysicalHeight;
-		return Found->Node->GetVirtualTop() + (Found->Node->GetVirtualBottom() - Found->Node->GetVirtualTop()) * FractionalHeight;
+		const FCachedGeometry& Found = PhysicalNodes[FoundIndex];
+		const float FractionalHeight = (InPhysical - Found.PhysicalTop) / Found.PhysicalHeight;
+		return Found.Node->GetVirtualTop() + (Found.Node->GetVirtualBottom() - Found.Node->GetVirtualTop()) * FractionalHeight;
 	}
 
 	if (PhysicalNodes.Num())
 	{
-		auto& Last = PhysicalNodes.Last();
+		const FCachedGeometry& Last = PhysicalNodes.Last();
 		return Last.Node->GetVirtualTop() + (InPhysical - Last.PhysicalTop);
 	}
 
@@ -404,30 +357,23 @@ float SSequencerTreeView::PhysicalToVirtual(float InPhysical) const
 
 float SSequencerTreeView::VirtualToPhysical(float InVirtual) const
 {
-	auto* Found = Utils::BinarySearch(PhysicalNodes, [&](const FCachedGeometry& In){
-
-		if (InVirtual < In.Node->GetVirtualTop())
-		{
-			return Utils::ESearchState::Before;
-		}
-		else if (InVirtual > In.Node->GetVirtualBottom())
-		{
-			return Utils::ESearchState::After;
-		}
-
-		return Utils::ESearchState::Found;
-
-	});
-
-	if (Found)
+	auto GetVirtualTop = [](const FCachedGeometry& In)
+	{
+		return In.Node->GetVirtualTop();
+	};
+	// Find the first node with a top after the specified value - the hit node must be the one preceeding this
+	const int32 FoundIndex = Algo::UpperBoundBy(PhysicalNodes, InVirtual, GetVirtualTop) - 1;
+	if (FoundIndex >= 0)
 	{
-		const float FractionalHeight = (InVirtual - Found->Node->GetVirtualTop()) / (Found->Node->GetVirtualBottom() - Found->Node->GetVirtualTop());
-		return Found->PhysicalTop + Found->PhysicalHeight * FractionalHeight;
+		const FCachedGeometry& Found = PhysicalNodes[FoundIndex];
+
+		const float FractionalHeight = (InVirtual - Found.Node->GetVirtualTop()) / (Found.Node->GetVirtualBottom() - Found.Node->GetVirtualTop());
+		return Found.PhysicalTop + Found.PhysicalHeight * FractionalHeight;
 	}
-	
+
 	if (PhysicalNodes.Num())
 	{
-		auto Last = PhysicalNodes.Last();
+		const FCachedGeometry& Last = PhysicalNodes.Last();
 		return Last.PhysicalTop + (InVirtual - Last.Node->GetVirtualTop());
 	}
 
@@ -439,7 +385,7 @@ void SSequencerTreeView::SetupColumns(const FArguments& InArgs)
 	FSequencer& Sequencer = SequencerNodeTree->GetSequencer();
 
 	// Define a column for the Outliner
-	auto GenerateOutliner = [=](const FDisplayNodeRef& InNode, const TSharedRef<SSequencerTreeViewRow>& InRow)
+	auto GenerateOutliner = [](const FDisplayNodeRef& InNode, const TSharedRef<SSequencerTreeViewRow>& InRow)
 	{
 		return InNode->GenerateContainerWidgetForOutliner(InRow);
 	};
@@ -447,9 +393,9 @@ void SSequencerTreeView::SetupColumns(const FArguments& InArgs)
 	Columns.Add("Outliner", FSequencerTreeViewColumn(GenerateOutliner, 1.f));
 
 	// Now populate the header row with the columns
-	for (auto& Pair : Columns)
+	for (TTuple<FName, FSequencerTreeViewColumn>& Pair : Columns)
 	{
-		if (Pair.Key != TrackAreaName || !Sequencer.GetShowCurveEditor())
+		if (Pair.Key != TrackAreaName)
 		{
 			HeaderRow->AddColumn(
 				SHeaderRow::Column(Pair.Key)
@@ -464,11 +410,7 @@ void SSequencerTreeView::UpdateTrackArea()
 	FSequencer& Sequencer = SequencerNodeTree->GetSequencer();
 
 	// Add or remove the column
-	if (Sequencer.GetShowCurveEditor())
-	{
-		HeaderRow->RemoveColumn(TrackAreaName);
-	}
-	else if (const auto* Column = Columns.Find(TrackAreaName))
+	if (const FSequencerTreeViewColumn* Column = Columns.Find(TrackAreaName))
 	{
 		HeaderRow->AddColumn(
 			SHeaderRow::Column(TrackAreaName)
@@ -499,7 +441,7 @@ void SSequencerTreeView::SynchronizeTreeSelectionWithSequencerSelection()
 			Private_ClearSelection();
 
 			FSequencer& Sequencer = SequencerNodeTree->GetSequencer();
-			for ( auto& Node : Sequencer.GetSelection().GetSelectedOutlinerNodes() )
+			for ( const TSharedRef<FSequencerDisplayNode>& Node : Sequencer.GetSelection().GetSelectedOutlinerNodes() )
 			{
 				if (Node->IsSelectable())
 				{
@@ -597,7 +539,7 @@ bool SSequencerTreeView::SynchronizeSequencerSelectionWithTreeSelection()
 		FSequencer& Sequencer = SequencerNodeTree->GetSequencer();
 		FSequencerSelection& Selection = Sequencer.GetSelection();
 		Selection.EmptySelectedOutlinerNodes();
-		for ( auto& Item : GetSelectedItems() )
+		for ( const TSharedRef<FSequencerDisplayNode>& Item : GetSelectedItems() )
 		{
 			Selection.AddToSelection( Item );
 		}
@@ -608,11 +550,14 @@ bool SSequencerTreeView::SynchronizeSequencerSelectionWithTreeSelection()
 
 TSharedPtr<SWidget> SSequencerTreeView::OnContextMenuOpening()
 {
-	const TSet<TSharedRef<FSequencerDisplayNode>> SelectedNodes = SequencerNodeTree->GetSequencer().GetSelection().GetSelectedOutlinerNodes();
-	auto SelectedNodesArray = SelectedNodes.Array();
-	if (SelectedNodes.Num() > 0 && SelectedNodesArray[0]->IsSelectable())
+	// Open a context menu for the first selected item if it is selectable
+	for (TSharedRef<FSequencerDisplayNode> SelectedNode : SequencerNodeTree->GetSequencer().GetSelection().GetSelectedOutlinerNodes())
 	{
-		return SelectedNodesArray[0]->OnSummonContextMenu();
+		if (SelectedNode->IsSelectable())
+		{
+			return SelectedNode->OnSummonContextMenu();
+		}
+		break;
 	}
 
 	// Otherwise, add a general menu for options
@@ -630,19 +575,23 @@ TSharedPtr<SWidget> SSequencerTreeView::OnContextMenuOpening()
 
 void SSequencerTreeView::Refresh()
 {
-	RootNodes.Reset(SequencerNodeTree->GetRootNodes().Num());
+	RootNodes.Reset();
+	Algo::CopyIf(SequencerNodeTree->GetRootNodes(), RootNodes, &FSequencerDisplayNode::IsVisible);
 
-	for (const auto& RootNode : SequencerNodeTree->GetRootNodes())
+	// Reset item expansion since we don't know if any expansion states may have changed in-between refreshes
 	{
-		if (RootNode->IsExpanded())
-		{
-			SetItemExpansion(RootNode, true);
-		}
+		STreeView::OnExpansionChanged.Unbind();
 
-		if (!RootNode->IsHidden())
+		ClearExpandedItems();
+		auto Traverse_SetExpansionStates = [this](FSequencerDisplayNode& InNode)
 		{
-			RootNodes.Add(RootNode);
-		}
+			this->SetItemExpansion(InNode.AsShared(), InNode.IsExpanded());
+			return true;
+		};
+		const bool bIncludeRootNode = false;
+		SequencerNodeTree->GetRootNode()->Traverse_ParentFirst(Traverse_SetExpansionStates, bIncludeRootNode);
+
+		STreeView::OnExpansionChanged.BindSP(this, &SSequencerTreeView::OnExpansionChanged);
 	}
 
 	// Force synchronization of selected tree view items here since the tree nodes may have been rebuilt
@@ -651,7 +600,7 @@ void SSequencerTreeView::Refresh()
 	SynchronizeTreeSelectionWithSequencerSelection();
 	bUpdatingTreeSelection = false;
 
-	RequestTreeRefresh();
+	RebuildList();
 }
 
 void SSequencerTreeView::ScrollByDelta(float DeltaInSlateUnits)
@@ -746,9 +695,9 @@ void SSequencerTreeView::ExpandCollapseNode(const FDisplayNodeRef& InNode, bool
 void SSequencerTreeView::OnExpansionChanged(FDisplayNodeRef InItem, bool bIsExpanded)
 {
 	InItem->SetExpansionState(bIsExpanded);
-	
+
 	// Expand any children that are also expanded
-	for (auto& Child : InItem->GetChildNodes())
+	for (const FDisplayNodeRef& Child : InItem->GetChildNodes())
 	{
 		if (Child->IsExpanded())
 		{
diff --git a/Engine/Source/Editor/Sequencer/Private/SSequencerTreeView.h b/Engine/Source/Editor/Sequencer/Private/SSequencerTreeView.h
--- a/Engine/Source/Editor/Sequencer/Private/SSequencerTreeView.h
+++ b/Engine/Source/Editor/Sequencer/Private/SSequencerTreeView.h
@@ -269,6 +269,9 @@ public:
 	/** Called to complete a drag and drop onto this drop. */
 	FReply OnAcceptDrop( const FDragDropEvent& DragDropEvent, EItemDropZone ItemDropZone, FDisplayNodeRef DisplayNode );
 
+	/** Gets the padding for this row based on whether it is a root node or not */
+	FMargin GetRowPadding() const;
+
 private:
 
 	/** Cached reference to a track lane that we relate to. This keeps the track lane alive (it's a weak widget) as long as we are in view. */
diff --git a/Engine/Source/Editor/Sequencer/Private/SectionHandle.cpp b/Engine/Source/Editor/Sequencer/Private/SectionHandle.cpp
new file mode 100644
--- /dev/null
+++ b/Engine/Source/Editor/Sequencer/Private/SectionHandle.cpp
@@ -0,0 +1,15 @@
+// Copyright 1998-2019 Epic Games, Inc. All Rights Reserved.
+
+#include "SectionHandle.h"
+#include "ISequencerSection.h"
+#include "DisplayNodes/SequencerTrackNode.h"
+
+TSharedRef<ISequencerSection> FSectionHandle::GetSectionInterface() const
+{
+	return TrackNode->GetSections()[SectionIndex];
+}
+
+UMovieSceneSection* FSectionHandle::GetSectionObject() const
+{
+	return GetSectionInterface()->GetSectionObject();
+}
\ No newline at end of file
diff --git a/Engine/Source/Editor/Sequencer/Private/SectionHandle.h b/Engine/Source/Editor/Sequencer/Private/SectionHandle.h
new file mode 100644
--- /dev/null
+++ b/Engine/Source/Editor/Sequencer/Private/SectionHandle.h
@@ -0,0 +1,63 @@
+// Copyright 1998-2019 Epic Games, Inc. All Rights Reserved.
+
+#pragma once
+
+#include "Templates/SharedPointer.h"
+
+class ISequencerSection;
+class UMovieSceneSection;
+class FSequencerTrackNode;
+
+/**
+ * Structure used to encapsulate a specific section on a track node.
+ * Authoritative mapping is stored in FSequencerNodeTree. FSectionHandle
+ * should not be held persistently except in contexts that are forcibly destroyed when the tree is refreshed.
+ */
+struct FSectionHandle
+{
+	/**
+	 * Construction from a track node, and the index of the section within it
+	 */
+	FSectionHandle(TSharedRef<FSequencerTrackNode> InTrackNode, int32 InSectionIndex)
+		: SectionIndex(InSectionIndex), TrackNode(MoveTemp(InTrackNode))
+	{ }
+
+	/**
+	 * Compare 2 handles for equality
+	 */
+	friend bool operator==(const FSectionHandle& A, const FSectionHandle& B)
+	{
+		return A.SectionIndex == B.SectionIndex && A.TrackNode == B.TrackNode;
+	}
+
+	/**
+	 * Retrieve the Sequencer section interface implementation from this handle
+	 */
+	TSharedRef<ISequencerSection> GetSectionInterface() const;
+
+	/**
+	 * Retrieve the actual section that this handle represents
+	 */
+	UMovieSceneSection* GetSectionObject() const;
+
+	/**
+	 * Access the track node that this section currently lives within (could be a sub track node)
+	 */
+	TSharedRef<FSequencerTrackNode> GetTrackNode() const
+	{
+		return TrackNode;
+	}
+
+	/**
+	 * Get the index of this section within its FSequencerTrackNode::GetSections array
+	 */
+	int32 GetSectionIndex() const
+	{
+		return SectionIndex;
+	}
+
+private:
+
+	int32 SectionIndex;
+	TSharedRef<FSequencerTrackNode> TrackNode;
+};
\ No newline at end of file
diff --git a/Engine/Source/Editor/Sequencer/Private/Sequencer.cpp b/Engine/Source/Editor/Sequencer/Private/Sequencer.cpp
--- a/Engine/Source/Editor/Sequencer/Private/Sequencer.cpp
+++ b/Engine/Source/Editor/Sequencer/Private/Sequencer.cpp
@@ -183,48 +183,59 @@ struct FSequencerTemplateStore : IMovieSceneSequenceTemplateStore
 struct FSequencerCurveEditorBounds : ICurveEditorBounds
 {
 	FSequencerCurveEditorBounds(TSharedRef<FSequencer> InSequencer)
-		: OutputMin(0), OutputMax(1), WeakSequencer(InSequencer)
-	{}
+		: WeakSequencer(InSequencer)
+	{
+		TRange<double> Bounds = InSequencer->GetViewRange();
+		InputMin = Bounds.GetLowerBoundValue();
+		InputMax = Bounds.GetUpperBoundValue();
+	}
 
 	virtual void GetInputBounds(double& OutMin, double& OutMax) const override
 	{
 		TSharedPtr<FSequencer> Sequencer = WeakSequencer.Pin();
 		if (Sequencer.IsValid())
 		{
-			TRange<double> Bounds = Sequencer->GetViewRange();
-			OutMin = Bounds.GetLowerBoundValue();
-			OutMax = Bounds.GetUpperBoundValue();
+			const bool bLinkTimeRange = Sequencer->GetSequencerSettings()->GetLinkCurveEditorTimeRange();
+			if (bLinkTimeRange)
+			{
+				TRange<double> Bounds = Sequencer->GetViewRange();
+				OutMin = Bounds.GetLowerBoundValue();
+				OutMax = Bounds.GetUpperBoundValue();
+			}
+			else
+			{
+				// If they don't want to link the time range with Sequencer we return the cached value.
+				OutMin = InputMin;
+				OutMax = InputMax;
+			}
 		}
 	}
-	virtual void GetOutputBounds(double& OutMin, double& OutMax) const override
-	{
-		OutMin = OutputMin;
-		OutMax = OutputMax;
-	}
 
 	virtual void SetInputBounds(double InMin, double InMax) override
 	{
 		TSharedPtr<FSequencer> Sequencer = WeakSequencer.Pin();
 		if (Sequencer.IsValid())
 		{
-			FFrameRate TickResolution = Sequencer->GetFocusedTickResolution();
-
-			if (InMin * TickResolution > TNumericLimits<int32>::Lowest() && InMax * TickResolution < TNumericLimits<int32>::Max())
+			const bool bLinkTimeRange = Sequencer->GetSequencerSettings()->GetLinkCurveEditorTimeRange();
+			if (bLinkTimeRange)
 			{
-				Sequencer->SetViewRange(TRange<double>(InMin, InMax), EViewRangeInterpolation::Immediate);
+				FFrameRate TickResolution = Sequencer->GetFocusedTickResolution();
+
+				if (InMin * TickResolution > TNumericLimits<int32>::Lowest() && InMax * TickResolution < TNumericLimits<int32>::Max())
+				{
+					Sequencer->SetViewRange(TRange<double>(InMin, InMax), EViewRangeInterpolation::Immediate);
+				}
 			}
-		}
-	}
-	virtual void SetOutputBounds(double InMin, double InMax) override
-	{
-		if (InMax - InMin < TNumericLimits<float>::Max())
-		{
-			OutputMin = InMin;
-			OutputMax = InMax;
+			
+			// We update these even if you are linked to the Sequencer Timeline so that when you turn off the link setting
+			// you don't pop to your last values, instead your view stays as is and just stops moving when Sequencer moves.
+			InputMin = InMin;
+			InputMax = InMax;
 		}
 	}
 
-	double OutputMin, OutputMax;
+	/** The min/max values for the viewing range. Only used if Curve Editor/Sequencer aren't linked ranges. */
+	double InputMin, InputMax;
 	TWeakPtr<FSequencer> WeakSequencer;
 };
 
@@ -240,25 +251,23 @@ public:
 	virtual void GetGridLinesX(TArray<float>& MajorGridLines, TArray<float>& MinorGridLines, TArray<FText>& MajorGridLabels) const override
 	{
 		TSharedPtr<FSequencer> Sequencer = WeakSequencer.Pin();
-
-		FCurveEditorScreenSpace ScreenSpace = GetScreenSpace();
+		FCurveEditorScreenSpaceH PanelInputSpace = GetPanelInputSpace();
 
 		double MajorGridStep  = 0.0;
 		int32  MinorDivisions = 0;
 
-		if (Sequencer.IsValid() && Sequencer->GetGridMetrics(ScreenSpace.GetPhysicalWidth(), MajorGridStep, MinorDivisions))
+		if (Sequencer.IsValid() && Sequencer->GetGridMetrics(PanelInputSpace.GetPhysicalWidth(), PanelInputSpace.GetInputMin(), PanelInputSpace.GetInputMax(), MajorGridStep, MinorDivisions))
 		{
-			const double FirstMajorLine = FMath::FloorToDouble(ScreenSpace.GetInputMin() / MajorGridStep) * MajorGridStep;
-			const double LastMajorLine  = FMath::CeilToDouble(ScreenSpace.GetInputMax() / MajorGridStep) * MajorGridStep;
+			const double FirstMajorLine = FMath::FloorToDouble(PanelInputSpace.GetInputMin() / MajorGridStep) * MajorGridStep;
+			const double LastMajorLine  = FMath::CeilToDouble(PanelInputSpace.GetInputMax() / MajorGridStep) * MajorGridStep;
 
 			for (double CurrentMajorLine = FirstMajorLine; CurrentMajorLine < LastMajorLine; CurrentMajorLine += MajorGridStep)
 			{
-				MajorGridLines.Add( ScreenSpace.SecondsToScreen(CurrentMajorLine) );
-				MajorGridLabels.Add(FText());
-				 
+				MajorGridLines.Add( PanelInputSpace.SecondsToScreen(CurrentMajorLine) );
+ 
 				for (int32 Step = 1; Step < MinorDivisions; ++Step)
 				{
-					MinorGridLines.Add( ScreenSpace.SecondsToScreen(CurrentMajorLine + Step*MajorGridStep/MinorDivisions) );
+					MinorGridLines.Add( PanelInputSpace.SecondsToScreen(CurrentMajorLine + Step*MajorGridStep/MinorDivisions) );
 				}
 			}
 		}
@@ -302,11 +311,15 @@ void FSequencer::InitSequencer(const FSequencerInitParams& InitParams, const TSh
 
 	Settings->GetOnEvaluateSubSequencesInIsolationChanged().AddSP(this, &FSequencer::RestorePreAnimatedState);
 	Settings->GetOnShowSelectedNodesOnlyChanged().AddSP(this, &FSequencer::OnSelectedNodesOnlyChanged);
-	Settings->GetOnCurveEditorCurveVisibilityChanged().AddSP(this, &FSequencer::SyncCurveEditorToSelection, false);
 
+	FCurveEditorInitParams CurveEditorInitParams;
+	{
+	}
+	
 	{
 		CurveEditorModel = MakeShared<FSequencerCurveEditor>(SharedThis(this));
 		CurveEditorModel->SetBounds(MakeUnique<FSequencerCurveEditorBounds>(SharedThis(this)));
+		CurveEditorModel->InitCurveEditor(CurveEditorInitParams);
 
 		CurveEditorModel->InputSnapEnabledAttribute   = MakeAttributeLambda([this]{ return Settings->GetIsSnapEnabled(); });
 		CurveEditorModel->OnInputSnapEnabledChanged   = FOnSetBoolean::CreateLambda([this](bool NewValue){ Settings->SetIsSnapEnabled(NewValue); });
@@ -446,7 +459,6 @@ FSequencer::FSequencer()
 	, bPerspectiveViewportPossessionEnabled( true )
 	, bPerspectiveViewportCameraCutEnabled( false )
 	, bIsEditingWithinLevelEditor( false )
-	, bShowCurveEditor( false )
 	, bNeedTreeRefresh( false )
 	, StoredPlaybackState( EMovieScenePlayerStatus::Stopped )
 	, NodeTree( MakeShareable( new FSequencerNodeTree( *this ) ) )
@@ -3998,8 +4010,6 @@ FFrameNumber FSequencer::OnGetNearestKey(FFrameTime InTime, bool bSearchAllTrack
 
 void FSequencer::OnScrubPositionChanged( FFrameTime NewScrubPosition, bool bScrubbing )
 {
-	bool bClampToViewRange = true;
-
 	if (PlaybackState == EMovieScenePlayerStatus::Scrubbing)
 	{
 		if (!bScrubbing)
@@ -4008,9 +4018,6 @@ void FSequencer::OnScrubPositionChanged( FFrameTime NewScrubPosition, bool bScru
 		}
 		else if (IsAutoScrollEnabled())
 		{
-			// Clamp to the view range when not auto-scrolling
-			bClampToViewRange = false;
-	
 			UpdateAutoScroll(NewScrubPosition / GetFocusedTickResolution());
 			
 			// When scrubbing, we animate auto-scrolled scrub position in Tick()
@@ -4021,23 +4028,6 @@ void FSequencer::OnScrubPositionChanged( FFrameTime NewScrubPosition, bool bScru
 		}
 	}
 
-	if (bClampToViewRange)
-	{
-		FFrameRate DisplayRate    = GetFocusedDisplayRate();
-		FFrameRate TickResolution = GetFocusedTickResolution();
-
-		FFrameTime LowerBound = (TargetViewRange.GetLowerBoundValue() * TickResolution).CeilToFrame();
-		FFrameTime UpperBound = (TargetViewRange.GetUpperBoundValue() * TickResolution).FloorToFrame();
-
-		if (Settings->GetIsSnapEnabled() && Settings->GetSnapPlayTimeToInterval())
-		{
-			LowerBound = FFrameRate::Snap(LowerBound, TickResolution, DisplayRate);
-			UpperBound = FFrameRate::Snap(UpperBound, TickResolution, DisplayRate);
-		}
-
-		NewScrubPosition = FMath::Clamp(NewScrubPosition, LowerBound, UpperBound);
-	}
-
 	if (!bScrubbing && FSlateApplication::Get().GetModifierKeys().IsShiftDown())
 	{
 		AutoScrubToTime(NewScrubPosition);
@@ -4492,7 +4482,7 @@ void FSequencer::PostUndo(bool bSuccess)
 {
 	NotifyMovieSceneDataChanged( EMovieSceneDataChangeType::Unknown );
 	SynchronizeSequencerSelectionWithExternalSelection();
-	SyncCurveEditorToSelection(false);
+
 	OnActivateSequenceEvent.Broadcast(ActiveTemplateIDs.Top());
 }
 
@@ -4743,8 +4733,14 @@ void FSequencer::UpdatePreviewLevelViewportClientFromCameraCut(FLevelEditorViewp
 
 void FSequencer::SetShowCurveEditor(bool bInShowCurveEditor)
 {
-	bShowCurveEditor = bInShowCurveEditor; 
-	SequencerWidget->OnCurveEditorVisibilityChanged();
+	SequencerWidget->OnCurveEditorVisibilityChanged(bInShowCurveEditor);
+}
+
+bool FSequencer::GetCurveEditorIsVisible() const
+{
+	// We always want to retrieve this directly from the UI instead of mirroring it to a local bool as there are
+	// a lot of ways the UI could get out of sync with a local bool (such as previously restored tab layouts)
+	return GetToolkitHost()->GetTabManager()->FindExistingLiveTab(FTabId(SSequencer::CurveEditorTabName)).IsValid();
 }
 
 void FSequencer::SaveCurrentMovieScene()
@@ -4965,7 +4961,6 @@ TArray<FGuid> FSequencer::AddActors(const TArray<TWeakObjectPtr<AActor> >& InAct
 void FSequencer::OnSelectedOutlinerNodesChanged()
 {
 	SynchronizeExternalSelectionWithSequencerSelection();
-	SyncCurveEditorToSelection(true);
 
 	FSequencerEdMode* SequencerEdMode = (FSequencerEdMode*)(GLevelEditorModeTools().GetActiveMode(FSequencerEdMode::EM_SequencerMode));
 	if (SequencerEdMode != nullptr)
@@ -5122,7 +5117,8 @@ void GetRootObjectBindingNodes(const TArray<TSharedRef<FSequencerDisplayNode>>&
 
 void FSequencer::SynchronizeSequencerSelectionWithExternalSelection()
 {
-	if ( bUpdatingExternalSelection || !IsLevelEditorSequencer() || ExactCast<ULevelSequence>(GetFocusedMovieSceneSequence()) == nullptr)
+	UMovieSceneSequence* Sequence = GetFocusedMovieSceneSequence();
+	if ( bUpdatingExternalSelection || !IsLevelEditorSequencer() || ExactCast<ULevelSequence>(Sequence) == nullptr)
 	{
 		return;
 	}
@@ -5141,69 +5137,63 @@ void FSequencer::SynchronizeSequencerSelectionWithExternalSelection()
 	ActorSelection->GetSelectedObjects<ASequencerKeyActor>(SelectedSequencerKeyActors);
 
 	TSet<TSharedRef<FSequencerDisplayNode>> NodesToSelect;
-	for (auto ObjectBinding : NodeTree->GetObjectBindingMap() )
+	for (const FMovieSceneBinding& Binding : Sequence->GetMovieScene()->GetBindings())
 	{
-		if (!ObjectBinding.Value.IsValid())
+		TSharedPtr<FSequencerObjectBindingNode> NodePtr = NodeTree->FindObjectBindingNode(Binding.GetObjectGuid());
+		if (!NodePtr)
 		{
 			continue;
 		}
 
-		TSharedRef<FSequencerObjectBindingNode> ObjectBindingNode = ObjectBinding.Value.ToSharedRef();
-		for ( TWeakObjectPtr<UObject> RuntimeObjectPtr : FindBoundObjects(ObjectBindingNode->GetObjectBinding(), ActiveTemplateIDs.Top()) )
+		TSharedRef<FSequencerObjectBindingNode> ObjectBindingNode = NodePtr.ToSharedRef();
+		for ( TWeakObjectPtr<> WeakObject : FindBoundObjects(Binding.GetObjectGuid(), ActiveTemplateIDs.Top()) )
 		{
-			UObject* RuntimeObject = RuntimeObjectPtr.Get();
-			if ( RuntimeObject != nullptr)
+			UObject* RuntimeObject = WeakObject.Get();
+			if (RuntimeObject == nullptr)
 			{
-				for (ASequencerKeyActor* KeyActor : SelectedSequencerKeyActors)
+				continue;
+			}
+
+			for (ASequencerKeyActor* KeyActor : SelectedSequencerKeyActors)
+			{
+				if (KeyActor->IsEditorOnly())
 				{
-					if (KeyActor->IsEditorOnly())
+					AActor* TrailActor = KeyActor->GetAssociatedActor();
+					if (TrailActor != nullptr && RuntimeObject == TrailActor)
 					{
-						AActor* TrailActor = KeyActor->GetAssociatedActor();
-						if (TrailActor != nullptr && RuntimeObject == TrailActor)
-						{
-							NodesToSelect.Add(ObjectBindingNode);
-							bAllAlreadySelected = false;
-							break;
-						}
+						NodesToSelect.Add(ObjectBindingNode);
+						bAllAlreadySelected = false;
+						break;
 					}
 				}
+			}
 
-				bool bActorSelected = ActorSelection->IsSelected( RuntimeObject );
-				bool bComponentSelected = GEditor->GetSelectedComponents()->IsSelected( RuntimeObject);
+			const bool bActorSelected = ActorSelection->IsSelected( RuntimeObject );
+			const bool bComponentSelected = GEditor->GetSelectedComponents()->IsSelected( RuntimeObject);
 
-				if (bActorSelected || bComponentSelected)
-				{
-					NodesToSelect.Add( ObjectBindingNode );
+			if (bActorSelected || bComponentSelected)
+			{
+				NodesToSelect.Add( ObjectBindingNode );
 
-					if (bAllAlreadySelected)
+				if (bAllAlreadySelected && !Selection.IsSelected(ObjectBindingNode))
+				{
+					// Traversal callback will exit prematurely if there are any selected children
+					auto Traverse_IsSelected = [this](FSequencerDisplayNode& InNode)
 					{
-						bool bAlreadySelected = Selection.IsSelected(ObjectBindingNode);
+						TSharedRef<FSequencerDisplayNode> SharedNode = InNode.AsShared();
+						return !this->Selection.IsSelected(SharedNode) && !this->Selection.NodeHasSelectedKeysOrSections(SharedNode);
+					};
 
-						if (!bAlreadySelected)
-						{
-							TSet<TSharedRef<FSequencerDisplayNode> > DescendantNodes;
-							SequencerHelpers::GetDescendantNodes(ObjectBindingNode, DescendantNodes);
-
-							for (auto DescendantNode : DescendantNodes)
-							{
-								if (Selection.IsSelected(DescendantNode) || Selection.NodeHasSelectedKeysOrSections(DescendantNode))
-								{
-									bAlreadySelected = true;
-									break;
-								}
-							}
-						}
-
-						if (!bAlreadySelected)
-						{
-							bAllAlreadySelected = false;
-						}
+					const bool bNoChildrenSelected = ObjectBindingNode->Traverse_ParentFirst(Traverse_IsSelected, false);
+					if (bNoChildrenSelected)
+					{
+						bAllAlreadySelected = false;
 					}
 				}
-				else if (Selection.IsSelected(ObjectBindingNode))
-				{
-					bAllAlreadySelected = false;
-				}
+			}
+			else if (Selection.IsSelected(ObjectBindingNode))
+			{
+				bAllAlreadySelected = false;
 			}
 		}
 	}
@@ -5248,161 +5238,6 @@ void FSequencer::OnSelectedNodesOnlyChanged()
 	SynchronizeSequencerSelectionWithExternalSelection();
 }
 
-void GatherKeyAreas(const TSet<TSharedRef<FSequencerDisplayNode>>& Selection, ECurveEditorCurveVisibility CurveVisibility, bool bAddChildren, TSharedRef<FSequencerDisplayNode> InNode, TSet<TSharedPtr<IKeyArea>>& OutKeyAreasToShow)
-{
-	// If we're only adding selected curves, and we've encountered a selected node, add all its child key areas
-	if (CurveVisibility == ECurveEditorCurveVisibility::SelectedCurves && Selection.Contains(InNode))
-	{
-		bAddChildren = true;
-	}
-
-	if (bAddChildren)
-	{
-		TSharedPtr<FSequencerSectionKeyAreaNode> KeyAreaNode;
-		if (InNode->GetType() == ESequencerNode::Track)
-		{
-			KeyAreaNode = StaticCastSharedRef<FSequencerTrackNode>(InNode)->GetTopLevelKeyNode();
-		}
-		else if (InNode->GetType() == ESequencerNode::KeyArea)
-		{
-			KeyAreaNode = StaticCastSharedRef<FSequencerSectionKeyAreaNode>(InNode);
-		}
-
-		if (KeyAreaNode.IsValid())
-		{
-			for (TSharedPtr<IKeyArea> KeyArea : KeyAreaNode->GetAllKeyAreas())
-			{
-				OutKeyAreasToShow.Add(KeyArea);
-			}
-		}
-	}
-
-	for (TSharedRef<FSequencerDisplayNode> Child : InNode->GetChildNodes())
-	{
-		GatherKeyAreas(Selection, CurveVisibility, bAddChildren, Child, OutKeyAreasToShow);
-	}
-}
-
-void FSequencer::SyncCurveEditorToSelection(bool bOutlinerSelectionChanged)
-{
-	if (!GetShowCurveEditor())
-	{
-		return;
-	}
-
-	ECurveEditorCurveVisibility CurveVisibility = Settings->GetCurveVisibility();
-	const TSet<TSharedRef<FSequencerDisplayNode>>& OutlinerSelection = Selection.GetSelectedOutlinerNodes();
-
-	// Traverse the tree to gather key areas
-	TSet<TSharedPtr<IKeyArea>> KeyAreasToShow;
-	for (const TSharedRef<FSequencerDisplayNode>& Node : NodeTree->GetRootNodes())
-	{
-		bool bAddChildren = CurveVisibility != ECurveEditorCurveVisibility::SelectedCurves;
-		GatherKeyAreas(OutlinerSelection, CurveVisibility, bAddChildren, Node, KeyAreasToShow);
-	}
-
-	// Remove unanimated curves if necessary
-	if (CurveVisibility == ECurveEditorCurveVisibility::AnimatedCurves)
-	{
-		TArray<TSharedPtr<IKeyArea>> UnanimatedCurves;
-		for (TSharedPtr<IKeyArea> KeyArea : KeyAreasToShow)
-		{
-			const FMovieSceneChannel* Channel = KeyArea->ResolveChannel();
-			if (Channel && Channel->GetNumKeys() == 0)
-			{
-				UnanimatedCurves.Add(KeyArea);
-			}
-		}
-
-		for (TSharedPtr<IKeyArea> Unanimated : UnanimatedCurves)
-		{
-			KeyAreasToShow.Remove(Unanimated);
-		}
-	}
-
-	// Cache the curve editor's current selection
-	TMap<const void*, FCurveModelID> ExistingCurveIDs;
-	for (const TTuple<FCurveModelID, TUniquePtr<FCurveModel>>& Pair : CurveEditorModel->GetCurves())
-	{
-		if (const void* Ptr = Pair.Value->GetCurve())
-		{
-			ExistingCurveIDs.Add(Ptr, Pair.Key);
-		}
-	}
-
-	bool bAnythingChanged = false;
-
-	// Add newly selected curves to the curve editor
-	for (TSharedPtr<IKeyArea> KeyArea : KeyAreasToShow)
-	{
-		const FMovieSceneChannel* Channel = KeyArea->ResolveChannel();
-		if (!Channel)
-		{
-			continue;
-		}
-		if (KeyArea->GetOwningSection()->ShowCurveForChannel(Channel))
-		{
-			if (ExistingCurveIDs.Contains(Channel))
-			{
-				ExistingCurveIDs.Remove(Channel);
-				continue;
-			}
-			else
-			{
-				TUniquePtr<FCurveModel> NewCurve = KeyArea->CreateCurveEditorModel(AsShared());
-				if (NewCurve.IsValid())
-				{
-					bAnythingChanged = true;
-					CurveEditorModel->AddCurve(MoveTemp(NewCurve));
-				}
-			}
-		}
-	}
-
-	// Remove anything that's no longer selected or shown
-	for (const TTuple<const void*, FCurveModelID>& Pair : ExistingCurveIDs)
-	{
-		bAnythingChanged = true;
-		CurveEditorModel->RemoveCurve(Pair.Value);
-	}
-
-	if (bAnythingChanged && CurveEditorModel->ShouldAutoFrame())
-	{
-		CurveEditorModel->ZoomToFit();
-	}
-	else if (bOutlinerSelectionChanged && CurveEditorModel->ShouldAutoFrame())
-	{
-		// If outliner selection changes, zoom to fit only the selected curves
-		TSet<TSharedPtr<IKeyArea>> SelectedKeyAreasToShow;
-		for (const TSharedRef<FSequencerDisplayNode>& Node : NodeTree->GetRootNodes())
-		{
-			GatherKeyAreas(OutlinerSelection, ECurveEditorCurveVisibility::SelectedCurves, false, Node, SelectedKeyAreasToShow);
-		}
-	
-		TArray<FCurveModelID> CurveModelIDs;
-		for (const TTuple<FCurveModelID, TUniquePtr<FCurveModel>>& Pair : CurveEditorModel->GetCurves())
-		{
-			if (const void* Ptr = Pair.Value->GetCurve())
-			{
-				for (TSharedPtr<IKeyArea> KeyArea : SelectedKeyAreasToShow)
-				{
-					const FMovieSceneChannel* Channel = KeyArea->ResolveChannel();
-					if (Channel && Ptr == Channel)
-					{
-						CurveModelIDs.Add(Pair.Key);
-						break;
-					}
-				}
-			}
-		}
-
-		if (CurveModelIDs.Num())
-		{
-			CurveEditorModel->ZoomToFitCurves(MakeArrayView(CurveModelIDs.GetData(), CurveModelIDs.Num()));
-		}
-	}
-}
-
 void FSequencer::ZoomToSelectedSections()
 {
 	FFrameRate TickResolution = GetFocusedTickResolution();
@@ -5480,11 +5315,11 @@ void FSequencer::GetSelectedSections(TArray<UMovieSceneSection*>& OutSelectedSec
 
 void FSequencer::SelectObject(FGuid ObjectBinding)
 {
-	const TSharedPtr<FSequencerObjectBindingNode>* Node = NodeTree->GetObjectBindingMap().Find(ObjectBinding);
-	if (Node != nullptr && Node->IsValid())
+	TSharedPtr<FSequencerObjectBindingNode> Node = NodeTree->FindObjectBindingNode(ObjectBinding);
+	if (Node.IsValid())
 	{
 		GetSelection().Empty();
-		GetSelection().AddToSelection(Node->ToSharedRef());
+		GetSelection().AddToSelection(Node.ToSharedRef());
 	}
 }
 
@@ -7901,13 +7736,9 @@ void FSequencer::SetKey()
 	{
 		if (OutlinerNode->GetType() == ESequencerNode::Track)
 		{
-			TSharedRef<FSequencerTrackNode> TrackNode = StaticCastSharedRef<FSequencerTrackNode>(OutlinerNode);
-
-			TSharedRef<FSequencerDisplayNode> ObjectBindingNode = OutlinerNode;
-			if (SequencerHelpers::FindObjectBindingNode(TrackNode, ObjectBindingNode))
+			if (TSharedPtr<FSequencerObjectBindingNode> ObjectBindingNode = OutlinerNode->FindParentObjectBindingNode())
 			{
-				FGuid ObjectGuid = StaticCastSharedRef<FSequencerObjectBindingNode>(ObjectBindingNode)->GetObjectBinding();
-				TrackNode->AddKey(ObjectGuid);
+				StaticCastSharedRef<FSequencerTrackNode>(OutlinerNode)->GetTrackEditor().AddKey(ObjectBindingNode->GetObjectBinding());
 			}
 		}
 	}
@@ -9458,9 +9289,9 @@ void FSequencer::BindCommands()
 
 	SequencerCommandBindings->MapAction(
 		Commands.ToggleShowCurveEditor,
-		FExecuteAction::CreateLambda( [this]{ SetShowCurveEditor(!GetShowCurveEditor()); } ),
+		FExecuteAction::CreateLambda( [this]{ SetShowCurveEditor(!GetCurveEditorIsVisible()); } ),
 		FCanExecuteAction::CreateLambda( [this]{ return !IsReadOnly(); } ),
-		FIsActionChecked::CreateLambda( [this]{ return GetShowCurveEditor(); } ) );
+		FIsActionChecked::CreateLambda( [this]{ return GetCurveEditorIsVisible(); } ) );
 
 	SequencerCommandBindings->MapAction(
 		Commands.ToggleLinkCurveEditorTimeRange,
@@ -9836,22 +9667,6 @@ void FSequencer::BindCommands()
 		FExecuteAction::CreateSP( this, &FSequencer::SetPlaybackRangeToAllShots ),
 		FCanExecuteAction::CreateSP( this, &FSequencer::IsViewingMasterSequence ) );
 
-	// Curve Visibility
-	SequencerCommandBindings->MapAction(Commands.SetAllCurveVisibility,
-		FExecuteAction::CreateLambda( [this]{ Settings->SetCurveVisibility( ECurveEditorCurveVisibility::AllCurves ); } ),
-		FCanExecuteAction::CreateLambda( []{ return true; } ),
-		FIsActionChecked::CreateLambda( [this]{ return Settings->GetCurveVisibility() == ECurveEditorCurveVisibility::AllCurves; } ) );
-
-	SequencerCommandBindings->MapAction(Commands.SetSelectedCurveVisibility,
-		FExecuteAction::CreateLambda( [this]{ Settings->SetCurveVisibility( ECurveEditorCurveVisibility::SelectedCurves ); } ),
-		FCanExecuteAction::CreateLambda( []{ return true; } ),
-		FIsActionChecked::CreateLambda( [this]{ return Settings->GetCurveVisibility() == ECurveEditorCurveVisibility::SelectedCurves; } ) );
-
-	SequencerCommandBindings->MapAction(Commands.SetAnimatedCurveVisibility,
-		FExecuteAction::CreateLambda( [this]{ Settings->SetCurveVisibility( ECurveEditorCurveVisibility::AnimatedCurves ); } ),
-		FCanExecuteAction::CreateLambda( []{ return true; } ),
-		FIsActionChecked::CreateLambda( [this]{ return Settings->GetCurveVisibility() == ECurveEditorCurveVisibility::AnimatedCurves; } ) );
-
 	// bind widget specific commands
 	SequencerWidget->BindCommands(SequencerCommandBindings);
 }
@@ -9964,13 +9779,13 @@ FKeyAttributes FSequencer::GetDefaultKeyAttributes() const
 	}
 }
 
-bool FSequencer::GetGridMetrics(float PhysicalWidth, double& OutMajorInterval, int32& OutMinorDivisions) const
+bool FSequencer::GetGridMetrics(const float PhysicalWidth, const double InViewStart, const double InViewEnd, double& OutMajorInterval, int32& OutMinorDivisions) const
 {
 	FSlateFontInfo SmallLayoutFont = FCoreStyle::GetDefaultFontStyle("Regular", 8);
 	TSharedRef<FSlateFontMeasure> FontMeasureService = FSlateApplication::Get().GetRenderer()->GetFontMeasureService();
 
-	double BiggestTime = GetViewRange().GetUpperBoundValue();
-	FString TickString = GetNumericTypeInterface()->ToString((BiggestTime * GetFocusedDisplayRate()).FrameNumber.Value);
+	// Use the end of the view as the longest number
+	FString TickString = GetNumericTypeInterface()->ToString((InViewEnd * GetFocusedDisplayRate()).FrameNumber.Value);
 	FVector2D MaxTextSize = FontMeasureService->Measure(TickString, SmallLayoutFont);
 
 	static float MajorTickMultiplier = 2.f;
@@ -9981,7 +9796,7 @@ bool FSequencer::GetGridMetrics(float PhysicalWidth, double& OutMajorInterval, i
 	if (PhysicalWidth > 0)
 	{
 		return GetFocusedDisplayRate().ComputeGridSpacing(
-			PhysicalWidth / GetViewRange().Size<double>(),
+			PhysicalWidth / (InViewEnd - InViewStart),
 			OutMajorInterval,
 			OutMinorDivisions,
 			MinTickPx,
@@ -10028,6 +9843,4 @@ void FSequencer::RecompileDirtyDirectors()
 		}
 	}
 }
-
-
 #undef LOCTEXT_NAMESPACE
diff --git a/Engine/Source/Editor/Sequencer/Private/Sequencer.h b/Engine/Source/Editor/Sequencer/Private/Sequencer.h
--- a/Engine/Source/Editor/Sequencer/Private/Sequencer.h
+++ b/Engine/Source/Editor/Sequencer/Private/Sequencer.h
@@ -141,7 +141,7 @@ public:
 	/**
 	 * Compute a major grid interval and number of minor divisions to display
 	 */
-	bool GetGridMetrics(float PhysicalWidth, double& OutMajorInterval, int32& OutMinorDivisions) const;
+	bool GetGridMetrics(const float PhysicalWidth, const double InViewStart, const double InViewEnd, double& OutMajorInterval, int32& OutMinorDivisions) const;
 
 public:
 
@@ -451,7 +451,8 @@ public:
 
 	/** @return Whether to show the curve editor or not */
 	void SetShowCurveEditor(bool bInShowCurveEditor);
-	bool GetShowCurveEditor() const { return bShowCurveEditor; }
+	/** @return If the curve editor is currently visible. */
+	bool GetCurveEditorIsVisible() const;
 
 	/** Called to save the current movie scene */
 	void SaveCurrentMovieScene();
@@ -528,8 +529,6 @@ public:
 
 	void OnSelectedNodesOnlyChanged();
 
-	void SyncCurveEditorToSelection(bool bOutlinerSelectionChanged);
-
 	TSharedPtr<FCurveEditor> GetCurveEditor() const
 	{
 		return CurveEditorModel;
@@ -1109,8 +1108,6 @@ private:
 	/** True if this sequencer is being edited within the level editor */
 	bool bIsEditingWithinLevelEditor;
 
-	bool bShowCurveEditor;
-
 	/** Whether the sequence should be editable or read only */
 	bool bReadOnly;
 
diff --git a/Engine/Source/Editor/Sequencer/Private/SequencerChannelTraits.cpp b/Engine/Source/Editor/Sequencer/Private/SequencerChannelTraits.cpp
--- a/Engine/Source/Editor/Sequencer/Private/SequencerChannelTraits.cpp
+++ b/Engine/Source/Editor/Sequencer/Private/SequencerChannelTraits.cpp
@@ -20,6 +20,11 @@ void DrawKeys(FMovieSceneChannel* Channel, TArrayView<const FKeyHandle> InHandle
 	}
 }
 
+bool SupportsCurveEditorModels(const FMovieSceneChannelHandle& ChannelHandle)
+{
+	return false;
+}
+
 TUniquePtr<FCurveModel> CreateCurveEditorModel(const FMovieSceneChannelHandle& ChannelHandle, UMovieSceneSection* OwningSection, TSharedRef<ISequencer> InSequencer)
 {
 	return nullptr;
diff --git a/Engine/Source/Editor/Sequencer/Private/SequencerCommonHelpers.cpp b/Engine/Source/Editor/Sequencer/Private/SequencerCommonHelpers.cpp
--- a/Engine/Source/Editor/Sequencer/Private/SequencerCommonHelpers.cpp
+++ b/Engine/Source/Editor/Sequencer/Private/SequencerCommonHelpers.cpp
@@ -152,29 +152,6 @@ void SequencerHelpers::GetAllSections(TSharedRef<FSequencerDisplayNode> DisplayN
 	}
 }
 
-bool SequencerHelpers::FindObjectBindingNode(TSharedRef<FSequencerDisplayNode> DisplayNode, TSharedRef<FSequencerDisplayNode>& ObjectBindingNode)
-{
-	TArray< TSharedPtr<FSequencerDisplayNode> > ParentNodes;
-	ParentNodes.Add(DisplayNode);
-
-	while (DisplayNode->GetParent().IsValid())
-	{
-		ParentNodes.Add(DisplayNode->GetParent());
-		DisplayNode = DisplayNode->GetParent().ToSharedRef();
-	}
-
-	for (int32 ParentNodeIndex = ParentNodes.Num() - 1; ParentNodeIndex >= 0; --ParentNodeIndex)
-	{
-		if (ParentNodes[ParentNodeIndex]->GetType() == ESequencerNode::Object)
-		{
-			ObjectBindingNode = ParentNodes[ParentNodeIndex].ToSharedRef();
-			return true;
-		}
-	}
-
-	return false;
-}
-
 bool IsSectionSelectedInNode(FSequencer& Sequencer, TSharedRef<FSequencerDisplayNode> InNode)
 {
 	if (InNode->GetType() == ESequencerNode::Track)
@@ -315,7 +292,7 @@ void SequencerHelpers::PerformDefaultSelection(FSequencer& Sequencer, const FPoi
 		}
 		else if (Hotspot->GetType() == ESequencerHotspot::Section || Hotspot->GetType() == ESequencerHotspot::EasingArea)
 		{
-			UMovieSceneSection* Section = static_cast<FSectionHotspot*>(Hotspot.Get())->Section.GetSectionObject();
+			UMovieSceneSection* Section = static_cast<FSectionHotspot*>(Hotspot.Get())->WeakSection.Get();
 			if (!Selection.IsSelected(Section))
 			{
 				ConditionallyClearSelection();
@@ -324,7 +301,7 @@ void SequencerHelpers::PerformDefaultSelection(FSequencer& Sequencer, const FPoi
 		}
 		else if (Hotspot->GetType() == ESequencerHotspot::SectionResize_L || Hotspot->GetType() == ESequencerHotspot::SectionResize_R)
 		{
-			UMovieSceneSection* Section = static_cast<FSectionResizeHotspot*>(Hotspot.Get())->Section.GetSectionObject();
+			UMovieSceneSection* Section = static_cast<FSectionResizeHotspot*>(Hotspot.Get())->WeakSection.Get();
 			if (!Selection.IsSelected(Section))
 			{
 				ConditionallyClearSelection();
@@ -367,7 +344,7 @@ void SequencerHelpers::PerformDefaultSelection(FSequencer& Sequencer, const FPoi
 	}
 	else if (Hotspot->GetType() == ESequencerHotspot::Section || Hotspot->GetType() == ESequencerHotspot::EasingArea)
 	{
-		UMovieSceneSection* Section = static_cast<FSectionHotspot*>(Hotspot.Get())->Section.GetSectionObject();
+		UMovieSceneSection* Section = static_cast<FSectionHotspot*>(Hotspot.Get())->WeakSection.Get();
 
 		// Never allow infinite sections to be selected through normal click (they're only selectable through right click)
 		if (Section->GetRange() != TRange<FFrameNumber>::All())
diff --git a/Engine/Source/Editor/Sequencer/Private/SequencerCommonHelpers.h b/Engine/Source/Editor/Sequencer/Private/SequencerCommonHelpers.h
--- a/Engine/Source/Editor/Sequencer/Private/SequencerCommonHelpers.h
+++ b/Engine/Source/Editor/Sequencer/Private/SequencerCommonHelpers.h
@@ -38,11 +38,6 @@ public:
 	 */
 	static void GetAllSections(TSharedRef<FSequencerDisplayNode> DisplayNode, TSet<TWeakObjectPtr<UMovieSceneSection>>& Sections);
 
-	/**
-	 * Find object binding node from the given display node
-	 */
-	static bool FindObjectBindingNode(TSharedRef<FSequencerDisplayNode> DisplayNode, TSharedRef<FSequencerDisplayNode>& ObjectBindingNode);
-
 	/**
 	 * Validate that the nodes with selected keys or sections actually are true
 	 */
diff --git a/Engine/Source/Editor/Sequencer/Private/SequencerContextMenus.cpp b/Engine/Source/Editor/Sequencer/Private/SequencerContextMenus.cpp
--- a/Engine/Source/Editor/Sequencer/Private/SequencerContextMenus.cpp
+++ b/Engine/Source/Editor/Sequencer/Private/SequencerContextMenus.cpp
@@ -472,11 +472,9 @@ FMovieSceneBlendTypeField FSectionContextMenu::GetSupportedBlendTypes() const
 {
 	FMovieSceneBlendTypeField BlendTypes = FMovieSceneBlendTypeField::All();
 
-	TArray<FSectionHandle> SelectedSections = StaticCastSharedRef<SSequencer>(Sequencer->GetSequencerWidget())->GetSectionHandles(Sequencer->GetSelection().GetSelectedSections());
-	for (const FSectionHandle& Handle : SelectedSections)
+	for (TWeakObjectPtr<UMovieSceneSection> WeakSection : Sequencer->GetSelection().GetSelectedSections())
 	{
-		UMovieSceneSection* Section = Handle.GetSectionObject();
-		if (Section)
+		if (UMovieSceneSection* Section = WeakSection.Get())
 		{
 			// Remove unsupported blend types
 			BlendTypes.Remove(Section->GetSupportedBlendTypes().Invert());
@@ -586,13 +584,11 @@ void FSectionContextMenu::AddBlendTypeMenu(FMenuBuilder& MenuBuilder)
 {
 	TArray<TWeakObjectPtr<UMovieSceneSection>> Sections;
 
-	TArray<FSectionHandle> SelectedSections = StaticCastSharedRef<SSequencer>(Sequencer->GetSequencerWidget())->GetSectionHandles(Sequencer->GetSelection().GetSelectedSections());
-	for (const FSectionHandle& Handle : SelectedSections)
+	for (TWeakObjectPtr<UMovieSceneSection> WeakSection : Sequencer->GetSelection().GetSelectedSections())
 	{
-		UMovieSceneSection* Section = Handle.GetSectionObject();
-		if (Section)
+		if (WeakSection.IsValid())
 		{
-			Sections.Add(Section);
+			Sections.Add(WeakSection);
 		}
 	}
 
@@ -602,16 +598,16 @@ void FSectionContextMenu::AddBlendTypeMenu(FMenuBuilder& MenuBuilder)
 
 void FSectionContextMenu::SelectAllKeys()
 {
-	TArray<FSectionHandle> SelectedSections = StaticCastSharedRef<SSequencer>(Sequencer->GetSequencerWidget())->GetSectionHandles(Sequencer->GetSelection().GetSelectedSections());
-	for (const FSectionHandle& Handle : SelectedSections)
+	for (const TWeakObjectPtr<UMovieSceneSection>& WeakSection : Sequencer->GetSelection().GetSelectedSections())
 	{
-		UMovieSceneSection* Section = Handle.GetSectionObject();
-		if (!Section)
+		UMovieSceneSection* Section = WeakSection.Get();
+		TOptional<FSectionHandle> SectionHandle = Sequencer->GetNodeTree()->GetSectionHandle(Section);
+		if (!SectionHandle)
 		{
 			continue;
 		}
 
-		FSectionLayout Layout(*Handle.TrackNode, Handle.SectionIndex);
+		FSectionLayout Layout(*SectionHandle->GetTrackNode(), SectionHandle->GetSectionIndex());
 		for (const FSectionLayoutElement& Element : Layout.GetElements())
 		{
 			for (TSharedRef<IKeyArea> KeyArea : Element.GetKeyAreas())
@@ -637,30 +633,27 @@ void FSectionContextMenu::CopyAllKeys()
 
 void FSectionContextMenu::TogglePrimeForRecording() const
 {
-	TArray<FSectionHandle> SelectedSections = StaticCastSharedRef<SSequencer>(Sequencer->GetSequencerWidget())->GetSectionHandles(Sequencer->GetSelection().GetSelectedSections());
-	if(SelectedSections.Num() > 0)
+	for (TWeakObjectPtr<UMovieSceneSection> WeakSection : Sequencer->GetSelection().GetSelectedSections())
 	{
-		const FSectionHandle& Handle = SelectedSections[0];
-		UMovieSceneSubSection* SubSection = Cast<UMovieSceneSubSection>(Handle.GetSectionObject());
+		UMovieSceneSubSection* SubSection = Cast<UMovieSceneSubSection>(WeakSection.Get());
 		if (SubSection)
 		{
 			SubSection->SetAsRecording(SubSection != UMovieSceneSubSection::GetRecordingSection());
-		}	
+			break;
+		}
 	}
 }
 
 
 bool FSectionContextMenu::IsPrimedForRecording() const
 {
-	TArray<FSectionHandle> SelectedSections = StaticCastSharedRef<SSequencer>(Sequencer->GetSequencerWidget())->GetSectionHandles(Sequencer->GetSelection().GetSelectedSections());
-	if(SelectedSections.Num() > 0)
+	for (TWeakObjectPtr<UMovieSceneSection> WeakSection : Sequencer->GetSelection().GetSelectedSections())
 	{
-		const FSectionHandle& Handle = SelectedSections[0];
-		UMovieSceneSubSection* SubSection = Cast<UMovieSceneSubSection>(Handle.GetSectionObject());
+		UMovieSceneSubSection* SubSection = Cast<UMovieSceneSubSection>(WeakSection.Get());
 		if (SubSection)
 		{
 			return SubSection == UMovieSceneSubSection::GetRecordingSection();
-		}	
+		}
 	}
 
 	return false;
@@ -668,15 +661,13 @@ bool FSectionContextMenu::IsPrimedForRecording() const
 
 bool FSectionContextMenu::CanPrimeForRecording() const
 {
-	TArray<FSectionHandle> SelectedSections = StaticCastSharedRef<SSequencer>(Sequencer->GetSequencerWidget())->GetSectionHandles(Sequencer->GetSelection().GetSelectedSections());
-	if(SelectedSections.Num() > 0)
+	for (TWeakObjectPtr<UMovieSceneSection> WeakSection : Sequencer->GetSelection().GetSelectedSections())
 	{
-		const FSectionHandle& Handle = SelectedSections[0];
-		UMovieSceneSubSection* SubSection = ExactCast<UMovieSceneSubSection>(Handle.GetSectionObject());
+		UMovieSceneSubSection* SubSection = ExactCast<UMovieSceneSubSection>(WeakSection.Get());
 		if (SubSection)
 		{
 			return true;
-		}	
+		}
 	}
 
 	return false;
@@ -685,12 +676,14 @@ bool FSectionContextMenu::CanPrimeForRecording() const
 
 void FSectionContextMenu::SetSectionToKey()
 {
-	TArray<FSectionHandle> SelectedSections = StaticCastSharedRef<SSequencer>(Sequencer->GetSequencerWidget())->GetSectionHandles(Sequencer->GetSelection().GetSelectedSections());
-	if (SelectedSections.Num() == 1)
+	if (Sequencer->GetSelection().GetSelectedSections().Num() != 1)
 	{
-		const FSectionHandle& Handle = SelectedSections[0];
-		UMovieSceneSection* Section = Handle.GetSectionObject();
-		if (Section)
+		return;
+	}
+
+	for (TWeakObjectPtr<UMovieSceneSection> WeakSection : Sequencer->GetSelection().GetSelectedSections())
+	{
+		if (UMovieSceneSection* Section = WeakSection.Get())
 		{
 			UMovieScenePropertyTrack* Track = Section->GetTypedOuter<UMovieScenePropertyTrack>();
 			if (Track)
@@ -700,17 +693,20 @@ void FSectionContextMenu::SetSectionToKey()
 				Track->SetSectionToKey(Section);
 			}
 		}
+		break;
 	}
 }
 
 bool FSectionContextMenu::CanSetSectionToKey() const
 {
-	TArray<FSectionHandle> SelectedSections = StaticCastSharedRef<SSequencer>(Sequencer->GetSequencerWidget())->GetSectionHandles(Sequencer->GetSelection().GetSelectedSections());
-	if (SelectedSections.Num() == 1)
+	if (Sequencer->GetSelection().GetSelectedSections().Num() != 1)
 	{
-		const FSectionHandle& Handle = SelectedSections[0];
-		UMovieSceneSection* Section = Handle.GetSectionObject();
-		if (Section)
+		return false;
+	}
+
+	for (TWeakObjectPtr<UMovieSceneSection> WeakSection : Sequencer->GetSelection().GetSelectedSections())
+	{
+		if (UMovieSceneSection* Section = WeakSection.Get())
 		{
 			UMovieScenePropertyTrack* Track = Section->GetTypedOuter<UMovieScenePropertyTrack>();
 			if (Track && Section->GetBlendType().IsValid() && (Section->GetBlendType().Get() == EMovieSceneBlendType::Absolute || Section->GetBlendType().Get() == EMovieSceneBlendType::Additive))
@@ -718,13 +714,15 @@ bool FSectionContextMenu::CanSetSectionToKey() const
 				return true;
 			}
 		}
+
+		break;
 	}
 	return false;
 }
 
 bool FSectionContextMenu::CanSelectAllKeys() const
 {
-	for (auto& Pair : ChannelsByType)
+	for (const TTuple<FName, TArray<FMovieSceneChannelHandle>>& Pair : ChannelsByType)
 	{
 		for (const FMovieSceneChannelHandle& Handle : Pair.Value)
 		{
@@ -1388,7 +1386,7 @@ void FPasteContextMenu::Setup()
 			}
 
 			TArray<UMovieSceneSection*> Sections;
-			for (auto Section : TrackNode->GetSections())
+			for (TSharedRef<ISequencerSection> Section : TrackNode->GetSections())
 			{
 				if (Section.Get().GetSectionObject())
 				{
@@ -1405,34 +1403,39 @@ void FPasteContextMenu::Setup()
 
 			if (SectionIndex != INDEX_NONE)
 			{
-				SectionHandles.Add(FSectionHandle(TrackNode, SectionIndex));
+				SectionHandles.Add(FSectionHandle(TrackNode.ToSharedRef(), SectionIndex));
 			}
 		}
 	}
 	else
 	{
 		// Use the selected sections
-		TSharedRef<SSequencer> SequencerWidget = StaticCastSharedRef<SSequencer>(Sequencer->GetSequencerWidget());
-		SectionHandles = SequencerWidget->GetSectionHandles(Sequencer->GetSelection().GetSelectedSections());
+		for (TWeakObjectPtr<UMovieSceneSection> WeakSection : Sequencer->GetSelection().GetSelectedSections())
+		{
+			if (TOptional<FSectionHandle> SectionHandle = Sequencer->GetNodeTree()->GetSectionHandle(WeakSection.Get()))
+			{
+				SectionHandles.Add(SectionHandle.GetValue());
+			}
+		}
 	}
 
 	TMap<FName, TArray<FSectionHandle>> SectionsByType;
 	for (const FSectionHandle& Section : SectionHandles)
 	{
-		UMovieSceneTrack* Track = Section.TrackNode->GetTrack();
+		UMovieSceneTrack* Track = Section.GetTrackNode()->GetTrack();
 		if (Track)
 		{
 			SectionsByType.FindOrAdd(Track->GetClass()->GetFName()).Add(Section);
 		}
 	}
 
-	for (auto& Pair : SectionsByType)
+	for (const TTuple<FName, TArray<FSectionHandle>>& Pair : SectionsByType)
 	{
 		FPasteDestination& Destination = PasteDestinations[PasteDestinations.AddDefaulted()];
 		if (Pair.Value.Num() == 1)
 		{
 			FString Path;
-			GetFullNodePath(*Pair.Value[0].TrackNode, Path);
+			GetFullNodePath(*Pair.Value[0].GetTrackNode(), Path);
 			Destination.Name = FText::FromString(Path);
 		}
 		else
@@ -1442,7 +1445,7 @@ void FPasteContextMenu::Setup()
 
 		for (const FSectionHandle& Section : Pair.Value)
 		{
-			GatherPasteDestinationsForNode(*Section.TrackNode, Section.GetSectionObject(), NAME_None, Destination.Reconcilers);
+			GatherPasteDestinationsForNode(*Section.GetTrackNode(), Section.GetSectionObject(), NAME_None, Destination.Reconcilers);
 		}
 
 		// Reconcile and remove invalid pastes
@@ -1732,7 +1735,7 @@ TOptional<int32> FEasingContextMenu::GetCurrentLength() const
 
 	for (const FEasingAreaHandle& Handle : Easings)
 	{
-		UMovieSceneSection* Section = Handle.Section.GetSectionObject();
+		UMovieSceneSection* Section = Handle.WeakSection.Get();
 		if (Section)
 		{
 			if (Handle.EasingType == ESequencerEasingType::In && Section->Easing.GetEaseInDuration() == Value.Get(Section->Easing.GetEaseInDuration()))
@@ -1757,7 +1760,7 @@ void FEasingContextMenu::OnUpdateLength(int32 NewLength)
 {
 	for (const FEasingAreaHandle& Handle : Easings)
 	{
-		if (UMovieSceneSection* Section = Handle.Section.GetSectionObject())
+		if (UMovieSceneSection* Section = Handle.WeakSection.Get())
 		{
 			Section->Modify();
 			if (Handle.EasingType == ESequencerEasingType::In)
@@ -1779,7 +1782,7 @@ ECheckBoxState FEasingContextMenu::GetAutoEasingCheckState() const
 	TOptional<bool> IsChecked;
 	for (const FEasingAreaHandle& Handle : Easings)
 	{
-		if (UMovieSceneSection* Section = Handle.Section.GetSectionObject())
+		if (UMovieSceneSection* Section = Handle.WeakSection.Get())
 		{
 			if (Handle.EasingType == ESequencerEasingType::In)
 			{
@@ -1810,7 +1813,7 @@ void FEasingContextMenu::SetAutoEasing(bool bAutoEasing)
 
 	for (const FEasingAreaHandle& Handle : Easings)
 	{
-		if (UMovieSceneSection* Section = Handle.Section.GetSectionObject())
+		if (UMovieSceneSection* Section = Handle.WeakSection.Get())
 		{
 			AllTracks.AddUnique(Section->GetTypedOuter<UMovieSceneTrack>());
 
@@ -1838,7 +1841,7 @@ FText FEasingContextMenu::GetEasingTypeText() const
 	UClass* ClassType = nullptr;
 	for (const FEasingAreaHandle& Handle : Easings)
 	{
-		if (UMovieSceneSection* Section = Handle.Section.GetSectionObject())
+		if (UMovieSceneSection* Section = Handle.WeakSection.Get())
 		{
 			UObject* Object = Handle.EasingType == ESequencerEasingType::In ? Section->Easing.EaseIn.GetObject() : Section->Easing.EaseOut.GetObject();
 			if (Object)
@@ -1902,7 +1905,7 @@ void FEasingContextMenu::OnEasingTypeChanged(UClass* NewClass)
 
 	for (const FEasingAreaHandle& Handle : Easings)
 	{
-		UMovieSceneSection* Section = Handle.Section.GetSectionObject();
+		UMovieSceneSection* Section = Handle.WeakSection.Get();
 		if (!Section)
 		{
 			continue;
@@ -1942,7 +1945,7 @@ void FEasingContextMenu::EasingOptionsMenu(FMenuBuilder& MenuBuilder)
 	TArray<UObject*> Objects;
 	for (const FEasingAreaHandle& Handle : Easings)
 	{
-		if (UMovieSceneSection* Section = Handle.Section.GetSectionObject())
+		if (UMovieSceneSection* Section = Handle.WeakSection.Get())
 		{
 			if (Handle.EasingType == ESequencerEasingType::In)
 			{
diff --git a/Engine/Source/Editor/Sequencer/Private/SequencerEdMode.cpp b/Engine/Source/Editor/Sequencer/Private/SequencerEdMode.cpp
--- a/Engine/Source/Editor/Sequencer/Private/SequencerEdMode.cpp
+++ b/Engine/Source/Editor/Sequencer/Private/SequencerEdMode.cpp
@@ -4,6 +4,8 @@
 #include "EditorViewportClient.h"
 #include "Curves/KeyHandle.h"
 #include "ISequencer.h"
+#include "MovieSceneSequence.h"
+#include "MovieScene.h"
 #include "DisplayNodes/SequencerDisplayNode.h"
 #include "Sequencer.h"
 #include "Framework/Application/SlateApplication.h"
@@ -165,38 +167,28 @@ void FSequencerEdMode::OnKeySelected(FViewport* Viewport, HMovieSceneKeyProxy* K
 			for (const FTrajectoryKey::FData KeyData : KeyProxy->Key.KeyData)
 			{
 				UMovieSceneSection* Section = KeyData.Section.Get();
-				if (Section && KeyData.KeyHandle.IsSet())
+				TOptional<FSectionHandle> SectionHandle = Sequencer->GetNodeTree()->GetSectionHandle(Section);
+				if (SectionHandle && KeyData.KeyHandle.IsSet())
 				{
-					TSet<TWeakObjectPtr<UMovieSceneSection>> SectionsToFind;
-					SectionsToFind.Add(Section);
+					TArray<TSharedRef<FSequencerSectionKeyAreaNode>> KeyAreaNodes;
+					SectionHandle->GetTrackNode()->GetChildKeyAreaNodesRecursively(KeyAreaNodes);
 
-					// Find the key area with the specified channel name
-					TArray<FSectionHandle> SectionHandles = StaticCastSharedRef<SSequencer>(Sequencer->GetSequencerWidget())->GetSectionHandles(SectionsToFind);
-
-					for (FSectionHandle Handle : SectionHandles)
+					for (TSharedRef<FSequencerSectionKeyAreaNode> KeyAreaNode : KeyAreaNodes)
 					{
-						TArray<TSharedRef<FSequencerSectionKeyAreaNode>> KeyAreaNodes;
-						Handle.TrackNode->GetChildKeyAreaNodesRecursively(KeyAreaNodes);
-
-						for (TSharedRef<FSequencerSectionKeyAreaNode> KeyAreaNode : KeyAreaNodes)
+						TSharedPtr<IKeyArea> KeyArea = KeyAreaNode->GetKeyArea(Section);
+						if (KeyArea.IsValid() && KeyArea->GetName() == KeyData.ChannelName)
 						{
-							TSharedPtr<IKeyArea> KeyArea = KeyAreaNode->GetKeyArea(Section);
-							if (KeyArea.IsValid() && KeyArea->GetName() == KeyData.ChannelName)
+							if (!bChangedSelection)
 							{
-								if (!bChangedSelection)
-								{
-									Sequencer->GetSelection().SuspendBroadcast();
-									bChangedSelection = true;
-								}
-
-								Sequencer->SelectKey(Section, KeyArea, KeyData.KeyHandle.GetValue(), bToggleSelection);
-								goto next_key;
+								Sequencer->GetSelection().SuspendBroadcast();
+								bChangedSelection = true;
 							}
+
+							Sequencer->SelectKey(Section, KeyArea, KeyData.KeyHandle.GetValue(), bToggleSelection);
+							break;
 						}
 					}
 				}
-next_key:
-				continue;
 			}
 			if (bChangedSelection)
 			{
@@ -275,48 +267,51 @@ void FSequencerEdMode::GetParents(TArray<const UObject *>& Parents, const UObjec
 /** This is not that scalable moving forward with stuff like the control rig , need a better caching solution there */
 bool FSequencerEdMode::GetParentTM(FTransform& CurrentRefTM, const TSharedPtr<FSequencer>& Sequencer, UObject* ParentObject, FFrameTime KeyTime)
 {
+	UMovieSceneSequence* Sequence = Sequencer->GetFocusedMovieSceneSequence();
+	if (!Sequence)
+	{
+		return false;
+	}
+
 	FGuid ObjectBinding = Sequencer->FindCachedObjectId(*ParentObject, Sequencer->GetFocusedTemplateID());
+	if (!ObjectBinding.IsValid())
+	{
+		return false;
+	}
 
-	if (ObjectBinding.IsValid())
+	const FMovieSceneBinding* Binding = Sequence->GetMovieScene()->FindBinding(ObjectBinding);
+	if (!Binding)
 	{
-		const TSharedPtr< FSequencerDisplayNode >& ObjectNode = Sequencer->GetNodeTree()->GetObjectBindingMap()[ObjectBinding];
+		return false;
+	}
 
-		for (const TSharedRef< FSequencerDisplayNode >& ChildNode : ObjectNode->GetChildNodes())
+	for (const UMovieSceneTrack* Track : Binding->GetTracks())
+	{
+		const UMovieScene3DTransformTrack* TransformTrack = Cast<UMovieScene3DTransformTrack>(Track);
+		if (!TransformTrack)
 		{
-			if (ChildNode->GetType() == ESequencerNode::Track)
-			{
-				const TSharedRef<FSequencerTrackNode> TrackNode = StaticCastSharedRef<FSequencerTrackNode>(ChildNode);
-				const UMovieSceneTrack* TrackNodeTrack = TrackNode->GetTrack();
-				const UMovieScene3DTransformTrack* TransformTrack = Cast<UMovieScene3DTransformTrack>(TrackNodeTrack);
-
-				if (TransformTrack != nullptr)
-				{
-					for (const UMovieSceneSection* Section : TransformTrack->GetAllSections())
-					{
-						if (Section->IsTimeWithinSection(KeyTime.FrameNumber))
-						{
-							const UMovieScene3DTransformSection* ParentSection = Cast<UMovieScene3DTransformSection>(Section);
+			continue;
+		}
 
-							if (ParentSection != nullptr)
-							{
-								FVector ParentKeyPos;
-								FRotator ParentKeyRot;
+		for (const UMovieSceneSection* Section : TransformTrack->GetAllSections())
+		{
+			if (Section->IsTimeWithinSection(KeyTime.FrameNumber))
+			{
+				continue;
+			}
 
-								FMovieSceneEvaluationTrack* EvalTrack = MovieSceneToolHelpers::GetEvaluationTrack(Sequencer.Get(), TransformTrack->GetSignature());
-								if (EvalTrack)
-								{
-									GetLocationAtTime(EvalTrack, ParentObject, KeyTime, ParentKeyPos, ParentKeyRot, Sequencer);
-									CurrentRefTM = FTransform(ParentKeyRot, ParentKeyPos);
-									return true;
-								}
-							}
-						}
-					}
-				}
+			FMovieSceneEvaluationTrack* EvalTrack = MovieSceneToolHelpers::GetEvaluationTrack(Sequencer.Get(), TransformTrack->GetSignature());
+			if (EvalTrack)
+			{
+				FVector ParentKeyPos;
+				FRotator ParentKeyRot;
+				GetLocationAtTime(EvalTrack, ParentObject, KeyTime, ParentKeyPos, ParentKeyRot, Sequencer);
+				CurrentRefTM = FTransform(ParentKeyRot, ParentKeyPos);
+				return true;
 			}
 		}
 	}
-	
+
 	return false;
 }
 
@@ -393,75 +388,70 @@ FTransform FSequencerEdMode::GetRefFrame(const TSharedPtr<FSequencer>& Sequencer
 
 FTransform FSequencerEdMode::GetRefFrame(const TSharedPtr<FSequencer>& Sequencer, const USceneComponent* SceneComponent, FFrameTime KeyTime)
 {
-	FTransform RefTM = FTransform::Identity;
+	UMovieSceneSequence* Sequence = Sequencer->GetFocusedMovieSceneSequence();
+	if (!Sequence)
+	{
+		return FTransform::Identity;
+	}
 
-	if (SceneComponent != nullptr && SceneComponent->GetAttachParent() != nullptr)
+	if (SceneComponent == nullptr || SceneComponent->GetAttachParent() == nullptr)
 	{
-		FTransform ParentRefTM;
+		return FTransform::Identity;
+	}
 
-		// If our parent is the root component, get the RefFrame from the Actor
-		if (SceneComponent->GetAttachParent() == SceneComponent->GetOwner()->GetRootComponent())
-		{
-			ParentRefTM = GetRefFrame(Sequencer, SceneComponent->GetAttachParent()->GetOwner(), KeyTime);
-		}
-		else
-		{
-			ParentRefTM = GetRefFrame(Sequencer, SceneComponent->GetAttachParent(), KeyTime);
-		}
-		
-		FTransform CurrentRefTM = SceneComponent->GetAttachParent()->GetRelativeTransform();
+	FTransform ParentRefTM;
+
+	// If our parent is the root component, get the RefFrame from the Actor
+	if (SceneComponent->GetAttachParent() == SceneComponent->GetOwner()->GetRootComponent())
+	{
+		ParentRefTM = GetRefFrame(Sequencer, SceneComponent->GetAttachParent()->GetOwner(), KeyTime);
+	}
+	else
+	{
+		ParentRefTM = GetRefFrame(Sequencer, SceneComponent->GetAttachParent(), KeyTime);
+	}
 
-		// Check if our parent is animated in this Sequencer
 
-		UObject* ParentObject = SceneComponent->GetAttachParent() == SceneComponent->GetOwner()->GetRootComponent() ? static_cast<UObject*>(SceneComponent->GetOwner()) : SceneComponent->GetAttachParent();
-		FGuid ObjectBinding = Sequencer->FindCachedObjectId(*ParentObject, Sequencer->GetFocusedTemplateID());
+	FTransform CurrentRefTM = SceneComponent->GetAttachParent()->GetRelativeTransform();
 
-		if (ObjectBinding.IsValid())
+	// Check if our parent is animated in this Sequencer
+	UObject* ParentObject = SceneComponent->GetAttachParent() == SceneComponent->GetOwner()->GetRootComponent() ? static_cast<UObject*>(SceneComponent->GetOwner()) : SceneComponent->GetAttachParent();
+	FGuid ObjectBinding = Sequencer->FindCachedObjectId(*ParentObject, Sequencer->GetFocusedTemplateID());
+
+	const FMovieSceneBinding* Binding = ObjectBinding.IsValid() ? Sequence->GetMovieScene()->FindBinding(ObjectBinding) : nullptr;
+	if (Binding)
+	{
+		for (const UMovieSceneTrack* Track : Binding->GetTracks())
 		{
-			const TSharedPtr< FSequencerDisplayNode >& ObjectNode = Sequencer->GetNodeTree()->GetObjectBindingMap()[ObjectBinding];
+			const UMovieScene3DTransformTrack* TransformTrack = Cast<UMovieScene3DTransformTrack>(Track);
+			if (!TransformTrack)
+			{
+				continue;
+			}
 
-			for (const TSharedRef< FSequencerDisplayNode >& ChildNode : ObjectNode->GetChildNodes())
+			for (const UMovieSceneSection* Section : TransformTrack->GetAllSections())
 			{
-				if (ChildNode->GetType() == ESequencerNode::Track)
+				if (Section->IsTimeWithinSection(KeyTime.FrameNumber))
 				{
-					const TSharedRef<FSequencerTrackNode> TrackNode = StaticCastSharedRef<FSequencerTrackNode>(ChildNode);
-					const UMovieSceneTrack* TrackNodeTrack = TrackNode->GetTrack();
-					const UMovieScene3DTransformTrack* TransformTrack = Cast<UMovieScene3DTransformTrack>(TrackNodeTrack);
-
-					if (TransformTrack != nullptr)
-					{
-						for (const UMovieSceneSection* Section : TransformTrack->GetAllSections())
-						{
-							if (Section->IsTimeWithinSection(KeyTime.FrameNumber))
-							{
-								const UMovieScene3DTransformSection* ParentSection = Cast<UMovieScene3DTransformSection>(Section);
-
-								if (ParentSection != nullptr)
-								{
-									FVector ParentKeyPos;
-									FRotator ParentKeyRot;
+					continue;
+				}
 
-									FMovieSceneEvaluationTrack* EvalTrack = MovieSceneToolHelpers::GetEvaluationTrack(Sequencer.Get(), TransformTrack->GetSignature());
-									if (EvalTrack)
-									{
-										GetLocationAtTime(EvalTrack, ParentObject, KeyTime, ParentKeyPos, ParentKeyRot, Sequencer);
+				FMovieSceneEvaluationTrack* EvalTrack = MovieSceneToolHelpers::GetEvaluationTrack(Sequencer.Get(), TransformTrack->GetSignature());
+				if (EvalTrack)
+				{
+					FVector ParentKeyPos;
+					FRotator ParentKeyRot;
+					GetLocationAtTime(EvalTrack, ParentObject, KeyTime, ParentKeyPos, ParentKeyRot, Sequencer);
 
-										CurrentRefTM = FTransform(ParentKeyRot, ParentKeyPos);
+					CurrentRefTM = FTransform(ParentKeyRot, ParentKeyPos);
 
-										return CurrentRefTM * ParentRefTM;
-									}
-								}
-							}
-						}
-					}
+					return CurrentRefTM * ParentRefTM;
 				}
 			}
 		}
-
-		RefTM = CurrentRefTM * ParentRefTM;
 	}
 
-	return RefTM;
+	return CurrentRefTM * ParentRefTM;
 }
 
 void FSequencerEdMode::GetLocationAtTime(FMovieSceneEvaluationTrack* Track, UObject* Object, FFrameTime KeyTime, FVector& KeyPos, FRotator& KeyRot, const TSharedPtr<FSequencer>& Sequencer)
@@ -481,7 +471,7 @@ void FSequencerEdMode::GetLocationAtTime(FMovieSceneEvaluationTrack* Track, UObj
 }
 
 void FSequencerEdMode::DrawTransformTrack(const TSharedPtr<FSequencer>& Sequencer, FPrimitiveDrawInterface* PDI,
-											UMovieScene3DTransformTrack* TransformTrack, const TArray<TWeakObjectPtr<UObject>>& BoundObjects, const bool bIsSelected)
+											UMovieScene3DTransformTrack* TransformTrack, TArrayView<const TWeakObjectPtr<>> BoundObjects, const bool bIsSelected)
 {
 	bool bHitTesting = true;
 	if( PDI != nullptr )
@@ -690,44 +680,46 @@ void FSequencerEdMode::DrawTracks3D(FPrimitiveDrawInterface* PDI)
 			continue;
 		}
 
-		TSet<TSharedRef<FSequencerDisplayNode> > ObjectBindingNodes;
+		UMovieSceneSequence* Sequence = Sequencer->GetFocusedMovieSceneSequence();
+		if (!Sequence)
+		{
+			continue;
+		}
 
-		// Map between object binding nodes and selection
-		TMap<TSharedRef<FSequencerDisplayNode>, bool > ObjectBindingNodesSelectionMap;
+		// Gather a map of object bindings to their implict selection state
+		TMap<const FMovieSceneBinding*, bool> ObjectBindingNodesSelectionMap;
 
-		for (auto ObjectBinding : Sequencer->GetNodeTree()->GetObjectBindingMap() )
+		const FSequencerSelection& Selection = Sequencer->GetSelection();
+		const TSharedRef<FSequencerNodeTree>& NodeTree  = Sequencer->GetNodeTree();
+		for (const FMovieSceneBinding& Binding : Sequence->GetMovieScene()->GetBindings())
 		{
-			if (!ObjectBinding.Value.IsValid())
+			TSharedPtr<FSequencerObjectBindingNode> ObjectBindingNode = NodeTree->FindObjectBindingNode(Binding.GetObjectGuid());
+			if (!ObjectBindingNode.IsValid())
 			{
 				continue;
 			}
-			TSharedRef<FSequencerObjectBindingNode> ObjectBindingNode = ObjectBinding.Value.ToSharedRef();
-			bool bSelected = Sequencer->GetSelection().IsSelected(ObjectBindingNode);
 
-			if (!bSelected)
+			bool bSelected = false;
+			auto Traverse_IsSelected = [&Selection, &bSelected](FSequencerDisplayNode& InNode)
 			{
-				TSet<TSharedRef<FSequencerDisplayNode> > DescendantNodes;
-				SequencerHelpers::GetDescendantNodes(ObjectBindingNode, DescendantNodes);
-
-				// If one of our child is selected, we're considered selected
-				for (auto& DescendantNode : DescendantNodes)
+				TSharedRef<FSequencerDisplayNode> Shared = InNode.AsShared();
+				if (Selection.IsSelected(Shared) || Selection.NodeHasSelectedKeysOrSections(Shared))
 				{
-					if (Sequencer->GetSelection().IsSelected(DescendantNode) ||
-						Sequencer->GetSelection().NodeHasSelectedKeysOrSections(DescendantNode))
-					{
-						bSelected = true;
-						break;
-					}
+					bSelected = true;
+					// Stop traversing
+					return false;
 				}
-			}
+
+				return true;
+			};
+
+			ObjectBindingNode->Traverse_ParentFirst(Traverse_IsSelected, true);
 
 			// If one of our parent is selected, we're considered selected
 			TSharedPtr<FSequencerDisplayNode> ParentNode = ObjectBindingNode->GetParent();
-
 			while (!bSelected && ParentNode.IsValid())
 			{
-				if (Sequencer->GetSelection().IsSelected(ParentNode.ToSharedRef()) ||
-					Sequencer->GetSelection().NodeHasSelectedKeysOrSections(ParentNode.ToSharedRef()))
+				if (Selection.IsSelected(ParentNode.ToSharedRef()) || Selection.NodeHasSelectedKeysOrSections(ParentNode.ToSharedRef()))
 				{
 					bSelected = true;
 				}
@@ -735,50 +727,38 @@ void FSequencerEdMode::DrawTracks3D(FPrimitiveDrawInterface* PDI)
 				ParentNode = ParentNode->GetParent();
 			}
 
-			ObjectBindingNodesSelectionMap.Add(ObjectBindingNode, bSelected);
+			ObjectBindingNodesSelectionMap.Add(&Binding, bSelected);
 		}
 
 		// Gather up the transform track nodes from the object binding nodes
-		for (auto& ObjectBindingNode : ObjectBindingNodesSelectionMap)
+		for (TTuple<const FMovieSceneBinding*, bool>& Pair : ObjectBindingNodesSelectionMap)
 		{
-			FGuid ObjectBinding = StaticCastSharedRef<FSequencerObjectBindingNode>(ObjectBindingNode.Key)->GetObjectBinding();
-
-			TArray<TWeakObjectPtr<UObject>> BoundObjects;
-			for (TWeakObjectPtr<> Ptr : Sequencer->FindObjectsInCurrentSequence(ObjectBinding))
+			for (UMovieSceneTrack* Track : Pair.Key->GetTracks())
 			{
-				BoundObjects.Add(Ptr);
-			}
+				UMovieScene3DTransformTrack* TransformTrack = Cast<UMovieScene3DTransformTrack>(Track);
+				if (!TransformTrack)
+				{
+					continue;
+				}
 
-			for (auto& DisplayNode : ObjectBindingNode.Key.Get().GetChildNodes())
-			{
-				if (DisplayNode->GetType() == ESequencerNode::Track)
+				// Ensure that we've got a mesh trail for this track
+				if (bDrawMeshTrails)
 				{
-					TSharedRef<FSequencerTrackNode> TrackNode = StaticCastSharedRef<FSequencerTrackNode>(DisplayNode);
-					UMovieScene3DTransformTrack* TransformTrack = Cast<UMovieScene3DTransformTrack>(TrackNode->GetTrack());
-					if (TransformTrack != nullptr)
+					const bool bHasMeshTrail = Algo::FindBy(MeshTrails, TransformTrack, &FMeshTrailData::Track) != nullptr;
+					if (!bHasMeshTrail)
 					{
-						// If we are drawing mesh trails but we haven't made one for this track yet
-						if (bDrawMeshTrails)
+						UViewportWorldInteraction* WorldInteraction = Cast<UViewportWorldInteraction>( GEditor->GetEditorWorldExtensionsManager()->GetEditorWorldExtensions( GetWorld() )->FindExtension( UViewportWorldInteraction::StaticClass() ) );
+						if( WorldInteraction != nullptr )
 						{
-							FMeshTrailData* TrailPtr = MeshTrails.FindByPredicate([TransformTrack](const FMeshTrailData InTrail)
-							{
-								return InTrail.Track == TransformTrack;
-							});
-							if (TrailPtr == nullptr)
-							{
-								UViewportWorldInteraction* WorldInteraction = Cast<UViewportWorldInteraction>( GEditor->GetEditorWorldExtensionsManager()->GetEditorWorldExtensions( GetWorld() )->FindExtension( UViewportWorldInteraction::StaticClass() ) );
-								if( WorldInteraction != nullptr )
-								{
-									ASequencerMeshTrail* TrailActor = WorldInteraction->SpawnTransientSceneActor<ASequencerMeshTrail>(TEXT("SequencerMeshTrail"), true);
-									FMeshTrailData MeshTrail = FMeshTrailData(TransformTrack, TrailActor);
-									MeshTrails.Add(MeshTrail);
-								}
-							}
+							ASequencerMeshTrail* TrailActor = WorldInteraction->SpawnTransientSceneActor<ASequencerMeshTrail>(TEXT("SequencerMeshTrail"), true);
+							FMeshTrailData MeshTrail = FMeshTrailData(TransformTrack, TrailActor);
+							MeshTrails.Add(MeshTrail);
 						}
-
-						DrawTransformTrack(Sequencer, PDI, TransformTrack, BoundObjects, ObjectBindingNode.Value);
 					}
 				}
+
+				const bool bIsSelected = Pair.Value;
+				DrawTransformTrack(Sequencer, PDI, TransformTrack, Sequencer->FindObjectsInCurrentSequence(Pair.Key->GetObjectGuid()), bIsSelected);
 			}
 		}
 	}
diff --git a/Engine/Source/Editor/Sequencer/Private/SequencerEdMode.h b/Engine/Source/Editor/Sequencer/Private/SequencerEdMode.h
--- a/Engine/Source/Editor/Sequencer/Private/SequencerEdMode.h
+++ b/Engine/Source/Editor/Sequencer/Private/SequencerEdMode.h
@@ -71,7 +71,7 @@ public:
 
 protected:
 	void DrawTracks3D(FPrimitiveDrawInterface* PDI);
-	void DrawTransformTrack(const TSharedPtr<FSequencer>& Sequencer, FPrimitiveDrawInterface* PDI, UMovieScene3DTransformTrack* TransformTrack, const TArray<TWeakObjectPtr<UObject>>& BoundObjects, const bool bIsSelected);
+	void DrawTransformTrack(const TSharedPtr<FSequencer>& Sequencer, FPrimitiveDrawInterface* PDI, UMovieScene3DTransformTrack* TransformTrack, TArrayView<const TWeakObjectPtr<>> BoundObjects, const bool bIsSelected);
 
 protected:
 	static void GetLocationAtTime(FMovieSceneEvaluationTrack* Track, UObject* BoundObject, FFrameTime KeyTime, FVector& KeyPos, FRotator& KeyRot, const TSharedPtr<FSequencer>& Sequencer);
diff --git a/Engine/Source/Editor/Sequencer/Private/SequencerHotspots.cpp b/Engine/Source/Editor/Sequencer/Private/SequencerHotspots.cpp
--- a/Engine/Source/Editor/Sequencer/Private/SequencerHotspots.cpp
+++ b/Engine/Source/Editor/Sequencer/Private/SequencerHotspots.cpp
@@ -17,15 +17,6 @@
 
 #define LOCTEXT_NAMESPACE "SequencerHotspots"
 
-TSharedRef<ISequencerSection> FSectionHandle::GetSectionInterface() const
-{
-	return TrackNode->GetSections()[SectionIndex];
-}
-
-UMovieSceneSection* FSectionHandle::GetSectionObject() const
-{
-	return GetSectionInterface()->GetSectionObject();
-}
 
 void FKeyHotspot::UpdateOnHover(SSequencerTrackArea& InTrackArea, ISequencer& InSequencer) const
 {
@@ -55,19 +46,19 @@ bool FKeyHotspot::PopulateContextMenu(FMenuBuilder& MenuBuilder, ISequencer& InS
 
 TOptional<FFrameNumber> FSectionHotspot::GetTime() const
 {
-	UMovieSceneSection* ThisSection = Section.GetSectionObject();
+	UMovieSceneSection* ThisSection = WeakSection.Get();
 	return ThisSection && ThisSection->HasStartFrame() ? ThisSection->GetInclusiveStartFrame() : TOptional<FFrameNumber>();
 }
 
 TOptional<FFrameTime> FSectionHotspot::GetOffsetTime() const
 {
-	UMovieSceneSection* ThisSection = Section.GetSectionObject();
+	UMovieSceneSection* ThisSection = WeakSection.Get();
 	return ThisSection ? ThisSection->GetOffsetTime() : TOptional<FFrameTime>();
 }
 
 void FSectionHotspot::UpdateOnHover(SSequencerTrackArea& InTrackArea, ISequencer& InSequencer) const
 {
-	UMovieSceneSection* ThisSection = Section.GetSectionObject();
+	UMovieSceneSection* ThisSection = WeakSection.Get();
 
 	// Move sections if they are selected
 	if (InSequencer.GetSelection().IsSelected(ThisSection))
@@ -103,30 +94,26 @@ void FSectionHotspot::UpdateOnHover(SSequencerTrackArea& InTrackArea, ISequencer
 
 bool FSectionHotspot::PopulateContextMenu(FMenuBuilder& MenuBuilder, ISequencer& InSequencer, FFrameTime MouseDownTime)
 {
-	FSequencer& Sequencer = static_cast<FSequencer&>(InSequencer);
-
-	TSharedPtr<ISequencerSection> SectionInterface = Section.TrackNode->GetSections()[Section.SectionIndex];
-
-	FGuid ObjectBinding;
-	if (Section.TrackNode.IsValid())
+	UMovieSceneSection* ThisSection = WeakSection.Get();
+	if (ThisSection)
 	{
-		TSharedPtr<FSequencerObjectBindingNode> ObjectBindingNode = Section.TrackNode->FindParentObjectBindingNode();
-		if (ObjectBindingNode.IsValid())
+		FSequencer& Sequencer = static_cast<FSequencer&>(InSequencer);
+		FSectionContextMenu::BuildMenu(MenuBuilder, Sequencer, MouseDownTime);
+
+		TOptional<FSectionHandle> SectionHandle = Sequencer.GetNodeTree()->GetSectionHandle(ThisSection);
+		if (SectionHandle.IsSet())
 		{
-			ObjectBinding = ObjectBindingNode->GetObjectBinding();
+			FGuid ObjectBinding = SectionHandle->GetTrackNode()->GetObjectGuid();
+			SectionHandle->GetSectionInterface()->BuildSectionContextMenu(MenuBuilder, ObjectBinding);
 		}
 	}
 
-	FSectionContextMenu::BuildMenu(MenuBuilder, Sequencer, MouseDownTime);
-
-	SectionInterface->BuildSectionContextMenu(MenuBuilder, ObjectBinding);
-
 	return true;
 }
 
 TOptional<FFrameNumber> FSectionResizeHotspot::GetTime() const
 {
-	UMovieSceneSection* ThisSection = Section.GetSectionObject();
+	UMovieSceneSection* ThisSection = WeakSection.Get();
 	if (!ThisSection)
 	{
 		return TOptional<FFrameNumber>();
@@ -141,25 +128,22 @@ void FSectionResizeHotspot::UpdateOnHover(SSequencerTrackArea& InTrackArea, ISeq
 
 TSharedPtr<ISequencerEditToolDragOperation> FSectionResizeHotspot::InitiateDrag(ISequencer& Sequencer)
 {
-	const auto& SelectedSections = Sequencer.GetSelection().GetSelectedSections();
-	auto SectionHandles = StaticCastSharedRef<SSequencer>(Sequencer.GetSequencerWidget())->GetSectionHandles(SelectedSections);
-	
-	if (!SelectedSections.Contains(Section.GetSectionObject()))
+	UMovieSceneSection* ThisSection = WeakSection.Get();
+
+	if (ThisSection && !Sequencer.GetSelection().GetSelectedSections().Contains(ThisSection))
 	{
 		Sequencer.GetSelection().Empty();
-		Sequencer.GetSelection().AddToSelection(Section.GetSectionObject());
+		Sequencer.GetSelection().AddToSelection(ThisSection);
 		SequencerHelpers::UpdateHoveredNodeFromSelectedSections(static_cast<FSequencer&>(Sequencer));
-
-		SectionHandles.Empty();
-		SectionHandles.Add(Section);
 	}
+
 	const bool bIsSlipping = false;
-	return MakeShareable( new FResizeSection(static_cast<FSequencer&>(Sequencer), SectionHandles, HandleType == Right, bIsSlipping) );
+	return MakeShareable( new FResizeSection(static_cast<FSequencer&>(Sequencer), Sequencer.GetSelection().GetSelectedSections(), HandleType == Right, bIsSlipping) );
 }
 
 TOptional<FFrameNumber> FSectionEasingHandleHotspot::GetTime() const
 {
-	UMovieSceneSection* ThisSection = Section.GetSectionObject();
+	UMovieSceneSection* ThisSection = WeakSection.Get();
 	if (ThisSection)
 	{
 		if (HandleType == ESequencerEasingType::In && !ThisSection->GetEaseInRange().IsEmpty())
@@ -181,33 +165,31 @@ void FSectionEasingHandleHotspot::UpdateOnHover(SSequencerTrackArea& InTrackArea
 
 bool FSectionEasingHandleHotspot::PopulateContextMenu(FMenuBuilder& MenuBuilder, ISequencer& Sequencer, FFrameTime MouseDownTime)
 {
-	FEasingContextMenu::BuildMenu(MenuBuilder, { FEasingAreaHandle{Section, HandleType} }, static_cast<FSequencer&>(Sequencer), MouseDownTime);
+	FEasingContextMenu::BuildMenu(MenuBuilder, { FEasingAreaHandle{WeakSection, HandleType} }, static_cast<FSequencer&>(Sequencer), MouseDownTime);
 	return true;
 }
 
 TSharedPtr<ISequencerEditToolDragOperation> FSectionEasingHandleHotspot::InitiateDrag(ISequencer& Sequencer)
 {
-	return MakeShareable( new FManipulateSectionEasing(static_cast<FSequencer&>(Sequencer), Section, HandleType == ESequencerEasingType::In) );
+	return MakeShareable( new FManipulateSectionEasing(static_cast<FSequencer&>(Sequencer), WeakSection, HandleType == ESequencerEasingType::In) );
 }
 
-bool FSectionEasingAreaHotspot::PopulateContextMenu(FMenuBuilder& MenuBuilder, ISequencer& Sequencer, FFrameTime MouseDownTime)
+bool FSectionEasingAreaHotspot::PopulateContextMenu(FMenuBuilder& MenuBuilder, ISequencer& InSequencer, FFrameTime MouseDownTime)
 {
-	FEasingContextMenu::BuildMenu(MenuBuilder, Easings, static_cast<FSequencer&>(Sequencer), MouseDownTime);
-
-	TSharedPtr<ISequencerSection> SectionInterface = Section.TrackNode->GetSections()[Section.SectionIndex];
+	FSequencer& Sequencer = static_cast<FSequencer&>(InSequencer);
+	FEasingContextMenu::BuildMenu(MenuBuilder, Easings, Sequencer, MouseDownTime);
 
-	FGuid ObjectBinding;
-	if (Section.TrackNode.IsValid())
+	UMovieSceneSection* ThisSection = WeakSection.Get();
+	if (ThisSection)
 	{
-		TSharedPtr<FSequencerObjectBindingNode> ObjectBindingNode = Section.TrackNode->FindParentObjectBindingNode();
-		if (ObjectBindingNode.IsValid())
+		TOptional<FSectionHandle> SectionHandle = Sequencer.GetNodeTree()->GetSectionHandle(ThisSection);
+		if (SectionHandle.IsSet())
 		{
-			ObjectBinding = ObjectBindingNode->GetObjectBinding();
+			FGuid ObjectBinding = SectionHandle->GetTrackNode()->GetObjectGuid();
+			SectionHandle->GetSectionInterface()->BuildSectionContextMenu(MenuBuilder, ObjectBinding);
 		}
 	}
 
-	SectionInterface->BuildSectionContextMenu(MenuBuilder, ObjectBinding);
-
 	return true;
 }
 
diff --git a/Engine/Source/Editor/Sequencer/Private/SequencerHotspots.h b/Engine/Source/Editor/Sequencer/Private/SequencerHotspots.h
--- a/Engine/Source/Editor/Sequencer/Private/SequencerHotspots.h
+++ b/Engine/Source/Editor/Sequencer/Private/SequencerHotspots.h
@@ -36,33 +36,12 @@ struct FKeyHotspot
 };
 
 
-/** Structure used to encapsulate a section and its track node */
-struct FSectionHandle
-{
-	FSectionHandle(TSharedPtr<FSequencerTrackNode> InTrackNode, int32 InSectionIndex)
-		: SectionIndex(InSectionIndex), TrackNode(MoveTemp(InTrackNode))
-	{ }
-
-	friend bool operator==(const FSectionHandle& A, const FSectionHandle& B)
-	{
-		return A.SectionIndex == B.SectionIndex && A.TrackNode == B.TrackNode;
-	}
-	
-	SEQUENCER_API TSharedRef<ISequencerSection> GetSectionInterface() const;
-
-	SEQUENCER_API UMovieSceneSection* GetSectionObject() const;
-
-	int32 SectionIndex;
-	TSharedPtr<FSequencerTrackNode> TrackNode;
-};
-
-
 /** A hotspot representing a section */
 struct FSectionHotspot
 	: ISequencerHotspot
 {
-	FSectionHotspot(FSectionHandle InSection)
-		: Section(InSection)
+	FSectionHotspot(UMovieSceneSection* InSection)
+		: WeakSection(InSection)
 	{ }
 
 	virtual ESequencerHotspot GetType() const override { return ESequencerHotspot::Section; }
@@ -72,8 +51,8 @@ struct FSectionHotspot
 	virtual TSharedPtr<ISequencerEditToolDragOperation> InitiateDrag(ISequencer&) override { return nullptr; }
 	virtual bool PopulateContextMenu(FMenuBuilder& MenuBuilder, ISequencer& Sequencer, FFrameTime MouseDownTime) override;
 
-	/** Handle to the section */
-	FSectionHandle Section;
+	/** The section */
+	TWeakObjectPtr<UMovieSceneSection> WeakSection;
 };
 
 
@@ -87,7 +66,7 @@ struct FSectionResizeHotspot
 		Right
 	};
 
-	FSectionResizeHotspot(EHandle InHandleType, FSectionHandle InSection) : Section(InSection), HandleType(InHandleType) {}
+	FSectionResizeHotspot(EHandle InHandleType, UMovieSceneSection* InSection) : WeakSection(InSection), HandleType(InHandleType) {}
 
 	virtual ESequencerHotspot GetType() const override { return HandleType == Left ? ESequencerHotspot::SectionResize_L : ESequencerHotspot::SectionResize_R; }
 	virtual void UpdateOnHover(SSequencerTrackArea& InTrackArea, ISequencer& InSequencer) const override;
@@ -95,8 +74,8 @@ struct FSectionResizeHotspot
 	virtual TSharedPtr<ISequencerEditToolDragOperation> InitiateDrag(ISequencer& Sequencer) override;
 	virtual FCursorReply GetCursor() const { return FCursorReply::Cursor( EMouseCursor::ResizeLeftRight ); }
 
-	/** Handle to the section */
-	FSectionHandle Section;
+	/** The section */
+	TWeakObjectPtr<UMovieSceneSection> WeakSection;
 
 private:
 
@@ -108,7 +87,7 @@ private:
 struct FSectionEasingHandleHotspot
 	: ISequencerHotspot
 {
-	FSectionEasingHandleHotspot(ESequencerEasingType InHandleType, FSectionHandle InSection) : Section(InSection), HandleType(InHandleType) {}
+	FSectionEasingHandleHotspot(ESequencerEasingType InHandleType, UMovieSceneSection* InSection) : WeakSection(InSection), HandleType(InHandleType) {}
 
 	virtual ESequencerHotspot GetType() const override { return HandleType == ESequencerEasingType::In ? ESequencerHotspot::EaseInHandle : ESequencerHotspot::EaseOutHandle; }
 	virtual void UpdateOnHover(SSequencerTrackArea& InTrackArea, ISequencer& InSequencer) const override;
@@ -118,7 +97,7 @@ struct FSectionEasingHandleHotspot
 	virtual FCursorReply GetCursor() const { return FCursorReply::Cursor( EMouseCursor::ResizeLeftRight ); }
 
 	/** Handle to the section */
-	FSectionHandle Section;
+	TWeakObjectPtr<UMovieSceneSection> WeakSection;
 
 private:
 
@@ -128,7 +107,7 @@ private:
 
 struct FEasingAreaHandle
 {
-	FSectionHandle Section;
+	TWeakObjectPtr<UMovieSceneSection> WeakSection;
 	ESequencerEasingType EasingType;
 };
 
@@ -136,12 +115,12 @@ struct FEasingAreaHandle
 struct FSectionEasingAreaHotspot
 	: FSectionHotspot
 {
-	FSectionEasingAreaHotspot(const TArray<FEasingAreaHandle>& InEasings, FSectionHandle InVisibleSection) : FSectionHotspot(InVisibleSection), Easings(InEasings) {}
+	FSectionEasingAreaHotspot(const TArray<FEasingAreaHandle>& InEasings, UMovieSceneSection* InVisibleSection) : FSectionHotspot(InVisibleSection), Easings(InEasings) {}
 
 	virtual ESequencerHotspot GetType() const override { return ESequencerHotspot::EasingArea; }
 	virtual bool PopulateContextMenu(FMenuBuilder& MenuBuilder, ISequencer& Sequencer, FFrameTime MouseDownTime) override;
 
-	bool Contains(FSectionHandle InSection) const { return Easings.ContainsByPredicate([=](const FEasingAreaHandle& InHandle){ return InHandle.Section == InSection; }); }
+	bool Contains(UMovieSceneSection* InSection) const { return Easings.ContainsByPredicate([=](const FEasingAreaHandle& InHandle){ return InHandle.WeakSection == InSection; }); }
 
 	/** Handles to the easings that exist on this hotspot */
 	TArray<FEasingAreaHandle> Easings;
diff --git a/Engine/Source/Editor/Sequencer/Private/SequencerNodeSortingMethods.h b/Engine/Source/Editor/Sequencer/Private/SequencerNodeSortingMethods.h
--- a/Engine/Source/Editor/Sequencer/Private/SequencerNodeSortingMethods.h
+++ b/Engine/Source/Editor/Sequencer/Private/SequencerNodeSortingMethods.h
@@ -3,6 +3,7 @@
 #pragma once
 #include "CoreMinimal.h"
 #include "DisplayNodes/SequencerDisplayNode.h"
+#include "Widgets/Views/STableRow.h"
 
 /**
 * This sorting delegate sorts based on the vertical position of the node represented by where it falls
@@ -17,108 +18,6 @@ struct FDisplayNodeTreePositionSorter
 	}
 };
 
-/**
-* This sorting delegate sorts based on category and then alphabetically after that. This replicates
-* existing behavior where all folders come first, then all tracks, and finally all object bindings.
-* Within each category these are sorted by display name alphabetically.
-*
-* This does not respect the user-defined sorting order.
-*/
-struct FDisplayNodeCategoricalSorter
-{
-	int32 NodeTypeToFolderSortId(ESequencerNode::Type NodeType) const
-	{
-		switch (NodeType)
-		{
-		case ESequencerNode::Folder:
-			return 0;
-		case ESequencerNode::Track:
-			return 1;
-		case ESequencerNode::Object:
-			return 2;
-		default:
-			return 3;
-		}
-	};
-
-	int32 NodeTypeToObjectSortId(ESequencerNode::Type NodeType) const
-	{
-		switch (NodeType)
-		{
-		case ESequencerNode::Object:
-			return 0;
-		case ESequencerNode::Track:
-			return 1;
-		default:
-			return 2;
-		}
-	};
-
-	bool operator()(const TSharedRef<FSequencerDisplayNode>& A, const TSharedRef<FSequencerDisplayNode>& B) const
-	{
-		TSharedPtr<FSequencerDisplayNode> ParentNode = A->GetParent();
-
-		// If the nodes are root nodes, or in folders and they are the same type, sort by name.
-		if ((ParentNode.IsValid() == false || ParentNode->GetType() == ESequencerNode::Folder) && A->GetType() == B->GetType())
-		{
-			return A->GetDisplayName().ToString() < B->GetDisplayName().ToString();
-		}
-
-		int32 SortIdA;
-		int32 SortIdB;
-
-		// Otherwise if they are root nodes or in folders use the folder sort id.
-		if (ParentNode.IsValid() == false || ParentNode->GetType() == ESequencerNode::Folder)
-		{
-			SortIdA = NodeTypeToFolderSortId(A->GetType());
-			SortIdB = NodeTypeToFolderSortId(B->GetType());
-		}
-		// Otherwise if they are in an object node use the object node sort id.
-		else if (ParentNode->GetType() == ESequencerNode::Object)
-		{
-			SortIdA = NodeTypeToObjectSortId(A->GetType());
-			SortIdB = NodeTypeToObjectSortId(B->GetType());
-		}
-		// Otherwise they are equal, and in a stable sort shouldn't change position.
-		else
-		{
-			SortIdA = 0;
-			SortIdB = 0;
-		}
-
-		return SortIdA < SortIdB;
-	}
-};
-
-
-/**
-* This sorting delegate sorts based on the sorting order specified by the node. This sorter falls back to the 
-* FDisplayNodeCategoricalSorter for Object nodes because Object nodes do not support sorting order for child nodes.
-*/
-struct FDisplayNodeSortingOrderSorter
-{
-	bool operator()(const TSharedRef<FSequencerDisplayNode>& A, const TSharedRef<FSequencerDisplayNode>& B) const
-	{
-		// If the parent of the object is a Object Binding then we use legacy sorting as we don't allow the user
-		// to reorder tracks within object binding nodes. 
-		TSharedPtr<FSequencerDisplayNode> ParentNode = A->GetParent();
-		if (ParentNode.IsValid() && ParentNode->GetType() == ESequencerNode::Object)
-		{
-			FDisplayNodeCategoricalSorter CategoricalSorter;
-			return CategoricalSorter(A, B);
-		}
-		
-		// If both nodes have been sorted before we just compare their sorting orders.
-		if (A->GetSortingOrder() >= 0 && B->GetSortingOrder() >= 0)
-		{
-			return A->GetSortingOrder() < B->GetSortingOrder();
-		}
-
-		// Otherwise if one of them has not placed before, we want the lower number higher so that
-		// the unsorted node shows up at the end of the list.
-		return A->GetSortingOrder() > B->GetSortingOrder();
-	}
-};
 
 /**
 * Sorts the supplied unsorted nodes and inserts them into the existing sorted nodes before assigning
