This is a significant overhaul to the Curve Editor used by Sequencer which adds a plugin-based architecture and extensibility. New tools and toolbar buttons can be added to all usages of the curve edi

diff --git a/Engine/Source/Editor/Sequencer/Private/SequencerNodeTree.cpp b/Engine/Source/Editor/Sequencer/Private/SequencerNodeTree.cpp
--- a/Engine/Source/Editor/Sequencer/Private/SequencerNodeTree.cpp
+++ b/Engine/Source/Editor/Sequencer/Private/SequencerNodeTree.cpp
@@ -13,425 +13,344 @@
 #include "Tracks/MovieSceneCinematicShotTrack.h"
 #include "Sequencer.h"
 #include "MovieSceneFolder.h"
-#include "SequencerSectionLayoutBuilder.h"
 #include "ISequencerTrackEditor.h"
-#include "DisplayNodes/SequencerSpacerNode.h"
+#include "DisplayNodes/SequencerRootNode.h"
 #include "Widgets/Views/STableRow.h"
+#include "CurveEditor.h"
 #include "SequencerNodeSortingMethods.h"
 
-void FSequencerNodeTree::Empty()
-{
-	RootNodes.Empty();
-	ObjectBindingMap.Empty();
-	Sequencer.GetSelection().EmptySelectedOutlinerNodes();
-	EditorMap.Empty();
-	FilteredNodes.Empty();
-	HoveredNode = nullptr;
-}
 
+FSequencerNodeTree::FSequencerNodeTree(FSequencer& InSequencer)
+	: RootNode(MakeShared<FSequencerRootNode>(*this))
+	, SerialNumber(0)
+	, Sequencer(InSequencer)
+{}
 
-void FSequencerNodeTree::AddObjectBindingAndTracks(const FMovieSceneBinding& Binding, TMap<FGuid, const FMovieSceneBinding*>& GuidToBindingMap, TArray< TSharedRef<FSequencerObjectBindingNode> >& OutNodeList)
+TSharedPtr<FSequencerObjectBindingNode> FSequencerNodeTree::FindObjectBindingNode(const FGuid& BindingID) const
 {
-	TSharedRef<FSequencerObjectBindingNode> ObjectBindingNode = AddObjectBinding( Binding.GetName(), Binding.GetObjectGuid(), GuidToBindingMap, OutNodeList );
-
-	for( UMovieSceneTrack* Track : Binding.GetTracks() )
-	{
-		if (!Sequencer.IsTrackVisible(Track))
-		{
-			continue;
-		}
-
-		// Create the new track node
-		TSharedRef<FSequencerTrackNode> TrackNode = MakeShared<FSequencerTrackNode>(*Track, *FindOrAddTypeEditor(*Track), false, nullptr, *this);
-
-		// Make the sub tracks and section interfaces for this node, and add it to the object binding node
-		// Note: MakeSubTracksAndSectionInterfaces may return a new parent node
-		ObjectBindingNode->AddTrackNode(MakeSubTracksAndSectionInterfaces(TrackNode, ObjectBindingNode->GetObjectBinding()));
-	}
+	return ObjectBindingToNode.FindRef(BindingID);
 }
 
-
-void FSequencerNodeTree::Update()
+void FSequencerNodeTree::RefreshNodes(UMovieScene* MovieScene)
 {
-	HoveredNode = nullptr;
-
-	// @todo Sequencer - This update pass is too aggressive.  Some nodes may still be valid
-	Empty();
+	check(MovieScene);
 
-	UMovieScene* MovieScene = Sequencer.GetFocusedMovieSceneSequence()->GetMovieScene();
-	UMovieSceneCinematicShotTrack* CinematicShotTrack = MovieScene->FindMasterTrack<UMovieSceneCinematicShotTrack>();
+	++SerialNumber;
 
-	// Get the master tracks  so we can get sections from them
-	const TArray<UMovieSceneTrack*>& MasterTracks = MovieScene->GetMasterTracks();
-	TArray<TSharedRef<FSequencerTrackNode>> MasterTrackNodes;
+	TSortedMap<FGuid, FGuid> ChildToParentBinding;
+	TSortedMap<FGuid, const FMovieSceneBinding*> AllBindings;
 
-	for (UMovieSceneTrack* Track : MasterTracks)
+	// Gather all object bindings in the sequence
+	for (const FMovieSceneBinding& Binding : MovieScene->GetBindings())
 	{
-		if (Track != CinematicShotTrack)
-		{
-			UMovieSceneTrack& TrackRef = *Track;
-
-			TSharedRef<FSequencerTrackNode> NodeToAdd = MakeSubTracksAndSectionInterfaces(MakeShared<FSequencerTrackNode>(TrackRef, *FindOrAddTypeEditor(TrackRef), true, nullptr, *this));
-			MasterTrackNodes.Add(NodeToAdd);
-		}
+		AllBindings.Add(Binding.GetObjectGuid(), &Binding);
 	}
 
-	const TArray<FMovieSceneBinding>& Bindings = MovieScene->GetBindings();
-	TMap<FGuid, const FMovieSceneBinding*> GuidToBindingMap;
-
-	for (const FMovieSceneBinding& Binding : Bindings)
+	// Populate the binding hierarchy
+	for (int32 PossessableIndex = 0; PossessableIndex < MovieScene->GetPossessableCount(); ++PossessableIndex)
 	{
-		GuidToBindingMap.Add(Binding.GetObjectGuid(), &Binding);
-	}
+		const FMovieScenePossessable& Possessable = MovieScene->GetPossessable(PossessableIndex);
 
-	// Make nodes for all object bindings
-	TArray<TSharedRef<FSequencerObjectBindingNode>> ObjectNodes;
-	for( const FMovieSceneBinding& Binding : Bindings )
-	{
-		if (!Sequencer.IsBindingVisible(Binding))
+		FGuid ThisID   = Possessable.GetGuid();
+		FGuid ParentID = Possessable.GetParent();
+
+		if (ParentID.IsValid())
 		{
-			continue;
+			ChildToParentBinding.Add(ThisID, ParentID);
 		}
-
-		AddObjectBindingAndTracks(Binding, GuidToBindingMap, ObjectNodes);
 	}
 
-	// If no bindings were added (presumably because of visibility) but there are bindings, add all regardless of visibility
-	if (!ObjectNodes.Num())
+	// Folders may also create hierarchy items for tracks and object bindings
 	{
-		for( const FMovieSceneBinding& Binding : Bindings )
+		for (UMovieSceneFolder* Folder : MovieScene->GetRootFolders())
 		{
-			AddObjectBindingAndTracks(Binding, GuidToBindingMap, ObjectNodes);
+			if (ensureAlwaysMsgf(Folder, TEXT("MovieScene data contains a null folder. This should never happen.")))
+			{
+				TSharedRef<FSequencerFolderNode> RootFolderNode = CreateOrUpdateFolder(Folder, AllBindings, ChildToParentBinding, MovieScene);
+				RootFolderNode->SetParent(RootNode);
+			}
 		}
 	}
 
-	// Cinematic shot track always comes first
-	if (CinematicShotTrack)
-	{
-		TSharedRef<FSequencerTrackNode> NodeToAdd = MakeSubTracksAndSectionInterfaces(MakeShared<FSequencerTrackNode>(*CinematicShotTrack, *FindOrAddTypeEditor(*CinematicShotTrack), false, nullptr, *this));
-		RootNodes.Add(NodeToAdd);
-	}
-
-	// Then comes the camera cut track
-	UMovieSceneTrack* CameraCutTrack = MovieScene->GetCameraCutTrack();
-	
-	if (CameraCutTrack)
+	// Object Bindings
 	{
-		TSharedRef<FSequencerTrackNode> NodeToAdd = MakeSubTracksAndSectionInterfaces(MakeShared<FSequencerTrackNode>(*CameraCutTrack, *FindOrAddTypeEditor(*CameraCutTrack), false, nullptr, *this));
-		RootNodes.Add(NodeToAdd);
-	}
-
-	// Add all other nodes after the camera cut track
-	TArray<TSharedRef<FSequencerDisplayNode>> FolderAndObjectAndTrackNodes;
-	TArray<TSharedRef<FSequencerDisplayNode>> MasterTrackNodesNotInFolders;
-	CreateAndPopulateFolderNodes( MasterTrackNodes, ObjectNodes, MovieScene->GetRootFolders(), FolderAndObjectAndTrackNodes, MasterTrackNodesNotInFolders );
-
-	// Merge the two lists together before sorting them together.
-	FolderAndObjectAndTrackNodes.Append(MasterTrackNodesNotInFolders);
+		for (const FMovieSceneBinding& Binding : MovieScene->GetBindings())
+		{
+			TSharedPtr<FSequencerObjectBindingNode> ObjectBindingNode = CreateOrUpdateObjectBinding(Binding.GetObjectGuid(), AllBindings, ChildToParentBinding, MovieScene);
+			if (!ObjectBindingNode.IsValid())
+			{
+				continue;
+			}
 
-	// Now sort the folders, tracks and objects together based on sorting order.
-	FolderAndObjectAndTrackNodes.Sort(FDisplayNodeSortingOrderSorter());
+			// Ensure it has a parent
+			if (!ObjectBindingNode->IsParentStillRelevant(SerialNumber))
+			{
+				ObjectBindingNode->SetParent(RootNode);
+			}
 
-	for (TSharedRef<FSequencerDisplayNode> Node : FolderAndObjectAndTrackNodes)
-	{
-		// Recursively sort the children of these tracks
-		Node->SortChildNodes(FDisplayNodeSortingOrderSorter());
+			// Create nodes for the object binding's tracks
+			for (UMovieSceneTrack* Track : Binding.GetTracks())
+			{
+				if (ensureAlwaysMsgf(Track, TEXT("MovieScene binding '%s' data contains a null track. This should never happen."), *Binding.GetName()))
+				{
+					TSharedPtr<FSequencerTrackNode> TrackNode = CreateOrUpdateTrack(Track, ETrackType::Object);
+					if (TrackNode.IsValid())
+					{
+						TrackNode->SetParent(ObjectBindingNode);
+					}
+				}
+			}
+		}
 	}
 
-	// Now that we've sorted the children we normalize their sorting index. This doesn't call Modify (as we're not part of a transaction) but modifies the
-	// in-memory sorting index of the backing data structures. This means the next time the tree is refreshed, the existing nodes will keep their sort
-	// and any new nodes will get pushed to the end again. When an asset is saved it'll write the sorting index to the asset and the next time it is loaded
-	// the Sort function will keep them in the same order even if they exist in a different order within the owning data structures.
-	for (int32 i = 0; i < FolderAndObjectAndTrackNodes.Num(); i++)
+	// Master tracks
 	{
-		FolderAndObjectAndTrackNodes[i]->SetSortingOrder(i);
-		FolderAndObjectAndTrackNodes[i]->Traverse_ParentFirst([](FSequencerDisplayNode& TraversalNode)
+		UMovieSceneTrack* CameraCutTrack = MovieScene->GetCameraCutTrack();
+		if (CameraCutTrack)
 		{
-			int32 ChildIndex = 0;
-			for (uint32 k = 0; k < TraversalNode.GetNumChildren(); k++)
+			TSharedPtr<FSequencerTrackNode> TrackNode = CreateOrUpdateTrack(CameraCutTrack, ETrackType::Master);
+			if (TrackNode.IsValid() && !TrackNode->IsParentStillRelevant(SerialNumber))
 			{
-				// Sometimes a node can have multiple display node children because the sections within a row
-				// have been re-arranged into an overlapping state. These rows are backed by the same data structure
-				// as their parents, so we skip them instead of incrementing the parent's sorting order.
-				if (TraversalNode.GetChildNodes()[k]->GetType() == ESequencerNode::Track)
+				TrackNode->SetParent(RootNode);
+			}
+		}
+
+		// Iterate all master tracks and generate nodes if necessary
+		for (UMovieSceneTrack* Track : MovieScene->GetMasterTracks())
+		{
+			if (ensureAlwaysMsgf(Track, TEXT("MovieScene data contains a null master track. This should never happen.")))
+			{
+				TSharedPtr<FSequencerTrackNode> TrackNode = CreateOrUpdateTrack(Track, ETrackType::Master);
+				if (TrackNode.IsValid() && !TrackNode->IsParentStillRelevant(SerialNumber))
 				{
-					TSharedRef<FSequencerTrackNode> FolderNode = StaticCastSharedRef<FSequencerTrackNode>(TraversalNode.GetChildNodes()[k]);
-					if (FolderNode->GetSubTrackMode() == FSequencerTrackNode::ESubTrackMode::SubTrack)
-					{
-						continue;
-					}
+					TrackNode->SetParent(RootNode);
 				}
-
-				TraversalNode.GetChildNodes()[k]->SetSortingOrder(ChildIndex);
-				ChildIndex++;
 			}
-			return true;
-		}, true);
+		}
 	}
 
-	RootNodes.Append(FolderAndObjectAndTrackNodes);
-	RootNodes.Reserve(FMath::Max(1, RootNodes.Num() - 1) * 2);
-	for (int32 Index = 1; Index < RootNodes.Num(); Index += 2)
+	// Remove anything that is no longer relevant (ie serial number is out of date)
+	for (auto It = FolderToNode.CreateIterator(); It; ++It)
 	{
-		RootNodes.Insert(MakeShareable(new FSequencerSpacerNode(1.f, nullptr, *this, false)), Index);
+		if (It->Value->TreeSerialNumber != SerialNumber)
+		{
+			It->Value->SetParent(nullptr);
+			It.RemoveCurrent();
+		}
 	}
-
-	// Always make space at the end of the tree
-	RootNodes.Add(MakeShared<FSequencerSpacerNode>(20.f, nullptr, *this, true));
-
-	// Set up virtual offsets, expansion states, and tints
-	float VerticalOffset = 0.f;
-
-	for (TSharedRef<FSequencerDisplayNode>& Node : RootNodes)
+	for (auto It = TrackToNode.CreateIterator(); It; ++It)
 	{
-		Node->Traverse_ParentFirst([&](FSequencerDisplayNode& InNode) {
-
-			// Set up the virtual node position
-			float VerticalTop = VerticalOffset;
-			VerticalOffset += InNode.GetNodeHeight() + InNode.GetNodePadding().Combined();
-			InNode.Initialize(VerticalTop, VerticalOffset);
-
-			return true;
-		});
+		if (It->Value->TreeSerialNumber != SerialNumber)
+		{
+			It->Value->SetParent(nullptr);
+			It.RemoveCurrent();
+		}
+	}
+	for (auto It = ObjectBindingToNode.CreateIterator(); It; ++It)
+	{
+		if (It->Value->TreeSerialNumber != SerialNumber)
+		{
+			It->Value->SetParent(nullptr);
+			It.RemoveCurrent();
+		}
 	}
-
-	// Re-filter the tree after updating 
-	// @todo sequencer: Newly added sections may need to be visible even when there is a filter
-	FilterNodes( FilterString );
-
-	OnUpdatedDelegate.Broadcast();
 }
 
-
-TSharedRef<ISequencerTrackEditor> FSequencerNodeTree::FindOrAddTypeEditor( UMovieSceneTrack& InTrack )
+TSharedPtr<FSequencerTrackNode> FSequencerNodeTree::CreateOrUpdateTrack(UMovieSceneTrack* Track, ETrackType TrackType)
 {
-	TSharedPtr<ISequencerTrackEditor> Editor = EditorMap.FindRef( &InTrack );
+	check(Track);
 
-	if( !Editor.IsValid() )
+	FObjectKey TrackKey(Track);
+	TSharedPtr<FSequencerTrackNode> TrackNode = TrackToNode.FindRef(TrackKey);
+	if (TrackNode.IsValid())
 	{
-		const TArray<TSharedPtr<ISequencerTrackEditor>>& TrackEditors = Sequencer.GetTrackEditors();
-
-		// Get a tool for each track
-		// @todo sequencer: Should probably only need to get this once and it shouldn't be done here. It depends on when movie scene tool modules are loaded
-		TSharedPtr<ISequencerTrackEditor> SupportedTool;
-
-		for (const auto& TrackEditor : TrackEditors)
+		// Should be implemented as a filter
+		if (!Sequencer.IsTrackVisible(Track))
 		{
-			if (TrackEditor->SupportsType(InTrack.GetClass()))
-			{
-				EditorMap.Add(&InTrack, TrackEditor);
-				Editor = TrackEditor;
-
-				break;
-			}
+			TrackNode->SetParent(nullptr);
+			TrackToNode.Remove(TrackKey);
+			return nullptr;
 		}
 	}
+	else
+	{
+		const bool bIsDraggable = TrackType == ETrackType::Master;
+		TrackNode = MakeShared<FSequencerTrackNode>(Track, *FindOrAddTypeEditor(Track), bIsDraggable, *this);
+		TrackToNode.Add(TrackKey, TrackNode);
+	}
 
-	return Editor.ToSharedRef();
+	// Assign the serial number for this node to indicate that it is still relevant
+	TrackNode->TreeSerialNumber = SerialNumber;
+	TrackNode->UpdateInnerHierarchy();
+	return TrackNode;
 }
 
-
-TSharedRef<FSequencerTrackNode> FSequencerNodeTree::MakeSubTracksAndSectionInterfaces(TSharedRef<FSequencerTrackNode> TrackNode, const FGuid& ObjectBinding)
+TSharedRef<FSequencerFolderNode> FSequencerNodeTree::CreateOrUpdateFolder(UMovieSceneFolder* Folder, const TSortedMap<FGuid, const FMovieSceneBinding*>& AllBindings, const TSortedMap<FGuid, FGuid>& ChildToParentBinding, const UMovieScene* InMovieScene)
 {
-	using ESubTrackMode = FSequencerTrackNode::ESubTrackMode;
+	check(Folder);
 
-	UMovieSceneTrack* Track = TrackNode->GetTrack();
-
-	check(Track);
-	check(!TrackNode->GetParent().IsValid());
+	FObjectKey FolderKey(Folder);
 
-	TArray<UMovieSceneSection*> MovieSceneSections = Track->GetAllSections();
-	if (MovieSceneSections.Num() == 0)
+	TSharedPtr<FSequencerFolderNode> FolderNode = FolderToNode.FindRef(FolderKey);
+	if (!FolderNode.IsValid())
 	{
-		return TrackNode;
+		FolderNode = MakeShared<FSequencerFolderNode>(*Folder, *this);
+		FolderToNode.Add(FolderKey, FolderNode.ToSharedRef());
 	}
 
-	Algo::SortBy(MovieSceneSections, &UMovieSceneSection::GetRowIndex);
+	// Assign the serial number for this node to indicate that it is still relevant
+	FolderNode->TreeSerialNumber = SerialNumber;
 
-	const bool bHasMultipleRows = MovieSceneSections.Last()->GetRowIndex() != 0;
-
-	TSharedRef<ISequencerTrackEditor> Editor = FindOrAddTypeEditor( *Track );
-
-	TArray<TSharedRef<FSequencerTrackNode>, TInlineAllocator<4>> SubTrackNodes;
-
-	TSharedRef<FSequencerTrackNode> ParentNode = TrackNode;
-	TSharedRef<FSequencerTrackNode> CurrentTrackNode = TrackNode;
-
-	for (int32 SectionIndex = 0; SectionIndex < MovieSceneSections.Num(); ++SectionIndex)
+	// Create the hierarchy for any child bindings
+	for (const FGuid& ID : Folder->GetChildObjectBindings())
 	{
-		UMovieSceneSection* SectionObject = MovieSceneSections[SectionIndex];
-		const int32 RowIndex = SectionObject->GetRowIndex();
-
-		if (CurrentTrackNode->GetSubTrackMode() == ESubTrackMode::SubTrack && RowIndex != CurrentTrackNode->GetRowIndex())
+		TSharedPtr<FSequencerObjectBindingNode> Binding = CreateOrUpdateObjectBinding(ID, AllBindings, ChildToParentBinding, InMovieScene);
+		if (Binding.IsValid())
 		{
-			CurrentTrackNode = MakeShared<FSequencerTrackNode>(*Track, *Editor, ParentNode->CanDrag(), ParentNode, *this);
-
-			CurrentTrackNode->SetSubTrackMode(ESubTrackMode::SubTrack);
-			CurrentTrackNode->SetRowIndex(RowIndex);
-			ParentNode->AddChildTrack(CurrentTrackNode);
+			Binding->SetParent(FolderNode);
 		}
+	}
 
-		// Make the section interface
-		TSharedRef<ISequencerSection> Section = Editor->MakeSectionInterface(*SectionObject, *Track, ObjectBinding);
-
-		// Ask the section to generate its inner layout
-		FSequencerSectionLayoutBuilder Builder(CurrentTrackNode);
-		Section->GenerateSectionLayout(Builder);
-
-		if (bHasMultipleRows && CurrentTrackNode == ParentNode)
+	// Create the hierarchy for any master tracks
+	for (UMovieSceneTrack* Track : Folder->GetChildMasterTracks())
+	{
+		if (ensureAlwaysMsgf(Track, TEXT("MovieScene folder '%s' data contains a null track. This should never happen."), *Folder->GetName()))
 		{
-			// Create a new parent node
-			ParentNode = MakeShared<FSequencerTrackNode>(*Track, *Editor, CurrentTrackNode->CanDrag(), nullptr, *this);
-			ParentNode->SetSubTrackMode(ESubTrackMode::ParentTrack);
-
-			CurrentTrackNode->SetSubTrackMode(ESubTrackMode::SubTrack);
-			CurrentTrackNode->SetRowIndex(RowIndex);
-			ParentNode->AddChildTrack(CurrentTrackNode);
+			TSharedPtr<FSequencerTrackNode> TrackNode = CreateOrUpdateTrack(Track, ETrackType::Master);
+			if (TrackNode.IsValid())
+			{
+				TrackNode->SetParent(FolderNode);
+			}
 		}
-
-		CurrentTrackNode->AddSection(Section);
 	}
 
-	return ParentNode;
-}
-
+	// Add child folders
+	for (UMovieSceneFolder* ChildFolder : Folder->GetChildFolders())
+	{
+		if (ensureAlwaysMsgf(ChildFolder, TEXT("MovieScene folder '%s' data contains a null child folder. This should never happen."), *Folder->GetName()))
+		{
+			TSharedRef<FSequencerFolderNode> ChildFolderNode = CreateOrUpdateFolder(ChildFolder, AllBindings, ChildToParentBinding, InMovieScene);
+			ChildFolderNode->SetParent(FolderNode);
+		}
+	}
 
-const TArray<TSharedRef<FSequencerDisplayNode>>& FSequencerNodeTree::GetRootNodes() const
-{
-	return RootNodes;
+	return FolderNode.ToSharedRef();
 }
 
-
-TSharedRef<FSequencerObjectBindingNode> FSequencerNodeTree::AddObjectBinding(const FString& ObjectName, const FGuid& ObjectBinding, TMap<FGuid, const FMovieSceneBinding*>& GuidToBindingMap, TArray<TSharedRef<FSequencerObjectBindingNode>>& OutNodeList)
+TSharedPtr<FSequencerObjectBindingNode> FSequencerNodeTree::CreateOrUpdateObjectBinding(const FGuid& BindingID, const TSortedMap<FGuid, const FMovieSceneBinding*>& AllBindings, const TSortedMap<FGuid, FGuid>& ChildToParentBinding, const UMovieScene* InMovieScene)
 {
-	TSharedPtr<FSequencerObjectBindingNode> ObjectNode;
-	TSharedPtr<FSequencerObjectBindingNode>* FoundObjectNode = ObjectBindingMap.Find(ObjectBinding);
-	if (FoundObjectNode != nullptr)
+	if (!ensureAlwaysMsgf(AllBindings.Contains(BindingID), TEXT("Attempting to add a binding that does not exist.")))
 	{
-		ObjectNode = *FoundObjectNode;
+		return nullptr;
 	}
-	else
+
+	TSharedPtr<FSequencerObjectBindingNode> ObjectBindingNode = ObjectBindingToNode.FindRef(BindingID);
+	if (!ObjectBindingNode.IsValid())
 	{
 		// The node name is the object guid
-		FName ObjectNodeName = *ObjectBinding.ToString();
+		FName ObjectNodeName = *BindingID.ToString();
 
-		// Try to get the parent object node if there is one.
-		TSharedPtr<FSequencerObjectBindingNode> ParentNode;
+		ObjectBindingNode = MakeShared<FSequencerObjectBindingNode>(ObjectNodeName, BindingID, *this);
+		ObjectBindingToNode.Add(BindingID, ObjectBindingNode);
+	}
 
-		UMovieSceneSequence* Sequence = Sequencer.GetFocusedMovieSceneSequence();
+	// Assign the serial number for this node to indicate that it is still relevant
+	ObjectBindingNode->TreeSerialNumber = SerialNumber;
 
-		// Prefer to use the parent spawnable if possible, rather than relying on runtime object presence
-		FMovieScenePossessable* Possessable = Sequence->GetMovieScene()->FindPossessable(ObjectBinding);
-		if (Possessable && Possessable->GetParent().IsValid())
+	// Create its parent and make the association
+	if (const FGuid* ParentGuid = ChildToParentBinding.Find(BindingID))
+	{
+		TSharedPtr<FSequencerObjectBindingNode> ParentBinding = CreateOrUpdateObjectBinding(*ParentGuid, AllBindings, ChildToParentBinding, InMovieScene);
+		if (ParentBinding.IsValid())
 		{
-			const FMovieSceneBinding* ParentBinding = GuidToBindingMap.FindRef(Possessable->GetParent());
-			if (ParentBinding)
-			{
-				ParentNode = AddObjectBinding( ParentBinding->GetName(), Possessable->GetParent(), GuidToBindingMap, OutNodeList );
-			}
+			ObjectBindingNode->SetParent(ParentBinding);
 		}
+	}
 
-		// get human readable name of the object
-		const FString& DisplayString = ObjectName;
+	return ObjectBindingNode;
+}
 
-		// Create the node.
-		ObjectNode = MakeShareable(new FSequencerObjectBindingNode(ObjectNodeName, FText::FromString(DisplayString), ObjectBinding, ParentNode, *this));
+void FSequencerNodeTree::Update()
+{
+	Sequencer.GetSelection().EmptySelectedOutlinerNodes();
 
-		if (ParentNode.IsValid())
-		{
-			ParentNode->AddObjectBindingNode(ObjectNode.ToSharedRef());
-		}
-		else
-		{
-			OutNodeList.Add( ObjectNode.ToSharedRef() );
-		}
+	EditorMap.Empty();
+	FilteredNodes.Empty();
+	SectionToHandle.Empty();
+	HoveredNode = nullptr;
 
-		// Map the guid to the object binding node for fast lookup later
-		ObjectBindingMap.Add( ObjectBinding, ObjectNode );
-	}
+	UMovieScene* MovieScene = Sequencer.GetFocusedMovieSceneSequence()->GetMovieScene();
+	RefreshNodes(MovieScene);
 
-	return ObjectNode.ToSharedRef();
-}
+	// Re-filter the tree after updating 
+	// @todo sequencer: Newly added sections may need to be visible even when there is a filter
+	FilterNodes( FilterString );
 
+	// Sort root nodes
+	RootNode->SortImmediateChildren();
 
-TSharedRef<FSequencerDisplayNode> CreateFolderNode(
-	UMovieSceneFolder& MovieSceneFolder, FSequencerNodeTree& NodeTree, 
-	TMap<UMovieSceneTrack*, TSharedRef<FSequencerTrackNode>>& MasterTrackToDisplayNodeMap,
-	TMap<FGuid, TSharedRef<FSequencerObjectBindingNode>>& ObjectGuidToDisplayNodeMap )
-{
-	TSharedRef<FSequencerFolderNode> FolderNode( new FSequencerFolderNode( MovieSceneFolder, TSharedPtr<FSequencerDisplayNode>(), NodeTree ) );
+	// Set up virtual offsets, expansion states, and tints
+	float VerticalOffset = 0.f;
 
-	for (UMovieSceneFolder* ChildFolder : MovieSceneFolder.GetChildFolders())
+	auto Traverse_OnTreeRefreshed = [this, &VerticalOffset](FSequencerDisplayNode& InNode)
 	{
-		if (ChildFolder != nullptr)
-		{
-			FolderNode->AddChildNode(CreateFolderNode(*ChildFolder, NodeTree, MasterTrackToDisplayNodeMap, ObjectGuidToDisplayNodeMap));
-		}
-	}
+		// Set up the virtual node position
+		float VerticalTop = VerticalOffset;
+		VerticalOffset += InNode.GetNodeHeight() + InNode.GetNodePadding().Combined();
+		InNode.OnTreeRefreshed(VerticalTop, VerticalOffset);
 
-	for ( UMovieSceneTrack* MasterTrack : MovieSceneFolder.GetChildMasterTracks() )
-	{
-		TSharedRef<FSequencerTrackNode>* TrackNodePtr = MasterTrackToDisplayNodeMap.Find( MasterTrack );
-		if ( TrackNodePtr != nullptr)
+		if (InNode.GetType() == ESequencerNode::Track)
 		{
-			// TODO: Log this.
-			FolderNode->AddChildNode( *TrackNodePtr );
-			MasterTrackToDisplayNodeMap.Remove( MasterTrack );
+			this->UpdateSectionHandles(StaticCastSharedRef<FSequencerTrackNode>(InNode.AsShared()));
 		}
-	}
+		return true;
+	};
 
-	for (const FGuid& ObjectGuid : MovieSceneFolder.GetChildObjectBindings() )
-	{
-		TSharedRef<FSequencerObjectBindingNode>* ObjectNodePtr = ObjectGuidToDisplayNodeMap.Find( ObjectGuid );
-		if ( ObjectNodePtr != nullptr )
-		{
-			// TODO: Log this.
-			FolderNode->AddChildNode( *ObjectNodePtr );
-			ObjectGuidToDisplayNodeMap.Remove( ObjectGuid );
-		}
-	}
+	const bool bIncludeRootNode = false;
+	RootNode->Traverse_ParentFirst(Traverse_OnTreeRefreshed, bIncludeRootNode);
+
+	// Ensure that the curve editor tree is up to date for our tree layout
+	UpdateCurveEditorTree();
 
-	return FolderNode;
+	OnUpdatedDelegate.Broadcast();
 }
 
 
-void FSequencerNodeTree::CreateAndPopulateFolderNodes( 
-	TArray<TSharedRef<FSequencerTrackNode>>& MasterTrackNodes, TArray<TSharedRef<FSequencerObjectBindingNode>>& ObjectNodes,
-	TArray<UMovieSceneFolder*>& MovieSceneFolders, TArray<TSharedRef<FSequencerDisplayNode>>& FolderAndObjectNodes, TArray<TSharedRef<FSequencerDisplayNode>>&  MasterTrackNodesNotInFolders )
+TSharedRef<ISequencerTrackEditor> FSequencerNodeTree::FindOrAddTypeEditor( UMovieSceneTrack* InTrack )
 {
-	TMap<UMovieSceneTrack*, TSharedRef<FSequencerTrackNode>> MasterTrackToDisplayNodeMap;
-	for ( TSharedRef<FSequencerTrackNode> MasterTrackNode : MasterTrackNodes )
-	{
-		MasterTrackToDisplayNodeMap.Add( MasterTrackNode->GetTrack(), MasterTrackNode );
-	}
+	TSharedPtr<ISequencerTrackEditor> Editor = EditorMap.FindRef( InTrack );
 
-	TMap<FGuid, TSharedRef<FSequencerObjectBindingNode>> ObjectGuidToDisplayNodeMap;
-	for ( TSharedRef<FSequencerObjectBindingNode> ObjectBindingNode : ObjectNodes )
+	if( !Editor.IsValid() )
 	{
-		ObjectGuidToDisplayNodeMap.Add( ObjectBindingNode->GetObjectBinding(), ObjectBindingNode );
-	}
+		const TArray<TSharedPtr<ISequencerTrackEditor>>& TrackEditors = Sequencer.GetTrackEditors();
 
-	for (UMovieSceneFolder* MovieSceneFolder : MovieSceneFolders)
-	{
-		if (MovieSceneFolder != nullptr)
+		// Get a tool for each track
+		// @todo sequencer: Should probably only need to get this once and it shouldn't be done here. It depends on when movie scene tool modules are loaded
+		TSharedPtr<ISequencerTrackEditor> SupportedTool;
+
+		for (const auto& TrackEditor : TrackEditors)
 		{
-			FolderAndObjectNodes.Add(CreateFolderNode(*MovieSceneFolder, *this, MasterTrackToDisplayNodeMap, ObjectGuidToDisplayNodeMap));
+			if (TrackEditor->SupportsType(InTrack->GetClass()))
+			{
+				EditorMap.Add(InTrack, TrackEditor);
+				Editor = TrackEditor;
+
+				break;
+			}
 		}
 	}
 
-	TArray<TSharedRef<FSequencerTrackNode>> NonFolderTrackNodes;
-	MasterTrackToDisplayNodeMap.GenerateValueArray( NonFolderTrackNodes );
-	for ( TSharedRef<FSequencerTrackNode> NonFolderTrackNode : NonFolderTrackNodes )
-	{
-		MasterTrackNodesNotInFolders.Add( NonFolderTrackNode );
-	}
+	return Editor.ToSharedRef();
+}
 
-	TArray<TSharedRef<FSequencerObjectBindingNode>> NonFolderObjectNodes;
-	ObjectGuidToDisplayNodeMap.GenerateValueArray( NonFolderObjectNodes );
-	for ( TSharedRef<FSequencerObjectBindingNode> NonFolderObjectNode : NonFolderObjectNodes )
-	{
-		FolderAndObjectNodes.Add( NonFolderObjectNode );
-	}
+TSharedRef<FSequencerDisplayNode> FSequencerNodeTree::GetRootNode() const
+{
+	return RootNode;
+}
+
+const TArray<TSharedRef<FSequencerDisplayNode>>& FSequencerNodeTree::GetRootNodes() const
+{
+	return RootNode->GetChildNodes();
 }
 
 void FSequencerNodeTree::MoveDisplayNodeToRoot(TSharedRef<FSequencerDisplayNode>& Node)
@@ -492,7 +411,7 @@ void FSequencerNodeTree::MoveDisplayNodeToRoot(TSharedRef<FSequencerDisplayNode>
 	}
 
 	// Clear the node's parent so that subsequent calls for GetNodePath correctly indicate that they no longer have a parent.
-	Node->ClearParent();
+	Node->SetParent(nullptr);
 
 	// Our children have changed parents which means that on subsequent creation they will retrieve their expansion state
 	// from the map using their new path. If the new path already exists the object goes to the state stored at that path.
@@ -512,26 +431,17 @@ void FSequencerNodeTree::MoveDisplayNodeToRoot(TSharedRef<FSequencerDisplayNode>
 
 void FSequencerNodeTree::SortAllNodesAndDescendants()
 {
-	// Sort the root first
-	SortAndSetSortingOrder(RootNodes, RootNodes, TOptional<EItemDropZone>(), FDisplayNodeCategoricalSorter(), nullptr);
-	
-	// Recursively sort our children looking for folders.
-	TArray<TSharedRef<FSequencerDisplayNode>> ChildNodes = GetRootNodes();
-	for (TSharedRef<FSequencerDisplayNode> Child : ChildNodes)
-	{
-		Child->Traverse_ParentFirst([&](FSequencerDisplayNode& Node)
-		{
-			// Folders are the only type of node that can have children that we can sort, so there is no need
-			// to follow the traversal all the way down.
-			if (Node.GetType() != ESequencerNode::Folder)
-				return false;
-	
-			SortAndSetSortingOrder(Node.GetChildNodes(), Node.GetChildNodes(), TOptional<EItemDropZone>(), FDisplayNodeCategoricalSorter(), nullptr);
-			return true;
-		}, true);
-	}
+	auto Traverse_ResetSortOrder = [](FSequencerDisplayNode& Node)
+	{
+		Node.ResortImmediateChildren();
+		return true;
+	};
+
+	const bool bIncludeRootNode = true;
+	RootNode->Traverse_ParentFirst(Traverse_ResetSortOrder, bIncludeRootNode);
 
 	// Refresh the tree so that our changes are visible.
+	// @todo: Is this necessary any more?
 	GetSequencer().RefreshTree();
 }
 
@@ -601,6 +511,30 @@ const TSharedPtr<FSequencerDisplayNode>& FSequencerNodeTree::GetHoveredNode() co
 	return HoveredNode;
 }
 
+void FSequencerNodeTree::UpdateSectionHandles(TSharedRef<FSequencerTrackNode> TrackNode)
+{
+	const TArray<TSharedRef<ISequencerSection>>& Sections = TrackNode->GetSections();
+	for (int32 SectionIndex = 0; SectionIndex < Sections.Num(); ++SectionIndex)
+	{
+		UMovieSceneSection* Section = Sections[SectionIndex]->GetSectionObject();
+		if (Section)
+		{
+			SectionToHandle.Add(Section, FSectionHandle(TrackNode, SectionIndex));
+		}
+	}
+}
+
+TOptional<FSectionHandle> FSequencerNodeTree::GetSectionHandle(const UMovieSceneSection* Section) const
+{
+	const FSectionHandle* Found = SectionToHandle.Find(Section);
+	if (Found)
+	{
+		return *Found;
+	}
+
+	return TOptional<FSectionHandle>();
+}
+
 static void AddChildNodes(const TSharedRef<FSequencerDisplayNode>& StartNode, TSet<TSharedRef<const FSequencerDisplayNode>>& OutFilteredNodes)
 {
 	OutFilteredNodes.Add(StartNode);
@@ -744,13 +678,7 @@ void FSequencerNodeTree::FilterNodes(const FString& InFilter)
 		FilterString.TrimStartAndEndInline();
 		FilterString.ParseIntoArray(FilterStrings, TEXT(" "), true /*bCullEmpty*/);
 
-		for (auto It = ObjectBindingMap.CreateIterator(); It; ++It)
-		{
-			// Recursively filter all nodes, matching them against the list of filter strings.  All filter strings must be matched
-			FilterNodesRecursive(Sequencer, It.Value().ToSharedRef(), FilterStrings, FilteredNodes);
-		}
-
-		for (auto It = RootNodes.CreateIterator(); It; ++It)
+		for (auto It = GetRootNodes().CreateConstIterator(); It; ++It)
 		{
 			// Recursively filter all nodes, matching them against the list of filter strings.  All filter strings must be matched
 			FilterNodesRecursive(Sequencer, *It, FilterStrings, FilteredNodes);
@@ -762,14 +690,80 @@ TArray< TSharedRef<FSequencerDisplayNode> > FSequencerNodeTree::GetAllNodes() co
 {
 	TArray< TSharedRef<FSequencerDisplayNode> > AllNodes;
 
-	for (const TSharedRef<FSequencerDisplayNode>& Node : RootNodes)
+	const bool bIncludeRootNode = false;
+	RootNode->Traverse_ParentFirst([&AllNodes](FSequencerDisplayNode& InNode) 
+	{
+		AllNodes.Add(InNode.AsShared());
+		return true;
+	}, bIncludeRootNode);
+
+	return AllNodes;
+}
+
+void FSequencerNodeTree::UpdateCurveEditorTree()
+{
+	FCurveEditor* CurveEditor = Sequencer.GetCurveEditor().Get();
+	auto Traverse_AddToCurveEditor = [this, CurveEditor](FSequencerDisplayNode& InNode)
+	{
+		if (InNode.GetType() == ESequencerNode::Track)
+		{
+			// Track nodes with top level key area's must be added
+			TSharedPtr<FSequencerSectionKeyAreaNode> TopLevelKeyArea = static_cast<const FSequencerTrackNode&>(InNode).GetTopLevelKeyNode();
+			if (TopLevelKeyArea.IsValid() && KeyAreaHasCurves(*TopLevelKeyArea))
+			{
+				this->AddToCurveEditor(InNode.AsShared(), CurveEditor);
+			}
+		}
+		else if (InNode.GetType() == ESequencerNode::KeyArea && KeyAreaHasCurves(static_cast<const FSequencerSectionKeyAreaNode&>(InNode)))
+		{
+			// Key area nodes are always added
+			this->AddToCurveEditor(InNode.AsShared(), CurveEditor);
+		}
+		return true;
+	};
+
+	static const bool bIncludeThisNode = false;
+	RootNode->Traverse_ChildFirst(Traverse_AddToCurveEditor, bIncludeThisNode);
+
+	// Remove no longer valid elements from the curve editor tree
+	for (auto It = CurveEditorTreeItemIDs.CreateIterator(); It; ++It)
 	{
-		Node->Traverse_ParentFirst([&](FSequencerDisplayNode& InNode) 
+		TSharedPtr<FSequencerDisplayNode> Node = It->Key.Pin();
+		if (!Node.IsValid() || Node->TreeSerialNumber != SerialNumber)
+		{
+			CurveEditor->RemoveTreeItem(It->Value);
+			It.RemoveCurrent();
+		}
+	}
+}
+
+bool FSequencerNodeTree::KeyAreaHasCurves(const FSequencerSectionKeyAreaNode& KeyAreaNode) const
+{
+	for (const TSharedRef<IKeyArea>& KeyArea : KeyAreaNode.GetAllKeyAreas())
+	{
+		const ISequencerChannelInterface* EditorInterface = KeyArea->FindChannelEditorInterface();
+		if (EditorInterface && EditorInterface->SupportsCurveEditorModels_Raw(KeyArea->GetChannel()))
 		{
-			AllNodes.Add(InNode.AsShared());
 			return true;
-		});
+		}
 	}
+	return false;
+}
 
-	return AllNodes;
+FCurveEditorTreeItemID FSequencerNodeTree::AddToCurveEditor(TSharedRef<FSequencerDisplayNode> DisplayNode, FCurveEditor* CurveEditor)
+{
+	if (FCurveEditorTreeItemID* Existing = CurveEditorTreeItemIDs.Find(DisplayNode))
+	{
+		return *Existing;
+	}
+
+	TSharedPtr<FSequencerDisplayNode> Parent = DisplayNode->GetParent();
+
+	FCurveEditorTreeItemID ParentID = Parent.IsValid() ? AddToCurveEditor(Parent.ToSharedRef(), CurveEditor) : FCurveEditorTreeItemID::Invalid();
+
+	FCurveEditorTreeItem* NewItem = CurveEditor->AddTreeItem(ParentID);
+	NewItem->SetWeakItem(StaticCastSharedRef<ICurveEditorTreeItem>(DisplayNode));
+
+	CurveEditorTreeItemIDs.Add(DisplayNode, NewItem->GetID());
+	return NewItem->GetID();
 }
\ No newline at end of file
diff --git a/Engine/Source/Editor/Sequencer/Private/SequencerNodeTree.h b/Engine/Source/Editor/Sequencer/Private/SequencerNodeTree.h
--- a/Engine/Source/Editor/Sequencer/Private/SequencerNodeTree.h
+++ b/Engine/Source/Editor/Sequencer/Private/SequencerNodeTree.h
@@ -4,15 +4,24 @@
 
 #include "CoreMinimal.h"
 #include "Misc/Guid.h"
+#include "UObject/ObjectKey.h"
+#include "Tree/CurveEditorTree.h"
+#include "DisplayNodes/SequencerDisplayNode.h"
+#include "SectionHandle.h"
 
 class FSequencer;
+class FCurveEditor;
 class FSequencerDisplayNode;
 class FSequencerObjectBindingNode;
+class FSequencerSectionKeyAreaNode;
 class FSequencerTrackNode;
+class FSequencerFolderNode;
 class ISequencerTrackEditor;
 class UMovieSceneFolder;
 class UMovieSceneTrack;
+class UMovieScene;
 struct FMovieSceneBinding;
+struct FCurveEditorTreeItemID;
 
 /**
  * Represents a tree of sequencer display nodes, used to populate the Sequencer UI with MovieScene data
@@ -23,20 +32,19 @@ public:
 	DECLARE_MULTICAST_DELEGATE(FOnUpdated);
 
 public:
-	FSequencerNodeTree( class FSequencer& InSequencer )
-		: Sequencer( InSequencer )
-	{}
-	
-	/**
-	 * Empties the entire tree
-	 */
-	void Empty();
+
+	FSequencerNodeTree(FSequencer& InSequencer);
 
 	/**
 	 * Updates the tree with sections from a MovieScene
 	 */
 	void Update();
 
+	/**
+	 * Access this tree's symbolic root node
+	 */
+	TSharedRef<FSequencerDisplayNode> GetRootNode() const;
+
 	/**
 	 * @return The root nodes of the tree
 	 */
@@ -101,9 +109,9 @@ public:
 	const TSharedPtr<FSequencerDisplayNode>& GetHoveredNode() const;
 
 	/*
-	 * Get the object binding map from guid to object binding nodes
+	 * Find the object binding node with the specified GUID
 	 */
-	const TMap< FGuid, TSharedPtr<FSequencerObjectBindingNode> >& GetObjectBindingMap() const { return ObjectBindingMap; }
+	TSharedPtr<FSequencerObjectBindingNode> FindObjectBindingNode(const FGuid& BindingID) const;
 
 	/*
 	 * Gets a multicast delegate which is called whenever the node tree has been updated.
@@ -116,53 +124,72 @@ public:
 	/** Sorts all nodes and their descendants by category then alphabetically.*/
 	void SortAllNodesAndDescendants();
 
-private:
+public:
+
 	/**
-	 * Finds or adds a type editor for the track
-	 *
-	 * @param Track	The type to find an editor for
-	 * @return The editor for the type
+	 * Finds the section handle relating to the specified section object, or Null if one was not found (perhaps, if it's filtered away)
 	 */
-	TSharedRef<ISequencerTrackEditor> FindOrAddTypeEditor( UMovieSceneTrack& Track );
+	TOptional<FSectionHandle> GetSectionHandle(const UMovieSceneSection* Section) const;
+
+private:
+
+	/** Population algorithm utilities */
+	void RefreshNodes(UMovieScene* MovieScene);
+
+	enum class ETrackType { Master, Object };
+
+	TSharedPtr<FSequencerTrackNode> CreateOrUpdateTrack(UMovieSceneTrack* Track, ETrackType TrackType);
+	TSharedRef<FSequencerFolderNode> CreateOrUpdateFolder(UMovieSceneFolder* Folder, const TSortedMap<FGuid, const FMovieSceneBinding*>& AllBindings, const TSortedMap<FGuid, FGuid>& ChildToParentBinding, const UMovieScene* InMovieScene);
+	TSharedPtr<FSequencerObjectBindingNode> CreateOrUpdateObjectBinding(const FGuid& BindingID, const TSortedMap<FGuid, const FMovieSceneBinding*>& AllBindings, const TSortedMap<FGuid, FGuid>& ChildToParentBinding, const UMovieScene* InMovieScene);
 
-	/* 
-	 * Makes sub-track nodes and section interfaces for a track node.
-	 * @param The track node to create sub-tracks and section interfaces for.
-	 * @return A new parent for the supplied track, or the track itself
+	/**
+	 * Creates section handles for all the sections contained in the specified track
 	 */
-	TSharedRef<FSequencerTrackNode> MakeSubTracksAndSectionInterfaces(TSharedRef<FSequencerTrackNode> TrackNode, const FGuid& ObjectBinding = FGuid());
+	void UpdateSectionHandles(TSharedRef<FSequencerTrackNode> TrackNode);
 
+private:
 	/**
-	 * Makes section interfaces for all sections in a track
+	 * Finds or adds a type editor for the track
 	 *
-	 * @param SectionAreaNode	The section area which section interfaces belong to
-	 * @param ObjectBinding Object binding ID that the track operates on
+	 * @param Track	The type to find an editor for
+	 * @return The editor for the type
 	 */
-	void MakeSectionInterfaces( TSharedRef<FSequencerTrackNode> TrackNode, const FGuid& ObjectBinding );
+	TSharedRef<ISequencerTrackEditor> FindOrAddTypeEditor( UMovieSceneTrack* Track );
 
-	/*
-	 * Helper function that creates object binding nodes along with their tracks. 
+	/**
+	 * Update the curve editor tree to include anything of relevance from this tree
 	 */
-	void AddObjectBindingAndTracks(const FMovieSceneBinding& ObjectBinding, TMap<FGuid, const FMovieSceneBinding*>& GuidToBindingMap, TArray< TSharedRef<FSequencerObjectBindingNode> >& OutNodeList);
+	void UpdateCurveEditorTree();
 
 	/**
-	 * Creates a new object binding node and any parent binding nodes.
+	 * Add the specified node to the curve editor, including all its parents if necessary
 	 */
-	TSharedRef<FSequencerObjectBindingNode> AddObjectBinding( const FString& ObjectName, const FGuid& ObjectBinding, TMap<FGuid, const FMovieSceneBinding*>& GuidToBindingMap, TArray< TSharedRef<FSequencerObjectBindingNode> >& OutNodeList );
+	FCurveEditorTreeItemID AddToCurveEditor(TSharedRef<FSequencerDisplayNode> Node, FCurveEditor* CurveEditor);
 
 	/**
-	 * Creates the tree of folder nodes and populates it with object and track nodes. 
+	 * Checks whether the specified key area node contains key areas that can create curve models
 	 */
-	void CreateAndPopulateFolderNodes( TArray<TSharedRef<FSequencerTrackNode>>& MasterTrackNodes, TArray<TSharedRef<FSequencerObjectBindingNode>>& ObjectNodes,
-		TArray<UMovieSceneFolder*>& MovieSceneFolders, TArray<TSharedRef<FSequencerDisplayNode>>& FolderAndObjectNodes, TArray<TSharedRef<FSequencerDisplayNode>>& MasterTrackNodesNotInFolders );
+	bool KeyAreaHasCurves(const FSequencerSectionKeyAreaNode& KeyAreaNode) const;
 
 private:
+
+	/** Symbolic root node that contains the actual displayed root nodes as children */
+	TSharedRef<FSequencerDisplayNode> RootNode;
+
+	/** A serially incrementing integer that is increased each time the tree is refreshed to track node relevance */
+	uint32 SerialNumber;
+	/** Map from FMovieSceneBinding::GetObjectGuid to display node */
+	TMap<FGuid,      TSharedPtr<FSequencerObjectBindingNode>> ObjectBindingToNode;
+	/** Map from UMovieSceneTrack object key to display node for any track (object tracks or master tracks) */
+	TMap<FObjectKey, TSharedPtr<FSequencerTrackNode>>         TrackToNode;
+	/** Map from UMovieSceneFolder object key to display node for any folder (root or child folder) */
+	TMap<FObjectKey, TSharedPtr<FSequencerFolderNode>>        FolderToNode;
+
+	/** Map from UMovieSceneSection* to its UI handle */
+	TMap<FObjectKey, FSectionHandle> SectionToHandle;
+
 	/** Tools for building movie scene section layouts.  One tool for each track */
 	TMap< UMovieSceneTrack*, TSharedPtr<ISequencerTrackEditor> > EditorMap;
-	/** Root nodes */
-	TArray< TSharedRef<FSequencerDisplayNode> > RootNodes;
-	/** Mapping of object binding guids to their node (for fast lookup) */
-	TMap< FGuid, TSharedPtr<FSequencerObjectBindingNode> > ObjectBindingMap;
 	/** Set of all filtered nodes */
 	TSet< TSharedRef<const FSequencerDisplayNode> > FilteredNodes;
 	/** Cardinal hovered node */
@@ -171,6 +198,8 @@ private:
 	FString FilterString;
 	/** Sequencer interface */
 	FSequencer& Sequencer;
+	/** Display node -> curve editor tree item ID mapping */
+	TMap<TWeakPtr<FSequencerDisplayNode>, FCurveEditorTreeItemID> CurveEditorTreeItemIDs;
 	/** A multicast delegate which is called whenever the node tree has been updated. */
 	FOnUpdated OnUpdatedDelegate;
 };
diff --git a/Engine/Source/Editor/Sequencer/Private/SequencerSectionLayoutBuilder.cpp b/Engine/Source/Editor/Sequencer/Private/SequencerSectionLayoutBuilder.cpp
--- a/Engine/Source/Editor/Sequencer/Private/SequencerSectionLayoutBuilder.cpp
+++ b/Engine/Source/Editor/Sequencer/Private/SequencerSectionLayoutBuilder.cpp
@@ -2,40 +2,129 @@
 
 #include "SequencerSectionLayoutBuilder.h"
 #include "DisplayNodes/SequencerSectionCategoryNode.h"
+#include "Channels/MovieSceneChannelEditorData.h"
+#include "IKeyArea.h"
 
 
-FSequencerSectionLayoutBuilder::FSequencerSectionLayoutBuilder( TSharedRef<FSequencerTrackNode> InRootNode )
-	: RootNode( InRootNode )
-	, CurrentNode( InRootNode )
+FSequencerSectionLayoutBuilder::FSequencerSectionLayoutBuilder(TSharedRef<FSequencerTrackNode> InRootTrackNode, UMovieSceneSection* InSection)
+	: RootNode(InRootTrackNode)
+	, CurrentNode(InRootTrackNode)
+	, Section(InSection)
 	, bHasAnyLayout(false)
 {}
 
 void FSequencerSectionLayoutBuilder::PushCategory( FName CategoryName, const FText& DisplayLabel )
 {
-	CurrentNode = CurrentNode->AddCategoryNode( CategoryName, DisplayLabel );
 	bHasAnyLayout = true;
+
+	TSharedPtr<FSequencerSectionCategoryNode> CategoryNode;
+
+	// Attempt to re-use an existing key area of the same name
+	for (const TSharedRef<FSequencerDisplayNode>& Child : CurrentNode->GetChildNodes())
+	{
+		if (Child->GetType() == ESequencerNode::Category && Child->GetNodeName() == CategoryName)
+		{
+			// Ensure its name is up to date
+			CategoryNode = StaticCastSharedRef<FSequencerSectionCategoryNode>(Child);
+		}
+	}
+
+	if (!CategoryNode)
+	{
+		CategoryNode = MakeShared<FSequencerSectionCategoryNode>(CategoryName, RootNode->GetParentTree());
+		CategoryNode->SetParent(CurrentNode);
+	}
+
+	CategoryNode->DisplayName = DisplayLabel;
+	CategoryNode->TreeSerialNumber = RootNode->TreeSerialNumber;
+	CurrentNode = CategoryNode.ToSharedRef();
 }
 
 void FSequencerSectionLayoutBuilder::PopCategory()
 {
 	// Pop a category if the current node is a category
-	if( CurrentNode->GetParent().IsValid() && CurrentNode->GetType() == ESequencerNode::Category )
+	if( TSharedPtr<FSequencerDisplayNode> Parent = CurrentNode->GetParent() )
 	{
-		CurrentNode = CurrentNode->GetParent().ToSharedRef();
+		if (CurrentNode->GetType() == ESequencerNode::Category)
+		{
+			CurrentNode = Parent.ToSharedRef();
+		}
 	}
 }
 
-void FSequencerSectionLayoutBuilder::SetSectionAsKeyArea( TSharedRef<IKeyArea> KeyArea )
+void FSequencerSectionLayoutBuilder::SetTopLevelChannel( const FMovieSceneChannelHandle& Channel )
 {
 	bHasAnyLayout = true;
-	KeyArea->SetName(RootNode->GetNodeName());
-	RootNode->SetSectionAsKeyArea( KeyArea ); 
+
+	ensureAlwaysMsgf(CurrentNode == RootNode, TEXT("Attempting to assign a top level channel when a category node is active. Top level key nodes will always be added to the outermost track node."));
+
+	TSharedPtr<FSequencerSectionKeyAreaNode> KeyAreaNode = RootNode->GetTopLevelKeyNode();
+	if (!KeyAreaNode.IsValid())
+	{
+		KeyAreaNode = MakeShared<FSequencerSectionKeyAreaNode>(RootNode->GetNodeName(), RootNode->GetParentTree());
+		// does this need a parent?
+		KeyAreaNode->SetParentDirectly(RootNode);
+
+		RootNode->SetTopLevelKeyNode(KeyAreaNode);
+	}
+
+	AddOrUpdateChannel(KeyAreaNode.ToSharedRef(), Channel);
 }
 
-void FSequencerSectionLayoutBuilder::AddKeyArea( FName KeyAreaName, const FText& DisplayName, TSharedRef<IKeyArea> KeyArea )
+void FSequencerSectionLayoutBuilder::AddChannel( const FMovieSceneChannelHandle& Channel )
 {
+	// @todo: this is all pretty crusy - we're currently linear-searching for both the child node, and the IKeyArea within that node
+	// Performance is generally acceptible however since we are dealing with small numbers of children, but this may need to be revisited.
+	const FMovieSceneChannelMetaData* MetaData = Channel.GetMetaData();
+	if (!ensureAlwaysMsgf(MetaData, TEXT("Attempting to add an expired channel handle to the node tree")))
+	{
+		return;
+	}
+
 	bHasAnyLayout = true;
-	KeyArea->SetName(KeyAreaName);
-	CurrentNode->AddKeyAreaNode( KeyAreaName, DisplayName, KeyArea );
+
+	TSharedPtr<FSequencerSectionKeyAreaNode> KeyAreaNode;
+
+	// Attempt to re-use an existing key area of the same name
+	for (const TSharedRef<FSequencerDisplayNode>& Child : CurrentNode->GetChildNodes())
+	{
+		if (Child->GetType() == ESequencerNode::KeyArea && Child->GetNodeName() == MetaData->Name)
+		{
+			KeyAreaNode = StaticCastSharedRef<FSequencerSectionKeyAreaNode>(Child);
+		}
+	}
+
+	if (!KeyAreaNode.IsValid())
+	{
+		// No existing node found make a new one
+		KeyAreaNode = MakeShared<FSequencerSectionKeyAreaNode>(MetaData->Name, CurrentNode->GetParentTree());
+		KeyAreaNode->DisplayName = MetaData->DisplayText;
+		KeyAreaNode->SetParent(CurrentNode);
+	}
+
+	AddOrUpdateChannel(KeyAreaNode.ToSharedRef(), Channel);
+}
+
+void FSequencerSectionLayoutBuilder::AddOrUpdateChannel(TSharedRef<FSequencerSectionKeyAreaNode> KeyAreaNode, const FMovieSceneChannelHandle& Channel)
+{
+	KeyAreaNode->TreeSerialNumber = RootNode->TreeSerialNumber;
+
+	TSharedPtr<IKeyArea> KeyArea = KeyAreaNode->GetKeyArea(Section);
+	if (!KeyArea)
+	{
+		// No key area for this section exists - create a new one
+		TSharedRef<IKeyArea> NewKeyArea = MakeShared<IKeyArea>(Section, Channel);
+		KeyAreaNode->AddKeyArea(NewKeyArea);
+	}
+	else if (KeyArea->GetChannel() != Channel)
+	{
+		// A key area exists but for a different channel handle so this needs re-creating
+		KeyArea->Reinitialize(Section, Channel);
+	}
+	else
+	{
+		// Just ensure the name is up to date
+		KeyArea->SetName(CurrentNode->GetNodeName());
+	}
 }
 
diff --git a/Engine/Source/Editor/Sequencer/Private/SequencerSectionLayoutBuilder.h b/Engine/Source/Editor/Sequencer/Private/SequencerSectionLayoutBuilder.h
--- a/Engine/Source/Editor/Sequencer/Private/SequencerSectionLayoutBuilder.h
+++ b/Engine/Source/Editor/Sequencer/Private/SequencerSectionLayoutBuilder.h
@@ -4,7 +4,6 @@
 
 #include "CoreMinimal.h"
 #include "DisplayNodes/SequencerDisplayNode.h"
-#include "IKeyArea.h"
 #include "DisplayNodes/SequencerTrackNode.h"
 #include "ISectionLayoutBuilder.h"
 
@@ -12,15 +11,15 @@ class FSequencerSectionLayoutBuilder
 	: public ISectionLayoutBuilder
 {
 public:
-	FSequencerSectionLayoutBuilder( TSharedRef<FSequencerTrackNode> InRootNode );
+	FSequencerSectionLayoutBuilder(TSharedRef<FSequencerTrackNode> InRootTrackNode, UMovieSceneSection* InSection);
 
 public:
 
 	// ISectionLayoutBuilder interface
 
 	virtual void PushCategory( FName CategoryName, const FText& DisplayLabel ) override;
-	virtual void SetSectionAsKeyArea( TSharedRef<IKeyArea> KeyArea ) override;
-	virtual void AddKeyArea( FName KeyAreaName, const FText& DisplayName, TSharedRef<IKeyArea> KeyArea ) override;
+	virtual void SetTopLevelChannel( const FMovieSceneChannelHandle& Channel ) override;
+	virtual void AddChannel( const FMovieSceneChannelHandle& Channel ) override;
 	virtual void PopCategory() override;
 
 	/** Check whether this section layout builder has been given any layout or not */
@@ -31,12 +30,17 @@ public:
 
 private:
 
+	void AddOrUpdateChannel(TSharedRef<FSequencerSectionKeyAreaNode> KeyAreaNode, const FMovieSceneChannelHandle& Channel);
+
 	/** Root node of the tree */
 	TSharedRef<FSequencerTrackNode> RootNode;
 
 	/** The current node that other nodes are added to */
 	TSharedRef<FSequencerDisplayNode> CurrentNode;
 
+	/** The section that we are building a layout for */
+	UMovieSceneSection* Section;
+
 	/** Boolean indicating whether this section layout builder has been given any layout or not */
 	bool bHasAnyLayout;
 };
diff --git a/Engine/Source/Editor/Sequencer/Private/SequencerSelection.cpp b/Engine/Source/Editor/Sequencer/Private/SequencerSelection.cpp
--- a/Engine/Source/Editor/Sequencer/Private/SequencerSelection.cpp
+++ b/Engine/Source/Editor/Sequencer/Private/SequencerSelection.cpp
@@ -67,22 +67,14 @@ TArray<FGuid> FSequencerSelection::GetBoundObjectsGuids()
 	}
 	for (TSharedRef<FSequencerDisplayNode> Node : SelectedNodes)
 	{
-		TSharedPtr<FSequencerObjectBindingNode> ObjectNode;
+		TSharedPtr<FSequencerObjectBindingNode> ObjectNode = Node->FindParentObjectBindingNode();
 		if (Node->GetType() == ESequencerNode::Object)
 		{
 			ObjectNode = StaticCastSharedRef<FSequencerObjectBindingNode>(Node);
 		}
 		else
 		{
-			TSharedPtr<FSequencerDisplayNode> ParentNode = Node->GetParent();
-			if (ParentNode.IsValid())
-			{
-				while (ParentNode->GetParent().IsValid() && ParentNode->GetType() != ESequencerNode::Object)
-				{
-					ParentNode = ParentNode->GetParent();
-				}
-				ObjectNode = StaticCastSharedPtr<FSequencerObjectBindingNode>(ParentNode);
-			}
+			ObjectNode = Node->FindParentObjectBindingNode();
 		}
 
 		if (ObjectNode.IsValid())
diff --git a/Engine/Source/Editor/Sequencer/Private/SequencerSettings.cpp b/Engine/Source/Editor/Sequencer/Private/SequencerSettings.cpp
--- a/Engine/Source/Editor/Sequencer/Private/SequencerSettings.cpp
+++ b/Engine/Source/Editor/Sequencer/Private/SequencerSettings.cpp
@@ -47,7 +47,6 @@ USequencerSettings::USequencerSettings( const FObjectInitializer& ObjectInitiali
 	bRerunConstructionScripts = true;
 	bVisualizePreAndPostRoll = true;
 	TrajectoryPathCap = 250;
-	CurveVisibility = ECurveEditorCurveVisibility::SelectedCurves;
 	FrameNumberDisplayFormat = EFrameNumberDisplayFormats::Seconds;
 }
 
@@ -628,26 +627,6 @@ void USequencerSettings::SetCompileDirectorOnEvaluate(bool bInCompileDirectorOnE
 	}
 }
 
-ECurveEditorCurveVisibility USequencerSettings::GetCurveVisibility() const
-{
-	return CurveVisibility;
-}
-
-void USequencerSettings::SetCurveVisibility(ECurveEditorCurveVisibility InCurveVisibility)
-{
-	if (CurveVisibility != InCurveVisibility)
-	{
-		CurveVisibility = InCurveVisibility;
-		OnCurveEditorCurveVisibilityChanged.Broadcast();
-		SaveConfig();
-	}
-}
-
-USequencerSettings::FOnCurveEditorCurveVisibilityChanged& USequencerSettings::GetOnCurveEditorCurveVisibilityChanged()
-{
-	return OnCurveEditorCurveVisibilityChanged;
-}
-
 USequencerSettings::FOnLoopStateChanged& USequencerSettings::GetOnLoopStateChanged()
 {
 	return OnLoopStateChangedEvent;
diff --git a/Engine/Source/Editor/Sequencer/Private/SequencerTimeSliderController.cpp b/Engine/Source/Editor/Sequencer/Private/SequencerTimeSliderController.cpp
--- a/Engine/Source/Editor/Sequencer/Private/SequencerTimeSliderController.cpp
+++ b/Engine/Source/Editor/Sequencer/Private/SequencerTimeSliderController.cpp
@@ -206,7 +206,7 @@ void FSequencerTimeSliderController::DrawTicks( FSlateWindowElementList& OutDraw
 
 	double MajorGridStep  = 0.0;
 	int32  MinorDivisions = 0;
-	if (!Sequencer->GetGridMetrics(InArgs.AllottedGeometry.Size.X, MajorGridStep, MinorDivisions))
+	if (!Sequencer->GetGridMetrics(InArgs.AllottedGeometry.Size.X, ViewRange.GetLowerBoundValue(), ViewRange.GetUpperBoundValue(), MajorGridStep, MinorDivisions))
 	{
 		return;
 	}
@@ -387,7 +387,7 @@ int32 FSequencerTimeSliderController::OnPaintTimeSlider( bool bMirrorLabels, con
 	const bool bEnabled = bParentEnabled;
 	const ESlateDrawEffect DrawEffects = bEnabled ? ESlateDrawEffect::None : ESlateDrawEffect::DisabledEffect;
 
-	TRange<double> LocalViewRange      = TimeSliderArgs.ViewRange.Get();
+	TRange<double> LocalViewRange      = GetViewRange();
 	const float    LocalViewRangeMin   = LocalViewRange.GetLowerBoundValue();
 	const float    LocalViewRangeMax   = LocalViewRange.GetUpperBoundValue();
 	const float    LocalSequenceLength = LocalViewRangeMax-LocalViewRangeMin;
@@ -493,7 +493,7 @@ int32 FSequencerTimeSliderController::OnPaintTimeSlider( bool bMirrorLabels, con
 			FFrameRate Resolution = GetTickResolution();
 			FFrameTime MouseDownTime[2];
 
-			FScrubRangeToScreen MouseDownRange(TimeSliderArgs.ViewRange.Get(), MouseDownGeometry.Size);
+			FScrubRangeToScreen MouseDownRange(GetViewRange(), MouseDownGeometry.Size);
 			MouseDownTime[0] = ComputeFrameTimeFromMouse(MouseDownGeometry, MouseDownPosition[0], MouseDownRange);
 			MouseDownTime[1] = ComputeFrameTimeFromMouse(MouseDownGeometry, MouseDownPosition[1], MouseDownRange);
 
@@ -731,7 +731,7 @@ FReply FSequencerTimeSliderController::OnMouseButtonUp( SWidget& WidgetOwner, co
 	bool bHandleLeftMouseButton  = MouseEvent.GetEffectingButton() == EKeys::LeftMouseButton  && WidgetOwner.HasMouseCapture();
 	bool bHandleRightMouseButton = MouseEvent.GetEffectingButton() == EKeys::RightMouseButton && WidgetOwner.HasMouseCapture() && TimeSliderArgs.AllowZoom ;
 
-	FScrubRangeToScreen RangeToScreen = FScrubRangeToScreen(TimeSliderArgs.ViewRange.Get(), MyGeometry.Size);
+	FScrubRangeToScreen RangeToScreen = FScrubRangeToScreen(GetViewRange(), MyGeometry.Size);
 	FFrameTime          MouseTime     = ComputeFrameTimeFromMouse(MyGeometry, MouseEvent.GetScreenSpacePosition(), RangeToScreen);
 
 	if ( bHandleRightMouseButton )
@@ -791,7 +791,7 @@ FReply FSequencerTimeSliderController::OnMouseButtonUp( SWidget& WidgetOwner, co
 			const bool bCanZoomOut = ViewRangeStack.Num() > 0;
 			if (bCanZoomIn || bCanZoomOut)
 			{
-				TRange<double> ViewRange = TimeSliderArgs.ViewRange.Get();
+				TRange<double> ViewRange = GetViewRange();
 				if (!bCanZoomIn)
 				{
 					ViewRange = ViewRangeStack.Pop();
@@ -865,7 +865,7 @@ FReply FSequencerTimeSliderController::OnMouseMove( SWidget& WidgetOwner, const
 		}
 		else
 		{
-			TRange<double> LocalViewRange = TimeSliderArgs.ViewRange.Get();
+			TRange<double> LocalViewRange = GetViewRange();
 			double LocalViewRangeMin = LocalViewRange.GetLowerBoundValue();
 			double LocalViewRangeMax = LocalViewRange.GetUpperBoundValue();
 
@@ -883,7 +883,7 @@ FReply FSequencerTimeSliderController::OnMouseMove( SWidget& WidgetOwner, const
 	}
 	else if (bHandleLeftMouseButton)
 	{
-		TRange<double> LocalViewRange = TimeSliderArgs.ViewRange.Get();
+		TRange<double> LocalViewRange = GetViewRange();
 		FScrubRangeToScreen RangeToScreen(LocalViewRange, MyGeometry.Size);
 		DistanceDragged += FMath::Abs( MouseEvent.GetCursorDelta().X );
 
@@ -985,6 +985,27 @@ FReply FSequencerTimeSliderController::OnMouseMove( SWidget& WidgetOwner, const
 
 void FSequencerTimeSliderController::CommitScrubPosition( FFrameTime NewValue, bool bIsScrubbing )
 {
+	// The user can scrub past the viewing range of the time slider controller, so we clamp it to the view range.
+	TSharedPtr<FSequencer> Sequencer = WeakSequencer.Pin();
+	if(Sequencer.IsValid())
+	{
+		FAnimatedRange ViewRange = GetViewRange();
+		
+		FFrameRate DisplayRate = Sequencer->GetFocusedDisplayRate();
+		FFrameRate TickResolution = Sequencer->GetFocusedTickResolution();
+
+		FFrameTime LowerBound = (ViewRange.GetLowerBoundValue() * TickResolution).CeilToFrame();
+		FFrameTime UpperBound = (ViewRange.GetUpperBoundValue() * TickResolution).FloorToFrame();
+
+		if (Sequencer->GetSequencerSettings()->GetIsSnapEnabled() && Sequencer->GetSequencerSettings()->GetSnapPlayTimeToInterval())
+		{
+			LowerBound = FFrameRate::Snap(LowerBound, TickResolution, DisplayRate);
+			UpperBound = FFrameRate::Snap(UpperBound, TickResolution, DisplayRate);
+		}
+
+		NewValue = FMath::Clamp(NewValue, LowerBound, UpperBound);
+	}
+
 	// Manage the scrub position ourselves if its not bound to a delegate
 	if ( !TimeSliderArgs.ScrubPosition.IsBound() )
 	{
@@ -1006,9 +1027,9 @@ FReply FSequencerTimeSliderController::OnMouseWheel( SWidget& WidgetOwner, const
 		if (TimeSliderArgs.Settings->GetZoomPosition() == ESequencerZoomPosition::SZP_CurrentTime)
 		{
 			const double ScrubPosition = TimeSliderArgs.ScrubPosition.Get() / GetTickResolution();
-			if (TimeSliderArgs.ViewRange.Get().Contains(ScrubPosition))
+			if (GetViewRange().Contains(ScrubPosition))
 			{
-				FScrubRangeToScreen RangeToScreen(TimeSliderArgs.ViewRange.Get(), MyGeometry.Size);
+				FScrubRangeToScreen RangeToScreen(GetViewRange(), MyGeometry.Size);
 				float TimePosition = RangeToScreen.InputToLocalX(ScrubPosition);
 				MouseFractionX = TimePosition / MyGeometry.GetLocalSize().X;
 			}
@@ -1037,7 +1058,7 @@ FCursorReply FSequencerTimeSliderController::OnCursorQuery( TSharedRef<const SWi
 		return FCursorReply::Unhandled();
 	}
 
-	FScrubRangeToScreen RangeToScreen(TimeSliderArgs.ViewRange.Get(), MyGeometry.Size);
+	FScrubRangeToScreen RangeToScreen(GetViewRange(), MyGeometry.Size);
 
 	const FFrameRate TickResolution   = GetTickResolution();
 	const bool       bLockedPlayRange = TimeSliderArgs.IsPlaybackRangeLocked.Get();
@@ -1068,7 +1089,7 @@ FCursorReply FSequencerTimeSliderController::OnCursorQuery( TSharedRef<const SWi
 	return FCursorReply::Unhandled();
 }
 
-int32 FSequencerTimeSliderController::OnPaintSectionView( const FGeometry& AllottedGeometry, const FSlateRect& MyCullingRect, FSlateWindowElementList& OutDrawElements, int32 LayerId, bool bEnabled, const FPaintSectionAreaViewArgs& Args ) const
+int32 FSequencerTimeSliderController::OnPaintViewArea( const FGeometry& AllottedGeometry, const FSlateRect& MyCullingRect, FSlateWindowElementList& OutDrawElements, int32 LayerId, bool bEnabled, const FPaintViewAreaArgs& Args ) const
 {
 	TSharedPtr<FSequencer> Sequencer = WeakSequencer.Pin();
 	if (!Sequencer.IsValid())
@@ -1078,7 +1099,7 @@ int32 FSequencerTimeSliderController::OnPaintSectionView( const FGeometry& Allot
 
 	const ESlateDrawEffect DrawEffects = bEnabled ? ESlateDrawEffect::None : ESlateDrawEffect::DisabledEffect;
 
-	TRange<double> LocalViewRange = TimeSliderArgs.ViewRange.Get();
+	TRange<double> LocalViewRange = GetViewRange();
 	FScrubRangeToScreen RangeToScreen( LocalViewRange, AllottedGeometry.Size );
 
 	if (Args.PlaybackRangeArgs.IsSet())
@@ -1344,7 +1365,7 @@ void FSequencerTimeSliderController::SetViewRange( double NewRangeMin, double Ne
 {
 	// Clamp to a minimum size to avoid zero-sized or negative visible ranges
 	double MinVisibleTimeRange = FFrameNumber(1) / GetTickResolution();
-	TRange<double> ExistingViewRange  = TimeSliderArgs.ViewRange.Get();
+	TRange<double> ExistingViewRange  = GetViewRange();
 	TRange<double> ExistingClampRange = TimeSliderArgs.ClampRange.Get();
 
 	if (NewRangeMax == ExistingViewRange.GetUpperBoundValue())
@@ -1400,7 +1421,7 @@ void FSequencerTimeSliderController::SetPlayRange( FFrameNumber RangeStart, int3
 
 bool FSequencerTimeSliderController::ZoomByDelta( float InDelta, float MousePositionFraction )
 {
-	TRange<double> LocalViewRange = TimeSliderArgs.ViewRange.Get().GetAnimationTarget();
+	TRange<double> LocalViewRange = GetViewRange().GetAnimationTarget();
 	double LocalViewRangeMax = LocalViewRange.GetUpperBoundValue();
 	double LocalViewRangeMin = LocalViewRange.GetLowerBoundValue();
 	const double OutputViewSize = LocalViewRangeMax - LocalViewRangeMin;
@@ -1421,7 +1442,7 @@ bool FSequencerTimeSliderController::ZoomByDelta( float InDelta, float MousePosi
 
 void FSequencerTimeSliderController::PanByDelta( float InDelta )
 {
-	TRange<double> LocalViewRange = TimeSliderArgs.ViewRange.Get().GetAnimationTarget();
+	TRange<double> LocalViewRange = GetViewRange().GetAnimationTarget();
 
 	double CurrentMin = LocalViewRange.GetLowerBoundValue();
 	double CurrentMax = LocalViewRange.GetUpperBoundValue();
diff --git a/Engine/Source/Editor/Sequencer/Private/Tools/EditToolDragOperations.cpp b/Engine/Source/Editor/Sequencer/Private/Tools/EditToolDragOperations.cpp
--- a/Engine/Source/Editor/Sequencer/Private/Tools/EditToolDragOperations.cpp
+++ b/Engine/Source/Editor/Sequencer/Private/Tools/EditToolDragOperations.cpp
@@ -20,12 +20,12 @@ struct FInvalidKeyAndSectionSnappingCandidates : ISequencerSnapCandidate
 	/**
 	 * Keys and Sections added to this ISequencerSnapCandidate will be ignored as potential candidates for snapping.
 	 */
-	FInvalidKeyAndSectionSnappingCandidates(const TSet<FSequencerSelectedKey>& InKeysToIgnore, const TArray<FSectionHandle>& InSectionsToIgnore)
+	FInvalidKeyAndSectionSnappingCandidates(const TSet<FSequencerSelectedKey>& InKeysToIgnore, const TArray<TWeakObjectPtr<UMovieSceneSection>>& InSectionsToIgnore)
 	{
 		KeysToExclude = InKeysToIgnore;
-		for (const FSectionHandle& SectionHandle : InSectionsToIgnore)
+		for (const TWeakObjectPtr<UMovieSceneSection>& WeakSection : InSectionsToIgnore)
 		{
-			SectionsToExclude.Add(SectionHandle.GetSectionObject());
+			SectionsToExclude.Add(WeakSection.Get());
 		}
 	}
 
@@ -73,39 +73,7 @@ TOptional<FSequencerSnapField::FSnapResult> SnapToInterval(const TArray<FFrameNu
 /** How many pixels near the mouse has to be before snapping occurs */
 const float PixelSnapWidth = 10.f;
 
-TRange<FFrameNumber> GetSectionBoundaries(const UMovieSceneSection* Section, TArray<FSectionHandle>& SectionHandles, TSharedPtr<FSequencerTrackNode> SequencerNode)
-{
-	// Only get boundaries for the sections that aren't being moved
-	TArray<const UMovieSceneSection*> SectionsBeingMoved;
-	for (auto SectionHandle : SectionHandles)
-	{
-		SectionsBeingMoved.Add(SectionHandle.GetSectionObject());
-	}
-
-	// Find the borders of where you can drag to
-	FFrameNumber LowerBound = TNumericLimits<int32>::Lowest(), UpperBound = TNumericLimits<int32>::Max();
-
-	// Also get the closest borders on either side
-	const TArray< TSharedRef<ISequencerSection> >& AllSections = SequencerNode->GetSections();
-	for (int32 SectionIndex = 0; SectionIndex < AllSections.Num(); ++SectionIndex)
-	{
-		const UMovieSceneSection* TestSection = AllSections[SectionIndex]->GetSectionObject();
-
-		if (!SectionsBeingMoved.Contains(TestSection) && Section->GetRowIndex() == TestSection->GetRowIndex())
-		{
-			if (TestSection->HasEndFrame() && Section->HasStartFrame() && TestSection->GetExclusiveEndFrame() <= Section->GetInclusiveStartFrame() && TestSection->GetExclusiveEndFrame() > LowerBound)
-			{
-				LowerBound = TestSection->GetExclusiveEndFrame();
-			}
-			if (TestSection->HasStartFrame() && Section->HasEndFrame() && TestSection->GetInclusiveStartFrame() >= Section->GetExclusiveEndFrame() && TestSection->GetInclusiveStartFrame() < UpperBound)
-			{
-				UpperBound = TestSection->GetInclusiveStartFrame();
-			}
-		}
-	}
 
-	return TRange<FFrameNumber>(LowerBound, UpperBound);
-}
 
 FEditToolDragOperation::FEditToolDragOperation( FSequencer& InSequencer )
 	: Sequencer(InSequencer)
@@ -123,14 +91,14 @@ int32 FEditToolDragOperation::OnPaint(const FGeometry& AllottedGeometry, const F
 	return LayerId;
 }
 
-void FEditToolDragOperation::BeginTransaction( TArray< FSectionHandle >& Sections, const FText& TransactionDesc )
+void FEditToolDragOperation::BeginTransaction( TArray< TWeakObjectPtr<UMovieSceneSection> >& Sections, const FText& TransactionDesc )
 {
 	// Begin an editor transaction and mark the section as transactional so it's state will be saved
 	Transaction.Reset( new FScopedTransaction(TransactionDesc) );
 
 	for (int32 SectionIndex = 0; SectionIndex < Sections.Num(); )
 	{
-		UMovieSceneSection* SectionObj = Sections[SectionIndex].GetSectionObject();
+		UMovieSceneSection* SectionObj = Sections[SectionIndex].Get();
 
 		SectionObj->SetFlags( RF_Transactional );
 		// Save the current state of the section
@@ -151,9 +119,9 @@ void FEditToolDragOperation::EndTransaction()
 	Sequencer.NotifyMovieSceneDataChanged( EMovieSceneDataChangeType::TrackValueChanged );
 }
 
-FResizeSection::FResizeSection( FSequencer& InSequencer, TArray<FSectionHandle> InSections, bool bInDraggingByEnd, bool bInIsSlipping )
+FResizeSection::FResizeSection( FSequencer& InSequencer, const TSet<TWeakObjectPtr<UMovieSceneSection>>& InSections, bool bInDraggingByEnd, bool bInIsSlipping )
 	: FEditToolDragOperation( InSequencer )
-	, Sections( MoveTemp(InSections) )
+	, Sections( InSections.Array() )
 	, bDraggingByEnd(bInDraggingByEnd)
 	, bIsSlipping(bInIsSlipping)
 	, MouseDownTime(0)
@@ -174,17 +142,15 @@ void FResizeSection::OnBeginDrag(const FPointerEvent& MouseEvent, FVector2D Loca
 	SectionInitTimes.Empty();
 
 	bool bIsDilating = MouseEvent.IsControlDown();
-	ISequencerModule& SequencerModule = FModuleManager::Get().LoadModuleChecked<ISequencerModule>("Sequencer");
+	PreDragSectionData.Empty();
 
-	for (auto& Handle : Sections)
+	for (TWeakObjectPtr<UMovieSceneSection> WeakSection : Sections)
 	{
-		UMovieSceneSection* Section = Handle.GetSectionObject();
+		UMovieSceneSection* Section = WeakSection.Get();
 
-		TSharedRef<ISequencerSection> SectionInterface = Handle.GetSectionInterface();
 		if (bIsDilating)
 		{
 			// Populate the resize data for this section
-			PreDragSectionData.Empty();
 			FPreDragSectionData ResizeData; 
 			ResizeData.MovieSection = Section;
 			ResizeData.InitialRange = Section->GetRange();
@@ -205,9 +171,9 @@ void FResizeSection::OnBeginDrag(const FPointerEvent& MouseEvent, FVector2D Loca
 			}
 			PreDragSectionData.Emplace(ResizeData);
 		}
-		else
+		else if (TOptional<FSectionHandle> SectionHandle = Sequencer.GetNodeTree()->GetSectionHandle(Section))
 		{
-			SectionInterface->BeginResizeSection();
+			SectionHandle->GetSectionInterface()->BeginResizeSection();
 		}
 
 		SectionInitTimes.Add(Section, bDraggingByEnd ? Section->GetExclusiveEndFrame() : Section->GetInclusiveStartFrame());
@@ -237,9 +203,9 @@ void FResizeSection::OnDrag(const FPointerEvent& MouseEvent, FVector2D LocalMous
 	if ( Settings->GetIsSnapEnabled() )
 	{
 		TArray<FFrameNumber> SectionTimes;
-		for (const FSectionHandle& Handle : Sections)
+		for (const TWeakObjectPtr<UMovieSceneSection>& WeakSection : Sections)
 		{
-			UMovieSceneSection* Section = Handle.GetSectionObject();
+			UMovieSceneSection* Section = WeakSection.Get();
 			SectionTimes.Add(SectionInitTimes[Section] + DeltaTime);
 		}
 
@@ -344,68 +310,66 @@ void FResizeSection::OnDrag(const FPointerEvent& MouseEvent, FVector2D LocalMous
 		}
 	}
 	/********************************************************************/
-	else for (const FSectionHandle& Handle : Sections)
+	else for (const TWeakObjectPtr<UMovieSceneSection>& WeakSection : Sections)
 	{
-		UMovieSceneSection* Section = Handle.GetSectionObject();
-
-		// Find the corresponding sequencer section to this movie scene section
-		for (const TSharedRef<ISequencerSection>& SequencerSection : Handle.TrackNode->GetSections())
+		UMovieSceneSection* Section = WeakSection.Get();
+		TOptional<FSectionHandle> SectionHandle = Sequencer.GetNodeTree()->GetSectionHandle(Section);
+		if (!SectionHandle)
 		{
-			if (SequencerSection->GetSectionObject() == Section)
-			{
-				FFrameNumber NewTime = SectionInitTimes[Section] + DeltaTime;
+			continue;
+		}
 
-				if( bDraggingByEnd )
-				{
-					FFrameNumber MinFrame = Section->HasStartFrame() ? Section->GetInclusiveStartFrame() : TNumericLimits<int32>::Lowest();
+		TSharedRef<ISequencerSection> SectionInterface = SectionHandle->GetSectionInterface();
 
-					// Dragging the end of a section
-					// Ensure we aren't shrinking past the start time
-					NewTime = FMath::Max( NewTime, MinFrame );
-				    if (bIsSlipping)
-					{
-						SequencerSection->SlipSection( NewTime );
-					}
-					else
-					{
-						SequencerSection->ResizeSection( SSRM_TrailingEdge, NewTime );
-					}
-				}
-				else
-				{
-					FFrameNumber MaxFrame = Section->HasEndFrame() ? Section->GetExclusiveEndFrame()-1 : TNumericLimits<int32>::Max();
+		FFrameNumber NewTime = SectionInitTimes[Section] + DeltaTime;
 
-					// Dragging the start of a section
-					// Ensure we arent expanding past the end time
-					NewTime = FMath::Min( NewTime, MaxFrame );
+		if( bDraggingByEnd )
+		{
+			FFrameNumber MinFrame = Section->HasStartFrame() ? Section->GetInclusiveStartFrame() : TNumericLimits<int32>::Lowest();
 
-					if (bIsSlipping)
-					{
-						SequencerSection->SlipSection( NewTime );
-					}
-					else
-					{
-						SequencerSection->ResizeSection( SSRM_LeadingEdge, NewTime );
-					}
-				}
+			// Dragging the end of a section
+			// Ensure we aren't shrinking past the start time
+			NewTime = FMath::Max( NewTime, MinFrame );
+			if (bIsSlipping)
+			{
+				SectionInterface->SlipSection( NewTime );
+			}
+			else
+			{
+				SectionInterface->ResizeSection( SSRM_TrailingEdge, NewTime );
+			}
+		}
+		else
+		{
+			FFrameNumber MaxFrame = Section->HasEndFrame() ? Section->GetExclusiveEndFrame()-1 : TNumericLimits<int32>::Max();
 
-				UMovieSceneTrack* OuterTrack = Section->GetTypedOuter<UMovieSceneTrack>();
-				if (OuterTrack)
-				{
-					OuterTrack->Modify();
-					OuterTrack->OnSectionMoved(*Section);
-				}
+			// Dragging the start of a section
+			// Ensure we arent expanding past the end time
+			NewTime = FMath::Min( NewTime, MaxFrame );
 
-				break;
+			if (bIsSlipping)
+			{
+				SectionInterface->SlipSection( NewTime );
+			}
+			else
+			{
+				SectionInterface->ResizeSection( SSRM_LeadingEdge, NewTime );
 			}
 		}
+
+		UMovieSceneTrack* OuterTrack = Section->GetTypedOuter<UMovieSceneTrack>();
+		if (OuterTrack)
+		{
+			OuterTrack->Modify();
+			OuterTrack->OnSectionMoved(*Section);
+		}
 	}
 
 	{
 		TSet<UMovieSceneTrack*> Tracks;
-		for (auto SectionHandle : Sections)
+		for (TWeakObjectPtr<UMovieSceneSection> WeakSection : Sections)
 		{
-			if (UMovieSceneTrack* Track = SectionHandle.GetSectionObject()->GetTypedOuter<UMovieSceneTrack>())
+			if (UMovieSceneTrack* Track = WeakSection.Get()->GetTypedOuter<UMovieSceneTrack>())
 			{
 				Tracks.Add(Track);
 			}
@@ -419,44 +383,6 @@ void FResizeSection::OnDrag(const FPointerEvent& MouseEvent, FVector2D LocalMous
 	Sequencer.NotifyMovieSceneDataChanged( EMovieSceneDataChangeType::TrackValueChanged );
 }
 
-void CollateTrackNodesByTrack(const TArray<TSharedRef<FSequencerDisplayNode>>& DisplayNodes, TMap<UMovieSceneTrack*, TArray<TSharedRef<FSequencerTrackNode>>>& TrackToTrackNodesMap)
-{
-	for (TSharedRef<FSequencerDisplayNode> DisplayNode : DisplayNodes)
-	{
-		if (DisplayNode->GetType() == ESequencerNode::Track)
-		{
-			TSharedRef<FSequencerTrackNode> TrackNode = StaticCastSharedRef<FSequencerTrackNode>(DisplayNode);
-			TArray<TSharedRef<FSequencerTrackNode>>* TrackNodes = TrackToTrackNodesMap.Find(TrackNode->GetTrack());
-			if (TrackNodes == nullptr)
-			{
-				TrackNodes = &TrackToTrackNodesMap.Add(TrackNode->GetTrack());
-			}
-			TrackNodes->Add(TrackNode);
-		}
-
-		CollateTrackNodesByTrack(DisplayNode->GetChildNodes(), TrackToTrackNodesMap);
-	}
-}
-
-bool TryUpdateHandleFromNewTrackNodes(const TArray<TSharedRef<FSequencerTrackNode>>& NewTrackNodes, FSectionHandle& SectionHandle)
-{
-	UMovieSceneSection* MovieSceneSection = SectionHandle.GetSectionObject();
-	for (TSharedRef<FSequencerTrackNode> NewTrackNode : NewTrackNodes)
-	{
-		const TArray<TSharedRef<ISequencerSection>> SequencerSections = NewTrackNode->GetSections();
-		for (int32 i = 0; i < SequencerSections.Num(); i++)
-		{
-			if (SequencerSections[i]->GetSectionObject() == MovieSceneSection)
-			{
-				SectionHandle.TrackNode = NewTrackNode;
-				SectionHandle.SectionIndex = i;
-				return true;
-			}
-		}
-	}
-	return false;
-}
-
 void FDuplicateKeysAndSections::OnBeginDrag(const FPointerEvent& MouseEvent, FVector2D LocalMousePos, const FVirtualTrackArea& VirtualTrackArea)
 {
 	// Begin an editor transaction and mark the section as transactional so it's state will be saved
@@ -479,14 +405,14 @@ void FDuplicateKeysAndSections::OnBeginDrag(const FPointerEvent& MouseEvent, FVe
 	// Duplicate our selections as well.
 	bool bDelayedStructureRebuild = false;
 
-	TArray<UMovieSceneSection*> SectionsToDuplicate;
-	for (const FSectionHandle& SectionHandle : Sections)
+	for (const TWeakObjectPtr<UMovieSceneSection>& WeakSection : Sections)
 	{
-		SectionsToDuplicate.Add(SectionHandle.GetSectionObject());
-	}
+		UMovieSceneSection* SectionToDuplicate = WeakSection.Get();
+		if (!SectionToDuplicate)
+		{
+			continue;
+		}
 
-	for (UMovieSceneSection* SectionToDuplicate : SectionsToDuplicate)
-	{
 		UMovieSceneSection* DuplicatedSection = DuplicateObject<UMovieSceneSection>(SectionToDuplicate, SectionToDuplicate->GetOuter());
 		UMovieSceneTrack* OwningTrack = SectionToDuplicate->GetTypedOuter<UMovieSceneTrack>();
 		OwningTrack->Modify();
@@ -512,9 +438,9 @@ void FDuplicateKeysAndSections::OnEndDrag(const FPointerEvent& MouseEvent, FVect
 	EndTransaction();
 }
 
-FManipulateSectionEasing::FManipulateSectionEasing( FSequencer& InSequencer, FSectionHandle InSection, bool _bEaseIn )
+FManipulateSectionEasing::FManipulateSectionEasing( FSequencer& InSequencer, TWeakObjectPtr<UMovieSceneSection> InSection, bool _bEaseIn )
 	: FEditToolDragOperation(InSequencer)
-	, Handle(InSection)
+	, WeakSection(InSection)
 	, bEaseIn(_bEaseIn)
 	, MouseDownTime(0)
 {
@@ -524,7 +450,7 @@ void FManipulateSectionEasing::OnBeginDrag(const FPointerEvent& MouseEvent, FVec
 {
 	Transaction.Reset( new FScopedTransaction(NSLOCTEXT("Sequencer", "DragSectionEasing", "Change Section Easing")) );
 
-	UMovieSceneSection* Section = Handle.GetSectionObject();
+	UMovieSceneSection* Section = WeakSection.Get();
 	Section->SetFlags( RF_Transactional );
 	Section->Modify();
 
@@ -550,12 +476,12 @@ void FManipulateSectionEasing::OnDrag(const FPointerEvent& MouseEvent, FVector2D
 	// Convert the current mouse position to a time
 	FFrameTime  DeltaTime = VirtualTrackArea.PixelToFrame(LocalMousePos.X) - MouseDownTime;
 
+	UMovieSceneSection* Section = WeakSection.Get();
+
 	// Snapping
 	if (Settings->GetIsSnapEnabled())
 	{
 		TArray<FFrameNumber> SnapTimes;
-
-		UMovieSceneSection* Section = Handle.GetSectionObject();
 		if (bEaseIn)
 		{
 			FFrameNumber DesiredTime = (DeltaTime + Section->GetInclusiveStartFrame() + InitValue.Get(0)).RoundToFrame();
@@ -590,8 +516,6 @@ void FManipulateSectionEasing::OnDrag(const FPointerEvent& MouseEvent, FVector2D
 		}
 	}
 
-	UMovieSceneSection* Section = Handle.GetSectionObject();
-
 	const int32 MaxEasingDuration = Section->HasStartFrame() && Section->HasEndFrame() ? MovieScene::DiscreteSize(Section->GetRange()) : TNumericLimits<int32>::Max() / 2;
 
 	if (bEaseIn)
@@ -620,12 +544,12 @@ void FManipulateSectionEasing::OnEndDrag(const FPointerEvent& MouseEvent, FVecto
 }
 
 
-FMoveKeysAndSections::FMoveKeysAndSections(FSequencer& InSequencer, const TSet<FSequencerSelectedKey>& InSelectedKeys, TArray<FSectionHandle> InSelectedSections, bool InbHotspotWasSection)
+FMoveKeysAndSections::FMoveKeysAndSections(FSequencer& InSequencer, const TSet<FSequencerSelectedKey>& InSelectedKeys, const TSet<TWeakObjectPtr<UMovieSceneSection>>& InSelectedSections, bool InbHotspotWasSection)
 	: FEditToolDragOperation(InSequencer)
 	, bHotspotWasSection(InbHotspotWasSection)
 {
 	// Filter out the keys on sections that are read only
-	for (auto SelectedKey : InSelectedKeys)
+	for (const FSequencerSelectedKey& SelectedKey : InSelectedKeys)
 	{
 		if (!SelectedKey.Section->IsReadOnly())
 		{
@@ -637,22 +561,14 @@ FMoveKeysAndSections::FMoveKeysAndSections(FSequencer& InSequencer, const TSet<F
 
 	// However, we don't want infinite sections to be movable, so we discard them from our selection.
 	// We support partially infinite (infinite on one side) sections however.
-	for (const FSectionHandle& SectionHandle : InSelectedSections)
+	for (const TWeakObjectPtr<UMovieSceneSection>& WeakSection : InSelectedSections)
 	{
-		const UMovieSceneSection* Section = SectionHandle.GetSectionObject();
+		const UMovieSceneSection* Section = WeakSection.Get();
 		if (Section->HasStartFrame() || Section->HasEndFrame())
 		{
-			Sections.Add(SectionHandle);
+			Sections.Add(WeakSection);
 		}
 	}
-
-	// Register a callback for when the node tree is updated so we can update our local Section Handle array.
-	SequencerNodeTreeUpdatedHandle = InSequencer.GetNodeTree()->OnUpdated().AddRaw(this, &FMoveKeysAndSections::OnSequencerNodeTreeUpdated);
-}
-
-FMoveKeysAndSections::~FMoveKeysAndSections()
-{
-	Sequencer.GetNodeTree()->OnUpdated().Remove(SequencerNodeTreeUpdatedHandle);
 }
 
 void FMoveKeysAndSections::OnBeginDrag(const FPointerEvent& MouseEvent, FVector2D LocalMousePos, const FVirtualTrackArea& VirtualTrackArea)
@@ -676,9 +592,9 @@ void FMoveKeysAndSections::OnBeginDrag(const FPointerEvent& MouseEvent, FVector2
 	// Now we store a relative offset to each key and section from the start position. This allows us to know how far away from
 	// the mouse each valid key/section was so we can restore their offset if needed.
 	RelativeOffsets.Reserve(Sections.Num() + Keys.Num());
-	for (int32 Index = 0; Index < Sections.Num(); ++Index)
+	for (TWeakObjectPtr<UMovieSceneSection> WeakSection : Sections)
 	{
-		UMovieSceneSection* Section = Sections[Index].GetSectionObject();
+		UMovieSceneSection* Section = WeakSection.Get();
 		FRelativeOffset Offset;
 
 		if (Section->HasStartFrame())
@@ -697,13 +613,14 @@ void FMoveKeysAndSections::OnBeginDrag(const FPointerEvent& MouseEvent, FVector2
 	// can be dragged above all other sections - this is accomplished by moving all other sections down. We store the row indices for all sections
 	// in all tracks that we're modifying so we can get them later to move them.
 	TSet<UMovieSceneTrack*> Tracks;
-	for (int32 Index = 0; Index < Sections.Num(); ++Index)
+	for (TWeakObjectPtr<UMovieSceneSection> WeakSection : Sections)
 	{
-		Tracks.Add(Sections[Index].TrackNode->GetTrack());
+		UMovieSceneSection* Section = WeakSection.Get();
+		Tracks.Add(Section->GetTypedOuter<UMovieSceneTrack>());
 	}
 	for (UMovieSceneTrack* Track : Tracks)
 	{
-		for (auto Section : Track->GetAllSections())
+		for (UMovieSceneSection* Section : Track->GetAllSections())
 		{
 			InitialSectionRowIndicies.Add(FInitialRowIndex{ Section, Section->GetRowIndex() });
 		}
@@ -811,9 +728,9 @@ void FMoveKeysAndSections::OnDrag(const FPointerEvent& MouseEvent, FVector2D Loc
 
 	// Get a list of the unique tracks in this selection and update their easing so previews draw interactively as you drag.
 	TSet<UMovieSceneTrack*> Tracks;
-	for (auto SectionHandle : Sections)
+	for (TWeakObjectPtr<UMovieSceneSection> WeakSection : Sections)
 	{
-		if (UMovieSceneTrack* Track = SectionHandle.GetSectionObject()->GetTypedOuter<UMovieSceneTrack>())
+		if (UMovieSceneTrack* Track = WeakSection.Get()->GetTypedOuter<UMovieSceneTrack>())
 		{
 			Tracks.Add(Track);
 		}
@@ -849,10 +766,10 @@ void FMoveKeysAndSections::OnEndDrag(const FPointerEvent& MouseEvent, FVector2D
 	bool bRowIndicesChanged = false;
 	TSet<UMovieSceneTrack*> Tracks;
 
-	for (auto& SectionHandle : Sections)
+	for (TWeakObjectPtr<UMovieSceneSection> WeakSection : Sections)
 	{
 		// Grab only unique tracks as multiple sections can reside on the same track.
-		Tracks.Add(SectionHandle.TrackNode->GetTrack());
+		Tracks.Add(WeakSection.Get()->GetTypedOuter<UMovieSceneTrack>());
 	}
 
 	for (UMovieSceneTrack* Track : Tracks)
@@ -866,9 +783,9 @@ void FMoveKeysAndSections::OnEndDrag(const FPointerEvent& MouseEvent, FVector2D
 		Sequencer.NotifyMovieSceneDataChanged(EMovieSceneDataChangeType::MovieSceneStructureItemsChanged);
 	}
 
-	for (const FSectionHandle& SectionHandle : Sections)
+	for (const TWeakObjectPtr<UMovieSceneSection>& WeakSection : Sections)
 	{
-		UMovieSceneSection* Section = SectionHandle.GetSectionObject();
+		UMovieSceneSection* Section = WeakSection.Get();
 		UMovieSceneTrack* OuterTrack = Cast<UMovieSceneTrack>(Section->GetOuter());
 
 		if (OuterTrack)
@@ -881,32 +798,13 @@ void FMoveKeysAndSections::OnEndDrag(const FPointerEvent& MouseEvent, FVector2D
 	EndTransaction();
 }
 
-void FMoveKeysAndSections::OnSequencerNodeTreeUpdated()
-{
-	TMap<UMovieSceneTrack*, TArray<TSharedRef<FSequencerTrackNode>>> TrackToTrackNodesMap;
-	CollateTrackNodesByTrack(Sequencer.GetNodeTree()->GetRootNodes(), TrackToTrackNodesMap);
-
-	// Update the track nodes in the handles based on the original track and section index.
-	for (FSectionHandle& SectionHandle : Sections)
-	{
-		TArray<TSharedRef<FSequencerTrackNode>>* NewTrackNodes = TrackToTrackNodesMap.Find(SectionHandle.TrackNode->GetTrack());
-		ensureMsgf(NewTrackNodes != nullptr, TEXT("Error rebuilding section handles:  Track not found after node tree update."));
-
-		if (NewTrackNodes != nullptr)
-		{
-			bool bHandleUpdated = TryUpdateHandleFromNewTrackNodes(*NewTrackNodes, SectionHandle);
-			ensureMsgf(bHandleUpdated, TEXT("Error rebuilding section handles: Track node with correct track and section index could not be found."));
-		}
-	}
-}
-
 void FMoveKeysAndSections::ModifyNonSelectedSections()
 {
 	for (const FSequencerSelectedKey& Key : Keys)
 	{
 		UMovieSceneSection* OwningSection = Key.Section;
 		const bool bHasBeenModified = ModifiedNonSelectedSections.Contains(OwningSection);
-		const bool bIsAlreadySelected = Sections.ContainsByPredicate([OwningSection](FSectionHandle Handle) { return Handle.GetSectionObject() == OwningSection; });
+		const bool bIsAlreadySelected = Sections.Contains(OwningSection);
 		if (!bHasBeenModified && !bIsAlreadySelected)
 		{
 			OwningSection->SetFlags(RF_Transactional);
@@ -918,6 +816,39 @@ void FMoveKeysAndSections::ModifyNonSelectedSections()
 	}
 }
 
+TRange<FFrameNumber> FMoveKeysAndSections::GetSectionBoundaries(const UMovieSceneSection* Section)
+{
+	// Find the borders of where you can drag to
+	FFrameNumber LowerBound = TNumericLimits<int32>::Lowest(), UpperBound = TNumericLimits<int32>::Max();
+
+	// Find the track node for this section
+	TOptional<FSectionHandle> SectionHandle = Sequencer.GetNodeTree()->GetSectionHandle(Section);
+	if (SectionHandle)
+	{
+		// Get the closest borders on either side
+		const TArray< TSharedRef<ISequencerSection> >& AllSections = SectionHandle->GetTrackNode()->GetSections();
+		for (TSharedRef<ISequencerSection> SectionInterface : AllSections)
+		{
+			const UMovieSceneSection* TestSection = SectionInterface->GetSectionObject();
+			if (!TestSection || Sections.Contains(TestSection))
+			{
+				continue;
+			}
+
+			if (TestSection->HasEndFrame() && Section->HasStartFrame() && TestSection->GetExclusiveEndFrame() <= Section->GetInclusiveStartFrame() && TestSection->GetExclusiveEndFrame() > LowerBound)
+			{
+				LowerBound = TestSection->GetExclusiveEndFrame();
+			}
+			if (TestSection->HasStartFrame() && Section->HasEndFrame() && TestSection->GetInclusiveStartFrame() >= Section->GetExclusiveEndFrame() && TestSection->GetInclusiveStartFrame() < UpperBound)
+			{
+				UpperBound = TestSection->GetInclusiveStartFrame();
+			}
+		}
+	}
+
+	return TRange<FFrameNumber>(LowerBound, UpperBound);
+}
+
 TOptional<FFrameNumber> FMoveKeysAndSections::GetMovementDeltaX(FFrameTime MouseTime)
 {
 	TOptional<FFrameNumber> DeltaX;
@@ -926,19 +857,17 @@ TOptional<FFrameNumber> FMoveKeysAndSections::GetMovementDeltaX(FFrameTime Mouse
 	const FFrameNumber MouseDeltaTime = (MouseTime - MouseTimePrev).FloorToFrame();
 
 	// Disallow movement if any of the sections can't move
-	for (int32 Index = 0; Index < Sections.Num(); ++Index)
+	for (TWeakObjectPtr<UMovieSceneSection> WeakSection : Sections)
 	{
-		const FSectionHandle& SectionHandle = Sections[Index];
-
 		// If we're moving a section that is blending with something then it's OK if it overlaps stuff, the blend amount will get updated at the end.
-		UMovieSceneSection* Section = SectionHandle.GetSectionObject();
-		if (Section->GetBlendType().IsValid())
+		UMovieSceneSection* Section = WeakSection.Get();
+		if (!Section || Section->GetBlendType().IsValid())
 		{
 			continue;
 		}
 
 		// We'll calculate this section's borders and clamp the possible delta time to be less than that
-		TRange<FFrameNumber> SectionBoundaries = GetSectionBoundaries(Section, Sections, SectionHandle.TrackNode);
+		TRange<FFrameNumber> SectionBoundaries = GetSectionBoundaries(Section);
 
 		FFrameNumber LeftMovementMaximum = MovieScene::DiscreteInclusiveLower(SectionBoundaries);
 		FFrameNumber RightMovementMaximum = MovieScene::DiscreteExclusiveUpper(SectionBoundaries);
@@ -983,40 +912,47 @@ bool FMoveKeysAndSections::HandleSectionMovement(FFrameTime MouseTime, FVector2D
 
 	// If sections are all on different rows, don't set row indices for anything because it leads to odd behavior.
 	bool bSectionsAreOnDifferentRows = false;
-	int32 FirstRowIndex = Sections[0].GetSectionObject()->GetRowIndex();
-	TArray<const UMovieSceneSection*> SectionsBeingMoved;
-	for (auto SectionHandle : Sections)
+	int32 FirstRowIndex = Sections[0].Get()->GetRowIndex();
+
+	for (TWeakObjectPtr<UMovieSceneSection> WeakSection : Sections)
 	{
-		if (FirstRowIndex != SectionHandle.GetSectionObject()->GetRowIndex())
+		UMovieSceneSection* Section = WeakSection.Get();
+		if (FirstRowIndex != Section->GetRowIndex())
 		{
 			bSectionsAreOnDifferentRows = true;
 		}
-		SectionsBeingMoved.Add(SectionHandle.GetSectionObject());
 	}
 
 	bool bRowIndexChanged = false;
-	for (int32 Index = 0; Index < Sections.Num(); ++Index)
+	for (TWeakObjectPtr<UMovieSceneSection> WeakSection : Sections)
 	{
-		auto& Handle = Sections[Index];
-		UMovieSceneSection* Section = Handle.GetSectionObject();
-
+		UMovieSceneSection* Section = WeakSection.Get();
+		UMovieSceneTrack* Track = Section->GetTypedOuter<UMovieSceneTrack>();
 
-		const TArray<UMovieSceneSection*>& AllSections = Handle.TrackNode->GetTrack()->GetAllSections();
+		const TArray<UMovieSceneSection*>& AllSections = Track->GetAllSections();
 
 		TArray<UMovieSceneSection*> NonDraggedSections;
 		for (UMovieSceneSection* TrackSection : AllSections)
 		{
-			if (!SectionsBeingMoved.Contains(TrackSection))
+			if (!Sections.Contains(TrackSection))
 			{
 				NonDraggedSections.Add(TrackSection);
 			}
 		}
 
+		TOptional<FSectionHandle> SectionHandle = Sequencer.GetNodeTree()->GetSectionHandle(Section);
+		if (!SectionHandle)
+		{
+			continue;
+		}
+
+		TSharedRef<FSequencerTrackNode> TrackNode = SectionHandle->GetTrackNode();
+
 		int32 TargetRowIndex = Section->GetRowIndex();
 
 		// Handle vertical dragging to re-arrange tracks. We don't support vertical rearranging if you're dragging via
 		// a key, as the built in offset causes it to always jump down a row even without moving the mouse.
-		if (Handle.TrackNode->GetTrack()->SupportsMultipleRows() && AllSections.Num() > 1 && bHotspotWasSection)
+		if (Track->SupportsMultipleRows() && AllSections.Num() > 1 && bHotspotWasSection)
 		{
 			// Compute the max row index whilst disregarding the one we're dragging
 			int32 MaxRowIndex = 0;
@@ -1029,17 +965,17 @@ bool FMoveKeysAndSections::HandleSectionMovement(FFrameTime MouseTime, FVector2D
 			}
 
 			// Handle sub-track and non-sub-track dragging
-			if (Handle.TrackNode->GetSubTrackMode() == FSequencerTrackNode::ESubTrackMode::None)
+			if (TrackNode->GetSubTrackMode() == FSequencerTrackNode::ESubTrackMode::None)
 			{
 				const int32 NumRows = FMath::Max(Section->GetRowIndex() + 1, MaxRowIndex);
 
 				// Find the total height of the track - this is necessary because tracks may contain key areas, but they will not use sub tracks unless there is more than one row
 				float VirtualSectionBottom = 0.f;
-				Handle.TrackNode->TraverseVisible_ParentFirst([&](FSequencerDisplayNode& Node) { VirtualSectionBottom = Node.GetVirtualBottom(); return true; }, true);
+				TrackNode->TraverseVisible_ParentFirst([&](FSequencerDisplayNode& Node) { VirtualSectionBottom = Node.GetVirtualBottom(); return true; }, true);
 
 				// Assume same height rows
-				const float VirtualSectionTop = Handle.TrackNode->GetVirtualTop();
-				const float VirtualSectionHeight = VirtualSectionBottom - Handle.TrackNode->GetVirtualTop();
+				const float VirtualSectionTop = TrackNode->GetVirtualTop();
+				const float VirtualSectionHeight = VirtualSectionBottom - TrackNode->GetVirtualTop();
 
 				const float VirtualRowHeight = VirtualSectionHeight / NumRows;
 				const float MouseOffsetWithinRow = VirtualMousePos.Y - (VirtualSectionTop + (VirtualRowHeight * TargetRowIndex));
@@ -1055,10 +991,13 @@ bool FMoveKeysAndSections::HandleSectionMovement(FFrameTime MouseTime, FVector2D
 				{
 					TargetRowIndex = -1;
 				}
+
+				// Ensure the track path is expanded
+				TrackNode->SetExpansionState(true);
 			}
-			else if (Handle.TrackNode->GetSubTrackMode() == FSequencerTrackNode::ESubTrackMode::SubTrack)
+			else if (TrackNode->GetSubTrackMode() == FSequencerTrackNode::ESubTrackMode::SubTrack)
 			{
-				TSharedPtr<FSequencerTrackNode> ParentTrack = StaticCastSharedPtr<FSequencerTrackNode>(Handle.TrackNode->GetParent());
+				TSharedPtr<FSequencerTrackNode> ParentTrack = StaticCastSharedPtr<FSequencerTrackNode>(TrackNode->GetParent());
 				if (ensure(ParentTrack.IsValid()))
 				{
 					for (int32 ChildIndex = 0; ChildIndex < ParentTrack->GetChildNodes().Num(); ++ChildIndex)
@@ -1106,9 +1045,9 @@ bool FMoveKeysAndSections::HandleSectionMovement(FFrameTime MouseTime, FVector2D
 				{
 					// If the sections being moved are all at the top, and all others are below it, do nothing
 					bool bSectionsBeingMovedAreAtTop = true;
-					for (auto InitialRowIndex : InitialSectionRowIndicies)
+					for (const FInitialRowIndex& InitialRowIndex : InitialSectionRowIndicies)
 					{
-						if (!SectionsBeingMoved.Contains(InitialRowIndex.Section))
+						if (!Sections.Contains(InitialRowIndex.Section))
 						{
 							if (InitialRowIndex.RowIndex <= FirstRowIndex)
 							{
@@ -1120,9 +1059,9 @@ bool FMoveKeysAndSections::HandleSectionMovement(FFrameTime MouseTime, FVector2D
 
 					if (!bSectionsBeingMovedAreAtTop)
 					{
-						for (auto InitialRowIndex : InitialSectionRowIndicies)
+						for (const FInitialRowIndex& InitialRowIndex : InitialSectionRowIndicies)
 						{
-							if (!SectionsBeingMoved.Contains(InitialRowIndex.Section))
+							if (!Sections.Contains(InitialRowIndex.Section))
 							{
 								InitialRowIndex.Section->Modify();
 								InitialRowIndex.Section->SetRowIndex(InitialRowIndex.RowIndex + 1);
@@ -1160,7 +1099,7 @@ void FMoveKeysAndSections::HandleKeyMovement(TOptional<FFrameNumber> MaxDeltaX,
 	for (int32 Index = 0; Index < CurrentKeyTimes.Num(); ++Index)
 	{
 		FSequencerSelectedKey& SelectedKey = KeysAsArray[Index];
-		const bool bOwningSectionIsSelected = Sections.ContainsByPredicate([SelectedKey](FSectionHandle Handle) { return Handle.GetSectionObject() == SelectedKey.Section; });
+		const bool bOwningSectionIsSelected = Sections.Contains(SelectedKey.Section);
 
 		// We don't want to apply delta if we have the key's section selected as well, otherwise they get double
 		// transformed (moving the section moves the keys + we add the delta to the key positions).
diff --git a/Engine/Source/Editor/Sequencer/Private/Tools/EditToolDragOperations.h b/Engine/Source/Editor/Sequencer/Private/Tools/EditToolDragOperations.h
--- a/Engine/Source/Editor/Sequencer/Private/Tools/EditToolDragOperations.h
+++ b/Engine/Source/Editor/Sequencer/Private/Tools/EditToolDragOperations.h
@@ -39,7 +39,7 @@ public:
 protected:
 
 	/** begin a new scoped transaction for this drag */
-	void BeginTransaction( TArray< FSectionHandle >& Sections, const FText& TransactionDesc );
+	void BeginTransaction( TArray< TWeakObjectPtr<UMovieSceneSection> >& Sections, const FText& TransactionDesc );
 
 	/** End an existing scoped transaction if one exists */
 	void EndTransaction();
@@ -66,7 +66,7 @@ class FResizeSection
 public:
 
 	/** Create and initialize a new instance. */
-	FResizeSection( FSequencer& InSequencer, TArray<FSectionHandle> Sections, bool bInDraggingByEnd, bool bIsSlipping );
+	FResizeSection( FSequencer& InSequencer, const TSet<TWeakObjectPtr<UMovieSceneSection>>& Sections, bool bInDraggingByEnd, bool bIsSlipping );
 
 public:
 
@@ -80,7 +80,7 @@ public:
 private:
 
 	/** The sections we are interacting with */
-	TArray<FSectionHandle> Sections;
+	TArray<TWeakObjectPtr<UMovieSceneSection>> Sections;
 
 	/********************************************************/
 	struct FPreDragChannelData
@@ -128,8 +128,7 @@ class FMoveKeysAndSections
 	: public FEditToolDragOperation
 {
 public:
-	FMoveKeysAndSections(FSequencer& InSequencer, const TSet<FSequencerSelectedKey>& InSelectedKeys, TArray<FSectionHandle> InSelectedSections, bool InbHotspotWasSection);
-	~FMoveKeysAndSections();
+	FMoveKeysAndSections(FSequencer& InSequencer, const TSet<FSequencerSelectedKey>& InSelectedKeys, const TSet<TWeakObjectPtr<UMovieSceneSection>>& InSelectedSections, bool InbHotspotWasSection);
 
 	// FEditToolDragOperation interface
 	virtual void OnBeginDrag(const FPointerEvent& MouseEvent, FVector2D LocalMousePos, const FVirtualTrackArea& VirtualTrackArea) override;
@@ -141,6 +140,10 @@ public:
 protected:
 	/** Calculate the possible horizontal movement we can, constrained by sections running into things. */
 	TOptional<FFrameNumber> GetMovementDeltaX(FFrameTime MouseTime);
+
+	/** Get the bounds within which the specified section can move based on its surrounding sections */
+	TRange<FFrameNumber> GetSectionBoundaries(const UMovieSceneSection* Section);
+
 	/** Move selected sections, if any. */
 	bool HandleSectionMovement(FFrameTime MouseTime, FVector2D VirtualMousePos, FVector2D LocalMousePos, TOptional<FFrameNumber> MaxDeltaX, FFrameNumber DesiredDeltaX);
 	/** Move selected keys, if any. */
@@ -153,7 +156,7 @@ protected:
 
 protected:
 	/** Array of sections that we're moving. */
-	TArray<FSectionHandle> Sections;
+	TArray<TWeakObjectPtr<UMovieSceneSection>> Sections;
 
 	/** Set of keys that are being moved. */
 	TSet<FSequencerSelectedKey> Keys;
@@ -200,9 +203,6 @@ protected:
 	/** Optional snap field to use when dragging */
 	TOptional<FSequencerSnapField> SnapField;
 
-	/** A handle for the sequencer node tree updated delegate. */
-	FDelegateHandle SequencerNodeTreeUpdatedHandle;
-
 	/** If the user is moving them via clicking on the Section then we'll allow vertical re-arranging, otherwise not. */
 	bool bHotspotWasSection;
 };
@@ -214,7 +214,7 @@ class FDuplicateKeysAndSections : public FMoveKeysAndSections
 {
 public:
 
-	FDuplicateKeysAndSections( FSequencer& InSequencer, const TSet<FSequencerSelectedKey>& InSelectedKeys, TArray<FSectionHandle> InSelectedSections, bool InbHotspotWasSection)
+	FDuplicateKeysAndSections( FSequencer& InSequencer, const TSet<FSequencerSelectedKey>& InSelectedKeys, const TSet<TWeakObjectPtr<UMovieSceneSection>>& InSelectedSections, bool InbHotspotWasSection)
 		: FMoveKeysAndSections(InSequencer, InSelectedKeys, InSelectedSections, InbHotspotWasSection)
 	{}
 
@@ -233,7 +233,7 @@ class FManipulateSectionEasing
 public:
 
 	/** Create and initialize a new instance. */
-	FManipulateSectionEasing( FSequencer& InSequencer, FSectionHandle InSection, bool bEaseIn );
+	FManipulateSectionEasing( FSequencer& InSequencer, TWeakObjectPtr<UMovieSceneSection> InSection, bool bEaseIn );
 
 public:
 
@@ -247,7 +247,7 @@ public:
 private:
 
 	/** The sections we are interacting with */
-	FSectionHandle Handle;
+	TWeakObjectPtr<UMovieSceneSection> WeakSection;
 
 	/** true if editing the section's ease in, false for ease out */
 	bool bEaseIn;
