Fix include

diff --git a/Engine/Source/Editor/Sequencer/Private/VirtualTrackArea.h b/Engine/Source/Editor/Sequencer/Private/VirtualTrackArea.h
--- a/Engine/Source/Editor/Sequencer/Private/VirtualTrackArea.h
+++ b/Engine/Source/Editor/Sequencer/Private/VirtualTrackArea.h
@@ -6,6 +6,7 @@
 #include "Layout/Geometry.h"
 #include "SequencerSelectedKey.h"
 #include "SequencerHotspots.h"
+#include "SectionHandle.h"
 #include "CommonMovieSceneTools.h"
 
 class FSequencer;
