Sequencer: Check for valid parent

diff --git a/Engine/Source/Editor/Sequencer/Private/SSequencerTreeView.cpp b/Engine/Source/Editor/Sequencer/Private/SSequencerTreeView.cpp
--- a/Engine/Source/Editor/Sequencer/Private/SSequencerTreeView.cpp
+++ b/Engine/Source/Editor/Sequencer/Private/SSequencerTreeView.cpp
@@ -44,7 +44,7 @@ FMargin SSequencerTreeViewRow::GetRowPadding() const
 	TSharedPtr<FSequencerDisplayNode> PinnedNode = Node.Pin();
 	TSharedPtr<FSequencerDisplayNode> ParentNode = PinnedNode ? PinnedNode->GetParentOrRoot() : nullptr;
 
-	if (ParentNode->GetType() == ESequencerNode::Root && ParentNode->GetChildNodes()[0] != PinnedNode)
+	if (ParentNode.IsValid() && ParentNode->GetType() == ESequencerNode::Root && ParentNode->GetChildNodes()[0] != PinnedNode)
 	{
 		return FMargin(0.f, 1.f, 0.f, 0.f);
 	}
