Deprecated removal

diff --git a/Engine/Source/Editor/MovieSceneTools/Private/K2Node_GetSequenceBindings.cpp b/Engine/Source/Editor/MovieSceneTools/Private/K2Node_GetSequenceBindings.cpp
deleted file mode 100644
diff --git a/Engine/Source/Editor/MovieSceneTools/Private/K2Node_GetSequenceBindings.h b/Engine/Source/Editor/MovieSceneTools/Private/K2Node_GetSequenceBindings.h
deleted file mode 100644
diff --git a/Engine/Source/Editor/Sequencer/Public/ISequencerModule.h b/Engine/Source/Editor/Sequencer/Public/ISequencerModule.h
--- a/Engine/Source/Editor/Sequencer/Public/ISequencerModule.h
+++ b/Engine/Source/Editor/Sequencer/Public/ISequencerModule.h
@@ -302,20 +302,6 @@ public:
 		return RegisterTrackEditor(FOnCreateTrackEditor::CreateStatic(PropertyTrackEditorType::CreateTrackEditor), PropertyTypes);
 	}
 
-public:
-
-	UE_DEPRECATED(4.16, "Please use RegisterTrackEditor")
-	FDelegateHandle RegisterTrackEditor_Handle(FOnCreateTrackEditor InOnCreateTrackEditor)
-	{
-		return RegisterTrackEditor(InOnCreateTrackEditor, TArrayView<FAnimatedPropertyKey>());
-	}
-	
-	UE_DEPRECATED(4.16, "Please use UnRegisterTrackEditor")
-	void UnRegisterTrackEditor_Handle(FDelegateHandle InHandle)
-	{
-		UnRegisterTrackEditor(InHandle);
-	}
-
 private:
 
 	/** Map of sequencer interfaces for movie scene channel types, keyed on channel UStruct name */
diff --git a/Engine/Source/Editor/Sequencer/Public/ISequencerSection.h b/Engine/Source/Editor/Sequencer/Public/ISequencerSection.h
--- a/Engine/Source/Editor/Sequencer/Public/ISequencerSection.h
+++ b/Engine/Source/Editor/Sequencer/Public/ISequencerSection.h
@@ -68,12 +68,6 @@ public:
 	 */
 	virtual TSharedRef<SWidget> GenerateSectionWidget() { return SNullWidget::NullWidget; }
 
-	UE_DEPRECATED(4.20, "Please override Sequencer::DrawKeys instead")
-	virtual const FSlateBrush* GetKeyBrush(FKeyHandle KeyHandle) const { return nullptr; }
-
-	UE_DEPRECATED(4.20, "Please override Sequencer::DrawKeys instead")
-	virtual FVector2D GetKeyBrushOrigin( FKeyHandle KeyHandle ) const { return FVector2D(0.0f, 0.0f); }
-
 	/**
 	 * Called when the section is double clicked
 	 *
diff --git a/Engine/Source/Editor/Sequencer/Public/ISequencerTrackEditor.h b/Engine/Source/Editor/Sequencer/Public/ISequencerTrackEditor.h
--- a/Engine/Source/Editor/Sequencer/Public/ISequencerTrackEditor.h
+++ b/Engine/Source/Editor/Sequencer/Public/ISequencerTrackEditor.h
@@ -128,9 +128,6 @@ public:
 	 */
 	virtual bool OnAllowDrop(const FDragDropEvent& DragDropEvent, UMovieSceneTrack* Track, int32 RowIndex, const FGuid& TargetObjectGuid) = 0;
 
-	UE_DEPRECATED(4.19, "Use OnAllowDrop with a given RowIndex and Guid")
-	virtual bool OnAllowDrop(const FDragDropEvent& DragDropEvent, UMovieSceneTrack* Track) { return OnAllowDrop(DragDropEvent, Track, 0, FGuid()); }
-
 	/**
 	 * Called when an asset is dropped directly onto a track.
 	 *
@@ -142,9 +139,6 @@ public:
 	 */	
 	virtual FReply OnDrop(const FDragDropEvent& DragDropEvent, UMovieSceneTrack* Track, int32 RowIndex, const FGuid& TargetObjectGuid) = 0;
 
-	UE_DEPRECATED(4.19, "Use OnDrop with a given RowIndex")
-	virtual FReply OnDrop(const FDragDropEvent& DragDropEvent, UMovieSceneTrack* Track) { return OnDrop(DragDropEvent, Track, 0, FGuid()); }
-
 	/**
 	 * Called to generate a section layout for a particular section.
 	 *
