Sequencer: Allow display nodes to show the add object bindings menu and the add tracks menu. By default, they do not, but folder nodes do, so you can add these tracks directly to a folder.

diff --git a/Engine/Source/Editor/Sequencer/Private/DisplayNodes/SequencerDisplayNode.cpp b/Engine/Source/Editor/Sequencer/Private/DisplayNodes/SequencerDisplayNode.cpp
--- a/Engine/Source/Editor/Sequencer/Private/DisplayNodes/SequencerDisplayNode.cpp
+++ b/Engine/Source/Editor/Sequencer/Private/DisplayNodes/SequencerDisplayNode.cpp
@@ -638,6 +638,21 @@ TSharedPtr<SWidget> FSequencerDisplayNode::OnSummonContextMenu()
 	const bool bShouldCloseWindowAfterMenuSelection = true;
 	FMenuBuilder MenuBuilder(bShouldCloseWindowAfterMenuSelection, GetSequencer().GetCommandBindings());
 
+	// let track editors & object bindings populate the menu
+	if (CanAddObjectBindingsMenu())
+	{
+		MenuBuilder.BeginSection("ObjectBindings");
+		GetSequencer().BuildAddObjectBindingsMenu(MenuBuilder);
+		MenuBuilder.EndSection();
+	}
+
+	if (CanAddTracksMenu())
+	{
+		MenuBuilder.BeginSection("AddTracks");
+		GetSequencer().BuildAddTrackMenu(MenuBuilder);
+		MenuBuilder.EndSection();
+	}
+
 	BuildContextMenu(MenuBuilder);
 
 	return MenuBuilder.MakeWidget();
diff --git a/Engine/Source/Editor/Sequencer/Private/DisplayNodes/SequencerDisplayNode.h b/Engine/Source/Editor/Sequencer/Private/DisplayNodes/SequencerDisplayNode.h
--- a/Engine/Source/Editor/Sequencer/Private/DisplayNodes/SequencerDisplayNode.h
+++ b/Engine/Source/Editor/Sequencer/Private/DisplayNodes/SequencerDisplayNode.h
@@ -275,6 +275,12 @@ public:
 	/** What sort of context menu this node summons */
 	virtual void BuildContextMenu(FMenuBuilder& MenuBuilder);
 
+	/** Can this node show the add object bindings menu? */
+	virtual bool CanAddObjectBindingsMenu() const { return false; }
+
+	/** Can this node show the add tracks menu? */
+	virtual bool CanAddTracksMenu() const { return false; }
+
 	/**
 	 * @return The name of the node (for identification purposes)
 	 */
diff --git a/Engine/Source/Editor/Sequencer/Private/DisplayNodes/SequencerFolderNode.h b/Engine/Source/Editor/Sequencer/Private/DisplayNodes/SequencerFolderNode.h
--- a/Engine/Source/Editor/Sequencer/Private/DisplayNodes/SequencerFolderNode.h
+++ b/Engine/Source/Editor/Sequencer/Private/DisplayNodes/SequencerFolderNode.h
@@ -31,6 +31,8 @@ public:
 	virtual TOptional<EItemDropZone> CanDrop( FSequencerDisplayNodeDragDropOp& DragDropOp, EItemDropZone ItemDropZone ) const override;
 	virtual void Drop( const TArray<TSharedRef<FSequencerDisplayNode>>& DraggedNodes, EItemDropZone ItemDropZone ) override;
 	virtual void BuildContextMenu(FMenuBuilder& MenuBuilder) override;
+	virtual bool CanAddObjectBindingsMenu() const override { return true; }
+	virtual bool CanAddTracksMenu() const override { return true; }
 	virtual int32 GetSortingOrder() const override;
 	virtual void SetSortingOrder(const int32 InSortingOrder) override;
 	virtual void ModifyAndSetSortingOrder(const int32 InSortingOrder) override;
