Sequencer: Add component/property tracks for selected objects

diff --git a/Engine/Plugins/MovieScene/LevelSequenceEditor/Source/LevelSequenceEditor/Private/LevelSequenceEditorToolkit.cpp b/Engine/Plugins/MovieScene/LevelSequenceEditor/Source/LevelSequenceEditor/Private/LevelSequenceEditorToolkit.cpp
--- a/Engine/Plugins/MovieScene/LevelSequenceEditor/Source/LevelSequenceEditor/Private/LevelSequenceEditorToolkit.cpp
+++ b/Engine/Plugins/MovieScene/LevelSequenceEditor/Source/LevelSequenceEditor/Private/LevelSequenceEditorToolkit.cpp
@@ -553,7 +553,24 @@ void FLevelSequenceEditorToolkit::OnSequencerReceivedFocus()
 
 void FLevelSequenceEditorToolkit::HandleAddComponentActionExecute(UActorComponent* Component)
 {
-	Sequencer->GetHandleToObject(Component);
+	const FScopedTransaction Transaction(LOCTEXT("AddComponent", "Add Component"));
+
+	FString ComponentName = Component->GetName();
+	USelection* SelectedActors = GEditor->GetSelectedActors();
+	for (FSelectionIterator Iter(*SelectedActors); Iter; ++Iter)
+	{
+		AActor* Actor = CastChecked<AActor>(*Iter);
+
+		TArray<UActorComponent*> OutComponents;
+		Actor->GetComponents(OutComponents);
+		for (UActorComponent* ActorComponent : OutComponents)
+		{
+			if (ActorComponent->GetName() == ComponentName)
+			{
+				Sequencer->GetHandleToObject(ActorComponent);
+			}
+		}
+	}
 }
 
 
diff --git a/Engine/Source/Editor/Sequencer/Private/DisplayNodes/SequencerObjectBindingNode.cpp b/Engine/Source/Editor/Sequencer/Private/DisplayNodes/SequencerObjectBindingNode.cpp
--- a/Engine/Source/Editor/Sequencer/Private/DisplayNodes/SequencerObjectBindingNode.cpp
+++ b/Engine/Source/Editor/Sequencer/Private/DisplayNodes/SequencerObjectBindingNode.cpp
@@ -1131,7 +1131,7 @@ void FSequencerObjectBindingNode::HandleLabelsSubMenuCreate(FMenuBuilder& MenuBu
 void FSequencerObjectBindingNode::HandlePropertyMenuItemExecute(FPropertyPath PropertyPath)
 {
 	FSequencer& Sequencer = GetSequencer();
-	UObject* BoundObject = GetSequencer().FindSpawnedObjectOrTemplate(ObjectBinding);
+	UObject* BoundObject = Sequencer.FindSpawnedObjectOrTemplate(ObjectBinding);
 
 	TArray<UObject*> KeyableBoundObjects;
 	if (BoundObject != nullptr)
@@ -1142,6 +1142,25 @@ void FSequencerObjectBindingNode::HandlePropertyMenuItemExecute(FPropertyPath Pr
 		}
 	}
 
+	for (const TSharedRef<FSequencerDisplayNode>& Node : Sequencer.GetSelection().GetSelectedOutlinerNodes())
+	{
+		if (Node->GetType() != ESequencerNode::Object)
+		{
+			continue;
+		}
+
+		auto ObjectBindingNode = StaticCastSharedRef<FSequencerObjectBindingNode>(Node);
+
+		FGuid Guid = ObjectBindingNode->GetObjectBinding();
+		for (auto RuntimeObject : Sequencer.FindBoundObjects(Guid, Sequencer.GetFocusedTemplateID()))
+		{
+			if (Sequencer.CanKeyProperty(FCanKeyPropertyParams(RuntimeObject->GetClass(), PropertyPath)))
+			{
+				KeyableBoundObjects.AddUnique(RuntimeObject.Get());
+			}
+		}
+	}
+
 	// When auto setting track defaults are disabled, force add a key so that the changed
 	// value is saved and is propagated to the property.
 	FKeyPropertyParams KeyPropertyParams(KeyableBoundObjects, PropertyPath, Sequencer.GetAutoSetTrackDefaults() == false ? ESequencerKeyMode::ManualKeyForced : ESequencerKeyMode::ManualKey);
