Sequencer: Set the min desired width for the key editors, rather than the override width

diff --git a/Engine/Source/Editor/Sequencer/Private/SKeyAreaEditorSwitcher.cpp b/Engine/Source/Editor/Sequencer/Private/SKeyAreaEditorSwitcher.cpp
--- a/Engine/Source/Editor/Sequencer/Private/SKeyAreaEditorSwitcher.cpp
+++ b/Engine/Source/Editor/Sequencer/Private/SKeyAreaEditorSwitcher.cpp
@@ -45,8 +45,8 @@ void SKeyAreaEditorSwitcher::Rebuild()
 			[
 				SNew(SBox)
 				.IsEnabled(!KeyAreaNode->GetSequencer().IsReadOnly())
-				.WidthOverride(100)
 				.HAlign(HAlign_Left)
+				.MinDesiredWidth(100)
 				.Visibility(this, &SKeyAreaEditorSwitcher::GetWidgetVisibility, Index)
 				[
 					KeyArea->CreateKeyEditor(KeyAreaNode->GetSequencer().AsShared(), ObjectBindingID)
