Sequencer: Focus on (scroll timeline to) newly created sections if they're not in view

diff --git a/Engine/Source/Editor/Sequencer/Private/Sequencer.cpp b/Engine/Source/Editor/Sequencer/Private/Sequencer.cpp
--- a/Engine/Source/Editor/Sequencer/Private/Sequencer.cpp
+++ b/Engine/Source/Editor/Sequencer/Private/Sequencer.cpp
@@ -5550,6 +5550,36 @@ void FSequencer::ThrobKeySelection()
 
 void FSequencer::ThrobSectionSelection()
 {
+	// Scrub to the beginning of newly created sections if they're out of view
+	TOptional<FFrameNumber> ScrubFrame;
+	for (TWeakObjectPtr<UMovieSceneSection> SelectedSectionPtr : Selection.GetSelectedSections())
+	{
+		if (SelectedSectionPtr.IsValid() && SelectedSectionPtr->HasStartFrame())
+		{
+			if (!ScrubFrame.IsSet() || (ScrubFrame.GetValue() > SelectedSectionPtr->GetInclusiveStartFrame()))
+			{
+				ScrubFrame = SelectedSectionPtr->GetInclusiveStartFrame();
+			}
+		}
+	}
+
+	if (ScrubFrame.IsSet())
+	{
+		float ScrubTime = GetFocusedDisplayRate().AsSeconds(FFrameRate::TransformTime(ScrubFrame.GetValue(), GetFocusedTickResolution(), GetFocusedDisplayRate()));
+
+		TRange<double> NewViewRange = GetViewRange();
+
+		if (!NewViewRange.Contains(ScrubTime))
+		{
+			double MidRange = (NewViewRange.GetUpperBoundValue() - NewViewRange.GetLowerBoundValue()) / 2.0 + NewViewRange.GetLowerBoundValue();
+
+			NewViewRange.SetLowerBoundValue(NewViewRange.GetLowerBoundValue() - (MidRange - ScrubTime));
+			NewViewRange.SetUpperBoundValue(NewViewRange.GetUpperBoundValue() - (MidRange - ScrubTime));
+
+			SetViewRange(NewViewRange, EViewRangeInterpolation::Animated);
+		}
+	}
+
 	SSequencerSection::ThrobSectionSelection();
 }
 
