Sequencer: Disable copy menu for category nodes (ie. copying Location, Rotation Scale individually is not supported)

diff --git a/Engine/Source/Editor/Sequencer/Private/Sequencer.cpp b/Engine/Source/Editor/Sequencer/Private/Sequencer.cpp
--- a/Engine/Source/Editor/Sequencer/Private/Sequencer.cpp
+++ b/Engine/Source/Editor/Sequencer/Private/Sequencer.cpp
@@ -9476,10 +9476,8 @@ void FSequencer::BindCommands()
 					// later on we will filter out the invalid nodes in CopySelection() or CutSelection()
 					return true;
 				}
-				else if (Node->GetParent().IsValid() && Node->GetParent()->GetType() == ESequencerNode::Track)
+				else if (Node->GetParent().IsValid() && Node->GetParent()->GetType() == ESequencerNode::Track && Node->GetType() != ESequencerNode::Category)
 				{
-					// Although copying only the child nodes (ex. translation) is not allowed, we still show the copy & cut button
-					// so that users are not misled and can achieve this in copy/cut the parent node (ex. transform)
 					return true;
 				}
 			}
