Sequencer: Fixed one-frame delay between hovering over a key and being able to select it

diff --git a/Engine/Source/Editor/Sequencer/Private/SSequencer.h b/Engine/Source/Editor/Sequencer/Private/SSequencer.h
--- a/Engine/Source/Editor/Sequencer/Private/SSequencer.h
+++ b/Engine/Source/Editor/Sequencer/Private/SSequencer.h
@@ -237,6 +237,9 @@ public:
 	/** Generate a helper structure that can be used to transform between phsyical space and virtual space in the track area */
 	FVirtualTrackArea GetVirtualTrackArea() const;
 
+	/** Access this widget's track area widget */
+	TSharedPtr<SSequencerTrackArea> GetTrackAreaWidget() const { return TrackArea; }
+
 	/** Get an array of section handles for the given set of movie scene sections */
 	TArray<FSectionHandle> GetSectionHandles(const TSet<TWeakObjectPtr<UMovieSceneSection>>& DesiredSections) const;
 
diff --git a/Engine/Source/Editor/Sequencer/Private/Sequencer.cpp b/Engine/Source/Editor/Sequencer/Private/Sequencer.cpp
--- a/Engine/Source/Editor/Sequencer/Private/Sequencer.cpp
+++ b/Engine/Source/Editor/Sequencer/Private/Sequencer.cpp
@@ -9097,6 +9097,12 @@ void FSequencer::SetHotspot(TSharedPtr<ISequencerHotspot> NewHotspot)
 	{
 		Hotspot = MoveTemp(NewHotspot);
 	}
+
+	// Simulate an update-on-hover for the new hotspot to ensure that any hover behavior doesn't have to wait until the next frame
+	if (Hotspot)
+	{
+		Hotspot->UpdateOnHover(*SequencerWidget->GetTrackAreaWidget(), *this);
+	}
 }
 
 void FSequencer::BindCommands()
