Sequencer: Added default key interpolation settings to main menu bar

diff --git a/Engine/Source/Editor/Sequencer/Private/SSequencer.cpp b/Engine/Source/Editor/Sequencer/Private/SSequencer.cpp
--- a/Engine/Source/Editor/Sequencer/Private/SSequencer.cpp
+++ b/Engine/Source/Editor/Sequencer/Private/SSequencer.cpp
@@ -843,41 +843,38 @@ TSharedRef<SWidget> SSequencer::MakeToolBar()
 
 		ToolBarBuilder.AddSeparator();
 
-		if( SequencerPtr.Pin()->IsLevelEditorSequencer() )
-		{
-			TAttribute<FSlateIcon> KeyGroupModeIcon;
-			KeyGroupModeIcon.Bind(TAttribute<FSlateIcon>::FGetter::CreateLambda([&] {
-				switch (SequencerPtr.Pin()->GetKeyGroupMode())
-				{
-				case EKeyGroupMode::KeyAll:
-					return FSequencerCommands::Get().SetKeyAll->GetIcon();
-				case EKeyGroupMode::KeyGroup:
-					return FSequencerCommands::Get().SetKeyGroup->GetIcon();
-				default: // EKeyGroupMode::KeyChanged
-					return FSequencerCommands::Get().SetKeyChanged->GetIcon();
-				}
-			}));
+		TAttribute<FSlateIcon> KeyGroupModeIcon;
+		KeyGroupModeIcon.Bind(TAttribute<FSlateIcon>::FGetter::CreateLambda([&] {
+			switch (SequencerPtr.Pin()->GetKeyGroupMode())
+			{
+			case EKeyGroupMode::KeyAll:
+				return FSequencerCommands::Get().SetKeyAll->GetIcon();
+			case EKeyGroupMode::KeyGroup:
+				return FSequencerCommands::Get().SetKeyGroup->GetIcon();
+			default: // EKeyGroupMode::KeyChanged
+				return FSequencerCommands::Get().SetKeyChanged->GetIcon();
+			}
+		}));
 
-			TAttribute<FText> KeyGroupModeToolTip;
-			KeyGroupModeToolTip.Bind(TAttribute<FText>::FGetter::CreateLambda([&] {
-				switch (SequencerPtr.Pin()->GetKeyGroupMode())
-				{
-				case EKeyGroupMode::KeyAll:
-					return FSequencerCommands::Get().SetKeyAll->GetDescription();
-				case EKeyGroupMode::KeyGroup:
-					return FSequencerCommands::Get().SetKeyGroup->GetDescription();
-				default: // EKeyGroupMode::KeyChanged
-					return FSequencerCommands::Get().SetKeyChanged->GetDescription();
-				}
-			}));
+		TAttribute<FText> KeyGroupModeToolTip;
+		KeyGroupModeToolTip.Bind(TAttribute<FText>::FGetter::CreateLambda([&] {
+			switch (SequencerPtr.Pin()->GetKeyGroupMode())
+			{
+			case EKeyGroupMode::KeyAll:
+				return FSequencerCommands::Get().SetKeyAll->GetDescription();
+			case EKeyGroupMode::KeyGroup:
+				return FSequencerCommands::Get().SetKeyGroup->GetDescription();
+			default: // EKeyGroupMode::KeyChanged
+				return FSequencerCommands::Get().SetKeyChanged->GetDescription();
+			}
+		}));
 
-			ToolBarBuilder.AddComboButton(
-				FUIAction(),
-				FOnGetContent::CreateSP(this, &SSequencer::MakeKeyGroupMenu),
-				LOCTEXT("KeyGroup", "Key All"),
-				KeyGroupModeToolTip,
-				KeyGroupModeIcon);
-		}
+		ToolBarBuilder.AddComboButton(
+			FUIAction(),
+			FOnGetContent::CreateSP(this, &SSequencer::MakeKeyGroupMenu),
+			LOCTEXT("KeyGroup", "Key All"),
+			KeyGroupModeToolTip,
+			KeyGroupModeIcon);
 
 		if (IVREditorModule::Get().IsVREditorModeActive() || (SequencerPtr.Pin()->IsLevelEditorSequencer() && ExactCast<ULevelSequence>(SequencerPtr.Pin()->GetFocusedMovieSceneSequence()) == nullptr))
 		{
@@ -1418,7 +1415,6 @@ TSharedRef<SWidget> SSequencer::MakeAutoChangeMenu()
 	MenuBuilder.AddMenuEntry(FSequencerCommands::Get().SetAutoChangeNone);
 
 	return MenuBuilder.MakeWidget();
-
 }
 
 TSharedRef<SWidget> SSequencer::MakeAllowEditsMenu()
@@ -1437,9 +1433,77 @@ TSharedRef<SWidget> SSequencer::MakeKeyGroupMenu()
 {
 	FMenuBuilder MenuBuilder(false, SequencerPtr.Pin()->GetCommandBindings());
 
-	MenuBuilder.AddMenuEntry(FSequencerCommands::Get().SetKeyAll);
-	MenuBuilder.AddMenuEntry(FSequencerCommands::Get().SetKeyGroup);
-	MenuBuilder.AddMenuEntry(FSequencerCommands::Get().SetKeyChanged);
+	if (SequencerPtr.Pin()->IsLevelEditorSequencer())
+	{
+		MenuBuilder.AddMenuEntry(FSequencerCommands::Get().SetKeyAll);
+		MenuBuilder.AddMenuEntry(FSequencerCommands::Get().SetKeyGroup);
+		MenuBuilder.AddMenuEntry(FSequencerCommands::Get().SetKeyChanged);
+	}
+
+	// Interpolation
+	MenuBuilder.BeginSection("SequencerInterpolation", LOCTEXT("KeyInterpolationMenu", "Default Key Interpolation"));
+	{
+		MenuBuilder.AddMenuEntry(
+			LOCTEXT("SetKeyInterpolationAuto", "Cubic (Auto)"),
+			LOCTEXT("SetKeyInterpolationAutoTooltip", "Set key interpolation to auto"),
+			FSlateIcon(FEditorStyle::GetStyleSetName(), "Sequencer.IconKeyAuto"),
+			FUIAction(
+				FExecuteAction::CreateLambda([this] { SequencerPtr.Pin()->SetKeyInterpolation(EMovieSceneKeyInterpolation::Auto); }),
+				FCanExecuteAction(),
+				FIsActionChecked::CreateLambda([this] { return SequencerPtr.Pin()->GetKeyInterpolation() == EMovieSceneKeyInterpolation::Auto; })),
+			NAME_None,
+			EUserInterfaceActionType::ToggleButton
+		);
+
+		MenuBuilder.AddMenuEntry(
+			LOCTEXT("SetKeyInterpolationUser", "Cubic (User)"),
+			LOCTEXT("SetKeyInterpolationUserTooltip", "Set key interpolation to user"),
+			FSlateIcon(FEditorStyle::GetStyleSetName(), "Sequencer.IconKeyUser"),
+			FUIAction(
+				FExecuteAction::CreateLambda([this] { SequencerPtr.Pin()->SetKeyInterpolation(EMovieSceneKeyInterpolation::User); }),
+				FCanExecuteAction(),
+				FIsActionChecked::CreateLambda([this] { return SequencerPtr.Pin()->GetKeyInterpolation() == EMovieSceneKeyInterpolation::User; })),
+			NAME_None,
+			EUserInterfaceActionType::ToggleButton
+		);
+
+		MenuBuilder.AddMenuEntry(
+			LOCTEXT("SetKeyInterpolationBreak", "Cubic (Break)"),
+			LOCTEXT("SetKeyInterpolationBreakTooltip", "Set key interpolation to break"),
+			FSlateIcon(FEditorStyle::GetStyleSetName(), "Sequencer.IconKeyBreak"),
+			FUIAction(
+				FExecuteAction::CreateLambda([this] { SequencerPtr.Pin()->SetKeyInterpolation(EMovieSceneKeyInterpolation::Break); }),
+				FCanExecuteAction(),
+				FIsActionChecked::CreateLambda([this] { return SequencerPtr.Pin()->GetKeyInterpolation() == EMovieSceneKeyInterpolation::Break; })),
+			NAME_None,
+			EUserInterfaceActionType::ToggleButton
+		);
+
+		MenuBuilder.AddMenuEntry(
+			LOCTEXT("SetKeyInterpolationLinear", "Linear"),
+			LOCTEXT("SetKeyInterpolationLinearTooltip", "Set key interpolation to linear"),
+			FSlateIcon(FEditorStyle::GetStyleSetName(), "Sequencer.IconKeyLinear"),
+			FUIAction(
+				FExecuteAction::CreateLambda([this] { SequencerPtr.Pin()->SetKeyInterpolation(EMovieSceneKeyInterpolation::Linear); }),
+				FCanExecuteAction(),
+				FIsActionChecked::CreateLambda([this] { return SequencerPtr.Pin()->GetKeyInterpolation() == EMovieSceneKeyInterpolation::Linear; })),
+			NAME_None,
+			EUserInterfaceActionType::ToggleButton
+		);
+
+		MenuBuilder.AddMenuEntry(
+			LOCTEXT("SetKeyInterpolationConstant", "Constant"),
+			LOCTEXT("SetKeyInterpolationConstantTooltip", "Set key interpolation to constant"),
+			FSlateIcon(FEditorStyle::GetStyleSetName(), "Sequencer.IconKeyConstant"),
+			FUIAction(
+				FExecuteAction::CreateLambda([this] { SequencerPtr.Pin()->SetKeyInterpolation(EMovieSceneKeyInterpolation::Constant); }),
+				FCanExecuteAction(),
+				FIsActionChecked::CreateLambda([this] { return SequencerPtr.Pin()->GetKeyInterpolation() == EMovieSceneKeyInterpolation::Constant; })),
+			NAME_None,
+			EUserInterfaceActionType::ToggleButton
+		);
+	}
+	MenuBuilder.EndSection(); // SequencerInterpolation
 
 	return MenuBuilder.MakeWidget();
 
