Sequencer: Display track node is red if the property track is not bound to valid property

diff --git a/Engine/Source/Editor/Sequencer/Private/DisplayNodes/SequencerTrackNode.cpp b/Engine/Source/Editor/Sequencer/Private/DisplayNodes/SequencerTrackNode.cpp
--- a/Engine/Source/Editor/Sequencer/Private/DisplayNodes/SequencerTrackNode.cpp
+++ b/Engine/Source/Editor/Sequencer/Private/DisplayNodes/SequencerTrackNode.cpp
@@ -26,6 +26,9 @@
 #include "MovieSceneSequence.h"
 #include "MovieScene.h"
 #include "MovieSceneFolder.h"
+#include "Tracks/MovieScenePropertyTrack.h"
+#include "Tracks/MovieScene3DTransformTrack.h"
+#include "MovieSceneCommonHelpers.h"
 
 #define LOCTEXT_NAMESPACE "SequencerTrackNode"
 
@@ -475,6 +478,42 @@ FText FSequencerTrackNode::GetDisplayName() const
 	return AssociatedTrack.IsValid() ? AssociatedTrack->GetDisplayName() : FText::GetEmpty();
 }
 
+FLinearColor FSequencerTrackNode::GetDisplayNameColor() const
+{
+	UMovieSceneTrack* Track = GetTrack();
+
+	// Display track node is red if the property track is not bound to valid property
+	if (UMovieScenePropertyTrack* PropertyTrack = Cast<UMovieScenePropertyTrack>(Track))
+	{
+		// 3D transform tracks don't map to property bindings as below
+		if (Track->IsA<UMovieScene3DTransformTrack>())
+		{
+			return FLinearColor::White;
+		}
+
+		FGuid ObjectBinding;
+		TSharedPtr<FSequencerDisplayNode> ParentSeqNode = GetParent();
+
+		if (ParentSeqNode.IsValid() && (ParentSeqNode->GetType() == ESequencerNode::Object))
+		{
+			ObjectBinding = StaticCastSharedPtr<FSequencerObjectBindingNode>(ParentSeqNode)->GetObjectBinding();
+		}
+
+		for (auto RuntimeObject : GetSequencer().FindBoundObjects(ObjectBinding, GetSequencer().GetFocusedTemplateID()))
+		{
+			FTrackInstancePropertyBindings PropertyBinding(FName(*PropertyTrack->GetName()), PropertyTrack->GetPropertyPath());
+
+			if (PropertyBinding.GetProperty(*RuntimeObject))
+			{
+				return FLinearColor::White;
+			}
+		}
+
+		return FLinearColor::Red;
+	}
+
+	return FLinearColor::White;
+}
 
 float FSequencerTrackNode::GetNodeHeight() const
 {
diff --git a/Engine/Source/Editor/Sequencer/Private/DisplayNodes/SequencerTrackNode.h b/Engine/Source/Editor/Sequencer/Private/DisplayNodes/SequencerTrackNode.h
--- a/Engine/Source/Editor/Sequencer/Private/DisplayNodes/SequencerTrackNode.h
+++ b/Engine/Source/Editor/Sequencer/Private/DisplayNodes/SequencerTrackNode.h
@@ -143,6 +143,7 @@ public:
 	virtual TSharedRef<SWidget> GetCustomOutlinerContent() override;
 	virtual void GetChildKeyAreaNodesRecursively(TArray<TSharedRef<FSequencerSectionKeyAreaNode>>& OutNodes) const override;
 	virtual FText GetDisplayName() const override;
+	virtual FLinearColor GetDisplayNameColor() const override;
 	virtual float GetNodeHeight() const override;
 	virtual FNodePadding GetNodePadding() const override;
 	virtual ESequencerNode::Type GetType() const override;
