Sequencer: Only allow setting/scrubbing time if mouse down was within the time scrubber region. This fixes a bug where if you double click on a shot and hold the mouse down on the 2nd click and drag, it will still set time.

Tested multiple combinations of setting time - click and release on the scrubber, click and release on the timeline, click and drag on the timeline, click and drag without the drag being in the timeline area.

diff --git a/Engine/Source/Editor/Sequencer/Private/SequencerTimeSliderController.cpp b/Engine/Source/Editor/Sequencer/Private/SequencerTimeSliderController.cpp
--- a/Engine/Source/Editor/Sequencer/Private/SequencerTimeSliderController.cpp
+++ b/Engine/Source/Editor/Sequencer/Private/SequencerTimeSliderController.cpp
@@ -49,6 +49,7 @@ FSequencerTimeSliderController::FSequencerTimeSliderController( const FTimeSlide
 	, TimeSliderArgs( InArgs )
 	, DistanceDragged( 0.0f )
 	, MouseDragType( DRAG_NONE )
+	, bMouseDownInRegion(false)
 	, bPanning( false )
 {
 	ScrubFillBrush              = FEditorStyle::GetBrush( TEXT( "Sequencer.Timeline.ScrubFill" ) );
@@ -720,9 +721,19 @@ int32 FSequencerTimeSliderController::DrawSubSequenceRange(const FGeometry& Allo
 
 FReply FSequencerTimeSliderController::OnMouseButtonDown( SWidget& WidgetOwner, const FGeometry& MyGeometry, const FPointerEvent& MouseEvent )
 {
+	MouseDragType = DRAG_NONE;
 	DistanceDragged = 0;
 	MouseDownPosition[0] = MouseDownPosition[1] = MouseEvent.GetScreenSpacePosition();
 	MouseDownGeometry = MyGeometry;
+	bMouseDownInRegion = false;
+
+	FVector2D CursorPos = MouseEvent.GetScreenSpacePosition();
+	FVector2D LocalPos = MouseDownGeometry.AbsoluteToLocal(CursorPos);
+	if (LocalPos.Y >= 0 && LocalPos.Y < MouseDownGeometry.GetLocalSize().Y)
+	{
+		bMouseDownInRegion = true;
+	}
+
 	return FReply::Unhandled();
 }
 
@@ -813,7 +824,7 @@ FReply FSequencerTimeSliderController::OnMouseButtonUp( SWidget& WidgetOwner, co
 				}
 			}
 		}
-		else
+		else if (bMouseDownInRegion)
 		{
 			TimeSliderArgs.OnEndScrubberMovement.ExecuteIfBound();
 
@@ -834,6 +845,7 @@ FReply FSequencerTimeSliderController::OnMouseButtonUp( SWidget& WidgetOwner, co
 
 		MouseDragType = DRAG_NONE;
 		DistanceDragged = 0.f;
+		bMouseDownInRegion = false;
 
 		return FReply::Handled().ReleaseMouseCapture();
 	}
@@ -930,7 +942,7 @@ FReply FSequencerTimeSliderController::OnMouseMove( SWidget& WidgetOwner, const
 				{
 					MouseDragType = DRAG_SETTING_RANGE;
 				}
-				else
+				else if (bMouseDownInRegion)
 				{
 					MouseDragType = DRAG_SCRUBBING_TIME;
 					TimeSliderArgs.OnBeginScrubberMovement.ExecuteIfBound();
diff --git a/Engine/Source/Editor/Sequencer/Public/SequencerTimeSliderController.h b/Engine/Source/Editor/Sequencer/Public/SequencerTimeSliderController.h
--- a/Engine/Source/Editor/Sequencer/Public/SequencerTimeSliderController.h
+++ b/Engine/Source/Editor/Sequencer/Public/SequencerTimeSliderController.h
@@ -321,6 +321,9 @@ private:
 	
 	DragType MouseDragType;
 	
+	/** If mouse down was in time scrubbing region, only allow setting time when mouse is pressed down in the region */
+	bool bMouseDownInRegion;
+
 	/** If we are currently panning the panel */
 	bool bPanning;
 
