Fix clicking tracks in sequencer to rename them by implementing tracking of row selection

diff --git a/Engine/Source/Editor/Sequencer/Private/SAnimationOutlinerTreeNode.cpp b/Engine/Source/Editor/Sequencer/Private/SAnimationOutlinerTreeNode.cpp
--- a/Engine/Source/Editor/Sequencer/Private/SAnimationOutlinerTreeNode.cpp
+++ b/Engine/Source/Editor/Sequencer/Private/SAnimationOutlinerTreeNode.cpp
@@ -20,6 +20,7 @@
 #include "Widgets/Colors/SColorPicker.h"
 #include "SequencerSectionPainter.h"
 #include "Widgets/Text/SInlineEditableTextBlock.h"
+#include "Framework/SlateDelegates.h"
 
 #define LOCTEXT_NAMESPACE "AnimationOutliner"
 
@@ -68,7 +69,8 @@ void SAnimationOutlinerTreeNode::Construct( const FArguments& InArgs, TSharedRef
 		.OnTextCommitted(this, &SAnimationOutlinerTreeNode::HandleNodeLabelTextCommitted)
 		.Text(this, &SAnimationOutlinerTreeNode::GetDisplayName)
 		.ToolTipText(this, &SAnimationOutlinerTreeNode::GetDisplayNameToolTipText)
-		.Clipping(EWidgetClipping::ClipToBounds);
+		.Clipping(EWidgetClipping::ClipToBounds)
+		.IsSelected(FIsSelected::CreateSP(InTableRow, &SSequencerTreeViewRow::IsSelectedExclusively));
 
 
 	Node->OnRenameRequested().AddRaw(this, &SAnimationOutlinerTreeNode::EnterRenameMode);
