When handling undo, recursively test parents of objects when testing if they belong to FSequencer so objects nested in other objects are detected.

diff --git a/Engine/Source/Editor/Sequencer/Private/Sequencer.cpp b/Engine/Source/Editor/Sequencer/Private/Sequencer.cpp
--- a/Engine/Source/Editor/Sequencer/Private/Sequencer.cpp
+++ b/Engine/Source/Editor/Sequencer/Private/Sequencer.cpp
@@ -4457,10 +4457,20 @@ bool FSequencer::MatchesContext(const FTransactionContext& InContext, const TArr
 	// Check if we care about the undo/redo
 	for (const TPair<UObject*, FTransactionObjectEvent>& TransactionObjectPair : TransactionObjects)
 	{
-		if (TransactionObjectPair.Value.HasPendingKillChange() || TransactionObjectPair.Key->GetClass()->IsChildOf(UMovieSceneSignedObject::StaticClass()))
+		if (TransactionObjectPair.Value.HasPendingKillChange())
 		{
 			return true;
 		}
+
+		UObject* Object = TransactionObjectPair.Key;
+		while (Object != nullptr)
+		{
+			if (Object->GetClass()->IsChildOf(UMovieSceneSignedObject::StaticClass()))
+			{
+				return true;
+			}
+			Object = Object->GetOuter();
+		}
 	}
 	return false;
 }
