Focus timeline on time slider when jumping to new position using playback controls

diff --git a/Engine/Source/Editor/Sequencer/Private/Sequencer.cpp b/Engine/Source/Editor/Sequencer/Private/Sequencer.cpp
--- a/Engine/Source/Editor/Sequencer/Private/Sequencer.cpp
+++ b/Engine/Source/Editor/Sequencer/Private/Sequencer.cpp
@@ -2266,12 +2266,15 @@ FQualifiedFrameTime FSequencer::GetGlobalTime() const
 	return FQualifiedFrameTime(RootTime, PlayPosition.GetOutputRate());
 }
 
-void FSequencer::SetLocalTime( FFrameTime NewTime, ESnapTimeMode SnapTimeMode )
+void FSequencer::SetLocalTime( FFrameTime NewTime, ESnapTimeMode SnapTimeMode, bool bForceScroll )
 {
 	FFrameRate LocalResolution = GetFocusedTickResolution();
 
 	// Ensure the time is in the current view
-	ScrollIntoView(NewTime / LocalResolution);
+	if (IsAutoScrollEnabled() || bForceScroll)
+	{
+		ScrollIntoView(NewTime / LocalResolution);
+	}
 
 	// Perform snapping
 	if ((SnapTimeMode & ESnapTimeMode::STM_Interval) && Settings->GetIsSnapEnabled())
@@ -2375,27 +2378,24 @@ void FSequencer::EvaluateInternal(FMovieSceneEvaluationRange InRange, bool bHasJ
 
 void FSequencer::ScrollIntoView(float InLocalTime)
 {
-	if (IsAutoScrollEnabled())
-	{
-		float RangeOffset = CalculateAutoscrollEncroachment(InLocalTime).Get(0.f);
+	float RangeOffset = CalculateAutoscrollEncroachment(InLocalTime).Get(0.f);
 		
-		// When not scrubbing, we auto scroll the view range immediately
-		if (RangeOffset != 0.f)
-		{
-			TRange<double> WorkingRange = GetClampRange();
-
-			// Adjust the offset so that the target range will be within the working range.
-			if (TargetViewRange.GetLowerBoundValue() + RangeOffset < WorkingRange.GetLowerBoundValue())
-			{
-				RangeOffset = WorkingRange.GetLowerBoundValue() - TargetViewRange.GetLowerBoundValue();
-			}
-			else if (TargetViewRange.GetUpperBoundValue() + RangeOffset > WorkingRange.GetUpperBoundValue())
-			{
-				RangeOffset = WorkingRange.GetUpperBoundValue() - TargetViewRange.GetUpperBoundValue();
-			}
+	// When not scrubbing, we auto scroll the view range immediately
+	if (RangeOffset != 0.f)
+	{
+		TRange<double> WorkingRange = GetClampRange();
 
-			SetViewRange(TRange<double>(TargetViewRange.GetLowerBoundValue() + RangeOffset, TargetViewRange.GetUpperBoundValue() + RangeOffset), EViewRangeInterpolation::Immediate);
+		// Adjust the offset so that the target range will be within the working range.
+		if (TargetViewRange.GetLowerBoundValue() + RangeOffset < WorkingRange.GetLowerBoundValue())
+		{
+			RangeOffset = WorkingRange.GetLowerBoundValue() - TargetViewRange.GetLowerBoundValue();
+		}
+		else if (TargetViewRange.GetUpperBoundValue() + RangeOffset > WorkingRange.GetUpperBoundValue())
+		{
+			RangeOffset = WorkingRange.GetUpperBoundValue() - TargetViewRange.GetUpperBoundValue();
 		}
+
+		SetViewRange(TRange<double>(TargetViewRange.GetLowerBoundValue() + RangeOffset, TargetViewRange.GetUpperBoundValue() + RangeOffset), EViewRangeInterpolation::Immediate);
 	}
 }
 
@@ -3549,7 +3549,7 @@ FReply FSequencer::OnStepForward()
 	FQualifiedFrameTime CurrentTime = GetLocalTime();
 
 	FFrameTime NewPosition = FFrameRate::TransformTime(CurrentTime.ConvertTo(DisplayRate).FloorToFrame() + 1, DisplayRate, CurrentTime.Rate);
-	SetLocalTime(NewPosition, ESnapTimeMode::STM_Interval);
+	SetLocalTime(NewPosition, ESnapTimeMode::STM_Interval, true);
 	return FReply::Handled();
 }
 
@@ -3563,7 +3563,7 @@ FReply FSequencer::OnStepBackward()
 
 	FFrameTime NewPosition = FFrameRate::TransformTime(CurrentTime.ConvertTo(DisplayRate).FloorToFrame() - 1, DisplayRate, CurrentTime.Rate);
 
-	SetLocalTime(NewPosition, ESnapTimeMode::STM_Interval);
+	SetLocalTime(NewPosition, ESnapTimeMode::STM_Interval, true);
 	return FReply::Handled();
 }
 
@@ -3571,7 +3571,7 @@ FReply FSequencer::OnStepBackward()
 FReply FSequencer::OnJumpToStart()
 {
 	SetPlaybackStatus(EMovieScenePlayerStatus::Stepping);
-	SetLocalTime(MovieScene::DiscreteInclusiveLower(GetPlaybackRange()));
+	SetLocalTime(MovieScene::DiscreteInclusiveLower(GetPlaybackRange()), ESnapTimeMode::STM_None, true);
 	return FReply::Handled();
 }
 
@@ -3589,7 +3589,7 @@ FReply FSequencer::OnJumpToEnd()
 	FFrameTime OneFrame = bInsetDisplayFrame ? FFrameRate::TransformTime(FFrameTime(1), DisplayRate, LocalResolution) : FFrameTime(1);
 	FFrameTime NewTime = MovieScene::DiscreteExclusiveUpper(GetPlaybackRange()) - OneFrame;
 
-	SetLocalTime(NewTime);
+	SetLocalTime(NewTime, ESnapTimeMode::STM_None, true);
 	return FReply::Handled();
 }
 
@@ -3663,6 +3663,11 @@ FReply FSequencer::JumpToPreviousKey()
 		if (NewTime.IsSet())
 		{
 			SetPlaybackStatus(EMovieScenePlayerStatus::Stepping);
+
+			// Ensure the time is in the current view
+			FFrameRate LocalResolution = GetFocusedTickResolution();
+			ScrollIntoView(NewTime.GetValue() / LocalResolution);
+
 			SetLocalTimeDirectly(NewTime.GetValue());
 		}
 	}
@@ -3687,6 +3692,11 @@ FReply FSequencer::JumpToNextKey()
 		if (NewTime.IsSet())
 		{
 			SetPlaybackStatus(EMovieScenePlayerStatus::Stepping);
+
+			// Ensure the time is in the current view
+			FFrameRate LocalResolution = GetFocusedTickResolution();
+			ScrollIntoView(NewTime.GetValue() / LocalResolution);
+
 			SetLocalTimeDirectly(NewTime.GetValue());
 		}
 	}
@@ -3776,7 +3786,10 @@ void FSequencer::SetLocalTimeLooped(FFrameTime NewLocalTime)
 
 	// Ensure the time is in the current view - must occur before the time cursor changes
 	UMovieScene* MovieScene = GetFocusedMovieSceneSequence()->GetMovieScene();
-	ScrollIntoView( (NewGlobalTime * RootToLocalTransform) / RootTickResolution );
+	if (IsAutoScrollEnabled())
+	{
+		ScrollIntoView((NewGlobalTime * RootToLocalTransform) / RootTickResolution);
+	}
 
 	FFrameTime NewPlayPosition = ConvertFrameTime(NewGlobalTime, RootTickResolution, PlayPosition.GetInputRate());
 
@@ -7833,7 +7846,7 @@ void FSequencer::StepToNextShot()
 	PopToSequenceInstance(ActiveTemplateIDs[ActiveTemplateIDs.Num()-2]);
 	FocusSequenceInstance(*NextShot);
 
-	SetLocalTime(FFrameTime(0));
+	SetLocalTime(FFrameTime(0), ESnapTimeMode::STM_None, true);
 }
 
 
@@ -7865,7 +7878,7 @@ void FSequencer::StepToPreviousShot()
 	PopToSequenceInstance(ActiveTemplateIDs[ActiveTemplateIDs.Num()-2]);
 	FocusSequenceInstance(*PreviousShot);
 
-	SetLocalTime(0);
+	SetLocalTime(0, ESnapTimeMode::STM_None, true);
 }
 
 
diff --git a/Engine/Source/Editor/Sequencer/Private/Sequencer.h b/Engine/Source/Editor/Sequencer/Private/Sequencer.h
--- a/Engine/Source/Editor/Sequencer/Private/Sequencer.h
+++ b/Engine/Source/Editor/Sequencer/Private/Sequencer.h
@@ -672,7 +672,7 @@ public:
 	virtual bool GetAutoSetTrackDefaults() const override;
 	virtual FQualifiedFrameTime GetLocalTime() const override;
 	virtual FQualifiedFrameTime GetGlobalTime() const override;
-	virtual void SetLocalTime(FFrameTime Time, ESnapTimeMode SnapTimeMode = ESnapTimeMode::STM_None) override;
+	virtual void SetLocalTime(FFrameTime Time, ESnapTimeMode SnapTimeMode = ESnapTimeMode::STM_None, bool bForceScroll = false) override;
 	virtual void SetLocalTimeDirectly(FFrameTime NewTime) override;
 	virtual void SetGlobalTime(FFrameTime Time) override;
 	virtual void ForceEvaluate() override;
diff --git a/Engine/Source/Editor/Sequencer/Public/ISequencer.h b/Engine/Source/Editor/Sequencer/Public/ISequencer.h
--- a/Engine/Source/Editor/Sequencer/Public/ISequencer.h
+++ b/Engine/Source/Editor/Sequencer/Public/ISequencer.h
@@ -322,8 +322,9 @@ public:
 	 *
 	 * @param Time The local time to set.
 	 * @param SnapTimeMode The type of time snapping allowed.
+	 * @param ForceScroll Whether to ignore autoscroll setting and always scroll view to new location
 	 */
-	virtual void SetLocalTime(FFrameTime Time, ESnapTimeMode SnapTimeMode = ESnapTimeMode::STM_None) = 0;
+	virtual void SetLocalTime(FFrameTime Time, ESnapTimeMode SnapTimeMode = ESnapTimeMode::STM_None, bool bForceScroll = false) = 0;
 
 	/** Set the current local time directly, with no other snapping, scrolling or manipulation */
 	virtual void SetLocalTimeDirectly(FFrameTime NewTime) = 0;
