Add support for converting multiple object bindings to spawnables in the sequencer timeline

diff --git a/Engine/Source/Editor/Sequencer/Private/Sequencer.cpp b/Engine/Source/Editor/Sequencer/Private/Sequencer.cpp
--- a/Engine/Source/Editor/Sequencer/Private/Sequencer.cpp
+++ b/Engine/Source/Editor/Sequencer/Private/Sequencer.cpp
@@ -4534,21 +4534,20 @@ void FSequencer::OnNewActorsDropped(const TArray<UObject*>& DroppedObjects, cons
 
 			if (bAddSpawnable)
 			{
-				FMovieSceneSpawnable* Spawnable = ConvertToSpawnableInternal(PossessableGuid);
-
-				ForceEvaluate();
-
-				for (TWeakObjectPtr<> WeakObject : FindBoundObjects(Spawnable->GetGuid(), ActiveTemplateIDs.Top()))
+				TArray< FMovieSceneSpawnable*> Spawnables = ConvertToSpawnableInternal(PossessableGuid);
+				if (Spawnables.Num() > 0)
 				{
-					AActor* SpawnedActor = Cast<AActor>(WeakObject.Get());
-					if (SpawnedActor)
+					for (TWeakObjectPtr<> WeakObject : FindBoundObjects(Spawnables[0]->GetGuid(), ActiveTemplateIDs.Top()))
 					{
-						SpawnedActors.Add(SpawnedActor);
-						NewActor = SpawnedActor;
+						AActor* SpawnedActor = Cast<AActor>(WeakObject.Get());
+						if (SpawnedActor)
+						{
+							SpawnedActors.Add(SpawnedActor);
+							NewActor = SpawnedActor;
+						}
 					}
+					NewGuid = Spawnables[0]->GetGuid();
 				}
-
-				NewGuid = Spawnable->GetGuid();
 			}
 
 			if (bCreateAndAttachCamera)
@@ -4573,9 +4572,7 @@ void FSequencer::OnNewActorsDropped(const TArray<UObject*>& DroppedObjects, cons
 
 				if (bAddSpawnable)
 				{
-					FMovieSceneSpawnable* Spawnable = ConvertToSpawnableInternal(NewCameraGuid);
-
-					ForceEvaluate();
+					FMovieSceneSpawnable* Spawnable = ConvertToSpawnableInternal(NewCameraGuid)[0];
 
 					for (TWeakObjectPtr<> WeakObject : FindBoundObjects(Spawnable->GetGuid(), ActiveTemplateIDs.Top()))
 					{
@@ -6298,7 +6295,7 @@ void FSequencer::CopySelectedObjects(TArray<TSharedPtr<FSequencerObjectBindingNo
 			}
 		}
 
-		const FMovieSceneBinding* Binding = MovieScene->GetBindings().FindByPredicate([=](const FMovieSceneBinding& InBinding){ return InBinding.GetObjectGuid() == ObjectBinding; });
+		const FMovieSceneBinding* Binding = MovieScene->FindBinding(ObjectBinding);
 		if (Binding)
 		{
 			CopyableBinding->Binding = *Binding;
@@ -7186,12 +7183,10 @@ void FSequencer::ConvertSelectedNodesToSpawnables()
 		FMovieScenePossessable* Possessable = MovieScene->FindPossessable(ObjectBindingNode->GetObjectBinding());
 		if (Possessable)
 		{
-			FMovieSceneSpawnable* Spawnable = ConvertToSpawnableInternal(Possessable->GetGuid());
+			TArray<FMovieSceneSpawnable*> Spawnables = ConvertToSpawnableInternal(Possessable->GetGuid());
 
-			if (Spawnable)
+			for (FMovieSceneSpawnable* Spawnable : Spawnables)
 			{
-				ForceEvaluate();
-
 				for (TWeakObjectPtr<> WeakObject : FindBoundObjects(Spawnable->GetGuid(), ActiveTemplateIDs.Top()))
 				{
 					if (AActor* SpawnedActor = Cast<AActor>(WeakObject.Get()))
@@ -7229,72 +7224,87 @@ void FSequencer::ConvertSelectedNodesToSpawnables()
 	NotifyMovieSceneDataChanged( EMovieSceneDataChangeType::MovieSceneStructureItemsChanged );
 }
 
-FMovieSceneSpawnable* FSequencer::ConvertToSpawnableInternal(FGuid PossessableGuid)
+TArray<FMovieSceneSpawnable*> FSequencer::ConvertToSpawnableInternal(FGuid PossessableGuid)
 {
 	UMovieSceneSequence* Sequence = GetFocusedMovieSceneSequence();
 	UMovieScene* MovieScene = Sequence->GetMovieScene();
 
 	if (MovieScene->IsReadOnly())
 	{
-		return nullptr;
+		return TArray<FMovieSceneSpawnable*>();
 	}
 
-	//@todo: this code doesn't work where multiple objects are bound
 	TArrayView<TWeakObjectPtr<>> FoundObjects = FindBoundObjects(PossessableGuid, ActiveTemplateIDs.Top());
-	if (FoundObjects.Num() != 1)
-	{
-		return nullptr;
-	}
-
-	UObject* FoundObject = FoundObjects[0].Get();
-	if (!FoundObject)
-	{
-		return nullptr;
-	}
 
 	Sequence->Modify();
 	MovieScene->Modify();
 
-	FMovieSceneSpawnable* Spawnable = MovieScene->FindSpawnable(AddSpawnable(*FoundObject));
-	if (Spawnable)
+	// Locate the folder containing the original possessable
+	UMovieSceneFolder* ParentFolder;
+	for (UMovieSceneFolder* Folder : MovieScene->GetRootFolders())
 	{
-		FGuid SpawnableGuid = Spawnable->GetGuid();
+		ParentFolder = Folder->FindFolderContaining(PossessableGuid);
+		if (ParentFolder != nullptr)
+		{
+			break;
+		}
+	}
 
-		// Remap all the spawnable's tracks and child bindings onto the new possessable
-		MovieScene->MoveBindingContents(PossessableGuid, SpawnableGuid);
+	TArray<FMovieSceneSpawnable*> CreatedSpawnables;
 
-		FMovieSceneBinding* PossessableBinding = (FMovieSceneBinding*)MovieScene->GetBindings().FindByPredicate([&](FMovieSceneBinding& Binding) { return Binding.GetObjectGuid() == PossessableGuid; });
+	if (Sequence->AllowsSpawnableObjects())
+	{
+		
+		FMovieSceneBinding* PossessableBinding = MovieScene->FindBinding(PossessableGuid);
 		check(PossessableBinding);
 
-		for (UMovieSceneFolder* Folder : MovieScene->GetRootFolders())
+		TArray<UMovieSceneTrack* > Tracks = PossessableBinding->GetTracks();
+
+		int32 SortingOrder = PossessableBinding->GetSortingOrder();
+
+		for (TWeakObjectPtr<> FoundObjectPtr : FoundObjects)
 		{
-			if (ReplaceFolderBindingGUID(Folder, PossessableGuid, SpawnableGuid))
+			if (FoundObjectPtr.IsValid())
 			{
-				break;
+				UObject* FoundObject = FoundObjectPtr.Get();
+				FMovieSceneSpawnable* Spawnable = MovieScene->FindSpawnable(AddSpawnable(*FoundObject));
+				if (Spawnable)
+				{
+					CreatedSpawnables.Add(Spawnable);
+
+					FGuid SpawnableGuid = Spawnable->GetGuid();
+
+					FMovieSceneBinding* SpawnableBinding = MovieScene->FindBinding(SpawnableGuid);
+					check(SpawnableBinding);
+
+					SpawnableBinding->SetSortingOrder(SortingOrder);
+
+					AActor* PossessedActor = Cast<AActor>(FoundObject);
+					Spawnable->AddChildPossessable(CreateBinding(*FoundObject, PossessedActor != nullptr ? PossessedActor->GetActorLabel() : FoundObject->GetName()));
+					
+					for (UMovieSceneTrack* Track : Tracks)
+					{
+						UMovieSceneTrack* DuplicatedTrack = Cast<UMovieSceneTrack>(StaticDuplicateObject(Track, MovieScene));
+						SpawnableBinding->AddTrack(*DuplicatedTrack);
+					}
+
+					TOptional<FTransformData> TransformData;
+					SpawnRegister->HandleConvertPossessableToSpawnable(FoundObject, *this, TransformData);
+					SpawnRegister->SetupDefaultsForSpawnable(nullptr, Spawnable->GetGuid(), TransformData, AsShared(), Settings);
+				}
 			}
 		}
 
-		int32 SortingOrder = PossessableBinding->GetSortingOrder();
-
 		if (MovieScene->RemovePossessable(PossessableGuid))
 		{
 			Sequence->UnbindPossessableObjects(PossessableGuid);
-
-			FMovieSceneBinding* SpawnableBinding = (FMovieSceneBinding*)MovieScene->GetBindings().FindByPredicate([&](FMovieSceneBinding& Binding) { return Binding.GetObjectGuid() == SpawnableGuid; });
-			check(SpawnableBinding);
-			
-			SpawnableBinding->SetSortingOrder(SortingOrder);
-
 		}
 
-		TOptional<FTransformData> TransformData;
-		SpawnRegister->HandleConvertPossessableToSpawnable(FoundObject, *this, TransformData);
-		SpawnRegister->SetupDefaultsForSpawnable(nullptr, Spawnable->GetGuid(), TransformData, AsShared(), Settings);
-
 		ForceEvaluate();
 	}
 
-	return Spawnable;
+
+	return CreatedSpawnables;
 }
 
 void FSequencer::ConvertToPossessable(TSharedRef<FSequencerObjectBindingNode> NodeToBeConverted)
@@ -7460,7 +7470,7 @@ FMovieScenePossessable* FSequencer::ConvertToPossessableInternal(FGuid Spawnable
 		// Remap all the spawnable's tracks and child bindings onto the new possessable
 		MovieScene->MoveBindingContents(OldSpawnableGuid, NewPossessableGuid);
 
-		FMovieSceneBinding* SpawnableBinding = (FMovieSceneBinding*)MovieScene->GetBindings().FindByPredicate([&](FMovieSceneBinding& Binding) { return Binding.GetObjectGuid() == OldSpawnableGuid; });
+		FMovieSceneBinding* SpawnableBinding = MovieScene->FindBinding(OldSpawnableGuid);
 		check(SpawnableBinding);
 
 		for (UMovieSceneFolder* Folder : MovieScene->GetRootFolders())
@@ -7478,7 +7488,7 @@ FMovieScenePossessable* FSequencer::ConvertToPossessableInternal(FGuid Spawnable
 		{
 			SpawnRegister->DestroySpawnedObject(OldSpawnableGuid, ActiveTemplateIDs.Top(), *this);
 
-			FMovieSceneBinding* PossessableBinding = (FMovieSceneBinding*)MovieScene->GetBindings().FindByPredicate([&](FMovieSceneBinding& Binding) { return Binding.GetObjectGuid() == NewPossessableGuid; });
+			FMovieSceneBinding* PossessableBinding = MovieScene->FindBinding(NewPossessableGuid);
 			check(PossessableBinding);
 			
 			PossessableBinding->SetSortingOrder(SortingOrder);
diff --git a/Engine/Source/Editor/Sequencer/Private/Sequencer.h b/Engine/Source/Editor/Sequencer/Private/Sequencer.h
--- a/Engine/Source/Editor/Sequencer/Private/Sequencer.h
+++ b/Engine/Source/Editor/Sequencer/Private/Sequencer.h
@@ -927,7 +927,7 @@ protected:
 	void UpdatePreviewLevelViewportClientFromCameraCut(FLevelEditorViewportClient& InViewportClient, UObject* InCameraObject, bool bJumpCut) const;
 
 	/** Internal conversion function that doesn't perform expensive reset/update tasks */
-	FMovieSceneSpawnable* ConvertToSpawnableInternal(FGuid PossessableGuid);
+	TArray<FMovieSceneSpawnable*> ConvertToSpawnableInternal(FGuid PossessableGuid);
 
 	/** Internal conversion function that doesn't perform expensive reset/update tasks */
 	FMovieScenePossessable* ConvertToPossessableInternal(FGuid SpawnableGuid);
diff --git a/Engine/Source/Runtime/MovieScene/Private/MovieSceneFolder.cpp b/Engine/Source/Runtime/MovieScene/Private/MovieSceneFolder.cpp
--- a/Engine/Source/Runtime/MovieScene/Private/MovieSceneFolder.cpp
+++ b/Engine/Source/Runtime/MovieScene/Private/MovieSceneFolder.cpp
@@ -76,6 +76,27 @@ void UMovieSceneFolder::RemoveChildObjectBinding( const FGuid& InObjectBinding )
 	ChildObjectBindings.Remove( InObjectBinding );
 }
 
+UMovieSceneFolder* UMovieSceneFolder::FindFolderContaining(const FGuid& InObjectBinding)
+{
+	for (FGuid ChildGuid : GetChildObjectBindings())
+	{
+		if (ChildGuid == InObjectBinding)
+		{
+			return this;
+		}
+	}
+
+	for (UMovieSceneFolder* ChildFolder : GetChildFolders())
+	{
+		UMovieSceneFolder* Folder = ChildFolder->FindFolderContaining(InObjectBinding);
+		if (Folder != nullptr)
+		{
+			return Folder;
+		}
+	}
+
+	return nullptr;
+}
 
 void UMovieSceneFolder::Serialize( FArchive& Archive )
 {
diff --git a/Engine/Source/Runtime/MovieScene/Public/MovieSceneFolder.h b/Engine/Source/Runtime/MovieScene/Public/MovieSceneFolder.h
--- a/Engine/Source/Runtime/MovieScene/Public/MovieSceneFolder.h
+++ b/Engine/Source/Runtime/MovieScene/Public/MovieSceneFolder.h
@@ -50,6 +50,9 @@ class MOVIESCENE_API UMovieSceneFolder : public UObject
 	/** Removes a guid for an object binding from this folder. */
 	void RemoveChildObjectBinding( const FGuid& InObjectBinding );
 
+	/** Searches for a guid in this folder and it's child folders, if found returns the folder containing the guid. */
+	UMovieSceneFolder* FindFolderContaining(const FGuid& InObjectBinding);
+
 
 	virtual void Serialize( FArchive& Archive );
 
