Sequencer: Dragging marks

Fixed missing transaction/modify for add/remove marked frame.
Right click menu for hovered marked frame now uses HitTestMark instead of testing the single frame that the mouse is over.

diff --git a/Engine/Source/Editor/Sequencer/Private/SSequencer.cpp b/Engine/Source/Editor/Sequencer/Private/SSequencer.cpp
--- a/Engine/Source/Editor/Sequencer/Private/SSequencer.cpp
+++ b/Engine/Source/Editor/Sequencer/Private/SSequencer.cpp
@@ -105,6 +105,8 @@ void SSequencer::Construct(const FArguments& InArgs, TSharedRef<FSequencer> InSe
 	OnPlaybackRangeEndDrag = InArgs._OnPlaybackRangeEndDrag;
 	OnSelectionRangeBeginDrag = InArgs._OnSelectionRangeBeginDrag;
 	OnSelectionRangeEndDrag = InArgs._OnSelectionRangeEndDrag;
+	OnMarkBeginDrag = InArgs._OnMarkBeginDrag;
+	OnMarkEndDrag = InArgs._OnMarkEndDrag;
 
 	OnReceivedFocus = InArgs._OnReceivedFocus;
 
@@ -154,6 +156,8 @@ void SSequencer::Construct(const FArguments& InArgs, TSharedRef<FSequencer> InSe
 		TimeSliderArgs.OnSelectionRangeChanged = InArgs._OnSelectionRangeChanged;
 		TimeSliderArgs.OnSelectionRangeBeginDrag = OnSelectionRangeBeginDrag;
 		TimeSliderArgs.OnSelectionRangeEndDrag = OnSelectionRangeEndDrag;
+		TimeSliderArgs.OnMarkBeginDrag = OnMarkBeginDrag;
+		TimeSliderArgs.OnMarkEndDrag = OnMarkEndDrag;
 		TimeSliderArgs.OnViewRangeChanged = InArgs._OnViewRangeChanged;
 		TimeSliderArgs.OnClampRangeChanged = InArgs._OnClampRangeChanged;
 		TimeSliderArgs.OnGetNearestKey = InArgs._OnGetNearestKey;
@@ -167,6 +171,7 @@ void SSequencer::Construct(const FArguments& InArgs, TSharedRef<FSequencer> InSe
 		TimeSliderArgs.SubSequenceRange = InArgs._SubSequenceRange;
 		TimeSliderArgs.VerticalFrames = InArgs._VerticalFrames;
 		TimeSliderArgs.MarkedFrames = InArgs._MarkedFrames;
+		TimeSliderArgs.OnSetMarkedFrame = InArgs._OnSetMarkedFrame;
 		TimeSliderArgs.OnMarkedFrameChanged = InArgs._OnMarkedFrameChanged;
 		TimeSliderArgs.OnClearAllMarkedFrames = InArgs._OnClearAllMarkedFrames;
 
diff --git a/Engine/Source/Editor/Sequencer/Private/SSequencer.h b/Engine/Source/Editor/Sequencer/Private/SSequencer.h
--- a/Engine/Source/Editor/Sequencer/Private/SSequencer.h
+++ b/Engine/Source/Editor/Sequencer/Private/SSequencer.h
@@ -138,6 +138,12 @@ public:
 		/** Called when the user has finished dragging the selection range */
 		SLATE_EVENT( FSimpleDelegate, OnSelectionRangeEndDrag )
 
+		/** Called when the user has begun dragging a mark */
+		SLATE_EVENT(FSimpleDelegate, OnMarkBeginDrag)
+
+		/** Called when the user has finished dragging a mark */
+		SLATE_EVENT(FSimpleDelegate, OnMarkEndDrag)
+
 		/** Whether the playback range is locked */
 		SLATE_ATTRIBUTE( bool, IsPlaybackRangeLocked )
 
@@ -150,8 +156,11 @@ public:
 		/** Called when the user changes the view range */
 		SLATE_EVENT( FOnViewRangeChanged, OnViewRangeChanged )
 
+		/** Called when the user sets a marked frame */
+		SLATE_EVENT(FOnSetMarkedFrame, OnSetMarkedFrame)
+
 		/** Called when the user changes on the set of marked frames */
-		SLATE_EVENT( FOnMarkedFrameChanged, OnMarkedFrameChanged )
+		SLATE_EVENT(FOnMarkedFrameChanged, OnMarkedFrameChanged)
 
 		/** Called when all marked frames should be cleared */
 		SLATE_EVENT( FSimpleDelegate, OnClearAllMarkedFrames)
@@ -545,6 +554,12 @@ private:
 	/** Called when the user has finished dragging the playback range */
 	FSimpleDelegate OnPlaybackRangeEndDrag;
 
+	/** Called when the user has begun dragging a mark */
+	FSimpleDelegate OnMarkBeginDrag;
+
+	/** Called when the user has finished dragging a mark */
+	FSimpleDelegate OnMarkEndDrag;
+
 	/** Called when any widget contained within sequencer has received focus */
 	FSimpleDelegate OnReceivedFocus;
 
diff --git a/Engine/Source/Editor/Sequencer/Private/Sequencer.cpp b/Engine/Source/Editor/Sequencer/Private/Sequencer.cpp
--- a/Engine/Source/Editor/Sequencer/Private/Sequencer.cpp
+++ b/Engine/Source/Editor/Sequencer/Private/Sequencer.cpp
@@ -361,7 +361,8 @@ void FSequencer::InitSequencer(const FSequencerInitParams& InitParams, const TSh
 		.SelectionRange( this, &FSequencer::GetSelectionRange )
 		.VerticalFrames(this, &FSequencer::GetVerticalFrames)
 		.MarkedFrames(this, &FSequencer::GetMarkedFrames)
-		.OnMarkedFrameChanged(this, &FSequencer::SetMarkedFrame )
+		.OnSetMarkedFrame(this, &FSequencer::SetMarkedFrame)
+		.OnMarkedFrameChanged(this, &FSequencer::OnMarkedFrameChanged)
 		.OnClearAllMarkedFrames(this, &FSequencer::ClearAllMarkedFrames )
 		.SubSequenceRange( this, &FSequencer::GetSubSequenceRange )
 		.OnPlaybackRangeChanged( this, &FSequencer::SetPlaybackRange )
@@ -370,6 +371,8 @@ void FSequencer::InitSequencer(const FSequencerInitParams& InitParams, const TSh
 		.OnSelectionRangeChanged( this, &FSequencer::SetSelectionRange )
 		.OnSelectionRangeBeginDrag( this, &FSequencer::OnSelectionRangeBeginDrag )
 		.OnSelectionRangeEndDrag( this, &FSequencer::OnSelectionRangeEndDrag )
+		.OnMarkBeginDrag(this, &FSequencer::OnMarkBeginDrag)
+		.OnMarkEndDrag(this, &FSequencer::OnMarkEndDrag)
 		.IsPlaybackRangeLocked( this, &FSequencer::IsPlaybackRangeLocked )
 		.OnTogglePlaybackRangeLocked( this, &FSequencer::TogglePlaybackRangeLocked )
 		.ScrubPosition( this, &FSequencer::GetLocalFrameTime )
@@ -4095,6 +4098,18 @@ void FSequencer::OnSelectionRangeEndDrag()
 }
 
 
+void FSequencer::OnMarkBeginDrag()
+{
+	GEditor->BeginTransaction(LOCTEXT("SetMark_Transaction", "Set Mark"));
+}
+
+
+void FSequencer::OnMarkEndDrag()
+{
+	GEditor->EndTransaction();
+}
+
+
 void FSequencer::StartAutoscroll(float UnitsPerS)
 {
 	AutoscrollOffset = UnitsPerS;
@@ -8186,17 +8201,37 @@ void FSequencer::ToggleMarkAtPlayPosition()
 			int32 MarkedFrameIndex = FocusedMovieScene->FindMarkedFrameByFrameNumber(TickFrameNumber);
 			if (MarkedFrameIndex != INDEX_NONE)
 			{
+				FScopedTransaction RemoveMarkedFrameTransaction(LOCTEXT("RemoveMarkedFrames_Transaction", "Remove Marked Frame"));
+
+				FocusedMovieScene->Modify();
 				FocusedMovieScene->RemoveMarkedFrame(MarkedFrameIndex);
 			}
 			else
 			{
+				FScopedTransaction AddMarkedFrameTransaction(LOCTEXT("AddMarkedFrame_Transaction", "Add Marked Frame"));
+
+				FocusedMovieScene->Modify();
 				FocusedMovieScene->AddMarkedFrame(FMovieSceneMarkedFrame(TickFrameNumber));
 			}
 		}
 	}
 }
 
-void FSequencer::SetMarkedFrame(FFrameNumber FrameNumber, bool bSetMark)
+void FSequencer::SetMarkedFrame(int32 InMarkIndex, FFrameNumber InFrameNumber)
+{
+	UMovieSceneSequence* FocusedMovieSequence = GetFocusedMovieSceneSequence();
+	if (FocusedMovieSequence != nullptr)
+	{
+		UMovieScene* FocusedMovieScene = FocusedMovieSequence->GetMovieScene();
+		if (FocusedMovieScene != nullptr)
+		{
+			FocusedMovieScene->Modify();
+			FocusedMovieScene->SetMarkedFrame(InMarkIndex, InFrameNumber);
+		}
+	}
+}
+
+void FSequencer::OnMarkedFrameChanged(FFrameNumber FrameNumber, bool bSetMark)
 {
 	UMovieSceneSequence* FocusedMovieSequence = GetFocusedMovieSceneSequence();
 	if (FocusedMovieSequence != nullptr)
@@ -8206,6 +8241,9 @@ void FSequencer::SetMarkedFrame(FFrameNumber FrameNumber, bool bSetMark)
 		{
 			if (bSetMark)
 			{
+				FScopedTransaction AddMarkedFrameTransaction(LOCTEXT("AddMarkedFrame_Transaction", "Add Marked Frame"));
+
+				FocusedMovieScene->Modify();
 				FocusedMovieScene->AddMarkedFrame(FMovieSceneMarkedFrame(FrameNumber));
 			}
 			else
@@ -8213,6 +8251,9 @@ void FSequencer::SetMarkedFrame(FFrameNumber FrameNumber, bool bSetMark)
 				int32 MarkedFrameIndex = FocusedMovieScene->FindMarkedFrameByFrameNumber(FrameNumber);
 				if (MarkedFrameIndex != INDEX_NONE)
 				{
+					FScopedTransaction RemoveMarkedFrameTransaction(LOCTEXT("RemoveMarkedFrame_Transaction", "Remove Marked Frame"));
+
+					FocusedMovieScene->Modify();
 					FocusedMovieScene->RemoveMarkedFrame(MarkedFrameIndex);
 				}
 			}
@@ -8228,6 +8269,9 @@ void FSequencer::ClearAllMarkedFrames()
 		UMovieScene* FocusedMovieScene = FocusedMovieSequence->GetMovieScene();
 		if (FocusedMovieScene != nullptr)
 		{
+			FScopedTransaction ClearAllMarkedFramesTransaction(LOCTEXT("ClearAllMarkedFrames_Transaction", "Clear All Marked Frames"));
+
+			FocusedMovieScene->Modify();
 			FocusedMovieScene->ClearMarkedFrames();
 		}
 	}
diff --git a/Engine/Source/Editor/Sequencer/Private/Sequencer.h b/Engine/Source/Editor/Sequencer/Private/Sequencer.h
--- a/Engine/Source/Editor/Sequencer/Private/Sequencer.h
+++ b/Engine/Source/Editor/Sequencer/Private/Sequencer.h
@@ -280,10 +280,16 @@ protected:
 	void StepToPreviousMark();
 
 	/**
-	 * @param	FrameNumber The FrameNumber in Ticks 
+	 * @param InMarkIndex The marked frame index to set
+	 * @param InFrameNumber The FrameNumber in Ticks
+	 */
+	void SetMarkedFrame(int32 InMarkIndex, FFrameNumber InFrameNumber);
+
+	/**
+	 * @param	FrameNumber The FrameNumber in Ticks
 	 * @param	bSetMark  true to set the mark, false to clear the mark
 	 */
-	void SetMarkedFrame(FFrameNumber FrameNumber, bool bSetMark);
+	void OnMarkedFrameChanged(FFrameNumber FrameNumber, bool bSetMark);
 
 	void ClearAllMarkedFrames();
 
@@ -822,6 +828,12 @@ protected:
 	/** Called when the user has finished dragging the selection range */
 	void OnSelectionRangeEndDrag();
 
+	/** Called when the user has begun dragging a mark */
+	void OnMarkBeginDrag();
+
+	/** Called when the user has finished dragging a mark */
+	void OnMarkEndDrag();
+
 	/** Get the unqualified local time */
 	FFrameTime GetLocalFrameTime() const { return GetLocalTime().Time; }
 
diff --git a/Engine/Source/Editor/Sequencer/Private/SequencerTimeSliderController.cpp b/Engine/Source/Editor/Sequencer/Private/SequencerTimeSliderController.cpp
--- a/Engine/Source/Editor/Sequencer/Private/SequencerTimeSliderController.cpp
+++ b/Engine/Source/Editor/Sequencer/Private/SequencerTimeSliderController.cpp
@@ -51,6 +51,7 @@ FSequencerTimeSliderController::FSequencerTimeSliderController( const FTimeSlide
 	, MouseDragType( DRAG_NONE )
 	, bMouseDownInRegion(false)
 	, bPanning( false )
+	, DragMarkIndex( INDEX_NONE )
 {
 	ScrubFillBrush              = FEditorStyle::GetBrush( TEXT( "Sequencer.Timeline.ScrubFill" ) );
 	FrameBlockScrubHandleUpBrush   = FEditorStyle::GetBrush( TEXT( "Sequencer.Timeline.FrameBlockScrubHandleUp" ) ); 
@@ -726,6 +727,7 @@ FReply FSequencerTimeSliderController::OnMouseButtonDown( SWidget& WidgetOwner,
 	MouseDownPosition[0] = MouseDownPosition[1] = MouseEvent.GetScreenSpacePosition();
 	MouseDownGeometry = MyGeometry;
 	bMouseDownInRegion = false;
+	DragMarkIndex = INDEX_NONE;
 
 	FVector2D CursorPos = MouseEvent.GetScreenSpacePosition();
 	FVector2D LocalPos = MouseDownGeometry.AbsoluteToLocal(CursorPos);
@@ -752,7 +754,7 @@ FReply FSequencerTimeSliderController::OnMouseButtonUp( SWidget& WidgetOwner, co
 			// Open a context menu if allowed
 			if (ContextMenuSuppression == 0 && TimeSliderArgs.PlaybackRange.IsSet())
 			{
-				TSharedRef<SWidget> MenuContent = OpenSetPlaybackRangeMenu(MouseTime.FrameNumber);
+				TSharedRef<SWidget> MenuContent = OpenSetPlaybackRangeMenu(MyGeometry, MouseEvent);
 				FSlateApplication::Get().PushMenu(
 					WidgetOwner.AsShared(),
 					MouseEvent.GetEventPath() != nullptr ? *MouseEvent.GetEventPath() : FWidgetPath(),
@@ -793,6 +795,10 @@ FReply FSequencerTimeSliderController::OnMouseButtonUp( SWidget& WidgetOwner, co
 		{
 			TimeSliderArgs.OnSelectionRangeEndDrag.ExecuteIfBound();
 		}
+		else if (MouseDragType == DRAG_MARK)
+		{
+			TimeSliderArgs.OnMarkEndDrag.ExecuteIfBound();
+		}
 		else if (MouseDragType == DRAG_SETTING_RANGE)
 		{
 			// Zooming
@@ -938,6 +944,11 @@ FReply FSequencerTimeSliderController::OnMouseMove( SWidget& WidgetOwner, const
 					MouseDragType = DRAG_PLAYBACK_START;
 					TimeSliderArgs.OnPlaybackRangeBeginDrag.ExecuteIfBound();
 				}
+				else if (!bHitScrubber && HitTestMark(RangeToScreen, MouseDownPixel, DragMarkIndex))
+				{
+					MouseDragType = DRAG_MARK;
+					TimeSliderArgs.OnMarkBeginDrag.ExecuteIfBound();
+				}
 				else if (FSlateApplication::Get().GetModifierKeys().AreModifersDown(EModifierKey::Control))
 				{
 					MouseDragType = DRAG_SETTING_RANGE;
@@ -973,6 +984,10 @@ FReply FSequencerTimeSliderController::OnMouseMove( SWidget& WidgetOwner, const
 			{
 				SetSelectionRangeEnd(MouseTime.FrameNumber);
 			}
+			else if (MouseDragType == DRAG_MARK)
+			{
+				SetMark(DragMarkIndex, MouseTime.FrameNumber);
+			}
 			else if (MouseDragType == DRAG_SCRUBBING_TIME)
 			{
 				// Delegate responsibility for clamping to the current viewrange to the client
@@ -1077,6 +1092,12 @@ FCursorReply FSequencerTimeSliderController::OnCursorQuery( TSharedRef<const SWi
 		return FCursorReply::Cursor(EMouseCursor::ResizeLeftRight);
 	}
 
+	int32 DummyMarkIndex = INDEX_NONE;
+	if (MouseDragType == DRAG_MARK || (!bHitScrubber && HitTestMark(RangeToScreen, HitTestPixel, DummyMarkIndex)))
+	{
+		return FCursorReply::Cursor(EMouseCursor::CardinalCross);
+	}
+
 	return FCursorReply::Unhandled();
 }
 
@@ -1174,8 +1195,12 @@ int32 FSequencerTimeSliderController::OnPaintSectionView( const FGeometry& Allot
 	return LayerId;
 }
 
-TSharedRef<SWidget> FSequencerTimeSliderController::OpenSetPlaybackRangeMenu(FFrameNumber FrameNumber)
+TSharedRef<SWidget> FSequencerTimeSliderController::OpenSetPlaybackRangeMenu(const FGeometry& MyGeometry, const FPointerEvent& MouseEvent)
 {
+	FScrubRangeToScreen RangeToScreen = FScrubRangeToScreen(TimeSliderArgs.ViewRange.Get(), MyGeometry.Size);
+	const float MousePixel = MyGeometry.AbsoluteToLocal(MouseEvent.GetScreenSpacePosition()).X;
+	FFrameNumber FrameNumber = ComputeFrameTimeFromMouse(MyGeometry, MouseEvent.GetScreenSpacePosition(), RangeToScreen).FrameNumber;
+
 	const bool bShouldCloseWindowAfterMenuSelection = true;
 	FMenuBuilder MenuBuilder(bShouldCloseWindowAfterMenuSelection, nullptr);
 
@@ -1262,9 +1287,11 @@ TSharedRef<SWidget> FSequencerTimeSliderController::OpenSetPlaybackRangeMenu(FFr
 		FFrameNumber DisplayFrameNumber = GetDisplayRate().AsFrameNumber(FrameNumber / GetTickResolution());
 
 		UMovieScene* MovieScene = WeakSequencer.Pin()->GetFocusedMovieSceneSequence()->GetMovieScene();
-		int32 MarkedIndex = MovieScene->FindMarkedFrameByFrameNumber(FrameNumber);
 		bool bHasMarks = MovieScene->GetMarkedFrames().Num() > 0;
 
+		int32 MarkedIndex = INDEX_NONE;
+		HitTestMark(RangeToScreen, MousePixel, MarkedIndex);
+
 		if (MarkedIndex != INDEX_NONE)
 		{
 			FPropertyEditorModule& PropertyEditorModule = FModuleManager::GetModuleChecked<FPropertyEditorModule>("PropertyEditor");
@@ -1469,6 +1496,36 @@ bool FSequencerTimeSliderController::HitTestRangeEnd(const FScrubRangeToScreen&
 		HitPixel <= RangeEndPixel + MouseTolerance + DragToleranceSlateUnits;
 }
 
+bool FSequencerTimeSliderController::HitTestMark(const FScrubRangeToScreen& RangeToScreen, float HitPixel, int32& OutMarkIndex) const
+{
+	const TArray<FMovieSceneMarkedFrame> & MarkedFrames = TimeSliderArgs.MarkedFrames.Get();
+	if (MarkedFrames.Num() < 1)
+	{
+		return false;
+	}
+
+	static float BrushSizeInStateUnits = 3.f, DragToleranceSlateUnits = 2.f, MouseTolerance = 2.f;
+
+	for (int32 MarkIndex = 0; MarkIndex < MarkedFrames.Num(); ++MarkIndex)
+	{
+		double Seconds = MarkedFrames[MarkIndex].FrameNumber / GetTickResolution();
+
+		float MarkPixel = RangeToScreen.InputToLocalX(Seconds);
+
+		// Hit test against the brush region to the left/right of the mark position, +/- DragToleranceSlateUnits
+		if ((HitPixel >= MarkPixel - MouseTolerance - DragToleranceSlateUnits &&
+			 HitPixel <= MarkPixel + MouseTolerance + BrushSizeInStateUnits + DragToleranceSlateUnits) ||
+			(HitPixel >= MarkPixel - MouseTolerance - BrushSizeInStateUnits - DragToleranceSlateUnits &&
+			 HitPixel <= MarkPixel + MouseTolerance + DragToleranceSlateUnits))
+		{
+			OutMarkIndex = MarkIndex;
+			return true;
+		}
+	}
+
+	return false;
+}
+
 FFrameTime FSequencerTimeSliderController::SnapTimeToNearestKey(const FScrubRangeToScreen& RangeToScreen, float CursorPos, FFrameTime InTime) const
 {
 	if (!WeakSequencer.IsValid())
@@ -1545,6 +1602,11 @@ void FSequencerTimeSliderController::SetSelectionRangeEnd(FFrameNumber NewEnd)
 	}
 }
 
+void FSequencerTimeSliderController::SetMark(int32 InMarkIndex, FFrameNumber FrameNumber)
+{
+	TimeSliderArgs.OnSetMarkedFrame.ExecuteIfBound(InMarkIndex, FrameNumber);
+}
+
 void FSequencerTimeSliderController::AddMarkAtFrame(FFrameNumber FrameNumber)
 {
 	TimeSliderArgs.OnMarkedFrameChanged.ExecuteIfBound(FrameNumber, true);
diff --git a/Engine/Source/Editor/Sequencer/Public/SequencerTimeSliderController.h b/Engine/Source/Editor/Sequencer/Public/SequencerTimeSliderController.h
--- a/Engine/Source/Editor/Sequencer/Public/SequencerTimeSliderController.h
+++ b/Engine/Source/Editor/Sequencer/Public/SequencerTimeSliderController.h
@@ -257,6 +257,13 @@ private:
 	 */
 	bool HitTestRangeEnd(const FScrubRangeToScreen& RangeToScreen, const TRange<double>& Range, float HitPixel) const;
 
+	/**
+	 * Hit test marks
+	 *
+	 * @return The mark index hit
+	 */
+	bool HitTestMark(const FScrubRangeToScreen& RangeToScreen, float HitPixel, int32& OutMarkIndex) const;
+
 	FFrameTime SnapTimeToNearestKey(const FScrubRangeToScreen& RangeToScreen, float CursorPos, FFrameTime InTime) const;
 
 	void SetPlaybackRangeStart(FFrameNumber NewStart);
@@ -265,7 +272,9 @@ private:
 	void SetSelectionRangeStart(FFrameNumber NewStart);
 	void SetSelectionRangeEnd(FFrameNumber NewEnd);
 
-	TSharedRef<SWidget> OpenSetPlaybackRangeMenu(FFrameNumber FrameNumber);
+	void SetMark(int32 InMarkIndex, FFrameNumber NewFrame);
+
+	TSharedRef<SWidget> OpenSetPlaybackRangeMenu(const FGeometry& MyGeometry, const FPointerEvent& MouseEvent);
 	FFrameTime ComputeScrubTimeFromMouse(const FGeometry& Geometry, FVector2D ScreenSpacePosition, FScrubRangeToScreen RangeToScreen) const;
 	FFrameTime ComputeFrameTimeFromMouse(const FGeometry& Geometry, FVector2D ScreenSpacePosition, FScrubRangeToScreen RangeToScreen, bool CheckSnapping = true) const;
 
@@ -316,6 +325,7 @@ private:
 		DRAG_PLAYBACK_END,
 		DRAG_SELECTION_START,
 		DRAG_SELECTION_END,
+		DRAG_MARK,
 		DRAG_NONE
 	};
 	
@@ -336,6 +346,9 @@ private:
 	/** Range stack */
 	TArray<TRange<double>> ViewRangeStack;
 
+	/** Index of mark being edited */
+	int32 DragMarkIndex;
+
 	/** When > 0, we should not show context menus */
 	int32 ContextMenuSuppression;
 	
diff --git a/Engine/Source/Editor/SequencerWidgets/Public/ITimeSlider.h b/Engine/Source/Editor/SequencerWidgets/Public/ITimeSlider.h
--- a/Engine/Source/Editor/SequencerWidgets/Public/ITimeSlider.h
+++ b/Engine/Source/Editor/SequencerWidgets/Public/ITimeSlider.h
@@ -28,7 +28,8 @@ DECLARE_DELEGATE_TwoParams( FOnScrubPositionChanged, FFrameTime, bool )
 DECLARE_DELEGATE_TwoParams( FOnViewRangeChanged, TRange<double>, EViewRangeInterpolation )
 DECLARE_DELEGATE_OneParam( FOnTimeRangeChanged, TRange<double> )
 DECLARE_DELEGATE_OneParam( FOnFrameRangeChanged, TRange<FFrameNumber> )
-DECLARE_DELEGATE_TwoParams( FOnMarkedFrameChanged, FFrameNumber, bool )
+DECLARE_DELEGATE_TwoParams(FOnSetMarkedFrame, int32, FFrameNumber)
+DECLARE_DELEGATE_TwoParams(FOnMarkedFrameChanged, FFrameNumber, bool)
 DECLARE_DELEGATE_RetVal_TwoParams( FFrameNumber, FOnGetNearestKey, FFrameTime, bool )
 
 /** Structure used to wrap up a range, and an optional animation target */
@@ -135,13 +136,22 @@ struct FTimeSliderArgs
 	/** Called right after the selection range has finished being dragged */
 	FSimpleDelegate OnSelectionRangeEndDrag;
 
+	/** Called right before a mark starts to be dragged */
+	FSimpleDelegate OnMarkBeginDrag;
+
+	/** Called right after a mark has finished being dragged */
+	FSimpleDelegate OnMarkEndDrag;
+
 	/** Attribute for the current sequence's vertical frames */
 	TAttribute<TSet<FFrameNumber>> VerticalFrames;
 
 	/** Attribute for the current sequence's marked frames */
 	TAttribute<TArray<FMovieSceneMarkedFrame>> MarkedFrames;
 
-	/** Called when the marked frames need to be updated */
+	/** Called when the marked frame needs to be set */
+	FOnSetMarkedFrame OnSetMarkedFrame;
+
+	/** Called when a marked frame is added/removed */
 	FOnMarkedFrameChanged OnMarkedFrameChanged;
 
 	/** Called when all marked frames should be cleared */
diff --git a/Engine/Source/Runtime/MovieScene/Private/MovieScene.cpp b/Engine/Source/Runtime/MovieScene/Private/MovieScene.cpp
--- a/Engine/Source/Runtime/MovieScene/Private/MovieScene.cpp
+++ b/Engine/Source/Runtime/MovieScene/Private/MovieScene.cpp
@@ -1033,6 +1033,16 @@ void UMovieScene::MoveBindingContents(const FGuid& SourceBindingId, const FGuid&
 	}
 }
 
+void UMovieScene::SetMarkedFrame(int32 InMarkIndex, FFrameNumber InFrameNumber)
+{
+	if (InMarkIndex < 0 && InMarkIndex > MarkedFrames.Num()-1)
+	{
+		return;
+	}
+
+	MarkedFrames[InMarkIndex].FrameNumber = InFrameNumber;
+}
+
 void UMovieScene::AddMarkedFrame(const FMovieSceneMarkedFrame &InMarkedFrame)
 {
 	FString NewLabel;
diff --git a/Engine/Source/Runtime/MovieScene/Public/MovieScene.h b/Engine/Source/Runtime/MovieScene/Public/MovieScene.h
--- a/Engine/Source/Runtime/MovieScene/Public/MovieScene.h
+++ b/Engine/Source/Runtime/MovieScene/Public/MovieScene.h
@@ -750,6 +750,14 @@ public:
 	 */
 	const TArray<FMovieSceneMarkedFrame>& GetMarkedFrames() const { return MarkedFrames; }
 
+	/*
+	 * Sets the frame number for the given marked frame index.
+	 *
+	 * @InMarkIndex The given user marked frame index to edit
+	 * @InFrameNumber The frame number to set
+	 */
+	void SetMarkedFrame(int32 InMarkIndex, FFrameNumber InFrameNumber);
+
 	/*
 	 * Add a given user marked frame.
 	 * A unique label will be generated if the marked frame label is empty
