Sequencer: Follow up to displaying property track bindings as red if not bound. This should only be done if there's a valid object track (ie. not for slomo, live link, or fade tracks)

diff --git a/Engine/Source/Editor/Sequencer/Private/DisplayNodes/SequencerTrackNode.cpp b/Engine/Source/Editor/Sequencer/Private/DisplayNodes/SequencerTrackNode.cpp
--- a/Engine/Source/Editor/Sequencer/Private/DisplayNodes/SequencerTrackNode.cpp
+++ b/Engine/Source/Editor/Sequencer/Private/DisplayNodes/SequencerTrackNode.cpp
@@ -499,17 +499,20 @@ FLinearColor FSequencerTrackNode::GetDisplayNameColor() const
 			ObjectBinding = StaticCastSharedPtr<FSequencerObjectBindingNode>(ParentSeqNode)->GetObjectBinding();
 		}
 
-		for (auto RuntimeObject : GetSequencer().FindBoundObjects(ObjectBinding, GetSequencer().GetFocusedTemplateID()))
+		if (ObjectBinding.IsValid())
 		{
-			FTrackInstancePropertyBindings PropertyBinding(FName(*PropertyTrack->GetName()), PropertyTrack->GetPropertyPath());
-
-			if (PropertyBinding.GetProperty(*RuntimeObject))
+			for (auto RuntimeObject : GetSequencer().FindBoundObjects(ObjectBinding, GetSequencer().GetFocusedTemplateID()))
 			{
-				return FLinearColor::White;
+				FTrackInstancePropertyBindings PropertyBinding(FName(*PropertyTrack->GetName()), PropertyTrack->GetPropertyPath());
+
+				if (PropertyBinding.GetProperty(*RuntimeObject))
+				{
+					return FLinearColor::White;
+				}
 			}
-		}
 
-		return FLinearColor::Red;
+			return FLinearColor::Red;
+		}
 	}
 
 	return FLinearColor::White;
