Navigate to controlling sequence for Actors via World Outliner/World Menu

diff --git a/Engine/Source/Editor/Sequencer/Private/LevelEditorSequencerIntegration.cpp b/Engine/Source/Editor/Sequencer/Private/LevelEditorSequencerIntegration.cpp
--- a/Engine/Source/Editor/Sequencer/Private/LevelEditorSequencerIntegration.cpp
+++ b/Engine/Source/Editor/Sequencer/Private/LevelEditorSequencerIntegration.cpp
@@ -19,6 +19,7 @@
 #include "Sequencer.h"
 #include "EditorModeManager.h"
 #include "SequencerCommands.h"
+#include "Sections/MovieSceneSubSection.h"
 #include "MovieSceneSequence.h"
 #include "MovieScene.h"
 #include "Tracks/MovieScenePropertyTrack.h"
@@ -624,6 +625,62 @@ void FLevelEditorSequencerIntegration::AddLevelViewportMenuExtender()
 	);
 }
 
+void FindActorInSequencesRecursive(AActor* InActor, FSequencer& Sequencer, FMovieSceneSequenceIDRef SequenceID, TArray<TPair<FMovieSceneSequenceID, FSequencer*> >& FoundInSequences)
+{
+	FMovieSceneRootEvaluationTemplateInstance& RootInstance = Sequencer.GetEvaluationTemplate();
+
+	// Find the sequence that corresponds to the sequence ID
+	UMovieSceneSequence* Sequence = RootInstance.GetSequence(SequenceID);
+	UMovieScene* MovieScene = Sequence ? Sequence->GetMovieScene() : nullptr;
+
+	// Recurse into child nodes
+	const FMovieSceneSequenceHierarchyNode* Node = RootInstance.GetHierarchy().FindNode(SequenceID);
+	if (Node)
+	{
+		for (FMovieSceneSequenceIDRef ChildID : Node->Children)
+		{
+			FindActorInSequencesRecursive(InActor, Sequencer, ChildID, FoundInSequences);
+		}
+	}
+
+	if (MovieScene)
+	{
+		FString SequenceName = Sequence->GetDisplayName().ToString();
+
+		// Search all possessables
+		for (int32 Index = 0; Index < MovieScene->GetPossessableCount(); ++Index)
+		{
+			FGuid ThisGuid = MovieScene->GetPossessable(Index).GetGuid();
+
+			for (TWeakObjectPtr<> WeakObject : Sequencer.FindBoundObjects(ThisGuid, SequenceID))
+			{
+				AActor* Actor = Cast<AActor>(WeakObject.Get());
+				if (Actor == InActor)
+				{
+					FoundInSequences.Add(TPair<FMovieSceneSequenceID, FSequencer*>(SequenceID, &Sequencer));
+					return;
+				}
+			}
+		}
+
+		// Search all spawnables
+		for (int32 Index = 0; Index < MovieScene->GetSpawnableCount(); ++Index)
+		{
+			FGuid ThisGuid = MovieScene->GetSpawnable(Index).GetGuid();
+
+			for (TWeakObjectPtr<> WeakObject : Sequencer.FindBoundObjects(ThisGuid, SequenceID))
+			{
+				AActor* Actor = Cast<AActor>(WeakObject.Get());
+				if (Actor == InActor)
+				{
+					FoundInSequences.Add(TPair<FMovieSceneSequenceID, FSequencer*>(SequenceID, &Sequencer));
+					return;
+				}
+			}
+		}
+	}
+}
+
 TSharedRef<FExtender> FLevelEditorSequencerIntegration::GetLevelViewportExtender(const TSharedRef<FUICommandList> CommandList, const TArray<AActor*> InActors)
 {
 	TSharedRef<FExtender> Extender = MakeShareable(new FExtender);
@@ -638,18 +695,47 @@ TSharedRef<FExtender> FLevelEditorSequencerIntegration::GetLevelViewportExtender
 		ActorName = FText::Format(LOCTEXT("ActorNamePlural", "{0} Actors"), FText::AsNumber(InActors.Num()));
 	}
 
+	TArray<TPair<FMovieSceneSequenceID, FSequencer*> > FoundInSequences;
+	IterateAllSequencers(
+		[&](FSequencer& In, const FLevelEditorSequencerIntegrationOptions& Options)
+	{
+		FindActorInSequencesRecursive(InActors[0], In, MovieSceneSequenceID::Root, FoundInSequences);
+	});
+
 	TSharedRef<FUICommandList> LevelEditorCommandBindings  = FModuleManager::GetModuleChecked<FLevelEditorModule>(TEXT("LevelEditor")).GetGlobalLevelEditorActions();
 	Extender->AddMenuExtension("ActorControl", EExtensionHook::After, LevelEditorCommandBindings, FMenuExtensionDelegate::CreateLambda(
-		[ActorName](FMenuBuilder& MenuBuilder){
-			MenuBuilder.BeginSection("SequenceRecording", LOCTEXT("SequenceRecordingHeading", "Sequence Recording"));
-			MenuBuilder.AddMenuEntry(FSequencerCommands::Get().RecordSelectedActors, NAME_None, FText::Format(LOCTEXT("RecordSelectedActorsText", "Record {0} In Sequencer"), ActorName));
-			MenuBuilder.EndSection();
+		[this, ActorName, Actor = InActors[0], FoundInSequences](FMenuBuilder& MenuBuilder) {
+		MenuBuilder.BeginSection("Sequencer", LOCTEXT("Sequencer", "Sequencer"));
+		MenuBuilder.AddMenuEntry(FSequencerCommands::Get().RecordSelectedActors, NAME_None, FText::Format(LOCTEXT("RecordSelectedActorsText", "Record {0} In Sequencer"), ActorName));
+
+		if (FoundInSequences.Num() > 0)
+		{
+			MenuBuilder.AddSubMenu(
+				LOCTEXT("BrowseToActorInSequencer", "Browse to Actor in Sequencer"),
+				FText(),
+				FNewMenuDelegate::CreateRaw(this, &FLevelEditorSequencerIntegration::MakeBrowseToSelectedActorSubMenu, Actor, FoundInSequences),
+				false,
+				FSlateIcon()
+			);
 		}
+
+		MenuBuilder.EndSection();
+	}
 	));
 
 	return Extender;
 }
 
+void FLevelEditorSequencerIntegration::MakeBrowseToSelectedActorSubMenu(FMenuBuilder& MenuBuilder, AActor* Actor, const TArray<TPair<FMovieSceneSequenceID, FSequencer*> > FoundInSequences)
+{
+	for (const TPair<FMovieSceneSequenceID, FSequencer*> Sequence : FoundInSequences)
+	{
+		FText ActorName = FText::Format(LOCTEXT("ActorNameSingular", "\"{0}\""), FText::FromString(Actor->GetActorLabel()));
+		FUIAction AddMenuAction(FExecuteAction::CreateLambda([this, Sequence]() {this->BrowseToSelectedActor(Sequence.Value, Sequence.Key); }));
+		MenuBuilder.AddMenuEntry(FText::Format(LOCTEXT("BrowseToSelectedActorText", "Browse to {0} in {1}"), ActorName, Sequence.Value->FindSubSection(Sequence.Key)->GetSequence()->GetDisplayName()), FText(), FSlateIcon(), AddMenuAction);
+	}
+}
+
 void FLevelEditorSequencerIntegration::ActivateDetailHandler()
 {
 	FName DetailHandlerName("DetailHandler");
@@ -758,6 +844,12 @@ void FLevelEditorSequencerIntegration::RecordSelectedActors()
 	);
 }
 
+void FLevelEditorSequencerIntegration::BrowseToSelectedActor(FSequencer* Sequencer, FMovieSceneSequenceID SequenceID)
+{
+	Sequencer->PopToSequenceInstance(MovieSceneSequenceID::Root);
+	Sequencer->FocusSequenceInstance(*Sequencer->FindSubSection(SequenceID));
+}
+
 namespace FaderConstants
 {	
 	/** The opacity when we are hovered */
@@ -1192,7 +1284,6 @@ void AddPropertiesToBindingsMap(TWeakPtr<FSequencer> Sequencer, UMovieSceneSeque
 	}
 }
 
-
 FString FLevelEditorSequencerBindingData::GetLevelSequencesForActor(TWeakPtr<FSequencer> Sequencer, const AActor* InActor)
 {
 	if (bActorBindingsDirty)
diff --git a/Engine/Source/Editor/Sequencer/Private/Sequencer.cpp b/Engine/Source/Editor/Sequencer/Private/Sequencer.cpp
--- a/Engine/Source/Editor/Sequencer/Private/Sequencer.cpp
+++ b/Engine/Source/Editor/Sequencer/Private/Sequencer.cpp
@@ -951,7 +951,15 @@ void FSequencer::PopToSequenceInstance(FMovieSceneSequenceIDRef SequenceID)
 		}
 
 		ResetPerMovieSceneData();
-		SequencerWidget->UpdateBreadcrumbs();
+
+		if (SequenceID == MovieSceneSequenceID::Root)
+		{
+			SequencerWidget->ResetBreadcrumbs();
+		}
+		else
+		{
+			SequencerWidget->UpdateBreadcrumbs();
+		}
 
 		OnActivateSequenceEvent.Broadcast(ActiveTemplateIDs.Top());
 
diff --git a/Engine/Source/Editor/Sequencer/Public/LevelEditorSequencerIntegration.h b/Engine/Source/Editor/Sequencer/Public/LevelEditorSequencerIntegration.h
--- a/Engine/Source/Editor/Sequencer/Public/LevelEditorSequencerIntegration.h
+++ b/Engine/Source/Editor/Sequencer/Public/LevelEditorSequencerIntegration.h
@@ -157,6 +157,9 @@ private:
 
 	void RecordSelectedActors();
 
+	void MakeBrowseToSelectedActorSubMenu(FMenuBuilder& MenuBuilder, AActor* Actor, const TArray<TPair<FMovieSceneSequenceID, FSequencer*> > FoundInSequences);
+	void BrowseToSelectedActor(FSequencer* Sequencer, FMovieSceneSequenceID SequenceId);
+
 	bool IsPropertyReadOnly(const FPropertyAndParent& InPropertyAndParent);
 
 private:
