Sequencer: Add section bounds to cached key positions for track lane. Renamed ObjectTrack to CombinedKeysTrack and allow tracks other than object tracks to display combined keyframes (ie. audio track)

diff --git a/Engine/Source/Editor/Sequencer/Private/DisplayNodes/SequencerDisplayNode.cpp b/Engine/Source/Editor/Sequencer/Private/DisplayNodes/SequencerDisplayNode.cpp
--- a/Engine/Source/Editor/Sequencer/Private/DisplayNodes/SequencerDisplayNode.cpp
+++ b/Engine/Source/Editor/Sequencer/Private/DisplayNodes/SequencerDisplayNode.cpp
@@ -58,12 +58,12 @@ struct FNameAndSignature
 	}
 };
 
-class SSequencerObjectTrack
+class SSequencerCombinedKeysTrack
 	: public SLeafWidget
 {
 public:
 
-	SLATE_BEGIN_ARGS(SSequencerObjectTrack) {}
+	SLATE_BEGIN_ARGS(SSequencerCombinedKeysTrack) {}
 		/** The view range of the section area */
 		SLATE_ATTRIBUTE( TRange<double>, ViewRange )
 		/** The tick resolution of the current sequence*/
@@ -80,8 +80,6 @@ public:
 		
 		ViewRange = InArgs._ViewRange;
 		TickResolution = InArgs._TickResolution;
-
-		check(RootNode->GetType() == ESequencerNode::Object);
 	}
 
 protected:
@@ -117,7 +115,7 @@ private:
 };
 
 
-void SSequencerObjectTrack::Tick( const FGeometry& AllottedGeometry, const double InCurrentTime, const float InDeltaTime )
+void SSequencerCombinedKeysTrack::Tick( const FGeometry& AllottedGeometry, const double InCurrentTime, const float InDeltaTime )
 {
 	SWidget::Tick(AllottedGeometry, InCurrentTime, InDeltaTime);
 
@@ -136,8 +134,16 @@ void SSequencerObjectTrack::Tick( const FGeometry& AllottedGeometry, const doubl
 	}
 }
 
-int32 SSequencerObjectTrack::OnPaint(const FPaintArgs& Args, const FGeometry& AllottedGeometry, const FSlateRect& MyCullingRect, FSlateWindowElementList& OutDrawElements, int32 LayerId, const FWidgetStyle& InWidgetStyle, bool bParentEnabled) const
+int32 SSequencerCombinedKeysTrack::OnPaint(const FPaintArgs& Args, const FGeometry& AllottedGeometry, const FSlateRect& MyCullingRect, FSlateWindowElementList& OutDrawElements, int32 LayerId, const FWidgetStyle& InWidgetStyle, bool bParentEnabled) const
 {
+	if (RootNode->GetType() == ESequencerNode::Track)
+	{
+		if (static_cast<FSequencerTrackNode&>(*RootNode).GetSubTrackMode() == FSequencerTrackNode::ESubTrackMode::ParentTrack && RootNode->IsExpanded())
+		{
+			return LayerId;
+		}
+	}
+
 	if (RootNode->GetSequencer().GetSequencerSettings()->GetShowCombinedKeyframes())
 	{
 		for (float KeyPosition : KeyDrawPositions)
@@ -164,14 +170,14 @@ int32 SSequencerObjectTrack::OnPaint(const FPaintArgs& Args, const FGeometry& Al
 }
 
 
-FVector2D SSequencerObjectTrack::ComputeDesiredSize( float ) const
+FVector2D SSequencerCombinedKeysTrack::ComputeDesiredSize( float ) const
 {
 	// Note: X Size is not used
 	return FVector2D( 100.0f, RootNode->GetNodeHeight() );
 }
 
 
-void SSequencerObjectTrack::GenerateCachedKeyPositions(const FGeometry& AllottedGeometry)
+void SSequencerCombinedKeysTrack::GenerateCachedKeyPositions(const FGeometry& AllottedGeometry)
 {
 	static float DuplicateThresholdPx = 3.f;
 
@@ -182,12 +188,25 @@ void SSequencerObjectTrack::GenerateCachedKeyPositions(const FGeometry& Allotted
 	// Unnamed key areas are uncacheable, so we track those separately
 	TArray<FSequencerCachedKeys> UncachableKeyTimes;
 
+	TArray<double> SectionBoundTimes;
+
 	// First off, accumulate (and cache) KeyDrawPositions as times, we convert to positions in the later loop
 	for (auto& CachePair : KeyCollectionSignature.GetKeyAreas())
 	{
 		TSharedRef<IKeyArea> KeyArea = CachePair.Key;
 
 		UMovieSceneSection* Section = KeyArea->GetOwningSection();
+
+		if (Section->HasStartFrame())
+		{
+			SectionBoundTimes.Add(Section->GetInclusiveStartFrame() / CachedTickResolution);
+		}
+
+		if (Section->HasEndFrame())
+		{
+			SectionBoundTimes.Add(Section->GetExclusiveEndFrame() / CachedTickResolution);
+		}
+
 		FNameAndSignature CacheKey{ CachePair.Value, KeyArea->GetName() };
 
 		// If we cached this last frame, use those key times again
@@ -243,6 +262,7 @@ void SSequencerObjectTrack::GenerateCachedKeyPositions(const FGeometry& Allotted
 		Uncached.GetKeysInRange(CachedViewRange, &Times, nullptr, nullptr);
 		AllIterators.Add(Times);
 	}
+	AllIterators.Add(TArrayView<const double>(SectionBoundTimes));
 
 	FTimeToPixel TimeToPixelConverter(AllottedGeometry, CachedViewRange, CachedTickResolution);
 
@@ -577,16 +597,10 @@ TSharedRef<SWidget> FSequencerDisplayNode::GenerateWidgetForSectionArea(const TA
 			.ViewRange(ViewRange);
 	}
 	
-	if (GetType() == ESequencerNode::Object)
-	{
-		return SNew(SSequencerObjectTrack, SharedThis(this))
-			.ViewRange(ViewRange)
-			.IsEnabled(!GetSequencer().IsReadOnly())
-			.TickResolution(this, &FSequencerDisplayNode::GetTickResolution);
-	}
-
-	// currently only section areas display widgets
-	return SNullWidget::NullWidget;
+	return SNew(SSequencerCombinedKeysTrack, SharedThis(this))
+		.ViewRange(ViewRange)
+		.IsEnabled(!GetSequencer().IsReadOnly())
+		.TickResolution(this, &FSequencerDisplayNode::GetTickResolution);
 }
 
 FFrameRate FSequencerDisplayNode::GetTickResolution() const
