When Navigating to controlling sequence for via outliner, scroll to selected actor.

diff --git a/Engine/Source/Editor/Sequencer/Private/LevelEditorSequencerIntegration.cpp b/Engine/Source/Editor/Sequencer/Private/LevelEditorSequencerIntegration.cpp
--- a/Engine/Source/Editor/Sequencer/Private/LevelEditorSequencerIntegration.cpp
+++ b/Engine/Source/Editor/Sequencer/Private/LevelEditorSequencerIntegration.cpp
@@ -731,7 +731,7 @@ void FLevelEditorSequencerIntegration::MakeBrowseToSelectedActorSubMenu(FMenuBui
 	for (const TPair<FMovieSceneSequenceID, FSequencer*> Sequence : FoundInSequences)
 	{
 		FText ActorName = FText::Format(LOCTEXT("ActorNameSingular", "\"{0}\""), FText::FromString(Actor->GetActorLabel()));
-		FUIAction AddMenuAction(FExecuteAction::CreateLambda([this, Sequence]() {this->BrowseToSelectedActor(Sequence.Value, Sequence.Key); }));
+		FUIAction AddMenuAction(FExecuteAction::CreateLambda([this, Actor, Sequence]() {this->BrowseToSelectedActor(Actor, Sequence.Value, Sequence.Key); }));
 		MenuBuilder.AddMenuEntry(FText::Format(LOCTEXT("BrowseToSelectedActorText", "Browse to {0} in {1}"), ActorName, Sequence.Value->FindSubSection(Sequence.Key)->GetSequence()->GetDisplayName()), FText(), FSlateIcon(), AddMenuAction);
 	}
 }
@@ -844,10 +844,11 @@ void FLevelEditorSequencerIntegration::RecordSelectedActors()
 	);
 }
 
-void FLevelEditorSequencerIntegration::BrowseToSelectedActor(FSequencer* Sequencer, FMovieSceneSequenceID SequenceID)
+void FLevelEditorSequencerIntegration::BrowseToSelectedActor(AActor* Actor, FSequencer* Sequencer, FMovieSceneSequenceID SequenceID)
 {
 	Sequencer->PopToSequenceInstance(MovieSceneSequenceID::Root);
 	Sequencer->FocusSequenceInstance(*Sequencer->FindSubSection(SequenceID));
+	Sequencer->SelectObject(Sequencer->FindObjectId(*Actor, SequenceID));
 }
 
 namespace FaderConstants
diff --git a/Engine/Source/Editor/Sequencer/Public/LevelEditorSequencerIntegration.h b/Engine/Source/Editor/Sequencer/Public/LevelEditorSequencerIntegration.h
--- a/Engine/Source/Editor/Sequencer/Public/LevelEditorSequencerIntegration.h
+++ b/Engine/Source/Editor/Sequencer/Public/LevelEditorSequencerIntegration.h
@@ -158,7 +158,7 @@ private:
 	void RecordSelectedActors();
 
 	void MakeBrowseToSelectedActorSubMenu(FMenuBuilder& MenuBuilder, AActor* Actor, const TArray<TPair<FMovieSceneSequenceID, FSequencer*> > FoundInSequences);
-	void BrowseToSelectedActor(FSequencer* Sequencer, FMovieSceneSequenceID SequenceId);
+	void BrowseToSelectedActor(AActor* Actor, FSequencer* Sequencer, FMovieSceneSequenceID SequenceId);
 
 	bool IsPropertyReadOnly(const FPropertyAndParent& InPropertyAndParent);
 
