Sequencer: Snap to playback range when dragging keys and sections if the setting is enabled.

diff --git a/Engine/Source/Editor/Sequencer/Private/Tools/EditToolDragOperations.cpp b/Engine/Source/Editor/Sequencer/Private/Tools/EditToolDragOperations.cpp
--- a/Engine/Source/Editor/Sequencer/Private/Tools/EditToolDragOperations.cpp
+++ b/Engine/Source/Editor/Sequencer/Private/Tools/EditToolDragOperations.cpp
@@ -797,6 +797,12 @@ void FMoveKeysAndSections::OnDrag(const FPointerEvent& MouseEvent, FVector2D Loc
 		}
 	}
 
+	if (Settings->ShouldKeepCursorInPlayRangeWhileScrubbing() && !Settings->ShouldKeepPlayRangeInSectionBounds())
+	{
+		MouseTime = MovieScene::ClampToDiscreteRange(MouseTime, Sequencer.GetPlaybackRange());
+	}
+
+
 	// We'll calculate a DeltaX based on limits on movement (snapping, section collision) and then use them on keys and sections below.
 	TOptional<FFrameNumber> MaxDeltaX = GetMovementDeltaX(MouseTime);
 
@@ -943,6 +949,19 @@ TOptional<FFrameNumber> FMoveKeysAndSections::GetMovementDeltaX(FFrameTime Mouse
 		FFrameNumber LeftMovementMaximum = MovieScene::DiscreteInclusiveLower(SectionBoundaries);
 		FFrameNumber RightMovementMaximum = MovieScene::DiscreteExclusiveUpper(SectionBoundaries);
 		
+		if (Settings->ShouldKeepCursorInPlayRangeWhileScrubbing() && !Settings->ShouldKeepPlayRangeInSectionBounds())
+		{
+			if (LeftMovementMaximum < Sequencer.GetPlaybackRange().GetLowerBoundValue())
+			{
+				LeftMovementMaximum = Sequencer.GetPlaybackRange().GetLowerBoundValue();
+			}
+
+			if (RightMovementMaximum > Sequencer.GetPlaybackRange().GetUpperBoundValue())
+			{
+				RightMovementMaximum = Sequencer.GetPlaybackRange().GetUpperBoundValue();
+			}
+		}
+
 		if (Section->HasStartFrame())
 		{
 			FFrameNumber NewStartTime = Section->GetInclusiveStartFrame() + MouseDeltaTime;
@@ -970,6 +989,43 @@ TOptional<FFrameNumber> FMoveKeysAndSections::GetMovementDeltaX(FFrameTime Mouse
 		}
 	}
 
+	if (Settings->ShouldKeepCursorInPlayRangeWhileScrubbing() && !Settings->ShouldKeepPlayRangeInSectionBounds())
+	{
+		TArray<FFrameNumber> CurrentKeyTimes;
+		CurrentKeyTimes.SetNum(KeysAsArray.Num());
+		GetKeyTimes(KeysAsArray, CurrentKeyTimes);
+
+		for (int32 Index = 0; Index < CurrentKeyTimes.Num(); ++Index)
+		{
+			FSequencerSelectedKey& SelectedKey = KeysAsArray[Index];
+			const bool bOwningSectionIsSelected = Sections.ContainsByPredicate([SelectedKey](FSectionHandle Handle) { return Handle.GetSectionObject() == SelectedKey.Section; });
+
+			// We don't want to apply delta if we have the key's section selected as well, otherwise they get double
+			// transformed (moving the section moves the keys + we add the delta to the key positions).
+			if (!bOwningSectionIsSelected)
+			{
+				FFrameNumber NewKeyTime = CurrentKeyTimes[Index] + MouseDeltaTime;
+				if (NewKeyTime < Sequencer.GetPlaybackRange().GetLowerBoundValue())
+				{
+					FFrameNumber ClampedDeltaTime = CurrentKeyTimes[Index] - Sequencer.GetPlaybackRange().GetLowerBoundValue();
+					if (!DeltaX.IsSet() || DeltaX.GetValue() > ClampedDeltaTime)
+					{
+						DeltaX = ClampedDeltaTime;
+					}
+				}
+
+				if (NewKeyTime > Sequencer.GetPlaybackRange().GetUpperBoundValue())
+				{
+					FFrameNumber ClampedDeltaTime = Sequencer.GetPlaybackRange().GetUpperBoundValue() - CurrentKeyTimes[Index];
+					if (!DeltaX.IsSet() || DeltaX.GetValue() > ClampedDeltaTime)
+					{
+						DeltaX = ClampedDeltaTime;
+					}
+				}
+			}
+		}
+	}
+
 	return DeltaX;
 }
 
