Sequencer: Add stretch/shrink box

diff --git a/Engine/Source/Editor/Sequencer/Private/SSequencer.cpp b/Engine/Source/Editor/Sequencer/Private/SSequencer.cpp
--- a/Engine/Source/Editor/Sequencer/Private/SSequencer.cpp
+++ b/Engine/Source/Editor/Sequencer/Private/SSequencer.cpp
@@ -62,6 +62,7 @@
 #include "FrameNumberDetailsCustomization.h"
 #include "SequencerSettings.h"
 #include "SSequencerTransformBox.h"
+#include "SSequencerStretchBox.h"
 #include "SSequencerDebugVisualizer.h"
 #include "ISequencerModule.h"
 #include "IVREditorModule.h"
@@ -541,6 +542,15 @@ void SSequencer::Construct(const FArguments& InArgs, TSharedRef<FSequencer> InSe
 					]
 
 					+ SGridPanel::Slot(Column1, Row2, SGridPanel::Layer(50))
+						.Padding(ResizeBarPadding)
+						.HAlign(HAlign_Left)
+						.VAlign(VAlign_Top)
+					[
+						// Stretch box
+						SAssignNew(StretchBox, SSequencerStretchBox, SequencerPtr.Pin().ToSharedRef(), *Settings, NumericTypeInterface.ToSharedRef())
+					]
+
+					+ SGridPanel::Slot(Column1, Row2, SGridPanel::Layer(60))
 					.Padding(ResizeBarPadding)
 					[
 						SAssignNew(TickResolutionOverlay, SSequencerTimePanel, SequencerPtr)
@@ -645,7 +655,12 @@ void SSequencer::BindCommands(TSharedRef<FUICommandList> SequencerCommandBinding
 
 	SequencerCommandBindings->MapAction(
 		FSequencerCommands::Get().ToggleShowTransformBox,
-		FExecuteAction::CreateLambda([this]{ TransformBox->ToggleVisibility(); })
+		FExecuteAction::CreateLambda([this] { TransformBox->ToggleVisibility(); })
+	);
+
+	SequencerCommandBindings->MapAction(
+		FSequencerCommands::Get().ToggleShowStretchBox,
+		FExecuteAction::CreateLambda([this] { StretchBox->ToggleVisibility(); })
 	);
 }
 
@@ -1334,7 +1349,8 @@ TSharedRef<SWidget> SSequencer::MakeSelectEditMenu()
 	TSharedPtr<FSequencer> Sequencer = SequencerPtr.Pin();
 
 	MenuBuilder.AddMenuEntry(FSequencerCommands::Get().ToggleShowTransformBox);
-	
+	MenuBuilder.AddMenuEntry(FSequencerCommands::Get().ToggleShowStretchBox);
+
 	if (SequencerPtr.Pin()->IsLevelEditorSequencer())
 	{
 		MenuBuilder.AddMenuEntry(FSequencerCommands::Get().BakeTransform);
diff --git a/Engine/Source/Editor/Sequencer/Private/SSequencer.h b/Engine/Source/Editor/Sequencer/Private/SSequencer.h
--- a/Engine/Source/Editor/Sequencer/Private/SSequencer.h
+++ b/Engine/Source/Editor/Sequencer/Private/SSequencer.h
@@ -28,6 +28,7 @@ class SSequencerLabelBrowser;
 class SSequencerTrackArea;
 class SSequencerTrackOutliner;
 class SSequencerTransformBox;
+class SSequencerStretchBox;
 class SSequencerTreeView;
 class USequencerSettings;
 struct FPaintPlaybackRangeArgs;
@@ -483,9 +484,13 @@ public:
 	/** This adds the specified path to the selection set to be restored the next time the tree view is refreshed. */
 	void AddAdditionalPathToSelectionSet(const FString& Path) { AdditionalSelectionsToAdd.Add(Path); }
 private:
+
 	/** Transform box widget. */
 	TSharedPtr<SSequencerTransformBox> TransformBox;
 
+	/** Stretch box widget. */
+	TSharedPtr<SSequencerStretchBox> StretchBox;
+
 	/** Section area widget */
 	TSharedPtr<SSequencerTrackArea> TrackArea;
 
diff --git a/Engine/Source/Editor/Sequencer/Private/SSequencerStretchBox.cpp b/Engine/Source/Editor/Sequencer/Private/SSequencerStretchBox.cpp
new file mode 100644
--- /dev/null
+++ b/Engine/Source/Editor/Sequencer/Private/SSequencerStretchBox.cpp
@@ -0,0 +1,182 @@
+// Copyright 1998-2019 Epic Games, Inc. All Rights Reserved.
+
+#include "SSequencerStretchBox.h"
+#include "Sequencer.h"
+#include "SequencerSettings.h"
+#include "SequencerCommonHelpers.h"
+#include "EditorStyleSet.h"
+#include "EditorFontGlyphs.h"
+#include "Widgets/Input/SButton.h"
+#include "Widgets/Layout/SSpacer.h"
+
+#define LOCTEXT_NAMESPACE "Sequencer"
+
+void SSequencerStretchBox::Construct(const FArguments& InArgs, const TSharedRef<FSequencer>& InSequencer, USequencerSettings& InSettings, const TSharedRef<INumericTypeInterface<double>>& InNumericTypeInterface)
+{
+	SequencerPtr = InSequencer;
+	Settings = &InSettings;
+	NumericTypeInterface = InNumericTypeInterface;
+	
+	CachedStretchTime = FFrameNumber(0);
+	CachedShrinkTime = FFrameNumber(0);
+
+	const FDockTabStyle* GenericTabStyle = &FCoreStyle::Get().GetWidgetStyle<FDockTabStyle>("Docking.Tab");
+	const FButtonStyle* const CloseButtonStyle = &GenericTabStyle->CloseButtonStyle;
+
+	ChildSlot
+	[
+		SAssignNew(Border, SBorder)
+			.BorderImage(FEditorStyle::GetBrush("ToolPanel.DarkGroupBorder"))
+			.Padding(6.0f)
+			.Visibility(EVisibility::Collapsed)
+			[
+				SNew(SHorizontalBox)
+
+				+ SHorizontalBox::Slot()
+					.AutoWidth()
+					.VAlign(VAlign_Center)
+					[
+						SNew(SButton)
+							.ToolTipText(LOCTEXT("StretchForward_Tooltip", "Stretch time from the current time forwards"))
+							.OnClicked(this, &SSequencerStretchBox::OnStretchButtonClicked)
+							[
+								SNew(STextBlock)
+								.Font(FEditorStyle::Get().GetFontStyle("FontAwesome.14"))
+								.Text(FEditorFontGlyphs::Angle_Right)
+							]
+					]
+
+				+ SHorizontalBox::Slot()
+					.Padding(6.0f, 0.0f, 0.0f, 0.0f)
+					.AutoWidth()
+					[
+						SAssignNew(StretchEntryBox, SNumericEntryBox<double>)
+						.MinDesiredValueWidth(32.0f)
+						.TypeInterface(NumericTypeInterface)
+						.ToolTipText(LOCTEXT("StretchDelta_Tooltip", "The amount of time to stretch forward in front of the current time"))
+						.OnValueCommitted(this, &SSequencerStretchBox::OnStretchCommitted)
+						.OnValueChanged(this, &SSequencerStretchBox::OnStretchChanged)
+						.Value_Lambda([this](){ return CachedStretchTime.Value; })
+					]
+
+				+ SHorizontalBox::Slot()
+					.Padding(6.0f, 0.0f, 0.0f, 0.0f)
+					.AutoWidth()
+					.VAlign(VAlign_Center)
+					[
+						SNew(SButton)
+							.ToolTipText(LOCTEXT("ShrinkForward_Tooltip", "Shrink time in front of the current time"))
+							.OnClicked(this, &SSequencerStretchBox::OnShrinkButtonClicked)
+							[
+								SNew(STextBlock)
+								.Font(FEditorStyle::Get().GetFontStyle("FontAwesome.14"))
+								.Text(FEditorFontGlyphs::Angle_Left)
+							]
+					]
+
+				+ SHorizontalBox::Slot()
+					.Padding(6.0f, 0.0f, 0.0f, 0.0f)
+					.AutoWidth()
+					[
+						SAssignNew(ShrinkEntryBox, SNumericEntryBox<double>)
+						.MinDesiredValueWidth(32.0f)
+						.TypeInterface(NumericTypeInterface)
+						.ToolTipText(LOCTEXT("ShrinkDelta_Tooltip", "The amount of time to shrink in front of the current time"))
+						.OnValueCommitted(this, &SSequencerStretchBox::OnShrinkCommitted)
+						.OnValueChanged(this, &SSequencerStretchBox::OnShrinkChanged)
+						.Value_Lambda([this]() { return CachedShrinkTime.Value; })
+					]
+
+				+ SHorizontalBox::Slot()
+				.Padding(3.0f, 0.0f, 0.0f, 0.0f)
+				.AutoWidth()
+				.HAlign(HAlign_Center)
+				.VAlign(VAlign_Center)
+				[
+					SNew(SButton)
+					.ButtonStyle( CloseButtonStyle )
+					.OnClicked( this, &SSequencerStretchBox::OnCloseButtonClicked )
+					.ContentPadding( 0 )
+					[
+						SNew(SSpacer)
+						.Size( CloseButtonStyle->Normal.ImageSize )
+					]
+				]
+			]
+	];
+}
+
+
+void SSequencerStretchBox::ToggleVisibility()
+{
+	FSlateApplication& SlateApplication = FSlateApplication::Get();
+
+	if (Border->GetVisibility() == EVisibility::Visible)
+	{
+		if (LastFocusedWidget.IsValid())
+		{
+			SlateApplication.SetAllUserFocus(LastFocusedWidget.Pin(), EFocusCause::Navigation);
+		}
+		
+		Border->SetVisibility(EVisibility::Collapsed);
+	}
+	else
+	{
+		Border->SetVisibility(EVisibility::Visible);
+		LastFocusedWidget = SlateApplication.GetUserFocusedWidget(0);
+		SlateApplication.SetAllUserFocus(StretchEntryBox, EFocusCause::Navigation);
+	}
+}
+
+void SSequencerStretchBox::OnStretchCommitted(double Value, ETextCommit::Type CommitType)
+{
+	OnStretchChanged(Value);
+}
+
+void SSequencerStretchBox::OnStretchChanged(double Value)
+{
+	CachedStretchTime = FFrameTime::FromDecimal(Value).GetFrame();
+}
+
+void SSequencerStretchBox::OnShrinkCommitted(double Value, ETextCommit::Type CommitType)
+{
+	OnShrinkChanged(Value);
+}
+
+void SSequencerStretchBox::OnShrinkChanged(double Value)
+{
+	CachedShrinkTime = FFrameTime::FromDecimal(Value).GetFrame();
+}
+
+FReply SSequencerStretchBox::OnStretchButtonClicked()
+{
+	if (CachedStretchTime != 0)
+	{
+		TSharedPtr<FSequencer> Sequencer = SequencerPtr.Pin();
+
+		Sequencer->StretchTime(CachedStretchTime);
+	}
+
+	return FReply::Handled();
+}
+
+FReply SSequencerStretchBox::OnShrinkButtonClicked()
+{
+	if (CachedShrinkTime != 0)
+	{
+		TSharedPtr<FSequencer> Sequencer = SequencerPtr.Pin();
+		
+		Sequencer->ShrinkTime(CachedShrinkTime);
+	}
+
+	return FReply::Handled();
+}
+
+FReply SSequencerStretchBox::OnCloseButtonClicked()
+{
+	ToggleVisibility();
+
+	return FReply::Handled();
+}
+
+#undef LOCTEXT_NAMESPACE
diff --git a/Engine/Source/Editor/Sequencer/Private/SSequencerStretchBox.h b/Engine/Source/Editor/Sequencer/Private/SSequencerStretchBox.h
new file mode 100644
--- /dev/null
+++ b/Engine/Source/Editor/Sequencer/Private/SSequencerStretchBox.h
@@ -0,0 +1,71 @@
+// Copyright 1998-2019 Epic Games, Inc. All Rights Reserved.
+
+#pragma once
+
+#include "CoreMinimal.h"
+#include "Widgets/Input/NumericTypeInterface.h"
+#include "Widgets/DeclarativeSyntaxSupport.h"
+#include "Input/Reply.h"
+#include "Widgets/SCompoundWidget.h"
+#include "Widgets/Input/SNumericEntryBox.h"
+#include "Misc/QualifiedFrameTime.h"
+
+class FSequencer;
+class USequencerSettings;
+
+class SSequencerStretchBox
+	: public SCompoundWidget
+{
+public:
+
+	SLATE_BEGIN_ARGS(SSequencerStretchBox) { }
+	SLATE_END_ARGS()
+
+	/** Construct the widget. */
+	void Construct(const FArguments& InArgs, const TSharedRef<FSequencer>& InSequencer, USequencerSettings& InSettings, const TSharedRef<INumericTypeInterface<double>>& InNumericTypeInterface);
+
+	/** Toggle the widget's visibility. */
+	void ToggleVisibility();
+
+private:
+
+	/** Callback for transform operations. */
+	FReply OnStretchButtonClicked();
+	FReply OnShrinkButtonClicked();
+	void OnStretchCommitted(double Value, ETextCommit::Type CommitType);
+	void OnStretchChanged(double Value);
+	void OnShrinkCommitted(double Value, ETextCommit::Type CommitType);
+	void OnShrinkChanged(double Value);
+
+	/** Callback for when the close button is clicked. */
+	FReply OnCloseButtonClicked();
+
+private:
+
+	/** The border widget. */
+	TSharedPtr<SWidget> Border;
+
+	/** The stretch box widget. */
+	TSharedPtr<SNumericEntryBox<double>> StretchEntryBox;
+
+	/** The shrink entry box widget. */
+	TSharedPtr<SNumericEntryBox<double>> ShrinkEntryBox;
+
+	/** The widget that focused prior to this transform box. */
+	TWeakPtr<SWidget> LastFocusedWidget;
+
+	/** Numeric type interface used for converting parsing and generating strings from numbers. */
+	TSharedPtr<INumericTypeInterface<double>> NumericTypeInterface;
+
+	/** The main sequencer interface. */
+	TWeakPtr<FSequencer> SequencerPtr;
+
+	/** Cached settings provided to the sequencer itself on creation. */
+	USequencerSettings* Settings;
+
+	/** Cached stretch time. */
+	FFrameNumber CachedStretchTime;
+
+	/** Cached shrink time. */
+	FFrameNumber CachedShrinkTime;
+};
diff --git a/Engine/Source/Editor/Sequencer/Private/Sequencer.cpp b/Engine/Source/Editor/Sequencer/Private/Sequencer.cpp
--- a/Engine/Source/Editor/Sequencer/Private/Sequencer.cpp
+++ b/Engine/Source/Editor/Sequencer/Private/Sequencer.cpp
@@ -1454,6 +1454,118 @@ void FSequencer::TranslateSelectedKeysAndSections(bool bTranslateLeft)
 	TransformSelectedKeysAndSections(Delta, 1.f);
 }
 
+void FSequencer::StretchTime(FFrameTime InDeltaTime)
+{
+	// From the current time, find all the keys and sections to the right and move them by InDeltaTime
+	UMovieScene* FocusedMovieScene = GetFocusedMovieSceneSequence()->GetMovieScene();
+	if (FocusedMovieScene->IsReadOnly())
+	{
+		return;
+	}
+
+	FScopedTransaction StretchTimeTransaction(NSLOCTEXT("Sequencer", "StretchTime", "Stretch Time"));
+
+	TRange<FFrameNumber> CachedSelectionRange = GetSelectionRange();
+
+	TRange<FFrameNumber> SelectionRange;
+
+	if (InDeltaTime > 0)
+	{
+		SelectionRange.SetLowerBound(GetLocalTime().Time.FrameNumber+1);
+		SelectionRange.SetUpperBound(TRangeBound<FFrameNumber>::Open());
+	}
+	else
+	{
+		SelectionRange.SetUpperBound(GetLocalTime().Time.FrameNumber-1);
+		SelectionRange.SetLowerBound(TRangeBound<FFrameNumber>::Open());
+	}
+
+	FocusedMovieScene->SetSelectionRange(SelectionRange);
+	SelectInSelectionRange(true, true);
+	TransformSelectedKeysAndSections(InDeltaTime, 1.f);
+
+	// Return state
+	FocusedMovieScene->SetSelectionRange(CachedSelectionRange);
+	Selection.Empty(); //todo restore key and section selection
+}
+
+void FSequencer::ShrinkTime(FFrameTime InDeltaTime)
+{
+	// From the current time, find all the keys and sections to the right and move them by -InDeltaTime
+	UMovieScene* FocusedMovieScene = GetFocusedMovieSceneSequence()->GetMovieScene();
+	if (FocusedMovieScene->IsReadOnly())
+	{
+		return;
+	}
+
+	FScopedTransaction StretchTimeTransaction(NSLOCTEXT("Sequencer", "ShrinkTime", "Shrink Time"));
+
+	TRange<FFrameNumber> CachedSelectionRange = GetSelectionRange();
+
+	// First, check if there's any keys/sections within InDeltaTime
+
+	TRange<FFrameNumber> CheckRange;
+
+	if (InDeltaTime > 0)
+	{
+		CheckRange.SetLowerBound(GetLocalTime().Time.FrameNumber + 1);
+		CheckRange.SetUpperBound(GetLocalTime().Time.FrameNumber + InDeltaTime.FrameNumber);
+	}
+	else
+	{
+		CheckRange.SetUpperBound(GetLocalTime().Time.FrameNumber - InDeltaTime.FrameNumber);
+		CheckRange.SetLowerBound(GetLocalTime().Time.FrameNumber - 1);
+	}
+
+	FocusedMovieScene->SetSelectionRange(CheckRange);
+	SelectInSelectionRange(true, true);
+
+	if (Selection.GetSelectedKeys().Num() > 0)
+	{
+		FNotificationInfo Info(FText::Format(NSLOCTEXT("Sequencer", "ShrinkTimeFailedKeys", "Shrink failed. There are {0} keys in between"), Selection.GetSelectedKeys().Num()));
+		Info.ExpireDuration = 5.0f;
+		FSlateNotificationManager::Get().AddNotification(Info)->SetCompletionState(SNotificationItem::CS_Fail);
+
+		// Return state
+		FocusedMovieScene->SetSelectionRange(CachedSelectionRange);
+		Selection.Empty(); //todo restore key and section selection
+		return;
+	}
+
+	if (Selection.GetSelectedSections().Num() > 0)
+	{
+		FNotificationInfo Info(FText::Format(NSLOCTEXT("Sequencer", "ShrinkTimeFailedSections", "Shrink failed. There are {0} sections in between"), Selection.GetSelectedSections().Num()));
+		Info.ExpireDuration = 5.0f;
+		FSlateNotificationManager::Get().AddNotification(Info)->SetCompletionState(SNotificationItem::CS_Fail);
+
+		// Return state
+		FocusedMovieScene->SetSelectionRange(CachedSelectionRange);
+		Selection.Empty(); //todo restore key and section selection
+		return;
+	}
+
+	TRange<FFrameNumber> SelectionRange;
+
+	if (InDeltaTime > 0)
+	{
+		SelectionRange.SetLowerBound(GetLocalTime().Time.FrameNumber + 1);
+		SelectionRange.SetUpperBound(TRangeBound<FFrameNumber>::Open());
+	}
+	else
+	{
+		SelectionRange.SetUpperBound(GetLocalTime().Time.FrameNumber - 1);
+		SelectionRange.SetLowerBound(TRangeBound<FFrameNumber>::Open());
+	}
+
+	FocusedMovieScene->SetSelectionRange(SelectionRange);
+	SelectInSelectionRange(true, true);
+	TransformSelectedKeysAndSections(-InDeltaTime, 1.f);
+
+	// Return state
+	FocusedMovieScene->SetSelectionRange(CachedSelectionRange);
+	Selection.Empty(); //todo restore key and section selection
+}
+
 void FSequencer::BakeTransform()
 {
 	UMovieScene* FocusedMovieScene = GetFocusedMovieSceneSequence()->GetMovieScene();
@@ -1894,7 +2006,7 @@ void FSequencer::SelectInSelectionRange(const TSharedRef<FSequencerDisplayNode>&
 
 			for (auto Section : OutSections)
 			{
-				if (Section.IsValid() && Section->GetRange().Overlaps(SelectionRange))
+				if (Section.IsValid() && Section->GetRange().Overlaps(SelectionRange) && Section->HasStartFrame() && Section->HasEndFrame())
 				{
 					Selection.AddToSelection(Section.Get());
 				}
diff --git a/Engine/Source/Editor/Sequencer/Private/Sequencer.h b/Engine/Source/Editor/Sequencer/Private/Sequencer.h
--- a/Engine/Source/Editor/Sequencer/Private/Sequencer.h
+++ b/Engine/Source/Editor/Sequencer/Private/Sequencer.h
@@ -250,6 +250,12 @@ public:
 	/** Translate the selected keys and section by the time snap interval */
 	void TranslateSelectedKeysAndSections(bool bTranslateLeft);
 
+	/** Stretch time*/
+	void StretchTime(FFrameTime InDeltaTime);
+
+	/** Shrink time*/
+	void ShrinkTime(FFrameTime InDeltaTime);
+
 	/** Bake transform */
 	void BakeTransform();
 
diff --git a/Engine/Source/Editor/Sequencer/Private/SequencerCommands.cpp b/Engine/Source/Editor/Sequencer/Private/SequencerCommands.cpp
--- a/Engine/Source/Editor/Sequencer/Private/SequencerCommands.cpp
+++ b/Engine/Source/Editor/Sequencer/Private/SequencerCommands.cpp
@@ -91,7 +91,8 @@ void FSequencerCommands::RegisterCommands()
 
 	UI_COMMAND( ChangeTimeDisplayFormat, "Change Time Display Format", "Rotates through supported display formats for time", EUserInterfaceActionType::Button, FInputChord(EModifierKey::Control, EKeys::T) );
 	UI_COMMAND( ToggleShowGotoBox, "Go to Time...", "Go to a particular point on the timeline", EUserInterfaceActionType::Button, FInputChord(EModifierKey::Control, EKeys::G) );
-	UI_COMMAND( ToggleShowTransformBox, "Transform Selection...", "Transform the selected keys and sections by a given amount", EUserInterfaceActionType::Button, FInputChord(EModifierKey::Control, EKeys::M) );
+	UI_COMMAND( ToggleShowTransformBox, "Transform Selection...", "Transform the selected keys and sections by a given amount", EUserInterfaceActionType::Button, FInputChord(EModifierKey::Control, EKeys::M));
+	UI_COMMAND( ToggleShowStretchBox, "Stretch/Shrink...", "Stretch or shrink a given amount, moving keys forwards/backwards as necessary", EUserInterfaceActionType::Button, FInputChord() );
 	UI_COMMAND( BakeTransform, "Bake Transform", "Bake transform in world space, removing any existing transform and attach tracks", EUserInterfaceActionType::Button, FInputChord() );
 	UI_COMMAND( SyncSectionsUsingSourceTimecode, "Sync Sections using Source Timecode", "Synchronize sections to the first selected section using source timecode", EUserInterfaceActionType::Button, FInputChord() );
 	UI_COMMAND( ToggleShowRangeSlider, "Show Range Slider", "Enables and disables showing the time range slider", EUserInterfaceActionType::ToggleButton, FInputChord() );
diff --git a/Engine/Source/Editor/Sequencer/Public/SequencerCommands.h b/Engine/Source/Editor/Sequencer/Public/SequencerCommands.h
--- a/Engine/Source/Editor/Sequencer/Public/SequencerCommands.h
+++ b/Engine/Source/Editor/Sequencer/Public/SequencerCommands.h
@@ -223,6 +223,9 @@ public:
 	/** Toggle the visibility of the transform box. */
 	TSharedPtr< FUICommandInfo > ToggleShowTransformBox;
 
+	/** Toggle the visibility of the stretch box. */
+	TSharedPtr< FUICommandInfo > ToggleShowStretchBox;
+
 	/** Bake transform. */
 	TSharedPtr< FUICommandInfo > BakeTransform;
 
