Don't allow keyboard navigation in SequencerTreeView to select nodes that are not selectable. Prevents selecting non-visible spacing nodes.

diff --git a/Engine/Source/Editor/Sequencer/Private/SSequencerTreeView.cpp b/Engine/Source/Editor/Sequencer/Private/SSequencerTreeView.cpp
--- a/Engine/Source/Editor/Sequencer/Private/SSequencerTreeView.cpp
+++ b/Engine/Source/Editor/Sequencer/Private/SSequencerTreeView.cpp
@@ -489,6 +489,188 @@ void SSequencerTreeView::OnRightMouseButtonUp(const FPointerEvent& MouseEvent)
 	bRightMouseButtonDown = false;
 }
 
+FReply SSequencerTreeView::OnKeyDown(const FGeometry& MyGeometry, const FKeyEvent& InKeyEvent)
+{
+	const TArray<FDisplayNodeRef>& ItemsSourceRef = (*this->ItemsSource);
+
+	// Don't respond to key-presses containing "Alt" as a modifier
+	if (ItemsSourceRef.Num() > 0 && !InKeyEvent.IsAltDown())
+	{
+		bool bWasHandled = false;
+		NullableItemType ItemNavigatedTo(nullptr);
+
+		// Check for selection manipulation keys (Up, Down, Home, End, PageUp, PageDown)
+		if (InKeyEvent.GetKey() == EKeys::Up)
+		{
+			int32 SelectionIndex = 0;
+			if (TListTypeTraits<FDisplayNodeRef>::IsPtrValid(SelectorItem))
+			{
+				SelectionIndex = ItemsSourceRef.Find(TListTypeTraits<FDisplayNodeRef>::NullableItemTypeConvertToItemType(SelectorItem));
+			}
+
+			--SelectionIndex;
+
+			for (; SelectionIndex >=0; --SelectionIndex)
+			{
+				if (ItemsSourceRef[SelectionIndex]->IsSelectable())
+				{
+					ItemNavigatedTo = ItemsSourceRef[SelectionIndex];
+					break;
+				}
+			}
+			bWasHandled = true;
+		}
+		else if (InKeyEvent.GetKey() == EKeys::Down)
+		{
+			int32 SelectionIndex = 0;
+			if (TListTypeTraits<FDisplayNodeRef>::IsPtrValid(SelectorItem))
+			{
+				SelectionIndex = ItemsSourceRef.Find(TListTypeTraits<FDisplayNodeRef>::NullableItemTypeConvertToItemType(SelectorItem));
+			}
+
+			++SelectionIndex;
+
+			for (; SelectionIndex < ItemsSourceRef.Num(); ++SelectionIndex)
+			{
+				if (ItemsSourceRef[SelectionIndex]->IsSelectable())
+				{
+					ItemNavigatedTo = ItemsSourceRef[SelectionIndex];
+					break;
+				}
+			}
+			bWasHandled = true;
+		}
+		else if (InKeyEvent.GetKey() == EKeys::Home)
+		{
+			// Select the first item
+			for (int32 SelectionIndex = 0; SelectionIndex < ItemsSourceRef.Num(); ++SelectionIndex)
+			{
+				if (ItemsSourceRef[SelectionIndex]->IsSelectable())
+				{
+					ItemNavigatedTo = ItemsSourceRef[SelectionIndex];
+					break;
+				}
+			}
+			bWasHandled = true;
+		}
+		else if (InKeyEvent.GetKey() == EKeys::End)
+		{
+			// Select the last item
+			for (int32 SelectionIndex = ItemsSourceRef.Num() -1; SelectionIndex >=0 ; --SelectionIndex)
+			{
+				if (ItemsSourceRef[SelectionIndex]->IsSelectable())
+				{
+					ItemNavigatedTo = ItemsSourceRef[SelectionIndex];
+					break;
+				}
+			}
+			bWasHandled = true;
+		}
+		else if (InKeyEvent.GetKey() == EKeys::PageUp)
+		{
+			int32 SelectionIndex = 0;
+			if (TListTypeTraits<FDisplayNodeRef>::IsPtrValid(SelectorItem))
+			{
+				SelectionIndex = ItemsSourceRef.Find(TListTypeTraits<FDisplayNodeRef>::NullableItemTypeConvertToItemType(SelectorItem));
+			}
+
+			int32 NumItemsInAPage = GetNumLiveWidgets();
+			int32 Remainder = NumItemsInAPage % GetNumItemsWide();
+			NumItemsInAPage -= Remainder;
+
+			if (SelectionIndex >= NumItemsInAPage)
+			{
+				// Select an item on the previous page
+				SelectionIndex = SelectionIndex - NumItemsInAPage;
+
+				// Scan up for the first selectable node
+				for (; SelectionIndex >= 0; --SelectionIndex)
+				{
+					if (ItemsSourceRef[SelectionIndex]->IsSelectable())
+					{
+						ItemNavigatedTo = ItemsSourceRef[SelectionIndex];
+						break;
+					}
+				}
+			}
+
+			// If we had less than a page to jump, or we haven't found a selectable node yet,
+			// scan back toward our current node until we find one.
+			if (!ItemNavigatedTo)
+			{
+				SelectionIndex = 0;
+				for (; SelectionIndex < ItemsSourceRef.Num(); ++SelectionIndex)
+				{
+					if (ItemsSourceRef[SelectionIndex]->IsSelectable())
+					{
+						ItemNavigatedTo = ItemsSourceRef[SelectionIndex];
+						break;
+					}
+				}
+			}
+
+			bWasHandled = true;
+		}
+		else if (InKeyEvent.GetKey() == EKeys::PageDown)
+		{
+			int32 SelectionIndex = 0;
+			if (TListTypeTraits<FDisplayNodeRef>::IsPtrValid(SelectorItem))
+			{
+				SelectionIndex = ItemsSourceRef.Find(TListTypeTraits<FDisplayNodeRef>::NullableItemTypeConvertToItemType(SelectorItem));
+			}
+
+			int32 NumItemsInAPage = GetNumLiveWidgets();
+			int32 Remainder = NumItemsInAPage % GetNumItemsWide();
+			NumItemsInAPage -= Remainder;
+
+
+			if (SelectionIndex < ItemsSourceRef.Num() - NumItemsInAPage)
+			{
+				// Select an item on the next page
+				SelectionIndex = SelectionIndex + NumItemsInAPage;
+
+				for (; SelectionIndex < ItemsSourceRef.Num(); ++SelectionIndex)
+				{
+					if (ItemsSourceRef[SelectionIndex]->IsSelectable())
+					{
+						ItemNavigatedTo = ItemsSourceRef[SelectionIndex];
+						break;
+					}
+				}
+			}
+
+			// If we had less than a page to jump, or we haven't found a selectable node yet,
+			// scan back toward our current node until we find one.
+			if (!ItemNavigatedTo)
+			{
+				SelectionIndex = ItemsSourceRef.Num() - 1;
+				for (; SelectionIndex >= 0; --SelectionIndex)
+				{
+					if (ItemsSourceRef[SelectionIndex]->IsSelectable())
+					{
+						ItemNavigatedTo = ItemsSourceRef[SelectionIndex];
+						break;
+					}
+				}
+			}
+			bWasHandled = true;
+		}
+
+		if (TListTypeTraits<FDisplayNodeRef>::IsPtrValid(ItemNavigatedTo))
+		{
+			FDisplayNodeRef ItemToSelect(TListTypeTraits<FDisplayNodeRef>::NullableItemTypeConvertToItemType(ItemNavigatedTo));
+			NavigationSelect(ItemToSelect, InKeyEvent);
+		}
+
+		if (bWasHandled)
+		{
+			return FReply::Handled();
+		}
+	}
+
+	return STreeView<FDisplayNodeRef>::OnKeyDown(MyGeometry, InKeyEvent);
+}
+	
 
 void SSequencerTreeView::SynchronizeTreeSelectionWithSequencerSelection()
 {
diff --git a/Engine/Source/Editor/Sequencer/Private/SSequencerTreeView.h b/Engine/Source/Editor/Sequencer/Private/SSequencerTreeView.h
--- a/Engine/Source/Editor/Sequencer/Private/SSequencerTreeView.h
+++ b/Engine/Source/Editor/Sequencer/Private/SSequencerTreeView.h
@@ -130,6 +130,10 @@ protected:
 	virtual void OnRightMouseButtonDown(const FPointerEvent& MouseEvent) override;
 	virtual void OnRightMouseButtonUp(const FPointerEvent& MouseEvent) override;
 
+public:
+
+	virtual FReply OnKeyDown(const FGeometry& MyGeometry, const FKeyEvent& InKeyEvent) override;
+
 private:
 
 	/** Updates the tree selection to match the current sequencer selection. */
