Implement a duplicate option for sequencer keys

diff --git a/Engine/Source/Editor/Sequencer/Private/Sequencer.cpp b/Engine/Source/Editor/Sequencer/Private/Sequencer.cpp
--- a/Engine/Source/Editor/Sequencer/Private/Sequencer.cpp
+++ b/Engine/Source/Editor/Sequencer/Private/Sequencer.cpp
@@ -129,6 +129,7 @@
 #include "Kismet2/KismetEditorUtilities.h"
 #include "ISerializedRecorder.h"
 #include "Features/IModularFeatures.h"
+#include "SequencerContextMenus.h"
 
 #define LOCTEXT_NAMESPACE "Sequencer"
 
@@ -8526,8 +8527,47 @@ void FSequencer::CutSelection()
 
 void FSequencer::DuplicateSelection()
 {
+	FScopedTransaction DuplicateSelectionTransaction(LOCTEXT("DuplicateSelection_Transaction", "Duplicate Selection"));
+
 	CopySelection();
 	DoPaste();
+
+	// DoPaste doesn't handle keys, and we want duplicated keys to shift by one display rate frame as an overlapping key isn't useful.
+	
+	if (!Selection.GetSelectedKeys().Num())
+	{
+		return;
+	}
+
+	// Offset by a visible amount
+	FFrameNumber FrameOffset = FFrameNumber((int32)GetDisplayRateDeltaFrameCount());
+
+	TArray<FSequencerSelectedKey> NewSelection;
+	for (const FSequencerSelectedKey& Key : Selection.GetSelectedKeys())
+	{
+		if (Key.IsValid())
+		{
+			TSharedPtr<IKeyArea> KeyArea = Key.KeyArea;
+			FKeyHandle KeyHandle = Key.KeyHandle.GetValue();
+
+			FKeyHandle NewKeyHandle = KeyArea->DuplicateKey(KeyHandle);
+			KeyArea->SetKeyTime(NewKeyHandle, KeyArea->GetKeyTime(KeyHandle) + FrameOffset);
+
+			NewSelection.Add(FSequencerSelectedKey(*KeyArea->GetOwningSection(), KeyArea, NewKeyHandle));
+		}
+	}
+	
+	Selection.SuspendBroadcast();
+	Selection.EmptySelectedKeys();
+
+	for (const FSequencerSelectedKey& Key : NewSelection)
+	{
+		Selection.AddToSelection(Key);
+	}
+	Selection.ResumeBroadcast();
+	Selection.GetOnKeySelectionChanged().Broadcast();
+
+	NotifyMovieSceneDataChanged(EMovieSceneDataChangeType::TrackValueChanged);
 }
 
 void FSequencer::CopySelectedKeys()
@@ -9726,6 +9766,12 @@ void FSequencer::BindCommands()
 	};
 
 	auto CanDuplicate = [this]{
+
+		if (Selection.GetSelectedKeys().Num() || Selection.GetSelectedSections().Num())
+		{
+			return true;
+		}
+
 		// For duplicate object tracks
 		TSet<TSharedRef<FSequencerDisplayNode>> SelectedNodes = Selection.GetNodesWithSelectedKeysOrSections();
 		if (SelectedNodes.Num() == 0)
diff --git a/Engine/Source/Editor/Sequencer/Private/SequencerContextMenus.cpp b/Engine/Source/Editor/Sequencer/Private/SequencerContextMenus.cpp
--- a/Engine/Source/Editor/Sequencer/Private/SequencerContextMenus.cpp
+++ b/Engine/Source/Editor/Sequencer/Private/SequencerContextMenus.cpp
@@ -154,6 +154,7 @@ void FKeyContextMenu::PopulateMenu(FMenuBuilder& MenuBuilder)
 		{
 			MenuBuilder.AddMenuEntry(FGenericCommands::Get().Cut);
 			MenuBuilder.AddMenuEntry(FGenericCommands::Get().Copy);
+			MenuBuilder.AddMenuEntry(FGenericCommands::Get().Duplicate);
 		}
 	}
 	MenuBuilder.EndSection(); // SequencerKeyEdit
