Add option to group selected items in Sequencer in a new folder

diff --git a/Engine/Source/Editor/Sequencer/Private/DisplayNodes/SequencerDisplayNode.cpp b/Engine/Source/Editor/Sequencer/Private/DisplayNodes/SequencerDisplayNode.cpp
--- a/Engine/Source/Editor/Sequencer/Private/DisplayNodes/SequencerDisplayNode.cpp
+++ b/Engine/Source/Editor/Sequencer/Private/DisplayNodes/SequencerDisplayNode.cpp
@@ -810,6 +810,7 @@ void FSequencerDisplayNode::BuildContextMenu(FMenuBuilder& MenuBuilder)
 	MenuBuilder.EndSection();
 
 	TArray<UMovieSceneTrack*> AllTracks;
+	TArray<TSharedRef<FSequencerDisplayNode> > DragableNodes;
 	for (TSharedRef<FSequencerDisplayNode> Node : GetSequencer().GetSelection().GetSelectedOutlinerNodes())
 	{
 		if (Node->GetType() == ESequencerNode::Track)
@@ -820,8 +821,26 @@ void FSequencerDisplayNode::BuildContextMenu(FMenuBuilder& MenuBuilder)
 				AllTracks.Add(Track);
 			}
 		}
+
+		if (Node->CanDrag())
+		{
+			DragableNodes.Add(Node);
+		}
 	}
-	
+
+	MenuBuilder.BeginSection("Organize", LOCTEXT("OrganizeContextMenuSectionName", "Organize"));
+	{
+		if (DragableNodes.Num())
+		{
+			MenuBuilder.AddMenuEntry(
+				LOCTEXT("MoveTracksToNewFolder", "Move to New Folder"),
+				LOCTEXT("MoveTracksToNewFolderTooltip", "Move the selected tracks to a new folder."),
+				FSlateIcon(FEditorStyle::GetStyleSetName(), "ContentBrowser.AssetTreeFolderOpen"),
+				FUIAction(FExecuteAction::CreateSP(&GetSequencer(), &FSequencer::MoveSelectedNodesToNewFolder)));
+		}
+	}
+	MenuBuilder.EndSection();
+
 	if (AllTracks.Num())
 	{
 		MenuBuilder.BeginSection("GeneralTrackOptions", NSLOCTEXT("Sequencer", "TrackNodeGeneralOptions", "Track Options"));
diff --git a/Engine/Source/Editor/Sequencer/Private/Sequencer.cpp b/Engine/Source/Editor/Sequencer/Private/Sequencer.cpp
--- a/Engine/Source/Editor/Sequencer/Private/Sequencer.cpp
+++ b/Engine/Source/Editor/Sequencer/Private/Sequencer.cpp
@@ -6280,6 +6280,279 @@ void FSequencer::DeleteSelectedNodes()
 	}
 }
 
+void FSequencer::MoveNodeToFolder(TSharedRef<FSequencerDisplayNode> NodeToMove, UMovieSceneFolder* DestinationFolder)
+{
+	TSharedPtr<FSequencerDisplayNode> ParentNode = NodeToMove->GetParent();
+	
+	if (DestinationFolder == nullptr)
+	{
+		return;
+	}
+
+	DestinationFolder->Modify();
+
+	switch (NodeToMove->GetType())
+	{
+		case ESequencerNode::Folder:
+		{
+			TSharedRef<FSequencerFolderNode> FolderNode = StaticCastSharedRef<FSequencerFolderNode>(NodeToMove);
+			if (ParentNode.IsValid())
+			{
+				checkf(ParentNode->GetType() == ESequencerNode::Folder, TEXT("Can not remove from unsupported parent node."));
+				TSharedPtr<FSequencerFolderNode> NodeParentFolder = StaticCastSharedPtr<FSequencerFolderNode>(ParentNode);
+				NodeParentFolder->GetFolder().Modify();
+				NodeParentFolder->GetFolder().RemoveChildFolder(&FolderNode->GetFolder());
+			}
+			else
+			{
+				UMovieScene* FocusedMovieScene = GetFocusedMovieSceneSequence()->GetMovieScene();
+
+				FocusedMovieScene->Modify();
+				FocusedMovieScene->GetRootFolders().Remove(&FolderNode->GetFolder());
+			}
+
+			DestinationFolder->AddChildFolder(&FolderNode->GetFolder());
+
+			break;
+		}
+		case ESequencerNode::Track:
+		{
+			TSharedRef<FSequencerTrackNode> TrackNode = StaticCastSharedRef<FSequencerTrackNode>(NodeToMove);
+
+			if (ParentNode.IsValid())
+			{
+				checkf(ParentNode->GetType() == ESequencerNode::Folder, TEXT("Can not remove from unsupported parent node."));
+				TSharedPtr<FSequencerFolderNode> NodeParentFolder = StaticCastSharedPtr<FSequencerFolderNode>(ParentNode);
+				NodeParentFolder->GetFolder().Modify();
+				NodeParentFolder->GetFolder().RemoveChildMasterTrack(TrackNode->GetTrack());
+			}
+
+			DestinationFolder->AddChildMasterTrack(TrackNode->GetTrack());
+			break;
+		}
+		case ESequencerNode::Object:
+		{
+			TSharedRef<FSequencerObjectBindingNode> ObjectBindingNode = StaticCastSharedRef<FSequencerObjectBindingNode>(NodeToMove);
+			if (ParentNode.IsValid())
+			{
+				checkf(ParentNode->GetType() == ESequencerNode::Folder, TEXT("Can not remove from unsupported parent node."));
+				TSharedPtr<FSequencerFolderNode> NodeParentFolder = StaticCastSharedPtr<FSequencerFolderNode>(ParentNode);
+				NodeParentFolder->GetFolder().Modify();
+				NodeParentFolder->GetFolder().RemoveChildObjectBinding(ObjectBindingNode->GetObjectBinding());
+			}
+
+			DestinationFolder->AddChildObjectBinding(ObjectBindingNode->GetObjectBinding());
+			break;
+		}
+	}
+}
+
+TArray<TSharedRef<FSequencerDisplayNode> > FSequencer::GetSelectedNodesToMove()
+{
+	TArray<TSharedRef<FSequencerDisplayNode> > NodesToMove;
+
+	// Build a list of the nodes we want to move.
+	for (TSharedRef<FSequencerDisplayNode> Node : GetSelection().GetSelectedOutlinerNodes())
+	{
+		// Only nodes that can be dragged can be moved in to a folder. They must also either be in the root or in a folder.
+		if (Node->CanDrag() && (!Node->GetParent().IsValid() || Node->GetParent()->GetType() == ESequencerNode::Folder))
+		{
+			NodesToMove.Add(Node);
+		}
+	}
+
+	if (!NodesToMove.Num())
+	{
+		return NodesToMove;
+	}
+
+	TArray<int32> NodesToRemove;
+
+	// Find nodes that are children of other nodes in the list
+	for (int32 NodeIndex = 0; NodeIndex < NodesToMove.Num(); ++NodeIndex)
+	{
+		TSharedPtr<FSequencerDisplayNode> Node = NodesToMove[NodeIndex];
+
+		for (TSharedRef<FSequencerDisplayNode> ParentNode : NodesToMove)
+		{
+			if (ParentNode == Node)
+			{
+				continue;
+			}
+
+			if (!ParentNode->Traverse_ParentFirst([&](FSequencerDisplayNode& InNode) { return &InNode != Node.Get(); }))
+			{
+				NodesToRemove.Add(NodeIndex);
+			}
+		}
+	}
+
+	// Remove the nodes that are children of other nodes in the list, as moving the parent will already be relocating them
+	while (NodesToRemove.Num() > 0)
+	{
+		int32 NodeIndex = NodesToRemove.Pop();
+		NodesToMove.RemoveAt(NodeIndex);
+	}
+
+	return NodesToMove;
+}
+
+void FSequencer::MoveSelectedNodesToNewFolder()
+{
+	UMovieScene* FocusedMovieScene = GetFocusedMovieSceneSequence()->GetMovieScene();
+
+	if (FocusedMovieScene->IsReadOnly())
+	{
+		return;
+	}
+
+	TArray<TSharedRef<FSequencerDisplayNode> > NodesToMove = GetSelectedNodesToMove();
+
+	if (!NodesToMove.Num())
+	{
+		return;
+	}
+
+	TArray<TArray<FString> > NodePathSplits;
+	int32 SharedPathLength = TNumericLimits<int32>::Max();
+
+	// Build a list of the paths for each node, split in to folder names
+	for (TSharedRef<FSequencerDisplayNode> Node : NodesToMove)
+	{
+		// Split the node's path in to segments
+		TArray<FString>& NodePath = NodePathSplits.AddDefaulted_GetRef();
+		Node->GetPathName().ParseIntoArray(NodePath, TEXT("."));
+
+		// Shared path obviously won't be larger than the shortest path
+		SharedPathLength = FMath::Min(SharedPathLength, NodePath.Num() - 1);
+	}
+
+	// If we have more than one, find the deepest folder shared by all paths
+	if (NodePathSplits.Num() > 1)
+	{
+		// Since we are looking for the shared path, we can arbitrarily choose the first path to compare against
+		TArray<FString>& ShareNodePathSplit = NodePathSplits[0];
+		for (int NodeIndex = 1; NodeIndex < NodePathSplits.Num(); ++NodeIndex)
+		{
+			if (SharedPathLength == 0)
+			{
+				break;
+			}
+
+			// Since all paths are at least as long as the shortest, we don't need to bounds check the path splits
+			for (int PathSplitIndex = 0; PathSplitIndex < SharedPathLength; ++PathSplitIndex)
+			{
+				if (NodePathSplits[NodeIndex][PathSplitIndex].Compare(ShareNodePathSplit[PathSplitIndex]))
+				{
+					SharedPathLength = PathSplitIndex;
+					break;
+				}
+			}
+		}
+	}
+
+	UMovieSceneFolder* ParentFolder = nullptr;
+
+	TArray<FString> FolderPath;
+	
+	// Walk up the shared path to find the deepest shared folder
+	for (int32 FolderPathIndex = 0; FolderPathIndex < SharedPathLength; ++FolderPathIndex)
+	{
+		FolderPath.Add(NodePathSplits[0][FolderPathIndex]);
+		FName DesiredFolderName = FName(*FolderPath[FolderPathIndex]);
+
+		TArray<UMovieSceneFolder*> FoldersToSearch;
+		if (!ParentFolder)
+		{
+			FoldersToSearch = FocusedMovieScene->GetRootFolders();
+		}
+		else
+		{
+			FoldersToSearch = ParentFolder->GetChildFolders();
+		}
+
+		for (UMovieSceneFolder* Folder : FoldersToSearch)
+		{
+			if (Folder->GetFolderName() == DesiredFolderName)
+			{
+				ParentFolder = Folder;
+				break;
+			}
+		}
+	}
+
+	TArray<FName> ExistingFolderNames;
+	if (!ParentFolder)
+	{
+		for (UMovieSceneFolder* SiblingFolder : FocusedMovieScene->GetRootFolders())
+		{
+			ExistingFolderNames.Add(SiblingFolder->GetFolderName());
+		}
+	}
+	else
+	{
+		for (UMovieSceneFolder* SiblingFolder : ParentFolder->GetChildFolders())
+		{
+			ExistingFolderNames.Add(SiblingFolder->GetFolderName());
+		}
+	}
+
+	FString NewFolderPath;
+	for (FString PathSection : FolderPath)
+	{
+		NewFolderPath.Append(PathSection);
+		NewFolderPath.AppendChar('.');
+	}
+
+	FScopedTransaction Transaction(LOCTEXT("MoveTracksToNewFolder", "Move to New Folder"));
+
+	// Create SharedFolder
+	FName UniqueName = FSequencerUtilities::GetUniqueName(FName("New Folder"), ExistingFolderNames);
+	UMovieSceneFolder* SharedFolder = NewObject<UMovieSceneFolder>( FocusedMovieScene, NAME_None, RF_Transactional );
+	SharedFolder->SetFolderName(UniqueName);
+	NewFolderPath.Append(UniqueName.ToString());
+
+	FolderPath.Add(UniqueName.ToString());
+	int SharedFolderPathLen = FolderPath.Num();
+
+	if (!ParentFolder)
+	{
+		FocusedMovieScene->Modify();
+		FocusedMovieScene->GetRootFolders().Add(SharedFolder);
+	}
+	else
+	{
+		ParentFolder->Modify();
+		ParentFolder->AddChildFolder(SharedFolder);
+	}
+
+	for (int32 NodeIndex = 0; NodeIndex < NodesToMove.Num() ; ++NodeIndex)
+	{
+		TSharedRef<FSequencerDisplayNode> Node = NodesToMove[NodeIndex];
+		TArray<FString>& NodePathSplit = NodePathSplits[NodeIndex];
+
+		// Reset to just the path to the shared folder
+		FolderPath.SetNum(SharedFolderPathLen);
+
+		// Append any relative path for the node
+		for (int32 FolderPathIndex = SharedPathLength; FolderPathIndex < NodePathSplit.Num() - 1; ++FolderPathIndex)
+		{
+			FolderPath.Add(NodePathSplit[FolderPathIndex]);
+		}
+
+		UMovieSceneFolder* DestinationFolder = CreateFoldersRecursively(FolderPath, 0, FocusedMovieScene, nullptr, FocusedMovieScene->GetRootFolders());
+		
+		MoveNodeToFolder(Node, DestinationFolder);
+	}
+
+	// Set the newly created folder as our selection
+	Selection.Empty();
+	SequencerWidget->AddAdditionalPathToSelectionSet(NewFolderPath);
+
+
+	NotifyMovieSceneDataChanged(EMovieSceneDataChangeType::MovieSceneStructureItemsChanged);
+}
+
 
 void ExportObjectBindingsToText(const TArray<UMovieSceneCopyableBinding*>& ObjectsToExport, FString& ExportedText)
 {
diff --git a/Engine/Source/Editor/Sequencer/Private/Sequencer.h b/Engine/Source/Editor/Sequencer/Private/Sequencer.h
--- a/Engine/Source/Editor/Sequencer/Private/Sequencer.h
+++ b/Engine/Source/Editor/Sequencer/Private/Sequencer.h
@@ -486,6 +486,13 @@ public:
 	void DeleteNode(TSharedRef<FSequencerDisplayNode> NodeToBeDeleted);
 	void DeleteSelectedNodes();
 
+	/** @return The list of nodes which must be moved to move the current selected nodes */
+	TArray<TSharedRef<FSequencerDisplayNode> > GetSelectedNodesToMove();
+
+	/** Called when a user executes the move to new folder menu item */
+	void MoveSelectedNodesToNewFolder();
+	void MoveNodeToFolder(TSharedRef<FSequencerDisplayNode> NodeToMove, UMovieSceneFolder* DestinationFolder);
+
 	/** Called when a user executes the copy track menu item */
 	void CopySelectedObjects(TArray<TSharedPtr<FSequencerObjectBindingNode>>& ObjectNodes, /*out*/ FString& ExportedText);
 	void CopySelectedTracks(TArray<TSharedPtr<FSequencerTrackNode>>& TrackNodes, /*out*/ FString& ExportedText);
