Sequencer: Allow adding tracks for all object classes, not restricted to selected object binding class

diff --git a/Engine/Source/Editor/Sequencer/Private/DisplayNodes/SequencerObjectBindingNode.cpp b/Engine/Source/Editor/Sequencer/Private/DisplayNodes/SequencerObjectBindingNode.cpp
--- a/Engine/Source/Editor/Sequencer/Private/DisplayNodes/SequencerObjectBindingNode.cpp
+++ b/Engine/Source/Editor/Sequencer/Private/DisplayNodes/SequencerObjectBindingNode.cpp
@@ -901,7 +901,7 @@ TSharedRef<SWidget> FSequencerObjectBindingNode::HandleAddTrackComboButtonGetMen
 		FGuid Guid = ObjectBindingNode->GetObjectBinding();
 		for (auto RuntimeObject : Sequencer.FindBoundObjects(Guid, Sequencer.GetFocusedTemplateID()))
 		{
-			if (RuntimeObject->GetClass() == ObjectClass)
+			////if (RuntimeObject->GetClass() == ObjectClass)
 			{
 				ObjectBindings.AddUnique(Guid);
 				continue;
