Sequencer: Expand filtered nodes

diff --git a/Engine/Source/Editor/Sequencer/Private/SSequencer.cpp b/Engine/Source/Editor/Sequencer/Private/SSequencer.cpp
--- a/Engine/Source/Editor/Sequencer/Private/SSequencer.cpp
+++ b/Engine/Source/Editor/Sequencer/Private/SSequencer.cpp
@@ -1798,6 +1798,9 @@ void SSequencer::OnOutlinerSearchChanged( const FText& Filter )
 		const FString FilterString = Filter.ToString();
 
 		Sequencer->GetNodeTree()->FilterNodes( FilterString );
+
+		Sequencer->GetNodeTree()->Update();
+
 		TreeView->Refresh();
 
 		if ( FilterString.StartsWith( TEXT( "label:" ) ) )
diff --git a/Engine/Source/Editor/Sequencer/Private/SequencerNodeTree.cpp b/Engine/Source/Editor/Sequencer/Private/SequencerNodeTree.cpp
--- a/Engine/Source/Editor/Sequencer/Private/SequencerNodeTree.cpp
+++ b/Engine/Source/Editor/Sequencer/Private/SequencerNodeTree.cpp
@@ -577,9 +577,16 @@ bool FSequencerNodeTree::GetDefaultExpansionState( const FSequencerDisplayNode&
 }
 
 
-bool FSequencerNodeTree::IsNodeFiltered( const TSharedRef<const FSequencerDisplayNode> Node ) const
+bool FSequencerNodeTree::IsNodeFiltered(const TSharedRef<const FSequencerDisplayNode> Node) const
 {
-	return FilteredNodes.Contains( Node );
+	for (auto It = FilteredNodes.CreateConstIterator(); It; ++It)
+	{
+		if ((*It) == Node)
+		{
+			return true;
+		}
+	}
+	return false;
 }
 
 void FSequencerNodeTree::SetHoveredNode(const TSharedPtr<FSequencerDisplayNode>& InHoveredNode)
@@ -595,7 +602,7 @@ const TSharedPtr<FSequencerDisplayNode>& FSequencerNodeTree::GetHoveredNode() co
 	return HoveredNode;
 }
 
-static void AddChildNodes(const TSharedRef<FSequencerDisplayNode>& StartNode, TSet<TSharedRef<const FSequencerDisplayNode>>& OutFilteredNodes)
+static void AddChildNodes(const TSharedRef<FSequencerDisplayNode>& StartNode, TSet<TSharedRef<FSequencerDisplayNode>>& OutFilteredNodes)
 {
 	OutFilteredNodes.Add(StartNode);
 
@@ -608,24 +615,39 @@ static void AddChildNodes(const TSharedRef<FSequencerDisplayNode>& StartNode, TS
 /*
  * Add node as filtered and include any parent folders
  */
-static void AddFilteredNode(const TSharedRef<FSequencerDisplayNode>& StartNode, TSet<TSharedRef<const FSequencerDisplayNode>>& OutFilteredNodes)
+static void AddFilteredNode(const TSharedRef<FSequencerDisplayNode>& StartNode, TSet<TSharedRef<FSequencerDisplayNode>>& OutFilteredNodes)
 {
+	if (!StartNode->IsExpanded())
+	{
+		StartNode->SetExpansionState(true);
+	}
+
 	AddChildNodes(StartNode, OutFilteredNodes);
 
 	// Gather parent folders up the chain
 	TSharedPtr<FSequencerDisplayNode> ParentNode = StartNode->GetParent();
 	while (ParentNode.IsValid() && ParentNode.Get()->GetType() == ESequencerNode::Folder)
 	{
+		if (!ParentNode.Get()->IsExpanded())
+		{
+			ParentNode.Get()->SetExpansionState(true);
+		}
+
 		OutFilteredNodes.Add(ParentNode.ToSharedRef());
 		ParentNode = ParentNode->GetParent();
 	}
 }
 
-static void AddParentNodes(const TSharedRef<FSequencerDisplayNode>& StartNode, TSet<TSharedRef<const FSequencerDisplayNode>>& OutFilteredNodes)
+static void AddParentNodes(const TSharedRef<FSequencerDisplayNode>& StartNode, TSet<TSharedRef<FSequencerDisplayNode>>& OutFilteredNodes)
 {
 	TSharedPtr<FSequencerDisplayNode> ParentNode = StartNode->GetParent();
 	if (ParentNode.IsValid())
 	{
+		if (!ParentNode.Get()->IsExpanded())
+		{
+			ParentNode.Get()->SetExpansionState(true);
+		}
+
 		OutFilteredNodes.Add(ParentNode.ToSharedRef());
 		AddParentNodes(ParentNode.ToSharedRef(), OutFilteredNodes);
 	}
@@ -639,7 +661,7 @@ static void AddParentNodes(const TSharedRef<FSequencerDisplayNode>& StartNode, T
  * @param OutFilteredNodes	The list of all filtered nodes
  * @return Whether the text filter was passed
  */
-static bool FilterNodesRecursive( FSequencer& Sequencer, const TSharedRef<FSequencerDisplayNode>& StartNode, const TArray<FString>& FilterStrings, TSet<TSharedRef<const FSequencerDisplayNode>>& OutFilteredNodes )
+static bool FilterNodesRecursive( FSequencer& Sequencer, const TSharedRef<FSequencerDisplayNode>& StartNode, const TArray<FString>& FilterStrings, TSet<TSharedRef<FSequencerDisplayNode>>& OutFilteredNodes )
 {
 	// check labels - only one of the labels needs to match
 	bool bMatchedLabel = false;
diff --git a/Engine/Source/Editor/Sequencer/Private/SequencerNodeTree.h b/Engine/Source/Editor/Sequencer/Private/SequencerNodeTree.h
--- a/Engine/Source/Editor/Sequencer/Private/SequencerNodeTree.h
+++ b/Engine/Source/Editor/Sequencer/Private/SequencerNodeTree.h
@@ -164,7 +164,7 @@ private:
 	/** Mapping of object binding guids to their node (for fast lookup) */
 	TMap< FGuid, TSharedPtr<FSequencerObjectBindingNode> > ObjectBindingMap;
 	/** Set of all filtered nodes */
-	TSet< TSharedRef<const FSequencerDisplayNode> > FilteredNodes;
+	TSet< TSharedRef<FSequencerDisplayNode> > FilteredNodes;
 	/** Cardinal hovered node */
 	TSharedPtr<FSequencerDisplayNode> HoveredNode;
 	/** Active filter string if any */
