Sequencer: Allow easing handles only if there are supported blend types

diff --git a/Engine/Source/Editor/Sequencer/Private/SSequencerSection.cpp b/Engine/Source/Editor/Sequencer/Private/SSequencerSection.cpp
--- a/Engine/Source/Editor/Sequencer/Private/SSequencerSection.cpp
+++ b/Engine/Source/Editor/Sequencer/Private/SSequencerSection.cpp
@@ -850,6 +850,12 @@ bool SSequencerSection::CheckForEasingHandleInteraction( const FPointerEvent& Mo
 		return false;
 	}
 
+	UMovieSceneTrack* Track = ThisSection->GetTypedOuter<UMovieSceneTrack>();
+	if (!Track || Track->GetSupportedBlendTypes().Num() == 0)
+	{
+		return false;
+	}
+
 	FTimeToPixel TimeToPixelConverter = ConstructTimeConverterForSection(MakeSectionGeometryWithoutHandles(SectionGeometry, SectionInterface), *ThisSection, GetSequencer());
 
 	const double MouseTime = TimeToPixelConverter.PixelToSeconds(SectionGeometry.AbsoluteToLocal(MouseEvent.GetScreenSpacePosition()).X);
