Sequencer: SetLocalTime now always scrolls into view unless playing and autoscroll setting is off.

Tested with autoscroll setting off that when playing, it does not scroll and when jumping to next key, it does autoscroll to the key.

Tested that when typing in a frame number to jump to that is out of view, it autoscrolls.

diff --git a/Engine/Source/Editor/Sequencer/Private/Sequencer.cpp b/Engine/Source/Editor/Sequencer/Private/Sequencer.cpp
--- a/Engine/Source/Editor/Sequencer/Private/Sequencer.cpp
+++ b/Engine/Source/Editor/Sequencer/Private/Sequencer.cpp
@@ -2390,12 +2390,12 @@ FQualifiedFrameTime FSequencer::GetGlobalTime() const
 	return FQualifiedFrameTime(RootTime, PlayPosition.GetOutputRate());
 }
 
-void FSequencer::SetLocalTime( FFrameTime NewTime, ESnapTimeMode SnapTimeMode, bool bForceScroll )
+void FSequencer::SetLocalTime( FFrameTime NewTime, ESnapTimeMode SnapTimeMode)
 {
 	FFrameRate LocalResolution = GetFocusedTickResolution();
 
 	// Ensure the time is in the current view
-	if (IsAutoScrollEnabled() || bForceScroll)
+	if (IsAutoScrollEnabled() || GetPlaybackStatus() != EMovieScenePlayerStatus::Playing)
 	{
 		ScrollIntoView(NewTime / LocalResolution);
 	}
@@ -3686,7 +3686,7 @@ FReply FSequencer::OnStepForward()
 	FQualifiedFrameTime CurrentTime = GetLocalTime();
 
 	FFrameTime NewPosition = FFrameRate::TransformTime(CurrentTime.ConvertTo(DisplayRate).FloorToFrame() + 1, DisplayRate, CurrentTime.Rate);
-	SetLocalTime(NewPosition, ESnapTimeMode::STM_Interval, true);
+	SetLocalTime(NewPosition, ESnapTimeMode::STM_Interval);
 	return FReply::Handled();
 }
 
@@ -3700,7 +3700,7 @@ FReply FSequencer::OnStepBackward()
 
 	FFrameTime NewPosition = FFrameRate::TransformTime(CurrentTime.ConvertTo(DisplayRate).FloorToFrame() - 1, DisplayRate, CurrentTime.Rate);
 
-	SetLocalTime(NewPosition, ESnapTimeMode::STM_Interval, true);
+	SetLocalTime(NewPosition, ESnapTimeMode::STM_Interval);
 	return FReply::Handled();
 }
 
@@ -3708,7 +3708,7 @@ FReply FSequencer::OnStepBackward()
 FReply FSequencer::OnJumpToStart()
 {
 	SetPlaybackStatus(EMovieScenePlayerStatus::Stepping);
-	SetLocalTime(MovieScene::DiscreteInclusiveLower(GetPlaybackRange()), ESnapTimeMode::STM_None, true);
+	SetLocalTime(MovieScene::DiscreteInclusiveLower(GetPlaybackRange()), ESnapTimeMode::STM_None);
 	return FReply::Handled();
 }
 
@@ -3726,7 +3726,7 @@ FReply FSequencer::OnJumpToEnd()
 	FFrameTime OneFrame = bInsetDisplayFrame ? FFrameRate::TransformTime(FFrameTime(1), DisplayRate, LocalResolution) : FFrameTime(1);
 	FFrameTime NewTime = MovieScene::DiscreteExclusiveUpper(GetPlaybackRange()) - OneFrame;
 
-	SetLocalTime(NewTime, ESnapTimeMode::STM_None, true);
+	SetLocalTime(NewTime, ESnapTimeMode::STM_None);
 	return FReply::Handled();
 }
 
@@ -8278,7 +8278,7 @@ void FSequencer::StepToNextShot()
 	PopToSequenceInstance(ActiveTemplateIDs[ActiveTemplateIDs.Num()-2]);
 	FocusSequenceInstance(*NextShot);
 
-	SetLocalTime(FFrameTime(0), ESnapTimeMode::STM_None, true);
+	SetLocalTime(FFrameTime(0), ESnapTimeMode::STM_None);
 }
 
 
@@ -8310,7 +8310,7 @@ void FSequencer::StepToPreviousShot()
 	PopToSequenceInstance(ActiveTemplateIDs[ActiveTemplateIDs.Num()-2]);
 	FocusSequenceInstance(*PreviousShot);
 
-	SetLocalTime(0, ESnapTimeMode::STM_None, true);
+	SetLocalTime(0, ESnapTimeMode::STM_None);
 }
 
 
diff --git a/Engine/Source/Editor/Sequencer/Private/Sequencer.h b/Engine/Source/Editor/Sequencer/Private/Sequencer.h
--- a/Engine/Source/Editor/Sequencer/Private/Sequencer.h
+++ b/Engine/Source/Editor/Sequencer/Private/Sequencer.h
@@ -691,7 +691,7 @@ public:
 	virtual bool GetAutoSetTrackDefaults() const override;
 	virtual FQualifiedFrameTime GetLocalTime() const override;
 	virtual FQualifiedFrameTime GetGlobalTime() const override;
-	virtual void SetLocalTime(FFrameTime Time, ESnapTimeMode SnapTimeMode = ESnapTimeMode::STM_None, bool bForceScroll = false) override;
+	virtual void SetLocalTime(FFrameTime Time, ESnapTimeMode SnapTimeMode = ESnapTimeMode::STM_None) override;
 	virtual void SetLocalTimeDirectly(FFrameTime NewTime) override;
 	virtual void SetGlobalTime(FFrameTime Time) override;
 	virtual void ForceEvaluate() override;
diff --git a/Engine/Source/Editor/Sequencer/Public/ISequencer.h b/Engine/Source/Editor/Sequencer/Public/ISequencer.h
--- a/Engine/Source/Editor/Sequencer/Public/ISequencer.h
+++ b/Engine/Source/Editor/Sequencer/Public/ISequencer.h
@@ -323,9 +323,8 @@ public:
 	 *
 	 * @param Time The local time to set.
 	 * @param SnapTimeMode The type of time snapping allowed.
-	 * @param ForceScroll Whether to ignore autoscroll setting and always scroll view to new location
 	 */
-	virtual void SetLocalTime(FFrameTime Time, ESnapTimeMode SnapTimeMode = ESnapTimeMode::STM_None, bool bForceScroll = false) = 0;
+	virtual void SetLocalTime(FFrameTime Time, ESnapTimeMode SnapTimeMode = ESnapTimeMode::STM_None) = 0;
 
 	/** Set the current local time directly, with no other snapping, scrolling or manipulation */
 	virtual void SetLocalTimeDirectly(FFrameTime NewTime) = 0;
diff --git a/Engine/Source/Editor/Sequencer/Public/SequencerSettings.h b/Engine/Source/Editor/Sequencer/Public/SequencerSettings.h
--- a/Engine/Source/Editor/Sequencer/Public/SequencerSettings.h
+++ b/Engine/Source/Editor/Sequencer/Public/SequencerSettings.h
@@ -229,9 +229,9 @@ public:
 	/** Set zoom in/out position (mouse position or current time). */
 	void SetZoomPosition(ESequencerZoomPosition InZoomPosition);
 
-	/** Gets whether or not auto-scroll is enabled. */
+	/** Gets whether or not auto-scroll is enabled when playing. */
 	bool GetAutoScrollEnabled() const;
-	/** Sets whether or not auto-scroll is enabled. */
+	/** Sets whether or not auto-scroll is enabled when playing. */
 	void SetAutoScrollEnabled(bool bInAutoScrollEnabled);
 	
 	/** Gets whether or not to link the curve editor time range. */
@@ -436,7 +436,7 @@ protected:
 	UPROPERTY( config, EditAnywhere, Category=Timeline )
 	TEnumAsByte<ESequencerZoomPosition> ZoomPosition;
 
-	/** Enable or disable auto scroll in the timeline. */
+	/** Enable or disable auto scroll in the timeline when playing. */
 	UPROPERTY( config, EditAnywhere, Category=Timeline )
 	bool bAutoScrollEnabled;
 
