Sequencer: Fix crash getting browser to selected actor sub menu. Allow navigating to actor in root sequence.

diff --git a/Engine/Source/Editor/Sequencer/Private/LevelEditorSequencerIntegration.cpp b/Engine/Source/Editor/Sequencer/Private/LevelEditorSequencerIntegration.cpp
--- a/Engine/Source/Editor/Sequencer/Private/LevelEditorSequencerIntegration.cpp
+++ b/Engine/Source/Editor/Sequencer/Private/LevelEditorSequencerIntegration.cpp
@@ -730,9 +730,23 @@ void FLevelEditorSequencerIntegration::MakeBrowseToSelectedActorSubMenu(FMenuBui
 {
 	for (const TPair<FMovieSceneSequenceID, FSequencer*> Sequence : FoundInSequences)
 	{
-		FText ActorName = FText::Format(LOCTEXT("ActorNameSingular", "\"{0}\""), FText::FromString(Actor->GetActorLabel()));
-		FUIAction AddMenuAction(FExecuteAction::CreateLambda([this, Actor, Sequence]() {this->BrowseToSelectedActor(Actor, Sequence.Value, Sequence.Key); }));
-		MenuBuilder.AddMenuEntry(FText::Format(LOCTEXT("BrowseToSelectedActorText", "Browse to {0} in {1}"), ActorName, Sequence.Value->FindSubSection(Sequence.Key)->GetSequence()->GetDisplayName()), FText(), FSlateIcon(), AddMenuAction);
+		UMovieSceneSequence* MovieSceneSequence = nullptr;
+		if (Sequence.Key == MovieSceneSequenceID::Root)
+		{
+			MovieSceneSequence = Sequence.Value->GetRootMovieSceneSequence();
+		}
+		else
+		{
+			UMovieSceneSubSection* SubSection = Sequence.Value->FindSubSection(Sequence.Key);
+			MovieSceneSequence = SubSection ? SubSection->GetSequence() : nullptr;
+		}
+		
+		if (MovieSceneSequence)
+		{
+			FText ActorName = FText::Format(LOCTEXT("ActorNameSingular", "\"{0}\""), FText::FromString(Actor->GetActorLabel()));
+			FUIAction AddMenuAction(FExecuteAction::CreateLambda([this, Actor, Sequence]() {this->BrowseToSelectedActor(Actor, Sequence.Value, Sequence.Key); }));
+			MenuBuilder.AddMenuEntry(FText::Format(LOCTEXT("BrowseToSelectedActorText", "Browse to {0} in {1}"), ActorName, MovieSceneSequence->GetDisplayName()), FText(), FSlateIcon(), AddMenuAction);
+		}
 	}
 }
 
@@ -847,7 +861,12 @@ void FLevelEditorSequencerIntegration::RecordSelectedActors()
 void FLevelEditorSequencerIntegration::BrowseToSelectedActor(AActor* Actor, FSequencer* Sequencer, FMovieSceneSequenceID SequenceID)
 {
 	Sequencer->PopToSequenceInstance(MovieSceneSequenceID::Root);
-	Sequencer->FocusSequenceInstance(*Sequencer->FindSubSection(SequenceID));
+
+	if (SequenceID != MovieSceneSequenceID::Root)
+	{
+		Sequencer->FocusSequenceInstance(*Sequencer->FindSubSection(SequenceID));
+	}
+
 	Sequencer->SelectObject(Sequencer->FindObjectId(*Actor, SequenceID));
 }
 
