Sequencer: Show error when sub display rates don't match. Changed error widgets to EVisibility::Visible so that their tooltips appear when hovered over.

diff --git a/Engine/Source/Editor/Sequencer/Private/SSequencerPlayRateCombo.cpp b/Engine/Source/Editor/Sequencer/Private/SSequencerPlayRateCombo.cpp
--- a/Engine/Source/Editor/Sequencer/Private/SSequencerPlayRateCombo.cpp
+++ b/Engine/Source/Editor/Sequencer/Private/SSequencerPlayRateCombo.cpp
@@ -70,12 +70,25 @@ void SSequencerPlayRateCombo::Construct(const FArguments& InArgs, TWeakPtr<FSequ
 			.VAlign(VAlign_Center)
 			[
 				SNew(STextBlock)
-				.ToolTipText(this, &SSequencerPlayRateCombo::GetFrameRateErrorDescription)
-				.Visibility(this, &SSequencerPlayRateCombo::GetFrameRateErrorVisibility)
+				.ToolTipText(this, &SSequencerPlayRateCombo::GetFrameRateIsMultipleOfErrorDescription)
+				.Visibility(this, &SSequencerPlayRateCombo::GetFrameRateIsMultipleOfErrorVisibility)
 				.TextStyle( InArgs._StyleSet, ISlateStyle::Join( InArgs._StyleName, ".Label" ) )
 				.Font(FEditorStyle::Get().GetFontStyle("FontAwesome.11"))
 				.Text(FEditorFontGlyphs::Exclamation_Triangle)
 			]
+
+			+ SHorizontalBox::Slot()
+			.Padding(0.f, 0.f, 3.f, 0.f)
+			.AutoWidth()
+			.VAlign(VAlign_Center)
+			[
+				SNew(STextBlock)
+				.ToolTipText(this, &SSequencerPlayRateCombo::GetFrameRateMismatchErrorDescription)
+				.Visibility(this, &SSequencerPlayRateCombo::GetFrameRateMismatchErrorVisibility)
+				.TextStyle(InArgs._StyleSet, ISlateStyle::Join(InArgs._StyleName, ".Label"))
+				.Font(FEditorStyle::Get().GetFontStyle("FontAwesome.11"))
+				.Text(FEditorFontGlyphs::Exclamation_Triangle)
+			]
 		]
 	];
 
@@ -89,17 +102,44 @@ EVisibility SSequencerPlayRateCombo::GetFrameLockedVisibility() const
 	return Sequence && Sequence->GetMovieScene()->GetEvaluationType() == EMovieSceneEvaluationType::FrameLocked ? EVisibility::HitTestInvisible : EVisibility::Collapsed;
 }
 
-EVisibility SSequencerPlayRateCombo::GetFrameRateErrorVisibility() const
+EVisibility SSequencerPlayRateCombo::GetFrameRateIsMultipleOfErrorVisibility() const
 {
 	TSharedPtr<FSequencer> Sequencer  = WeakSequencer.Pin();
 	if (!Sequencer.IsValid() || Sequencer->GetFocusedDisplayRate().IsMultipleOf(Sequencer->GetFocusedTickResolution()))
 	{
 		return EVisibility::Collapsed;
 	}
-	return EVisibility::HitTestInvisible;
+	return EVisibility::Visible;
 }
 
-FText SSequencerPlayRateCombo::GetFrameRateErrorDescription() const
+EVisibility SSequencerPlayRateCombo::GetFrameRateMismatchErrorVisibility() const
+{
+	TSharedPtr<FSequencer> Sequencer = WeakSequencer.Pin();
+	if (Sequencer.IsValid())
+	{
+		FFrameRate DisplayRate = Sequencer->GetRootDisplayRate();
+
+		const FMovieSceneSequenceHierarchy& Hierarchy = Sequencer->GetEvaluationTemplate().GetHierarchy();
+		for (const TTuple<FMovieSceneSequenceID, FMovieSceneSubSequenceData>& Pair : Hierarchy.AllSubSequenceData())
+		{
+			if (UMovieSceneSequence* SubSequence = Pair.Value.GetSequence())
+			{
+				if (SubSequence->GetMovieScene())
+				{
+					FFrameRate SubDisplayRate = SubSequence->GetMovieScene()->GetDisplayRate();
+					if (SubDisplayRate != DisplayRate)
+					{
+						return EVisibility::Visible;
+					}
+				}
+			}
+		}
+	}
+
+	return EVisibility::Collapsed;
+}
+
+FText SSequencerPlayRateCombo::GetFrameRateIsMultipleOfErrorDescription() const
 {
 	TSharedPtr<FSequencer> Sequencer  = WeakSequencer.Pin();
 	if (Sequencer.IsValid())
@@ -113,7 +153,37 @@ FText SSequencerPlayRateCombo::GetFrameRateErrorDescription() const
 		FText DisplayRateText     = DisplayRateInfo     ? DisplayRateInfo->DisplayName    : FText::Format(LOCTEXT("DisplayRateFormat", "{0} fps"), DisplayRate.AsDecimal());
 		FText TickResolutionText  = TickResolutionInfo  ? TickResolutionInfo->DisplayName : FText::Format(LOCTEXT("TickResolutionFormat", "{0} ticks every second"), TickResolution.AsDecimal());
 
-		return FText::Format(LOCTEXT("FrameRateErrorDescription", "The current display rate of {0} is incompatible with this sequence's tick resolution of {1} ticks per second."), DisplayRateText, TickResolutionText);
+		return FText::Format(LOCTEXT("FrameRateIsMultipleOfErrorDescription", "The current display rate of {0} is incompatible with this sequence's tick resolution of {1} ticks per second."), DisplayRateText, TickResolutionText);
+	}
+
+	return FText();
+}
+
+FText SSequencerPlayRateCombo::GetFrameRateMismatchErrorDescription() const
+{
+	TSharedPtr<FSequencer> Sequencer = WeakSequencer.Pin();
+	if (Sequencer.IsValid())
+	{
+		FFrameRate DisplayRate = Sequencer->GetRootDisplayRate();
+
+		const FMovieSceneSequenceHierarchy& Hierarchy = Sequencer->GetEvaluationTemplate().GetHierarchy();
+		for (const TTuple<FMovieSceneSequenceID, FMovieSceneSubSequenceData>& Pair : Hierarchy.AllSubSequenceData())
+		{
+			if (UMovieSceneSequence* SubSequence = Pair.Value.GetSequence())
+			{
+				FFrameRate SubDisplayRate = SubSequence->GetMovieScene()->GetDisplayRate();
+				if (SubSequence->GetMovieScene() && DisplayRate != SubDisplayRate)
+				{
+					const FCommonFrameRateInfo* DisplayRateInfo = FCommonFrameRates::Find(DisplayRate);
+					const FCommonFrameRateInfo* SubDisplayRateInfo = FCommonFrameRates::Find(SubDisplayRate);
+
+					FText DisplayRateText = DisplayRateInfo ? DisplayRateInfo->DisplayName : FText::Format(LOCTEXT("DisplayRateFormat", "{0} fps"), DisplayRate.AsDecimal());
+					FText SubDisplayRateText = SubDisplayRateInfo ? SubDisplayRateInfo->DisplayName : FText::Format(LOCTEXT("SubDisplayRateFormat", "{0} fps"), SubDisplayRate.AsDecimal());
+
+					return FText::Format(LOCTEXT("FrameRateMismatchDescription", "At least one mismatch in display rate: {0} is at {1} and {2} is at {3}"), Sequencer->GetRootMovieSceneSequence()->GetDisplayName(), DisplayRateText, SubSequence->GetDisplayName(), SubDisplayRateText);
+				}
+			}
+		}
 	}
 
 	return FText();
diff --git a/Engine/Source/Editor/Sequencer/Private/SSequencerPlayRateCombo.h b/Engine/Source/Editor/Sequencer/Private/SSequencerPlayRateCombo.h
--- a/Engine/Source/Editor/Sequencer/Private/SSequencerPlayRateCombo.h
+++ b/Engine/Source/Editor/Sequencer/Private/SSequencerPlayRateCombo.h
@@ -56,10 +56,12 @@ private:
 
 	void SetClockSource(EUpdateClockSource NewClockSource);
 
-	FText GetFrameRateErrorDescription() const;
+	FText GetFrameRateIsMultipleOfErrorDescription() const;
+	FText GetFrameRateMismatchErrorDescription() const;
 
 	EVisibility GetFrameLockedVisibility() const;
-	EVisibility GetFrameRateErrorVisibility() const;
+	EVisibility GetFrameRateIsMultipleOfErrorVisibility() const;
+	EVisibility GetFrameRateMismatchErrorVisibility() const;
 
 	/** Sequencer pointer */
 	TWeakPtr<FSequencer> WeakSequencer;
