Sequencer: Add option to force refresh details when initializing sequencer. This is disabled in the actor sequence since it can cause the details panel refresh to be reentrant because its initialized from the "Open in Tab" button in the details panel.

diff --git a/Engine/Plugins/MovieScene/ActorSequence/Source/ActorSequenceEditor/Private/ActorSequenceEditorTabSummoner.cpp b/Engine/Plugins/MovieScene/ActorSequence/Source/ActorSequenceEditor/Private/ActorSequenceEditorTabSummoner.cpp
--- a/Engine/Plugins/MovieScene/ActorSequence/Source/ActorSequenceEditor/Private/ActorSequenceEditorTabSummoner.cpp
+++ b/Engine/Plugins/MovieScene/ActorSequence/Source/ActorSequenceEditor/Private/ActorSequenceEditorTabSummoner.cpp
@@ -340,6 +340,7 @@ public:
 		FLevelEditorSequencerIntegrationOptions Options;
 		Options.bRequiresLevelEvents = true;
 		Options.bRequiresActorEvents = false;
+		Options.bForceRefreshDetails = false;
 		Options.bCanRecord = false;
 
 		FLevelEditorSequencerIntegration::Get().AddSequencer(Sequencer.ToSharedRef(), Options);
diff --git a/Engine/Source/Editor/Sequencer/Private/LevelEditorSequencerIntegration.cpp b/Engine/Source/Editor/Sequencer/Private/LevelEditorSequencerIntegration.cpp
--- a/Engine/Source/Editor/Sequencer/Private/LevelEditorSequencerIntegration.cpp
+++ b/Engine/Source/Editor/Sequencer/Private/LevelEditorSequencerIntegration.cpp
@@ -161,7 +161,7 @@ void FLevelEditorSequencerIntegration::IterateAllSequencers(TFunctionRef<void(FS
 	}
 }
 
-void FLevelEditorSequencerIntegration::Initialize()
+void FLevelEditorSequencerIntegration::Initialize(const FLevelEditorSequencerIntegrationOptions& Options)
 {
 	AcquiredResources.Release();
 
@@ -217,7 +217,7 @@ void FLevelEditorSequencerIntegration::Initialize()
 	}
 
 	AddLevelViewportMenuExtender();
-	ActivateDetailHandler();
+	ActivateDetailHandler(Options);
 	ActivateSequencerEditorMode();
 	BindLevelEditorCommands();
 
@@ -253,7 +253,7 @@ void FLevelEditorSequencerIntegration::Initialize()
 		);
 	}
 
-	const bool bForceRefresh = true;
+	bool bForceRefresh = Options.bForceRefreshDetails;
 	UpdateDetails(bForceRefresh);
 }
 
@@ -750,7 +750,7 @@ void FLevelEditorSequencerIntegration::MakeBrowseToSelectedActorSubMenu(FMenuBui
 	}
 }
 
-void FLevelEditorSequencerIntegration::ActivateDetailHandler()
+void FLevelEditorSequencerIntegration::ActivateDetailHandler(const FLevelEditorSequencerIntegrationOptions& Options)
 {
 	FName DetailHandlerName("DetailHandler");
 
@@ -818,12 +818,17 @@ void FLevelEditorSequencerIntegration::ActivateDetailHandler()
 		};
 
 	AcquiredResources.Add(DetailHandlerName, DeactivateDetailKeyframeHandler);
-	AcquiredResources.Add(DetailHandlerRefreshName, RefreshDetailHandler);
+
+	if (Options.bForceRefreshDetails)
+	{
+		AcquiredResources.Add(DetailHandlerRefreshName, RefreshDetailHandler);
+	}
 }
 
 void FLevelEditorSequencerIntegration::OnPropertyEditorOpened()
 {
-	ActivateDetailHandler();
+	FLevelEditorSequencerIntegrationOptions Options;
+	ActivateDetailHandler(Options);
 }
 
 void FLevelEditorSequencerIntegration::BindLevelEditorCommands()
@@ -1109,7 +1114,7 @@ void FLevelEditorSequencerIntegration::AddSequencer(TSharedRef<ISequencer> InSeq
 {
 	if (!BoundSequencers.Num())
 	{
-		Initialize();
+		Initialize(Options);
 	}
 
 	KeyFrameHandler->Add(InSequencer);
diff --git a/Engine/Source/Editor/Sequencer/Public/LevelEditorSequencerIntegration.h b/Engine/Source/Editor/Sequencer/Public/LevelEditorSequencerIntegration.h
--- a/Engine/Source/Editor/Sequencer/Public/LevelEditorSequencerIntegration.h
+++ b/Engine/Source/Editor/Sequencer/Public/LevelEditorSequencerIntegration.h
@@ -26,11 +26,13 @@ struct FLevelEditorSequencerIntegrationOptions
 	FLevelEditorSequencerIntegrationOptions()
 		: bRequiresLevelEvents(true)
 		, bRequiresActorEvents(false)
+		, bForceRefreshDetails(true)
 		, bCanRecord(false)
 	{}
 
 	bool bRequiresLevelEvents : 1;
 	bool bRequiresActorEvents : 1;
+	bool bForceRefreshDetails : 1;
 	bool bCanRecord : 1;
 };
 
@@ -73,7 +75,7 @@ public:
 
 	static FLevelEditorSequencerIntegration& Get();
 
-	void Initialize();
+	void Initialize(const FLevelEditorSequencerIntegrationOptions& Options);
 
 	void AddSequencer(TSharedRef<ISequencer> InSequencer, const FLevelEditorSequencerIntegrationOptions& Options);
 
@@ -166,7 +168,7 @@ private:
 
 	void ActivateSequencerEditorMode();
 	void AddLevelViewportMenuExtender();
-	void ActivateDetailHandler();
+	void ActivateDetailHandler(const FLevelEditorSequencerIntegrationOptions& Options);
 	void AttachOutlinerColumn();
 	void DetachOutlinerColumn();
 	void ActivateRealtimeViewports();
