Add filters to sequencer

diff --git a/Engine/Plugins/Animation/LiveLink/Source/LiveLinkEditor/Private/Sequencer/LiveLinkSequencerFilters.cpp b/Engine/Plugins/Animation/LiveLink/Source/LiveLinkEditor/Private/Sequencer/LiveLinkSequencerFilters.cpp
new file mode 100644
--- /dev/null
+++ b/Engine/Plugins/Animation/LiveLink/Source/LiveLinkEditor/Private/Sequencer/LiveLinkSequencerFilters.cpp
@@ -0,0 +1,29 @@
+// Copyright 1998-2019 Epic Games, Inc. All Rights Reserved.
+
+#include "LiveLinkSequencerFilters.h"
+#include "MovieScene/MovieSceneLiveLinkTrack.h"
+#include "LiveLinkComponent.h"
+#include "Styling/SlateIconFinder.h"
+
+#define LOCTEXT_NAMESPACE "LiveLinkSequencerTrackFilters"
+
+//////////////////////////////////////////////////////////////////////////
+//
+
+class FSequencerTrackFilter_LiveLinkTracks : public FSequencerTrackFilter_ClassType<UMovieSceneLiveLinkTrack>
+{
+	virtual FString GetName() const override { return TEXT("SequencerLiveLinkTrackFilter"); }
+	virtual FText GetDisplayName() const override { return LOCTEXT("SequencerTrackFilter_LiveLinkTracks", "Live Link"); }
+	virtual FText GetToolTipText() const override { return LOCTEXT("SequencerTrackFilter_LiveLinkTracksToolTip", "Show only Live Link tracks."); }
+	virtual FSlateIcon GetIcon() const { return FSlateIconFinder::FindIconForClass(ULiveLinkComponent::StaticClass()); }
+};
+
+//////////////////////////////////////////////////////////////////////////
+//
+
+void ULiveLinkSequencerTrackFilter::AddTrackFilterExtensions(TArray< TSharedRef<class FSequencerTrackFilter> >& InOutFilterList) const
+{
+	InOutFilterList.Add(MakeShared<FSequencerTrackFilter_LiveLinkTracks>());
+}
+
+#undef LOCTEXT_NAMESPACE
diff --git a/Engine/Plugins/Animation/LiveLink/Source/LiveLinkEditor/Private/Sequencer/LiveLinkSequencerFilters.h b/Engine/Plugins/Animation/LiveLink/Source/LiveLinkEditor/Private/Sequencer/LiveLinkSequencerFilters.h
new file mode 100644
--- /dev/null
+++ b/Engine/Plugins/Animation/LiveLink/Source/LiveLinkEditor/Private/Sequencer/LiveLinkSequencerFilters.h
@@ -0,0 +1,20 @@
+// Copyright 1998-2019 Epic Games, Inc. All Rights Reserved.
+
+#pragma once
+
+#include "CoreMinimal.h"
+#include "UObject/ObjectMacros.h"
+#include "SequencerTrackFilterBase.h"
+#include "SequencerTrackFilterExtension.h"
+#include "LiveLinkSequencerFilters.generated.h"
+
+UCLASS()
+class ULiveLinkSequencerTrackFilter : public USequencerTrackFilterExtension
+{
+public:
+	GENERATED_BODY()
+
+	// USequencerTrackFilterExtension interface
+	virtual void AddTrackFilterExtensions(TArray< TSharedRef<class FSequencerTrackFilter> >& InOutFilterList) const override;
+	// End of USequencerTrackFilterExtension interface
+};
diff --git a/Engine/Plugins/FX/Niagara/Source/NiagaraEditor/Private/Sequencer/NiagaraSequence/NiagaraSequencerFilters.cpp b/Engine/Plugins/FX/Niagara/Source/NiagaraEditor/Private/Sequencer/NiagaraSequence/NiagaraSequencerFilters.cpp
new file mode 100644
--- /dev/null
+++ b/Engine/Plugins/FX/Niagara/Source/NiagaraEditor/Private/Sequencer/NiagaraSequence/NiagaraSequencerFilters.cpp
@@ -0,0 +1,49 @@
+// Copyright 1998-2019 Epic Games, Inc. All Rights Reserved.
+
+#include "NiagaraSequencerFilters.h"
+#include "MovieScene/MovieSceneNiagaraTrack.h"
+#include "NiagaraActor.h"
+#include "NiagaraComponent.h"
+
+#define LOCTEXT_NAMESPACE "NiagaraSequencerTrackFilters"
+
+//////////////////////////////////////////////////////////////////////////
+//
+
+class FSequencerTrackFilter_NiagaraTracks : public FSequencerTrackFilter
+{
+	virtual FString GetName() const override { return TEXT("SequencerNiagaraTrackFilter"); }
+	virtual FText GetDisplayName() const override { return LOCTEXT("SequencerTrackFilter_NiagaraTracks", "Niagara"); }
+	virtual FText GetToolTipText() const override { return LOCTEXT("SequencerTrackFilter_NiagaraTracksToolTip", "Show only Niagara tracks."); }
+	virtual FSlateIcon GetIcon() const { return FSlateIcon(FEditorStyle::GetStyleSetName(), "ClassIcon.ParticleSystem"); }
+
+	// IFilter implementation
+	virtual bool PassesFilter(FTrackFilterType InItem) const override
+	{
+		if (InItem->IsA(UMovieSceneNiagaraTrack::StaticClass()) || InItem->IsA(ANiagaraActor::StaticClass()) || InItem->IsA(UNiagaraComponent::StaticClass()))
+		{
+			return true;
+		}
+
+		const AActor* Actor = Cast<const AActor>(InItem);
+		if (Actor)
+		{
+			if (Actor->FindComponentByClass(UNiagaraComponent::StaticClass()))
+			{
+				return true;
+			}
+		}
+
+		return false;
+	}
+};
+
+//////////////////////////////////////////////////////////////////////////
+//
+
+void UNiagaraSequencerTrackFilter::AddTrackFilterExtensions(TArray< TSharedRef<class FSequencerTrackFilter> >& InOutFilterList) const
+{
+	InOutFilterList.Add(MakeShared<FSequencerTrackFilter_NiagaraTracks>());
+}
+
+#undef LOCTEXT_NAMESPACE
diff --git a/Engine/Plugins/FX/Niagara/Source/NiagaraEditor/Private/Sequencer/NiagaraSequence/NiagaraSequencerFilters.h b/Engine/Plugins/FX/Niagara/Source/NiagaraEditor/Private/Sequencer/NiagaraSequence/NiagaraSequencerFilters.h
new file mode 100644
--- /dev/null
+++ b/Engine/Plugins/FX/Niagara/Source/NiagaraEditor/Private/Sequencer/NiagaraSequence/NiagaraSequencerFilters.h
@@ -0,0 +1,20 @@
+// Copyright 1998-2019 Epic Games, Inc. All Rights Reserved.
+
+#pragma once
+
+#include "CoreMinimal.h"
+#include "UObject/ObjectMacros.h"
+#include "SequencerTrackFilterBase.h"
+#include "SequencerTrackFilterExtension.h"
+#include "NiagaraSequencerFilters.generated.h"
+
+UCLASS()
+class UNiagaraSequencerTrackFilter : public USequencerTrackFilterExtension
+{
+public:
+	GENERATED_BODY()
+
+	// USequencerTrackFilterExtension interface
+	virtual void AddTrackFilterExtensions(TArray< TSharedRef<class FSequencerTrackFilter> >& InOutFilterList) const override;
+	// End of USequencerTrackFilterExtension interface
+};
diff --git a/Engine/Source/Editor/Sequencer/Private/SSequencer.cpp b/Engine/Source/Editor/Sequencer/Private/SSequencer.cpp
--- a/Engine/Source/Editor/Sequencer/Private/SSequencer.cpp
+++ b/Engine/Source/Editor/Sequencer/Private/SSequencer.cpp
@@ -78,6 +78,8 @@
 #include "MovieSceneCopyableBinding.h"
 #include "MovieSceneCopyableTrack.h"
 #include "Fonts/FontMeasure.h"
+#include "SequencerTrackFilters.h"
+#include "SequencerTrackFilterExtension.h"
 
 #define LOCTEXT_NAMESPACE "Sequencer"
 
@@ -100,6 +102,8 @@ void SSequencer::Construct(const FArguments& InArgs, TSharedRef<FSequencer> InSe
 
 	Settings = InSequencer->GetSequencerSettings();
 
+	InitializeTrackFilters();
+
 	ISequencerWidgetsModule& SequencerWidgets = FModuleManager::Get().LoadModuleChecked<ISequencerWidgetsModule>( "SequencerWidgets" );
 
 	OnPlaybackRangeBeginDrag = InArgs._OnPlaybackRangeBeginDrag;
@@ -370,12 +374,20 @@ void SSequencer::Construct(const FArguments& InArgs, TSharedRef<FSequencer> InSe
 							[
 								MakeAddButton()
 							]
+							
+							+ SHorizontalBox::Slot()
+							.AutoWidth()
+							.VAlign(VAlign_Center)
+							.Padding(FMargin(0.f, 0.f, CommonPadding, 0.f))
+							[
+								MakeFilterButton()
+							]
 
 							+ SHorizontalBox::Slot()
 							.VAlign(VAlign_Center)
 							[
 								SAssignNew(SearchBox, SSearchBox)
-								.HintText(LOCTEXT("FilterNodesHint", "Filter"))
+								.HintText(LOCTEXT("SearchNodesHint", "Search Tracks"))
 								.OnTextChanged( this, &SSequencer::OnOutlinerSearchChanged )
 							]
 							+ SHorizontalBox::Slot()
@@ -694,6 +706,36 @@ TSharedRef<INumericTypeInterface<double>> SSequencer::GetNumericTypeInterface()
 	return NumericTypeInterface.ToSharedRef();
 }
 
+void SSequencer::InitializeTrackFilters()
+{
+	// Add all built-in track filters here
+	AllTrackFilters.Add(MakeShared<FSequencerTrackFilter_AudioTracks>());
+	AllTrackFilters.Add(MakeShared<FSequencerTrackFilter_EventTracks>());
+	AllTrackFilters.Add(MakeShared<FSequencerTrackFilter_LevelVisibilityTracks>());
+	AllTrackFilters.Add(MakeShared<FSequencerTrackFilter_ParticleTracks>());
+
+	AllTrackFilters.Add(MakeShared<FSequencerTrackFilter_CameraObjects>());
+	AllTrackFilters.Add(MakeShared<FSequencerTrackFilter_LightObjects>());
+	AllTrackFilters.Add(MakeShared<FSequencerTrackFilter_SkeletalMeshObjects>());
+
+	// Add any global user-defined frontend filters
+	for (TObjectIterator<USequencerTrackFilterExtension> ExtensionIt(RF_NoFlags); ExtensionIt; ++ExtensionIt)
+	{
+		if (USequencerTrackFilterExtension* PotentialExtension = *ExtensionIt)
+		{
+			if (PotentialExtension->HasAnyFlags(RF_ClassDefaultObject) && !PotentialExtension->GetClass()->HasAnyClassFlags(CLASS_Deprecated | CLASS_Abstract))
+			{
+				// Grab the filters
+				TArray< TSharedRef<FSequencerTrackFilter> > ExtendedTrackFilters;
+				PotentialExtension->AddTrackFilterExtensions(ExtendedTrackFilters);
+				AllTrackFilters.Append(ExtendedTrackFilters);
+			}
+		}
+	}
+
+	// Sort by display name
+	AllTrackFilters.Sort([](const TSharedRef<FSequencerTrackFilter>& LHS, const TSharedRef<FSequencerTrackFilter>& RHS) { return LHS->GetDisplayName().ToString() < RHS->GetDisplayName().ToString(); });
+}
 
 /* SSequencer callbacks
  *****************************************************************************/
@@ -796,6 +838,40 @@ TSharedRef<SWidget> SSequencer::MakeAddButton()
 	];
 }
 
+TSharedRef<SWidget> SSequencer::MakeFilterButton()
+{
+	return SNew(SComboButton)
+	.ComboButtonStyle(FEditorStyle::Get(), "GenericFilters.ComboButtonStyle")
+	.ForegroundColor(FLinearColor::White)
+	.ContentPadding(0)
+	.ToolTipText(LOCTEXT("AddTrackFilterToolTip", "Add a track filter."))
+	.OnGetMenuContent(this, &SSequencer::MakeFilterMenu)
+	.HasDownArrow(true)
+	.ContentPadding(FMargin(1, 0))
+	.ButtonContent()
+	[
+		SNew(SHorizontalBox)
+
+		+ SHorizontalBox::Slot()
+		.AutoWidth()
+		[
+			SNew(STextBlock)
+			.TextStyle(FEditorStyle::Get(), "GenericFilters.TextStyle")
+			.Font(FEditorStyle::Get().GetFontStyle("FontAwesome.9"))
+			.Text(FEditorFontGlyphs::Filter)
+		]
+
+		+ SHorizontalBox::Slot()
+		.AutoWidth()
+		.Padding(2, 0, 0, 0)
+		[
+			SNew(STextBlock)
+			.TextStyle(FEditorStyle::Get(), "GenericFilters.TextStyle")
+			.Text(LOCTEXT("Filters", "Filters"))
+		]
+	];
+}
+
 TSharedRef<SWidget> SSequencer::MakeToolBar()
 {
 	ISequencerModule& SequencerModule = FModuleManager::GetModuleChecked<ISequencerModule>("Sequencer");
@@ -1045,6 +1121,74 @@ TSharedRef<SWidget> SSequencer::MakeAddMenu()
 	return MenuBuilder.MakeWidget();
 }
 
+TSharedRef<SWidget> SSequencer::MakeFilterMenu()
+{
+	FMenuBuilder MenuBuilder(true, nullptr, AddMenuExtender);
+
+	// let track editors & object bindings populate the menu
+	TSharedPtr<FSequencer> Sequencer = SequencerPtr.Pin();
+
+	MenuBuilder.BeginSection("SequencerTracksResetFilters");
+	{
+		MenuBuilder.AddMenuEntry(
+			LOCTEXT("FilterListResetFilters", "Reset Filters"),
+			LOCTEXT("FilterListResetToolTip", "Resets current filter selection"),
+			FSlateIcon(),
+			FUIAction(FExecuteAction::CreateSP(this, &SSequencer::OnResetFilters))
+		);
+	}
+	MenuBuilder.EndSection();
+
+	MenuBuilder.BeginSection("TrackFilters");
+	
+	for (TSharedRef<FSequencerTrackFilter> TrackFilter : AllTrackFilters)
+	{
+		if (TrackFilter->SupportsSequence(Sequencer->GetFocusedMovieSceneSequence()))
+		{
+			MenuBuilder.AddMenuEntry(
+				TrackFilter->GetDisplayName(),
+				TrackFilter->GetToolTipText(),
+				TrackFilter->GetIcon(),
+				FUIAction(
+				FExecuteAction::CreateSP(this, &SSequencer::OnTrackFilterClicked, TrackFilter),
+				FCanExecuteAction(),
+				FIsActionChecked::CreateSP(this, &SSequencer::IsTrackFilterActive, TrackFilter)),
+				NAME_None,
+				EUserInterfaceActionType::ToggleButton
+			);
+		}
+	}
+
+	MenuBuilder.EndSection();
+
+	return MenuBuilder.MakeWidget();
+}
+
+void SSequencer::OnResetFilters()
+{
+	TSharedPtr<FSequencer> Sequencer = SequencerPtr.Pin();
+	Sequencer->GetNodeTree()->RemoveAllFilters();
+}
+
+void SSequencer::OnTrackFilterClicked(TSharedRef<FSequencerTrackFilter> TrackFilter)
+{
+	TSharedPtr<FSequencer> Sequencer = SequencerPtr.Pin();
+	if (IsTrackFilterActive(TrackFilter))
+	{
+		Sequencer->GetNodeTree()->RemoveFilter(TrackFilter);
+	}
+	else
+	{
+		Sequencer->GetNodeTree()->AddFilter(TrackFilter);
+	}
+}
+
+bool SSequencer::IsTrackFilterActive(TSharedRef<FSequencerTrackFilter> TrackFilter) const
+{
+	TSharedPtr<FSequencer> Sequencer = SequencerPtr.Pin();
+	return Sequencer->GetNodeTree()->IsTrackFilterActive(TrackFilter);
+}
+
 TSharedRef<SWidget> SSequencer::MakeGeneralMenu()
 {
 	FMenuBuilder MenuBuilder( true, SequencerPtr.Pin()->GetCommandBindings() );
@@ -1799,8 +1943,6 @@ void SSequencer::OnOutlinerSearchChanged( const FText& Filter )
 
 		Sequencer->GetNodeTree()->FilterNodes( FilterString );
 
-		Sequencer->GetNodeTree()->Update();
-
 		TreeView->Refresh();
 
 		if ( FilterString.StartsWith( TEXT( "label:" ) ) )
diff --git a/Engine/Source/Editor/Sequencer/Private/SSequencer.h b/Engine/Source/Editor/Sequencer/Private/SSequencer.h
--- a/Engine/Source/Editor/Sequencer/Private/SSequencer.h
+++ b/Engine/Source/Editor/Sequencer/Private/SSequencer.h
@@ -31,6 +31,7 @@ class SSequencerTransformBox;
 class SSequencerStretchBox;
 class SSequencerTreeView;
 class USequencerSettings;
+class FSequencerTrackFilter;
 struct FPaintPlaybackRangeArgs;
 struct FSectionHandle;
 
@@ -284,6 +285,9 @@ protected:
 
 private:
 	
+	/** Initalizes a list of all track filter objects */
+	void InitializeTrackFilters();
+
 	/** Handles key selection changes. */
 	void HandleKeySelectionChanged();
 
@@ -314,9 +318,13 @@ private:
 	/** Makes add button. */
 	TSharedRef<SWidget> MakeAddButton();
 
+	TSharedRef<SWidget> MakeFilterButton();
+
 	/** Makes the add menu for the toolbar. */
 	TSharedRef<SWidget> MakeAddMenu();
 
+	TSharedRef<SWidget> MakeFilterMenu();
+
 	/** Makes the general menu for the toolbar. */
 	TSharedRef<SWidget> MakeGeneralMenu();
 
@@ -355,6 +363,10 @@ public:
 
 private:
 
+	void OnResetFilters();
+	void OnTrackFilterClicked(TSharedRef<FSequencerTrackFilter> TrackFilter);
+	bool IsTrackFilterActive(TSharedRef<FSequencerTrackFilter> TrackFilter) const;
+
 	/**
 	* Called when the time snap interval changes.
 	*
@@ -580,4 +592,7 @@ private:
 	TArray<FString> AdditionalSelectionsToAdd;
 
 	TSharedPtr<SWidget> TickResolutionOverlay;
+
+	/** All possible track filter objects */
+	TArray< TSharedRef<FSequencerTrackFilter> > AllTrackFilters;
 };
diff --git a/Engine/Source/Editor/Sequencer/Private/Sequencer.cpp b/Engine/Source/Editor/Sequencer/Private/Sequencer.cpp
--- a/Engine/Source/Editor/Sequencer/Private/Sequencer.cpp
+++ b/Engine/Source/Editor/Sequencer/Private/Sequencer.cpp
@@ -554,7 +554,7 @@ void FSequencer::Tick(float InDeltaTime)
 		}
 	}
 
-	if (bNeedTreeRefresh)
+	if (bNeedTreeRefresh || NodeTree->NeedsFilterUpdate())
 	{
 		SelectionPreview.Empty();
 
diff --git a/Engine/Source/Editor/Sequencer/Private/SequencerNodeTree.cpp b/Engine/Source/Editor/Sequencer/Private/SequencerNodeTree.cpp
--- a/Engine/Source/Editor/Sequencer/Private/SequencerNodeTree.cpp
+++ b/Engine/Source/Editor/Sequencer/Private/SequencerNodeTree.cpp
@@ -18,6 +18,21 @@
 #include "DisplayNodes/SequencerSpacerNode.h"
 #include "Widgets/Views/STableRow.h"
 #include "SequencerNodeSortingMethods.h"
+#include "SequencerTrackFilters.h"
+
+FSequencerNodeTree::~FSequencerNodeTree()
+{
+	if (TrackFilters.IsValid())
+	{
+		TrackFilters->OnChanged().RemoveAll(this);
+	}
+}
+
+FSequencerNodeTree::FSequencerNodeTree(class FSequencer& InSequencer) : Sequencer(InSequencer), bFilterUpdateRequested(false)
+{
+	TrackFilters = MakeShared<FSequencerTrackFilterCollection>();
+	TrackFilters->OnChanged().AddRaw(this, &FSequencerNodeTree::RequestFilterUpdate);
+}
 
 void FSequencerNodeTree::Empty()
 {
@@ -197,7 +212,8 @@ void FSequencerNodeTree::Update()
 
 	// Re-filter the tree after updating 
 	// @todo sequencer: Newly added sections may need to be visible even when there is a filter
-	FilterNodes( FilterString );
+	bFilterUpdateRequested = true;
+	UpdateFilters();
 
 	OnUpdatedDelegate.Broadcast();
 }
@@ -301,6 +317,10 @@ const TArray<TSharedRef<FSequencerDisplayNode>>& FSequencerNodeTree::GetRootNode
 	return RootNodes;
 }
 
+bool FSequencerNodeTree::HasActiveFilter() const
+{
+	return (!FilterString.IsEmpty() || TrackFilters->Num() > 0);
+}
 
 TSharedRef<FSequencerObjectBindingNode> FSequencerNodeTree::AddObjectBinding(const FString& ObjectName, const FGuid& ObjectBinding, TMap<FGuid, const FMovieSceneBinding*>& GuidToBindingMap, TArray<TSharedRef<FSequencerObjectBindingNode>>& OutNodeList)
 {
@@ -535,6 +555,26 @@ void FSequencerNodeTree::SortAllNodesAndDescendants()
 	GetSequencer().RefreshTree();
 }
 
+void FSequencerNodeTree::AddFilter(TSharedRef<FSequencerTrackFilter> TrackFilter)
+{
+	TrackFilters->Add(TrackFilter);
+}
+
+int32 FSequencerNodeTree::RemoveFilter(TSharedRef<FSequencerTrackFilter> TrackFilter)
+{
+	return TrackFilters->Remove(TrackFilter);
+}
+
+void FSequencerNodeTree::RemoveAllFilters()
+{
+	TrackFilters->RemoveAll();
+}
+
+bool FSequencerNodeTree::IsTrackFilterActive(TSharedRef<FSequencerTrackFilter> TrackFilter) const
+{
+	return TrackFilters->Contains(TrackFilter);
+}
+
 void FSequencerNodeTree::SaveExpansionState(const FSequencerDisplayNode& Node, bool bExpanded)
 {	
 	// @todo Sequencer - This should be moved to the sequence level
@@ -632,7 +672,7 @@ static void AddFilteredNode(const TSharedRef<FSequencerDisplayNode>& StartNode,
 
 	// Gather parent folders up the chain
 	TSharedPtr<FSequencerDisplayNode> ParentNode = StartNode->GetParent();
-	while (ParentNode.IsValid() && ParentNode.Get()->GetType() == ESequencerNode::Folder)
+	while (ParentNode.IsValid())
 	{
 		if (!ParentNode.Get()->IsExpanded())
 		{
@@ -663,121 +703,125 @@ static void AddParentNodes(const TSharedRef<FSequencerDisplayNode>& StartNode, T
  * Recursively filters nodes
  *
  * @param StartNode			The node to start from
- * @param FilterStrings		The filter strings which need to be matched
+ * @param Filters			The filter collection to test against
  * @param OutFilteredNodes	The list of all filtered nodes
- * @return Whether the text filter was passed
  */
-static bool FilterNodesRecursive( FSequencer& Sequencer, const TSharedRef<FSequencerDisplayNode>& StartNode, const TArray<FString>& FilterStrings, TSet<TSharedRef<FSequencerDisplayNode>>& OutFilteredNodes )
-{
-	// check labels - only one of the labels needs to match
-	bool bMatchedLabel = false;
-	bool bObjectHasLabels = false;
-	for (const FString& String : FilterStrings)
-	{
-		if (String.StartsWith(TEXT("label:")) && String.Len() > 6)
-		{
-			if (StartNode->GetType() == ESequencerNode::Object)
-			{
-				bObjectHasLabels = true;
-				auto ObjectBindingNode = StaticCastSharedRef<FSequencerObjectBindingNode>(StartNode);
-				auto Labels = Sequencer.GetLabelManager().GetObjectLabels(ObjectBindingNode->GetObjectBinding());
 
-				if (Labels != nullptr && Labels->Strings.Contains(String.RightChop(6)))
-				{
-					bMatchedLabel = true;
-					break;
-				}
-			}
-			else if (!StartNode->GetParent().IsValid())
-			{
-				return false;
-			}
-		}
-	}
-
-	if (bObjectHasLabels && !bMatchedLabel)
+static void FilterNodesRecursive( FSequencer& Sequencer, const TSharedRef<FSequencerDisplayNode>& StartNode, TSharedPtr<FSequencerTrackFilterCollection> Filters, const TArray<FString>& FilterStrings, TSet<TSharedRef<FSequencerDisplayNode>>& OutFilteredNodes )
+{
+	for (TSharedRef<FSequencerDisplayNode> Node : StartNode->GetChildNodes())
 	{
-		return false;
+		FilterNodesRecursive(Sequencer, Node, Filters, FilterStrings, OutFilteredNodes);
 	}
 
-	// assume the filter is acceptable
-	bool bPassedTextFilter = true;
-
-	// check each string in the filter strings list against 
-	for (const FString& String : FilterStrings)
+	bool bPasssedAnyFilters = false;
+	if (StartNode->GetType() == ESequencerNode::Track)
 	{
-		if (!String.StartsWith(TEXT("label:")) && !StartNode->GetDisplayName().ToString().Contains(String)) 
+		UMovieSceneTrack* Track = static_cast<const FSequencerTrackNode&>(StartNode.Get()).GetTrack();
+		if (Filters->PassesAnyFilters(Track))
 		{
-			bPassedTextFilter = false;
-			break;
+			bPasssedAnyFilters = true;
 		}
 	}
-
-	// whether or the start node is in the filter
-	bool bInFilter = false;
-
-	if (bPassedTextFilter)
+	else if (StartNode->GetType() == ESequencerNode::Object)
 	{
-		// This node is now filtered
-		AddFilteredNode(StartNode, OutFilteredNodes);
-
-		bInFilter = true;
+		const FSequencerObjectBindingNode ObjectNode = static_cast<const FSequencerObjectBindingNode&>(StartNode.Get());
+		for (TWeakObjectPtr<>& Object : Sequencer.FindObjectsInCurrentSequence(ObjectNode.GetObjectBinding()))
+		{
+			if (Filters->PassesAnyFilters(Object.Get()))
+			{
+				bPasssedAnyFilters = true;
+				break;
+			}
+		}
 	}
 
-	// check each child node to determine if it is filtered
-	if (StartNode->GetType() != ESequencerNode::Folder)
+	if (Filters->Num() == 0 || bPasssedAnyFilters)
 	{
-		const TArray<TSharedRef<FSequencerDisplayNode>>& ChildNodes = StartNode->GetChildNodes();
-
-		for (const auto& Node : ChildNodes)
+		// If we have a filter string, make sure we match
+		if (FilterStrings.Num() > 0)
 		{
-			// Mark the parent as filtered if any child node was filtered
-			bPassedTextFilter |= FilterNodesRecursive(Sequencer, Node, FilterStrings, OutFilteredNodes);
+			// check labels - only one of the labels needs to match
+			bool bMatchedLabel = false;
+			bool bObjectHasLabels = false;
+			for (const FString& String : FilterStrings)
+			{
+				if (String.StartsWith(TEXT("label:")) && String.Len() > 6)
+				{
+					if (StartNode->GetType() == ESequencerNode::Object)
+					{
+						bObjectHasLabels = true;
+						auto ObjectBindingNode = StaticCastSharedRef<FSequencerObjectBindingNode>(StartNode);
+						auto Labels = Sequencer.GetLabelManager().GetObjectLabels(ObjectBindingNode->GetObjectBinding());
+
+						if (Labels != nullptr && Labels->Strings.Contains(String.RightChop(6)))
+						{
+							bMatchedLabel = true;
+							break;
+						}
+					}
+					else if (!StartNode->GetParent().IsValid())
+					{
+						return;
+					}
+				}
+			}
 
-			if (bPassedTextFilter && !bInFilter)
+			if (bObjectHasLabels && !bMatchedLabel)
 			{
-				AddParentNodes(Node, OutFilteredNodes);
+				return;
+			}
 
-				bInFilter = true;
+			// check each string in the filter strings list against 
+			for (const FString& String : FilterStrings)
+			{
+				if (!String.StartsWith(TEXT("label:")) && !StartNode->GetDisplayName().ToString().Contains(String))
+				{
+					return;
+				}
 			}
 		}
+		AddFilteredNode(StartNode, OutFilteredNodes);
 	}
-
-	return bPassedTextFilter;
 }
 
-
-void FSequencerNodeTree::FilterNodes(const FString& InFilter)
+void FSequencerNodeTree::UpdateFilters()
 {
+	if (!bFilterUpdateRequested)
+		return;
+
 	FilteredNodes.Empty();
 
-	if (InFilter.IsEmpty())
-	{
-		// No filter
-		FilterString.Empty();
-	}
-	else
+	if (TrackFilters->Num() > 0 || !FilterString.IsEmpty())
 	{
 		// Build a list of strings that must be matched
 		TArray<FString> FilterStrings;
 
-		FilterString = InFilter;
 		// Remove whitespace from the front and back of the string
 		FilterString.TrimStartAndEndInline();
 		FilterString.ParseIntoArray(FilterStrings, TEXT(" "), true /*bCullEmpty*/);
 
 		for (auto It = ObjectBindingMap.CreateIterator(); It; ++It)
 		{
-			// Recursively filter all nodes, matching them against the list of filter strings.  All filter strings must be matched
-			FilterNodesRecursive(Sequencer, It.Value().ToSharedRef(), FilterStrings, FilteredNodes);
+			FilterNodesRecursive(Sequencer, It.Value().ToSharedRef(), TrackFilters, FilterStrings, FilteredNodes);
 		}
 
 		for (auto It = RootNodes.CreateIterator(); It; ++It)
 		{
-			// Recursively filter all nodes, matching them against the list of filter strings.  All filter strings must be matched
-			FilterNodesRecursive(Sequencer, *It, FilterStrings, FilteredNodes);
+			FilterNodesRecursive(Sequencer, *It, TrackFilters, FilterStrings, FilteredNodes);
 		}
 	}
+
+	bFilterUpdateRequested = false;
+}
+
+void FSequencerNodeTree::FilterNodes(const FString& InFilter)
+{
+	if (InFilter != FilterString)
+	{
+		FilterString = InFilter;
+		bFilterUpdateRequested = true;
+	}
 }
 
 TArray< TSharedRef<FSequencerDisplayNode> > FSequencerNodeTree::GetAllNodes() const
diff --git a/Engine/Source/Editor/Sequencer/Private/SequencerNodeTree.h b/Engine/Source/Editor/Sequencer/Private/SequencerNodeTree.h
--- a/Engine/Source/Editor/Sequencer/Private/SequencerNodeTree.h
+++ b/Engine/Source/Editor/Sequencer/Private/SequencerNodeTree.h
@@ -13,6 +13,8 @@ class ISequencerTrackEditor;
 class UMovieSceneFolder;
 class UMovieSceneTrack;
 struct FMovieSceneBinding;
+class FSequencerTrackFilter;
+class FSequencerTrackFilterCollection;
 
 /**
  * Represents a tree of sequencer display nodes, used to populate the Sequencer UI with MovieScene data
@@ -23,9 +25,9 @@ public:
 	DECLARE_MULTICAST_DELEGATE(FOnUpdated);
 
 public:
-	FSequencerNodeTree( class FSequencer& InSequencer )
-		: Sequencer( InSequencer )
-	{}
+	FSequencerNodeTree(class FSequencer& InSequencer);
+	
+	~FSequencerNodeTree();
 	
 	/**
 	 * Empties the entire tree
@@ -43,7 +45,7 @@ public:
 	const TArray< TSharedRef<FSequencerDisplayNode> >& GetRootNodes() const;
 
 	/** @return Whether or not there is an active filter */
-	bool HasActiveFilter() const { return !FilterString.IsEmpty(); }
+	bool HasActiveFilter() const;
 
 	/**
 	 * Returns whether or not a node is filtered
@@ -51,7 +53,17 @@ public:
 	 * @param Node	The node to check if it is filtered
 	 */
 	bool IsNodeFiltered( const TSharedRef<const FSequencerDisplayNode> Node ) const;
-	
+
+	/**
+	 * Schedules an update of all filters
+	 */
+	void RequestFilterUpdate() { bFilterUpdateRequested = true; }
+
+	/**
+	 * @return Whether there is a filter update scheduled
+	 */
+	bool NeedsFilterUpdate() const { return bFilterUpdateRequested; }
+
 	/**
 	 * Filters the nodes based on the passed in filter terms
 	 *
@@ -116,7 +128,19 @@ public:
 	/** Sorts all nodes and their descendants by category then alphabetically.*/
 	void SortAllNodesAndDescendants();
 
+	void AddFilter(TSharedRef<FSequencerTrackFilter> TrackFilter);
+	int32 RemoveFilter(TSharedRef<FSequencerTrackFilter> TrackFilter);
+	void RemoveAllFilters();
+	bool IsTrackFilterActive(TSharedRef<FSequencerTrackFilter> TrackFilter) const;
+
 private:
+
+	/**
+	 * Update the list of filters nodes based on current filter settings, if an update is scheduled
+	 * This is called by Update();
+	 */
+	void UpdateFilters();
+
 	/**
 	 * Finds or adds a type editor for the track
 	 *
@@ -173,4 +197,9 @@ private:
 	FSequencer& Sequencer;
 	/** A multicast delegate which is called whenever the node tree has been updated. */
 	FOnUpdated OnUpdatedDelegate;
+
+	/** Active track filters */
+	TSharedPtr<FSequencerTrackFilterCollection> TrackFilters;
+
+	bool bFilterUpdateRequested;
 };
diff --git a/Engine/Source/Editor/Sequencer/Private/SequencerTrackFilters.h b/Engine/Source/Editor/Sequencer/Private/SequencerTrackFilters.h
new file mode 100644
--- /dev/null
+++ b/Engine/Source/Editor/Sequencer/Private/SequencerTrackFilters.h
@@ -0,0 +1,137 @@
+// Copyright 1998-2019 Epic Games, Inc. All Rights Reserved.
+
+#pragma once
+
+#include "SequencerTrackFilterBase.h"
+#include "Misc/IFilter.h"
+#include "Misc/FilterCollection.h"
+#include "CollectionManagerTypes.h"
+#include "Styling/SlateIconFinder.h"
+#include "SequencerDisplayNode.h"
+#include "MovieSceneSequence.h"
+
+#include "Tracks/MovieSceneAudioTrack.h"
+#include "Tracks/MovieSceneEventTrack.h"
+#include "Tracks/MovieSceneLevelVisibilityTrack.h"
+#include "Tracks/MovieSceneParticleTrack.h"
+
+#include "Camera/CameraComponent.h"
+#include "Components/LightComponentBase.h"
+#include "Components/SkeletalMeshComponent.h"
+#include "Particles/ParticleSystem.h"
+
+#define LOCTEXT_NAMESPACE "Sequencer"
+
+class FSequencerTrackFilterCollection : public TFilterCollection<FTrackFilterType>
+{
+public:
+	/**
+	 *	Returns whether the specified Item passes any of the filters in the collection
+	 *
+	 *	@param	InItem	The Item to check against all child filter restrictions
+	 *	@return			Whether the Item passed any child filter restrictions
+	 */
+	 // @todo Maybe this should get moved in to TFilterCollection
+	bool PassesAnyFilters(/*ItemType*/ FTrackFilterType InItem) const
+	{
+		for (int Index = 0; Index < ChildFilters.Num(); Index++)
+		{
+			if (ChildFilters[Index]->PassesFilter(InItem))
+			{
+				return true;
+			}
+		}
+
+		return false;
+	}
+
+	// @todo Maybe this should get moved in to TFilterCollection
+	bool Contains(const TSharedPtr< IFilter< FTrackFilterType > >& InItem) const
+	{
+		for (const TSharedPtr<IFilter<FTrackFilterType>>& Filter : ChildFilters)
+		{
+			if (InItem == Filter)
+			{
+				return true;
+			}
+		}
+		return false;
+	}
+
+	// @todo Maybe this should get moved in to TFilterCollection
+	void RemoveAll()
+	{
+		for (auto Iterator = ChildFilters.CreateIterator(); Iterator; ++Iterator)
+		{
+			const TSharedPtr< IFilter< FTrackFilterType > >& Filter = *Iterator;
+
+			if (Filter.IsValid())
+			{
+				Filter->OnChanged().RemoveAll(this);
+			}
+		}
+
+		ChildFilters.Empty();
+
+		ChangedEvent.Broadcast();
+	}
+};
+
+class FSequencerTrackFilter_AudioTracks : public FSequencerTrackFilter_ClassType< UMovieSceneAudioTrack >
+{
+	virtual FString GetName() const override { return TEXT("SequencerAudioTracksFilter"); }
+	virtual FText GetDisplayName() const override { return LOCTEXT("SequencerTrackFilter_AudioTracks", "Audio"); }
+	virtual FText GetToolTipText() const override { return LOCTEXT("SequencerTrackFilter_AudioTracksToolTip", "Show only Audio tracks."); }
+	virtual FSlateIcon GetIcon() const override { return FSlateIcon(FEditorStyle::GetStyleSetName(), "Sequencer.Tracks.Audio"); }
+
+};
+
+class FSequencerTrackFilter_EventTracks : public FSequencerTrackFilter_ClassType< UMovieSceneEventTrack >
+{
+	virtual FString GetName() const override { return TEXT("SequencerEventTracksFilter"); }
+	virtual FText GetDisplayName() const override { return LOCTEXT("SequencerTrackFilter_EventTracks", "Event"); }
+	virtual FText GetToolTipText() const override { return LOCTEXT("SequencerTrackFilter_EventTracksToolTip", "Show only Event tracks."); }
+	virtual FSlateIcon GetIcon() const override { return FSlateIcon(FEditorStyle::GetStyleSetName(), "Sequencer.Tracks.Event"); }
+};
+
+class FSequencerTrackFilter_LevelVisibilityTracks : public FSequencerTrackFilter_ClassType< UMovieSceneLevelVisibilityTrack >
+{
+	virtual FString GetName() const override { return TEXT("SequencerLevelVisibilityTracksFilter"); }
+	virtual FText GetDisplayName() const override { return LOCTEXT("SequencerTrackFilter_LevelVisibilityTracks", "Level Visibility"); }
+	virtual FText GetToolTipText() const override { return LOCTEXT("SequencerTrackFilter_LevelVisibilityTracksToolTip", "Show only Level Visibility tracks."); }
+	virtual FSlateIcon GetIcon() const override { return FSlateIcon(FEditorStyle::GetStyleSetName(), "Sequencer.Tracks.LevelVisibility"); }
+};
+
+class FSequencerTrackFilter_ParticleTracks : public FSequencerTrackFilter_ClassType< UMovieSceneParticleTrack >
+{
+	virtual FString GetName() const override { return TEXT("SequencerParticleTracksFilter"); }
+	virtual FText GetDisplayName() const override { return LOCTEXT("SequencerTrackFilter_ParticleTracks", "Particle Systems"); }
+	virtual FText GetToolTipText() const override { return LOCTEXT("SequencerTrackFilter_ParticleTracksToolTip", "Show only Particle System tracks."); }
+	virtual FSlateIcon GetIcon() const override { return FSlateIconFinder::FindIconForClass(UParticleSystem::StaticClass()); }
+};
+
+class FSequencerTrackFilter_SkeletalMeshObjects : public FSequencerTrackFilter_ComponentType< USkeletalMeshComponent >
+{
+	virtual FString GetName() const override { return TEXT("SequencerSkeletalMeshObjectsFilter"); }
+	virtual FText GetDisplayName() const override { return LOCTEXT("SequencerTrackFilter_SkeletalMeshObjects", "Skeletal Mesh"); }
+	virtual FText GetToolTipText() const override { return LOCTEXT("SequencerTrackFilter_SkeletalMeshObjectsToolTip", "Show only SkeletalMesh objects."); }
+	virtual FSlateIcon GetIcon() const override { return FSlateIconFinder::FindIconForClass(USkeletalMeshComponent::StaticClass()); }
+};
+
+class FSequencerTrackFilter_CameraObjects : public FSequencerTrackFilter_ComponentType< UCameraComponent >
+{
+	virtual FString GetName() const override { return TEXT("SequencerCameraObjectsFilter"); }
+	virtual FText GetDisplayName() const override { return LOCTEXT("SequencerTrackFilter_CameraObjects", "Cameras"); }
+	virtual FText GetToolTipText() const override { return LOCTEXT("SequencerTrackFilter_CameraObjectsToolTip", "Show only Camera objects."); }
+	virtual FSlateIcon GetIcon() const override { return FSlateIconFinder::FindIconForClass(UCameraComponent::StaticClass()); }
+};
+
+class FSequencerTrackFilter_LightObjects : public FSequencerTrackFilter_ComponentType< ULightComponentBase >
+{
+	virtual FString GetName() const override { return TEXT("SequencerLightObjectsFilter"); }
+	virtual FText GetDisplayName() const override { return LOCTEXT("SequencerTrackFilter_LightObjects", "Lights"); }
+	virtual FText GetToolTipText() const override { return LOCTEXT("SequencerTrackFilter_LightObjectsToolTip", "Show only Light objects."); }
+	virtual FSlateIcon GetIcon() const override { return FSlateIcon(FEditorStyle::GetStyleSetName(), "ClassIcon.Light"); }
+};
+
+#undef LOCTEXT_NAMESPACE
\ No newline at end of file
diff --git a/Engine/Source/Editor/Sequencer/Public/SequencerTrackFilterBase.h b/Engine/Source/Editor/Sequencer/Public/SequencerTrackFilterBase.h
new file mode 100644
--- /dev/null
+++ b/Engine/Source/Editor/Sequencer/Public/SequencerTrackFilterBase.h
@@ -0,0 +1,86 @@
+// Copyright 1998-2019 Epic Games, Inc. All Rights Reserved.
+
+#pragma once
+
+#include "CoreMinimal.h"
+#include "Misc/IFilter.h"
+#include "MovieSceneSequence.h"
+
+class FMenuBuilder;
+
+typedef const UObject* FTrackFilterType;
+
+class FSequencerTrackFilter : public IFilter<FTrackFilterType>
+{
+public:
+	/** Returns the system name for this filter */
+	virtual FString GetName() const = 0;
+
+	/** Returns the human readable name for this filter */
+	virtual FText GetDisplayName() const = 0;
+
+	/** Returns the tooltip for this filter, shown in the filters menu */
+	virtual FText GetToolTipText() const = 0;
+
+	/** Returns the icon to use in menu entries */
+	virtual FSlateIcon GetIcon() const { return FSlateIcon(); }
+
+	/** Returns whether the filter supports the sequence type */
+	virtual bool SupportsSequence(UMovieSceneSequence* InSequence) const
+	{
+		static UClass* LevelSequenceClass = FindObject<UClass>(ANY_PACKAGE, TEXT("LevelSequence"), true);
+		static UClass* WidgetAnimationClass = FindObject<UClass>(ANY_PACKAGE, TEXT("WidgetAnimation"), true);
+		return InSequence != nullptr &&
+			((LevelSequenceClass != nullptr && InSequence->GetClass()->IsChildOf(LevelSequenceClass)) ||
+			(WidgetAnimationClass != nullptr && InSequence->GetClass()->IsChildOf(WidgetAnimationClass)));
+	}
+
+	// IFilter implementation
+	DECLARE_DERIVED_EVENT(FSequencerTrackFilter, IFilter<FTrackFilterType>::FChangedEvent, FChangedEvent);
+	virtual FChangedEvent& OnChanged() override { return ChangedEvent; }
+
+protected:
+	void BroadcastChangedEvent() const { ChangedEvent.Broadcast(); }
+
+private:
+	FChangedEvent ChangedEvent;
+};
+
+/** Helper template for building Sequencer track filters based on a track or object type */
+template<typename TrackType>
+class FSequencerTrackFilter_ClassType : public FSequencerTrackFilter
+{
+	// IFilter implementation
+	virtual bool PassesFilter(FTrackFilterType InItem) const override
+	{
+		if (InItem->IsA(TrackType::StaticClass()))
+		{
+			return true;
+		}
+		return false;
+	}
+};
+
+/** Helper template for building Sequencer track filters based on a component */
+template<typename ComponentType>
+class FSequencerTrackFilter_ComponentType : public FSequencerTrackFilter
+{
+	// IFilter implementation
+	virtual bool PassesFilter(FTrackFilterType InItem) const override
+	{
+		if (InItem->IsA(ComponentType::StaticClass()))
+		{
+			return true;
+		}
+
+		const AActor* Actor = Cast<const AActor>(InItem);
+		if (Actor)
+		{
+			if (Actor->FindComponentByClass(ComponentType::StaticClass()))
+			{
+				return true;
+			}
+		}
+		return false;
+	}
+};
diff --git a/Engine/Source/Editor/Sequencer/Public/SequencerTrackFilterExtension.h b/Engine/Source/Editor/Sequencer/Public/SequencerTrackFilterExtension.h
new file mode 100644
--- /dev/null
+++ b/Engine/Source/Editor/Sequencer/Public/SequencerTrackFilterExtension.h
@@ -0,0 +1,20 @@
+// Copyright 1998-2019 Epic Games, Inc. All Rights Reserved.
+
+#pragma once
+
+#include "CoreMinimal.h"
+#include "UObject/ObjectMacros.h"
+#include "UObject/Object.h"
+#include "SequencerTrackFilterBase.h"
+#include "SequencerTrackFilterExtension.generated.h"
+
+// Override this class in order to make an additional track filter available in Sequencer
+UCLASS(Abstract)
+class SEQUENCER_API USequencerTrackFilterExtension : public UObject
+{
+	GENERATED_BODY()
+
+public:
+	// Override this method to add new track filters
+	virtual void AddTrackFilterExtensions(TArray< TSharedRef<class FSequencerTrackFilter> >& InOutFilterList) const { }
+};
