add missing header include to fix clang build

diff --git a/Engine/Source/Editor/Sequencer/Public/SequencerTrackFilterBase.h b/Engine/Source/Editor/Sequencer/Public/SequencerTrackFilterBase.h
--- a/Engine/Source/Editor/Sequencer/Public/SequencerTrackFilterBase.h
+++ b/Engine/Source/Editor/Sequencer/Public/SequencerTrackFilterBase.h
@@ -3,6 +3,7 @@
 #pragma once
 
 #include "CoreMinimal.h"
+#include "Textures/SlateIcon.h"
 #include "Misc/IFilter.h"
 #include "MovieSceneSequence.h"
 
