Sequencer: Fix static analysis

diff --git a/Engine/Source/Editor/Sequencer/Private/SSequencerPlayRateCombo.cpp b/Engine/Source/Editor/Sequencer/Private/SSequencerPlayRateCombo.cpp
--- a/Engine/Source/Editor/Sequencer/Private/SSequencerPlayRateCombo.cpp
+++ b/Engine/Source/Editor/Sequencer/Private/SSequencerPlayRateCombo.cpp
@@ -171,16 +171,19 @@ FText SSequencerPlayRateCombo::GetFrameRateMismatchErrorDescription() const
 		{
 			if (UMovieSceneSequence* SubSequence = Pair.Value.GetSequence())
 			{
-				FFrameRate SubDisplayRate = SubSequence->GetMovieScene()->GetDisplayRate();
-				if (SubSequence->GetMovieScene() && DisplayRate != SubDisplayRate)
+				if (SubSequence->GetMovieScene())
 				{
-					const FCommonFrameRateInfo* DisplayRateInfo = FCommonFrameRates::Find(DisplayRate);
-					const FCommonFrameRateInfo* SubDisplayRateInfo = FCommonFrameRates::Find(SubDisplayRate);
+					FFrameRate SubDisplayRate = SubSequence->GetMovieScene()->GetDisplayRate();
+					if (DisplayRate != SubDisplayRate)
+					{
+						const FCommonFrameRateInfo* DisplayRateInfo = FCommonFrameRates::Find(DisplayRate);
+						const FCommonFrameRateInfo* SubDisplayRateInfo = FCommonFrameRates::Find(SubDisplayRate);
 
-					FText DisplayRateText = DisplayRateInfo ? DisplayRateInfo->DisplayName : FText::Format(LOCTEXT("DisplayRateFormat", "{0} fps"), DisplayRate.AsDecimal());
-					FText SubDisplayRateText = SubDisplayRateInfo ? SubDisplayRateInfo->DisplayName : FText::Format(LOCTEXT("SubDisplayRateFormat", "{0} fps"), SubDisplayRate.AsDecimal());
+						FText DisplayRateText = DisplayRateInfo ? DisplayRateInfo->DisplayName : FText::Format(LOCTEXT("DisplayRateFormat", "{0} fps"), DisplayRate.AsDecimal());
+						FText SubDisplayRateText = SubDisplayRateInfo ? SubDisplayRateInfo->DisplayName : FText::Format(LOCTEXT("SubDisplayRateFormat", "{0} fps"), SubDisplayRate.AsDecimal());
 
-					return FText::Format(LOCTEXT("FrameRateMismatchDescription", "At least one mismatch in display rate: {0} is at {1} and {2} is at {3}"), Sequencer->GetRootMovieSceneSequence()->GetDisplayName(), DisplayRateText, SubSequence->GetDisplayName(), SubDisplayRateText);
+						return FText::Format(LOCTEXT("FrameRateMismatchDescription", "At least one mismatch in display rate: {0} is at {1} and {2} is at {3}"), Sequencer->GetRootMovieSceneSequence()->GetDisplayName(), DisplayRateText, SubSequence->GetDisplayName(), SubDisplayRateText);
+					}
 				}
 			}
 		}
