add missing header include to fix build issue

diff --git a/Engine/Source/Editor/Sequencer/Public/SequencerTrackFilterBase.h b/Engine/Source/Editor/Sequencer/Public/SequencerTrackFilterBase.h
--- a/Engine/Source/Editor/Sequencer/Public/SequencerTrackFilterBase.h
+++ b/Engine/Source/Editor/Sequencer/Public/SequencerTrackFilterBase.h
@@ -4,6 +4,7 @@
 
 #include "CoreMinimal.h"
 #include "Textures/SlateIcon.h"
+#include "GameFramework/Actor.h"
 #include "Misc/IFilter.h"
 #include "MovieSceneSequence.h"
 
