Reimplement converting possessables to spawnables in sequencer to more correctly handle multiple bindings.

diff --git a/Engine/Source/Editor/Sequencer/Private/Sequencer.cpp b/Engine/Source/Editor/Sequencer/Private/Sequencer.cpp
--- a/Engine/Source/Editor/Sequencer/Private/Sequencer.cpp
+++ b/Engine/Source/Editor/Sequencer/Private/Sequencer.cpp
@@ -7646,85 +7646,199 @@ void FSequencer::ConvertSelectedNodesToSpawnables()
 	NotifyMovieSceneDataChanged( EMovieSceneDataChangeType::MovieSceneStructureItemsChanged );
 }
 
-TArray<FMovieSceneSpawnable*> FSequencer::ConvertToSpawnableInternal(FGuid PossessableGuid)
+TArray<FGuid> FSequencer::ExpandMultiplePossessableBindings(FGuid PossessableGuid)
 {
 	UMovieSceneSequence* Sequence = GetFocusedMovieSceneSequence();
 	UMovieScene* MovieScene = Sequence->GetMovieScene();
 
+	TArray<FGuid> NewPossessableGuids;
+
 	if (MovieScene->IsReadOnly())
 	{
-		return TArray<FMovieSceneSpawnable*>();
+		return TArray<FGuid>();
 	}
 
-	TArrayView<TWeakObjectPtr<>> FoundObjects = FindBoundObjects(PossessableGuid, ActiveTemplateIDs.Top());
+	
+	// Create a copy of the TArrayView of bound objects, as the underlying array will get destroyed
+	TArray<TWeakObjectPtr<>> FoundObjects;
+	for (TWeakObjectPtr<> BoundObject : FindBoundObjects(PossessableGuid, ActiveTemplateIDs.Top()))
+	{
+		FoundObjects.Insert(BoundObject,0);
+	}
+
+	if (FoundObjects.Num() < 2)
+	{
+		// If less than two objects, nothing to do, return the same Guid
+		NewPossessableGuids.Add(PossessableGuid);
+		return NewPossessableGuids;
+	}
 
 	Sequence->Modify();
 	MovieScene->Modify();
 
-	// Locate the folder containing the original possessable
-	UMovieSceneFolder* ParentFolder;
-	for (UMovieSceneFolder* Folder : MovieScene->GetRootFolders())
+	FMovieSceneBinding* PossessableBinding = (FMovieSceneBinding*)MovieScene->GetBindings().FindByPredicate([&](FMovieSceneBinding& Binding) { return Binding.GetObjectGuid() == PossessableGuid; });
+
+	// First gather the children
+	TArray<FGuid> ChildPossessableGuids;
+	for (int32 Index = 0; Index < MovieScene->GetPossessableCount(); ++Index)
 	{
-		ParentFolder = Folder->FindFolderContaining(PossessableGuid);
-		if (ParentFolder != nullptr)
+		FMovieScenePossessable& Possessable = MovieScene->GetPossessable(Index);
+		if (Possessable.GetParent() == PossessableGuid)
 		{
-			break;
+			ChildPossessableGuids.Add(Possessable.GetGuid());
 		}
 	}
 
-	TArray<FMovieSceneSpawnable*> CreatedSpawnables;
+	TArray<UMovieSceneTrack* > Tracks = PossessableBinding->StealTracks();
 
-	if (Sequence->AllowsSpawnableObjects())
+	// Remove binding to stop any children from claiming the old guid as their parent
+	if (MovieScene->RemovePossessable(PossessableGuid))
 	{
-		
-		FMovieSceneBinding* PossessableBinding = MovieScene->FindBinding(PossessableGuid);
-		check(PossessableBinding);
+		Sequence->UnbindPossessableObjects(PossessableGuid);
+	}
+
+	for (TWeakObjectPtr<> FoundObjectPtr : FoundObjects)
+	{
+		UObject* FoundObject = FoundObjectPtr.Get();
+		check(FoundObject);
 
-		TArray<UMovieSceneTrack* > Tracks = PossessableBinding->GetTracks();
+		FoundObject->Modify();
 
-		int32 SortingOrder = PossessableBinding->GetSortingOrder();
+		UObject* BindingContext = GetPlaybackContext();
 
-		for (TWeakObjectPtr<> FoundObjectPtr : FoundObjects)
+		// Find this object's parent object, if it has one.
+		UObject* ParentObject = Sequence->GetParentObject(FoundObject);
+		if (ParentObject)
 		{
-			if (FoundObjectPtr.IsValid())
+			BindingContext = ParentObject;
+		}
+
+		// Create a new Possessable for this object
+		AActor* PossessedActor = Cast<AActor>(FoundObject);
+		const FGuid NewPossessableGuid = MovieScene->AddPossessable(PossessedActor != nullptr ? PossessedActor->GetActorLabel() : FoundObject->GetName(), FoundObject->GetClass());
+		FMovieScenePossessable* NewPossessable = MovieScene->FindPossessable(NewPossessableGuid);
+		if (NewPossessable)
+		{
+			FMovieSceneBinding* NewPossessableBinding = (FMovieSceneBinding*)MovieScene->GetBindings().FindByPredicate([&](FMovieSceneBinding& Binding) { return Binding.GetObjectGuid() == NewPossessableGuid; });
+			
+			if (ParentObject)
 			{
-				UObject* FoundObject = FoundObjectPtr.Get();
-				FMovieSceneSpawnable* Spawnable = MovieScene->FindSpawnable(AddSpawnable(*FoundObject));
-				if (Spawnable)
-				{
-					CreatedSpawnables.Add(Spawnable);
+				FGuid ParentGuid = FindObjectId(*ParentObject, ActiveTemplateIDs.Top());
+				NewPossessable->SetParent(ParentGuid);
+			}
 
-					FGuid SpawnableGuid = Spawnable->GetGuid();
+			Sequence->BindPossessableObject(NewPossessableGuid, *FoundObject, BindingContext);
+			NewPossessableGuids.Add(NewPossessableGuid);
 
-					FMovieSceneBinding* SpawnableBinding = MovieScene->FindBinding(SpawnableGuid);
-					check(SpawnableBinding);
+			// Create copies of the tracks
+			for (UMovieSceneTrack* Track : Tracks)
+			{
+				UMovieSceneTrack* DuplicatedTrack = Cast<UMovieSceneTrack>(StaticDuplicateObject(Track, MovieScene));
+				NewPossessableBinding->AddTrack(*DuplicatedTrack);
+			}
+		}
+	}
 
-					SpawnableBinding->SetSortingOrder(SortingOrder);
+	// Finally, recurse in to any children
+	for (FGuid ChildPossessableGuid : ChildPossessableGuids)
+	{
+		ExpandMultiplePossessableBindings(ChildPossessableGuid);
+	}
 
-					AActor* PossessedActor = Cast<AActor>(FoundObject);
-					Spawnable->AddChildPossessable(CreateBinding(*FoundObject, PossessedActor != nullptr ? PossessedActor->GetActorLabel() : FoundObject->GetName()));
-					
-					for (UMovieSceneTrack* Track : Tracks)
-					{
-						UMovieSceneTrack* DuplicatedTrack = Cast<UMovieSceneTrack>(StaticDuplicateObject(Track, MovieScene));
-						SpawnableBinding->AddTrack(*DuplicatedTrack);
-					}
+	NotifyMovieSceneDataChanged(EMovieSceneDataChangeType::MovieSceneStructureItemsChanged);
 
-					TOptional<FTransformData> TransformData;
-					SpawnRegister->HandleConvertPossessableToSpawnable(FoundObject, *this, TransformData);
-					SpawnRegister->SetupDefaultsForSpawnable(nullptr, Spawnable->GetGuid(), TransformData, AsShared(), Settings);
-				}
-			}
-		}
+	return NewPossessableGuids;
+}
 
-		if (MovieScene->RemovePossessable(PossessableGuid))
+TArray<FMovieSceneSpawnable*> FSequencer::ConvertToSpawnableInternal(FGuid PossessableGuid)
+{
+	UMovieSceneSequence* Sequence = GetFocusedMovieSceneSequence();
+	UMovieScene* MovieScene = Sequence->GetMovieScene();
+
+	if (MovieScene->IsReadOnly() || !Sequence->AllowsSpawnableObjects())
+	{
+		return TArray<FMovieSceneSpawnable*>();
+	}
+
+	TArrayView<TWeakObjectPtr<>> FoundObjects = FindBoundObjects(PossessableGuid, ActiveTemplateIDs.Top());
+
+	TArray<FMovieSceneSpawnable*> CreatedSpawnables;
+
+	if (FoundObjects.Num() > 1)
+	{
+		// Expand to individual possessables for each bound object, then convert each one individually
+		TArray<FGuid> ExpandedPossessableGuids = ExpandMultiplePossessableBindings(PossessableGuid);
+		for (FGuid NewPossessableGuid : ExpandedPossessableGuids)
 		{
-			Sequence->UnbindPossessableObjects(PossessableGuid);
+			CreatedSpawnables.Append(ConvertToSpawnableInternal(NewPossessableGuid));
 		}
 
 		ForceEvaluate();
 	}
+	else
+	{
+		UObject* FoundObject = FoundObjects[0].Get();
+		if (!FoundObject)
+		{
+			return TArray<FMovieSceneSpawnable*>();
+		}
+
+		Sequence->Modify();
+		MovieScene->Modify();
+
+		// Locate the folder containing the original possessable
+		UMovieSceneFolder* ParentFolder;
+		for (UMovieSceneFolder* Folder : MovieScene->GetRootFolders())
+		{
+			ParentFolder = Folder->FindFolderContaining(PossessableGuid);
+			if (ParentFolder != nullptr)
+			{
+				break;
+			}
+		}
+
+		FMovieSceneSpawnable* Spawnable = MovieScene->FindSpawnable(AddSpawnable(*FoundObject));
+		if (Spawnable)
+		{
+			CreatedSpawnables.Add(Spawnable);
+			FGuid SpawnableGuid = Spawnable->GetGuid();
+
+			// Remap all the spawnable's tracks and child bindings onto the new possessable
+			MovieScene->MoveBindingContents(PossessableGuid, SpawnableGuid);
+
+			FMovieSceneBinding* PossessableBinding = (FMovieSceneBinding*)MovieScene->GetBindings().FindByPredicate([&](FMovieSceneBinding& Binding) { return Binding.GetObjectGuid() == PossessableGuid; });
+			check(PossessableBinding);
+
+			for (UMovieSceneFolder* Folder : MovieScene->GetRootFolders())
+			{
+				if (ReplaceFolderBindingGUID(Folder, PossessableGuid, SpawnableGuid))
+				{
+					break;
+				}
+			}
+
+			int32 SortingOrder = PossessableBinding->GetSortingOrder();
 
+			if (MovieScene->RemovePossessable(PossessableGuid))
+			{
+				Sequence->UnbindPossessableObjects(PossessableGuid);
+
+				FMovieSceneBinding* SpawnableBinding = (FMovieSceneBinding*)MovieScene->GetBindings().FindByPredicate([&](FMovieSceneBinding& Binding) { return Binding.GetObjectGuid() == SpawnableGuid; });
+				check(SpawnableBinding);
+
+				SpawnableBinding->SetSortingOrder(SortingOrder);
+
+			}
+
+			TOptional<FTransformData> TransformData;
+			SpawnRegister->HandleConvertPossessableToSpawnable(FoundObject, *this, TransformData);
+			SpawnRegister->SetupDefaultsForSpawnable(nullptr, Spawnable->GetGuid(), TransformData, AsShared(), Settings);
+
+			ForceEvaluate();
+
+			NotifyMovieSceneDataChanged(EMovieSceneDataChangeType::MovieSceneStructureItemsChanged);
+		}
+	}
 
 	return CreatedSpawnables;
 }
diff --git a/Engine/Source/Editor/Sequencer/Private/Sequencer.h b/Engine/Source/Editor/Sequencer/Private/Sequencer.h
--- a/Engine/Source/Editor/Sequencer/Private/Sequencer.h
+++ b/Engine/Source/Editor/Sequencer/Private/Sequencer.h
@@ -951,6 +951,9 @@ protected:
 	/** Updates a viewport client from camera cut data */
 	void UpdatePreviewLevelViewportClientFromCameraCut(FLevelEditorViewportClient& InViewportClient, UObject* InCameraObject, bool bJumpCut) const;
 
+	/** Expands Possessables with multiple bindings into indidual Possessables for each binding */
+	TArray<FGuid> ExpandMultiplePossessableBindings(FGuid PossessableGuid);
+
 	/** Internal conversion function that doesn't perform expensive reset/update tasks */
 	TArray<FMovieSceneSpawnable*> ConvertToSpawnableInternal(FGuid PossessableGuid);
 
