Default SequencerTrackFilters to only support level sequences, only enable audio and event filters for UMG.

diff --git a/Engine/Source/Editor/Sequencer/Private/SequencerTrackFilters.h b/Engine/Source/Editor/Sequencer/Private/SequencerTrackFilters.h
--- a/Engine/Source/Editor/Sequencer/Private/SequencerTrackFilters.h
+++ b/Engine/Source/Editor/Sequencer/Private/SequencerTrackFilters.h
@@ -85,6 +85,14 @@ class FSequencerTrackFilter_AudioTracks : public FSequencerTrackFilter_ClassType
 	virtual FText GetToolTipText() const override { return LOCTEXT("SequencerTrackFilter_AudioTracksToolTip", "Show only Audio tracks."); }
 	virtual FSlateIcon GetIcon() const override { return FSlateIcon(FEditorStyle::GetStyleSetName(), "Sequencer.Tracks.Audio"); }
 
+	virtual bool SupportsSequence(UMovieSceneSequence* InSequence) const override
+	{
+		static UClass* LevelSequenceClass = FindObject<UClass>(ANY_PACKAGE, TEXT("LevelSequence"), true);
+		static UClass* WidgetAnimationClass = FindObject<UClass>(ANY_PACKAGE, TEXT("WidgetAnimation"), true);
+		return InSequence != nullptr &&
+			((LevelSequenceClass != nullptr && InSequence->GetClass()->IsChildOf(LevelSequenceClass)) ||
+			(WidgetAnimationClass != nullptr && InSequence->GetClass()->IsChildOf(WidgetAnimationClass)));
+	}
 };
 
 class FSequencerTrackFilter_EventTracks : public FSequencerTrackFilter_ClassType< UMovieSceneEventTrack >
@@ -93,6 +101,15 @@ class FSequencerTrackFilter_EventTracks : public FSequencerTrackFilter_ClassType
 	virtual FText GetDisplayName() const override { return LOCTEXT("SequencerTrackFilter_EventTracks", "Event"); }
 	virtual FText GetToolTipText() const override { return LOCTEXT("SequencerTrackFilter_EventTracksToolTip", "Show only Event tracks."); }
 	virtual FSlateIcon GetIcon() const override { return FSlateIcon(FEditorStyle::GetStyleSetName(), "Sequencer.Tracks.Event"); }
+
+	virtual bool SupportsSequence(UMovieSceneSequence* InSequence) const override
+	{
+		static UClass* LevelSequenceClass = FindObject<UClass>(ANY_PACKAGE, TEXT("LevelSequence"), true);
+		static UClass* WidgetAnimationClass = FindObject<UClass>(ANY_PACKAGE, TEXT("WidgetAnimation"), true);
+		return InSequence != nullptr &&
+			((LevelSequenceClass != nullptr && InSequence->GetClass()->IsChildOf(LevelSequenceClass)) ||
+			(WidgetAnimationClass != nullptr && InSequence->GetClass()->IsChildOf(WidgetAnimationClass)));
+	}
 };
 
 class FSequencerTrackFilter_LevelVisibilityTracks : public FSequencerTrackFilter_ClassType< UMovieSceneLevelVisibilityTrack >
diff --git a/Engine/Source/Editor/Sequencer/Public/SequencerTrackFilterBase.h b/Engine/Source/Editor/Sequencer/Public/SequencerTrackFilterBase.h
--- a/Engine/Source/Editor/Sequencer/Public/SequencerTrackFilterBase.h
+++ b/Engine/Source/Editor/Sequencer/Public/SequencerTrackFilterBase.h
@@ -31,10 +31,7 @@ public:
 	virtual bool SupportsSequence(UMovieSceneSequence* InSequence) const
 	{
 		static UClass* LevelSequenceClass = FindObject<UClass>(ANY_PACKAGE, TEXT("LevelSequence"), true);
-		static UClass* WidgetAnimationClass = FindObject<UClass>(ANY_PACKAGE, TEXT("WidgetAnimation"), true);
-		return InSequence != nullptr &&
-			((LevelSequenceClass != nullptr && InSequence->GetClass()->IsChildOf(LevelSequenceClass)) ||
-			(WidgetAnimationClass != nullptr && InSequence->GetClass()->IsChildOf(WidgetAnimationClass)));
+		return InSequence != nullptr && (LevelSequenceClass != nullptr && InSequence->GetClass()->IsChildOf(LevelSequenceClass));
 	}
 
 	// IFilter implementation
