Sequencer: Skip coloring for primitive material track

diff --git a/Engine/Source/Editor/Sequencer/Private/DisplayNodes/SequencerTrackNode.cpp b/Engine/Source/Editor/Sequencer/Private/DisplayNodes/SequencerTrackNode.cpp
--- a/Engine/Source/Editor/Sequencer/Private/DisplayNodes/SequencerTrackNode.cpp
+++ b/Engine/Source/Editor/Sequencer/Private/DisplayNodes/SequencerTrackNode.cpp
@@ -28,6 +28,7 @@
 #include "MovieSceneFolder.h"
 #include "Tracks/MovieScenePropertyTrack.h"
 #include "Tracks/MovieScene3DTransformTrack.h"
+#include "Tracks/MovieScenePrimitiveMaterialTrack.h"
 #include "MovieSceneCommonHelpers.h"
 
 #define LOCTEXT_NAMESPACE "SequencerTrackNode"
@@ -486,7 +487,7 @@ FLinearColor FSequencerTrackNode::GetDisplayNameColor() const
 	if (UMovieScenePropertyTrack* PropertyTrack = Cast<UMovieScenePropertyTrack>(Track))
 	{
 		// 3D transform tracks don't map to property bindings as below
-		if (Track->IsA<UMovieScene3DTransformTrack>())
+		if (Track->IsA<UMovieScene3DTransformTrack>() || Track->IsA<UMovieScenePrimitiveMaterialTrack>())
 		{
 			return FLinearColor::White;
 		}
