Sequencer: Clear/restore focus only on SetLocalTimeDirectly and change to weak ptr

diff --git a/Engine/Source/Editor/Sequencer/Private/Sequencer.cpp b/Engine/Source/Editor/Sequencer/Private/Sequencer.cpp
--- a/Engine/Source/Editor/Sequencer/Private/Sequencer.cpp
+++ b/Engine/Source/Editor/Sequencer/Private/Sequencer.cpp
@@ -2419,14 +2419,7 @@ void FSequencer::SetLocalTime( FFrameTime NewTime, ESnapTimeMode SnapTimeMode)
 
 void FSequencer::SetLocalTimeDirectly(FFrameTime NewTime)
 {
-	// Transform the time to the root time-space
-	SetGlobalTime(NewTime * RootToLocalTransform.Inverse());
-}
-
-
-void FSequencer::SetGlobalTime(FFrameTime NewTime)
-{
-	TSharedPtr<SWidget> PreviousFocusedWidget = FSlateApplication::Get().GetKeyboardFocusedWidget();
+	TWeakPtr<SWidget> PreviousFocusedWidget = FSlateApplication::Get().GetKeyboardFocusedWidget();
 
 	// Clear focus before setting time in case there's a key editor value selected that gets committed to a newly selected key on UserMovedFocus
 	if (GetPlaybackStatus() == EMovieScenePlayerStatus::Stopped)
@@ -2434,6 +2427,18 @@ void FSequencer::SetGlobalTime(FFrameTime NewTime)
 		FSlateApplication::Get().ClearKeyboardFocus(EFocusCause::Cleared);
 	}
 
+	// Transform the time to the root time-space
+	SetGlobalTime(NewTime * RootToLocalTransform.Inverse());
+
+	if (PreviousFocusedWidget.IsValid())
+	{
+		FSlateApplication::Get().SetKeyboardFocus(PreviousFocusedWidget.Pin());
+	}
+}
+
+
+void FSequencer::SetGlobalTime(FFrameTime NewTime)
+{
 	NewTime = ConvertFrameTime(NewTime, GetRootTickResolution(), PlayPosition.GetInputRate());
 	if (PlayPosition.GetEvaluationType() == EMovieSceneEvaluationType::FrameLocked)
 	{
@@ -2453,11 +2458,6 @@ void FSequencer::SetGlobalTime(FFrameTime NewTime)
 		SetPlaybackStatus(EMovieScenePlayerStatus::Stopped);
 		AutoScrubTarget.Reset();
 	}
-
-	if (PreviousFocusedWidget.IsValid())
-	{
-		FSlateApplication::Get().SetKeyboardFocus(PreviousFocusedWidget);
-	}
 }
 
 void FSequencer::ForceEvaluate()
