Sequencer: Minor fix for reached end condition when current time is already at the start or end.

diff --git a/Engine/Source/Editor/Sequencer/Private/Sequencer.cpp b/Engine/Source/Editor/Sequencer/Private/Sequencer.cpp
--- a/Engine/Source/Editor/Sequencer/Private/Sequencer.cpp
+++ b/Engine/Source/Editor/Sequencer/Private/Sequencer.cpp
@@ -3884,11 +3884,11 @@ void FSequencer::SetLocalTimeLooped(FFrameTime NewLocalTime)
 		bool bReachedEnd = false;
 		if (PlaybackSpeed > 0)
 		{
-			bReachedEnd = GetLocalTime().Time < MaxInclusiveTime && NewLocalTime >= MaxInclusiveTime;
+			bReachedEnd = GetLocalTime().Time <= MaxInclusiveTime && NewLocalTime >= MaxInclusiveTime;
 		}
 		else
 		{
-			bReachedEnd = GetLocalTime().Time > MinInclusiveTime && NewLocalTime <= MinInclusiveTime;
+			bReachedEnd = GetLocalTime().Time >= MinInclusiveTime && NewLocalTime <= MinInclusiveTime;
 		}
 
 		// Stop if we hit the playback range end
