Filter object binding sequencer tracks by sublevel

diff --git a/Engine/Source/Editor/Sequencer/Private/SSequencer.cpp b/Engine/Source/Editor/Sequencer/Private/SSequencer.cpp
--- a/Engine/Source/Editor/Sequencer/Private/SSequencer.cpp
+++ b/Engine/Source/Editor/Sequencer/Private/SSequencer.cpp
@@ -1146,6 +1146,16 @@ TSharedRef<SWidget> SSequencer::MakeFilterMenu()
 	}
 	MenuBuilder.EndSection();
 
+	UObject* PlaybackContext = Sequencer->GetPlaybackContext();
+	UWorld* World = PlaybackContext ? PlaybackContext->GetWorld() : nullptr;
+
+	if (World && World->GetLevels().Num() > 1)
+	{
+		MenuBuilder.BeginSection("TrackLevelFilters");
+		MenuBuilder.AddSubMenu(LOCTEXT("LevelFilters", "Level Filters"), LOCTEXT("LevelFiltersToolTip", "Filter object tracks by level"), FNewMenuDelegate::CreateRaw(this, &SSequencer::FillLevelFilterMenu));
+		MenuBuilder.EndSection();
+	}
+
 	MenuBuilder.BeginSection("TrackFilters");
 	
 	for (TSharedRef<FSequencerTrackFilter> TrackFilter : AllTrackFilters)
@@ -1171,6 +1181,33 @@ TSharedRef<SWidget> SSequencer::MakeFilterMenu()
 	return MenuBuilder.MakeWidget();
 }
 
+void SSequencer::FillLevelFilterMenu(FMenuBuilder& InMenuBarBuilder)
+{
+	TSharedPtr<FSequencer> Sequencer = SequencerPtr.Pin();
+	UObject* PlaybackContext = Sequencer->GetPlaybackContext();
+	UWorld* World = PlaybackContext ? PlaybackContext->GetWorld() : nullptr;
+
+	if (World)
+	{
+		const TArray<ULevel*> Levels = World->GetLevels();
+		for (ULevel* Level : Levels)
+		{
+			FString LevelName = FPackageName::GetShortName(Level->GetOutermost()->GetName());
+			InMenuBarBuilder.AddMenuEntry(
+				FText::FromString(LevelName),
+				FText::FromString(Level->GetOutermost()->GetName()),
+				FSlateIcon(),
+				FUIAction(
+					FExecuteAction::CreateSP(this, &SSequencer::OnTrackLevelFilterClicked, LevelName),
+					FCanExecuteAction(),
+					FIsActionChecked::CreateSP(this, &SSequencer::IsTrackLevelFilterActive, LevelName)),
+				NAME_None,
+				EUserInterfaceActionType::ToggleButton
+			);
+		}
+	}
+}
+
 void SSequencer::OnResetFilters()
 {
 	TSharedPtr<FSequencer> Sequencer = SequencerPtr.Pin();
@@ -1196,6 +1233,25 @@ bool SSequencer::IsTrackFilterActive(TSharedRef<FSequencerTrackFilter> TrackFilt
 	return Sequencer->GetNodeTree()->IsTrackFilterActive(TrackFilter);
 }
 
+void SSequencer::OnTrackLevelFilterClicked(const FString LevelName)
+{
+	TSharedPtr<FSequencer> Sequencer = SequencerPtr.Pin();
+	if (IsTrackLevelFilterActive(LevelName))
+	{
+		Sequencer->GetNodeTree()->RemoveLevelFilter(LevelName);
+	}
+	else
+	{
+		Sequencer->GetNodeTree()->AddLevelFilter(LevelName);
+	}
+}
+
+bool SSequencer::IsTrackLevelFilterActive(const FString LevelName) const
+{
+	TSharedPtr<FSequencer> Sequencer = SequencerPtr.Pin();
+	return Sequencer->GetNodeTree()->IsTrackLevelFilterActive(LevelName);
+}
+
 TSharedRef<SWidget> SSequencer::MakeGeneralMenu()
 {
 	FMenuBuilder MenuBuilder( true, SequencerPtr.Pin()->GetCommandBindings() );
diff --git a/Engine/Source/Editor/Sequencer/Private/SSequencer.h b/Engine/Source/Editor/Sequencer/Private/SSequencer.h
--- a/Engine/Source/Editor/Sequencer/Private/SSequencer.h
+++ b/Engine/Source/Editor/Sequencer/Private/SSequencer.h
@@ -367,6 +367,11 @@ private:
 	void OnTrackFilterClicked(TSharedRef<FSequencerTrackFilter> TrackFilter);
 	bool IsTrackFilterActive(TSharedRef<FSequencerTrackFilter> TrackFilter) const;
 
+	void OnTrackLevelFilterClicked(const FString LevelName);
+	bool IsTrackLevelFilterActive(const FString LevelName) const;
+
+	void FillLevelFilterMenu(FMenuBuilder& InMenuBarBuilder);
+
 	/**
 	* Called when the time snap interval changes.
 	*
diff --git a/Engine/Source/Editor/Sequencer/Private/SequencerNodeTree.cpp b/Engine/Source/Editor/Sequencer/Private/SequencerNodeTree.cpp
--- a/Engine/Source/Editor/Sequencer/Private/SequencerNodeTree.cpp
+++ b/Engine/Source/Editor/Sequencer/Private/SequencerNodeTree.cpp
@@ -26,12 +26,18 @@ FSequencerNodeTree::~FSequencerNodeTree()
 	{
 		TrackFilters->OnChanged().RemoveAll(this);
 	}
+	if (TrackFilterLevelFilter.IsValid())
+	{
+		TrackFilterLevelFilter->OnChanged().RemoveAll(this);
+	}
 }
 
 FSequencerNodeTree::FSequencerNodeTree(class FSequencer& InSequencer) : Sequencer(InSequencer), bFilterUpdateRequested(false)
 {
 	TrackFilters = MakeShared<FSequencerTrackFilterCollection>();
 	TrackFilters->OnChanged().AddRaw(this, &FSequencerNodeTree::RequestFilterUpdate);
+	TrackFilterLevelFilter = MakeShared< FSequencerTrackFilter_LevelFilter>();
+	TrackFilterLevelFilter->OnChanged().AddRaw(this, &FSequencerNodeTree::RequestFilterUpdate);
 }
 
 void FSequencerNodeTree::Empty()
@@ -319,7 +325,7 @@ const TArray<TSharedRef<FSequencerDisplayNode>>& FSequencerNodeTree::GetRootNode
 
 bool FSequencerNodeTree::HasActiveFilter() const
 {
-	return (!FilterString.IsEmpty() || TrackFilters->Num() > 0);
+	return (!FilterString.IsEmpty() || TrackFilters->Num() > 0 || TrackFilterLevelFilter->IsActive());
 }
 
 TSharedRef<FSequencerObjectBindingNode> FSequencerNodeTree::AddObjectBinding(const FString& ObjectName, const FGuid& ObjectBinding, TMap<FGuid, const FMovieSceneBinding*>& GuidToBindingMap, TArray<TSharedRef<FSequencerObjectBindingNode>>& OutNodeList)
@@ -568,6 +574,7 @@ int32 FSequencerNodeTree::RemoveFilter(TSharedRef<FSequencerTrackFilter> TrackFi
 void FSequencerNodeTree::RemoveAllFilters()
 {
 	TrackFilters->RemoveAll();
+	TrackFilterLevelFilter->ResetFilter();
 }
 
 bool FSequencerNodeTree::IsTrackFilterActive(TSharedRef<FSequencerTrackFilter> TrackFilter) const
@@ -575,6 +582,22 @@ bool FSequencerNodeTree::IsTrackFilterActive(TSharedRef<FSequencerTrackFilter> T
 	return TrackFilters->Contains(TrackFilter);
 }
 
+void FSequencerNodeTree::AddLevelFilter(const FString& LevelName)
+{
+	TrackFilterLevelFilter->UnhideLevel(LevelName);
+}
+
+void FSequencerNodeTree::RemoveLevelFilter(const FString& LevelName)
+{
+	TrackFilterLevelFilter->HideLevel(LevelName);
+}
+
+bool FSequencerNodeTree::IsTrackLevelFilterActive(const FString& LevelName) const
+{
+	return !TrackFilterLevelFilter->IsLevelHidden(LevelName);
+}
+
+
 void FSequencerNodeTree::SaveExpansionState(const FSequencerDisplayNode& Node, bool bExpanded)
 {	
 	// @todo Sequencer - This should be moved to the sequence level
@@ -707,20 +730,49 @@ static void AddParentNodes(const TSharedRef<FSequencerDisplayNode>& StartNode, T
  * @param OutFilteredNodes	The list of all filtered nodes
  */
 
-static void FilterNodesRecursive( FSequencer& Sequencer, const TSharedRef<FSequencerDisplayNode>& StartNode, TSharedPtr<FSequencerTrackFilterCollection> Filters, const TArray<FString>& FilterStrings, TSet<TSharedRef<FSequencerDisplayNode>>& OutFilteredNodes )
+static void FilterNodesRecursive( FSequencer& Sequencer, const TSharedRef<FSequencerDisplayNode>& StartNode, TSharedPtr<FSequencerTrackFilterCollection> Filters, const TArray<FString>& FilterStrings, TSharedPtr<FSequencerTrackFilter_LevelFilter> LevelTrackFilter, TSet<TSharedRef<FSequencerDisplayNode>>& OutFilteredNodes )
 {
 	for (TSharedRef<FSequencerDisplayNode> Node : StartNode->GetChildNodes())
 	{
-		FilterNodesRecursive(Sequencer, Node, Filters, FilterStrings, OutFilteredNodes);
+		FilterNodesRecursive(Sequencer, Node, Filters, FilterStrings, LevelTrackFilter, OutFilteredNodes);
 	}
 
 	bool bPasssedAnyFilters = false;
+
 	if (StartNode->GetType() == ESequencerNode::Track)
 	{
 		UMovieSceneTrack* Track = static_cast<const FSequencerTrackNode&>(StartNode.Get()).GetTrack();
-		if (Filters->PassesAnyFilters(Track))
+		if (Filters->Num() == 0 || Filters->PassesAnyFilters(Track))
 		{
 			bPasssedAnyFilters = true;
+
+			// Track nodes do not belong to a level, but might be a child of an objectbinding node that does
+			if (LevelTrackFilter->IsActive())
+			{
+				TSharedPtr<const FSequencerDisplayNode> ParentNode = StartNode->GetParent();
+				while (ParentNode.IsValid())
+				{
+					if (ParentNode->GetType() == ESequencerNode::Object)
+					{
+						// The track belongs to an objectbinding node, start by assuming it doesn't match the level filter
+						bPasssedAnyFilters = false;
+
+						const FSequencerObjectBindingNode* ObjectNode = static_cast<const FSequencerObjectBindingNode*>(ParentNode.Get());
+						for (TWeakObjectPtr<>& Object : Sequencer.FindObjectsInCurrentSequence(ObjectNode->GetObjectBinding()))
+						{
+							if (LevelTrackFilter->PassesFilter(Object.Get()))
+							{
+								// If at least one of the objects on the objectbinding node pass the level filter, show the track
+								bPasssedAnyFilters = true;
+								break;
+							}
+						}
+
+						break;
+					}
+					ParentNode = ParentNode->GetParent();
+				}
+			}
 		}
 	}
 	else if (StartNode->GetType() == ESequencerNode::Object)
@@ -728,7 +780,8 @@ static void FilterNodesRecursive( FSequencer& Sequencer, const TSharedRef<FSeque
 		const FSequencerObjectBindingNode ObjectNode = static_cast<const FSequencerObjectBindingNode&>(StartNode.Get());
 		for (TWeakObjectPtr<>& Object : Sequencer.FindObjectsInCurrentSequence(ObjectNode.GetObjectBinding()))
 		{
-			if (Filters->PassesAnyFilters(Object.Get()))
+			if ((Filters->Num() == 0 || Filters->PassesAnyFilters(Object.Get()))
+				&& LevelTrackFilter->PassesFilter(Object.Get()))
 			{
 				bPasssedAnyFilters = true;
 				break;
@@ -736,7 +789,7 @@ static void FilterNodesRecursive( FSequencer& Sequencer, const TSharedRef<FSeque
 		}
 	}
 
-	if (Filters->Num() == 0 || bPasssedAnyFilters)
+	if (bPasssedAnyFilters)
 	{
 		// If we have a filter string, make sure we match
 		if (FilterStrings.Num() > 0)
@@ -792,7 +845,9 @@ void FSequencerNodeTree::UpdateFilters()
 
 	FilteredNodes.Empty();
 
-	if (TrackFilters->Num() > 0 || !FilterString.IsEmpty())
+	TrackFilterLevelFilter->UpdateWorld(Sequencer.GetPlaybackContext()->GetWorld());
+
+	if (TrackFilters->Num() > 0 || !FilterString.IsEmpty() || TrackFilterLevelFilter->IsActive())
 	{
 		// Build a list of strings that must be matched
 		TArray<FString> FilterStrings;
@@ -803,12 +858,12 @@ void FSequencerNodeTree::UpdateFilters()
 
 		for (auto It = ObjectBindingMap.CreateIterator(); It; ++It)
 		{
-			FilterNodesRecursive(Sequencer, It.Value().ToSharedRef(), TrackFilters, FilterStrings, FilteredNodes);
+			FilterNodesRecursive(Sequencer, It.Value().ToSharedRef(), TrackFilters, FilterStrings, TrackFilterLevelFilter, FilteredNodes);
 		}
 
 		for (auto It = RootNodes.CreateIterator(); It; ++It)
 		{
-			FilterNodesRecursive(Sequencer, *It, TrackFilters, FilterStrings, FilteredNodes);
+			FilterNodesRecursive(Sequencer, *It, TrackFilters, FilterStrings, TrackFilterLevelFilter, FilteredNodes);
 		}
 	}
 
diff --git a/Engine/Source/Editor/Sequencer/Private/SequencerNodeTree.h b/Engine/Source/Editor/Sequencer/Private/SequencerNodeTree.h
--- a/Engine/Source/Editor/Sequencer/Private/SequencerNodeTree.h
+++ b/Engine/Source/Editor/Sequencer/Private/SequencerNodeTree.h
@@ -15,6 +15,7 @@ class UMovieSceneTrack;
 struct FMovieSceneBinding;
 class FSequencerTrackFilter;
 class FSequencerTrackFilterCollection;
+class FSequencerTrackFilter_LevelFilter;
 
 /**
  * Represents a tree of sequencer display nodes, used to populate the Sequencer UI with MovieScene data
@@ -133,6 +134,10 @@ public:
 	void RemoveAllFilters();
 	bool IsTrackFilterActive(TSharedRef<FSequencerTrackFilter> TrackFilter) const;
 
+	void AddLevelFilter(const FString& LevelName);
+	void RemoveLevelFilter(const FString& LevelName);
+	bool IsTrackLevelFilterActive(const FString& LevelName) const;
+
 private:
 
 	/**
@@ -200,6 +205,9 @@ private:
 
 	/** Active track filters */
 	TSharedPtr<FSequencerTrackFilterCollection> TrackFilters;
+	
+	/** Level based track filtering */
+	TSharedPtr<FSequencerTrackFilter_LevelFilter> TrackFilterLevelFilter;
 
 	bool bFilterUpdateRequested;
 };
diff --git a/Engine/Source/Editor/Sequencer/Private/SequencerTrackFilters.cpp b/Engine/Source/Editor/Sequencer/Private/SequencerTrackFilters.cpp
new file mode 100644
--- /dev/null
+++ b/Engine/Source/Editor/Sequencer/Private/SequencerTrackFilters.cpp
@@ -0,0 +1,113 @@
+// Copyright 1998-2019 Epic Games, Inc. All Rights Reserved.
+
+#include "SequencerTrackFilters.h"
+#include "Engine/World.h"
+
+FSequencerTrackFilter_LevelFilter::~FSequencerTrackFilter_LevelFilter()
+{
+	if (CachedWorld.IsValid())
+	{
+		CachedWorld.Get()->OnLevelsChanged().RemoveAll(this);
+		CachedWorld.Reset();
+	}
+}
+
+bool FSequencerTrackFilter_LevelFilter::PassesFilter(FTrackFilterType InItem) const
+{
+	// For anything in a level, outermost should refer to the ULevel that contains it
+	FString OutermostName = FPackageName::GetShortName(InItem->GetOutermost()->GetName());
+	
+	// Pass anything that is not in a hidden level
+	return !HiddenLevels.Contains(OutermostName);
+}
+
+void FSequencerTrackFilter_LevelFilter::ResetFilter()
+{
+	HiddenLevels.Empty();
+
+	BroadcastChangedEvent();
+}
+
+bool FSequencerTrackFilter_LevelFilter::IsLevelHidden(const FString& LevelName) const
+{
+	return HiddenLevels.Contains(LevelName);
+}
+
+void FSequencerTrackFilter_LevelFilter::HideLevel(const FString& LevelName)
+{
+	HiddenLevels.AddUnique(LevelName);
+
+	BroadcastChangedEvent();
+}
+
+void FSequencerTrackFilter_LevelFilter::UnhideLevel(const FString& LevelName)
+{
+	HiddenLevels.Remove(LevelName);
+
+	BroadcastChangedEvent();
+}
+
+void FSequencerTrackFilter_LevelFilter::UpdateWorld(UWorld* World)
+{
+	if (!CachedWorld.IsValid() || CachedWorld.Get() != World)
+	{
+		if (CachedWorld.IsValid())
+		{
+			CachedWorld.Get()->OnLevelsChanged().RemoveAll(this);
+		}
+		
+		CachedWorld.Reset();
+	
+		if (IsValid(World))
+		{
+			CachedWorld = World;
+			CachedWorld.Get()->OnLevelsChanged().AddRaw(this, &FSequencerTrackFilter_LevelFilter::HandleLevelsChanged);
+		}
+
+		HandleLevelsChanged();
+	}
+}
+
+void FSequencerTrackFilter_LevelFilter::HandleLevelsChanged()
+{
+	if (!CachedWorld.IsValid())
+	{
+		ResetFilter();
+		return;
+	}
+
+	const TArray<ULevel*>& WorldLevels = CachedWorld->GetLevels();
+	
+	if (WorldLevels.Num() < 2)
+	{
+		ResetFilter();
+		return;
+	}
+
+	// Build a list of level names contained in the current world
+	TArray<FString> WorldLevelNames;
+	for (const ULevel* Level : WorldLevels)
+	{
+		if (IsValid(Level))
+		{
+			FString LevelName = FPackageName::GetShortName(Level->GetOutermost()->GetName());
+			WorldLevelNames.Add(LevelName);
+		}
+	}
+
+	// Rebuild our list of hidden level names to only include levels which are still in the world
+	TArray<FString> OldHiddenLevels = HiddenLevels;
+	HiddenLevels.Empty();
+	for (FString LevelName : OldHiddenLevels)
+	{
+		if (WorldLevelNames.Contains(LevelName))
+		{
+			HiddenLevels.Add(LevelName);
+		}
+	}
+
+	if (OldHiddenLevels.Num() != HiddenLevels.Num())
+	{
+		BroadcastChangedEvent();
+	}
+}
\ No newline at end of file
diff --git a/Engine/Source/Editor/Sequencer/Private/SequencerTrackFilters.h b/Engine/Source/Editor/Sequencer/Private/SequencerTrackFilters.h
--- a/Engine/Source/Editor/Sequencer/Private/SequencerTrackFilters.h
+++ b/Engine/Source/Editor/Sequencer/Private/SequencerTrackFilters.h
@@ -21,6 +21,8 @@
 #include "Components/SkeletalMeshComponent.h"
 #include "Particles/ParticleSystem.h"
 
+class UWorld;
+
 #define LOCTEXT_NAMESPACE "Sequencer"
 
 class FSequencerTrackFilterCollection : public TFilterCollection<FTrackFilterType>
@@ -152,4 +154,33 @@ class FSequencerTrackFilter_LightObjects : public FSequencerTrackFilter_Componen
 	virtual FSlateIcon GetIcon() const override { return FSlateIcon(FEditorStyle::GetStyleSetName(), "ClassIcon.Light"); }
 };
 
+class FSequencerTrackFilter_LevelFilter : public FSequencerTrackFilter
+{
+public:
+	~FSequencerTrackFilter_LevelFilter();
+
+	virtual bool PassesFilter(FTrackFilterType InItem) const override;
+
+	virtual FString GetName() const override { return TEXT("SequencerSubLevelFilter"); }
+	virtual FText GetDisplayName() const override { return FText::GetEmpty(); }
+	virtual FText GetToolTipText() const override { return FText::GetEmpty(); }
+
+	void UpdateWorld(UWorld* World);
+	void ResetFilter();
+
+	bool IsActive() const { return HiddenLevels.Num() > 0; }
+
+	bool IsLevelHidden(const FString& LevelName) const;
+	void HideLevel(const FString& LevelName);
+	void UnhideLevel(const FString& LevelName);
+
+private:
+	void HandleLevelsChanged();
+
+	// List of sublevels which should not pass filter
+	TArray<FString> HiddenLevels;
+
+	TWeakObjectPtr<UWorld> CachedWorld;
+};
+
 #undef LOCTEXT_NAMESPACE
\ No newline at end of file
