Sequencer: Fix crash with null world

diff --git a/Engine/Source/Editor/Sequencer/Private/SequencerNodeTree.cpp b/Engine/Source/Editor/Sequencer/Private/SequencerNodeTree.cpp
--- a/Engine/Source/Editor/Sequencer/Private/SequencerNodeTree.cpp
+++ b/Engine/Source/Editor/Sequencer/Private/SequencerNodeTree.cpp
@@ -841,11 +841,19 @@ static void FilterNodesRecursive( FSequencer& Sequencer, const TSharedRef<FSeque
 void FSequencerNodeTree::UpdateFilters()
 {
 	if (!bFilterUpdateRequested)
+	{
 		return;
+	}
 
 	FilteredNodes.Empty();
 
-	TrackFilterLevelFilter->UpdateWorld(Sequencer.GetPlaybackContext()->GetWorld());
+
+	UObject* PlaybackContext = Sequencer.GetPlaybackContext();
+	UWorld* World = PlaybackContext ? PlaybackContext->GetWorld() : nullptr;
+	if (World)
+	{
+		TrackFilterLevelFilter->UpdateWorld(World);
+	}
 
 	if (TrackFilters->Num() > 0 || !FilterString.IsEmpty() || TrackFilterLevelFilter->IsActive())
 	{
