Sequencer: Pass nullptr to TrackFilterLevelFilter when no valid world to unhook previous world

diff --git a/Engine/Source/Editor/Sequencer/Private/SequencerNodeTree.cpp b/Engine/Source/Editor/Sequencer/Private/SequencerNodeTree.cpp
--- a/Engine/Source/Editor/Sequencer/Private/SequencerNodeTree.cpp
+++ b/Engine/Source/Editor/Sequencer/Private/SequencerNodeTree.cpp
@@ -850,10 +850,7 @@ void FSequencerNodeTree::UpdateFilters()
 
 	UObject* PlaybackContext = Sequencer.GetPlaybackContext();
 	UWorld* World = PlaybackContext ? PlaybackContext->GetWorld() : nullptr;
-	if (World)
-	{
-		TrackFilterLevelFilter->UpdateWorld(World);
-	}
+	TrackFilterLevelFilter->UpdateWorld(World);
 
 	if (TrackFilters->Num() > 0 || !FilterString.IsEmpty() || TrackFilterLevelFilter->IsActive())
 	{
