Sequencer: Fix dragging with blended sections

diff --git a/Engine/Source/Editor/Sequencer/Private/Tools/EditToolDragOperations.cpp b/Engine/Source/Editor/Sequencer/Private/Tools/EditToolDragOperations.cpp
--- a/Engine/Source/Editor/Sequencer/Private/Tools/EditToolDragOperations.cpp
+++ b/Engine/Source/Editor/Sequencer/Private/Tools/EditToolDragOperations.cpp
@@ -802,7 +802,6 @@ void FMoveKeysAndSections::OnDrag(const FPointerEvent& MouseEvent, FVector2D Loc
 		MouseTime = MovieScene::ClampToDiscreteRange(MouseTime, Sequencer.GetPlaybackRange());
 	}
 
-
 	// We'll calculate a DeltaX based on limits on movement (snapping, section collision) and then use them on keys and sections below.
 	TOptional<FFrameNumber> MaxDeltaX = GetMovementDeltaX(MouseTime);
 
@@ -934,56 +933,64 @@ TOptional<FFrameNumber> FMoveKeysAndSections::GetMovementDeltaX(FFrameTime Mouse
 	// Disallow movement if any of the sections can't move
 	for (int32 Index = 0; Index < Sections.Num(); ++Index)
 	{
+		TOptional<FFrameNumber> LeftMovementMaximum;
+		TOptional<FFrameNumber> RightMovementMaximum;
+
 		const FSectionHandle& SectionHandle = Sections[Index];
 
 		// If we're moving a section that is blending with something then it's OK if it overlaps stuff, the blend amount will get updated at the end.
 		UMovieSceneSection* Section = SectionHandle.GetSectionObject();
-		if (Section->GetBlendType().IsValid())
-		{
-			continue;
-		}
 
 		// We'll calculate this section's borders and clamp the possible delta time to be less than that
-		TRange<FFrameNumber> SectionBoundaries = GetSectionBoundaries(Section, Sections, SectionHandle.TrackNode);
-
-		FFrameNumber LeftMovementMaximum = MovieScene::DiscreteInclusiveLower(SectionBoundaries);
-		FFrameNumber RightMovementMaximum = MovieScene::DiscreteExclusiveUpper(SectionBoundaries);
+		
+		if (!Section->GetBlendType().IsValid())
+		{
+			TRange<FFrameNumber> SectionBoundaries = GetSectionBoundaries(Section, Sections, SectionHandle.TrackNode);
+			LeftMovementMaximum = MovieScene::DiscreteInclusiveLower(SectionBoundaries);
+			RightMovementMaximum = MovieScene::DiscreteExclusiveUpper(SectionBoundaries);
+		}
 		
 		if (Settings->ShouldKeepCursorInPlayRangeWhileScrubbing() && !Settings->ShouldKeepPlayRangeInSectionBounds())
 		{
-			if (LeftMovementMaximum < Sequencer.GetPlaybackRange().GetLowerBoundValue())
+			if (!LeftMovementMaximum.IsSet() || LeftMovementMaximum.GetValue() < Sequencer.GetPlaybackRange().GetLowerBoundValue())
 			{
 				LeftMovementMaximum = Sequencer.GetPlaybackRange().GetLowerBoundValue();
 			}
 
-			if (RightMovementMaximum > Sequencer.GetPlaybackRange().GetUpperBoundValue())
+			if (!RightMovementMaximum.IsSet() || RightMovementMaximum.GetValue() > Sequencer.GetPlaybackRange().GetUpperBoundValue())
 			{
 				RightMovementMaximum = Sequencer.GetPlaybackRange().GetUpperBoundValue();
 			}
 		}
 
-		if (Section->HasStartFrame())
+		if (LeftMovementMaximum.IsSet())
 		{
-			FFrameNumber NewStartTime = Section->GetInclusiveStartFrame() + MouseDeltaTime;
-			if (NewStartTime < LeftMovementMaximum)
+			if (Section->HasStartFrame())
 			{
-				FFrameNumber ClampedDeltaTime = LeftMovementMaximum - Section->GetInclusiveStartFrame();
-				if (!DeltaX.IsSet() || DeltaX.GetValue() > ClampedDeltaTime)
+				FFrameNumber NewStartTime = Section->GetInclusiveStartFrame() + MouseDeltaTime;
+				if (NewStartTime < LeftMovementMaximum.GetValue())
 				{
-					DeltaX = ClampedDeltaTime;
+					FFrameNumber ClampedDeltaTime = LeftMovementMaximum.GetValue() - Section->GetInclusiveStartFrame();
+					if (!DeltaX.IsSet() || DeltaX.GetValue() > ClampedDeltaTime)
+					{
+						DeltaX = ClampedDeltaTime;
+					}
 				}
 			}
 		}
-		
-		if (Section->HasEndFrame())
+
+		if (RightMovementMaximum.IsSet())
 		{
-			FFrameNumber NewEndTime = Section->GetExclusiveEndFrame() + MouseDeltaTime;
-			if (NewEndTime > RightMovementMaximum)
+			if (Section->HasEndFrame())
 			{
-				FFrameNumber ClampedDeltaTime = RightMovementMaximum - Section->GetExclusiveEndFrame();
-				if (!DeltaX.IsSet() || DeltaX.GetValue() > ClampedDeltaTime)
+				FFrameNumber NewEndTime = Section->GetExclusiveEndFrame() + MouseDeltaTime;
+				if (NewEndTime > RightMovementMaximum.GetValue())
 				{
-					DeltaX = ClampedDeltaTime;
+					FFrameNumber ClampedDeltaTime = RightMovementMaximum.GetValue() - Section->GetExclusiveEndFrame();
+					if (!DeltaX.IsSet() || DeltaX.GetValue() > ClampedDeltaTime)
+					{
+						DeltaX = ClampedDeltaTime;
+					}
 				}
 			}
 		}
diff --git a/Engine/Source/Editor/Sequencer/Private/Tools/SequencerEditTool_Movement.cpp b/Engine/Source/Editor/Sequencer/Private/Tools/SequencerEditTool_Movement.cpp
--- a/Engine/Source/Editor/Sequencer/Private/Tools/SequencerEditTool_Movement.cpp
+++ b/Engine/Source/Editor/Sequencer/Private/Tools/SequencerEditTool_Movement.cpp
@@ -16,6 +16,7 @@
 #include "Widgets/Layout/SBox.h"
 #include "Framework/Application/SlateApplication.h"
 #include "Misc/Timecode.h"
+#include "MovieSceneTimeHelpers.h"
 
 const FName FSequencerEditTool_Movement::Identifier = "Movement";
 
@@ -76,9 +77,17 @@ FReply FSequencerEditTool_Movement::OnMouseMove(SWidget& OwnerWidget, const FGeo
 			{
 				DragPosition = MyGeometry.AbsoluteToLocal(MouseEvent.GetScreenSpacePosition());
 
+				if (Sequencer.GetSequencerSettings()->ShouldKeepCursorInPlayRangeWhileScrubbing() && !Sequencer.GetSequencerSettings()->ShouldKeepPlayRangeInSectionBounds())
+				{
+					DragPosition.X = FMath::Max(DragPosition.X, 0.f);
+					FFrameTime CurrentTime = VirtualTrackArea.PixelToFrame(DragPosition.X);
+					CurrentTime = MovieScene::ClampToDiscreteRange(CurrentTime, Sequencer.GetPlaybackRange());
+					DragPosition.X = VirtualTrackArea.FrameToPixel(CurrentTime);
+				}
+					
 				double CurrentTime = VirtualTrackArea.PixelToSeconds(DragPosition.X);
 				Sequencer.UpdateAutoScroll(CurrentTime);
-				
+
 				DragOperation->OnDrag(MouseEvent, DragPosition, VirtualTrackArea);
 			}
 		}
