Sequencer: Do no rescroll to parent track when selecting sequener nodes

diff --git a/Engine/Source/Editor/Sequencer/Private/SSequencer.cpp b/Engine/Source/Editor/Sequencer/Private/SSequencer.cpp
--- a/Engine/Source/Editor/Sequencer/Private/SSequencer.cpp
+++ b/Engine/Source/Editor/Sequencer/Private/SSequencer.cpp
@@ -810,7 +810,6 @@ void SSequencer::Construct(const FArguments& InArgs, TSharedRef<FSequencer> InSe
 
 	InSequencer->GetSelection().GetOnKeySelectionChanged().AddSP(this, &SSequencer::HandleKeySelectionChanged);
 	InSequencer->GetSelection().GetOnSectionSelectionChanged().AddSP(this, &SSequencer::HandleSectionSelectionChanged);
-	InSequencer->GetSelection().GetOnOutlinerNodeSelectionChanged().AddSP(this, &SSequencer::HandleOutlinerNodeSelectionChanged);
 
 	ResetBreadcrumbs();
 }
@@ -957,27 +956,6 @@ void SSequencer::HandleSectionSelectionChanged()
 {
 }
 
-
-void SSequencer::HandleOutlinerNodeSelectionChanged()
-{
-	const TSet<TSharedRef<FSequencerDisplayNode>>& OutlinerSelection = SequencerPtr.Pin()->GetSelection().GetSelectedOutlinerNodes();
-	if ( OutlinerSelection.Num() == 1 )
-	{
-		for ( auto& Node : OutlinerSelection )
-		{
-			auto Parent = Node->GetParent();
-			while (Parent.IsValid())
-			{
-				TreeView->SetItemExpansion(Parent->AsShared(), true);
-				Parent = Parent->GetParent();
-			}
-
-			TreeView->RequestScrollIntoView( Node );
-			break;
-		}
-	}
-}
-
 TSharedRef<SWidget> SSequencer::MakeAddButton()
 {
 	return SNew(SComboButton)
diff --git a/Engine/Source/Editor/Sequencer/Private/SSequencer.h b/Engine/Source/Editor/Sequencer/Private/SSequencer.h
--- a/Engine/Source/Editor/Sequencer/Private/SSequencer.h
+++ b/Engine/Source/Editor/Sequencer/Private/SSequencer.h
@@ -298,9 +298,6 @@ private:
 	/** Handles section selection changes. */
 	void HandleSectionSelectionChanged();
 
-	/** Handles changes to the selected outliner nodes. */
-	void HandleOutlinerNodeSelectionChanged();
-
 	/** Syncs the current node selection to the curve editor. */
 	void SyncCurveEditorToSelection();
 
diff --git a/Engine/Source/Editor/Sequencer/Private/Sequencer.cpp b/Engine/Source/Editor/Sequencer/Private/Sequencer.cpp
--- a/Engine/Source/Editor/Sequencer/Private/Sequencer.cpp
+++ b/Engine/Source/Editor/Sequencer/Private/Sequencer.cpp
@@ -5374,6 +5374,25 @@ void FSequencer::SynchronizeSequencerSelectionWithExternalSelection()
 		{
 			Selection.AddToSelection( NodeToSelect );
 		}
+		
+		TSharedPtr<SSequencerTreeView> TreeView = SequencerWidget->GetTreeView();
+		const TSet<TSharedRef<FSequencerDisplayNode>>& OutlinerSelection = GetSelection().GetSelectedOutlinerNodes();
+		if (OutlinerSelection.Num() == 1)
+		{
+			for (auto& Node : OutlinerSelection)
+			{
+				auto Parent = Node->GetParent();
+				while (Parent.IsValid())
+				{
+					TreeView->SetItemExpansion(Parent->AsShared(), true);
+					Parent = Parent->GetParent();
+				}
+
+				TreeView->RequestScrollIntoView(Node);
+				break;
+			}
+		}
+
 		Selection.ResumeBroadcast();
 		Selection.GetOnOutlinerNodeSelectionChanged().Broadcast();
 	}
