Prevent keyframing on array properties that have a setter function because FTrackInstancePropertyBindings::InvokeSetterFunction doesn't support it.

diff --git a/Engine/Source/Editor/Sequencer/Private/SequencerObjectChangeListener.cpp b/Engine/Source/Editor/Sequencer/Private/SequencerObjectChangeListener.cpp
--- a/Engine/Source/Editor/Sequencer/Private/SequencerObjectChangeListener.cpp
+++ b/Engine/Source/Editor/Sequencer/Private/SequencerObjectChangeListener.cpp
@@ -217,7 +217,11 @@ const FOnAnimatablePropertyChanged* FSequencerObjectChangeListener::FindProperty
 		bool bFoundValidFunction = false;
 		if (Function && !Function->HasMetaData(DeprecatedFunctionName))
 		{
-			bFoundValidFunction = true;
+			// FIXME: FTrackInstancePropertyBindings::InvokeSetterFunction doesn't support array properties.
+			if (!Cast<const UArrayProperty>(&Property))
+			{
+				bFoundValidFunction = true;
+			}
 		}
 
 		bool bFoundValidInterp = false;
