Sequencer: Update with the channel metadata name.

Otherwise, the key area was getting the name Location instead of Location.X. This fixes a subtle bug that appeared when copying and pasting keys. On undo, the key areas would get their names updated from Location.X to Location, Location.Y to Location, and Location.Z to Location. Then the keys would get copied with those key area names and pasted all to the same channel.

diff --git a/Engine/Source/Editor/Sequencer/Private/SequencerSectionLayoutBuilder.cpp b/Engine/Source/Editor/Sequencer/Private/SequencerSectionLayoutBuilder.cpp
--- a/Engine/Source/Editor/Sequencer/Private/SequencerSectionLayoutBuilder.cpp
+++ b/Engine/Source/Editor/Sequencer/Private/SequencerSectionLayoutBuilder.cpp
@@ -124,7 +124,13 @@ void FSequencerSectionLayoutBuilder::AddOrUpdateChannel(TSharedRef<FSequencerSec
 	else
 	{
 		// Just ensure the name is up to date
-		KeyArea->SetName(CurrentNode->GetNodeName());
+		const FMovieSceneChannelMetaData* MetaData = Channel.GetMetaData();
+		if (!ensureAlwaysMsgf(MetaData, TEXT("Attempting to update an expired channel handle to the node tree")))
+		{
+			return;
+		}
+
+		KeyArea->SetName(MetaData->Name);
 	}
 }
 
