Fixes a case where an object binding could get added to a folder but when the object binding was removed via Ctrl + Z the folder would not remove it as a child and the UI would trip an ensure.

diff --git a/Engine/Source/Editor/Sequencer/Private/Sequencer.cpp b/Engine/Source/Editor/Sequencer/Private/Sequencer.cpp
--- a/Engine/Source/Editor/Sequencer/Private/Sequencer.cpp
+++ b/Engine/Source/Editor/Sequencer/Private/Sequencer.cpp
@@ -5113,6 +5113,7 @@ TArray<FGuid> FSequencer::AddActors(const TArray<TWeakObjectPtr<AActor> >& InAct
 		{
 			for (const FGuid& Possessable : PossessableGuids)
 			{
+				SelectedParentFolders[0]->Modify();
 				SelectedParentFolders[0]->AddChildObjectBinding(Possessable);
 			}
 		}
