Fixes Sequencer objects/folders that get moved from a folder back to root via undo/redo not correctly updating the UI.

diff --git a/Engine/Source/Editor/Sequencer/Private/SequencerNodeTree.cpp b/Engine/Source/Editor/Sequencer/Private/SequencerNodeTree.cpp
--- a/Engine/Source/Editor/Sequencer/Private/SequencerNodeTree.cpp
+++ b/Engine/Source/Editor/Sequencer/Private/SequencerNodeTree.cpp
@@ -85,7 +85,7 @@ void FSequencerNodeTree::RefreshNodes(UMovieScene* MovieScene)
 		{
 			if (ensureAlwaysMsgf(Folder, TEXT("MovieScene data contains a null folder. This should never happen.")))
 			{
-				TSharedRef<FSequencerFolderNode> RootFolderNode = CreateOrUpdateFolder(Folder, AllBindings, ChildToParentBinding, MovieScene);
+				TSharedRef<FSequencerFolderNode> RootFolderNode = CreateOrUpdateFolder(Folder, AllBindings, ChildToParentBinding);
 				RootFolderNode->SetParent(RootNode);
 			}
 		}
@@ -95,18 +95,22 @@ void FSequencerNodeTree::RefreshNodes(UMovieScene* MovieScene)
 	{
 		for (const FMovieSceneBinding& Binding : MovieScene->GetBindings())
 		{
-			TSharedPtr<FSequencerObjectBindingNode> ObjectBindingNode = CreateOrUpdateObjectBinding(Binding.GetObjectGuid(), AllBindings, ChildToParentBinding, MovieScene);
-			if (!ObjectBindingNode.IsValid())
+			TSharedPtr<FSequencerObjectBindingNode> ObjectBindingNode = FindOrCreateObjectBinding(Binding.GetObjectGuid(), AllBindings, ChildToParentBinding);
+			if (!ObjectBindingNode)
 			{
 				continue;
 			}
 
-			// Ensure it has a parent
-			if (!ObjectBindingNode->IsParentStillRelevant(SerialNumber))
+			const bool bParentAlreadyDefined = ObjectBindingNode->TreeSerialNumber == SerialNumber;
+			// Check if we've already seen this object binding as part of the folder pass above.
+			if (!bParentAlreadyDefined || !ObjectBindingNode->IsParentStillRelevant(SerialNumber))
 			{
 				ObjectBindingNode->SetParent(RootNode);
 			}
 
+			// Always ensure the serial number is up-to-date. This is only strictly necessary for root object bindings outside of folders, as all others will have been updated already
+			ObjectBindingNode->TreeSerialNumber = SerialNumber;
+
 			// Create nodes for the object binding's tracks
 			for (UMovieSceneTrack* Track : Binding.GetTracks())
 			{
@@ -209,7 +213,7 @@ TSharedPtr<FSequencerTrackNode> FSequencerNodeTree::CreateOrUpdateTrack(UMovieSc
 	return TrackNode;
 }
 
-TSharedRef<FSequencerFolderNode> FSequencerNodeTree::CreateOrUpdateFolder(UMovieSceneFolder* Folder, const TSortedMap<FGuid, const FMovieSceneBinding*>& AllBindings, const TSortedMap<FGuid, FGuid>& ChildToParentBinding, const UMovieScene* InMovieScene)
+TSharedRef<FSequencerFolderNode> FSequencerNodeTree::CreateOrUpdateFolder(UMovieSceneFolder* Folder, const TSortedMap<FGuid, const FMovieSceneBinding*>& AllBindings, const TSortedMap<FGuid, FGuid>& ChildToParentBinding)
 {
 	check(Folder);
 
@@ -228,10 +232,13 @@ TSharedRef<FSequencerFolderNode> FSequencerNodeTree::CreateOrUpdateFolder(UMovie
 	// Create the hierarchy for any child bindings
 	for (const FGuid& ID : Folder->GetChildObjectBindings())
 	{
-		TSharedPtr<FSequencerObjectBindingNode> Binding = CreateOrUpdateObjectBinding(ID, AllBindings, ChildToParentBinding, InMovieScene);
+		TSharedPtr<FSequencerObjectBindingNode> Binding = FindOrCreateObjectBinding(ID, AllBindings, ChildToParentBinding);
 		if (Binding.IsValid())
 		{
 			Binding->SetParent(FolderNode);
+
+			// Indicate that this binding has had its parent defined within this update cycle
+			Binding->TreeSerialNumber = SerialNumber;
 		}
 	}
 
@@ -253,7 +260,7 @@ TSharedRef<FSequencerFolderNode> FSequencerNodeTree::CreateOrUpdateFolder(UMovie
 	{
 		if (ensureAlwaysMsgf(ChildFolder, TEXT("MovieScene folder '%s' data contains a null child folder. This should never happen."), *Folder->GetName()))
 		{
-			TSharedRef<FSequencerFolderNode> ChildFolderNode = CreateOrUpdateFolder(ChildFolder, AllBindings, ChildToParentBinding, InMovieScene);
+			TSharedRef<FSequencerFolderNode> ChildFolderNode = CreateOrUpdateFolder(ChildFolder, AllBindings, ChildToParentBinding);
 			ChildFolderNode->SetParent(FolderNode);
 		}
 	}
@@ -266,7 +273,7 @@ bool FSequencerNodeTree::HasActiveFilter() const
 	return (!FilterString.IsEmpty() || TrackFilters->Num() > 0 || TrackFilterLevelFilter->IsActive());
 }
 
-TSharedPtr<FSequencerObjectBindingNode> FSequencerNodeTree::CreateOrUpdateObjectBinding(const FGuid& BindingID, const TSortedMap<FGuid, const FMovieSceneBinding*>& AllBindings, const TSortedMap<FGuid, FGuid>& ChildToParentBinding, const UMovieScene* InMovieScene)
+TSharedPtr<FSequencerObjectBindingNode> FSequencerNodeTree::FindOrCreateObjectBinding(const FGuid& BindingID, const TSortedMap<FGuid, const FMovieSceneBinding*>& AllBindings, const TSortedMap<FGuid, FGuid>& ChildToParentBinding)
 {
 	if (!ensureAlwaysMsgf(AllBindings.Contains(BindingID), TEXT("Attempting to add a binding that does not exist.")))
 	{
@@ -283,16 +290,17 @@ TSharedPtr<FSequencerObjectBindingNode> FSequencerNodeTree::CreateOrUpdateObject
 		ObjectBindingToNode.Add(BindingID, ObjectBindingNode);
 	}
 
-	// Assign the serial number for this node to indicate that it is still relevant
-	ObjectBindingNode->TreeSerialNumber = SerialNumber;
-
 	// Create its parent and make the association
 	if (const FGuid* ParentGuid = ChildToParentBinding.Find(BindingID))
 	{
-		TSharedPtr<FSequencerObjectBindingNode> ParentBinding = CreateOrUpdateObjectBinding(*ParentGuid, AllBindings, ChildToParentBinding, InMovieScene);
+		TSharedPtr<FSequencerObjectBindingNode> ParentBinding = FindOrCreateObjectBinding(*ParentGuid, AllBindings, ChildToParentBinding);
+
 		if (ParentBinding.IsValid())
 		{
 			ObjectBindingNode->SetParent(ParentBinding);
+
+			// Assign the serial number for this node to indicate that it has received its correct parent
+			ObjectBindingNode->TreeSerialNumber = SerialNumber;
 		}
 	}
 
diff --git a/Engine/Source/Editor/Sequencer/Private/SequencerNodeTree.h b/Engine/Source/Editor/Sequencer/Private/SequencerNodeTree.h
--- a/Engine/Source/Editor/Sequencer/Private/SequencerNodeTree.h
+++ b/Engine/Source/Editor/Sequencer/Private/SequencerNodeTree.h
@@ -162,9 +162,35 @@ private:
 
 	enum class ETrackType { Master, Object };
 
+	/**
+	 * Create or update a track node for the specified track object, updating its serial number.
+	 *
+	 * @param Track                   Pointer to the track to create or update a node for
+	 * @param TrackType               Whether this track is a master track or contained within an object
+	 * @return A shared pointer to a track node with its serial number updated. Will return nullptr for tracks that are forcibly hidden through FSequencer::IsTrackVisible.
+	 */
 	TSharedPtr<FSequencerTrackNode> CreateOrUpdateTrack(UMovieSceneTrack* Track, ETrackType TrackType);
-	TSharedRef<FSequencerFolderNode> CreateOrUpdateFolder(UMovieSceneFolder* Folder, const TSortedMap<FGuid, const FMovieSceneBinding*>& AllBindings, const TSortedMap<FGuid, FGuid>& ChildToParentBinding, const UMovieScene* InMovieScene);
-	TSharedPtr<FSequencerObjectBindingNode> CreateOrUpdateObjectBinding(const FGuid& BindingID, const TSortedMap<FGuid, const FMovieSceneBinding*>& AllBindings, const TSortedMap<FGuid, FGuid>& ChildToParentBinding, const UMovieScene* InMovieScene);
+
+	/**
+	 * Create or update a folder node for the specified folder and all its decendent child folders and object bindings, updating their serial numbers in the process
+	 *
+	 * @param Folder                  Pointer to the movie scene folder to create a node for
+	 * @param AllBindings             A map from guid to binding pointer for all the bindings in the sequence
+	 * @param ChildToParentBinding    Child to parent GUID map used for creating parent items
+	 * @return A shared reference to a folder node. The resulting node's serial number will always be up-to-date, as will all its child folders, and any immediate child object bindings.
+	 */
+	TSharedRef<FSequencerFolderNode> CreateOrUpdateFolder(UMovieSceneFolder* Folder, const TSortedMap<FGuid, const FMovieSceneBinding*>& AllBindings, const TSortedMap<FGuid, FGuid>& ChildToParentBinding);
+
+	/**
+	 * Find an existing object binding node (or create a new one) for the specified binding ID without updating its tree serial number, creating any parent object binding nodes in the process.
+	 * @note: Will only update FSequencerDisplayNode::TreeSerialNumber for object bindings that have a known and valid parent object binding.
+	 *
+	 * @param BindingID               The Guid of the object binding within UMovieScene::GetBindings
+	 * @param AllBindings             A map from guid to binding pointer for all the bindings in the sequence
+	 * @param ChildToParentBinding    Child to parent GUID map used for creating parent items
+	 * @return A shared pointer to an object binding node or nullptr if the supplied BindingID was not valid for this sequence. The resulting node's serial number will be up-to-date if it is a child of another binding, or it was previously added to a folder.
+	 */
+	TSharedPtr<FSequencerObjectBindingNode> FindOrCreateObjectBinding(const FGuid& BindingID, const TSortedMap<FGuid, const FMovieSceneBinding*>& AllBindings, const TSortedMap<FGuid, FGuid>& ChildToParentBinding);
 
 	/**
 	 * Creates section handles for all the sections contained in the specified track
