Sequencer: Add options to enable/disable all level filters

diff --git a/Engine/Source/Editor/Sequencer/Private/SSequencer.cpp b/Engine/Source/Editor/Sequencer/Private/SSequencer.cpp
--- a/Engine/Source/Editor/Sequencer/Private/SSequencer.cpp
+++ b/Engine/Source/Editor/Sequencer/Private/SSequencer.cpp
@@ -1414,7 +1414,7 @@ TSharedRef<SWidget> SSequencer::MakeAddMenu()
 
 TSharedRef<SWidget> SSequencer::MakeFilterMenu()
 {
-	FMenuBuilder MenuBuilder(true, nullptr, AddMenuExtender);
+	FMenuBuilder MenuBuilder(false, nullptr, AddMenuExtender);
 
 	// let track editors & object bindings populate the menu
 	TSharedPtr<FSequencer> Sequencer = SequencerPtr.Pin();
@@ -1436,7 +1436,7 @@ TSharedRef<SWidget> SSequencer::MakeFilterMenu()
 	if (World && World->GetLevels().Num() > 1)
 	{
 		MenuBuilder.BeginSection("TrackLevelFilters");
-		MenuBuilder.AddSubMenu(LOCTEXT("LevelFilters", "Level Filters"), LOCTEXT("LevelFiltersToolTip", "Filter object tracks by level"), FNewMenuDelegate::CreateRaw(this, &SSequencer::FillLevelFilterMenu));
+		MenuBuilder.AddSubMenu(LOCTEXT("LevelFilters", "Level Filters"), LOCTEXT("LevelFiltersToolTip", "Filter object tracks by level"), FNewMenuDelegate::CreateRaw(this, &SSequencer::FillLevelFilterMenu), false);
 		MenuBuilder.EndSection();
 	}
 
@@ -1474,6 +1474,28 @@ void SSequencer::FillLevelFilterMenu(FMenuBuilder& InMenuBarBuilder)
 	if (World)
 	{
 		const TArray<ULevel*> Levels = World->GetLevels();
+
+		if (Levels.Num() > 0)
+		{
+			InMenuBarBuilder.BeginSection("SequencerTracksResetLevelFilters");
+
+			InMenuBarBuilder.AddMenuEntry(
+				LOCTEXT("EnableAllLevelFilters", "Enable All"),
+				LOCTEXT("EnableAllLevelFiltersToolTip", "Enables all level filters"),
+				FSlateIcon(),
+				FUIAction(FExecuteAction::CreateSP(this, &SSequencer::OnEnableAllLevelFilters, true)));
+
+			InMenuBarBuilder.AddMenuEntry(
+				LOCTEXT("DisableAllLevelFilters", "Disable All"),
+				LOCTEXT("DisableAllLevelFiltersToolTip", "Disable all level filters"),
+				FSlateIcon(),
+				FUIAction(FExecuteAction::CreateSP(this, &SSequencer::OnEnableAllLevelFilters, false)));
+
+			InMenuBarBuilder.EndSection();
+
+			InMenuBarBuilder.AddMenuSeparator();
+		}
+
 		for (ULevel* Level : Levels)
 		{
 			FString LevelName = FPackageName::GetShortName(Level->GetOutermost()->GetName());
@@ -1517,6 +1539,32 @@ bool SSequencer::IsTrackFilterActive(TSharedRef<FSequencerTrackFilter> TrackFilt
 	return Sequencer->GetNodeTree()->IsTrackFilterActive(TrackFilter);
 }
 
+void SSequencer::OnEnableAllLevelFilters(bool bEnableAll)
+{
+	TSharedPtr<FSequencer> Sequencer = SequencerPtr.Pin();
+	UObject* PlaybackContext = Sequencer->GetPlaybackContext();
+	UWorld* World = PlaybackContext ? PlaybackContext->GetWorld() : nullptr;
+
+	if (World)
+	{
+		const TArray<ULevel*> Levels = World->GetLevels();
+
+		for (ULevel* Level : Levels)
+		{
+			FString LevelName = FPackageName::GetShortName(Level->GetOutermost()->GetName());
+
+			if (bEnableAll)
+			{
+				Sequencer->GetNodeTree()->AddLevelFilter(LevelName);
+			}
+			else
+			{
+				Sequencer->GetNodeTree()->RemoveLevelFilter(LevelName);
+			}
+		}
+	}
+}
+
 void SSequencer::OnTrackLevelFilterClicked(const FString LevelName)
 {
 	TSharedPtr<FSequencer> Sequencer = SequencerPtr.Pin();
diff --git a/Engine/Source/Editor/Sequencer/Private/SSequencer.h b/Engine/Source/Editor/Sequencer/Private/SSequencer.h
--- a/Engine/Source/Editor/Sequencer/Private/SSequencer.h
+++ b/Engine/Source/Editor/Sequencer/Private/SSequencer.h
@@ -389,6 +389,7 @@ private:
 	void OnTrackFilterClicked(TSharedRef<FSequencerTrackFilter> TrackFilter);
 	bool IsTrackFilterActive(TSharedRef<FSequencerTrackFilter> TrackFilter) const;
 
+	void OnEnableAllLevelFilters(bool bEnableAll);
 	void OnTrackLevelFilterClicked(const FString LevelName);
 	bool IsTrackLevelFilterActive(const FString LevelName) const;
 
