Curve Editor: Filter to object's curves instead of just a single channel

diff --git a/Engine/Source/Editor/CurveEditor/Private/Tree/CurveEditorTree.cpp b/Engine/Source/Editor/CurveEditor/Private/Tree/CurveEditorTree.cpp
--- a/Engine/Source/Editor/CurveEditor/Private/Tree/CurveEditorTree.cpp
+++ b/Engine/Source/Editor/CurveEditor/Private/Tree/CurveEditorTree.cpp
@@ -230,24 +230,6 @@ bool FCurveEditorTree::PerformFilterPass(TArrayView<const FCurveEditorTreeFilter
 {
 	bool bAnyMatched = false;
 
-	for (FCurveEditorTreeItemID ItemID : ItemsToFilter)
-	{
-		// If it failed a previous pass, don't consider it for this pass
-		if (FilterStates.Get(ItemID) != ECurveEditorTreeFilterState::NoMatch)
-		{
-			TSharedPtr<ICurveEditorTreeItem> TreeItemImpl = GetItem(ItemID).GetItem();
-			if (TreeItemImpl)
-			{
-				const bool bMatchesFilter = Algo::AnyOf(FilterPtrs, [TreeItemImpl](const FCurveEditorTreeFilter* Filter){ return TreeItemImpl->PassesFilter(Filter); });
-				if (bMatchesFilter)
-				{
-					InheritedState = ECurveEditorTreeFilterState::NoMatch;
-					break;
-				}
-			}
-		}
-	}
-
 	for (FCurveEditorTreeItemID ItemID : ItemsToFilter)
 	{
 		// If it failed a previous pass, don't consider it for this pass
diff --git a/Engine/Source/Editor/Sequencer/Private/DisplayNodes/SequencerDisplayNode.cpp b/Engine/Source/Editor/Sequencer/Private/DisplayNodes/SequencerDisplayNode.cpp
--- a/Engine/Source/Editor/Sequencer/Private/DisplayNodes/SequencerDisplayNode.cpp
+++ b/Engine/Source/Editor/Sequencer/Private/DisplayNodes/SequencerDisplayNode.cpp
@@ -31,6 +31,7 @@
 #include "ScopedTransaction.h"
 #include "SequencerKeyTimeCache.h"
 #include "SequencerNodeSortingMethods.h"
+#include "SequencerSelectionCurveFilter.h"
 #include "SequencerKeyCollection.h"
 
 #define LOCTEXT_NAMESPACE "SequencerDisplayNode"
@@ -1195,7 +1196,8 @@ bool FSequencerDisplayNode::PassesFilter(const FCurveEditorTreeFilter* InFilter)
 	}
 	else if (InFilter->GetType() == ISequencerModule::GetSequencerSelectionFilterType())
 	{
-		return GetSequencer().GetSelection().GetSelectedOutlinerNodes().Contains(ConstCastSharedRef<FSequencerDisplayNode>(AsShared()));
+		const FSequencerSelectionCurveFilter* Filter = static_cast<const FSequencerSelectionCurveFilter*>(InFilter);
+		return Filter->Match(AsShared());
 	}
 	return false;
 }
diff --git a/Engine/Source/Editor/Sequencer/Private/SSequencer.cpp b/Engine/Source/Editor/Sequencer/Private/SSequencer.cpp
--- a/Engine/Source/Editor/Sequencer/Private/SSequencer.cpp
+++ b/Engine/Source/Editor/Sequencer/Private/SSequencer.cpp
@@ -75,6 +75,7 @@
 #include "Tree/CurveEditorTreeFilter.h"
 #include "Tree/SCurveEditorTreeTextFilter.h"
 #include "Tree/SCurveEditorTreeFilterStatusBar.h"
+#include "SequencerSelectionCurveFilter.h"
 #include "SCurveKeyDetailPanel.h"
 #include "MovieSceneTimeHelpers.h"
 #include "FrameNumberNumericInterface.h"
@@ -1022,9 +1023,11 @@ void SSequencer::HandleOutlinerNodeSelectionChanged()
 		{
 			if (!SequencerSelectionCurveEditorFilter)
 			{
-				static const int32 FilterPass = -1000;
-				SequencerSelectionCurveEditorFilter = MakeShared<FCurveEditorTreeFilter>(ISequencerModule::GetSequencerSelectionFilterType(), FilterPass);
+				SequencerSelectionCurveEditorFilter = MakeShared<FSequencerSelectionCurveFilter>();
 			}
+
+			SequencerSelectionCurveEditorFilter->Update(Sequencer->GetSelection().GetSelectedOutlinerNodes());
+
 			CurveEditor->GetTree()->AddFilter(SequencerSelectionCurveEditorFilter);
 		}
 		// If we're not isolating to the selection (or there is no selection) remove the filter
diff --git a/Engine/Source/Editor/Sequencer/Private/SSequencer.h b/Engine/Source/Editor/Sequencer/Private/SSequencer.h
--- a/Engine/Source/Editor/Sequencer/Private/SSequencer.h
+++ b/Engine/Source/Editor/Sequencer/Private/SSequencer.h
@@ -36,6 +36,7 @@ class SDockTab;
 class USequencerSettings;
 class FSequencerTrackFilter;
 struct FPaintPlaybackRangeArgs;
+struct FSequencerSelectionCurveFilter;
 
 namespace SequencerLayoutConstants
 {
@@ -533,7 +534,7 @@ private:
 	TSharedPtr<SCurveEditorTree> CurveEditorTree;
 
 	/** Curve editor filter that shows only the selected nodes */
-	TSharedPtr<FCurveEditorTreeFilter> SequencerSelectionCurveEditorFilter;
+	TSharedPtr<FSequencerSelectionCurveFilter> SequencerSelectionCurveEditorFilter;
 
 	/** The breadcrumb trail widget for this sequencer */
 	TSharedPtr<SBreadcrumbTrail<FSequencerBreadcrumb>> BreadcrumbTrail;
diff --git a/Engine/Source/Editor/Sequencer/Private/SequencerSelectionCurveFilter.h b/Engine/Source/Editor/Sequencer/Private/SequencerSelectionCurveFilter.h
new file mode 100644
--- /dev/null
+++ b/Engine/Source/Editor/Sequencer/Private/SequencerSelectionCurveFilter.h
@@ -0,0 +1,54 @@
+// Copyright 1998-2019 Epic Games, Inc. All Rights Reserved.
+
+#include "CoreMinimal.h"
+#include "Tree/CurveEditorTreeFilter.h"
+#include "DisplayNodes/SequencerDisplayNode.h"
+
+
+/**
+ * A specialized filter for showing items in the curve editor selected from the sequencer panel.
+ * This filter will store the selected nodes and all the parents of the selected nodes in a NodesToFilter set.
+ * An item will pass if it is either directly selected or has a parent in the set.
+ */
+struct FSequencerSelectionCurveFilter : FCurveEditorTreeFilter
+{
+	static const int32 FilterPass = -1000;
+
+	FSequencerSelectionCurveFilter()
+		: FCurveEditorTreeFilter(ISequencerModule::GetSequencerSelectionFilterType(), FilterPass)
+	{}
+
+	/**
+	 * Adds all selected nodes and their object parents to the NodesToFilter set
+	 */
+	void Update(const TSet<TSharedRef<FSequencerDisplayNode>>& SelectedNodes)
+	{
+		NodesToFilter.Reserve(SelectedNodes.Num());
+
+		for (const TSharedRef<FSequencerDisplayNode>& SelectedNode : SelectedNodes)
+		{
+			NodesToFilter.Add(SelectedNode);
+
+			TSharedPtr<FSequencerDisplayNode> Parent = SelectedNode->GetParent();
+			while (Parent.IsValid())
+			{
+				if (Parent->GetType() == ESequencerNode::Object)
+				{
+					NodesToFilter.Add(Parent.ToSharedRef());
+					break;
+				}
+
+				Parent = Parent->GetParent();
+			}
+		}
+	}
+
+	bool Match(TSharedRef<const FSequencerDisplayNode> InNode) const
+	{
+		return NodesToFilter.Contains(InNode);
+	}
+
+private:
+
+	TSet<TSharedRef<const FSequencerDisplayNode>> NodesToFilter;
+};
\ No newline at end of file
diff --git a/Engine/Source/Editor/Sequencer/Public/ISequencerModule.h b/Engine/Source/Editor/Sequencer/Public/ISequencerModule.h
--- a/Engine/Source/Editor/Sequencer/Public/ISequencerModule.h
+++ b/Engine/Source/Editor/Sequencer/Public/ISequencerModule.h
@@ -261,7 +261,7 @@ public:
 
 
 	/**
-	 * Retrieve the unique identifer for the sequencer selection curve editor filter
+	 * Retrieve the unique identifer for the sequencer selection curve editor filter (of type FSequencerSelectionCurveFilter)
 	 */
 	static ECurveEditorTreeFilterType GetSequencerSelectionFilterType();
 
