Sequencer: Added missing pragma once

diff --git a/Engine/Source/Editor/Sequencer/Private/SequencerSelectionCurveFilter.h b/Engine/Source/Editor/Sequencer/Private/SequencerSelectionCurveFilter.h
--- a/Engine/Source/Editor/Sequencer/Private/SequencerSelectionCurveFilter.h
+++ b/Engine/Source/Editor/Sequencer/Private/SequencerSelectionCurveFilter.h
@@ -1,5 +1,7 @@
 // Copyright 1998-2019 Epic Games, Inc. All Rights Reserved.
 
+#pragma once
+
 #include "CoreMinimal.h"
 #include "Tree/CurveEditorTreeFilter.h"
 #include "DisplayNodes/SequencerDisplayNode.h"
