Sequencer: Decouple keep cursor in playback range while scrubbing with key/section snapping to playback range

diff --git a/Engine/Source/Editor/Sequencer/Private/SSequencer.cpp b/Engine/Source/Editor/Sequencer/Private/SSequencer.cpp
--- a/Engine/Source/Editor/Sequencer/Private/SSequencer.cpp
+++ b/Engine/Source/Editor/Sequencer/Private/SSequencer.cpp
@@ -1933,10 +1933,16 @@ TSharedRef<SWidget> SSequencer::MakeSnapMenu()
 	}
 	MenuBuilder.EndSection();
 
-	MenuBuilder.BeginSection( "SectionSnapping", LOCTEXT( "SnappingMenuSectionHeader", "Section Snapping" ) );
+	MenuBuilder.BeginSection("SectionSnapping", LOCTEXT("SnappingMenuSectionHeader", "Section Snapping"));
 	{
-		MenuBuilder.AddMenuEntry( FSequencerCommands::Get().ToggleSnapSectionTimesToInterval );
-		MenuBuilder.AddMenuEntry( FSequencerCommands::Get().ToggleSnapSectionTimesToSections );
+		MenuBuilder.AddMenuEntry(FSequencerCommands::Get().ToggleSnapSectionTimesToInterval);
+		MenuBuilder.AddMenuEntry(FSequencerCommands::Get().ToggleSnapSectionTimesToSections);
+	}
+	MenuBuilder.EndSection();
+
+	MenuBuilder.BeginSection("KeyAndSectionSnapping", LOCTEXT("SnappingMenuKeysAndSectionHeader", "Keys and Sections Snapping"));
+	{
+		MenuBuilder.AddMenuEntry(FSequencerCommands::Get().ToggleSnapKeysAndSectionsToPlayRange);
 	}
 	MenuBuilder.EndSection();
 
diff --git a/Engine/Source/Editor/Sequencer/Private/Sequencer.cpp b/Engine/Source/Editor/Sequencer/Private/Sequencer.cpp
--- a/Engine/Source/Editor/Sequencer/Private/Sequencer.cpp
+++ b/Engine/Source/Editor/Sequencer/Private/Sequencer.cpp
@@ -9900,6 +9900,12 @@ void FSequencer::BindCommands()
 		FCanExecuteAction::CreateLambda( []{ return true; } ),
 		FIsActionChecked::CreateLambda( [this]{ return Settings->GetSnapSectionTimesToSections(); } ) );
 
+	SequencerCommandBindings->MapAction(
+		Commands.ToggleSnapKeysAndSectionsToPlayRange,
+		FExecuteAction::CreateLambda([this] { Settings->SetSnapKeysAndSectionsToPlayRange(!Settings->GetSnapKeysAndSectionsToPlayRange()); }),
+		FCanExecuteAction::CreateLambda([] { return true; }),
+		FIsActionChecked::CreateLambda([this] { return Settings->GetSnapKeysAndSectionsToPlayRange(); }));
+
 	SequencerCommandBindings->MapAction(
 		Commands.ToggleSnapPlayTimeToKeys,
 		FExecuteAction::CreateLambda( [this]{ Settings->SetSnapPlayTimeToKeys( !Settings->GetSnapPlayTimeToKeys() ); } ),
@@ -10074,9 +10080,9 @@ void FSequencer::BindCommands()
 
 	SequencerCommandBindings->MapAction(
 		Commands.ToggleKeepCursorInPlaybackRangeWhileScrubbing,
-		FExecuteAction::CreateLambda( [this]{ Settings->SetKeepCursorInPlayRangeWhileScrubbing( !Settings->ShouldKeepCursorInPlayRangeWhileScrubbing() ); } ),
-		FCanExecuteAction::CreateLambda( []{ return true; } ),
-		FIsActionChecked::CreateLambda( [this]{ return Settings->ShouldKeepCursorInPlayRangeWhileScrubbing(); } ) );
+		FExecuteAction::CreateLambda([this] { Settings->SetKeepCursorInPlayRangeWhileScrubbing(!Settings->ShouldKeepCursorInPlayRangeWhileScrubbing()); }),
+		FCanExecuteAction::CreateLambda([] { return true; }),
+		FIsActionChecked::CreateLambda([this] { return Settings->ShouldKeepCursorInPlayRangeWhileScrubbing(); }));
 
 	SequencerCommandBindings->MapAction(
 		Commands.ToggleKeepCursorInPlaybackRange,
diff --git a/Engine/Source/Editor/Sequencer/Private/SequencerCommands.cpp b/Engine/Source/Editor/Sequencer/Private/SequencerCommands.cpp
--- a/Engine/Source/Editor/Sequencer/Private/SequencerCommands.cpp
+++ b/Engine/Source/Editor/Sequencer/Private/SequencerCommands.cpp
@@ -104,6 +104,7 @@ void FSequencerCommands::RegisterCommands()
 
 	UI_COMMAND( ToggleSnapSectionTimesToInterval, "Snap to the Interval", "Snap sections to the time snapping interval", EUserInterfaceActionType::ToggleButton, FInputChord() );
 	UI_COMMAND( ToggleSnapSectionTimesToSections, "Snap to Keys and Sections", "Snap sections to other keys and sections in this sequence", EUserInterfaceActionType::ToggleButton, FInputChord() );
+	UI_COMMAND( ToggleSnapKeysAndSectionsToPlayRange, "Snap Keys and Sections to the Playback Range", "When checked, keys and sections will be snapped to the playback range bounds", EUserInterfaceActionType::ToggleButton, FInputChord());
 
 	UI_COMMAND( ToggleSnapPlayTimeToKeys, "Snap to Keys While Scrubbing", "Snap the current time to keys of the selected track while scrubbing", EUserInterfaceActionType::ToggleButton, FInputChord() );
 	UI_COMMAND( ToggleSnapPlayTimeToInterval, "Snap to the Interval While Scrubbing", "Snap the current time to the time snapping interval while scrubbing", EUserInterfaceActionType::ToggleButton, FInputChord() );
diff --git a/Engine/Source/Editor/Sequencer/Private/SequencerSettings.cpp b/Engine/Source/Editor/Sequencer/Private/SequencerSettings.cpp
--- a/Engine/Source/Editor/Sequencer/Private/SequencerSettings.cpp
+++ b/Engine/Source/Editor/Sequencer/Private/SequencerSettings.cpp
@@ -36,6 +36,7 @@ USequencerSettings::USequencerSettings( const FObjectInitializer& ObjectInitiali
 	bSynchronizeCurveEditorSelection = true;
 	bIsolateCurveEditorToSelection = true;
 	LoopMode = ESequencerLoopMode::SLM_NoLoop;
+	bSnapKeysAndSectionsToPlayRange = false;
 	bKeepCursorInPlayRangeWhileScrubbing = false;
 	bKeepCursorInPlayRange = true;
 	bKeepPlayRangeInSectionBounds = true;
@@ -242,6 +243,20 @@ void USequencerSettings::SetSnapSectionTimesToSections( bool InbSnapSectionTimes
 	}
 }
 
+bool USequencerSettings::GetSnapKeysAndSectionsToPlayRange() const
+{
+	return bSnapKeysAndSectionsToPlayRange;
+}
+
+void USequencerSettings::SetSnapKeysAndSectionsToPlayRange(bool bInSnapKeysAndSectionsToPlayRange)
+{
+	if (bSnapKeysAndSectionsToPlayRange != bInSnapKeysAndSectionsToPlayRange)
+	{
+		bSnapKeysAndSectionsToPlayRange = bInSnapKeysAndSectionsToPlayRange;
+		SaveConfig();
+	}
+}
+
 bool USequencerSettings::GetSnapPlayTimeToKeys() const
 {
 	return bSnapPlayTimeToKeys;
diff --git a/Engine/Source/Editor/Sequencer/Private/Tools/EditToolDragOperations.cpp b/Engine/Source/Editor/Sequencer/Private/Tools/EditToolDragOperations.cpp
--- a/Engine/Source/Editor/Sequencer/Private/Tools/EditToolDragOperations.cpp
+++ b/Engine/Source/Editor/Sequencer/Private/Tools/EditToolDragOperations.cpp
@@ -714,7 +714,7 @@ void FMoveKeysAndSections::OnDrag(const FPointerEvent& MouseEvent, FVector2D Loc
 		}
 	}
 
-	if (Settings->ShouldKeepCursorInPlayRangeWhileScrubbing() && !Settings->ShouldKeepPlayRangeInSectionBounds())
+	if (Settings->GetIsSnapEnabled() && Settings->GetSnapKeysAndSectionsToPlayRange() && !Settings->ShouldKeepPlayRangeInSectionBounds())
 	{
 		MouseTime = MovieScene::ClampToDiscreteRange(MouseTime, Sequencer.GetPlaybackRange());
 	}
@@ -883,7 +883,7 @@ TOptional<FFrameNumber> FMoveKeysAndSections::GetMovementDeltaX(FFrameTime Mouse
 			RightMovementMaximum = MovieScene::DiscreteExclusiveUpper(SectionBoundaries);
 		}
 		
-		if (Settings->ShouldKeepCursorInPlayRangeWhileScrubbing() && !Settings->ShouldKeepPlayRangeInSectionBounds())
+		if (Settings->GetIsSnapEnabled() && Settings->GetSnapKeysAndSectionsToPlayRange() && !Settings->ShouldKeepPlayRangeInSectionBounds())
 		{
 			if (!LeftMovementMaximum.IsSet() || LeftMovementMaximum.GetValue() < Sequencer.GetPlaybackRange().GetLowerBoundValue())
 			{
@@ -929,7 +929,7 @@ TOptional<FFrameNumber> FMoveKeysAndSections::GetMovementDeltaX(FFrameTime Mouse
 		}
 	}
 
-	if (Settings->ShouldKeepCursorInPlayRangeWhileScrubbing() && !Settings->ShouldKeepPlayRangeInSectionBounds())
+	if (Settings->GetIsSnapEnabled() && Settings->GetSnapKeysAndSectionsToPlayRange() && !Settings->ShouldKeepPlayRangeInSectionBounds())
 	{
 		TArray<FFrameNumber> CurrentKeyTimes;
 		CurrentKeyTimes.SetNum(KeysAsArray.Num());
diff --git a/Engine/Source/Editor/Sequencer/Private/Tools/SequencerEditTool_Movement.cpp b/Engine/Source/Editor/Sequencer/Private/Tools/SequencerEditTool_Movement.cpp
--- a/Engine/Source/Editor/Sequencer/Private/Tools/SequencerEditTool_Movement.cpp
+++ b/Engine/Source/Editor/Sequencer/Private/Tools/SequencerEditTool_Movement.cpp
@@ -77,7 +77,7 @@ FReply FSequencerEditTool_Movement::OnMouseMove(SWidget& OwnerWidget, const FGeo
 			{
 				DragPosition = MyGeometry.AbsoluteToLocal(MouseEvent.GetScreenSpacePosition());
 
-				if (Sequencer.GetSequencerSettings()->ShouldKeepCursorInPlayRangeWhileScrubbing() && !Sequencer.GetSequencerSettings()->ShouldKeepPlayRangeInSectionBounds())
+				if (Sequencer.GetSequencerSettings()->GetIsSnapEnabled() && Sequencer.GetSequencerSettings()->GetSnapKeysAndSectionsToPlayRange() && !Sequencer.GetSequencerSettings()->ShouldKeepPlayRangeInSectionBounds())
 				{
 					DragPosition.X = FMath::Max(DragPosition.X, 0.f);
 					FFrameTime CurrentTime = VirtualTrackArea.PixelToFrame(DragPosition.X);
diff --git a/Engine/Source/Editor/Sequencer/Public/SequencerCommands.h b/Engine/Source/Editor/Sequencer/Public/SequencerCommands.h
--- a/Engine/Source/Editor/Sequencer/Public/SequencerCommands.h
+++ b/Engine/Source/Editor/Sequencer/Public/SequencerCommands.h
@@ -256,6 +256,9 @@ public:
 	/** Toggles whether or not sections should snap to other sections. */
 	TSharedPtr< FUICommandInfo > ToggleSnapSectionTimesToSections;
 
+	/** Toggle constraining keys and sections in the play range */
+	TSharedPtr< FUICommandInfo > ToggleSnapKeysAndSectionsToPlayRange;
+
 	/** Toggles whether or not snap to key times while scrubbing. */
 	TSharedPtr< FUICommandInfo > ToggleSnapPlayTimeToKeys;
 
diff --git a/Engine/Source/Editor/Sequencer/Public/SequencerSettings.h b/Engine/Source/Editor/Sequencer/Public/SequencerSettings.h
--- a/Engine/Source/Editor/Sequencer/Public/SequencerSettings.h
+++ b/Engine/Source/Editor/Sequencer/Public/SequencerSettings.h
@@ -163,6 +163,11 @@ public:
 	/** sets whether or not to snap sections to other sections. */
 	void SetSnapSectionTimesToSections( bool InbSnapSectionTimesToSections );
 
+	/** @return true if keys and sections should be kept within the playback range when moving them */
+	bool GetSnapKeysAndSectionsToPlayRange() const;
+	/** Set whether or not keys and sections should be kept within the playback range when moving them */
+	void SetSnapKeysAndSectionsToPlayRange(bool bInSnapKeysAndSectionsToPlayRange);
+
 	/** Gets whether or not to snap the play time to keys while scrubbing. */
 	bool GetSnapPlayTimeToKeys() const;
 	/** Sets whether or not to snap the play time to keys while scrubbing. */
@@ -385,6 +390,10 @@ protected:
 	UPROPERTY( config, EditAnywhere, Category=Snapping )
 	bool bSnapSectionTimesToSections;
 
+	/** Enable or disable keeping keys and sections in the playback range. */
+	UPROPERTY(config, EditAnywhere, Category = Timeline)
+	bool bSnapKeysAndSectionsToPlayRange;
+
 	/** Enable or disable snapping the current time to keys of the selected track while scrubbing. */
 	UPROPERTY( config, EditAnywhere, Category=Snapping )
 	bool bSnapPlayTimeToKeys;
@@ -445,7 +454,7 @@ protected:
 	TEnumAsByte<ESequencerLoopMode> LoopMode;
 
 	/** Enable or disable keeping the cursor in the current playback range while scrubbing. */
-	UPROPERTY( config, EditAnywhere, Category=Timeline )
+	UPROPERTY(config, EditAnywhere, Category = Timeline)
 	bool bKeepCursorInPlayRangeWhileScrubbing;
 
 	/** Enable or disable keeping the cursor in the current playback range during playback. */
