Sequencer: Changed population of tree-node hierarchy on refresh to use a separate pass, rather than trying to do both creation and hierarchy off the node's serial number

diff --git a/Engine/Source/Editor/Sequencer/Private/SequencerNodeTree.cpp b/Engine/Source/Editor/Sequencer/Private/SequencerNodeTree.cpp
--- a/Engine/Source/Editor/Sequencer/Private/SequencerNodeTree.cpp
+++ b/Engine/Source/Editor/Sequencer/Private/SequencerNodeTree.cpp
@@ -79,38 +79,23 @@ void FSequencerNodeTree::RefreshNodes(UMovieScene* MovieScene)
 		}
 	}
 
-	// Folders may also create hierarchy items for tracks and object bindings
-	{
-		for (UMovieSceneFolder* Folder : MovieScene->GetRootFolders())
-		{
-			if (ensureAlwaysMsgf(Folder, TEXT("MovieScene data contains a null folder. This should never happen.")))
-			{
-				TSharedRef<FSequencerFolderNode> RootFolderNode = CreateOrUpdateFolder(Folder, AllBindings, ChildToParentBinding);
-				RootFolderNode->SetParent(RootNode);
-			}
-		}
-	}
+	TMap<FSequencerDisplayNode*, TSharedPtr<FSequencerDisplayNode>> ChildToParentMap;
 
 	// Object Bindings
 	{
 		for (const FMovieSceneBinding& Binding : MovieScene->GetBindings())
 		{
-			TSharedPtr<FSequencerObjectBindingNode> ObjectBindingNode = FindOrCreateObjectBinding(Binding.GetObjectGuid(), AllBindings, ChildToParentBinding);
+			TSharedPtr<FSequencerObjectBindingNode> ObjectBindingNode = FindOrCreateObjectBinding(Binding.GetObjectGuid(), AllBindings, ChildToParentBinding, &ChildToParentMap);
 			if (!ObjectBindingNode)
 			{
 				continue;
 			}
 
-			const bool bParentAlreadyDefined = ObjectBindingNode->TreeSerialNumber == SerialNumber;
-			// Check if we've already seen this object binding as part of the folder pass above.
-			if (!bParentAlreadyDefined || !ObjectBindingNode->IsParentStillRelevant(SerialNumber))
+			if (!ChildToParentMap.Contains(ObjectBindingNode.Get()))
 			{
-				ObjectBindingNode->SetParent(RootNode);
+				ChildToParentMap.Add(ObjectBindingNode.Get(), RootNode);
 			}
 
-			// Always ensure the serial number is up-to-date. This is only strictly necessary for root object bindings outside of folders, as all others will have been updated already
-			ObjectBindingNode->TreeSerialNumber = SerialNumber;
-
 			// Create nodes for the object binding's tracks
 			for (UMovieSceneTrack* Track : Binding.GetTracks())
 			{
@@ -119,7 +104,7 @@ void FSequencerNodeTree::RefreshNodes(UMovieScene* MovieScene)
 					TSharedPtr<FSequencerTrackNode> TrackNode = CreateOrUpdateTrack(Track, ETrackType::Object);
 					if (TrackNode.IsValid())
 					{
-						TrackNode->SetParent(ObjectBindingNode);
+						ChildToParentMap.Add(TrackNode.Get(), ObjectBindingNode);
 					}
 				}
 			}
@@ -132,9 +117,9 @@ void FSequencerNodeTree::RefreshNodes(UMovieScene* MovieScene)
 		if (CameraCutTrack)
 		{
 			TSharedPtr<FSequencerTrackNode> TrackNode = CreateOrUpdateTrack(CameraCutTrack, ETrackType::Master);
-			if (TrackNode.IsValid() && !TrackNode->IsParentStillRelevant(SerialNumber))
+			if (TrackNode.IsValid())
 			{
-				TrackNode->SetParent(RootNode);
+				ChildToParentMap.Add(TrackNode.Get(), RootNode);
 			}
 		}
 
@@ -144,14 +129,31 @@ void FSequencerNodeTree::RefreshNodes(UMovieScene* MovieScene)
 			if (ensureAlwaysMsgf(Track, TEXT("MovieScene data contains a null master track. This should never happen.")))
 			{
 				TSharedPtr<FSequencerTrackNode> TrackNode = CreateOrUpdateTrack(Track, ETrackType::Master);
-				if (TrackNode.IsValid() && !TrackNode->IsParentStillRelevant(SerialNumber))
+				if (TrackNode.IsValid())
 				{
-					TrackNode->SetParent(RootNode);
+					ChildToParentMap.Add(TrackNode.Get(), RootNode);
 				}
 			}
 		}
 	}
 
+	// Folders may also create hierarchy items for tracks and object bindings
+	{
+		for (UMovieSceneFolder* Folder : MovieScene->GetRootFolders())
+		{
+			if (ensureAlwaysMsgf(Folder, TEXT("MovieScene data contains a null folder. This should never happen.")))
+			{
+				TSharedRef<FSequencerFolderNode> RootFolderNode = CreateOrUpdateFolder(Folder, AllBindings, ChildToParentBinding, &ChildToParentMap);
+				RootFolderNode->SetParent(RootNode);
+			}
+		}
+	}
+
+	for (TTuple<FSequencerDisplayNode*, TSharedPtr<FSequencerDisplayNode>> Pair : ChildToParentMap)
+	{
+		Pair.Key->SetParent(Pair.Value);
+	}
+
 	// Remove anything that is no longer relevant (ie serial number is out of date)
 	for (auto It = FolderToNode.CreateIterator(); It; ++It)
 	{
@@ -213,7 +215,7 @@ TSharedPtr<FSequencerTrackNode> FSequencerNodeTree::CreateOrUpdateTrack(UMovieSc
 	return TrackNode;
 }
 
-TSharedRef<FSequencerFolderNode> FSequencerNodeTree::CreateOrUpdateFolder(UMovieSceneFolder* Folder, const TSortedMap<FGuid, const FMovieSceneBinding*>& AllBindings, const TSortedMap<FGuid, FGuid>& ChildToParentBinding)
+TSharedRef<FSequencerFolderNode> FSequencerNodeTree::CreateOrUpdateFolder(UMovieSceneFolder* Folder, const TSortedMap<FGuid, const FMovieSceneBinding*>& AllBindings, const TSortedMap<FGuid, FGuid>& ChildToParentBinding, TMap<FSequencerDisplayNode*, TSharedPtr<FSequencerDisplayNode>>* OutChildToParentMap)
 {
 	check(Folder);
 
@@ -232,13 +234,10 @@ TSharedRef<FSequencerFolderNode> FSequencerNodeTree::CreateOrUpdateFolder(UMovie
 	// Create the hierarchy for any child bindings
 	for (const FGuid& ID : Folder->GetChildObjectBindings())
 	{
-		TSharedPtr<FSequencerObjectBindingNode> Binding = FindOrCreateObjectBinding(ID, AllBindings, ChildToParentBinding);
+		TSharedPtr<FSequencerObjectBindingNode> Binding = FindOrCreateObjectBinding(ID, AllBindings, ChildToParentBinding, OutChildToParentMap);
 		if (Binding.IsValid())
 		{
-			Binding->SetParent(FolderNode);
-
-			// Indicate that this binding has had its parent defined within this update cycle
-			Binding->TreeSerialNumber = SerialNumber;
+			OutChildToParentMap->Add(Binding.Get(), FolderNode);
 		}
 	}
 
@@ -250,7 +249,7 @@ TSharedRef<FSequencerFolderNode> FSequencerNodeTree::CreateOrUpdateFolder(UMovie
 			TSharedPtr<FSequencerTrackNode> TrackNode = CreateOrUpdateTrack(Track, ETrackType::Master);
 			if (TrackNode.IsValid())
 			{
-				TrackNode->SetParent(FolderNode);
+				OutChildToParentMap->Add(TrackNode.Get(), FolderNode);
 			}
 		}
 	}
@@ -260,8 +259,8 @@ TSharedRef<FSequencerFolderNode> FSequencerNodeTree::CreateOrUpdateFolder(UMovie
 	{
 		if (ensureAlwaysMsgf(ChildFolder, TEXT("MovieScene folder '%s' data contains a null child folder. This should never happen."), *Folder->GetName()))
 		{
-			TSharedRef<FSequencerFolderNode> ChildFolderNode = CreateOrUpdateFolder(ChildFolder, AllBindings, ChildToParentBinding);
-			ChildFolderNode->SetParent(FolderNode);
+			TSharedRef<FSequencerFolderNode> ChildFolderNode = CreateOrUpdateFolder(ChildFolder, AllBindings, ChildToParentBinding, OutChildToParentMap);
+			OutChildToParentMap->Add(&ChildFolderNode.Get(), FolderNode);
 		}
 	}
 
@@ -273,7 +272,7 @@ bool FSequencerNodeTree::HasActiveFilter() const
 	return (!FilterString.IsEmpty() || TrackFilters->Num() > 0 || TrackFilterLevelFilter->IsActive());
 }
 
-TSharedPtr<FSequencerObjectBindingNode> FSequencerNodeTree::FindOrCreateObjectBinding(const FGuid& BindingID, const TSortedMap<FGuid, const FMovieSceneBinding*>& AllBindings, const TSortedMap<FGuid, FGuid>& ChildToParentBinding)
+TSharedPtr<FSequencerObjectBindingNode> FSequencerNodeTree::FindOrCreateObjectBinding(const FGuid& BindingID, const TSortedMap<FGuid, const FMovieSceneBinding*>& AllBindings, const TSortedMap<FGuid, FGuid>& ChildToParentBinding, TMap<FSequencerDisplayNode*, TSharedPtr<FSequencerDisplayNode>>* OutChildToParentMap)
 {
 	if (!ensureAlwaysMsgf(AllBindings.Contains(BindingID), TEXT("Attempting to add a binding that does not exist.")))
 	{
@@ -290,17 +289,15 @@ TSharedPtr<FSequencerObjectBindingNode> FSequencerNodeTree::FindOrCreateObjectBi
 		ObjectBindingToNode.Add(BindingID, ObjectBindingNode);
 	}
 
+	ObjectBindingNode->TreeSerialNumber = SerialNumber;
+
 	// Create its parent and make the association
 	if (const FGuid* ParentGuid = ChildToParentBinding.Find(BindingID))
 	{
-		TSharedPtr<FSequencerObjectBindingNode> ParentBinding = FindOrCreateObjectBinding(*ParentGuid, AllBindings, ChildToParentBinding);
-
+		TSharedPtr<FSequencerObjectBindingNode> ParentBinding = FindOrCreateObjectBinding(*ParentGuid, AllBindings, ChildToParentBinding, OutChildToParentMap);
 		if (ParentBinding.IsValid())
 		{
-			ObjectBindingNode->SetParent(ParentBinding);
-
-			// Assign the serial number for this node to indicate that it has received its correct parent
-			ObjectBindingNode->TreeSerialNumber = SerialNumber;
+			OutChildToParentMap->Add(ObjectBindingNode.Get(), ParentBinding);
 		}
 	}
 
diff --git a/Engine/Source/Editor/Sequencer/Private/SequencerNodeTree.h b/Engine/Source/Editor/Sequencer/Private/SequencerNodeTree.h
--- a/Engine/Source/Editor/Sequencer/Private/SequencerNodeTree.h
+++ b/Engine/Source/Editor/Sequencer/Private/SequencerNodeTree.h
@@ -185,9 +185,10 @@ private:
 	 * @param Folder                  Pointer to the movie scene folder to create a node for
 	 * @param AllBindings             A map from guid to binding pointer for all the bindings in the sequence
 	 * @param ChildToParentBinding    Child to parent GUID map used for creating parent items
+	 * @param OutChildToParentMap     Pointer to a map that should be populated with child->parent relationships to be set up after creation of all nodes
 	 * @return A shared reference to a folder node. The resulting node's serial number will always be up-to-date, as will all its child folders, and any immediate child object bindings.
 	 */
-	TSharedRef<FSequencerFolderNode> CreateOrUpdateFolder(UMovieSceneFolder* Folder, const TSortedMap<FGuid, const FMovieSceneBinding*>& AllBindings, const TSortedMap<FGuid, FGuid>& ChildToParentBinding);
+	TSharedRef<FSequencerFolderNode> CreateOrUpdateFolder(UMovieSceneFolder* Folder, const TSortedMap<FGuid, const FMovieSceneBinding*>& AllBindings, const TSortedMap<FGuid, FGuid>& ChildToParentBinding, TMap<FSequencerDisplayNode*, TSharedPtr<FSequencerDisplayNode>>* OutChildToParentMap);
 
 	/**
 	 * Find an existing object binding node (or create a new one) for the specified binding ID without updating its tree serial number, creating any parent object binding nodes in the process.
@@ -196,9 +197,10 @@ private:
 	 * @param BindingID               The Guid of the object binding within UMovieScene::GetBindings
 	 * @param AllBindings             A map from guid to binding pointer for all the bindings in the sequence
 	 * @param ChildToParentBinding    Child to parent GUID map used for creating parent items
+	 * @param OutChildToParentMap     Pointer to a map that should be populated with child->parent relationships to be set up after creation of all nodes
 	 * @return A shared pointer to an object binding node or nullptr if the supplied BindingID was not valid for this sequence. The resulting node's serial number will be up-to-date if it is a child of another binding, or it was previously added to a folder.
 	 */
-	TSharedPtr<FSequencerObjectBindingNode> FindOrCreateObjectBinding(const FGuid& BindingID, const TSortedMap<FGuid, const FMovieSceneBinding*>& AllBindings, const TSortedMap<FGuid, FGuid>& ChildToParentBinding);
+	TSharedPtr<FSequencerObjectBindingNode> FindOrCreateObjectBinding(const FGuid& BindingID, const TSortedMap<FGuid, const FMovieSceneBinding*>& AllBindings, const TSortedMap<FGuid, FGuid>& ChildToParentBinding, TMap<FSequencerDisplayNode*, TSharedPtr<FSequencerDisplayNode>>* OutChildToParentMap);
 
 	/**
 	 * Creates section handles for all the sections contained in the specified track
