Update to new function name
#rnx

diff --git a/Engine/Source/Editor/Sequencer/Private/SSequencerTreeView.cpp b/Engine/Source/Editor/Sequencer/Private/SSequencerTreeView.cpp
--- a/Engine/Source/Editor/Sequencer/Private/SSequencerTreeView.cpp
+++ b/Engine/Source/Editor/Sequencer/Private/SSequencerTreeView.cpp
@@ -517,7 +517,7 @@ FReply SSequencerTreeView::OnKeyDown(const FGeometry& MyGeometry, const FKeyEven
 			}
 
 			int32 NumItemsInAPage = GetNumLiveWidgets();
-			int32 Remainder = NumItemsInAPage % GetNumItemsWide();
+			int32 Remainder = NumItemsInAPage % GetNumItemsPerLine();
 			NumItemsInAPage -= Remainder;
 
 			if (SelectionIndex >= NumItemsInAPage)
@@ -562,7 +562,7 @@ FReply SSequencerTreeView::OnKeyDown(const FGeometry& MyGeometry, const FKeyEven
 			}
 
 			int32 NumItemsInAPage = GetNumLiveWidgets();
-			int32 Remainder = NumItemsInAPage % GetNumItemsWide();
+			int32 Remainder = NumItemsInAPage % GetNumItemsPerLine();
 			NumItemsInAPage -= Remainder;
 
 
