Sequencer: Force refresh immediately upon split section. Otherwise, the trimmed section will be in a delayed restore state after the new section, causing the evaluation to be out of date.

diff --git a/Engine/Source/Editor/Sequencer/Private/SequencerContextMenus.cpp b/Engine/Source/Editor/Sequencer/Private/SequencerContextMenus.cpp
--- a/Engine/Source/Editor/Sequencer/Private/SequencerContextMenus.cpp
+++ b/Engine/Source/Editor/Sequencer/Private/SequencerContextMenus.cpp
@@ -754,7 +754,7 @@ void FSectionContextMenu::SplitSection()
 	FQualifiedFrameTime SplitFrame = FQualifiedFrameTime(CurrentFrame, Sequencer->GetFocusedTickResolution());
 
 	MovieSceneToolHelpers::SplitSection(Sequencer->GetSelection().GetSelectedSections(), SplitFrame);
-	Sequencer->NotifyMovieSceneDataChanged( EMovieSceneDataChangeType::MovieSceneStructureItemAdded );
+	Sequencer->NotifyMovieSceneDataChanged( EMovieSceneDataChangeType::RefreshAllImmediately );
 }
 
 
