Sequencer: Downgraded a check for an ensure when re-adding a channel wtih a duplicate name to the sequencer tree

  - If a channel name exists multiple times within the channel meta-data, the assumption that adding a new node would increase the insertion index was invalid due to an existing valid entry being moved. This is technically bad content but we can handle it more gracefully than a fatal assert, so I'm downgrading it to an ensure.

diff --git a/Engine/Source/Editor/Sequencer/Private/SequencerSectionLayoutBuilder.cpp b/Engine/Source/Editor/Sequencer/Private/SequencerSectionLayoutBuilder.cpp
--- a/Engine/Source/Editor/Sequencer/Private/SequencerSectionLayoutBuilder.cpp
+++ b/Engine/Source/Editor/Sequencer/Private/SequencerSectionLayoutBuilder.cpp
@@ -24,7 +24,10 @@ TSharedPtr<FSequencerDisplayNode> FindAndRelocateExistingNode(TSharedRef<FSequen
 
 	const TArray<TSharedRef<FSequencerDisplayNode>>& CurrentChildren = ParentNode->GetChildNodes();
 
-	check(ExpectedIndex <= CurrentChildren.Num());
+	if (!ensureAlwaysMsgf(ExpectedIndex <= CurrentChildren.Num(), TEXT("Invalid desired index specified")))
+	{
+		ExpectedIndex = FMath::Clamp(ExpectedIndex, 0, CurrentChildren.Num());
+	}
 
 	// Common up-to-date case: check the desired insert index for an existing node that matches
 	if (CurrentChildren.IsValidIndex(ExpectedIndex) && MatchNode(CurrentChildren[ExpectedIndex]))
@@ -136,7 +139,8 @@ void FSequencerSectionLayoutBuilder::AddChannel( const FMovieSceneChannelHandle&
 	AddOrUpdateChannel(KeyAreaNode.ToSharedRef(), Channel);
 
 	// Move onto the next index at this level
-	++InsertIndexStack.Last();
+	int32& NextIndex = InsertIndexStack.Last();
+	NextIndex = FMath::Clamp(NextIndex + 1, 0, CurrentNode->GetChildNodes().Num());
 }
 
 void FSequencerSectionLayoutBuilder::AddOrUpdateChannel(TSharedRef<FSequencerSectionKeyAreaNode> KeyAreaNode, const FMovieSceneChannelHandle& Channel)
