Automatically deselect all sections after pasting keys in Sequencer via the Context Menu. Using the context menu automatically selects the section, so after paste, on infinite sections, the UI would appear to be stuck in a movement mode when really it was just that an infinite section was selected which we don't support outside of instantaneous events.

diff --git a/Engine/Source/Editor/Sequencer/Private/SequencerContextMenus.cpp b/Engine/Source/Editor/Sequencer/Private/SequencerContextMenus.cpp
--- a/Engine/Source/Editor/Sequencer/Private/SequencerContextMenus.cpp
+++ b/Engine/Source/Editor/Sequencer/Private/SequencerContextMenus.cpp
@@ -1558,6 +1558,7 @@ void FPasteContextMenu::PasteInto(int32 DestinationIndex, FName KeyAreaName)
 		// @todo sequencer: selection in transactions
 		FSequencerSelection& Selection = Sequencer->GetSelection();
 		Selection.SuspendBroadcast();
+		Selection.EmptySelectedSections();
 		Selection.EmptySelectedKeys();
 
 		for (const FSequencerSelectedKey& Key : NewSelection)
@@ -1566,6 +1567,7 @@ void FPasteContextMenu::PasteInto(int32 DestinationIndex, FName KeyAreaName)
 		}
 		Selection.ResumeBroadcast();
 		Selection.GetOnKeySelectionChanged().Broadcast();
+		Selection.GetOnSectionSelectionChanged().Broadcast();
 
 		Sequencer->OnClipboardUsed(Args.Clipboard);
 		Sequencer->NotifyMovieSceneDataChanged(EMovieSceneDataChangeType::TrackValueChanged);
