Sequencer: Paste onto selected object binding track

diff --git a/Engine/Source/Editor/Sequencer/Private/Sequencer.cpp b/Engine/Source/Editor/Sequencer/Private/Sequencer.cpp
--- a/Engine/Source/Editor/Sequencer/Private/Sequencer.cpp
+++ b/Engine/Source/Editor/Sequencer/Private/Sequencer.cpp
@@ -6737,20 +6737,12 @@ void FSequencer::DoPaste()
 
 bool FSequencer::PasteObjectBindings(const FString& TextToImport, TArray<FNotificationInfo>& PasteErrors)
 {
-	TArray<UMovieSceneCopyableBinding*> ImportedBindings;
-	ImportObjectBindingsFromText(TextToImport, ImportedBindings);
-
-	if (ImportedBindings.Num() == 0)
-	{
-		return false;
-	}
-	
 	TArray<UMovieSceneFolder*> SelectedParentFolders;
 	FString NewNodePath;
 	CalculateSelectedFolderAndPath(SelectedParentFolders, NewNodePath);
 
 	FScopedTransaction Transaction(FGenericCommands::Get().Paste->GetDescription());
-	
+
 	UMovieSceneSequence* OwnerSequence = GetFocusedMovieSceneSequence();
 	UObject* BindingContext = GetPlaybackContext();
 
@@ -6759,68 +6751,111 @@ bool FSequencer::PasteObjectBindings(const FString& TextToImport, TArray<FNotifi
 	TArray<FGuid> PossessableGuids;
 
 	TArray<FMovieSceneBinding> BindingsPasted;
-	for (UMovieSceneCopyableBinding* CopyableBinding : ImportedBindings)
+
+	TSet<TSharedRef<FSequencerDisplayNode>> SelectedNodes = Selection.GetSelectedOutlinerNodes();
+
+	TArray<FGuid> SelectedParentGuids;
+	for (TSharedRef<FSequencerDisplayNode> Node : SelectedNodes)
 	{
-		// Clear transient flags on the imported tracks
-		for (UMovieSceneTrack* CopiedTrack : CopyableBinding->Tracks)
+		if (Node->GetType() != ESequencerNode::Object)
 		{
-			CopiedTrack->ClearFlags(RF_Transient);
-			TArray<UObject*> Subobjects;
-			GetObjectsWithOuter(CopiedTrack, Subobjects);
-			for (UObject* Subobject : Subobjects)
-			{
-				Subobject->ClearFlags(RF_Transient);
-			}
+			continue;
 		}
 
-		if (CopyableBinding->Possessable.GetGuid().IsValid())
+		TSharedPtr<FSequencerObjectBindingNode> ObjectNode = StaticCastSharedRef<FSequencerObjectBindingNode>(Node);
+		if (ObjectNode.IsValid())
 		{
-			FGuid NewGuid = FGuid::NewGuid();
+			SelectedParentGuids.Add(ObjectNode->GetObjectBinding());
+		}
+	}
 
-			FMovieSceneBinding NewBinding(NewGuid, CopyableBinding->Binding.GetName(), CopyableBinding->Tracks);
+	int32 NumTargets = 1;
+	if (SelectedParentGuids.Num() > 1)
+	{
+		NumTargets = SelectedParentGuids.Num();
+	}
 
-			FMovieScenePossessable NewPossessable = CopyableBinding->Possessable;
-			NewPossessable.SetGuid(NewGuid);
+	for (int32 TargetIndex = 0; TargetIndex < NumTargets; ++TargetIndex)
+	{
+		TArray<UMovieSceneCopyableBinding*> ImportedBindings;
+		ImportObjectBindingsFromText(TextToImport, ImportedBindings);
 
-			MovieScene->AddPossessable(NewPossessable, NewBinding);
+		if (ImportedBindings.Num() == 0)
+		{
+			return false;
+		}
 
-			OldToNewGuidMap.Add(CopyableBinding->Possessable.GetGuid(), NewGuid);
+		for (UMovieSceneCopyableBinding* CopyableBinding : ImportedBindings)
+		{
+			// Clear transient flags on the imported tracks
+			for (UMovieSceneTrack* CopiedTrack : CopyableBinding->Tracks)
+			{
+				CopiedTrack->ClearFlags(RF_Transient);
+				TArray<UObject*> Subobjects;
+				GetObjectsWithOuter(CopiedTrack, Subobjects);
+				for (UObject* Subobject : Subobjects)
+				{
+					Subobject->ClearFlags(RF_Transient);
+				}
+			}
 
-			BindingsPasted.Add(NewBinding);
+			if (CopyableBinding->Possessable.GetGuid().IsValid())
+			{
+				FGuid NewGuid = FGuid::NewGuid();
 
-			PossessableGuids.Add(NewGuid);
-		}
-		else if (CopyableBinding->Spawnable.GetGuid().IsValid())
-		{
-			// We need to let the sequence create the spawnable so that it has everything set up properly internally.
-			// This is required to get spawnables with the correct references to object templates, object templates with
-			// correct owners, etc. However, making a new spawnable also creates the binding for us - this is a problem
-			// because we need to use our binding (which has tracks associated with it). To solve this, we let it create
-			// an object template based off of our (transient package owned) template, then find the newly created binding
-			// and update it.
-			FGuid NewGuid = MakeNewSpawnable(*CopyableBinding->SpawnableObjectTemplate, nullptr, false);
-			FMovieSceneBinding NewBinding(NewGuid, CopyableBinding->Binding.GetName(), CopyableBinding->Tracks);
-			FMovieSceneSpawnable* Spawnable = MovieScene->FindSpawnable(NewGuid);
+				FMovieSceneBinding NewBinding(NewGuid, CopyableBinding->Binding.GetName(), CopyableBinding->Tracks);
 
-			// Copy the name of the original spawnable too.
-			Spawnable->SetName(CopyableBinding->Spawnable.GetName());
+				FMovieScenePossessable NewPossessable = CopyableBinding->Possessable;
+				NewPossessable.SetGuid(NewGuid);
 
-			// Clear the transient flags on the copyable binding before assigning to the new spawnable
-			for (auto Track : NewBinding.GetTracks())
-			{
-				Track->ClearFlags(RF_Transient);
-				for (auto Section : Track->GetAllSections())
+				MovieScene->AddPossessable(NewPossessable, NewBinding);
+
+				OldToNewGuidMap.Add(CopyableBinding->Possessable.GetGuid(), NewGuid);
+
+				BindingsPasted.Add(NewBinding);
+
+				PossessableGuids.Add(NewGuid);
+
+				if (FMovieScenePossessable* Possessable = MovieScene->FindPossessable(NewGuid))
 				{
-					Section->ClearFlags(RF_Transient);
+					if (TargetIndex < SelectedParentGuids.Num())
+					{
+						Possessable->SetParent(SelectedParentGuids[TargetIndex]);
+					}
 				}
 			}
+			else if (CopyableBinding->Spawnable.GetGuid().IsValid())
+			{
+				// We need to let the sequence create the spawnable so that it has everything set up properly internally.
+				// This is required to get spawnables with the correct references to object templates, object templates with
+				// correct owners, etc. However, making a new spawnable also creates the binding for us - this is a problem
+				// because we need to use our binding (which has tracks associated with it). To solve this, we let it create
+				// an object template based off of our (transient package owned) template, then find the newly created binding
+				// and update it.
+				FGuid NewGuid = MakeNewSpawnable(*CopyableBinding->SpawnableObjectTemplate, nullptr, false);
+				FMovieSceneBinding NewBinding(NewGuid, CopyableBinding->Binding.GetName(), CopyableBinding->Tracks);
+				FMovieSceneSpawnable* Spawnable = MovieScene->FindSpawnable(NewGuid);
+
+				// Copy the name of the original spawnable too.
+				Spawnable->SetName(CopyableBinding->Spawnable.GetName());
 
-			// Replace the auto-generated binding with our deserialized bindings (which has our tracks)
-			MovieScene->ReplaceBinding(NewGuid, NewBinding);
+				// Clear the transient flags on the copyable binding before assigning to the new spawnable
+				for (auto Track : NewBinding.GetTracks())
+				{
+					Track->ClearFlags(RF_Transient);
+					for (auto Section : Track->GetAllSections())
+					{
+						Section->ClearFlags(RF_Transient);
+					}
+				}
+
+				// Replace the auto-generated binding with our deserialized bindings (which has our tracks)
+				MovieScene->ReplaceBinding(NewGuid, NewBinding);
 
-			OldToNewGuidMap.Add(CopyableBinding->Spawnable.GetGuid(), NewGuid);
+				OldToNewGuidMap.Add(CopyableBinding->Spawnable.GetGuid(), NewGuid);
 
-			BindingsPasted.Add(NewBinding);
+				BindingsPasted.Add(NewBinding);
+			}
 		}
 	}
 
@@ -6828,7 +6863,7 @@ bool FSequencer::PasteObjectBindings(const FString& TextToImport, TArray<FNotifi
 	for (auto PossessableGuid : PossessableGuids)
 	{
 		FMovieScenePossessable* Possessable = MovieScene->FindPossessable(PossessableGuid);
-		if (Possessable && OldToNewGuidMap.Contains(Possessable->GetParent()))
+		if (Possessable && OldToNewGuidMap.Contains(Possessable->GetParent()) && PossessableGuid != OldToNewGuidMap[Possessable->GetParent()])
 		{
 			Possessable->SetParent(OldToNewGuidMap[Possessable->GetParent()]);
 		}
