Sequencer: Don't clear node selection if it's already empty

diff --git a/Engine/Source/Editor/Sequencer/Private/Sequencer.cpp b/Engine/Source/Editor/Sequencer/Private/Sequencer.cpp
--- a/Engine/Source/Editor/Sequencer/Private/Sequencer.cpp
+++ b/Engine/Source/Editor/Sequencer/Private/Sequencer.cpp
@@ -5405,7 +5405,7 @@ void FSequencer::SynchronizeSequencerSelectionWithExternalSelection()
 		}
 	}
 
-	if (!bAllAlreadySelected || NodesToSelect.Num() == 0)
+	if (!bAllAlreadySelected || (NodesToSelect.Num() == 0 && Selection.GetSelectedOutlinerNodes().Num()))
 	{
 		Selection.SuspendBroadcast();
 		Selection.EmptySelectedOutlinerNodes();
