Sequencer: Update inner hierarchy after the parents have been set up so that UpdateSection() is given the proper object binding guid.

diff --git a/Engine/Source/Editor/Sequencer/Private/SequencerNodeTree.cpp b/Engine/Source/Editor/Sequencer/Private/SequencerNodeTree.cpp
--- a/Engine/Source/Editor/Sequencer/Private/SequencerNodeTree.cpp
+++ b/Engine/Source/Editor/Sequencer/Private/SequencerNodeTree.cpp
@@ -170,6 +170,11 @@ void FSequencerNodeTree::RefreshNodes(UMovieScene* MovieScene)
 			It->Value->SetParent(nullptr);
 			It.RemoveCurrent();
 		}
+		else
+		{
+			// Update after the above SetParent() because track sections need to have valid parent object bindings set up
+			It->Value->UpdateInnerHierarchy();
+		}
 	}
 	for (auto It = ObjectBindingToNode.CreateIterator(); It; ++It)
 	{
@@ -211,7 +216,6 @@ TSharedPtr<FSequencerTrackNode> FSequencerNodeTree::CreateOrUpdateTrack(UMovieSc
 
 	// Assign the serial number for this node to indicate that it is still relevant
 	TrackNode->TreeSerialNumber = SerialNumber;
-	TrackNode->UpdateInnerHierarchy();
 	return TrackNode;
 }
 
