Sequencer: Reimplement 'Show Selected Nodes Only' as part of filtering. Add support for option to UMG animations.

diff --git a/Engine/Source/Editor/Sequencer/Private/Sequencer.cpp b/Engine/Source/Editor/Sequencer/Private/Sequencer.cpp
--- a/Engine/Source/Editor/Sequencer/Private/Sequencer.cpp
+++ b/Engine/Source/Editor/Sequencer/Private/Sequencer.cpp
@@ -5324,9 +5324,16 @@ void GetRootObjectBindingNodes(const TArray<TSharedRef<FSequencerDisplayNode>>&
 
 void FSequencer::SynchronizeSequencerSelectionWithExternalSelection()
 {
+	if ( bUpdatingExternalSelection )
+	{
+		return;
+	}
+
 	UMovieSceneSequence* Sequence = GetFocusedMovieSceneSequence();
-	if ( bUpdatingExternalSelection || !IsLevelEditorSequencer() || ExactCast<ULevelSequence>(Sequence) == nullptr)
+	if( !IsLevelEditorSequencer() || ExactCast<ULevelSequence>(Sequence) == nullptr )
 	{
+		// Only level sequences have a full update here, but we still want filters to update for UMG animations
+		NodeTree->RequestFilterUpdate();
 		return;
 	}
 
diff --git a/Engine/Source/Editor/Sequencer/Private/SequencerNodeTree.cpp b/Engine/Source/Editor/Sequencer/Private/SequencerNodeTree.cpp
--- a/Engine/Source/Editor/Sequencer/Private/SequencerNodeTree.cpp
+++ b/Engine/Source/Editor/Sequencer/Private/SequencerNodeTree.cpp
@@ -273,7 +273,7 @@ TSharedRef<FSequencerFolderNode> FSequencerNodeTree::CreateOrUpdateFolder(UMovie
 
 bool FSequencerNodeTree::HasActiveFilter() const
 {
-	return (!FilterString.IsEmpty() || TrackFilters->Num() > 0 || TrackFilterLevelFilter->IsActive());
+	return (!FilterString.IsEmpty() || TrackFilters->Num() > 0 || TrackFilterLevelFilter->IsActive() || Sequencer.GetSequencerSettings()->GetShowSelectedNodesOnly());
 }
 
 TSharedPtr<FSequencerObjectBindingNode> FSequencerNodeTree::FindOrCreateObjectBinding(const FGuid& BindingID, const TSortedMap<FGuid, const FMovieSceneBinding*>& AllBindings, const TSortedMap<FGuid, FGuid>& ChildToParentBinding, TMap<FSequencerDisplayNode*, TSharedPtr<FSequencerDisplayNode>>* OutChildToParentMap)
@@ -716,6 +716,26 @@ static void FilterNodesRecursive( FSequencer& Sequencer, const TSharedRef<FSeque
 					ParentNode = ParentNode->GetParent();
 				}
 			}
+
+			if (bPasssedAnyFilters && Sequencer.GetSequencerSettings()->GetShowSelectedNodesOnly())
+			{
+				TSharedPtr<const FSequencerDisplayNode> ParentNode = StartNode->GetParent();
+				while (ParentNode.IsValid())
+				{
+					if (ParentNode->GetType() == ESequencerNode::Object)
+					{
+						const FSequencerObjectBindingNode* ObjectNode = static_cast<const FSequencerObjectBindingNode*>(ParentNode.Get());
+						const FMovieSceneBinding* Binding = Sequencer.GetFocusedMovieSceneSequence()->GetMovieScene()->FindBinding(ObjectNode->GetObjectBinding());
+						if (!(Binding && Sequencer.IsBindingVisible(*Binding)))
+						{
+							bPasssedAnyFilters = false;
+						}
+
+						break;
+					}
+					ParentNode = ParentNode->GetParent();
+				}
+			}
 		}
 	}
 	else if (StartNode->GetType() == ESequencerNode::Object)
@@ -730,6 +750,16 @@ static void FilterNodesRecursive( FSequencer& Sequencer, const TSharedRef<FSeque
 				break;
 			}
 		}
+
+		if (bPasssedAnyFilters && Sequencer.GetSequencerSettings()->GetShowSelectedNodesOnly())
+		{
+			UMovieScene* MovieScene = Sequencer.GetFocusedMovieSceneSequence()->GetMovieScene();
+			const FMovieSceneBinding* Binding = Sequencer.GetFocusedMovieSceneSequence()->GetMovieScene()->FindBinding(ObjectNode.GetObjectBinding());
+			if (Binding && !Sequencer.IsBindingVisible(*Binding))
+			{
+				bPasssedAnyFilters = false;
+			}
+		}
 	}
 
 	if (bPasssedAnyFilters)
@@ -795,7 +825,7 @@ void FSequencerNodeTree::UpdateFilters()
 	UWorld* World = PlaybackContext ? PlaybackContext->GetWorld() : nullptr;
 	TrackFilterLevelFilter->UpdateWorld(World);
 
-	if (TrackFilters->Num() > 0 || !FilterString.IsEmpty() || TrackFilterLevelFilter->IsActive())
+	if (TrackFilters->Num() > 0 || !FilterString.IsEmpty() || TrackFilterLevelFilter->IsActive() || Sequencer.GetSequencerSettings()->GetShowSelectedNodesOnly())
 	{
 		// Build a list of strings that must be matched
 		TArray<FString> FilterStrings;
diff --git a/Engine/Source/Editor/UMGEditor/Private/WidgetBlueprintEditor.cpp b/Engine/Source/Editor/UMGEditor/Private/WidgetBlueprintEditor.cpp
--- a/Engine/Source/Editor/UMGEditor/Private/WidgetBlueprintEditor.cpp
+++ b/Engine/Source/Editor/UMGEditor/Private/WidgetBlueprintEditor.cpp
@@ -296,6 +296,32 @@ void FWidgetBlueprintEditor::SelectObjects(const TSet<UObject*>& Objects)
 	OnSelectedWidgetsChanged.Broadcast();
 }
 
+bool FWidgetBlueprintEditor::IsBindingSelected(const FMovieSceneBinding& InBinding)
+{
+	TSet<FWidgetReference> Widgets = GetSelectedWidgets();
+	if (Widgets.Num() == 0)
+	{
+		return true;
+	}
+
+	UMovieSceneSequence* AnimationSequence = GetSequencer().Get()->GetFocusedMovieSceneSequence();
+	UObject* BindingContext = GetAnimationPlaybackContext();
+	TArray<UObject*, TInlineAllocator<1>> BoundObjects = AnimationSequence->LocateBoundObjects(InBinding.GetObjectGuid(), BindingContext);
+
+	if (BoundObjects.Num() == 0)
+	{
+		return false;
+	}
+	else if (Cast<UPanelSlot>(BoundObjects[0]))
+	{
+		return Widgets.Contains(GetReferenceFromPreview(Cast<UPanelSlot>(BoundObjects[0])->Content));
+	}
+	else
+	{
+		return Widgets.Contains(GetReferenceFromPreview(Cast<UWidget>(BoundObjects[0])));
+	}
+}
+
 void FWidgetBlueprintEditor::SetSelectedNamedSlot(TOptional<FNamedSlotSelection> InSelectedNamedSlot)
 {
 	OnSelectedWidgetsChanging.Broadcast();
@@ -810,6 +836,11 @@ TSharedPtr<ISequencer>& FWidgetBlueprintEditor::GetSequencer()
 		Sequencer->OnMovieSceneBindingsPasted().AddSP( this, &FWidgetBlueprintEditor::OnMovieSceneBindingsPasted );
 		// Change selected widgets in the sequencer tree view
 		Sequencer->GetSelectionChangedObjectGuids().AddSP(this, &FWidgetBlueprintEditor::SyncSelectedWidgetsWithSequencerSelection);
+		OnSelectedWidgetsChanged.AddSP(this, &FWidgetBlueprintEditor::SyncSequencerSelectionToSelectedWidgets);
+		
+		// Allow sequencer to test which bindings are selected
+		Sequencer->OnGetIsBindingVisible().BindRaw(this, &FWidgetBlueprintEditor::IsBindingSelected);
+
 		ChangeViewedAnimation(*UWidgetAnimation::GetNullAnimation());
 	}
 
@@ -1818,6 +1849,11 @@ void FWidgetBlueprintEditor::SyncSelectedWidgetsWithSequencerSelection(TArray<FG
 	}
 }
 
+void FWidgetBlueprintEditor::SyncSequencerSelectionToSelectedWidgets()
+{
+	GetSequencer()->ExternalSelectionHasChanged();
+}
+
 void FWidgetBlueprintEditor::UpdateTrackName(FGuid ObjectId)
 {
 	UWidgetAnimation* WidgetAnimation = Cast<UWidgetAnimation>(Sequencer->GetFocusedMovieSceneSequence());
diff --git a/Engine/Source/Editor/UMGEditor/Public/WidgetBlueprintEditor.h b/Engine/Source/Editor/UMGEditor/Public/WidgetBlueprintEditor.h
--- a/Engine/Source/Editor/UMGEditor/Public/WidgetBlueprintEditor.h
+++ b/Engine/Source/Editor/UMGEditor/Public/WidgetBlueprintEditor.h
@@ -108,6 +108,9 @@ public:
 	/** Sets the currently selected set of objects */
 	void SelectObjects(const TSet<UObject*>& Objects);
 
+	/** Called to determine whether a binding is selected in the tree view */
+	bool IsBindingSelected(const FMovieSceneBinding& InBinding);
+
 	/** Sets the selected named slot */
 	void SetSelectedNamedSlot(TOptional<FNamedSlotSelection> SelectedNamedSlot);
 
@@ -294,6 +297,9 @@ private:
 	/** Fire off when sequencer selection changed */
 	void SyncSelectedWidgetsWithSequencerSelection(TArray<FGuid> ObjectGuids);
 
+	/** Tell sequencer the selected widgets changed */
+	void SyncSequencerSelectionToSelectedWidgets();
+
 	/** Get the animation playback context */
 	UObject* GetAnimationPlaybackContext() const { return GetPreview(); }
 
