Sequencer: Allow folder name to be the same as a non folder sibling

diff --git a/Engine/Source/Editor/Sequencer/Private/DisplayNodes/SequencerFolderNode.cpp b/Engine/Source/Editor/Sequencer/Private/DisplayNodes/SequencerFolderNode.cpp
--- a/Engine/Source/Editor/Sequencer/Private/DisplayNodes/SequencerFolderNode.cpp
+++ b/Engine/Source/Editor/Sequencer/Private/DisplayNodes/SequencerFolderNode.cpp
@@ -76,7 +76,7 @@ void FSequencerFolderNode::SetDisplayName( const FText& NewDisplayName )
 		{
 			for ( TSharedRef<FSequencerDisplayNode> SiblingNode : ParentSeqNode->GetChildNodes() )
 			{
-				if ( SiblingNode != AsShared() )
+				if ( SiblingNode != AsShared() && SiblingNode->GetType() == ESequencerNode::Folder )
 				{
 					SiblingNames.Add( FName(*SiblingNode->GetDisplayName().ToString()) );
 				}
@@ -86,7 +86,7 @@ void FSequencerFolderNode::SetDisplayName( const FText& NewDisplayName )
 		{
 			for (TSharedRef<FSequencerDisplayNode> RootNode : GetParentTree().GetRootNodes())
 			{
-				if ( RootNode != AsShared() )
+				if ( RootNode != AsShared() && RootNode->GetType() == ESequencerNode::Folder )
 				{
 					SiblingNames.Add( FName(*RootNode->GetDisplayName().ToString()) );
 				}
