Fix for crash when scrolling the sequence node list while the the track context menu is open.

diff --git a/Engine/Source/Editor/Sequencer/Private/SSequencerTrackArea.cpp b/Engine/Source/Editor/Sequencer/Private/SSequencerTrackArea.cpp
--- a/Engine/Source/Editor/Sequencer/Private/SSequencerTrackArea.cpp
+++ b/Engine/Source/Editor/Sequencer/Private/SSequencerTrackArea.cpp
@@ -86,7 +86,7 @@ void SSequencerTrackArea::OnArrangeChildren( const FGeometry& AllottedGeometry,
 		const FTrackAreaSlot& CurChild = Children[ChildIndex];
 
 		const EVisibility ChildVisibility = CurChild.GetWidget()->GetVisibility();
-		if (!ArrangedChildren.Accepts(ChildVisibility))
+		if (!PinnedTrackLane.IsValid() || !ArrangedChildren.Accepts(ChildVisibility))
 		{
 			continue;
 		}
