Back out changelist 7977368
Turns out the code doesn't exist in the release branch.

diff --git a/Engine/Source/Editor/Sequencer/Private/SSequencerTrackArea.cpp b/Engine/Source/Editor/Sequencer/Private/SSequencerTrackArea.cpp
--- a/Engine/Source/Editor/Sequencer/Private/SSequencerTrackArea.cpp
+++ b/Engine/Source/Editor/Sequencer/Private/SSequencerTrackArea.cpp
@@ -86,7 +86,7 @@ void SSequencerTrackArea::OnArrangeChildren( const FGeometry& AllottedGeometry,
 		const FTrackAreaSlot& CurChild = Children[ChildIndex];
 
 		const EVisibility ChildVisibility = CurChild.GetWidget()->GetVisibility();
-		if (!PinnedTrackLane.IsValid() || !ArrangedChildren.Accepts(ChildVisibility))
+		if (!ArrangedChildren.Accepts(ChildVisibility))
 		{
 			continue;
 		}
