Sequencer: Fix crash when converting an unbound possessable to spawanble

diff --git a/Engine/Source/Editor/Sequencer/Private/Sequencer.cpp b/Engine/Source/Editor/Sequencer/Private/Sequencer.cpp
--- a/Engine/Source/Editor/Sequencer/Private/Sequencer.cpp
+++ b/Engine/Source/Editor/Sequencer/Private/Sequencer.cpp
@@ -7720,7 +7720,13 @@ TArray<FMovieSceneSpawnable*> FSequencer::ConvertToSpawnableInternal(FGuid Posse
 
 	TArray<FMovieSceneSpawnable*> CreatedSpawnables;
 
-	if (FoundObjects.Num() > 1)
+	if (FoundObjects.Num() == 0)
+	{
+		FMovieScenePossessable* Possessable = MovieScene->FindPossessable(PossessableGuid);
+
+		UE_LOG(LogSequencer, Error, TEXT("Failed to convert %s to spawnable because there are no objects bound to it"), Possessable ? *Possessable->GetName() : TEXT(""));
+	}
+	else if (FoundObjects.Num() > 1)
 	{
 		// Expand to individual possessables for each bound object, then convert each one individually
 		TArray<FGuid> ExpandedPossessableGuids = ExpandMultiplePossessableBindings(PossessableGuid);
