Sequencer: Speculative fix for a potential nullptr dereference

diff --git a/Engine/Source/Editor/Sequencer/Private/DisplayNodes/SequencerTrackNode.cpp b/Engine/Source/Editor/Sequencer/Private/DisplayNodes/SequencerTrackNode.cpp
--- a/Engine/Source/Editor/Sequencer/Private/DisplayNodes/SequencerTrackNode.cpp
+++ b/Engine/Source/Editor/Sequencer/Private/DisplayNodes/SequencerTrackNode.cpp
@@ -709,13 +709,15 @@ FLinearColor FSequencerTrackNode::GetDisplayNameColor() const
 
 		if (ObjectBinding.IsValid())
 		{
-			for (auto RuntimeObject : GetSequencer().FindBoundObjects(ObjectBinding, GetSequencer().GetFocusedTemplateID()))
+			for (TWeakObjectPtr<> WeakObject : GetSequencer().FindBoundObjects(ObjectBinding, GetSequencer().GetFocusedTemplateID()))
 			{
-				FTrackInstancePropertyBindings PropertyBinding(FName(*PropertyTrack->GetName()), PropertyTrack->GetPropertyPath());
-
-				if (PropertyBinding.GetProperty(*RuntimeObject))
+				if (UObject* Object = WeakObject.Get())
 				{
-					return FLinearColor::White;
+					FTrackInstancePropertyBindings PropertyBinding(PropertyTrack->GetPropertyName(), PropertyTrack->GetPropertyPath());
+					if (PropertyBinding.GetProperty(*Object))
+					{
+						return FLinearColor::White;
+					}
 				}
 			}
 
