Sequencer no longer tries to update your internal selection across tree rebuilds. This was causing issues with UMG selection & tree filtering on multi-row sections (ie: transforms) because the old logic would always pick the first key area, ie: X on restore. In reality, this is no longer needed because we no longer destroy the section, key area or tree nodes on tree update so there is no need to try to persist across refresh.

diff --git a/Engine/Source/Editor/Sequencer/Private/SSequencer.cpp b/Engine/Source/Editor/Sequencer/Private/SSequencer.cpp
--- a/Engine/Source/Editor/Sequencer/Private/SSequencer.cpp
+++ b/Engine/Source/Editor/Sequencer/Private/SSequencer.cpp
@@ -2167,84 +2167,6 @@ void RestoreSelectionState(const TArray<TSharedRef<FSequencerDisplayNode>>& Disp
 	}
 }
 
-void RestoreSectionSelection(const TSet<TWeakObjectPtr<UMovieSceneSection> >& SelectedSections, FSequencerSelection& Selection)
-{
-	for (auto Section : SelectedSections)
-	{
-		if (Section.IsValid())
-		{
-			Selection.AddToSelection(Section.Get());
-		}
-	}
-}
-
-/** Attempt to restore key selection from the specified set of selected keys. Only works for key areas that have the same key handles as their expired counterparts (this is generally the case) */
-void RestoreKeySelection(const TSet<FSequencerSelectedKey>& OldKeys, FSequencerSelection& Selection, FSequencerNodeTree& Tree)
-{
-	// Store a map of previous section/key area pairs to their current pairs
-	TMap<FSequencerSelectedKey, FSequencerSelectedKey> OldToNew;
-
-	for (FSequencerSelectedKey OldKeyTemplate : OldKeys)
-	{
-		// Cache of this key's handle for assignment to the new handle
-		TOptional<FKeyHandle> OldKeyHandle = OldKeyTemplate.KeyHandle;
-		// Reset the key handle so we can reuse cached section/key area pairs
-		OldKeyTemplate.KeyHandle.Reset();
-
-		FSequencerSelectedKey NewKeyTemplate = OldToNew.FindRef(OldKeyTemplate);
-		if (!NewKeyTemplate.Section)
-		{
-			// Not cached yet, so we'll need to search for it
-			for (const TSharedRef<FSequencerDisplayNode>& RootNode : Tree.GetRootNodes())
-			{
-				auto FindKeyArea =
-					[&](FSequencerDisplayNode& InNode)
-					{
-						FSequencerSectionKeyAreaNode* KeyAreaNode = nullptr;
-
-						if (InNode.GetType() == ESequencerNode::KeyArea)
-						{
-							KeyAreaNode = static_cast<FSequencerSectionKeyAreaNode*>(&InNode);
-						}
-						else if (InNode.GetType() == ESequencerNode::Track)
-						{
-							KeyAreaNode = static_cast<FSequencerTrackNode&>(InNode).GetTopLevelKeyNode().Get();
-						}
-
-						if (KeyAreaNode)
-						{
-							for (const TSharedRef<IKeyArea>& KeyArea : KeyAreaNode->GetAllKeyAreas())
-							{
-								if (KeyArea->GetOwningSection() == OldKeyTemplate.Section)
-								{
-									NewKeyTemplate.Section = OldKeyTemplate.Section;
-									NewKeyTemplate.KeyArea = KeyArea;
-									OldToNew.Add(OldKeyTemplate, NewKeyTemplate);
-									// stop iterating
-									return false;
-								}
-							}
-						}
-						return true;
-					};
-				
-				// If the traversal returned false, we've found what we're looking for - no need to look at any more nodes
-				if (!RootNode->Traverse_ParentFirst(FindKeyArea))
-				{
-					break;
-				}
-			}
-		}
-
-		// If we've got a curretn section/key area pair, we can add this key to the selection
-		if (NewKeyTemplate.Section)
-		{
-			NewKeyTemplate.KeyHandle = OldKeyHandle;
-			Selection.AddToSelection(NewKeyTemplate);
-		}
-	}
-}
-
 void SSequencer::UpdateLayoutTree()
 {
 	TrackArea->Empty();
@@ -2273,19 +2195,15 @@ void SSequencer::UpdateLayoutTree()
 
 		// Suspend broadcasting selection changes because we don't want unnecessary rebuilds.
 		Sequencer->GetSelection().SuspendBroadcast();
-		Sequencer->GetSelection().Empty();
 
 		// Update the node tree
 		Sequencer->GetNodeTree()->Update();
 
-		// Restore the selection state.
+		// Restore the selection state. This is still needed to apply the AdditionalSelectionsToAdd hack.
 		RestoreSelectionState(Sequencer->GetNodeTree()->GetRootNodes(), SelectedPathNames, SequencerPtr.Pin()->GetSelection());	// Update to actor selection.
 
 		// This must come after the selection state has been restored so that the tree and curve editor are populated with the correctly selected nodes
 		TreeView->Refresh();
-
-		RestoreKeySelection(SelectedKeys, Sequencer->GetSelection(), *Sequencer->GetNodeTree());
-		RestoreSectionSelection(SelectedSections, Sequencer->GetSelection());
 		
 		// If we've manually specified an additional selection to add it's because the item was newly created.
 		// Now that the treeview has been refreshed and selection restored, we'll try to focus the first item
