Working check-in: new gadget intros

diff --git a/Engine/Source/Editor/Sequencer/Private/DisplayNodes/SequencerObjectBindingNode.cpp b/Engine/Source/Editor/Sequencer/Private/DisplayNodes/SequencerObjectBindingNode.cpp
--- a/Engine/Source/Editor/Sequencer/Private/DisplayNodes/SequencerObjectBindingNode.cpp
+++ b/Engine/Source/Editor/Sequencer/Private/DisplayNodes/SequencerObjectBindingNode.cpp
@@ -36,6 +36,7 @@
 #include "MovieSceneSequence.h"
 #include "MovieScene.h"
 #include "MovieSceneFolder.h"
+#include "SExposeBindingWidget.h"
 #include "ISequencerTrackEditor.h"
 
 #include "Tracks/MovieSceneSpawnTrack.h"
@@ -156,6 +157,12 @@ void FSequencerObjectBindingNode::BuildContextMenu(FMenuBuilder& MenuBuilder)
 		MenuBuilder.PushExtender(Extender.ToSharedRef());
 	}
 
+	MenuBuilder.AddSubMenu(
+		LOCTEXT("ExposeBindingLabel", "Expose"),
+		LOCTEXT("ExposeBindingTooltip", "Specifies options for exposing this binding to external systems as a persistent name."),
+		FNewMenuDelegate::CreateSP(this, &FSequencerObjectBindingNode::AddExposeMenu)
+	);
+
 	if (GetSequencer().IsLevelEditorSequencer())
 	{
 		FSequencer*  Sequencer  = &GetSequencer();
@@ -458,6 +465,10 @@ void FSequencerObjectBindingNode::AddAssignActorMenu(FMenuBuilder& MenuBuilder)
 	GetSequencer().AssignActor(MenuBuilder, ObjectBinding);
 }
 
+void FSequencerObjectBindingNode::AddExposeMenu(FMenuBuilder& MenuBuilder)
+{
+	MenuBuilder.AddWidget(SNew(SExposeBindingWidget, GetSequencer().AsShared()), FText(), true, false);
+}
 
 bool FSequencerObjectBindingNode::CanRenameNode() const
 {
diff --git a/Engine/Source/Editor/Sequencer/Private/DisplayNodes/SequencerObjectBindingNode.h b/Engine/Source/Editor/Sequencer/Private/DisplayNodes/SequencerObjectBindingNode.h
--- a/Engine/Source/Editor/Sequencer/Private/DisplayNodes/SequencerObjectBindingNode.h
+++ b/Engine/Source/Editor/Sequencer/Private/DisplayNodes/SequencerObjectBindingNode.h
@@ -80,6 +80,7 @@ protected:
 	void AddSpawnOwnershipMenu(FMenuBuilder& MenuBuilder);
 	void AddSpawnLevelMenu(FMenuBuilder& MenuBuilder);
 	void AddAssignActorMenu(FMenuBuilder& MenuBuilder);
+	void AddExposeMenu(FMenuBuilder& MenuBuilder);
 
 	/** Get class for object binding */
 	const UClass* GetClassForObjectBinding() const;
diff --git a/Engine/Source/Editor/Sequencer/Private/SExposeBindingWidget.cpp b/Engine/Source/Editor/Sequencer/Private/SExposeBindingWidget.cpp
new file mode 100644
--- /dev/null
+++ b/Engine/Source/Editor/Sequencer/Private/SExposeBindingWidget.cpp
@@ -0,0 +1,228 @@
+// Copyright 1998-2019 Epic Games, Inc. All Rights Reserved.
+
+#include "SExposeBindingWidget.h"
+
+#include "Sequencer.h"
+#include "MovieSceneSequence.h"
+#include "MovieScene.h"
+#include "MovieSceneObjectBindingID.h"
+
+#include "DisplayNodes/SequencerObjectBindingNode.h"
+
+#include "Widgets/SNullWidget.h"
+#include "Widgets/SBoxPanel.h"
+#include "Widgets/Layout/SBox.h"
+#include "Widgets/Text/STextBlock.h"
+#include "Widgets/Input/SEditableTextBox.h"
+#include "Widgets/Input/SButton.h"
+
+#include "EditorStyleSet.h"
+#include "EditorFontGlyphs.h"
+#include "ScopedTransaction.h"
+
+#include "Algo/Sort.h"
+
+#define LOCTEXT_NAMESPACE "SExposeBindingWidget"
+
+void SExposeBindingWidget::Construct(const FArguments& InArgs, TWeakPtr<ISequencer> InWeakSequencer)
+{
+	WeakSequencer = InWeakSequencer;
+
+	Reconstruct();
+}
+
+void SExposeBindingWidget::Reconstruct()
+{
+	TSharedPtr<ISequencer> Sequencer  = WeakSequencer.Pin();
+	UMovieSceneSequence*   Sequence   = Sequencer ? Sequencer->GetFocusedMovieSceneSequence() : nullptr;
+	UMovieScene*           MovieScene = Sequence  ? Sequence->GetMovieScene()                 : nullptr;
+
+	if (!MovieScene)
+	{
+		ChildSlot
+		[
+			SNullWidget::NullWidget
+		];
+
+		return;
+	}
+
+	TArray<FGuid> SelectedObjectBindings;
+	for (const TSharedRef<FSequencerDisplayNode>& Node : Sequencer->GetSelection().GetSelectedOutlinerNodes())
+	{
+		if (Node->GetType() == ESequencerNode::Object)
+		{
+			SelectedObjectBindings.Add(StaticCastSharedRef<FSequencerObjectBindingNode>(Node)->GetObjectBinding());
+		}
+	}
+
+	TSharedRef<SVerticalBox> VerticalBox = SNew(SVerticalBox);
+
+	FMovieSceneSequenceID      SequenceID = Sequencer->GetFocusedTemplateID();
+
+	TArray<FName> ExposedNames;
+
+	for (const TTuple<FName, FMovieSceneObjectBindingIDs>& Pair : MovieScene->AllBindingGroups())
+	{
+		bool bContainsAll = false;
+		for (const FGuid& ID : SelectedObjectBindings)
+		{
+			FMovieSceneObjectBindingID ThisBindingID(ID, SequenceID);
+			if (!Pair.Value.IDs.Contains(ThisBindingID))
+			{
+				bContainsAll = false;
+				break;
+			}
+			
+			bContainsAll = true;
+		}
+
+		if (bContainsAll)
+		{
+			ExposedNames.Add(Pair.Key);
+		}
+	}
+
+	// Sort exposed names alphabetically
+	Algo::Sort(ExposedNames, [](const FName& A, const FName& B) { return A.Compare(B) < 0; });
+
+	for (FName ExposedName : ExposedNames)
+	{
+		VerticalBox->AddSlot()
+		[
+			SNew(SHorizontalBox)
+
+			+ SHorizontalBox::Slot()
+			.VAlign(VAlign_Center)
+			.Padding(FMargin(7.f, 5.f))
+			[
+				SNew(STextBlock)
+				.Text(FText::FromName(ExposedName))
+			]
+
+			+ SHorizontalBox::Slot()
+			.AutoWidth()
+			.Padding(FMargin(7.f, 5.f))
+			[
+				SNew(SButton)
+				.ButtonStyle(FEditorStyle::Get(), "HoverHintOnly")
+				.OnClicked(this, &SExposeBindingWidget::RemoveFromExposedName, ExposedName)
+				[
+					SNew(STextBlock)
+					.ColorAndOpacity(FLinearColor::Red)
+					.Font(FEditorStyle::Get().GetFontStyle("FontAwesome.11"))
+					.Text(FEditorFontGlyphs::Times)
+				]
+			]
+		];
+	}
+
+	TSharedRef<SEditableTextBox> EditableText = SNew(SEditableTextBox)
+	.OnTextCommitted(this, &SExposeBindingWidget::OnNewTextCommitted)
+	.HintText(LOCTEXT("AddNew_Hint", "Add new"));
+
+	auto OnClicked = [this, WeakEditableText = TWeakPtr<SEditableTextBox>(EditableText)]() -> FReply
+	{
+		TSharedPtr<SEditableTextBox> LocalEditableText = WeakEditableText.Pin();
+		if (LocalEditableText)
+		{
+			this->OnNewTextCommitted(LocalEditableText->GetText(), ETextCommit::OnEnter);
+		}
+		return FReply::Handled();
+	};
+
+	VerticalBox->AddSlot()
+	[
+		SNew(SHorizontalBox)
+
+		+ SHorizontalBox::Slot()
+		.VAlign(VAlign_Center)
+		.Padding(FMargin(7.f, 5.f))
+		[
+			SNew(SBox)
+			.MinDesiredWidth(100.f)
+			[
+				EditableText
+			]
+		]
+
+		+ SHorizontalBox::Slot()
+		.AutoWidth()
+		.Padding(FMargin(7.f, 5.f))
+		[
+			SNew(SButton)
+			.ButtonStyle(FEditorStyle::Get(), "HoverHintOnly")
+			.OnClicked_Lambda(OnClicked)
+			[
+				SNew(STextBlock)
+				.ColorAndOpacity(FLinearColor::Green)
+				.Font(FEditorStyle::Get().GetFontStyle("FontAwesome.11"))
+				.Text(FEditorFontGlyphs::Plus)
+			]
+		]
+	];
+
+	ChildSlot
+	[
+		VerticalBox
+	];
+}
+
+void SExposeBindingWidget::OnNewTextCommitted(const FText& InNewText, ETextCommit::Type CommitType)
+{
+	if (CommitType == ETextCommit::OnEnter && !InNewText.IsEmpty())
+	{
+		FName NewName = *InNewText.ToString();
+		ExposeAsName(NewName);
+	}
+}
+
+void SExposeBindingWidget::ExposeAsName(FName InNewName)
+{
+	TSharedPtr<ISequencer> Sequencer = WeakSequencer.Pin();
+	UMovieSceneSequence*   Sequence = Sequencer ? Sequencer->GetFocusedMovieSceneSequence() : nullptr;
+	UMovieScene*           MovieScene = Sequence ? Sequence->GetMovieScene() : nullptr;
+
+	FScopedTransaction Transaction(FText::Format(LOCTEXT("ExposeBinding_Transaction", "Expose binding as {0}"), FText::FromName(InNewName)));
+
+	MovieScene->Modify();
+
+	for (const TSharedRef<FSequencerDisplayNode>& Node : Sequencer->GetSelection().GetSelectedOutlinerNodes())
+	{
+		if (Node->GetType() == ESequencerNode::Object)
+		{
+			FGuid ObjectBindingID = StaticCastSharedRef<FSequencerObjectBindingNode>(Node)->GetObjectBinding();
+
+			MovieScene->ExposeBinding(InNewName, FMovieSceneObjectBindingID(ObjectBindingID, Sequencer->GetFocusedTemplateID(), EMovieSceneObjectBindingSpace::Local));
+		}
+	}
+
+	Reconstruct();
+}
+
+FReply SExposeBindingWidget::RemoveFromExposedName(FName InNameToRemove)
+{
+	TSharedPtr<ISequencer> Sequencer  = WeakSequencer.Pin();
+	UMovieSceneSequence*   Sequence   = Sequencer ? Sequencer->GetFocusedMovieSceneSequence() : nullptr;
+	UMovieScene*           MovieScene = Sequence  ? Sequence->GetMovieScene()                 : nullptr;
+
+	FScopedTransaction Transaction(FText::Format(LOCTEXT("ExposeBinding_Transaction", "Stop exposing binding as {0}"), FText::FromName(InNameToRemove)));
+
+	MovieScene->Modify();
+
+	for (const TSharedRef<FSequencerDisplayNode>& Node : Sequencer->GetSelection().GetSelectedOutlinerNodes())
+	{
+		if (Node->GetType() == ESequencerNode::Object)
+		{
+			FGuid ObjectBindingID = StaticCastSharedRef<FSequencerObjectBindingNode>(Node)->GetObjectBinding();
+
+			MovieScene->RemoveExposedBinding(InNameToRemove, FMovieSceneObjectBindingID(ObjectBindingID, Sequencer->GetFocusedTemplateID(), EMovieSceneObjectBindingSpace::Local));
+		}
+	}
+
+	Reconstruct();
+	return FReply::Handled();
+}
+
+
+#undef LOCTEXT_NAMESPACE
\ No newline at end of file
diff --git a/Engine/Source/Editor/Sequencer/Private/SExposeBindingWidget.h b/Engine/Source/Editor/Sequencer/Private/SExposeBindingWidget.h
new file mode 100644
--- /dev/null
+++ b/Engine/Source/Editor/Sequencer/Private/SExposeBindingWidget.h
@@ -0,0 +1,31 @@
+// Copyright 1998-2019 Epic Games, Inc. All Rights Reserved.
+
+#include "Widgets/DeclarativeSyntaxSupport.h"
+#include "Widgets/SCompoundWidget.h"
+
+class ISequencer;
+
+class SExposeBindingWidget : public SCompoundWidget
+{
+public:
+	SLATE_BEGIN_ARGS(SExposeBindingWidget){}
+	SLATE_END_ARGS()
+
+
+	void Construct(const FArguments& InArgs, TWeakPtr<ISequencer> InWeakSequencer);
+
+private:
+
+	void Reconstruct();
+
+	void OnNewTextCommitted(const FText& InNewText, ETextCommit::Type CommitType);
+
+	void ExposeAsName(FName InNewName);
+
+	FReply RemoveFromExposedName(FName InNameToRemove);
+
+private:
+
+	/** The sequencer UI instance that is currently open */
+	TWeakPtr<ISequencer> WeakSequencer;
+};
\ No newline at end of file
diff --git a/Engine/Source/Editor/Sequencer/Private/SExposedBindingsWidget.cpp b/Engine/Source/Editor/Sequencer/Private/SExposedBindingsWidget.cpp
new file mode 100644
--- /dev/null
+++ b/Engine/Source/Editor/Sequencer/Private/SExposedBindingsWidget.cpp
@@ -0,0 +1,373 @@
+// Copyright 1998-2019 Epic Games, Inc. All Rights Reserved.
+
+#include "SExposedBindingsWidget.h"
+
+#include "Sequencer.h"
+#include "MovieSceneSequence.h"
+#include "MovieScene.h"
+#include "MovieSceneObjectBindingID.h"
+#include "MovieSceneObjectBindingIDPicker.h"
+
+#include "Widgets/SNullWidget.h"
+#include "Widgets/SBoxPanel.h"
+#include "Widgets/Layout/SBox.h"
+#include "Widgets/Text/STextBlock.h"
+#include "Widgets/Input/SEditableTextBox.h"
+#include "Widgets/Input/SButton.h"
+#include "Widgets/Input/SComboButton.h"
+#include "Framework/MultiBox/MultiBoxBuilder.h"
+
+#include "EditorStyleSet.h"
+#include "EditorFontGlyphs.h"
+#include "ScopedTransaction.h"
+
+#include "Algo/Sort.h"
+
+#define LOCTEXT_NAMESPACE "SExposedBindingsWidget"
+
+
+struct SExposedBindingPicker : public SCompoundWidget, public FMovieSceneObjectBindingIDPicker
+{
+	SLATE_BEGIN_ARGS(SExposedBindingPicker){}
+	SLATE_END_ARGS()
+
+
+	void Construct(const FArguments& InArgs, FName InExposedName, UMovieSceneSequence* InSequence, FMovieSceneObjectBindingID InCurrentValue)
+	{
+		ExposedName = InExposedName;
+		Sequence = InSequence;
+		CurrentValue = InCurrentValue;
+
+		Initialize();
+
+		ChildSlot
+		[
+			SNew(SComboButton)
+			.ToolTipText(this, &SExposedBindingPicker::GetToolTipText)
+			.OnGetMenuContent(this, &SExposedBindingPicker::GetPickerMenu)
+			.ContentPadding(FMargin(4.0, 2.0))
+			.ButtonContent()
+			[
+				GetCurrentItemWidget(SNew(STextBlock))
+			]
+		];
+	}
+
+private:
+
+	virtual UMovieSceneSequence* GetSequence() const override
+	{
+		return Sequence;
+	}
+
+	virtual void SetCurrentValue(const FMovieSceneObjectBindingID& InBindingId) override
+	{
+		UMovieScene* MovieScene = Sequence->GetMovieScene();
+
+		MovieScene->Modify();
+		MovieScene->RemoveExposedBinding(ExposedName, CurrentValue);
+		MovieScene->ExposeBinding(ExposedName, InBindingId);
+
+		CurrentValue = InBindingId;
+	}
+
+	virtual FMovieSceneObjectBindingID GetCurrentValue() const
+	{
+		return CurrentValue;
+	}
+
+	FName ExposedName;
+	UMovieSceneSequence* Sequence;
+	FMovieSceneObjectBindingID CurrentValue;
+};
+
+struct SExposedNameSubMenuContent : public SCompoundWidget
+{
+	SLATE_BEGIN_ARGS(SExposedNameSubMenuContent){}
+	SLATE_END_ARGS()
+
+	void Construct(const FArguments& InArgs, FName InExposedName, UMovieSceneSequence* InSequence)
+	{
+		ExposedName = InExposedName;
+		Sequence = InSequence;
+
+		Reconstruct();
+	}
+
+	void Reconstruct()
+	{
+		FMenuBuilder MenuBuilder(true, nullptr);
+
+		const FMovieSceneObjectBindingIDs* IDs = Sequence->GetMovieScene()->AllBindingGroups().Find(ExposedName);
+		if (IDs)
+		{
+			for (FMovieSceneObjectBindingID ID : IDs->IDs)
+			{
+				MenuBuilder.AddWidget(
+					SNew(SHorizontalBox)
+
+					+ SHorizontalBox::Slot()
+					.Padding(FMargin(0.f, 0.f, 5.f, 0.f))
+					.AutoWidth()
+					.VAlign(VAlign_Center)
+					[
+						SNew(SButton)
+						.ButtonStyle(FEditorStyle::Get(), "HoverHintOnly")
+						.OnClicked(this, &SExposedNameSubMenuContent::OnRemove, ID)
+						[
+							SNew(STextBlock)
+							.ColorAndOpacity(FLinearColor::Red)
+							.Font(FEditorStyle::Get().GetFontStyle("FontAwesome.11"))
+							.Text(FEditorFontGlyphs::Times)
+						]
+					]
+
+					+ SHorizontalBox::Slot()
+					.VAlign(VAlign_Center)
+					[
+						SNew(SExposedBindingPicker, ExposedName, Sequence, ID)
+					],
+					FText(),
+					true
+				);
+			}
+		}
+
+		MenuBuilder.AddWidget(
+			SNew(SBox)
+			.HAlign(HAlign_Left)
+			[
+				SNew(SButton)
+				.ButtonStyle(FEditorStyle::Get(), "HoverHintOnly")
+				.OnClicked(this, &SExposedNameSubMenuContent::AddEmpty)
+				[
+					SNew(STextBlock)
+					.ColorAndOpacity(FLinearColor::Green)
+					.Font(FEditorStyle::Get().GetFontStyle("FontAwesome.11"))
+					.Text(FEditorFontGlyphs::Plus)
+				]
+			],
+			FText(),
+			true
+		);
+
+		ChildSlot
+		[
+			MenuBuilder.MakeWidget()
+		];
+	}
+
+	FReply AddEmpty()
+	{
+		UMovieScene* MovieScene = Sequence->GetMovieScene();
+
+		FScopedTransaction Transaction(FText::Format(LOCTEXT("AddNewExposedBinding", "Add new binding for exposed name {0}"), FText::FromName(ExposedName)));
+
+		MovieScene->Modify();
+		MovieScene->ExposeBinding(ExposedName, FMovieSceneObjectBindingID());
+
+		Reconstruct();
+
+		return FReply::Handled();
+	}
+
+	FReply OnRemove(FMovieSceneObjectBindingID ID)
+	{
+		UMovieScene* MovieScene = Sequence->GetMovieScene();
+
+		FScopedTransaction Transaction(FText::Format(LOCTEXT("RemoveExposedBinding", "Remove binding from exposed name {0}"), FText::FromName(ExposedName)));
+
+		MovieScene->Modify();
+		MovieScene->RemoveExposedBinding(ExposedName, ID);
+
+		Reconstruct();
+		return FReply::Handled();
+	}
+
+	FName ExposedName;
+	UMovieSceneSequence* Sequence;
+};
+
+void SExposedBindingsWidget::Construct(const FArguments& InArgs, TWeakPtr<ISequencer> InWeakSequencer)
+{
+	WeakSequencer = InWeakSequencer;
+
+	Reconstruct();
+}
+
+void SExposedBindingsWidget::Reconstruct()
+{
+	TSharedPtr<ISequencer> Sequencer  = WeakSequencer.Pin();
+	UMovieSceneSequence*   Sequence   = Sequencer ? Sequencer->GetFocusedMovieSceneSequence() : nullptr;
+	UMovieScene*           MovieScene = Sequence  ? Sequence->GetMovieScene()                 : nullptr;
+
+	FMenuBuilder MenuBuilder(true, nullptr);
+
+	if (!MovieScene)
+	{
+		ChildSlot
+		[
+			SNullWidget::NullWidget
+		];
+
+		return;
+	}
+
+	TSharedRef<SVerticalBox> VerticalBox = SNew(SVerticalBox);
+
+	TArray<FName> ExposedNames;
+
+	const TMap<FName, FMovieSceneObjectBindingIDs>& BindingMap = MovieScene->AllBindingGroups();
+	for (const TTuple<FName, FMovieSceneObjectBindingIDs>& Pair : BindingMap)
+	{
+		ExposedNames.Add(Pair.Key);
+	}
+
+	// Sort exposed names alphabetically
+	Algo::Sort(ExposedNames, [](const FName& A, const FName& B) { return A.Compare(B) < 0; });
+
+	for (FName ExposedName : ExposedNames)
+	{
+		TSharedRef<SWidget> MenuContent =
+			SNew(SHorizontalBox)
+
+			+ SHorizontalBox::Slot()
+			.Padding(FMargin(0.f, 0.f, 5.f, 0.f))
+			.VAlign(VAlign_Center)
+			.AutoWidth()
+			[
+				SNew(SButton)
+				.ButtonStyle(FEditorStyle::Get(), "HoverHintOnly")
+				.OnClicked(this, &SExposedBindingsWidget::RemoveExposedName, ExposedName)
+				[
+					SNew(STextBlock)
+					.ColorAndOpacity(FLinearColor::Red)
+					.Font(FEditorStyle::Get().GetFontStyle("FontAwesome.11"))
+					.Text(FEditorFontGlyphs::Times)
+				]
+			]
+
+			+ SHorizontalBox::Slot()
+			.VAlign(VAlign_Center)
+			[
+				SNew(STextBlock)
+				.Text(this, &SExposedBindingsWidget::GetSubMenuLabel, ExposedName)
+			];
+
+		MenuBuilder.AddSubMenu(MenuContent, FNewMenuDelegate::CreateSP(this, &SExposedBindingsWidget::MakeSubMenu, ExposedName));
+	}
+
+	TSharedRef<SEditableTextBox> EditableText = SNew(SEditableTextBox)
+	.OnTextCommitted(this, &SExposedBindingsWidget::OnNewTextCommitted)
+	.HintText(LOCTEXT("AddNew_Hint", "Add new"));
+
+	auto OnClicked = [this, WeakEditableText = TWeakPtr<SEditableTextBox>(EditableText)]() -> FReply
+	{
+		TSharedPtr<SEditableTextBox> LocalEditableText = WeakEditableText.Pin();
+		if (LocalEditableText)
+		{
+			this->OnNewTextCommitted(LocalEditableText->GetText(), ETextCommit::OnEnter);
+		}
+		return FReply::Handled();
+	};
+
+	TSharedRef<SWidget> NewNameContent =
+		SNew(SHorizontalBox)
+
+		+ SHorizontalBox::Slot()
+		.VAlign(VAlign_Center)
+		.HAlign(HAlign_Left)
+		.Padding(FMargin(0.f, 0.f, 5.f, 0.f))
+		.AutoWidth()
+		[
+			SNew(SButton)
+			.ButtonStyle(FEditorStyle::Get(), "HoverHintOnly")
+			.OnClicked_Lambda(OnClicked)
+			[
+				SNew(STextBlock)
+				.ColorAndOpacity(FLinearColor::Green)
+				.Font(FEditorStyle::Get().GetFontStyle("FontAwesome.11"))
+				.Text(FEditorFontGlyphs::Plus)
+			]
+		]
+
+		+ SHorizontalBox::Slot()
+		.VAlign(VAlign_Center)
+		[
+			SNew(SBox)
+			.MinDesiredWidth(100.f)
+			[
+				EditableText
+			]
+		];
+
+	MenuBuilder.AddWidget(NewNameContent, FText(), true);
+
+	ChildSlot
+	[
+		MenuBuilder.MakeWidget()
+	];
+}
+
+FText SExposedBindingsWidget::GetSubMenuLabel(FName ExposedName) const
+{
+	TSharedPtr<ISequencer> Sequencer  = WeakSequencer.Pin();
+	UMovieSceneSequence*   Sequence   = Sequencer ? Sequencer->GetFocusedMovieSceneSequence() : nullptr;
+	UMovieScene*           MovieScene = Sequence  ? Sequence->GetMovieScene()                 : nullptr;
+	if (!MovieScene)
+	{
+		return FText();
+	}
+
+	return FText::Format(LOCTEXT("NamedBinding_Label", "{0} ({1} {1}|plural(one=binding,other=Bindings))")
+		, FText::FromName(ExposedName)
+		, FText::AsNumber(MovieScene->AllBindingGroups().FindChecked(ExposedName).IDs.Num())
+	);
+}
+
+void SExposedBindingsWidget::MakeSubMenu(FMenuBuilder& MenuBuilder, FName ExposedName)
+{
+	TSharedPtr<ISequencer> Sequencer = WeakSequencer.Pin();
+	MenuBuilder.AddWidget(SNew(SExposedNameSubMenuContent, ExposedName, Sequencer->GetFocusedMovieSceneSequence()), FText(), true);
+}
+
+void SExposedBindingsWidget::OnNewTextCommitted(const FText& InNewText, ETextCommit::Type CommitType)
+{
+	if (CommitType == ETextCommit::OnEnter && !InNewText.IsEmpty())
+	{
+		FName NewName = *InNewText.ToString();
+		ExposeAsName(NewName);
+	}
+}
+
+void SExposedBindingsWidget::ExposeAsName(FName InNewName)
+{
+	TSharedPtr<ISequencer> Sequencer = WeakSequencer.Pin();
+	UMovieSceneSequence*   Sequence = Sequencer ? Sequencer->GetFocusedMovieSceneSequence() : nullptr;
+	UMovieScene*           MovieScene = Sequence ? Sequence->GetMovieScene() : nullptr;
+
+	FScopedTransaction Transaction(FText::Format(LOCTEXT("ExposeBinding_Transaction", "Expose binding as {0}"), FText::FromName(InNewName)));
+
+	MovieScene->Modify();
+	MovieScene->ExposeBinding(InNewName);
+
+	Reconstruct();
+}
+
+FReply SExposedBindingsWidget::RemoveExposedName(FName InNameToRemove)
+{
+	TSharedPtr<ISequencer> Sequencer  = WeakSequencer.Pin();
+	UMovieSceneSequence*   Sequence   = Sequencer ? Sequencer->GetFocusedMovieSceneSequence() : nullptr;
+	UMovieScene*           MovieScene = Sequence  ? Sequence->GetMovieScene()                 : nullptr;
+
+	FScopedTransaction Transaction(FText::Format(LOCTEXT("ExposeBinding_Transaction", "Stop exposing binding {0}"), FText::FromName(InNameToRemove)));
+
+	MovieScene->Modify();
+	MovieScene->RemoveExposedBinding(InNameToRemove);
+
+	Reconstruct();
+	return FReply::Handled();
+}
+
+
+#undef LOCTEXT_NAMESPACE
\ No newline at end of file
diff --git a/Engine/Source/Editor/Sequencer/Private/SExposedBindingsWidget.h b/Engine/Source/Editor/Sequencer/Private/SExposedBindingsWidget.h
new file mode 100644
--- /dev/null
+++ b/Engine/Source/Editor/Sequencer/Private/SExposedBindingsWidget.h
@@ -0,0 +1,36 @@
+// Copyright 1998-2019 Epic Games, Inc. All Rights Reserved.
+
+#include "Widgets/DeclarativeSyntaxSupport.h"
+#include "Widgets/SCompoundWidget.h"
+
+class ISequencer;
+class FMenuBuilder;
+
+class SExposedBindingsWidget : public SCompoundWidget
+{
+public:
+
+	SLATE_BEGIN_ARGS(SExposedBindingsWidget){}
+	SLATE_END_ARGS()
+
+	void Construct(const FArguments& InArgs, TWeakPtr<ISequencer> InWeakSequencer);
+
+private:
+
+	void Reconstruct();
+
+	void MakeSubMenu(FMenuBuilder& MenuBuilder, FName ExposedName);
+
+	void OnNewTextCommitted(const FText& InNewText, ETextCommit::Type CommitType);
+
+	void ExposeAsName(FName InNewName);
+
+	FReply RemoveExposedName(FName InNameToRemove);
+
+	FText GetSubMenuLabel(FName ExposedName) const;
+
+private:
+
+	/** The sequencer UI instance that is currently open */
+	TWeakPtr<ISequencer> WeakSequencer;
+};
\ No newline at end of file
diff --git a/Engine/Source/Editor/Sequencer/Private/SSequencer.cpp b/Engine/Source/Editor/Sequencer/Private/SSequencer.cpp
--- a/Engine/Source/Editor/Sequencer/Private/SSequencer.cpp
+++ b/Engine/Source/Editor/Sequencer/Private/SSequencer.cpp
@@ -82,6 +82,7 @@
 #include "LevelSequence.h"
 #include "SequencerLog.h"
 #include "MovieSceneCopyableBinding.h"
+#include "SExposedBindingsWidget.h"
 #include "MovieSceneCopyableTrack.h"
 #include "IPropertyRowGenerator.h"
 #include "Fonts/FontMeasure.h"
@@ -1616,14 +1617,23 @@ TSharedRef<SWidget> SSequencer::MakeGeneralMenu()
 
 	MenuBuilder.AddMenuEntry(FSequencerCommands::Get().ToggleShowGotoBox);
 
-	MenuBuilder.AddMenuSeparator();
+	MenuBuilder.BeginSection( "Bindings", LOCTEXT( "BindingsMenuHeader", "Bindings" ) );
 
 	if (SequencerPtr.Pin()->IsLevelEditorSequencer())
 	{
+
 		MenuBuilder.AddMenuEntry(FSequencerCommands::Get().FixActorReferences);
 		MenuBuilder.AddMenuEntry(FSequencerCommands::Get().RebindPossessableReferences);
 	}
 
+	MenuBuilder.AddSubMenu(
+		LOCTEXT("ExposeBindingLabel", "Exposed Binding Groups"),
+		LOCTEXT("ExposeBindingTooltip", "Specifies options for exposing this binding to external systems as a persistent name."),
+		FNewMenuDelegate::CreateSP(this, &SSequencer::PopulateExposeBindingsMenu)
+	);
+
+	MenuBuilder.EndSection();
+
 	if ( SequencerPtr.Pin()->IsLevelEditorSequencer() )
 	{
 		MenuBuilder.AddMenuSeparator();
@@ -1635,6 +1645,11 @@ TSharedRef<SWidget> SSequencer::MakeGeneralMenu()
 	return MenuBuilder.MakeWidget();
 }
 
+void SSequencer::PopulateExposeBindingsMenu(FMenuBuilder& InMenuBuilder)
+{
+	InMenuBuilder.AddWidget(SNew(SExposedBindingsWidget, SequencerPtr), FText(), true);
+}
+
 void SSequencer::FillPlaybackSpeedMenu(FMenuBuilder& InMenuBarBuilder)
 {
 	const int32 NumPlaybackSpeeds = 7;
diff --git a/Engine/Source/Editor/Sequencer/Private/SSequencer.h b/Engine/Source/Editor/Sequencer/Private/SSequencer.h
--- a/Engine/Source/Editor/Sequencer/Private/SSequencer.h
+++ b/Engine/Source/Editor/Sequencer/Private/SSequencer.h
@@ -370,6 +370,8 @@ private:
 	/** Makes the key group menu for the toolbar. */
 	TSharedRef<SWidget> MakeKeyGroupMenu();
 
+	void PopulateExposeBindingsMenu(FMenuBuilder& InMenuBuilder);
+
 	/** Makes the playback speed menu for the toolbar. */
 	void FillPlaybackSpeedMenu(FMenuBuilder& InMenuBuilder);
 
diff --git a/Engine/Source/Runtime/LevelSequence/Private/LevelSequenceActor.cpp b/Engine/Source/Runtime/LevelSequence/Private/LevelSequenceActor.cpp
--- a/Engine/Source/Runtime/LevelSequence/Private/LevelSequenceActor.cpp
+++ b/Engine/Source/Runtime/LevelSequence/Private/LevelSequenceActor.cpp
@@ -7,6 +7,8 @@
 #include "LevelSequenceBurnIn.h"
 #include "DefaultLevelSequenceInstanceData.h"
 #include "Engine/ActorChannel.h"
+#include "Logging/MessageLog.h"
+#include "Misc/UObjectToken.h"
 #include "Net/UnrealNetwork.h"
 #include "LevelSequenceModule.h"
 
@@ -307,7 +309,106 @@ void ALevelSequenceActor::RefreshBurnIn()
 	}
 }
 
+void ALevelSequenceActor::SetBinding(FMovieSceneObjectBindingID Binding, const TArray<AActor*>& Actors, bool bAllowBindingsFromAsset)
+{
+	if (!Binding.IsValid())
+	{
+		FMessageLog("PIE")
+			.Warning(NSLOCTEXT("LevelSequenceActor", "SetBinding_Warning", "The specified binding ID is not valid"))
+			->AddToken(FUObjectToken::Create(this));
+	}
+	else
+	{
+		BindingOverrides->SetBinding(Binding, TArray<UObject*>(Actors), bAllowBindingsFromAsset);
+		if (SequencePlayer)
+		{
+			SequencePlayer->State.Invalidate(Binding.GetGuid(), Binding.GetSequenceID());
+		}
+	}
+}
+
+void ALevelSequenceActor::AddBinding(FMovieSceneObjectBindingID Binding, AActor* Actor, bool bAllowBindingsFromAsset)
+{
+	if (!Binding.IsValid())
+	{
+		FMessageLog("PIE")
+			.Warning(NSLOCTEXT("LevelSequenceActor", "AddBinding_Warning", "The specified binding ID is not valid"))
+			->AddToken(FUObjectToken::Create(this));
+	}
+	else
+	{
+		BindingOverrides->AddBinding(Binding, Actor, bAllowBindingsFromAsset);
+		if (SequencePlayer)
+		{
+			SequencePlayer->State.Invalidate(Binding.GetGuid(), Binding.GetSequenceID());
+		}
+	}
+}
 
+void ALevelSequenceActor::RemoveBinding(FMovieSceneObjectBindingID Binding, AActor* Actor)
+{
+	if (!Binding.IsValid())
+	{
+		FMessageLog("PIE")
+			.Warning(NSLOCTEXT("LevelSequenceActor", "RemoveBinding_Warning", "The specified binding ID is not valid"))
+			->AddToken(FUObjectToken::Create(this));
+	}
+	else
+	{
+		BindingOverrides->RemoveBinding(Binding, Actor);
+		if (SequencePlayer)
+		{
+			SequencePlayer->State.Invalidate(Binding.GetGuid(), Binding.GetSequenceID());
+		}
+	}
+}
+
+void ALevelSequenceActor::ResetBinding(FMovieSceneObjectBindingID Binding)
+{
+	if (!Binding.IsValid())
+	{
+		FMessageLog("PIE")
+			.Warning(NSLOCTEXT("LevelSequenceActor", "ResetBinding_Warning", "The specified binding ID is not valid"))
+			->AddToken(FUObjectToken::Create(this));
+	}
+	else
+	{
+		BindingOverrides->ResetBinding(Binding);
+		if (SequencePlayer)
+		{
+			SequencePlayer->State.Invalidate(Binding.GetGuid(), Binding.GetSequenceID());
+		}
+	}
+}
+
+void ALevelSequenceActor::ResetBindings()
+{
+	BindingOverrides->ResetBindings();
+	if (SequencePlayer)
+	{
+		SequencePlayer->State.ClearObjectCaches(*SequencePlayer);
+	}
+}
+
+FMovieSceneObjectBindingID ALevelSequenceActor::FindNamedBinding(FName InBindingName) const
+{
+	if (ensureAlways(SequencePlayer))
+	{
+		return SequencePlayer->GetSequence()->FindNamedBinding(InBindingName);
+	}
+	return FMovieSceneObjectBindingID();
+}
+
+const TArray<FMovieSceneObjectBindingID>& ALevelSequenceActor::FindNamedBindings(FName InBindingName) const
+{
+	if (ensureAlways(SequencePlayer))
+	{
+		return SequencePlayer->GetSequence()->FindNamedBindings(InBindingName);
+	}
+
+	static TArray<FMovieSceneObjectBindingID> EmptyBindings;
+	return EmptyBindings;
+}
 
 #if WITH_EDITOR
 
diff --git a/Engine/Source/Runtime/LevelSequence/Public/LevelSequenceActor.h b/Engine/Source/Runtime/LevelSequence/Public/LevelSequenceActor.h
--- a/Engine/Source/Runtime/LevelSequence/Public/LevelSequenceActor.h
+++ b/Engine/Source/Runtime/LevelSequence/Public/LevelSequenceActor.h
@@ -170,70 +170,55 @@ public:
 
 public:
 
-	/** Overrides the specified binding with the specified actors, optionally still allowing the bindings defined in the Level Sequence asset
+	/**
+	 * Overrides the specified binding with the specified actors, optionally still allowing the bindings defined in the Level Sequence asset
 	 *
 	 * @param Binding Binding to modify
 	 * @param Actors Actors to bind
 	 * @param bAllowBindingsFromAsset Allow bindings from the level sequence asset
 	 */
 	UFUNCTION(BlueprintCallable, Category="Game|Cinematic|Bindings")
-	void SetBinding(FMovieSceneObjectBindingID Binding, const TArray<AActor*>& Actors, bool bAllowBindingsFromAsset = false)
-	{
-		BindingOverrides->SetBinding(Binding, TArray<UObject*>(Actors), bAllowBindingsFromAsset);
-		if (SequencePlayer)
-		{
-			SequencePlayer->State.Invalidate(Binding.GetGuid(), Binding.GetSequenceID());
-		}
-	}
+	void SetBinding(FMovieSceneObjectBindingID Binding, const TArray<AActor*>& Actors, bool bAllowBindingsFromAsset = false);
 
-	/** Adds the specified actor to the overridden bindings for the specified binding ID, optionally still allowing the bindings defined in the Level Sequence asset
+	/**
+	 * Adds the specified actor to the overridden bindings for the specified binding ID, optionally still allowing the bindings defined in the Level Sequence asset
 	 *
 	 * @param Binding Binding to modify
 	 * @param Actor Actor to bind
 	 * @param bAllowBindingsFromAsset Allow bindings from the level sequence asset
 	 */
 	UFUNCTION(BlueprintCallable, Category="Game|Cinematic|Bindings")
-	void AddBinding(FMovieSceneObjectBindingID Binding, AActor* Actor, bool bAllowBindingsFromAsset = false)
-	{
-		BindingOverrides->AddBinding(Binding, Actor, bAllowBindingsFromAsset);
-		if (SequencePlayer)
-		{
-			SequencePlayer->State.Invalidate(Binding.GetGuid(), Binding.GetSequenceID());
-		}
-	}
+	void AddBinding(FMovieSceneObjectBindingID Binding, AActor* Actor, bool bAllowBindingsFromAsset = false);
 
-	/** Removes the specified actor from the specified binding's actor array */
+	/**
+	 * Removes the specified actor from the specified binding's actor array
+	 */
 	UFUNCTION(BlueprintCallable, Category="Game|Cinematic|Bindings")
-	void RemoveBinding(FMovieSceneObjectBindingID Binding, AActor* Actor)
-	{
-		BindingOverrides->RemoveBinding(Binding, Actor);
-		if (SequencePlayer)
-		{
-			SequencePlayer->State.Invalidate(Binding.GetGuid(), Binding.GetSequenceID());
-		}
-	}
+	void RemoveBinding(FMovieSceneObjectBindingID Binding, AActor* Actor);
 
-	/** Resets the specified binding back to the defaults defined by the Level Sequence asset */
+	/**
+	 * Resets the specified binding back to the defaults defined by the Level Sequence asset
+	 */
 	UFUNCTION(BlueprintCallable, Category="Game|Cinematic|Bindings")
-	void ResetBinding(FMovieSceneObjectBindingID Binding)
-	{
-		BindingOverrides->ResetBinding(Binding);
-		if (SequencePlayer)
-		{
-			SequencePlayer->State.Invalidate(Binding.GetGuid(), Binding.GetSequenceID());
-		}
-	}
+	void ResetBinding(FMovieSceneObjectBindingID Binding);
 
-	/** Resets all overridden bindings back to the defaults defined by the Level Sequence asset */
+	/**
+	 * Resets all overridden bindings back to the defaults defined by the Level Sequence asset
+	 */
 	UFUNCTION(BlueprintCallable, Category="Game|Cinematic|Bindings")
-	void ResetBindings()
-	{
-		BindingOverrides->ResetBindings();
-		if (SequencePlayer)
-		{
-			SequencePlayer->State.ClearObjectCaches(*SequencePlayer);
-		}
-	}
+	void ResetBindings();
+
+	/**
+	 * Find the first object binding ID associated with the specified name (set up through RMB->Expose on Object bindings from within sequencer)
+	 */
+	UFUNCTION(BlueprintCallable, Category = "Game|Cinematic|Bindings")
+	FMovieSceneObjectBindingID FindNamedBinding(FName InBindingName) const;
+
+	/**
+	 * Find all object binding IDs associated with the specified name (set up through RMB->Expose on Object bindings from within sequencer)
+	 */
+	UFUNCTION(BlueprintCallable, Category = "Game|Cinematic|Bindings")
+	const TArray<FMovieSceneObjectBindingID>& FindNamedBindings(FName InBindingName) const;
 
 protected:
 
diff --git a/Engine/Source/Runtime/MovieScene/Private/MovieScene.cpp b/Engine/Source/Runtime/MovieScene/Private/MovieScene.cpp
--- a/Engine/Source/Runtime/MovieScene/Private/MovieScene.cpp
+++ b/Engine/Source/Runtime/MovieScene/Private/MovieScene.cpp
@@ -409,6 +409,28 @@ FText UMovieScene::GetObjectDisplayName(const FGuid& ObjectId)
 }
 
 
+void UMovieScene::ExposeBinding(FName ExposeAs)
+{
+	BindingGroups.Add(ExposeAs);
+}
+
+void UMovieScene::ExposeBinding(FName ExposeAs, FMovieSceneObjectBindingID BindingToExpose)
+{
+	BindingGroups.FindOrAdd(ExposeAs).IDs.AddUnique(BindingToExpose);
+}
+
+void UMovieScene::RemoveExposedBinding(FName ExposedAs, FMovieSceneObjectBindingID BindingToExpose)
+{
+	if (FMovieSceneObjectBindingIDs* Array = BindingGroups.Find(ExposedAs))
+	{
+		Array->IDs.Remove(BindingToExpose);
+	}
+}
+
+void UMovieScene::RemoveExposedBinding(FName ExposedAs)
+{
+	BindingGroups.Remove(ExposedAs);
+}
 
 #if WITH_EDITORONLY_DATA
 void UMovieScene::SetObjectDisplayName(const FGuid& ObjectId, const FText& DisplayName)
diff --git a/Engine/Source/Runtime/MovieScene/Private/MovieSceneSequence.cpp b/Engine/Source/Runtime/MovieScene/Private/MovieSceneSequence.cpp
--- a/Engine/Source/Runtime/MovieScene/Private/MovieSceneSequence.cpp
+++ b/Engine/Source/Runtime/MovieScene/Private/MovieSceneSequence.cpp
@@ -8,6 +8,8 @@
 #include "Tracks/MovieSceneSubTrack.h"
 #include "Sections/MovieSceneSubSection.h"
 #include "Compilation/MovieSceneCompiler.h"
+#include "Logging/MessageLog.h"
+#include "Misc/UObjectToken.h"
 #include "Interfaces/ITargetPlatform.h"
 
 UMovieSceneSequence::UMovieSceneSequence(const FObjectInitializer& Init)
@@ -93,3 +95,29 @@ FGuid UMovieSceneSequence::FindPossessableObjectId(UObject& Object, UObject* Con
 	}
 	return FGuid();
 }
+
+FMovieSceneObjectBindingID UMovieSceneSequence::FindNamedBinding(FName InBindingName) const
+{
+	for (FMovieSceneObjectBindingID ID : FindNamedBindings(InBindingName))
+	{
+		return ID;
+	}
+
+	FMessageLog("PIE")
+		.Warning(NSLOCTEXT("UMovieSceneSequence", "FindNamedBinding_Warning", "Attempted to find a named binding that did not exist"))
+		->AddToken(FUObjectToken::Create(this));
+
+	return FMovieSceneObjectBindingID();
+}
+
+const TArray<FMovieSceneObjectBindingID>& UMovieSceneSequence::FindNamedBindings(FName InBindingName) const
+{
+	const FMovieSceneObjectBindingIDs* BindingIDs = GetMovieScene()->AllBindingGroups().Find(InBindingName);
+	if (BindingIDs)
+	{
+		return BindingIDs->IDs;
+	}
+
+	static TArray<FMovieSceneObjectBindingID> EmptyBindings;
+	return EmptyBindings;
+}
\ No newline at end of file
diff --git a/Engine/Source/Runtime/MovieScene/Public/MovieScene.h b/Engine/Source/Runtime/MovieScene/Public/MovieScene.h
--- a/Engine/Source/Runtime/MovieScene/Public/MovieScene.h
+++ b/Engine/Source/Runtime/MovieScene/Public/MovieScene.h
@@ -167,6 +167,18 @@ struct FMovieSceneTrackLabels
 };
 
 
+/**
+ * Structure that comprises a list of object binding IDs
+ */
+USTRUCT()
+struct FMovieSceneObjectBindingIDs
+{
+	GENERATED_BODY()
+
+	UPROPERTY()
+	TArray<FMovieSceneObjectBindingID> IDs;
+};
+
 /**
  * Implements a movie scene asset.
  */
@@ -801,6 +813,34 @@ public:
 	 */
 	int32 FindNextMarkedFrame(FFrameNumber InFrameNumber, bool bForward);
 
+	/*
+	 * Retrieve all the exposed binding groups for this movie scene
+	 */
+	const TMap<FName, FMovieSceneObjectBindingIDs>& AllBindingGroups() const
+	{
+		return BindingGroups;
+	}
+
+	/*
+	 * Add a new binding group for the specified name
+	 */
+	void ExposeBinding(FName ExposeAs);
+
+	/*
+	 * Expose the specified binding ID using the specified name
+	 */
+	void ExposeBinding(FName ExposeAs, FMovieSceneObjectBindingID BindingToExpose);
+
+	/*
+	 * Remove the specified binding ID from the specified expose group name
+	 */
+	void RemoveExposedBinding(FName ExposedAs, FMovieSceneObjectBindingID BindingToExpose);
+
+	/*
+	 * Remove the entire group with the specified name
+	 */
+	void RemoveExposedBinding(FName ExposedAs);
+
 protected:
 
 	/**
@@ -845,6 +885,10 @@ private:
 	UPROPERTY()
 	TArray<FMovieSceneBinding> ObjectBindings;
 
+	/** Map of persistent named bindings for this sequence */
+	UPROPERTY()
+	TMap<FName, FMovieSceneObjectBindingIDs> BindingGroups;
+
 	/** Master tracks which are not bound to spawned or possessed objects */
 	UPROPERTY(Instanced)
 	TArray<UMovieSceneTrack*> MasterTracks;
diff --git a/Engine/Source/Runtime/MovieScene/Public/MovieSceneSequence.h b/Engine/Source/Runtime/MovieScene/Public/MovieSceneSequence.h
--- a/Engine/Source/Runtime/MovieScene/Public/MovieSceneSequence.h
+++ b/Engine/Source/Runtime/MovieScene/Public/MovieSceneSequence.h
@@ -179,6 +179,18 @@ public:
 	 */
 	virtual UObject* CreateDirectorInstance(IMovieScenePlayer& Player) { return nullptr; }
 
+	/**
+	 * Find the first object binding ID associated with the specified name (set up through RMB->Expose on Object bindings from within sequencer)
+	 */
+	UFUNCTION(BlueprintCallable, Category = "Game|Cinematic|Bindings")
+	MOVIESCENE_API FMovieSceneObjectBindingID FindNamedBinding(FName InBindingName) const;
+
+	/**
+	 * Find all object binding IDs associated with the specified name (set up through RMB->Expose on Object bindings from within sequencer)
+	 */
+	UFUNCTION(BlueprintCallable, Category = "Game|Cinematic|Bindings")
+	MOVIESCENE_API const TArray<FMovieSceneObjectBindingID>& FindNamedBindings(FName InBindingName) const;
+
 public:
 
 	MOVIESCENE_API virtual void PostLoad() override;
