Adding pragma once to remove class redefinitions

diff --git a/Engine/Source/Editor/Sequencer/Private/SExposeBindingWidget.h b/Engine/Source/Editor/Sequencer/Private/SExposeBindingWidget.h
--- a/Engine/Source/Editor/Sequencer/Private/SExposeBindingWidget.h
+++ b/Engine/Source/Editor/Sequencer/Private/SExposeBindingWidget.h
@@ -1,5 +1,7 @@
 // Copyright 1998-2019 Epic Games, Inc. All Rights Reserved.
 
+#pragma once
+
 #include "Widgets/DeclarativeSyntaxSupport.h"
 #include "Widgets/SCompoundWidget.h"
 
diff --git a/Engine/Source/Editor/Sequencer/Private/SExposedBindingsWidget.h b/Engine/Source/Editor/Sequencer/Private/SExposedBindingsWidget.h
--- a/Engine/Source/Editor/Sequencer/Private/SExposedBindingsWidget.h
+++ b/Engine/Source/Editor/Sequencer/Private/SExposedBindingsWidget.h
@@ -1,5 +1,7 @@
 // Copyright 1998-2019 Epic Games, Inc. All Rights Reserved.
 
+#pragma once
+
 #include "Widgets/DeclarativeSyntaxSupport.h"
 #include "Widgets/SCompoundWidget.h"
 
