Sequencer: Reset playback position on play if current play position is out of bounds.

diff --git a/Engine/Source/Editor/Sequencer/Private/Sequencer.cpp b/Engine/Source/Editor/Sequencer/Private/Sequencer.cpp
--- a/Engine/Source/Editor/Sequencer/Private/Sequencer.cpp
+++ b/Engine/Source/Editor/Sequencer/Private/Sequencer.cpp
@@ -3608,6 +3608,17 @@ FReply FSequencer::OnPlay(bool bTogglePlay)
 	}
 	else
 	{
+		TRange<FFrameNumber> TimeBounds = GetTimeBounds();
+
+		FFrameNumber MinInclusiveTime = MovieScene::DiscreteInclusiveLower(TimeBounds);
+		FFrameNumber MaxInclusiveTime = MovieScene::DiscreteExclusiveUpper(TimeBounds) - 1;
+
+		if (GetLocalTime().Time <= MinInclusiveTime || GetLocalTime().Time >= MaxInclusiveTime)
+		{
+			FFrameTime NewLocalTime = (PlaybackSpeed > 0 ? MinInclusiveTime : MaxInclusiveTime) * RootToLocalTransform.Inverse();
+			SetLocalTimeDirectly(NewLocalTime);
+		}
+
 		SetPlaybackStatus(EMovieScenePlayerStatus::Playing);
 
 		// Make sure Slate ticks during playback
