Sequencer: Notify the editor of actor selection changes before we start selecting their components.

This is because it will fire the necessary events to (among other things) repopulate the SCSEditor widget which shows the component hierarchy of the selected actor. Then, when we select a child component, the details panel will be able to find that child component in the widget's tree and properly show its details.

diff --git a/Engine/Source/Editor/Sequencer/Private/Sequencer.cpp b/Engine/Source/Editor/Sequencer/Private/Sequencer.cpp
--- a/Engine/Source/Editor/Sequencer/Private/Sequencer.cpp
+++ b/Engine/Source/Editor/Sequencer/Private/Sequencer.cpp
@@ -5296,6 +5296,8 @@ void FSequencer::SynchronizeExternalSelectionWithSequencerSelection()
 
 	GEditor->GetSelectedActors()->EndBatchSelectOperation();
 
+	GEditor->NoteSelectionChange();
+
 	if (SelectedSequencerComponents.Num())
 	{
 		GEditor->GetSelectedComponents()->Modify();
@@ -5310,9 +5312,9 @@ void FSequencer::SynchronizeExternalSelectionWithSequencerSelection()
 		}
 
 		GEditor->GetSelectedComponents()->EndBatchSelectOperation();
+
+		GEditor->NoteSelectionChange();
 	}
-		
-	GEditor->NoteSelectionChange();
 }
 
 
