Sequencer: Fix `GetObjectGuid`, which would always return null.

diff --git a/Engine/Source/Editor/Sequencer/Private/DisplayNodes/SequencerDisplayNode.cpp b/Engine/Source/Editor/Sequencer/Private/DisplayNodes/SequencerDisplayNode.cpp
--- a/Engine/Source/Editor/Sequencer/Private/DisplayNodes/SequencerDisplayNode.cpp
+++ b/Engine/Source/Editor/Sequencer/Private/DisplayNodes/SequencerDisplayNode.cpp
@@ -549,7 +549,7 @@ TSharedPtr<FSequencerObjectBindingNode> FSequencerDisplayNode::FindParentObjectB
 FGuid FSequencerDisplayNode::GetObjectGuid() const
 {
 	TSharedPtr<FSequencerObjectBindingNode> ObjectBindingNode = FindParentObjectBindingNode();
-	return ObjectBindingNode ? ObjectBindingNode->GetObjectGuid() : FGuid();
+	return ObjectBindingNode ? ObjectBindingNode->GetObjectBinding() : FGuid();
 }
 
 bool FSequencerDisplayNode::Traverse_ChildFirst(const TFunctionRef<bool(FSequencerDisplayNode&)>& InPredicate, bool bIncludeThisNode)
