Sequencer: Fix crash trying to convert an invalid/unbound objects when converting to spawnable.

diff --git a/Engine/Source/Editor/Sequencer/Private/Sequencer.cpp b/Engine/Source/Editor/Sequencer/Private/Sequencer.cpp
--- a/Engine/Source/Editor/Sequencer/Private/Sequencer.cpp
+++ b/Engine/Source/Editor/Sequencer/Private/Sequencer.cpp
@@ -7669,7 +7669,10 @@ TArray<FGuid> FSequencer::ExpandMultiplePossessableBindings(FGuid PossessableGui
 	for (TWeakObjectPtr<> FoundObjectPtr : FoundObjects)
 	{
 		UObject* FoundObject = FoundObjectPtr.Get();
-		check(FoundObject);
+		if (!FoundObject)
+		{
+			continue;
+		}
 
 		FoundObject->Modify();
 
