Sequencer: Set value snap to 0.1f

diff --git a/Engine/Source/Editor/Sequencer/Private/SequencerSettings.cpp b/Engine/Source/Editor/Sequencer/Private/SequencerSettings.cpp
--- a/Engine/Source/Editor/Sequencer/Private/SequencerSettings.cpp
+++ b/Engine/Source/Editor/Sequencer/Private/SequencerSettings.cpp
@@ -25,7 +25,7 @@ USequencerSettings::USequencerSettings( const FObjectInitializer& ObjectInitiali
 	bSnapPlayTimeToInterval = true;
 	bSnapPlayTimeToPressedKey = true;
 	bSnapPlayTimeToDraggedKey = true;
-	CurveValueSnapInterval = 10.0f;
+	CurveValueSnapInterval = 0.1f;
 	bSnapCurveValueToInterval = true;
 	bLabelBrowserVisible = false;
 	bShowSelectedNodesOnly = false;
