Fixed Actor Sequences not showing a Curve Editor because they had no Toolkit host. Fixed a crash using the keyboard shortcut to open the Curve Editor when the Curve Editor was not supported (ie: Niagara).

diff --git a/Engine/Plugins/MovieScene/ActorSequence/Source/ActorSequenceEditor/Private/ActorSequenceEditorTabSummoner.cpp b/Engine/Plugins/MovieScene/ActorSequence/Source/ActorSequenceEditor/Private/ActorSequenceEditorTabSummoner.cpp
--- a/Engine/Plugins/MovieScene/ActorSequence/Source/ActorSequenceEditor/Private/ActorSequenceEditorTabSummoner.cpp
+++ b/Engine/Plugins/MovieScene/ActorSequence/Source/ActorSequenceEditor/Private/ActorSequenceEditorTabSummoner.cpp
@@ -203,6 +203,17 @@ public:
 		GEditor->OnBlueprintPreCompile().Remove(OnBlueprintPreCompileHandle);
 		FCoreUObjectDelegates::OnObjectSaved.Remove(OnObjectSavedHandle);
 	}
+	
+	TSharedRef<SDockTab> SpawnCurveEditorTab(const FSpawnTabArgs&)
+	{
+		const FSlateIcon SequencerGraphIcon = FSlateIcon(FEditorStyle::GetStyleSetName(), "GenericCurveEditor.TabIcon");
+		return SNew(SDockTab)
+			.Icon(SequencerGraphIcon.GetIcon())
+			.Label(NSLOCTEXT("Sequencer", "SequencerMainGraphEditorTitle", "Sequencer Curves"))
+			[
+				SNullWidget::NullWidget
+			];
+	}
 
 	void Construct(const FArguments&, TWeakPtr<FBlueprintEditor> InBlueprintEditor)
 	{
@@ -211,6 +222,13 @@ public:
 
 		WeakBlueprintEditor = InBlueprintEditor;
 
+		{
+			// Register an empty tab to spawn the Curve Editor in so that layouts restore properly.
+			WeakBlueprintEditor.Pin()->GetTabManager()->RegisterTabSpawner("SequencerGraphEditor", 
+				FOnSpawnTab::CreateSP(this, &SActorSequenceEditorWidgetImpl::SpawnCurveEditorTab))
+				.SetMenuType(ETabSpawnerMenuType::Type::Hidden);
+		}
+
 		ChildSlot
 		[
 			SAssignNew(Content, SBox)
@@ -309,6 +327,12 @@ public:
 			SequencerInitParams.RootSequence = NewSequence;
 			SequencerInitParams.EventContexts = TAttribute<TArray<UObject*>>(this, &SActorSequenceEditorWidgetImpl::GetEventContexts);
 			SequencerInitParams.PlaybackContext = TAttribute<UObject*>(this, &SActorSequenceEditorWidgetImpl::GetPlaybackContext);
+			
+			if (WeakBlueprintEditor.IsValid())
+			{
+				SequencerInitParams.ToolkitHost = WeakBlueprintEditor.Pin()->GetToolkitHost();
+				SequencerInitParams.HostCapabilities.bSupportsCurveEditor = true;
+			}
 
 			TSharedRef<FExtender> AddMenuExtender = MakeShareable(new FExtender);
 
diff --git a/Engine/Plugins/MovieScene/LevelSequenceEditor/Source/LevelSequenceEditor/Private/LevelSequenceEditorToolkit.cpp b/Engine/Plugins/MovieScene/LevelSequenceEditor/Source/LevelSequenceEditor/Private/LevelSequenceEditorToolkit.cpp
--- a/Engine/Plugins/MovieScene/LevelSequenceEditor/Source/LevelSequenceEditor/Private/LevelSequenceEditorToolkit.cpp
+++ b/Engine/Plugins/MovieScene/LevelSequenceEditor/Source/LevelSequenceEditor/Private/LevelSequenceEditorToolkit.cpp
@@ -184,6 +184,9 @@ void FLevelSequenceEditorToolkit::Initialize(const EToolkitMode::Type Mode, cons
 		SequencerInitParams.ViewParams.ScrubberStyle = ESequencerScrubberStyle::FrameBlock;
 		SequencerInitParams.ViewParams.OnReceivedFocus.BindRaw(this, &FLevelSequenceEditorToolkit::OnSequencerReceivedFocus);
 
+		SequencerInitParams.HostCapabilities.bSupportsCurveEditor = true;
+		SequencerInitParams.HostCapabilities.bSupportsSaveMovieSceneAsset = true;
+
 		TSharedRef<FExtender> ToolbarExtender = MakeShared<FExtender>();
 		ToolbarExtender->AddToolBarExtension("Base Commands", EExtensionHook::Before, nullptr, FToolBarExtensionDelegate::CreateSP(this, &FLevelSequenceEditorToolkit::ExtendSequencerToolbar));
 		SequencerInitParams.ViewParams.ToolbarExtender = ToolbarExtender;
diff --git a/Engine/Source/Editor/Sequencer/Private/SSequencer.cpp b/Engine/Source/Editor/Sequencer/Private/SSequencer.cpp
--- a/Engine/Source/Editor/Sequencer/Private/SSequencer.cpp
+++ b/Engine/Source/Editor/Sequencer/Private/SSequencer.cpp
@@ -408,16 +408,18 @@ void SSequencer::Construct(const FArguments& InArgs, TSharedRef<FSequencer> InSe
 	TrackArea->SetTreeView(TreeView);
 
 	TAttribute<FAnimatedRange> ViewRangeAttribute = InArgs._ViewRange;
-	
-	// We create a custom Time Slider Controller which is just a wrapper around the actual one, but is aware of our custom bounds logic. Currently the range the
-	// bar displays is tied to Sequencer timeline and not the Bounds, so we need a way of changing it to look at the Bounds but only for the Curve Editor time
-	// slider controller. We want everything else to just pass through though.
-	TSharedRef<ITimeSliderController> CurveEditorTimeSliderController = MakeShared<FSequencerCurveEditorTimeSliderController>(TimeSliderArgs, SequencerPtr, InSequencer->GetCurveEditor().ToSharedRef());
 
-	// Initialize the Curve Editor Widget if there is a tab manager to spawn our extra tab in.
-	// Some areas that use Sequencer don't use our curve editor. In this case no button is shown on the UI.
-	if (InSequencer->GetToolkitHost().IsValid())
+	if (InSequencer->GetHostCapabilities().bSupportsCurveEditor)
 	{
+		// If they've said they want to support the curve editor then they need to provide a toolkit host
+		// so that we know where to spawn our tab into.
+		check(InSequencer->GetToolkitHost().IsValid());
+
+		// We create a custom Time Slider Controller which is just a wrapper around the actual one, but is aware of our custom bounds logic. Currently the range the
+		// bar displays is tied to Sequencer timeline and not the Bounds, so we need a way of changing it to look at the Bounds but only for the Curve Editor time
+		// slider controller. We want everything else to just pass through though.
+		TSharedRef<ITimeSliderController> CurveEditorTimeSliderController = MakeShared<FSequencerCurveEditorTimeSliderController>(TimeSliderArgs, SequencerPtr, InSequencer->GetCurveEditor().ToSharedRef());
+
 		CurveEditorTree = SNew(SCurveEditorTree, InSequencer->GetCurveEditor());
 		TSharedRef<SCurveEditorPanel> CurveEditorWidget = SNew(SCurveEditorPanel, InSequencer->GetCurveEditor().ToSharedRef())
 			// Grid lines match the color specified in FSequencerTimeSliderController::OnPaintViewArea
@@ -1179,23 +1181,30 @@ TSharedRef<SWidget> SSequencer::MakeToolBar()
 				}
 			}));
 
-			ToolBarBuilder.AddToolBarButton(
-				FUIAction(FExecuteAction::CreateSP(this, &SSequencer::OnSaveMovieSceneClicked)),
-				NAME_None,
-				LOCTEXT("SaveDirtyPackages", "Save"),
-				LOCTEXT("SaveDirtyPackagesTooltip", "Saves the current sequence and any subsequences"),
-				SaveIcon
-			);
+			if (SequencerPtr.Pin()->GetHostCapabilities().bSupportsSaveMovieSceneAsset)
+			{
+				ToolBarBuilder.AddToolBarButton(
+					FUIAction(FExecuteAction::CreateSP(this, &SSequencer::OnSaveMovieSceneClicked)),
+					NAME_None,
+					LOCTEXT("SaveDirtyPackages", "Save"),
+					LOCTEXT("SaveDirtyPackagesTooltip", "Saves the current sequence and any subsequences"),
+					SaveIcon
+				);
 
-			ToolBarBuilder.AddToolBarButton(
-				FUIAction(FExecuteAction::CreateSP(this, &SSequencer::OnSaveMovieSceneAsClicked)),
-				NAME_None,
-				LOCTEXT("SaveAs", "Save As"),
-				LOCTEXT("SaveAsTooltip", "Saves the current sequence under a different name"),
-				FSlateIcon(FEditorStyle::GetStyleSetName(), "Sequencer.SaveAs")
-			);
+				ToolBarBuilder.AddToolBarButton(
+					FUIAction(FExecuteAction::CreateSP(this, &SSequencer::OnSaveMovieSceneAsClicked)),
+					NAME_None,
+					LOCTEXT("SaveAs", "Save As"),
+					LOCTEXT("SaveAsTooltip", "Saves the current sequence under a different name"),
+					FSlateIcon(FEditorStyle::GetStyleSetName(), "Sequencer.SaveAs")
+				);
+			}
+
+			if (SequencerPtr.Pin()->GetHostCapabilities().bSupportsDiscardChanges)
+			{
+				ToolBarBuilder.AddToolBarButton( FSequencerCommands::Get().DiscardChanges );
+			}
 
-			//ToolBarBuilder.AddToolBarButton( FSequencerCommands::Get().DiscardChanges );
 			ToolBarBuilder.AddToolBarButton( FSequencerCommands::Get().FindInContentBrowser );
 			ToolBarBuilder.AddToolBarButton( FSequencerCommands::Get().CreateCamera );
 			ToolBarBuilder.AddToolBarButton( FSequencerCommands::Get().RenderMovie );
@@ -1366,8 +1375,8 @@ TSharedRef<SWidget> SSequencer::MakeToolBar()
 
 	ToolBarBuilder.BeginSection("Curve Editor");
 	{
-		// Only add the button if we have a toolkit host to spawn tabs in
-		if (SequencerPtr.Pin()->GetToolkitHost().IsValid())
+		// Only add the button if supported
+		if (SequencerPtr.Pin()->GetHostCapabilities().bSupportsCurveEditor)
 		{
 			ToolBarBuilder.AddToolBarButton( FSequencerCommands::Get().ToggleShowCurveEditor );
 		}
@@ -2133,11 +2142,14 @@ SSequencer::~SSequencer()
 	TSharedPtr<FSequencer> Sequencer = SequencerPtr.Pin();
 	if(Sequencer)
 	{
-		FTabId TabId = FTabId(SSequencer::CurveEditorTabName);
-		TSharedPtr<SDockTab> CurveEditorTab = Sequencer->GetToolkitHost()->GetTabManager()->FindExistingLiveTab(TabId);
-		if (CurveEditorTab)
+		if (Sequencer->GetHostCapabilities().bSupportsCurveEditor)
 		{
-			CurveEditorTab->RequestCloseTab();
+			FTabId TabId = FTabId(SSequencer::CurveEditorTabName);
+			TSharedPtr<SDockTab> CurveEditorTab = Sequencer->GetToolkitHost()->GetTabManager()->FindExistingLiveTab(TabId);
+			if (CurveEditorTab)
+			{
+				CurveEditorTab->RequestCloseTab();
+			}
 		}
 	}
 }
@@ -2795,6 +2807,12 @@ void SSequencer::OnCurveEditorVisibilityChanged(bool bShouldBeVisible)
 	TSharedPtr<FSequencer> Sequencer = SequencerPtr.Pin();
 	FTabId TabId = FTabId(SSequencer::CurveEditorTabName);
 
+	// Curve Editor may not be supported
+	if (!Sequencer->GetHostCapabilities().bSupportsCurveEditor)
+	{
+		return;
+	}
+
 	if (bShouldBeVisible)
 	{
 		// Request the Tab Manager invoke the tab. This will spawn the tab if needed, otherwise pull it to focus. This assumes
diff --git a/Engine/Source/Editor/Sequencer/Private/Sequencer.cpp b/Engine/Source/Editor/Sequencer/Private/Sequencer.cpp
--- a/Engine/Source/Editor/Sequencer/Private/Sequencer.cpp
+++ b/Engine/Source/Editor/Sequencer/Private/Sequencer.cpp
@@ -281,6 +281,7 @@ void FSequencer::InitSequencer(const FSequencerInitParams& InitParams, const TSh
 {
 	bIsEditingWithinLevelEditor = InitParams.bEditWithinLevelEditor;
 	ScrubStyle = InitParams.ViewParams.ScrubberStyle;
+	HostCapabilities = InitParams.HostCapabilities;
 
 	SilentModeCount = 0;
 	bReadOnly = InitParams.ViewParams.bReadOnly;
@@ -4955,6 +4956,12 @@ void FSequencer::SetShowCurveEditor(bool bInShowCurveEditor)
 
 bool FSequencer::GetCurveEditorIsVisible() const
 {
+	// Some Sequencer usages don't support the Curve Editor
+	if (!GetHostCapabilities().bSupportsCurveEditor)
+	{
+		return false;
+	}
+
 	// We always want to retrieve this directly from the UI instead of mirroring it to a local bool as there are
 	// a lot of ways the UI could get out of sync with a local bool (such as previously restored tab layouts)
 	return GetToolkitHost()->GetTabManager()->FindExistingLiveTab(FTabId(SSequencer::CurveEditorTabName)).IsValid();
@@ -5032,13 +5039,14 @@ void FSequencer::SaveCurrentMovieScene()
 
 void FSequencer::SaveCurrentMovieSceneAs()
 {
-	TSharedPtr<IToolkitHost> MyToolkitHost = GetToolkitHost();
-
-	if (!MyToolkitHost.IsValid())
+	if (!GetHostCapabilities().bSupportsSaveMovieSceneAsset)
 	{
 		return;
 	}
 
+	TSharedPtr<IToolkitHost> MyToolkitHost = GetToolkitHost();
+	check(MyToolkitHost);
+
 	TArray<UObject*> AssetsToSave;
 	AssetsToSave.Add(GetCurrentAsset());
 
diff --git a/Engine/Source/Editor/Sequencer/Private/Sequencer.h b/Engine/Source/Editor/Sequencer/Private/Sequencer.h
--- a/Engine/Source/Editor/Sequencer/Private/Sequencer.h
+++ b/Engine/Source/Editor/Sequencer/Private/Sequencer.h
@@ -462,6 +462,8 @@ public:
 	/** @return The toolkit that this sequencer is hosted in (if any) */
 	TSharedPtr<IToolkitHost> GetToolkitHost() const { return ToolkitHost.Pin(); }
 
+	const FSequencerHostCapabilities& GetHostCapabilities() const { return HostCapabilities; }
+
 	/** @return Whether or not this sequencer is used in the level editor */
 	bool IsLevelEditorSequencer() const { return bIsEditingWithinLevelEditor; }
 
@@ -1076,6 +1078,9 @@ private:
 	/** The asset editor that created this Sequencer if any */
 	TWeakPtr<IToolkitHost> ToolkitHost;
 
+	/** A copy of the supported features/capabilities we were initialized with. */
+	FSequencerHostCapabilities HostCapabilities;
+
 	TWeakObjectPtr<UMovieSceneSequence> RootSequence;
 	FMovieSceneRootEvaluationTemplateInstance RootTemplateInstance;
 
diff --git a/Engine/Source/Editor/Sequencer/Public/ISequencerModule.h b/Engine/Source/Editor/Sequencer/Public/ISequencerModule.h
--- a/Engine/Source/Editor/Sequencer/Public/ISequencerModule.h
+++ b/Engine/Source/Editor/Sequencer/Public/ISequencerModule.h
@@ -93,6 +93,28 @@ struct FSequencerViewParams
 	{ }
 };
 
+/**
+ * Sequencer host functionality capabilities. These are no longer
+ * based on whether or not there is a Toolkit host as we may have
+ * a toolkit host outside of conditions where these are supported.
+ */
+struct FSequencerHostCapabilities
+{
+	/** Should we show the Save-As button in the toolbar? */
+	bool bSupportsSaveMovieSceneAsset;
+
+	/** Do we support discarding the changes */
+	bool bSupportsDiscardChanges;
+
+	/** Do we support the curve editor */
+	bool bSupportsCurveEditor;
+
+	FSequencerHostCapabilities()
+		: bSupportsSaveMovieSceneAsset(false)
+		, bSupportsDiscardChanges(false)
+		, bSupportsCurveEditor(false)
+	{}
+};
 
 /**
  * Sequencer initialization parameters.
@@ -108,6 +130,9 @@ struct FSequencerInitParams
 	/** View parameters */
 	FSequencerViewParams ViewParams;
 
+	/** Immutable capability set specified when our instance is created. Used to specify which feature set is supported. */
+	FSequencerHostCapabilities HostCapabilities;
+
 	/** Whether or not sequencer should be edited within the level editor */
 	bool bEditWithinLevelEditor;
 
diff --git a/Engine/Source/Editor/UMGEditor/Private/WidgetBlueprintEditor.cpp b/Engine/Source/Editor/UMGEditor/Private/WidgetBlueprintEditor.cpp
--- a/Engine/Source/Editor/UMGEditor/Private/WidgetBlueprintEditor.cpp
+++ b/Engine/Source/Editor/UMGEditor/Private/WidgetBlueprintEditor.cpp
@@ -832,6 +832,8 @@ TSharedPtr<ISequencer>& FWidgetBlueprintEditor::GetSequencer()
 			SequencerInitParams.ToolkitHost = GetToolkitHost();
 			SequencerInitParams.PlaybackContext = TAttribute<UObject*>(this, &FWidgetBlueprintEditor::GetAnimationPlaybackContext);
 			SequencerInitParams.EventContexts = TAttribute<TArray<UObject*>>(this, &FWidgetBlueprintEditor::GetAnimationEventContexts);
+
+			SequencerInitParams.HostCapabilities.bSupportsCurveEditor = true;
 		};
 
 		Sequencer = FModuleManager::LoadModuleChecked<ISequencerModule>("Sequencer").CreateSequencer(SequencerInitParams);
