Sequencer: Fix setting global time when out of bounds so that it's calling SetGlobalTime instead of SetLocalTime.

diff --git a/Engine/Source/Editor/Sequencer/Private/Sequencer.cpp b/Engine/Source/Editor/Sequencer/Private/Sequencer.cpp
--- a/Engine/Source/Editor/Sequencer/Private/Sequencer.cpp
+++ b/Engine/Source/Editor/Sequencer/Private/Sequencer.cpp
@@ -3616,8 +3616,8 @@ FReply FSequencer::OnPlay(bool bTogglePlay)
 
 		if (GetLocalTime().Time <= MinInclusiveTime || GetLocalTime().Time >= MaxInclusiveTime)
 		{
-			FFrameTime NewLocalTime = (PlaybackSpeed > 0 ? MinInclusiveTime : MaxInclusiveTime) * RootToLocalTransform.Inverse();
-			SetLocalTimeDirectly(NewLocalTime);
+			FFrameTime NewGlobalTime = (PlaybackSpeed > 0 ? MinInclusiveTime : MaxInclusiveTime) * RootToLocalTransform.Inverse();
+			SetGlobalTime(NewGlobalTime);
 		}
 
 		SetPlaybackStatus(EMovieScenePlayerStatus::Playing);
