Sequencer: Fix for selecting child key areas if any parent is collapsed

diff --git a/Engine/Source/Editor/Sequencer/Private/Tools/SequencerEntityVisitor.cpp b/Engine/Source/Editor/Sequencer/Private/Tools/SequencerEntityVisitor.cpp
--- a/Engine/Source/Editor/Sequencer/Private/Tools/SequencerEntityVisitor.cpp
+++ b/Engine/Source/Editor/Sequencer/Private/Tools/SequencerEntityVisitor.cpp
@@ -81,7 +81,7 @@ void FSequencerEntityWalker::ConditionallyIntersectNode(const ISequencerEntityVi
 
 		if (Range.IntersectKeyArea(InNode, VirtualKeySize.Y))
 		{
-			VisitKeyAnyAreas(Visitor, InNode);
+			VisitKeyAnyAreas(Visitor, InNode, !InNode->IsExpanded());
 		}
 	}
 
@@ -99,7 +99,7 @@ void FSequencerEntityWalker::ConditionallyIntersectNode(const ISequencerEntityVi
 	}
 }
 
-void FSequencerEntityWalker::VisitKeyAnyAreas(const ISequencerEntityVisitor& Visitor, const TSharedRef<FSequencerDisplayNode>& InNode)
+void FSequencerEntityWalker::VisitKeyAnyAreas(const ISequencerEntityVisitor& Visitor, const TSharedRef<FSequencerDisplayNode>& InNode, bool bAnyParentCollapsed )
 {
 	if (!Visitor.CheckEntityMask(ESequencerEntity::Key))
 	{
@@ -129,11 +129,11 @@ void FSequencerEntityWalker::VisitKeyAnyAreas(const ISequencerEntityVisitor& Vis
 		}
 	}
 	// Otherwise it might be a collapsed node that contains key areas as children. If so we visit them as if they were a part of this track so that key groupings are visited properly.
-	else if (!InNode->IsExpanded())
+	else if (bAnyParentCollapsed)
 	{
 		for (TSharedRef<FSequencerDisplayNode> ChildNode : InNode->GetChildNodes())
 		{
-			VisitKeyAnyAreas(Visitor, ChildNode);
+			VisitKeyAnyAreas(Visitor, ChildNode, bAnyParentCollapsed || !InNode->IsExpanded());
 		}
 	}
 }
diff --git a/Engine/Source/Editor/Sequencer/Private/Tools/SequencerEntityVisitor.h b/Engine/Source/Editor/Sequencer/Private/Tools/SequencerEntityVisitor.h
--- a/Engine/Source/Editor/Sequencer/Private/Tools/SequencerEntityVisitor.h
+++ b/Engine/Source/Editor/Sequencer/Private/Tools/SequencerEntityVisitor.h
@@ -82,7 +82,7 @@ private:
 	/** Visit any keys within the specified track node that overlap the range's horizontal space */
 	void VisitTrackNode(const ISequencerEntityVisitor& Visitor, const TSharedRef<FSequencerTrackNode>& InNode);
 	/** Visit any keys within any key area nodes that belong to the specified node that overlap the range's horizontal space */
-	void VisitKeyAnyAreas(const ISequencerEntityVisitor& Visitor, const TSharedRef<FSequencerDisplayNode>& InNode);
+	void VisitKeyAnyAreas(const ISequencerEntityVisitor& Visitor, const TSharedRef<FSequencerDisplayNode>& InNode, bool bAnyParentCollapsed);
 	/** Visit any keys within the specified key area that overlap the range's horizontal space */
 	void VisitKeyArea(const ISequencerEntityVisitor& Visitor, const TSharedRef<IKeyArea>& KeyArea, UMovieSceneSection* Section, const TSharedRef<FSequencerDisplayNode>& InNode);
 
