Sequencer: Remove unnecessary force evaluate in stopped state which was introduced in 7047268.

diff --git a/Engine/Source/Editor/Sequencer/Private/Sequencer.cpp b/Engine/Source/Editor/Sequencer/Private/Sequencer.cpp
--- a/Engine/Source/Editor/Sequencer/Private/Sequencer.cpp
+++ b/Engine/Source/Editor/Sequencer/Private/Sequencer.cpp
@@ -4320,9 +4320,6 @@ void FSequencer::OnEndScrubbing()
 	AutoscrubOffset.Reset();
 	StopAutoscroll();
 
-	// Force an evaluation in the stopped state
-	ForceEvaluate();
-
 	OnEndScrubbingDelegate.Broadcast();
 }
 
