Sequencer: Fix crash in node tree filter with invalid object bindings

diff --git a/Engine/Source/Editor/Sequencer/Private/SequencerNodeTree.cpp b/Engine/Source/Editor/Sequencer/Private/SequencerNodeTree.cpp
--- a/Engine/Source/Editor/Sequencer/Private/SequencerNodeTree.cpp
+++ b/Engine/Source/Editor/Sequencer/Private/SequencerNodeTree.cpp
@@ -941,7 +941,7 @@ static void FilterNodesRecursive( FSequencer& Sequencer, const TSharedRef<FSeque
 						const FSequencerObjectBindingNode* ObjectNode = static_cast<const FSequencerObjectBindingNode*>(ParentNode.Get());
 						for (TWeakObjectPtr<>& Object : Sequencer.FindObjectsInCurrentSequence(ObjectNode->GetObjectBinding()))
 						{
-							if (LevelTrackFilter->PassesFilter(Object.Get()))
+							if (Object.IsValid() && LevelTrackFilter->PassesFilter(Object.Get()))
 							{
 								// If at least one of the objects on the objectbinding node pass the level filter, show the track
 								bPasssedAnyFilters = true;
@@ -987,7 +987,7 @@ static void FilterNodesRecursive( FSequencer& Sequencer, const TSharedRef<FSeque
 		const FSequencerObjectBindingNode ObjectNode = static_cast<const FSequencerObjectBindingNode&>(StartNode.Get());
 		for (TWeakObjectPtr<>& Object : Sequencer.FindObjectsInCurrentSequence(ObjectNode.GetObjectBinding()))
 		{
-			if ((Filters->Num() == 0 || Filters->PassesAnyFilters(Object.Get()))
+			if (Object.IsValid() && (Filters->Num() == 0 || Filters->PassesAnyFilters(Object.Get()))
 				&& LevelTrackFilter->PassesFilter(Object.Get()))
 			{
 				bPasssedAnyFilters = true;
