Sequencer: Redo fix for CL #6619179 so that AddGivenTrack allows for tracks with the same name and the check is only done on paste.

diff --git a/Engine/Source/Editor/Sequencer/Private/Sequencer.cpp b/Engine/Source/Editor/Sequencer/Private/Sequencer.cpp
--- a/Engine/Source/Editor/Sequencer/Private/Sequencer.cpp
+++ b/Engine/Source/Editor/Sequencer/Private/Sequencer.cpp
@@ -7244,6 +7244,7 @@ bool FSequencer::PasteTracks(const FString& TextToImport, TArray<FNotificationIn
 	CalculateSelectedFolderAndPath(SelectedParentFolders, NewNodePath);
 
 	UMovieSceneSequence* OwnerSequence = GetFocusedMovieSceneSequence();
+	UMovieScene* OwnerMovieScene = OwnerSequence->GetMovieScene();
 	UObject* BindingContext = GetPlaybackContext();
 
 	TSet<TSharedRef<FSequencerDisplayNode>> SelectedNodes = Selection.GetSelectedOutlinerNodes();
@@ -7302,6 +7303,24 @@ bool FSequencer::PasteTracks(const FString& TextToImport, TArray<FNotificationIn
 						Subobject->ClearFlags(RF_Transient);
 					}
 
+					// Remove tracks with the same name before adding
+					for (const FMovieSceneBinding& Binding : OwnerMovieScene->GetBindings())
+					{
+						if (Binding.GetObjectGuid() == ObjectGuid)
+						{
+							// Tracks of the same class should be unique per name.
+							for (UMovieSceneTrack* Track : Binding.GetTracks())
+							{
+								if (Track->GetClass() == NewTrack->GetClass() && Track->GetTrackName() == NewTrack->GetTrackName())
+								{
+									// If a track of the same class and name exists, remove it so the new track replaces it
+									OwnerMovieScene->RemoveTrack(*Track);
+									break;
+								}
+							}
+						}
+					}
+
 					if (!GetFocusedMovieSceneSequence()->GetMovieScene()->AddGivenTrack(NewTrack, ObjectGuid))
 					{
 						continue;
diff --git a/Engine/Source/Runtime/MovieScene/Private/MovieScene.cpp b/Engine/Source/Runtime/MovieScene/Private/MovieScene.cpp
--- a/Engine/Source/Runtime/MovieScene/Private/MovieScene.cpp
+++ b/Engine/Source/Runtime/MovieScene/Private/MovieScene.cpp
@@ -769,17 +769,6 @@ bool UMovieScene::AddGivenTrack(UMovieSceneTrack* InTrack, const FGuid& ObjectGu
 	{
 		if (Binding.GetObjectGuid() == ObjectGuid)
 		{
-			// Tracks of the same class should be unique per name.
-			for (UMovieSceneTrack* Track : Binding.GetTracks())
-			{
-				if (Track->GetClass() == InTrack->GetClass() && Track->GetTrackName() == InTrack->GetTrackName())
-				{
-					// If a track of the same class and name exists, remove it so the new track replaces it
-					Binding.RemoveTrack(*Track);
-					break;
-				}
-			}
-
 			InTrack->Rename(nullptr, this);
 			check(InTrack);
 			Binding.AddTrack(*InTrack);
