Sequencer: Remove parent binding if component not updated

diff --git a/Engine/Source/Editor/Sequencer/Private/Sequencer.cpp b/Engine/Source/Editor/Sequencer/Private/Sequencer.cpp
--- a/Engine/Source/Editor/Sequencer/Private/Sequencer.cpp
+++ b/Engine/Source/Editor/Sequencer/Private/Sequencer.cpp
@@ -6111,13 +6111,22 @@ FGuid FSequencer::DoAssignActor(AActor*const* InActors, int32 NumActors, FGuid I
 				FGuid ComponentGuid = FindObjectId(*ComponentToReplace, ActiveTemplateIDs.Top());
 				if (ComponentGuid.IsValid())
 				{
+					bool bComponentWasUpdated = false;
 					for (UActorComponent* NewComponent : Actor->GetComponents())
 					{
 						if (NewComponent->GetFullName(Actor) == ComponentToReplace->GetFullName(ActorToReplace))
 						{
 							UpdateComponent( ComponentGuid, NewComponent );
+							bComponentWasUpdated = true;
 						}
 					}
+
+					// Clear the parent guid since this possessable component doesn't match to any component on the new actor
+					if (!bComponentWasUpdated)
+					{
+						FMovieScenePossessable* ThisPossessable = OwnerMovieScene->FindPossessable(ComponentGuid);
+						ThisPossessable->SetParent(FGuid());
+					}
 				}
 			}
 		}
