Sequencer: If the sequence is locked, draw sections as locked (red border) and dim display node text.

diff --git a/Engine/Source/Editor/Sequencer/Private/DisplayNodes/SequencerDisplayNode.cpp b/Engine/Source/Editor/Sequencer/Private/DisplayNodes/SequencerDisplayNode.cpp
--- a/Engine/Source/Editor/Sequencer/Private/DisplayNodes/SequencerDisplayNode.cpp
+++ b/Engine/Source/Editor/Sequencer/Private/DisplayNodes/SequencerDisplayNode.cpp
@@ -710,6 +710,11 @@ bool FSequencerDisplayNode::IsDimmed() const
 		return true;
 	};
 
+	if (GetSequencer().IsReadOnly())
+	{
+		return true;
+	}
+
 	FSequencerDisplayNode *This = const_cast<FSequencerDisplayNode*>(this);
 	//if empty with no key areas or sections then it's active, otherwise
 	//find first child with active section, then it's active, else inactive.
diff --git a/Engine/Source/Editor/Sequencer/Private/SSequencerSection.cpp b/Engine/Source/Editor/Sequencer/Private/SSequencerSection.cpp
--- a/Engine/Source/Editor/Sequencer/Private/SSequencerSection.cpp
+++ b/Engine/Source/Editor/Sequencer/Private/SSequencerSection.cpp
@@ -1041,7 +1041,7 @@ int32 SSequencerSection::OnPaint( const FPaintArgs& Args, const FGeometry& Allot
 
 	UMovieSceneTrack* Track = SectionObject->GetTypedOuter<UMovieSceneTrack>();
 	const bool bEnabled = bParentEnabled && SectionObject->IsActive() && !(Track && Track->IsEvalDisabled());
-	const bool bLocked = SectionObject->IsLocked();
+	const bool bLocked = SectionObject->IsLocked() || SectionObject->IsReadOnly();
 
 	UMovieScenePropertyTrack* PropertyTrack = Cast<UMovieScenePropertyTrack>(Track);
 	bool bSetSectionToKey = false;
