Sequencer: Mute/Solo should be able to apply to a track, not it's parent binding.

diff --git a/Engine/Source/Editor/Sequencer/Private/DisplayNodes/SequencerDisplayNode.cpp b/Engine/Source/Editor/Sequencer/Private/DisplayNodes/SequencerDisplayNode.cpp
--- a/Engine/Source/Editor/Sequencer/Private/DisplayNodes/SequencerDisplayNode.cpp
+++ b/Engine/Source/Editor/Sequencer/Private/DisplayNodes/SequencerDisplayNode.cpp
@@ -1190,7 +1190,8 @@ FSequencerDisplayNode* FSequencerDisplayNode::GetBaseNode() const
 {
 	ESequencerNode::Type Type = GetType();
 
-	if (IsRootNode() || Type == ESequencerNode::Folder || Type == ESequencerNode::Object || GetParentOrRoot()->GetType() == ESequencerNode::Folder)
+	if (IsRootNode() || Type == ESequencerNode::Folder || Type == ESequencerNode::Object
+		|| (Type == ESequencerNode::Track && static_cast<const FSequencerTrackNode*>(this)->GetSubTrackMode() != FSequencerTrackNode::ESubTrackMode::SubTrack))
 	{
 		return (FSequencerDisplayNode*)this;
 	}
