Sequencer: Invalidate bound objects after rerunning construction scripts, so that they are re-evaluated when FindBoundObject is calleds. This fixes an issue where the binding in the tree view is red and the menu is empty.

diff --git a/Engine/Source/Editor/Sequencer/Private/Sequencer.cpp b/Engine/Source/Editor/Sequencer/Private/Sequencer.cpp
--- a/Engine/Source/Editor/Sequencer/Private/Sequencer.cpp
+++ b/Engine/Source/Editor/Sequencer/Private/Sequencer.cpp
@@ -1073,14 +1073,16 @@ void FSequencer::RerunConstructionScripts()
 		
 	UMovieSceneSequence* Sequence = RootTemplate.GetSequence(MovieSceneSequenceID::Root);
 
-	GetConstructionScriptActors(Sequence->GetMovieScene(), MovieSceneSequenceID::Root, BoundActors);
+	TArray < TPair<FMovieSceneSequenceID, FGuid> > BoundGuids;
+
+	GetConstructionScriptActors(Sequence->GetMovieScene(), MovieSceneSequenceID::Root, BoundActors, BoundGuids);
 
 	for (FMovieSceneSequenceIDRef SequenceID : RootTemplateInstance.GetThisFrameMetaData().ActiveSequences)
 	{
 		UMovieSceneSequence* SubSequence = RootTemplateInstance.GetSequence(SequenceID);
 		if (SubSequence)
 		{
-			GetConstructionScriptActors(SubSequence->GetMovieScene(), SequenceID, BoundActors);
+			GetConstructionScriptActors(SubSequence->GetMovieScene(), SequenceID, BoundActors, BoundGuids);
 		}
 	}
 
@@ -1091,9 +1093,14 @@ void FSequencer::RerunConstructionScripts()
 			BoundActor.Get()->RerunConstructionScripts();
 		}
 	}
+
+	for (TPair<FMovieSceneSequenceID, FGuid> BoundGuid : BoundGuids)
+	{
+		State.Invalidate(BoundGuid.Value, BoundGuid.Key);
+	}
 }
 
-void FSequencer::GetConstructionScriptActors(UMovieScene* MovieScene, FMovieSceneSequenceIDRef SequenceID, TSet<TWeakObjectPtr<AActor> >& BoundActors)
+void FSequencer::GetConstructionScriptActors(UMovieScene* MovieScene, FMovieSceneSequenceIDRef SequenceID, TSet<TWeakObjectPtr<AActor> >& BoundActors, TArray < TPair<FMovieSceneSequenceID, FGuid> >& BoundGuids)
 {
 	for (int32 Index = 0; Index < MovieScene->GetPossessableCount(); ++Index)
 	{
@@ -1111,6 +1118,7 @@ void FSequencer::GetConstructionScriptActors(UMovieScene* MovieScene, FMovieScen
 					if (Blueprint && Blueprint->bRunConstructionScriptInSequencer)
 					{
 						BoundActors.Add(Actor);
+						BoundGuids.Add(TPair<FMovieSceneSequenceID, FGuid>(SequenceID, ThisGuid));
 					}
 				}
 			}
@@ -1133,6 +1141,7 @@ void FSequencer::GetConstructionScriptActors(UMovieScene* MovieScene, FMovieScen
 					if (Blueprint && Blueprint->bRunConstructionScriptInSequencer)
 					{
 						BoundActors.Add(Actor);
+						BoundGuids.Add(TPair<FMovieSceneSequenceID, FGuid>(SequenceID, ThisGuid));
 					}
 				}
 			}
diff --git a/Engine/Source/Editor/Sequencer/Private/Sequencer.h b/Engine/Source/Editor/Sequencer/Private/Sequencer.h
--- a/Engine/Source/Editor/Sequencer/Private/Sequencer.h
+++ b/Engine/Source/Editor/Sequencer/Private/Sequencer.h
@@ -1010,7 +1010,7 @@ protected:
 	void RerunConstructionScripts();
 
 	/** Get actors that want to rerun construction scripts */
-	void GetConstructionScriptActors(UMovieScene*, FMovieSceneSequenceIDRef SequenceID, TSet<TWeakObjectPtr<AActor> >& BoundActors);
+	void GetConstructionScriptActors(UMovieScene*, FMovieSceneSequenceIDRef SequenceID, TSet<TWeakObjectPtr<AActor> >& BoundActors, TArray < TPair<FMovieSceneSequenceID, FGuid> >& BoundGuids);
 
 	/** Check whether we're viewing the master sequence or not */
 	bool IsViewingMasterSequence() const { return ActiveTemplateIDs.Num() == 1; }
