Sequencer: Reworked concept of exposed bindings to be presented as a tags in the Sequencer UI

  - Tags are now specified on object bindings and can be used to persistently reference similar bindings across different sequences
  - Tags are now always specified and stored on the root sequence, making it easier to rebind objects from within sub sequences
  - A new object binding tag manager window is available that gives a bird's eye view of all bindings and tags for the whole sequence
  - New functions are available on Level Sequence Actors for interacting with bindings by tag: SetBindingByTag, AddBindingByTag, RemoveBindingByTag, FindBindingByTag
  - Tags are now shown on the sequencer tree view for visibility

diff --git a/Engine/Source/Editor/EditorStyle/Private/SlateEditorStyle.cpp b/Engine/Source/Editor/EditorStyle/Private/SlateEditorStyle.cpp
--- a/Engine/Source/Editor/EditorStyle/Private/SlateEditorStyle.cpp
+++ b/Engine/Source/Editor/EditorStyle/Private/SlateEditorStyle.cpp
@@ -2451,6 +2451,22 @@ void FSlateEditorStyle::FStyle::SetupSequencerStyles()
 		Set("Sequencer.Section.BackgroundText", DEFAULT_FONT("Bold", 24));
 		Set("Sequencer.Section.EmptySpace", new BOX_BRUSH(TEXT("Sequencer/EmptySpace"), FMargin(0.f, 7.f / 14.f)));
 
+		Set("Sequencer.ExposedNamePill_BG", new BOX_BRUSH(TEXT("Sequencer/ExposedNamePill_BG"), FMargin(14.f / 30.f), FLinearColor(1.f, 1.f, 1.f, .8f)));
+		Set("Sequencer.ExposedNamePill", FButtonStyle()
+			.SetNormal(  BOX_BRUSH(TEXT("Sequencer/ExposedNamePill_BG"), FMargin(14.f / 30.f), FLinearColor(1.f, 1.f, 1.f, .8f)) )
+			.SetHovered( BOX_BRUSH(TEXT("Sequencer/ExposedNamePill_BG_Hovered"), FMargin(14.f / 30.f), FLinearColor::White) )
+			.SetPressed( BOX_BRUSH(TEXT("Sequencer/ExposedNamePill_BG_Pressed"), FMargin(14.f / 30.f), FLinearColor::White) )
+			.SetNormalPadding( FMargin(0,0,0,0) )
+			.SetPressedPadding( FMargin(0,0,0,0) )
+		);
+		Set("Sequencer.ExposedNamePill.Input", FEditableTextBoxStyle()
+			.SetBackgroundImageNormal(FSlateColorBrush(FLinearColor::Transparent))
+			.SetBackgroundImageHovered(FSlateColorBrush(FLinearColor::Transparent))
+			.SetBackgroundImageFocused(FSlateColorBrush(FLinearColor::Transparent))
+			.SetBackgroundImageReadOnly(FSlateColorBrush(FLinearColor::Transparent))
+		);
+
+
 		Set("Sequencer.AnimationOutliner.ColorStrip", FButtonStyle()
 			.SetNormal(FSlateNoResource())
 			.SetHovered(FSlateNoResource())
diff --git a/Engine/Source/Editor/MovieSceneTools/Private/MovieSceneObjectBindingIDPicker.cpp b/Engine/Source/Editor/MovieSceneTools/Private/MovieSceneObjectBindingIDPicker.cpp
--- a/Engine/Source/Editor/MovieSceneTools/Private/MovieSceneObjectBindingIDPicker.cpp
+++ b/Engine/Source/Editor/MovieSceneTools/Private/MovieSceneObjectBindingIDPicker.cpp
@@ -5,14 +5,11 @@
 #include "MovieSceneBindingOwnerInterface.h"
 #include "MovieSceneSequence.h"
 #include "MovieScene.h"
+#include "SequenceBindingTree.h"
 #include "Widgets/Input/SComboButton.h"
 #include "Widgets/Text/STextBlock.h"
-#include "Tracks/MovieSceneSubTrack.h"
-#include "Sections/MovieSceneSubSection.h"
-#include "Sections/MovieSceneCinematicShotSection.h"
 #include "Framework/MultiBox/MultiBoxBuilder.h"
 #include "Textures/SlateIcon.h"
-#include "Styling/SlateIconFinder.h"
 #include "EditorStyleSet.h"
 #include "Widgets/Images/SImage.h"
 #include "Widgets/SOverlay.h"
@@ -23,328 +20,6 @@
 
 #define LOCTEXT_NAMESPACE "MovieSceneObjectBindingIDPicker"
 
-DECLARE_DELEGATE_OneParam(FOnSelectionChanged, const FMovieSceneObjectBindingID&);
-
-/** Node that represents an object binding, or a sub sequence (where the guid is zero) */
-struct FSequenceBindingNode
-{
-	FSequenceBindingNode(FText InDisplayString, FMovieSceneObjectBindingID InBindingID, FSlateIcon InIcon)
-		: BindingID(InBindingID)
-		, DisplayString(InDisplayString)
-		, Icon(InIcon)
-		, bIsSpawnable(false)
-	{}
-
-	/** Add a child */
-	void AddChild(TSharedRef<FSequenceBindingNode> Child)
-	{
-		Child->ParentID = BindingID;
-		Children.Add(Child);
-	}
-
-	/** This object's ID, and its parent's */
-	FMovieSceneObjectBindingID BindingID, ParentID;
-	/** The display string that represents this node */
-	FText DisplayString;
-	/** A representative icon for the node */
-	FSlateIcon Icon;
-	/** Whether this is a spawnable or not */
-	bool bIsSpawnable;
-	/** Array holding this node's children */
-	TArray<TSharedRef<FSequenceBindingNode>> Children;
-};
-
-/** Stack of sequence IDs from parent to child */
-struct FSequenceIDStack
-{
-	/** Get the current accumulated sequence ID */
-	FMovieSceneSequenceID GetCurrent() const
-	{
-		FMovieSceneSequenceID ID = MovieSceneSequenceID::Root;
-		for (int32 Index = IDs.Num() - 1; Index >= 0; --Index)
-		{
-			ID = ID.AccumulateParentID(IDs[Index]);
-		}
-		return ID;
-	}
-
-	/** Push a sequence ID onto the stack */
-	void Push(FMovieSceneSequenceID InSequenceID) { IDs.Add(InSequenceID); }
-
-	/** Pop the last sequence ID off the stack */
-	void Pop() { IDs.RemoveAt(IDs.Num() - 1, 1, false); }
-	
-private:
-	TArray<FMovieSceneSequenceID> IDs;
-};
-
-
-/** Data structure used internally to represent the bindings of a sequence recursively */
-struct FSequenceBindingTree
-{
-	/**
-	 * Construct the tree structure from the specified sequence.
-	 *
-	 * @param InSequence			The sequence to generate the tree for
-	 * @param InActiveSequence		A sequence at which point we can start to generate localally resolving IDs
-	 * @param InActiveSequenceID	The sequence ID for the above sequence within the root context
-	 */
-	void Build(UMovieSceneSequence* InSequence, FObjectKey InActiveSequence, FMovieSceneSequenceID InActiveSequenceID)
-	{
-		bIsEmpty = true;
-
-		// Reset state
-		ActiveSequenceID = InActiveSequenceID;
-		ActiveSequence = InActiveSequence;
-		Hierarchy.Reset();
-
-		ActiveSequenceNode = nullptr;
-
-		// Create a node for the root sequence
-		FMovieSceneObjectBindingID RootSequenceID;
-		TSharedRef<FSequenceBindingNode> RootSequenceNode = MakeShared<FSequenceBindingNode>(FText(), RootSequenceID, FSlateIcon());
-		Hierarchy.Add(RootSequenceID, RootSequenceNode);
-
-		TopLevelNode = RootSequenceNode;
-
-		if (InSequence)
-		{
-			RootSequenceNode->DisplayString = FText::FromString(InSequence->GetName());
-			RootSequenceNode->Icon = FSlateIconFinder::FindIconForClass(InSequence->GetClass());
-
-			// Build the tree
-			FSequenceIDStack SequenceIDStack;
-			Build(InSequence, SequenceIDStack);
-
-			// Sort the tree
-			Sort(RootSequenceNode);
-
-			// We don't show cross-references to the same sequence since this would result in erroneous mixtures of absolute and local bindings
-			if (ActiveSequenceNode.IsValid() && ActiveSequenceNode != RootSequenceNode)
-			{
-				// Remove it from its parent, and put it at the root for quick access
-				TSharedPtr<FSequenceBindingNode> ActiveParent = Hierarchy.FindChecked(ActiveSequenceNode->ParentID);
-				ActiveParent->Children.Remove(ActiveSequenceNode.ToSharedRef());
-
-				// Make a new top level node (with an invalid ID)
-				FMovieSceneObjectBindingID TopLevelID = FMovieSceneObjectBindingID(FGuid(), FMovieSceneSequenceID());
-				TopLevelNode = MakeShared<FSequenceBindingNode>(FText(), TopLevelID, FSlateIcon());
-
-				// Override the display string and icon
-				ActiveSequenceNode->DisplayString = LOCTEXT("ThisSequenceText", "This Sequence");
-				ActiveSequenceNode->Icon = FSlateIcon();
-
-				TopLevelNode->Children.Add(ActiveSequenceNode.ToSharedRef());
-				TopLevelNode->Children.Add(RootSequenceNode);
-			}
-		}
-	}
-
-	/** Get the root of the tree */
-	TSharedRef<FSequenceBindingNode> GetRootNode() const
-	{
-		return TopLevelNode.ToSharedRef();
-	}
-
-	/** Find a node in the tree */
-	TSharedPtr<FSequenceBindingNode> FindNode(FMovieSceneObjectBindingID BindingID) const
-	{
-		return Hierarchy.FindRef(BindingID);
-	}
-
-	bool IsEmpty() const
-	{
-		return bIsEmpty;
-	}
-
-private:
-
-	/** Recursive sort helper for a sequence binding node */
-	static void Sort(TSharedRef<FSequenceBindingNode> Node)
-	{
-		Node->Children.Sort(
-			[](TSharedRef<FSequenceBindingNode> A, TSharedRef<FSequenceBindingNode> B)
-			{
-				// Sort shots first
-				if (A->BindingID.IsValid() != B->BindingID.IsValid())
-				{
-					return !A->BindingID.IsValid();
-				}
-				return A->DisplayString.CompareToCaseIgnored(B->DisplayString) < 0;
-			}
-		);
-
-		for (TSharedRef<FSequenceBindingNode> Child : Node->Children)
-		{
-			Sort(Child);
-		}
-	}
-
-	/** Recursive builder function that iterates into sub sequences */
-	void Build(UMovieSceneSequence* InSequence, FSequenceIDStack& SequenceIDStack)
-	{
-		check(InSequence);
-
-		UMovieScene* MovieScene = InSequence->GetMovieScene();
-		if (!MovieScene)
-		{
-			return;
-		}
-		
-		if (ActiveSequence == InSequence)
-		{
-			// Don't allow cross-references to the same sequence (ie, re-entrant references)
-			if (SequenceIDStack.GetCurrent() != ActiveSequenceID)
-			{
-				return;
-			}
-
-			// Keep track of the active sequence node
-			ActiveSequenceNode = Hierarchy.FindChecked(FMovieSceneObjectBindingID(FGuid(), SequenceIDStack.GetCurrent()));
-		}
-
-		// Iterate all sub sections
-		for (const UMovieSceneTrack* MasterTrack : MovieScene->GetMasterTracks())
-		{
-			const UMovieSceneSubTrack* SubTrack = Cast<const UMovieSceneSubTrack>(MasterTrack);
-			if (SubTrack)
-			{
-				for (UMovieSceneSection* Section : SubTrack->GetAllSections())
-				{
-					UMovieSceneSubSection* SubSection = Cast<UMovieSceneSubSection>(Section);
-					UMovieSceneSequence* SubSequence = SubSection ? SubSection->GetSequence() : nullptr;
-					if (SubSequence && Section->IsActive())
-					{
-						// Hold onto the current parent ID before adding our ID onto the stack
-						FMovieSceneSequenceID ParentID = SequenceIDStack.GetCurrent();
-						SequenceIDStack.Push(SubSection->GetSequenceID());
-						
-						FMovieSceneObjectBindingID CurrentID(FGuid(), SequenceIDStack.GetCurrent());
-
-						UMovieSceneCinematicShotSection* ShotSection = Cast<UMovieSceneCinematicShotSection>(Section);
-						FText DisplayString = ShotSection ? FText::FromString(ShotSection->GetShotDisplayName()) : FText::FromName(SubSequence->GetFName());
-						FSlateIcon Icon(FEditorStyle::GetStyleSetName(), ShotSection ? "Sequencer.Tracks.CinematicShot" : "Sequencer.Tracks.Sub");
-						
-						TSharedRef<FSequenceBindingNode> NewNode = MakeShared<FSequenceBindingNode>(DisplayString, CurrentID, Icon);
-						ensure(!Hierarchy.Contains(CurrentID));
-						Hierarchy.Add(CurrentID, NewNode);
-
-						EnsureParent(FGuid(), MovieScene, ParentID)->AddChild(NewNode);
-
-						Build(SubSequence, SequenceIDStack);
-
-						SequenceIDStack.Pop();
-					}
-				}
-			}
-		}
-
-		FMovieSceneSequenceID CurrentSequenceID = SequenceIDStack.GetCurrent();
-
-		// Add all spawnables first (since possessables can be children of spawnables)
-		int32 SpawnableCount = MovieScene->GetSpawnableCount();
-		for (int32 Index = 0; Index < SpawnableCount; ++Index)
-		{
-			const FMovieSceneSpawnable& Spawnable = MovieScene->GetSpawnable(Index);
-			
-			FMovieSceneObjectBindingID ID(Spawnable.GetGuid(), CurrentSequenceID);
-
-			FSlateIcon Icon = FSlateIconFinder::FindIconForClass(Spawnable.GetObjectTemplate()->GetClass());
-			TSharedRef<FSequenceBindingNode> NewNode = MakeShared<FSequenceBindingNode>(MovieScene->GetObjectDisplayName(Spawnable.GetGuid()), ID, Icon);
-			NewNode->bIsSpawnable = true;
-
-			EnsureParent(FGuid(), MovieScene, CurrentSequenceID)->AddChild(NewNode);
-			ensure(!Hierarchy.Contains(ID));
-			Hierarchy.Add(ID, NewNode);
-
-			bIsEmpty = false;
-		}
-
-		// Add all possessables
-		const int32 PossessableCount = MovieScene->GetPossessableCount();
-		for (int32 Index = 0; Index < PossessableCount; ++Index)
-		{
-			const FMovieScenePossessable& Possessable = MovieScene->GetPossessable(Index);
-			if (InSequence->CanRebindPossessable(Possessable))
-			{
-				FMovieSceneObjectBindingID ID(Possessable.GetGuid(), CurrentSequenceID);
-
-				FSlateIcon Icon = FSlateIconFinder::FindIconForClass(Possessable.GetPossessedObjectClass());
-				TSharedRef<FSequenceBindingNode> NewNode = MakeShared<FSequenceBindingNode>(MovieScene->GetObjectDisplayName(Possessable.GetGuid()), ID, Icon);
-
-				EnsureParent(Possessable.GetParent(), MovieScene, CurrentSequenceID)->AddChild(NewNode);
-				ensure(!Hierarchy.Contains(ID));
-				Hierarchy.Add(ID, NewNode);
-
-				bIsEmpty = false;
-			}
-		}
-	}
-
-	/** Ensure that a parent node exists for the specified object */
-	TSharedRef<FSequenceBindingNode> EnsureParent(const FGuid& InParentGuid, UMovieScene* InMovieScene, FMovieSceneSequenceID SequenceID)
-	{
-		FMovieSceneObjectBindingID ParentPtr(InParentGuid, SequenceID);
-
-		// If the node already exists
-		TSharedPtr<FSequenceBindingNode> Parent = Hierarchy.FindRef(ParentPtr);
-		if (Parent.IsValid())
-		{
-			return Parent.ToSharedRef();
-		}
-
-		// Non-object binding nodes should have already been added externally to EnsureParent
-		check(InParentGuid.IsValid());
-
-		// Need to add it
-		FGuid AddToGuid;
-		if (FMovieScenePossessable* GrandParentPossessable = InMovieScene->FindPossessable(InParentGuid))
-		{
-			AddToGuid = GrandParentPossessable->GetGuid();
-		}
-
-		// Deduce the icon for the node
-		FSlateIcon Icon;
-		bool bIsSpawnable = false;
-		{
-			const FMovieScenePossessable* Possessable = InMovieScene->FindPossessable(InParentGuid);
-			const FMovieSceneSpawnable* Spawnable = Possessable ? nullptr : InMovieScene->FindSpawnable(InParentGuid);
-			if (Possessable || Spawnable)
-			{
-				Icon = FSlateIconFinder::FindIconForClass(Possessable ? Possessable->GetPossessedObjectClass() : Spawnable->GetObjectTemplate()->GetClass());
-			}
-
-			bIsSpawnable = Spawnable != nullptr;
-		}
-
-		TSharedRef<FSequenceBindingNode> NewNode = MakeShared<FSequenceBindingNode>(InMovieScene->GetObjectDisplayName(InParentGuid), ParentPtr, Icon);
-		NewNode->bIsSpawnable = bIsSpawnable;
-		
-		ensure(!Hierarchy.Contains(ParentPtr));
-		Hierarchy.Add(ParentPtr, NewNode);
-		bIsEmpty = false;
-
-		EnsureParent(AddToGuid, InMovieScene, SequenceID)->AddChild(NewNode);
-
-		return NewNode;
-	}
-
-private:
-
-	/** The ID of the currently 'active' sequence from which to generate relative IDs */
-	FMovieSceneSequenceID ActiveSequenceID;
-	/** The currently 'active' sequence from which to generate relative IDs */
-	FObjectKey ActiveSequence;
-	/** The node relating to the currently active sequence ID (if any) */
-	TSharedPtr<FSequenceBindingNode> ActiveSequenceNode;
-	/** The top level (root) node in the tree */
-	TSharedPtr<FSequenceBindingNode> TopLevelNode;
-	/** Map of hierarchical information */
-	TMap<FMovieSceneObjectBindingID, TSharedPtr<FSequenceBindingNode>> Hierarchy;
-	/** Whether the tree is considered empty */
-	bool bIsEmpty;
-};
-
 bool FMovieSceneObjectBindingIDPicker::IsEmpty() const
 {
 	return !DataTree.IsValid() || DataTree->IsEmpty();
@@ -363,7 +38,7 @@ void FMovieSceneObjectBindingIDPicker::Initialize()
 	UMovieSceneSequence* ActiveSequence = Sequencer.IsValid() ? Sequencer->GetFocusedMovieSceneSequence() : GetSequence();
 	FMovieSceneSequenceID ActiveSequenceID = Sequencer.IsValid() ? Sequencer->GetFocusedTemplateID() : MovieSceneSequenceID::Root;
 
-	DataTree->Build(Sequence, ActiveSequence, ActiveSequenceID);
+	DataTree->ConditionalRebuild(Sequence, ActiveSequence, ActiveSequenceID);
 
 	UpdateCachedData();
 }
diff --git a/Engine/Source/Editor/MovieSceneTools/Private/ObjectBindingTagCache.cpp b/Engine/Source/Editor/MovieSceneTools/Private/ObjectBindingTagCache.cpp
new file mode 100644
--- /dev/null
+++ b/Engine/Source/Editor/MovieSceneTools/Private/ObjectBindingTagCache.cpp
@@ -0,0 +1,79 @@
+// Copyright 1998-2019 Epic Games, Inc. All Rights Reserved.
+
+#include "ObjectBindingTagCache.h"
+
+#include "MovieSceneSequence.h"
+#include "MovieScene.h"
+#include "MovieSceneObjectBindingID.h"
+
+#include "HAL/IConsoleManager.h"
+#include "Misc/SecureHash.h"
+
+float GSequencerTagSaturation = 0.6f;
+FAutoConsoleVariableRef CVarSequencerTagSaturation(
+	TEXT("Sequencer.TagSaturation"),
+	GSequencerTagSaturation,
+	TEXT("Specifies how saturated object binding tags should appear in the Sequencer UI.\n"),
+	ECVF_Default
+);
+
+
+void FObjectBindingTagCache::ConditionalUpdate(UMovieSceneSequence* RootSequence)
+{
+	if (!RootSequence)
+	{
+		if (RootSequenceSignature.IsValid())
+		{
+			RootSequenceSignature = FGuid();
+			ExposedNameReverseLUT.Empty();
+			OnUpdatedEvent.Broadcast(this);
+		}
+		return;
+	}
+
+	if (RootSequenceSignature.IsValid() && RootSequence->GetSignature() == RootSequenceSignature)
+	{
+		return;
+	}
+
+	RootSequenceSignature = RootSequence->GetSignature();
+
+	ExposedNameReverseLUT.Empty();
+
+	for (const TTuple<FName, FMovieSceneObjectBindingIDs>& Pair : RootSequence->GetMovieScene()->AllTaggedBindings())
+	{
+		FName Tag = Pair.Key;
+
+		// Define and cache a color for this type if necessary
+		if (!ExposedNameColors.Contains(Tag))
+		{
+			// Could use SHA1 here which might give a better distribution of color hues but I don't have good sample data to test with.
+			// CRC32 is fine for know until we know better
+			// BYTE HashBytes[20];
+			// FString TagString = Tag.ToString();
+
+			// FSHA1::HashBuffer(*TagString, sizeof(TCHAR)*TagString.Len(), HashBytes);
+
+			// Down sample the 20 byte hash to a 4 byte range
+			// uint32 StringHash = 0;
+			// for (int32 i = 0; i < 20; i+=4)
+			// {
+			//  	StringHash = HashCombine(StringHash, *reinterpret_cast<uint32*>(&HashBytes[i]));
+			// }
+
+			uint32 StringHash = GetTypeHash(Tag.ToString());
+			float Hue = (double(StringHash) / MAX_uint32) * 360.f;
+			FLinearColor ColorTintRGB = FLinearColor(Hue, GSequencerTagSaturation, .5f).HSVToLinearRGB();
+
+			ExposedNameColors.Add(Tag, ColorTintRGB);
+		}
+
+		for (FMovieSceneObjectBindingID ID : Pair.Value.IDs)
+		{
+			ExposedNameReverseLUT.Add(ID, Pair.Key);
+		}
+	}
+
+	OnUpdatedEvent.Broadcast(this);
+}
+
diff --git a/Engine/Source/Editor/MovieSceneTools/Private/SObjectBindingTag.cpp b/Engine/Source/Editor/MovieSceneTools/Private/SObjectBindingTag.cpp
new file mode 100644
--- /dev/null
+++ b/Engine/Source/Editor/MovieSceneTools/Private/SObjectBindingTag.cpp
@@ -0,0 +1,227 @@
+// Copyright 1998-2019 Epic Games, Inc. All Rights Reserved.
+
+#include "SObjectBindingTag.h"
+
+#include "MovieSceneObjectBindingID.h"
+#include "ObjectBindingTagCache.h"
+
+#include "SlateOptMacros.h"
+#include "Widgets/SNullWidget.h"
+#include "Widgets/SBoxPanel.h"
+#include "Widgets/Images/SImage.h"
+#include "Widgets/Text/STextBlock.h"
+#include "Widgets/Input/SEditableTextBox.h"
+#include "Widgets/Input/SButton.h"
+#include "Widgets/Notifications/SNotificationList.h"
+#include "Framework/Notifications/NotificationManager.h"
+
+#include "EditorStyleSet.h"
+#include "EditorFontGlyphs.h"
+
+#define LOCTEXT_NAMESPACE "SObjectBindingTag"
+
+void SObjectBindingTags::Construct(const FArguments& InArgs, const FMovieSceneObjectBindingID& InBindingID, FObjectBindingTagCache* BindingCache)
+{
+	BindingID = InBindingID;
+	OnTagDeletedEvent = InArgs._OnTagDeleted;
+
+	BindingCache->OnUpdatedEvent.AddSP(this, &SObjectBindingTags::OnBindingCacheUpdated);
+	OnBindingCacheUpdated(BindingCache);
+}
+
+void SObjectBindingTags::OnBindingCacheUpdated(const FObjectBindingTagCache* BindingCache)
+{
+	TSharedPtr<SHorizontalBox> Box;
+	float LeftPadding = 0.f;
+	for (auto It = BindingCache->IterateTags(BindingID); It; ++It)
+	{
+		FName TagName = It.Value();
+
+		if (!Box)
+		{
+			Box = SNew(SHorizontalBox);
+		}
+
+		FSimpleDelegate OnDeleteDelegate;
+		if (OnTagDeletedEvent.IsBound())
+		{
+			OnDeleteDelegate = FSimpleDelegate::CreateSP(this, &SObjectBindingTags::OnTagDeleted, TagName);
+		}
+
+		Box->AddSlot()
+		.Padding(FMargin(LeftPadding, 0.f, 0.f, 0.f))
+		.AutoWidth()
+		[
+			SNew(SObjectBindingTag)
+			.OnDeleted(OnDeleteDelegate)
+			.ColorTint(BindingCache->GetTagColor(TagName))
+			.Text(FText::FromName(TagName))
+			.ToolTipText(FText::Format(LOCTEXT("TagToolTip", "This object binding can be looked up, resolved, and overridden in the master sequence by using the tag '{0}'"), FText::FromName(TagName)))
+		];
+
+		LeftPadding = 5.f;
+	}
+
+	ChildSlot
+	[
+		Box.IsValid()
+		? Box.ToSharedRef()
+		: SNullWidget::NullWidget
+	];
+}
+
+void SObjectBindingTags::OnTagDeleted(FName TagName)
+{
+	OnTagDeletedEvent.Execute(TagName);
+}
+
+
+BEGIN_SLATE_FUNCTION_BUILD_OPTIMIZATION
+void SObjectBindingTag::Construct(const FArguments& InArgs)
+{
+	OnDeleted   = InArgs._OnDeleted;
+	OnClicked   = InArgs._OnClicked;
+	OnCreateNew = InArgs._OnCreateNew;
+
+	TSharedRef<SHorizontalBox> ContentBox = SNew(SHorizontalBox);
+
+	if (InArgs._OnCreateNew.IsBound())
+	{
+		ContentBox->AddSlot()
+		.VAlign(VAlign_Center)
+		[
+			SAssignNew(EditableText, SEditableTextBox)
+			.Font(FEditorStyle::GetFontStyle("TinyText"))
+			.Style(FEditorStyle::Get(), "Sequencer.ExposedNamePill.Input")
+			.BackgroundColor(FLinearColor::Transparent)
+			.OnTextCommitted(this, &SObjectBindingTag::OnNewTextCommitted)
+			.HintText(LOCTEXT("AddNew_Hint", "Enter New Name"))
+		];
+	}
+	else
+	{
+		ContentBox->AddSlot()
+		.VAlign(VAlign_Center)
+		[
+			SNew(STextBlock)
+			.Font(FEditorStyle::GetFontStyle("TinyText"))
+			.Text(InArgs._Text)
+		];
+	}
+
+	// Create button
+	if (OnCreateNew.IsBound())
+	{
+		ContentBox->AddSlot()
+		.Padding(FMargin(5.f, 0.f, 0.f, 0.f))
+		.VAlign(VAlign_Center)
+		.AutoWidth()
+		[
+			SNew(SButton)
+			.ContentPadding(FMargin(0.f))
+			.ButtonStyle(FEditorStyle::Get(), "HoverHintOnly")
+			.OnClicked(this, &SObjectBindingTag::HandleCreateButtonClicked)
+			.HAlign(HAlign_Center)
+			.VAlign(VAlign_Center)
+			[
+				SNew(STextBlock)
+				.Font(FEditorStyle::Get().GetFontStyle("FontAwesome.9"))
+				.Text(FEditorFontGlyphs::Plus)
+			]
+		];
+	}
+	// Delete button
+	if (InArgs._OnDeleted.IsBound())
+	{
+		ContentBox->AddSlot()
+		.Padding(FMargin(5.f, 0.f, 0.f, 0.f))
+		.VAlign(VAlign_Center)
+		.AutoWidth()
+		[
+			SNew(SButton)
+			.ContentPadding(FMargin(0.f))
+			.ButtonStyle(FEditorStyle::Get(), "HoverHintOnly")
+			.OnClicked(this, &SObjectBindingTag::HandleDeleteButtonClicked)
+			.HAlign(HAlign_Center)
+			.VAlign(VAlign_Center)
+			[
+				SNew(STextBlock)
+				.Font(FEditorStyle::Get().GetFontStyle("FontAwesome.9"))
+				.Text(FEditorFontGlyphs::Times)
+			]
+		];
+	}
+
+	if (OnClicked.IsBound())
+	{
+		// If this is clickable, we need to be a button
+		ChildSlot
+		[
+			SNew(SButton)
+			.ToolTipText(InArgs._ToolTipText)
+			.ButtonStyle(FEditorStyle::Get(), "Sequencer.ExposedNamePill")
+			.ButtonColorAndOpacity(InArgs._ColorTint)
+			.ContentPadding(FMargin(8.f, 2.f))
+			.OnClicked(this, &SObjectBindingTag::HandlePillClicked)
+			[
+				ContentBox
+			]
+		];
+	}
+	else
+	{
+		// else it's just a border
+		ChildSlot
+		[
+			SNew(SBorder)
+			.ToolTipText(InArgs._ToolTipText)
+			.BorderImage(FEditorStyle::GetBrush("Sequencer.ExposedNamePill_BG"))
+			.BorderBackgroundColor(InArgs._ColorTint)
+			.Padding(FMargin(8.f, 2.f))
+			[
+				ContentBox
+			]
+		];
+	}
+}
+END_SLATE_FUNCTION_BUILD_OPTIMIZATION
+
+void SObjectBindingTag::OnNewTextCommitted(const FText& InNewText, ETextCommit::Type CommitType)
+{
+	if (!InNewText.IsEmpty() && CommitType == ETextCommit::OnEnter)
+	{
+		FString NewName = InNewText.ToString();
+		FText FailureReason, Context = LOCTEXT("TagName", "Object binding tags");
+		if (!FName::IsValidXName(NewName, INVALID_NAME_CHARACTERS, &FailureReason, &Context))
+		{
+			FNotificationInfo NotificationInfo(FailureReason);
+			NotificationInfo.ExpireDuration = 4.0f;
+			FSlateNotificationManager::Get().AddNotification(NotificationInfo);
+			return;
+		}
+
+		OnCreateNew.Execute(FName(*NewName));
+
+		EditableText->SetText(FText());
+	}
+}
+
+FReply SObjectBindingTag::HandleCreateButtonClicked()
+{
+	OnNewTextCommitted(EditableText->GetText(), ETextCommit::OnEnter);
+	return FReply::Handled();
+}
+
+FReply SObjectBindingTag::HandleDeleteButtonClicked()
+{
+	OnDeleted.Execute();
+	return FReply::Handled();
+}
+
+FReply SObjectBindingTag::HandlePillClicked()
+{
+	OnClicked.Execute();
+	return FReply::Handled();
+}
+
+#undef LOCTEXT_NAMESPACE
\ No newline at end of file
diff --git a/Engine/Source/Editor/MovieSceneTools/Private/SequenceBindingTree.cpp b/Engine/Source/Editor/MovieSceneTools/Private/SequenceBindingTree.cpp
new file mode 100644
--- /dev/null
+++ b/Engine/Source/Editor/MovieSceneTools/Private/SequenceBindingTree.cpp
@@ -0,0 +1,306 @@
+// Copyright 1998-2019 Epic Games, Inc. All Rights Reserved.
+
+#include "SequenceBindingTree.h"
+#include "MovieSceneSequence.h"
+#include "MovieScene.h"
+#include "Tracks/MovieSceneSubTrack.h"
+#include "Sections/MovieSceneSubSection.h"
+#include "Sections/MovieSceneCinematicShotSection.h"
+
+#include "Styling/SlateIconFinder.h"
+#include "Textures/SlateIcon.h"
+#include "EditorStyleSet.h"
+
+#define LOCTEXT_NAMESPACE "MovieSceneObjectBindingIDPicker"
+
+/** Stack of sequence IDs from parent to child */
+struct FSequenceBindingTree::FSequenceIDStack
+{
+	/** Get the current accumulated sequence ID */
+	FMovieSceneSequenceID GetCurrent() const
+	{
+		FMovieSceneSequenceID ID = MovieSceneSequenceID::Root;
+		for (int32 Index = IDs.Num() - 1; Index >= 0; --Index)
+		{
+			ID = ID.AccumulateParentID(IDs[Index]);
+		}
+		return ID;
+	}
+
+	/** Push a sequence ID onto the stack */
+	void Push(FMovieSceneSequenceID InSequenceID) { IDs.Add(InSequenceID); }
+
+	/** Pop the last sequence ID off the stack */
+	void Pop() { IDs.RemoveAt(IDs.Num() - 1, 1, false); }
+	
+private:
+	TArray<FMovieSceneSequenceID> IDs;
+};
+
+bool FSequenceBindingTree::ConditionalRebuild(UMovieSceneSequence* InSequence, FObjectKey InActiveSequence, FMovieSceneSequenceID InActiveSequenceID)
+{
+	struct FInternal
+	{
+		static bool IsOutOfDate(UMovieSceneSequence* ThisSequence, TMap<FObjectKey, FGuid>* InCachedSequenceSignatures)
+		{
+			FGuid* ExistingSignature = InCachedSequenceSignatures->Find(ThisSequence);
+			if (!ExistingSignature || *ExistingSignature != ThisSequence->GetSignature())
+			{
+				return true;
+			}
+
+			UMovieScene* ThisMovieScene = ThisSequence->GetMovieScene();
+			for (const UMovieSceneTrack* MasterTrack : ThisMovieScene->GetMasterTracks())
+			{
+				const UMovieSceneSubTrack* SubTrack = Cast<const UMovieSceneSubTrack>(MasterTrack);
+				if (SubTrack)
+				{
+					for (UMovieSceneSection* Section : SubTrack->GetAllSections())
+					{
+						UMovieSceneSubSection* SubSection = Cast<UMovieSceneSubSection>(Section);
+						UMovieSceneSequence* SubSequence = SubSection ? SubSection->GetSequence() : nullptr;
+						if (SubSequence && Section->IsActive())
+						{
+							if (IsOutOfDate(SubSequence, InCachedSequenceSignatures))
+							{
+								return true;
+							}
+						}
+					}
+				}
+			}
+			return false;
+		}
+	};
+	
+
+	if (FInternal::IsOutOfDate(InSequence, &CachedSequenceSignatures))
+	{
+		ForceRebuild(InSequence, InActiveSequence, InActiveSequenceID);
+		return true;
+	}
+	return false;
+}
+
+void FSequenceBindingTree::ForceRebuild(UMovieSceneSequence* InSequence, FObjectKey InActiveSequence, FMovieSceneSequenceID InActiveSequenceID)
+{
+	CachedSequenceSignatures.Empty();
+	bIsEmpty = true;
+
+	// Reset state
+	ActiveSequenceID = InActiveSequenceID;
+	ActiveSequence = InActiveSequence;
+	Hierarchy.Reset();
+
+	ActiveSequenceNode = nullptr;
+
+	// Create a node for the root sequence
+	FMovieSceneObjectBindingID RootSequenceID;
+	TSharedRef<FSequenceBindingNode> RootSequenceNode = MakeShared<FSequenceBindingNode>(FText(), RootSequenceID, FSlateIcon());
+	Hierarchy.Add(RootSequenceID, RootSequenceNode);
+
+	TopLevelNode = RootSequenceNode;
+
+	if (InSequence)
+	{
+		RootSequenceNode->DisplayString = FText::FromString(InSequence->GetName());
+		RootSequenceNode->Icon = FSlateIconFinder::FindIconForClass(InSequence->GetClass());
+
+		// Build the tree
+		FSequenceIDStack SequenceIDStack;
+		Build(InSequence, SequenceIDStack);
+
+		// Sort the tree
+		Sort(RootSequenceNode);
+
+		// We don't show cross-references to the same sequence since this would result in erroneous mixtures of absolute and local bindings
+		if (ActiveSequenceNode.IsValid() && ActiveSequenceNode != RootSequenceNode)
+		{
+			// Remove it from its parent, and put it at the root for quick access
+			TSharedPtr<FSequenceBindingNode> ActiveParent = Hierarchy.FindChecked(ActiveSequenceNode->ParentID);
+			ActiveParent->Children.Remove(ActiveSequenceNode.ToSharedRef());
+
+			// Make a new top level node (with an invalid ID)
+			FMovieSceneObjectBindingID TopLevelID = FMovieSceneObjectBindingID(FGuid(), FMovieSceneSequenceID());
+			TopLevelNode = MakeShared<FSequenceBindingNode>(FText(), TopLevelID, FSlateIcon());
+
+			// Override the display string and icon
+			ActiveSequenceNode->DisplayString = LOCTEXT("ThisSequenceText", "This Sequence");
+			ActiveSequenceNode->Icon = FSlateIcon();
+
+			TopLevelNode->Children.Add(ActiveSequenceNode.ToSharedRef());
+			TopLevelNode->Children.Add(RootSequenceNode);
+		}
+	}
+}
+
+void FSequenceBindingTree::Sort(TSharedRef<FSequenceBindingNode> Node)
+{
+	Node->Children.Sort(
+		[](TSharedRef<FSequenceBindingNode> A, TSharedRef<FSequenceBindingNode> B)
+		{
+			// Sort shots first
+			if (A->BindingID.IsValid() != B->BindingID.IsValid())
+			{
+				return !A->BindingID.IsValid();
+			}
+			return A->DisplayString.CompareToCaseIgnored(B->DisplayString) < 0;
+		}
+	);
+
+	for (TSharedRef<FSequenceBindingNode> Child : Node->Children)
+	{
+		Sort(Child);
+	}
+}
+
+void FSequenceBindingTree::Build(UMovieSceneSequence* InSequence, FSequenceIDStack& SequenceIDStack)
+{
+	check(InSequence);
+
+	UMovieScene* MovieScene = InSequence->GetMovieScene();
+	if (!MovieScene)
+	{
+		return;
+	}
+	
+	if (ActiveSequence == InSequence)
+	{
+		// Don't allow cross-references to the same sequence (ie, re-entrant references)
+		if (SequenceIDStack.GetCurrent() != ActiveSequenceID)
+		{
+			return;
+		}
+
+		// Keep track of the active sequence node
+		ActiveSequenceNode = Hierarchy.FindChecked(FMovieSceneObjectBindingID(FGuid(), SequenceIDStack.GetCurrent()));
+	}
+
+	// Iterate all sub sections
+	for (const UMovieSceneTrack* MasterTrack : MovieScene->GetMasterTracks())
+	{
+		const UMovieSceneSubTrack* SubTrack = Cast<const UMovieSceneSubTrack>(MasterTrack);
+		if (SubTrack)
+		{
+			for (UMovieSceneSection* Section : SubTrack->GetAllSections())
+			{
+				UMovieSceneSubSection* SubSection = Cast<UMovieSceneSubSection>(Section);
+				UMovieSceneSequence* SubSequence = SubSection ? SubSection->GetSequence() : nullptr;
+				if (SubSequence && Section->IsActive())
+				{
+					// Hold onto the current parent ID before adding our ID onto the stack
+					FMovieSceneSequenceID ParentID = SequenceIDStack.GetCurrent();
+					SequenceIDStack.Push(SubSection->GetSequenceID());
+					
+					FMovieSceneObjectBindingID CurrentID(FGuid(), SequenceIDStack.GetCurrent());
+
+					UMovieSceneCinematicShotSection* ShotSection = Cast<UMovieSceneCinematicShotSection>(Section);
+					FText DisplayString = ShotSection ? FText::FromString(ShotSection->GetShotDisplayName()) : FText::FromName(SubSequence->GetFName());
+					FSlateIcon Icon(FEditorStyle::GetStyleSetName(), ShotSection ? "Sequencer.Tracks.CinematicShot" : "Sequencer.Tracks.Sub");
+					
+					TSharedRef<FSequenceBindingNode> NewNode = MakeShared<FSequenceBindingNode>(DisplayString, CurrentID, Icon);
+					ensure(!Hierarchy.Contains(CurrentID));
+					Hierarchy.Add(CurrentID, NewNode);
+
+					EnsureParent(FGuid(), MovieScene, ParentID)->AddChild(NewNode);
+
+					Build(SubSequence, SequenceIDStack);
+
+					SequenceIDStack.Pop();
+				}
+			}
+		}
+	}
+
+	FMovieSceneSequenceID CurrentSequenceID = SequenceIDStack.GetCurrent();
+
+	CachedSequenceSignatures.Add(InSequence, InSequence->GetSignature());
+
+	// Add all spawnables first (since possessables can be children of spawnables)
+	int32 SpawnableCount = MovieScene->GetSpawnableCount();
+	for (int32 Index = 0; Index < SpawnableCount; ++Index)
+	{
+		const FMovieSceneSpawnable& Spawnable = MovieScene->GetSpawnable(Index);
+		
+		FMovieSceneObjectBindingID ID(Spawnable.GetGuid(), CurrentSequenceID);
+
+		FSlateIcon Icon = FSlateIconFinder::FindIconForClass(Spawnable.GetObjectTemplate()->GetClass());
+		TSharedRef<FSequenceBindingNode> NewNode = MakeShared<FSequenceBindingNode>(MovieScene->GetObjectDisplayName(Spawnable.GetGuid()), ID, Icon);
+		NewNode->bIsSpawnable = true;
+
+		EnsureParent(FGuid(), MovieScene, CurrentSequenceID)->AddChild(NewNode);
+		ensure(!Hierarchy.Contains(ID));
+		Hierarchy.Add(ID, NewNode);
+
+		bIsEmpty = false;
+	}
+
+	// Add all possessables
+	const int32 PossessableCount = MovieScene->GetPossessableCount();
+	for (int32 Index = 0; Index < PossessableCount; ++Index)
+	{
+		const FMovieScenePossessable& Possessable = MovieScene->GetPossessable(Index);
+		if (InSequence->CanRebindPossessable(Possessable))
+		{
+			FMovieSceneObjectBindingID ID(Possessable.GetGuid(), CurrentSequenceID);
+
+			FSlateIcon Icon = FSlateIconFinder::FindIconForClass(Possessable.GetPossessedObjectClass());
+			TSharedRef<FSequenceBindingNode> NewNode = MakeShared<FSequenceBindingNode>(MovieScene->GetObjectDisplayName(Possessable.GetGuid()), ID, Icon);
+
+			EnsureParent(Possessable.GetParent(), MovieScene, CurrentSequenceID)->AddChild(NewNode);
+			ensure(!Hierarchy.Contains(ID));
+			Hierarchy.Add(ID, NewNode);
+
+			bIsEmpty = false;
+		}
+	}
+}
+
+TSharedRef<FSequenceBindingNode> FSequenceBindingTree::EnsureParent(const FGuid& InParentGuid, UMovieScene* InMovieScene, FMovieSceneSequenceID SequenceID)
+{
+	FMovieSceneObjectBindingID ParentPtr(InParentGuid, SequenceID);
+
+	// If the node already exists
+	TSharedPtr<FSequenceBindingNode> Parent = Hierarchy.FindRef(ParentPtr);
+	if (Parent.IsValid())
+	{
+		return Parent.ToSharedRef();
+	}
+
+	// Non-object binding nodes should have already been added externally to EnsureParent
+	check(InParentGuid.IsValid());
+
+	// Need to add it
+	FGuid AddToGuid;
+	if (FMovieScenePossessable* GrandParentPossessable = InMovieScene->FindPossessable(InParentGuid))
+	{
+		AddToGuid = GrandParentPossessable->GetGuid();
+	}
+
+	// Deduce the icon for the node
+	FSlateIcon Icon;
+	bool bIsSpawnable = false;
+	{
+		const FMovieScenePossessable* Possessable = InMovieScene->FindPossessable(InParentGuid);
+		const FMovieSceneSpawnable* Spawnable = Possessable ? nullptr : InMovieScene->FindSpawnable(InParentGuid);
+		if (Possessable || Spawnable)
+		{
+			Icon = FSlateIconFinder::FindIconForClass(Possessable ? Possessable->GetPossessedObjectClass() : Spawnable->GetObjectTemplate()->GetClass());
+		}
+
+		bIsSpawnable = Spawnable != nullptr;
+	}
+
+	TSharedRef<FSequenceBindingNode> NewNode = MakeShared<FSequenceBindingNode>(InMovieScene->GetObjectDisplayName(InParentGuid), ParentPtr, Icon);
+	NewNode->bIsSpawnable = bIsSpawnable;
+	
+	ensure(!Hierarchy.Contains(ParentPtr));
+	Hierarchy.Add(ParentPtr, NewNode);
+	bIsEmpty = false;
+
+	EnsureParent(AddToGuid, InMovieScene, SequenceID)->AddChild(NewNode);
+
+	return NewNode;
+}
+
+#undef LOCTEXT_NAMESPACE
diff --git a/Engine/Source/Editor/MovieSceneTools/Public/ObjectBindingTagCache.h b/Engine/Source/Editor/MovieSceneTools/Public/ObjectBindingTagCache.h
new file mode 100644
--- /dev/null
+++ b/Engine/Source/Editor/MovieSceneTools/Public/ObjectBindingTagCache.h
@@ -0,0 +1,73 @@
+// Copyright 1998-2019 Epic Games, Inc. All Rights Reserved.
+
+#pragma once
+
+#include "CoreTypes.h"
+#include "Delegates/Delegate.h"
+#include "Misc/Guid.h"
+#include "Containers/Map.h"
+#include "Containers/SortedMap.h"
+#include "MovieSceneObjectBindingID.h"
+
+class UMovieSceneSequence;
+
+
+/**
+ * Owned by an FSequencer instance, this class tracks tags for object bindings, and maintains a reverse lookup map from binding to tag(s)
+ * along with other information for the tags such as color tints.
+ */
+class FObjectBindingTagCache
+{
+public:
+
+	using TagIterator = TMultiMap<FMovieSceneObjectBindingID, FName>::TConstKeyIterator;
+
+	DECLARE_MULTICAST_DELEGATE_OneParam(FOnUpdated, const FObjectBindingTagCache*);
+
+	/** Public event that is broadcast when this cache changes in some way */
+	FOnUpdated OnUpdatedEvent;
+
+	/*
+	 * Conditionally update this cache based on the state of the specified root sequence
+	 */
+	MOVIESCENETOOLS_API void ConditionalUpdate(UMovieSceneSequence* RootSequence);
+
+
+	/*
+	 * Retrieve a color for the specified tag name
+	 */
+	FLinearColor GetTagColor(const FName& TagName) const
+	{
+		const FLinearColor* Color = ExposedNameColors.Find(TagName);
+		return Color ? *Color : FLinearColor::White;
+	}
+
+
+	/*
+	 * Check whether the specified binding ID has the specified tag
+	 */
+	bool HasTag(const FMovieSceneObjectBindingID& BindingID, const FName& TagName) const
+	{
+		return ExposedNameReverseLUT.FindPair(BindingID, TagName);
+	}
+
+
+	/*
+	 * Iterate all the tags for the specified binding ID
+	 */
+	TagIterator IterateTags(const FMovieSceneObjectBindingID& BindingID) const
+	{
+		return ExposedNameReverseLUT.CreateConstKeyIterator(BindingID);
+	}
+
+private:
+
+	/** Sorted map of colors by tag name */
+	TSortedMap<FName, FLinearColor, FDefaultAllocator, FNameFastLess> ExposedNameColors;
+
+	/** The signature of the root sequence last time we were updated */
+	FGuid RootSequenceSignature;
+
+	/** Multi map (1->many) of a binding ID to any tags that are associated with it */
+	TMultiMap<FMovieSceneObjectBindingID, FName> ExposedNameReverseLUT;
+};
\ No newline at end of file
diff --git a/Engine/Source/Editor/MovieSceneTools/Public/SObjectBindingTag.h b/Engine/Source/Editor/MovieSceneTools/Public/SObjectBindingTag.h
new file mode 100644
--- /dev/null
+++ b/Engine/Source/Editor/MovieSceneTools/Public/SObjectBindingTag.h
@@ -0,0 +1,102 @@
+// Copyright 1998-2019 Epic Games, Inc. All Rights Reserved.
+
+#pragma once
+
+#include "CoreTypes.h"
+#include "MovieSceneObjectBindingID.h"
+
+#include "Widgets/DeclarativeSyntaxSupport.h"
+#include "Widgets/SCompoundWidget.h"
+
+class SEditableTextBox;
+class FObjectBindingTagCache;
+
+/** A horizontally grouped collection of tags that apply to a single object binding. This widget automatically updates when the tags for a binding change. */
+class SObjectBindingTags : public SCompoundWidget
+{
+public:
+
+	/** A simple delegate that is passed the tag name applicable to the operation */
+	DECLARE_DELEGATE_OneParam(FTagEvent, FName)
+
+	SLATE_BEGIN_ARGS(SObjectBindingTags){}
+
+		/** (Optional) When bound, will show a delete button on the tags, and call this delegate when clicked */
+		SLATE_EVENT(FTagEvent, OnTagDeleted)
+
+	SLATE_END_ARGS()
+
+	/** Construct this widget. Binding cache can be retrieved through FSequencer::GetObjectTagCache()*/
+	MOVIESCENETOOLS_API void Construct(const FArguments& InArgs, const FMovieSceneObjectBindingID& InBindingID, FObjectBindingTagCache* BindingCache);
+
+private:
+
+	/** Called when the binding tag cache that holds information about the binding's tags changes */
+	void OnBindingCacheUpdated(const FObjectBindingTagCache* BindingCache);
+
+	/** Called when a tag is to be deleted */
+	void OnTagDeleted(FName TagName);
+
+private:
+
+	FTagEvent OnTagDeletedEvent;
+	FMovieSceneObjectBindingID BindingID;
+};
+
+/** A single named tag widget for an object binding within a sequence, represented as a rounded 'pill' */
+class SObjectBindingTag : public SCompoundWidget
+{
+public:
+
+	DECLARE_DELEGATE_OneParam(FOnCreateNew, FName)
+
+	SLATE_BEGIN_ARGS(SObjectBindingTag) : _ColorTint(FLinearColor::White) {}
+
+		/** (Optional) When bound, shows a small delete button on the tag that invokes this delegate when clicked */
+		SLATE_EVENT(FSimpleDelegate, OnDeleted)
+
+		/** (Optional) When bound, the text portion of the tag will be a text input allowing the user to create a new named tag */
+		SLATE_EVENT(FOnCreateNew, OnCreateNew)
+
+		/** (Optional) When bound, the whole tag will be clickable, and invoke this delegate when clicked */
+		SLATE_EVENT(FSimpleDelegate, OnClicked)
+
+		/** Text to display on the tag (when OnCreateNew is not specified) */
+		SLATE_ATTRIBUTE(FText, Text)
+
+		/** Tool tip text for this whole widget */
+		SLATE_ATTRIBUTE(FText, ToolTipText)
+
+		/** Color tint for the whole widget */
+		SLATE_ATTRIBUTE(FSlateColor, ColorTint)
+
+	SLATE_END_ARGS()
+
+	/**
+	 * Construct this widget
+	 */
+	MOVIESCENETOOLS_API void Construct(const FArguments& InArgs);
+
+private:
+
+	void OnNewTextCommitted(const FText& InNewText, ETextCommit::Type CommitType);
+
+	FReply HandleCreateButtonClicked();
+
+	FReply HandleDeleteButtonClicked();
+
+	FReply HandlePillClicked();
+
+	virtual FVector2D ComputeDesiredSize(float LayoutScaleMultiplier) const override
+	{
+		const float DesiredX = SCompoundWidget::ComputeDesiredSize(LayoutScaleMultiplier).X;
+		return FVector2D(DesiredX, 24.f);
+	}
+
+private:
+	FOnCreateNew OnCreateNew;
+	FSimpleDelegate OnDeleted;
+	FSimpleDelegate OnClicked;
+
+	TSharedPtr<SEditableTextBox> EditableText;
+};
\ No newline at end of file
diff --git a/Engine/Source/Editor/MovieSceneTools/Public/SequenceBindingTree.h b/Engine/Source/Editor/MovieSceneTools/Public/SequenceBindingTree.h
new file mode 100644
--- /dev/null
+++ b/Engine/Source/Editor/MovieSceneTools/Public/SequenceBindingTree.h
@@ -0,0 +1,115 @@
+// Copyright 1998-2019 Epic Games, Inc. All Rights Reserved.
+
+
+#include "CoreTypes.h"
+#include "Misc/Guid.h"
+#include "Containers/Map.h"
+#include "UObject/ObjectKey.h"
+#include "Textures/SlateIcon.h"
+#include "Templates/SharedPointer.h"
+#include "Internationalization/Text.h"
+
+#include "MovieSceneSequenceID.h"
+#include "MovieSceneObjectBindingID.h"
+
+class UMovieScene;
+class UMovieSceneSequence;
+
+/** Node that represents an object binding, or a sub sequence (where the guid is zero) */
+struct FSequenceBindingNode
+{
+	FSequenceBindingNode(FText InDisplayString, FMovieSceneObjectBindingID InBindingID, FSlateIcon InIcon)
+		: BindingID(InBindingID)
+		, DisplayString(InDisplayString)
+		, Icon(InIcon)
+		, bIsSpawnable(false)
+	{}
+
+	/** Add a child */
+	void AddChild(TSharedRef<FSequenceBindingNode> Child)
+	{
+		Child->ParentID = BindingID;
+		Children.Add(Child);
+	}
+
+	/** This object's ID, and its parent's */
+	FMovieSceneObjectBindingID BindingID, ParentID;
+	/** The display string that represents this node */
+	FText DisplayString;
+	/** A representative icon for the node */
+	FSlateIcon Icon;
+	/** Whether this is a spawnable or not */
+	bool bIsSpawnable;
+	/** Array holding this node's children */
+	TArray<TSharedRef<FSequenceBindingNode>> Children;
+};
+
+/** Data structure used internally to represent the bindings of a sequence recursively */
+struct MOVIESCENETOOLS_API FSequenceBindingTree
+{
+	/**
+	 * Conditionally reconstruct the tree structure from the specified sequence if any of the sequences have changed since the last population
+	 *
+	 * @param InSequence			The sequence to generate the tree for
+	 * @param InActiveSequence		A sequence at which point we can start to generate localally resolving IDs
+	 * @param InActiveSequenceID	The sequence ID for the above sequence within the root context
+	 * @retrun True if this tree was rebuilt, false otherwise
+	 */
+	bool ConditionalRebuild(UMovieSceneSequence* InSequence, FObjectKey InActiveSequence, FMovieSceneSequenceID InActiveSequenceID);
+
+	/**
+	 * Construct the tree structure from the specified sequence.
+	 *
+	 * @param InSequence			The sequence to generate the tree for
+	 * @param InActiveSequence		A sequence at which point we can start to generate localally resolving IDs
+	 * @param InActiveSequenceID	The sequence ID for the above sequence within the root context
+	 */
+	void ForceRebuild(UMovieSceneSequence* InSequence, FObjectKey InActiveSequence, FMovieSceneSequenceID InActiveSequenceID);
+
+	/** Get the root of the tree */
+	TSharedRef<FSequenceBindingNode> GetRootNode() const
+	{
+		return TopLevelNode.ToSharedRef();
+	}
+
+	/** Find a node in the tree */
+	TSharedPtr<FSequenceBindingNode> FindNode(FMovieSceneObjectBindingID BindingID) const
+	{
+		return Hierarchy.FindRef(BindingID);
+	}
+
+	bool IsEmpty() const
+	{
+		return bIsEmpty;
+	}
+
+private:
+
+	struct FSequenceIDStack;
+
+	/** Recursive sort helper for a sequence binding node */
+	static void Sort(TSharedRef<FSequenceBindingNode> Node);
+
+	/** Recursive builder function that iterates into sub sequences */
+	void Build(UMovieSceneSequence* InSequence, FSequenceIDStack& SequenceIDStack);
+
+	/** Ensure that a parent node exists for the specified object */
+	TSharedRef<FSequenceBindingNode> EnsureParent(const FGuid& InParentGuid, UMovieScene* InMovieScene, FMovieSceneSequenceID SequenceID);
+
+private:
+
+	/** The ID of the currently 'active' sequence from which to generate relative IDs */
+	FMovieSceneSequenceID ActiveSequenceID;
+	/** The currently 'active' sequence from which to generate relative IDs */
+	FObjectKey ActiveSequence;
+	/** The node relating to the currently active sequence ID (if any) */
+	TSharedPtr<FSequenceBindingNode> ActiveSequenceNode;
+	/** The top level (root) node in the tree */
+	TSharedPtr<FSequenceBindingNode> TopLevelNode;
+	/** Map of hierarchical information */
+	TMap<FMovieSceneObjectBindingID, TSharedPtr<FSequenceBindingNode>> Hierarchy;
+	/** Map of sequence to its signature the last time we were built */
+	TMap<FObjectKey, FGuid> CachedSequenceSignatures;
+	/** Whether the tree is considered empty */
+	bool bIsEmpty;
+};
\ No newline at end of file
diff --git a/Engine/Source/Editor/Sequencer/Private/DisplayNodes/SequencerDisplayNode.h b/Engine/Source/Editor/Sequencer/Private/DisplayNodes/SequencerDisplayNode.h
--- a/Engine/Source/Editor/Sequencer/Private/DisplayNodes/SequencerDisplayNode.h
+++ b/Engine/Source/Editor/Sequencer/Private/DisplayNodes/SequencerDisplayNode.h
@@ -247,6 +247,13 @@ public:
 	 */
 	virtual TSharedRef<SWidget> GetCustomOutlinerContent();
 
+	/**
+	 * Creates an additional label widget to appear immediately beside this node's label on the tree
+	 * 
+	 * @return Content to display on the outliner node
+	 */
+	virtual TSharedPtr<SWidget> GetAdditionalOutlinerLabel() { return nullptr; }
+
 	/**
 	 * Generates a widget for display in the section area portion of the track area
 	 * 
diff --git a/Engine/Source/Editor/Sequencer/Private/DisplayNodes/SequencerObjectBindingNode.cpp b/Engine/Source/Editor/Sequencer/Private/DisplayNodes/SequencerObjectBindingNode.cpp
--- a/Engine/Source/Editor/Sequencer/Private/DisplayNodes/SequencerObjectBindingNode.cpp
+++ b/Engine/Source/Editor/Sequencer/Private/DisplayNodes/SequencerObjectBindingNode.cpp
@@ -36,7 +36,8 @@
 #include "MovieSceneSequence.h"
 #include "MovieScene.h"
 #include "MovieSceneFolder.h"
-#include "SExposeBindingWidget.h"
+#include "ObjectBindingTagCache.h"
+#include "SObjectBindingTag.h"
 #include "ISequencerTrackEditor.h"
 
 #include "Tracks/MovieSceneSpawnTrack.h"
@@ -157,15 +158,10 @@ void FSequencerObjectBindingNode::BuildContextMenu(FMenuBuilder& MenuBuilder)
 		MenuBuilder.PushExtender(Extender.ToSharedRef());
 	}
 
-	MenuBuilder.AddSubMenu(
-		LOCTEXT("ExposeBindingLabel", "Expose"),
-		LOCTEXT("ExposeBindingTooltip", "Specifies options for exposing this binding to external systems as a persistent name."),
-		FNewMenuDelegate::CreateSP(this, &FSequencerObjectBindingNode::AddExposeMenu)
-	);
+	FSequencer* Sequencer = &GetSequencer();
 
-	if (GetSequencer().IsLevelEditorSequencer())
+	if (Sequencer->IsLevelEditorSequencer())
 	{
-		FSequencer*  Sequencer  = &GetSequencer();
 		UMovieScene* MovieScene = Sequencer->GetFocusedMovieSceneSequence()->GetMovieScene();
 		FMovieSceneSpawnable* Spawnable = MovieScene->FindSpawnable(ObjectBinding);
 
@@ -291,6 +287,12 @@ void FSequencerObjectBindingNode::BuildContextMenu(FMenuBuilder& MenuBuilder)
 			LOCTEXT("LabelsSubMenuTip", "Add or remove labels on this track"),
 			FNewMenuDelegate::CreateSP(this, &FSequencerObjectBindingNode::HandleLabelsSubMenuCreate)
 		);
+
+		MenuBuilder.AddSubMenu(
+			LOCTEXT("TagsLabel", "Tags"),
+			LOCTEXT("TagsTooltip", "Show this object binding's tags"),
+			FNewMenuDelegate::CreateSP(this, &FSequencerObjectBindingNode::AddTagMenu)
+		);
 	}
 	MenuBuilder.EndSection();
 
@@ -465,9 +467,162 @@ void FSequencerObjectBindingNode::AddAssignActorMenu(FMenuBuilder& MenuBuilder)
 	GetSequencer().AssignActor(MenuBuilder, ObjectBinding);
 }
 
-void FSequencerObjectBindingNode::AddExposeMenu(FMenuBuilder& MenuBuilder)
+
+void FSequencerObjectBindingNode::AddTagMenu(FMenuBuilder& MenuBuilder)
+{
+	MenuBuilder.AddMenuEntry(FSequencerCommands::Get().OpenTaggedBindingManager);
+
+	FSequencer* Sequencer = &GetSequencer();
+
+	UMovieSceneSequence* Sequence   = Sequencer->GetRootMovieSceneSequence();
+	UMovieScene*         MovieScene = Sequence->GetMovieScene();
+
+	MenuBuilder.BeginSection(NAME_None, LOCTEXT("ObjectTagsHeader", "Object Tags"));
+	{
+		TSet<FName> AllTags;
+
+		// Gather all the tags on all currently selected object binding IDs
+		FMovieSceneSequenceID SequenceID = Sequencer->GetFocusedTemplateID();
+		for (const TSharedRef<FSequencerDisplayNode>& Node : GetSequencer().GetSelection().GetSelectedOutlinerNodes())
+		{
+			if (Node->GetType() == ESequencerNode::Object)
+			{
+				const FGuid& ObjectID = StaticCastSharedRef<FSequencerObjectBindingNode>(Node)->GetObjectBinding();
+
+				FMovieSceneObjectBindingID BindingID(ObjectID, SequenceID);
+				for (auto It = Sequencer->GetObjectBindingTagCache()->IterateTags(BindingID); It; ++It)
+				{
+					AllTags.Add(It.Value());
+				}
+			}
+		}
+
+		bool bIsReadOnly = MovieScene->IsReadOnly();
+		for (const FName& TagName : AllTags)
+		{
+			MenuBuilder.AddMenuEntry(
+				FText::FromName(TagName),
+				FText(),
+				FSlateIcon(),
+				FUIAction(
+					FExecuteAction::CreateSP(this, &FSequencerObjectBindingNode::ToggleTag, TagName),
+					FCanExecuteAction::CreateLambda([bIsReadOnly] { return bIsReadOnly == false; }),
+					FGetActionCheckState::CreateSP(this, &FSequencerObjectBindingNode::GetTagCheckState, TagName)
+				),
+				NAME_None,
+				EUserInterfaceActionType::ToggleButton
+			);
+		}
+	}
+	MenuBuilder.EndSection();
+
+	MenuBuilder.BeginSection(NAME_None, LOCTEXT("AddNewHeader", "Add Tag"));
+	{
+		if (!MovieScene->IsReadOnly())
+		{
+			TSharedRef<SWidget> Widget =
+				SNew(SObjectBindingTag)
+				.OnCreateNew(this, &FSequencerObjectBindingNode::HandleAddTag);
+
+			MenuBuilder.AddWidget(Widget, FText());
+		}
+	}
+	MenuBuilder.EndSection();
+}
+
+ECheckBoxState FSequencerObjectBindingNode::GetTagCheckState(FName TagName)
+{
+	ECheckBoxState CheckBoxState = ECheckBoxState::Undetermined;
+
+	FSequencer&           Sequencer  = GetSequencer();
+	UMovieScene*          MovieScene = Sequencer.GetRootMovieSceneSequence()->GetMovieScene();
+	FMovieSceneSequenceID SequenceID = Sequencer.GetFocusedTemplateID();
+
+	for (const TSharedRef<FSequencerDisplayNode>& Node : Sequencer.GetSelection().GetSelectedOutlinerNodes())
+	{
+		if (Node->GetType() == ESequencerNode::Object)
+		{
+			const FGuid& ObjectID = StaticCastSharedRef<FSequencerObjectBindingNode>(Node)->GetObjectBinding();
+
+			FMovieSceneObjectBindingID BindingID(ObjectID, SequenceID);
+			ECheckBoxState ThisCheckState = Sequencer.GetObjectBindingTagCache()->HasTag(BindingID, TagName)
+				? ECheckBoxState::Checked
+				: ECheckBoxState::Unchecked;
+
+			if (CheckBoxState == ECheckBoxState::Undetermined)
+			{
+				CheckBoxState = ThisCheckState;
+			}
+			else if (CheckBoxState != ThisCheckState)
+			{
+				return ECheckBoxState::Undetermined;
+			}
+		}
+	}
+
+	return CheckBoxState;
+}
+
+void FSequencerObjectBindingNode::ToggleTag(FName TagName)
 {
-	MenuBuilder.AddWidget(SNew(SExposeBindingWidget, GetSequencer().AsShared()), FText(), true, false);
+	FSequencer&           Sequencer  = GetSequencer();
+	UMovieScene*          MovieScene = Sequencer.GetRootMovieSceneSequence()->GetMovieScene();
+	FMovieSceneSequenceID SequenceID = Sequencer.GetFocusedTemplateID();
+
+	for (const TSharedRef<FSequencerDisplayNode>& Node : Sequencer.GetSelection().GetSelectedOutlinerNodes())
+	{
+		if (Node->GetType() == ESequencerNode::Object)
+		{
+			const FGuid& ObjectID = StaticCastSharedRef<FSequencerObjectBindingNode>(Node)->GetObjectBinding();
+
+			FMovieSceneObjectBindingID BindingID(ObjectID, SequenceID);
+			if (!Sequencer.GetObjectBindingTagCache()->HasTag(BindingID, TagName))
+			{
+				HandleAddTag(TagName);
+				return;
+			}
+		}
+	}
+
+	HandleDeleteTag(TagName);
+}
+
+void FSequencerObjectBindingNode::HandleDeleteTag(FName TagName)
+{
+	FScopedTransaction Transaction(FText::Format(LOCTEXT("RemoveBindingTag", "Remove tag '{0}' from binding(s)"), FText::FromName(TagName)));
+
+	UMovieScene* MovieScene = GetSequencer().GetRootMovieSceneSequence()->GetMovieScene();
+	MovieScene->Modify();
+
+	FMovieSceneSequenceID SequenceID = GetSequencer().GetFocusedTemplateID();
+	for (const TSharedRef<FSequencerDisplayNode>& Node : GetSequencer().GetSelection().GetSelectedOutlinerNodes())
+	{
+		if (Node->GetType() == ESequencerNode::Object)
+		{
+			const FGuid& ObjectID = StaticCastSharedRef<FSequencerObjectBindingNode>(Node)->GetObjectBinding();
+
+			MovieScene->UntagBinding(TagName, FMovieSceneObjectBindingID(ObjectID, SequenceID));
+		}
+	}
+}
+
+void FSequencerObjectBindingNode::HandleAddTag(FName TagName)
+{
+	FScopedTransaction Transaction(FText::Format(LOCTEXT("CreateBindingTag", "Add new tag {0} to binding(s)"), FText::FromName(TagName)));
+
+	UMovieScene* MovieScene = GetSequencer().GetRootMovieSceneSequence()->GetMovieScene();
+	MovieScene->Modify();
+
+	FMovieSceneSequenceID SequenceID = GetSequencer().GetFocusedTemplateID();
+	for (const TSharedRef<FSequencerDisplayNode>& Node : GetSequencer().GetSelection().GetSelectedOutlinerNodes())
+	{
+		if (Node->GetType() == ESequencerNode::Object)
+		{
+			const FGuid& ObjectID = StaticCastSharedRef<FSequencerObjectBindingNode>(Node)->GetObjectBinding();
+
+			MovieScene->TagBinding(TagName, FMovieSceneObjectBindingID(ObjectID, SequenceID));
+		}
+	}
 }
 
 bool FSequencerObjectBindingNode::CanRenameNode() const
@@ -479,6 +634,7 @@ TSharedRef<SWidget> FSequencerObjectBindingNode::GetCustomOutlinerContent()
 {
 	// Create a container edit box
 	TSharedRef<SHorizontalBox> BoxPanel = SNew(SHorizontalBox)
+
 		+ SHorizontalBox::Slot()
 		[
 			SNew(SSpacer)
@@ -500,6 +656,14 @@ TSharedRef<SWidget> FSequencerObjectBindingNode::GetCustomOutlinerContent()
 	return BoxPanel;
 }
 
+TSharedPtr<SWidget> FSequencerObjectBindingNode::GetAdditionalOutlinerLabel()
+{
+	FSequencer& Sequencer = GetSequencer();
+	FMovieSceneObjectBindingID BindingID(ObjectBinding, Sequencer.GetFocusedTemplateID());
+
+	return SNew(SObjectBindingTags, BindingID, Sequencer.GetObjectBindingTagCache());
+}
+
 
 FText FSequencerObjectBindingNode::GetDisplayName() const
 {
diff --git a/Engine/Source/Editor/Sequencer/Private/DisplayNodes/SequencerObjectBindingNode.h b/Engine/Source/Editor/Sequencer/Private/DisplayNodes/SequencerObjectBindingNode.h
--- a/Engine/Source/Editor/Sequencer/Private/DisplayNodes/SequencerObjectBindingNode.h
+++ b/Engine/Source/Editor/Sequencer/Private/DisplayNodes/SequencerObjectBindingNode.h
@@ -12,6 +12,8 @@ class FSequencerTrackNode;
 class FMenuBuilder;
 struct FSlateBrush;
 
+enum class ECheckBoxState : uint8;
+
 /** Enumeration specifying what kind of object binding this is */
 enum class EObjectBindingType
 {
@@ -56,6 +58,7 @@ public:
 	virtual void BuildContextMenu(FMenuBuilder& MenuBuilder) override;
 	virtual bool CanRenameNode() const override;
 	virtual TSharedRef<SWidget> GetCustomOutlinerContent() override;
+	virtual TSharedPtr<SWidget> GetAdditionalOutlinerLabel() override;
 	virtual FText GetDisplayName() const override;
 	virtual FLinearColor GetDisplayNameColor() const override;
 	virtual FText GetDisplayNameToolTipText() const override;
@@ -80,7 +83,7 @@ protected:
 	void AddSpawnOwnershipMenu(FMenuBuilder& MenuBuilder);
 	void AddSpawnLevelMenu(FMenuBuilder& MenuBuilder);
 	void AddAssignActorMenu(FMenuBuilder& MenuBuilder);
-	void AddExposeMenu(FMenuBuilder& MenuBuilder);
+	void AddTagMenu(FMenuBuilder& MenuBuilder);
 
 	/** Get class for object binding */
 	const UClass* GetClassForObjectBinding() const;
@@ -95,6 +98,14 @@ private:
 	
 	void HandlePropertyMenuItemExecute(FPropertyPath PropertyPath);
 
+	ECheckBoxState GetTagCheckState(FName TagName);
+
+	void ToggleTag(FName TagName);
+
+	void HandleDeleteTag(FName TagName);
+
+	void HandleAddTag(FName TagName);
+
 private:
 
 	/** The binding to live objects */
diff --git a/Engine/Source/Editor/Sequencer/Private/SAnimationOutlinerTreeNode.cpp b/Engine/Source/Editor/Sequencer/Private/SAnimationOutlinerTreeNode.cpp
--- a/Engine/Source/Editor/Sequencer/Private/SAnimationOutlinerTreeNode.cpp
+++ b/Engine/Source/Editor/Sequencer/Private/SAnimationOutlinerTreeNode.cpp
@@ -19,6 +19,7 @@
 #include "SSequencerTreeView.h"
 #include "Widgets/Colors/SColorPicker.h"
 #include "SequencerSectionPainter.h"
+#include "ObjectBindingTagCache.h"
 #include "Widgets/Text/SInlineEditableTextBlock.h"
 #include "Framework/SlateDelegates.h"
 
@@ -73,6 +74,26 @@ void SAnimationOutlinerTreeNode::Construct( const FArguments& InArgs, TSharedRef
 		.Clipping(EWidgetClipping::ClipToBounds)
 		.IsSelected(FIsSelected::CreateSP(InTableRow, &SSequencerTreeViewRow::IsSelectedExclusively));
 
+	TSharedRef<SWidget> LabelContent = EditableLabel.ToSharedRef();
+
+	if (TSharedPtr<SWidget> AdditionalLabelContent = Node->GetAdditionalOutlinerLabel())
+	{
+		LabelContent = SNew(SHorizontalBox)
+
+		+ SHorizontalBox::Slot()
+		.AutoWidth()
+		.Padding(FMargin(0.f, 0.f, 5.f, 0.f))
+		[
+			LabelContent
+		]
+
+		+ SHorizontalBox::Slot()
+		[
+			AdditionalLabelContent.ToSharedRef()
+		];
+
+		LabelContent->SetClipping(EWidgetClipping::ClipToBounds);
+	}
 
 	Node->OnRenameRequested().AddRaw(this, &SAnimationOutlinerTreeNode::EnterRenameMode);
 
@@ -152,7 +173,7 @@ void SAnimationOutlinerTreeNode::Construct( const FArguments& InArgs, TSharedRef
 							.VAlign(VAlign_Center)
 							.Padding(FMargin(0.f, 0.f, 4.f, 0.f))
 							[
-								EditableLabel.ToSharedRef()
+								LabelContent
 							]
 
 							// Arbitrary customization slot
diff --git a/Engine/Source/Editor/Sequencer/Private/SExposeBindingWidget.cpp b/Engine/Source/Editor/Sequencer/Private/SExposeBindingWidget.cpp
deleted file mode 100644
diff --git a/Engine/Source/Editor/Sequencer/Private/SExposedBindingsWidget.cpp b/Engine/Source/Editor/Sequencer/Private/SExposedBindingsWidget.cpp
deleted file mode 100644
diff --git a/Engine/Source/Editor/Sequencer/Private/SExposedBindingsWidget.h b/Engine/Source/Editor/Sequencer/Private/SExposedBindingsWidget.h
deleted file mode 100644
diff --git a/Engine/Source/Editor/Sequencer/Private/SObjectBindingTagManager.cpp b/Engine/Source/Editor/Sequencer/Private/SObjectBindingTagManager.cpp
new file mode 100644
--- /dev/null
+++ b/Engine/Source/Editor/Sequencer/Private/SObjectBindingTagManager.cpp
@@ -0,0 +1,474 @@
+// Copyright 1998-2019 Epic Games, Inc. All Rights Reserved.
+
+#include "SObjectBindingTagManager.h"
+
+#include "Sequencer.h"
+#include "MovieSceneSequence.h"
+#include "MovieScene.h"
+#include "SequenceBindingTree.h"
+#include "MovieSceneObjectBindingID.h"
+#include "MovieSceneObjectBindingIDPicker.h"
+#include "ObjectBindingTagCache.h"
+#include "SObjectBindingTag.h"
+
+#include "SlateOptMacros.h"
+#include "Widgets/SNullWidget.h"
+#include "Widgets/SBoxPanel.h"
+#include "Widgets/Layout/SBox.h"
+#include "Widgets/Layout/SScrollBorder.h"
+#include "Widgets/Images/SImage.h"
+#include "Widgets/SOverlay.h"
+#include "Widgets/Text/STextBlock.h"
+#include "Widgets/Input/SEditableTextBox.h"
+#include "Widgets/Input/SButton.h"
+#include "Widgets/Input/SComboButton.h"
+#include "Widgets/Views/SHeaderRow.h"
+#include "Widgets/Notifications/SNotificationList.h"
+#include "Framework/MultiBox/MultiBoxBuilder.h"
+#include "Framework/Notifications/NotificationManager.h"
+
+#include "EditorStyleSet.h"
+#include "EditorFontGlyphs.h"
+#include "ScopedTransaction.h"
+
+#include "Algo/Sort.h"
+
+#define LOCTEXT_NAMESPACE "SObjectBindingTagManager"
+
+class SSequenceBindingNodeRow : public SMultiColumnTableRow<TSharedRef<FSequenceBindingNode>>
+{
+	SLATE_BEGIN_ARGS(SSequenceBindingNodeRow) {}
+	SLATE_END_ARGS()
+
+	static FName Column_Label;
+	static FName Column_Tags;
+
+public:
+
+	void Construct(const FArguments& InArgs, const TSharedRef<STableViewBase>& InOwnerTableView, TWeakPtr<FSequenceBindingNode> InBindingNode, TWeakPtr<SObjectBindingTagManager> InWeakTabManager)
+	{
+		WeakBindingNode = InBindingNode;
+		WeakTabManager  = InWeakTabManager;
+
+		SMultiColumnTableRow<TSharedRef<FSequenceBindingNode>>
+			::Construct(FSuperRowType::FArguments().Padding(5.f), InOwnerTableView);
+	}
+
+public:
+
+	virtual TSharedRef<SWidget> GenerateWidgetForColumn(const FName& ColumnName) override
+	{
+		TSharedPtr<FSequenceBindingNode> BindingNode = WeakBindingNode.Pin();
+		if (!BindingNode)
+		{
+			return SNullWidget::NullWidget;
+		}
+
+		if (ColumnName == Column_Label)
+		{
+			return GenerateLabelContent(BindingNode);
+		}
+		else if (ColumnName == Column_Tags)
+		{
+			TSharedRef<SWidget> TagContent = SNullWidget::NullWidget;
+			if (BindingNode->BindingID.IsValid())
+			{
+				TagContent = GenerateTagsContent(BindingNode);
+			}
+
+			return SNew(SBox)
+				.HeightOverride(24.f)
+				[
+					TagContent
+				];
+		}
+
+		return SNullWidget::NullWidget;
+	}
+
+	TSharedRef<SWidget> GenerateTagsContent(TSharedPtr<FSequenceBindingNode> BindingNode)
+	{
+		TSharedPtr<FSequencer> Sequencer = WeakTabManager.Pin()->GetSequencer();
+		if (!Sequencer)
+		{
+			return SNullWidget::NullWidget;
+		}
+
+		FObjectBindingTagCache* BindingCache = Sequencer->GetObjectBindingTagCache();
+
+		TArray<FName, TInlineAllocator<4>> AllTags;
+
+		for (auto It = BindingCache->IterateTags(BindingNode->BindingID); It; ++It)
+		{
+			AllTags.Add(It.Value());
+		}
+
+		if (AllTags.Num() == 0)
+		{
+			return SNullWidget::NullWidget;
+		}
+
+		Algo::Sort(AllTags, [](const FName& A, const FName& B) { return A.Compare(B) < 0; });
+
+		TSharedRef<SHorizontalBox> PillBox = SNew(SHorizontalBox);
+
+		FMargin Padding(0.f, 0.f, 4.f, 0.f);
+
+		bool bReadOnly = Sequencer->GetRootMovieSceneSequence()->GetMovieScene()->IsReadOnly();
+		for (FName TagName : AllTags)
+		{
+			FSimpleDelegate OnDeleted;
+			if (!bReadOnly)
+			{
+				OnDeleted = FSimpleDelegate::CreateSP(this, &SSequenceBindingNodeRow::HandleDeleteClicked, TagName);
+			}
+
+			PillBox->AddSlot()
+			.AutoWidth()
+			.Padding(Padding)
+			[
+				SNew(SObjectBindingTag)
+				.ColorTint(BindingCache->GetTagColor(TagName))
+				.Text(FText::FromName(TagName))
+				.OnDeleted(OnDeleted)
+			];
+
+			Padding = FMargin(4.f, 0.f);
+		}
+
+		return PillBox;
+	}
+
+	BEGIN_SLATE_FUNCTION_BUILD_OPTIMIZATION
+	TSharedRef<SWidget> GenerateLabelContent(TSharedPtr<FSequenceBindingNode> BindingNode)
+	{
+		return SNew(SHorizontalBox)
+
+		+ SHorizontalBox::Slot()
+		.VAlign(VAlign_Center)
+		.AutoWidth()
+		[
+			SNew(SExpanderArrow, SharedThis(this))
+		]
+
+		+ SHorizontalBox::Slot()
+		.VAlign(VAlign_Center)
+		.Padding(FMargin(5.f, 0.f, 5.f, 0.f))
+		.AutoWidth()
+		[
+			SNew(SOverlay)
+
+			+ SOverlay::Slot()
+			[
+				SNew(SImage)
+				.Image(BindingNode->Icon.GetIcon())
+			]
+
+			+ SOverlay::Slot()
+			.VAlign(VAlign_Top)
+			.HAlign(HAlign_Right)
+			[
+				SNew(SImage)
+				.Visibility(BindingNode->bIsSpawnable ? EVisibility::Visible : EVisibility::Collapsed)
+				.Image(FEditorStyle::GetBrush("Sequencer.SpawnableIconOverlay"))
+			]
+		]
+
+		+ SHorizontalBox::Slot()
+		.VAlign(VAlign_Center)
+		[
+			SNew(STextBlock)
+			.Text(BindingNode->DisplayString)
+		];
+	}
+	END_SLATE_FUNCTION_BUILD_OPTIMIZATION
+
+	void HandleDeleteClicked(FName TagName)
+	{
+		TSharedPtr<SObjectBindingTagManager> TabManager  = WeakTabManager.Pin();
+		if (TabManager)
+		{
+			TabManager->UntagSelection(TagName, WeakBindingNode.Pin());
+		}
+	}
+
+private:
+
+	TWeakPtr<FSequencer> WeakSequencer;
+	TWeakPtr<SObjectBindingTagManager> WeakTabManager;
+	TWeakPtr<FSequenceBindingNode> WeakBindingNode;
+};
+
+FName SSequenceBindingNodeRow::Column_Label("Label");
+FName SSequenceBindingNodeRow::Column_Tags("Tags");
+
+void SObjectBindingTagManager::Construct(const FArguments& InArgs, TWeakPtr<FSequencer> InWeakSequencer)
+{
+	WeakSequencer = InWeakSequencer;
+
+	TSharedPtr<FSequencer> Sequencer = WeakSequencer.Pin();
+	UMovieSceneSequence*   Sequence = Sequencer ? Sequencer->GetRootMovieSceneSequence() : nullptr;
+	if (!ensure(Sequence))
+	{
+		return;
+	}
+
+	BindingTree = MakeShared<FSequenceBindingTree>();
+	BindingTree->ForceRebuild(Sequence, Sequence, MovieSceneSequenceID::Root);
+
+	TWeakPtr<SObjectBindingTagManager> WeakTabManager = SharedThis(this);
+	auto HandleGenerateRow = [WeakTabManager](TSharedRef<FSequenceBindingNode> InNode, const TSharedRef<STableViewBase>& InOwnerTableView) -> TSharedRef<ITableRow>
+	{
+		return SNew(SSequenceBindingNodeRow, InOwnerTableView, InNode, WeakTabManager);
+	};
+
+	auto HandleGetChildren = [](TSharedRef<FSequenceBindingNode> InParent, TArray<TSharedRef<FSequenceBindingNode>>& OutChildren)
+	{
+		OutChildren.Append(InParent->Children);
+	};
+
+	TSharedRef<SHeaderRow> HeaderRow = SNew(SHeaderRow)
+		+ SHeaderRow::Column(SSequenceBindingNodeRow::Column_Label)
+		.DefaultLabel(LOCTEXT("HeaderLabel_Text", "Object Binding"))
+		.VAlignHeader(VAlign_Center)
+		.FillWidth(0.4)
+
+		+ SHeaderRow::Column(SSequenceBindingNodeRow::Column_Tags)
+		.DefaultLabel(LOCTEXT("HeaderTags_Text", "Tags"))
+		.VAlignHeader(VAlign_Center)
+		.FillWidth(0.6);
+
+	TreeView = SNew(STreeView<TSharedRef<FSequenceBindingNode>>)
+		.OnGenerateRow_Lambda(HandleGenerateRow)
+		.OnGetChildren_Lambda(HandleGetChildren)
+		.TreeItemsSource(&BindingTree->GetRootNode()->Children)
+		.OnContextMenuOpening( this, &SObjectBindingTagManager::OnContextMenuOpening)
+		.HeaderRow(HeaderRow);
+
+	ChildSlot
+	[
+		SNew(SBorder)
+		.BorderImage(FEditorStyle::GetBrush("ToolPanel.GroupBorder"))
+		[
+			SNew(SVerticalBox)
+
+			+ SVerticalBox::Slot()
+			[
+				SNew(SScrollBorder, TreeView.ToSharedRef())
+				[
+					TreeView.ToSharedRef()
+				]
+			]
+
+			+ SVerticalBox::Slot()
+			.AutoHeight()
+			.Padding(FMargin(10.f))
+			[
+				SAssignNew(Tags, SHorizontalBox)
+			]
+		]
+	];
+
+	FObjectBindingTagCache* BindingCache = Sequencer->GetObjectBindingTagCache();
+	BindingCache->OnUpdatedEvent.AddSP(this, &SObjectBindingTagManager::OnBindingCacheUpdated);
+	OnBindingCacheUpdated(BindingCache);
+
+	ExpandAllItems();
+}
+
+
+void SObjectBindingTagManager::OnBindingCacheUpdated(const FObjectBindingTagCache* BindingCache)
+{
+	Tags->ClearChildren();
+
+	// BindingCache is owned by FSequencer so it must still be valid
+	UMovieSceneSequence* Sequence = WeakSequencer.Pin()->GetRootMovieSceneSequence();
+	if (!Sequence)
+	{
+		return;
+	}
+
+	float LeftPadding = 0.f;
+	for (const TTuple<FName, FMovieSceneObjectBindingIDs>& Pair : Sequence->GetMovieScene()->AllTaggedBindings())
+	{
+		FName TagName = Pair.Key;
+
+		Tags->AddSlot()
+		.Padding(FMargin(LeftPadding, 0.f, 0.f, 0.f))
+		.AutoWidth()
+		[
+			SNew(SObjectBindingTag)
+			.OnDeleted(this, &SObjectBindingTagManager::RemoveTag, TagName)
+			.ColorTint(BindingCache->GetTagColor(TagName))
+			.Text(FText::FromName(TagName))
+		];
+
+		LeftPadding = 5.f;
+	}
+}
+
+void SObjectBindingTagManager::ExpandAllItems()
+{
+	struct FInternal
+	{
+		static void ExpandAll(TSharedRef<FSequenceBindingNode> Node, STreeView<TSharedRef<FSequenceBindingNode>>* InTreeView)
+		{
+			InTreeView->SetItemExpansion(Node, true);
+			for (TSharedRef<FSequenceBindingNode> Child : Node->Children)
+			{
+				ExpandAll(Child, InTreeView);
+			}
+		}
+	};
+	FInternal::ExpandAll(BindingTree->GetRootNode(), TreeView.Get());
+}
+
+void SObjectBindingTagManager::Tick(const FGeometry& AllottedGeometry, const double InCurrentTime, const float InDeltaTime)
+{
+	TSharedPtr<ISequencer> Sequencer  = WeakSequencer.Pin();
+	UMovieSceneSequence*   Sequence   = Sequencer ? Sequencer->GetRootMovieSceneSequence() : nullptr;
+	if (Sequence)
+	{
+		const bool bRebuilt = BindingTree->ConditionalRebuild(Sequence, Sequence, MovieSceneSequenceID::Root);
+		if (bRebuilt)
+		{
+			ExpandAllItems();
+			TreeView->SetTreeItemsSource(&BindingTree->GetRootNode()->Children);
+			TreeView->RequestTreeRefresh();
+		}
+	}
+}
+
+TSharedPtr<SWidget> SObjectBindingTagManager::OnContextMenuOpening()
+{
+	FMenuBuilder MenuBuilder(true, nullptr);
+
+	MenuBuilder.BeginSection(NAME_None, LOCTEXT("AddNewHeader", "Add Tag"));
+	{
+		MenuBuilder.AddWidget(
+			SNew(SBox)
+			.Padding(FMargin(7.f, 2.f))
+			[
+				SNew(SObjectBindingTag)
+				.ToolTipText(LOCTEXT("AddNew_Tooltip", "Adds a new persistent tag to the selected bindings"))
+				.OnCreateNew(this, &SObjectBindingTagManager::TagSelectionAs)
+			]
+			, FText()
+			, true
+		);
+	}
+	MenuBuilder.EndSection();
+
+	return MenuBuilder.MakeWidget();
+}
+
+void SObjectBindingTagManager::TagSelectionAs(FName NewName)
+{
+	TSharedPtr<ISequencer> Sequencer  = WeakSequencer.Pin();
+	UMovieSceneSequence*   Sequence   = Sequencer ? Sequencer->GetRootMovieSceneSequence() : nullptr;
+	UMovieScene*           MovieScene = Sequence  ? Sequence->GetMovieScene()              : nullptr;
+
+	if (!MovieScene)
+	{
+		return;
+	}
+	else if (MovieScene->IsReadOnly())
+	{
+		FNotificationInfo NotificationInfo(LOCTEXT("SequenceLocked", "Unable to change bindings while the sequence is locked."));
+		NotificationInfo.ExpireDuration = 4.0f;
+		FSlateNotificationManager::Get().AddNotification(NotificationInfo);
+		return;
+	}
+
+	FScopedTransaction Transaction(FText::Format(LOCTEXT("TagBinding", "Tag binding(s) as {0}"), FText::FromName(NewName)));
+
+	bool bMadeAnyChanges = false;
+
+	MovieScene->Modify();
+	for (TSharedRef<FSequenceBindingNode> SelectedItem : TreeView->GetSelectedItems())
+	{
+		if (SelectedItem->BindingID.IsValid())
+		{
+			bMadeAnyChanges = true;
+			MovieScene->TagBinding(NewName, SelectedItem->BindingID);
+		}
+	}
+
+	if (!bMadeAnyChanges)
+	{
+		Transaction.Cancel();
+	}
+}
+
+void SObjectBindingTagManager::UntagSelection(FName TagName, TSharedPtr<FSequenceBindingNode> InstigatorNode)
+{
+	TSharedPtr<ISequencer> Sequencer  = WeakSequencer.Pin();
+	UMovieSceneSequence*   Sequence   = Sequencer ? Sequencer->GetRootMovieSceneSequence() : nullptr;
+	UMovieScene*           MovieScene = Sequence  ? Sequence->GetMovieScene()              : nullptr;
+
+	if (!MovieScene)
+	{
+		return;
+	}
+	else if (MovieScene->IsReadOnly())
+	{
+		FNotificationInfo NotificationInfo(LOCTEXT("SequenceLocked", "Unable to change bindings while the sequence is locked."));
+		NotificationInfo.ExpireDuration = 4.0f;
+		FSlateNotificationManager::Get().AddNotification(NotificationInfo);
+		return;
+	}
+
+	FScopedTransaction Transaction(FText::Format(LOCTEXT("UntagBinding", "Untag {0} from binding(s)"), FText::FromName(TagName)));
+
+	bool bMadeAnyChanges = false;
+
+	MovieScene->Modify();
+	for (TSharedRef<FSequenceBindingNode> SelectedItem : TreeView->GetSelectedItems())
+	{
+		if (SelectedItem->BindingID.IsValid())
+		{
+			bMadeAnyChanges = true;
+			MovieScene->UntagBinding(TagName, SelectedItem->BindingID);
+		}
+	}
+	if (InstigatorNode && !TreeView->IsItemSelected(InstigatorNode.ToSharedRef()))
+	{
+		if (InstigatorNode->BindingID.IsValid())
+		{
+			bMadeAnyChanges = true;
+			MovieScene->UntagBinding(TagName, InstigatorNode->BindingID);
+		}
+	}
+
+	if (!bMadeAnyChanges)
+	{
+		Transaction.Cancel();
+	}
+}
+
+void SObjectBindingTagManager::RemoveTag(FName TagName)
+{
+	TSharedPtr<ISequencer> Sequencer  = WeakSequencer.Pin();
+	UMovieSceneSequence*   Sequence   = Sequencer ? Sequencer->GetRootMovieSceneSequence() : nullptr;
+	UMovieScene*           MovieScene = Sequence  ? Sequence->GetMovieScene()              : nullptr;
+
+	if (!MovieScene)
+	{
+		return;
+	}
+	else if (MovieScene->IsReadOnly())
+	{
+		FNotificationInfo NotificationInfo(LOCTEXT("SequenceLocked", "Unable to change bindings while the sequence is locked."));
+		NotificationInfo.ExpireDuration = 4.0f;
+		FSlateNotificationManager::Get().AddNotification(NotificationInfo);
+	}
+	else
+	{
+		FScopedTransaction Transaction(FText::Format(LOCTEXT("RemoveTag", "Remove tag {0}"), FText::FromName(TagName)));
+
+		MovieScene->Modify();
+		MovieScene->RemoveTag(TagName);
+	}
+}
+
+#undef LOCTEXT_NAMESPACE
\ No newline at end of file
diff --git a/Engine/Source/Editor/Sequencer/Private/SObjectBindingTagManager.h b/Engine/Source/Editor/Sequencer/Private/SObjectBindingTagManager.h
new file mode 100644
--- /dev/null
+++ b/Engine/Source/Editor/Sequencer/Private/SObjectBindingTagManager.h
@@ -0,0 +1,57 @@
+// Copyright 1998-2019 Epic Games, Inc. All Rights Reserved.
+
+#pragma once
+
+#include "Widgets/DeclarativeSyntaxSupport.h"
+#include "Widgets/SCompoundWidget.h"
+#include "Widgets/Views/STreeView.h"
+#include "MovieSceneObjectBindingID.h"
+
+struct FSequenceBindingNode;
+struct FSequenceBindingTree;
+
+class FSequencer;
+class FMenuBuilder;
+class FObjectBindingTagCache;
+
+class SObjectBindingTagManager : public SCompoundWidget
+{
+public:
+
+	SLATE_BEGIN_ARGS(SObjectBindingTagManager){}
+	SLATE_END_ARGS()
+
+	void Construct(const FArguments& InArgs, TWeakPtr<FSequencer> InWeakSequencer);
+
+	TSharedPtr<FSequencer> GetSequencer() const
+	{
+		return WeakSequencer.Pin();
+	}
+
+	void TagSelectionAs(FName NewName);
+
+	void UntagSelection(FName TagName, TSharedPtr<FSequenceBindingNode> InstigatorNode);
+
+	void RemoveTag(FName TagName);
+
+private:
+
+	virtual void Tick(const FGeometry& AllottedGeometry, const double InCurrentTime, const float InDeltaTime) override;
+
+	void OnBindingCacheUpdated(const FObjectBindingTagCache* BindingCache);
+
+	void ExpandAllItems();
+
+	TSharedPtr<SWidget> OnContextMenuOpening();
+
+private:
+
+	/** The sequencer UI instance that is currently open */
+	TWeakPtr<FSequencer> WeakSequencer;
+
+	TSharedPtr<FSequenceBindingTree> BindingTree;
+
+	TSharedPtr<STreeView<TSharedRef<FSequenceBindingNode>>> TreeView;
+
+	TSharedPtr<SHorizontalBox> Tags;
+};
\ No newline at end of file
diff --git a/Engine/Source/Editor/Sequencer/Private/SSequencer.cpp b/Engine/Source/Editor/Sequencer/Private/SSequencer.cpp
--- a/Engine/Source/Editor/Sequencer/Private/SSequencer.cpp
+++ b/Engine/Source/Editor/Sequencer/Private/SSequencer.cpp
@@ -85,7 +85,7 @@
 #include "LevelSequence.h"
 #include "SequencerLog.h"
 #include "MovieSceneCopyableBinding.h"
-#include "SExposedBindingsWidget.h"
+#include "SObjectBindingTagManager.h"
 #include "MovieSceneCopyableTrack.h"
 #include "IPropertyRowGenerator.h"
 #include "Fonts/FontMeasure.h"
@@ -1044,6 +1044,11 @@ void SSequencer::BindCommands(TSharedRef<FUICommandList> SequencerCommandBinding
 		FSequencerCommands::Get().OpenDirectorBlueprint,
 		FExecuteAction::CreateLambda(OpenDirectorBlueprint)
 	);
+
+	SequencerCommandBindings->MapAction(
+		FSequencerCommands::Get().OpenTaggedBindingManager,
+		FExecuteAction::CreateSP(this, &SSequencer::OpenTaggedBindingManager)
+	);
 }
 
 void SSequencer::ShowTickResolutionOverlay()
@@ -1799,11 +1804,7 @@ TSharedRef<SWidget> SSequencer::MakeGeneralMenu()
 		MenuBuilder.AddMenuEntry(FSequencerCommands::Get().RebindPossessableReferences);
 	}
 
-	MenuBuilder.AddSubMenu(
-		LOCTEXT("ExposeBindingLabel", "Exposed Binding Groups"),
-		LOCTEXT("ExposeBindingTooltip", "Specifies options for exposing this binding to external systems as a persistent name."),
-		FNewMenuDelegate::CreateSP(this, &SSequencer::PopulateExposeBindingsMenu)
-	);
+	MenuBuilder.AddMenuEntry(FSequencerCommands::Get().OpenTaggedBindingManager);
 
 	MenuBuilder.EndSection();
 
@@ -1818,9 +1819,34 @@ TSharedRef<SWidget> SSequencer::MakeGeneralMenu()
 	return MenuBuilder.MakeWidget();
 }
 
-void SSequencer::PopulateExposeBindingsMenu(FMenuBuilder& InMenuBuilder)
+void SSequencer::OpenTaggedBindingManager()
 {
-	InMenuBuilder.AddWidget(SNew(SExposedBindingsWidget, SequencerPtr), FText(), true);
+	if (TSharedPtr<SWindow> Window = WeakExposedBindingsWindow.Pin())
+	{
+		Window->DrawAttention(FWindowDrawAttentionParameters());
+		return;
+	}
+
+	TSharedRef<SWindow> ExposedBindingsWindow = SNew(SWindow)
+		.Title(FText::Format(LOCTEXT("ExposedBindings_Title", "Bindings Exposed in {0}"), FText::FromName(SequencerPtr.Pin()->GetRootMovieSceneSequence()->GetFName())))
+		.SupportsMaximize(false)
+		.ClientSize(FVector2D(600.f, 500.f))
+		.Content()
+		[
+			SNew(SObjectBindingTagManager, SequencerPtr)
+		];
+
+	TSharedPtr<SWindow> ParentWindow = FSlateApplication::Get().FindWidgetWindow(AsShared());
+	if (ParentWindow)
+	{
+		FSlateApplication::Get().AddWindowAsNativeChild(ExposedBindingsWindow, ParentWindow.ToSharedRef());
+	}
+	else
+	{
+		FSlateApplication::Get().AddWindow(ExposedBindingsWindow);
+	}
+
+	WeakExposedBindingsWindow = ExposedBindingsWindow;
 }
 
 void SSequencer::FillPlaybackSpeedMenu(FMenuBuilder& InMenuBarBuilder)
@@ -2316,6 +2342,11 @@ SSequencer::~SSequencer()
 			}
 		}
 	}
+
+	if (TSharedPtr<SWindow> Window = WeakExposedBindingsWindow.Pin())
+	{
+		Window->DestroyWindowImmediately();
+	}
 }
 
 
diff --git a/Engine/Source/Editor/Sequencer/Private/SSequencer.h b/Engine/Source/Editor/Sequencer/Private/SSequencer.h
--- a/Engine/Source/Editor/Sequencer/Private/SSequencer.h
+++ b/Engine/Source/Editor/Sequencer/Private/SSequencer.h
@@ -33,6 +33,7 @@ class SSequencerStretchBox;
 class SSequencerTreeView;
 class SCurveEditorPanel;
 class SDockTab;
+class SWindow;
 class USequencerSettings;
 class FSequencerTrackFilter;
 struct FPaintPlaybackRangeArgs;
@@ -394,7 +395,7 @@ private:
 	/** Makes the key group menu for the toolbar. */
 	TSharedRef<SWidget> MakeKeyGroupMenu();
 
-	void PopulateExposeBindingsMenu(FMenuBuilder& InMenuBuilder);
+	void OpenTaggedBindingManager();
 
 	/** Makes the playback speed menu for the toolbar. */
 	void FillPlaybackSpeedMenu(FMenuBuilder& InMenuBuilder);
@@ -697,6 +698,8 @@ private:
 	/** All possible track filter objects */
 	TArray< TSharedRef<FSequencerTrackFilter> > AllTrackFilters;
 
+	TWeakPtr<SWindow> WeakExposedBindingsWindow;
+
 public:
 	static const FName CurveEditorTabName;
 
diff --git a/Engine/Source/Editor/Sequencer/Private/Sequencer.cpp b/Engine/Source/Editor/Sequencer/Private/Sequencer.cpp
--- a/Engine/Source/Editor/Sequencer/Private/Sequencer.cpp
+++ b/Engine/Source/Editor/Sequencer/Private/Sequencer.cpp
@@ -106,6 +106,7 @@
 #include "DesktopPlatformModule.h"
 #include "Factories.h"
 #include "FbxExporter.h"
+#include "ObjectBindingTagCache.h"
 #include "UnrealExporter.h"
 #include "ISequencerEditorObjectBinding.h"
 #include "LevelSequence.h"
@@ -319,6 +320,8 @@ void FSequencer::InitSequencer(const FSequencerInitParams& InitParams, const TSh
 	Settings->GetOnEvaluateSubSequencesInIsolationChanged().AddSP(this, &FSequencer::RestorePreAnimatedState);
 	Settings->GetOnShowSelectedNodesOnlyChanged().AddSP(this, &FSequencer::OnSelectedNodesOnlyChanged);
 
+	ObjectBindingTagCache = MakeUnique<FObjectBindingTagCache>();
+
 	FCurveEditorInitParams CurveEditorInitParams;
 	{
 	}
@@ -569,6 +572,8 @@ void FSequencer::Tick(float InDeltaTime)
 	// Ensure the time bases for our playback position are kept up to date with the root sequence
 	UpdateTimeBases();
 
+	ObjectBindingTagCache->ConditionalUpdate(RootSequence.Get());
+
 	Selection.Tick();
 
 	if (PlaybackContextAttribute.IsBound())
diff --git a/Engine/Source/Editor/Sequencer/Private/Sequencer.h b/Engine/Source/Editor/Sequencer/Private/Sequencer.h
--- a/Engine/Source/Editor/Sequencer/Private/Sequencer.h
+++ b/Engine/Source/Editor/Sequencer/Private/Sequencer.h
@@ -53,6 +53,7 @@ class IMenu;
 class FCurveEditor;
 class ISequencerEditTool;
 class FSequencerKeyCollection;
+class FObjectBindingTagCache;
 class ISequencerTrackEditor;
 class ISequencerEditorObjectBinding;
 class SSequencer;
@@ -189,6 +190,11 @@ public:
 		return NodeTree;
 	}
 
+	FObjectBindingTagCache* GetObjectBindingTagCache()
+	{
+		return ObjectBindingTagCache.Get();
+	}
+
 	bool IsPerspectiveViewportPossessionEnabled() const override
 	{
 		return bPerspectiveViewportPossessionEnabled;
@@ -1260,4 +1266,6 @@ private:
 
 	/** A signature that will suppress auto evaluation when it is the only change dirtying the template. */
 	TOptional<TTuple<TWeakObjectPtr<UMovieSceneSequence>, FGuid>> SuppressAutoEvalSignature;
+
+	TUniquePtr<FObjectBindingTagCache> ObjectBindingTagCache;
 };
diff --git a/Engine/Source/Editor/Sequencer/Private/SequencerCommands.cpp b/Engine/Source/Editor/Sequencer/Private/SequencerCommands.cpp
--- a/Engine/Source/Editor/Sequencer/Private/SequencerCommands.cpp
+++ b/Engine/Source/Editor/Sequencer/Private/SequencerCommands.cpp
@@ -95,6 +95,7 @@ void FSequencerCommands::RegisterCommands()
 	UI_COMMAND( ToggleShowTransformBox, "Transform Selection...", "Transform the selected keys and sections by a given amount", EUserInterfaceActionType::Button, FInputChord(EModifierKey::Control, EKeys::M));
 	UI_COMMAND( ToggleShowStretchBox, "Stretch/Shrink...", "Stretch or shrink a given amount, moving keys forwards/backwards as necessary", EUserInterfaceActionType::Button, FInputChord() );
 	UI_COMMAND( OpenDirectorBlueprint, "Open Director Blueprint", "Opens the director blueprint for this sequence.", EUserInterfaceActionType::Button, FInputChord() );
+	UI_COMMAND( OpenTaggedBindingManager, "Open Binding Tag Manager...", "Specifies options for tagging bindings within this sequence for external systems to reference as a persistent name.", EUserInterfaceActionType::Button, FInputChord() );
 	UI_COMMAND( BakeTransform, "Bake Transform", "Bake transform in world space, removing any existing transform and attach tracks", EUserInterfaceActionType::Button, FInputChord() );
 	UI_COMMAND( SyncSectionsUsingSourceTimecode, "Sync Sections using Source Timecode", "Synchronize sections to the first selected section using source timecode", EUserInterfaceActionType::Button, FInputChord() );
 	UI_COMMAND( ToggleShowRangeSlider, "Show Range Slider", "Enables and disables showing the time range slider", EUserInterfaceActionType::ToggleButton, FInputChord() );
diff --git a/Engine/Source/Editor/Sequencer/Public/SequencerCommands.h b/Engine/Source/Editor/Sequencer/Public/SequencerCommands.h
--- a/Engine/Source/Editor/Sequencer/Public/SequencerCommands.h
+++ b/Engine/Source/Editor/Sequencer/Public/SequencerCommands.h
@@ -231,6 +231,9 @@ public:
 
 	/** Opens the director blueprint for a sequence. */
 	TSharedPtr< FUICommandInfo > OpenDirectorBlueprint;
+
+	/** Opens the tagged binding manager. */
+	TSharedPtr< FUICommandInfo > OpenTaggedBindingManager;
 	
 	/** Sets the tree search widget as the focused widget in Slate for easy typing. */
 	TSharedPtr< FUICommandInfo > QuickTreeSearch;
diff --git a/Engine/Source/Runtime/LevelSequence/Private/LevelSequenceActor.cpp b/Engine/Source/Runtime/LevelSequence/Private/LevelSequenceActor.cpp
--- a/Engine/Source/Runtime/LevelSequence/Private/LevelSequenceActor.cpp
+++ b/Engine/Source/Runtime/LevelSequence/Private/LevelSequenceActor.cpp
@@ -19,6 +19,14 @@
 	#include "SceneOutlinerFilters.h"
 #endif
 
+bool GLevelSequenceActor_InvalidBindingTagWarnings = true;
+FAutoConsoleVariableRef CVarLevelSequenceActor_InvalidBindingTagWarnings(
+	TEXT("LevelSequence.InvalidBindingTagWarnings"),
+	GLevelSequenceActor_InvalidBindingTagWarnings,
+	TEXT("Whether to emit a warning when invalid object binding tags are used to override bindings or not.\n"),
+	ECVF_Default
+);
+
 ALevelSequenceActor::ALevelSequenceActor(const FObjectInitializer& Init)
 	: Super(Init)
 	, bShowBurnin(true)
@@ -355,6 +363,25 @@ void ALevelSequenceActor::SetBinding(FMovieSceneObjectBindingID Binding, const T
 	}
 }
 
+void ALevelSequenceActor::SetBindingByTag(FName BindingTag, const TArray<AActor*>& Actors, bool bAllowBindingsFromAsset)
+{
+	const UMovieSceneSequence*         Sequence = GetSequence();
+	const FMovieSceneObjectBindingIDs* Bindings = Sequence ? Sequence->GetMovieScene()->AllTaggedBindings().Find(BindingTag) : nullptr;
+	if (Bindings)
+	{
+		for (FMovieSceneObjectBindingID ID : Bindings->IDs)
+		{
+			SetBinding(ID, Actors, bAllowBindingsFromAsset);
+		}
+	}
+	else if (GLevelSequenceActor_InvalidBindingTagWarnings)
+	{
+		FMessageLog("PIE")
+			.Warning(FText::Format(NSLOCTEXT("LevelSequenceActor", "SetBindingByTag", "Sequence did not contain any bindings with the tag '{0}'"), FText::FromName(BindingTag)))
+			->AddToken(FUObjectToken::Create(this));
+	}
+}
+
 void ALevelSequenceActor::AddBinding(FMovieSceneObjectBindingID Binding, AActor* Actor, bool bAllowBindingsFromAsset)
 {
 	if (!Binding.IsValid())
@@ -373,6 +400,25 @@ void ALevelSequenceActor::AddBinding(FMovieSceneObjectBindingID Binding, AActor*
 	}
 }
 
+void ALevelSequenceActor::AddBindingByTag(FName BindingTag, AActor* Actor, bool bAllowBindingsFromAsset)
+{
+	const UMovieSceneSequence*         Sequence = GetSequence();
+	const FMovieSceneObjectBindingIDs* Bindings = Sequence ? Sequence->GetMovieScene()->AllTaggedBindings().Find(BindingTag) : nullptr;
+	if (Bindings)
+	{
+		for (FMovieSceneObjectBindingID ID : Bindings->IDs)
+		{
+			AddBinding(ID, Actor, bAllowBindingsFromAsset);
+		}
+	}
+	else if (GLevelSequenceActor_InvalidBindingTagWarnings)
+	{
+		FMessageLog("PIE")
+			.Warning(FText::Format(NSLOCTEXT("LevelSequenceActor", "AddBindingByTag", "Sequence did not contain any bindings with the tag '{0}'"), FText::FromName(BindingTag)))
+			->AddToken(FUObjectToken::Create(this));
+	}
+}
+
 void ALevelSequenceActor::RemoveBinding(FMovieSceneObjectBindingID Binding, AActor* Actor)
 {
 	if (!Binding.IsValid())
@@ -391,6 +437,25 @@ void ALevelSequenceActor::RemoveBinding(FMovieSceneObjectBindingID Binding, AAct
 	}
 }
 
+void ALevelSequenceActor::RemoveBindingByTag(FName BindingTag, AActor* Actor)
+{
+	const UMovieSceneSequence*         Sequence = GetSequence();
+	const FMovieSceneObjectBindingIDs* Bindings = Sequence ? Sequence->GetMovieScene()->AllTaggedBindings().Find(BindingTag) : nullptr;
+	if (Bindings)
+	{
+		for (FMovieSceneObjectBindingID ID : Bindings->IDs)
+		{
+			RemoveBinding(ID, Actor);
+		}
+	}
+	else if (GLevelSequenceActor_InvalidBindingTagWarnings)
+	{
+		FMessageLog("PIE")
+			.Warning(FText::Format(NSLOCTEXT("LevelSequenceActor", "RemoveBindingByTag", "Sequence did not contain any bindings with the tag '{0}'"), FText::FromName(BindingTag)))
+			->AddToken(FUObjectToken::Create(this));
+	}
+}
+
 void ALevelSequenceActor::ResetBinding(FMovieSceneObjectBindingID Binding)
 {
 	if (!Binding.IsValid())
@@ -422,7 +487,7 @@ FMovieSceneObjectBindingID ALevelSequenceActor::FindNamedBinding(FName InBinding
 {
 	if (ensureAlways(SequencePlayer))
 	{
-		return SequencePlayer->GetSequence()->FindNamedBinding(InBindingName);
+		return SequencePlayer->GetSequence()->FindBindingByTag(InBindingName);
 	}
 	return FMovieSceneObjectBindingID();
 }
@@ -431,7 +496,7 @@ const TArray<FMovieSceneObjectBindingID>& ALevelSequenceActor::FindNamedBindings
 {
 	if (ensureAlways(SequencePlayer))
 	{
-		return SequencePlayer->GetSequence()->FindNamedBindings(InBindingName);
+		return SequencePlayer->GetSequence()->FindBindingsByTag(InBindingName);
 	}
 
 	static TArray<FMovieSceneObjectBindingID> EmptyBindings;
diff --git a/Engine/Source/Runtime/LevelSequence/Public/LevelSequenceActor.h b/Engine/Source/Runtime/LevelSequence/Public/LevelSequenceActor.h
--- a/Engine/Source/Runtime/LevelSequence/Public/LevelSequenceActor.h
+++ b/Engine/Source/Runtime/LevelSequence/Public/LevelSequenceActor.h
@@ -183,6 +183,16 @@ public:
 	UFUNCTION(BlueprintCallable, Category="Game|Cinematic|Bindings")
 	void SetBinding(FMovieSceneObjectBindingID Binding, const TArray<AActor*>& Actors, bool bAllowBindingsFromAsset = false);
 
+	/**
+	 * Assigns an set of actors to all the bindings tagged with the specified name in this sequence. Object Bindings can be tagged within the sequence UI by RMB -> Tags... on the object binding in the tree.
+	 *
+	 * @param BindingTag   The unique tag name to lookup bindings with
+	 * @param Actors       The actors to assign to all the tagged bindings
+	 * @param bAllowBindingsFromAsset Whether to continue to allow bindings from the level sequence asset (true) or not (false)
+	 */
+	UFUNCTION(BlueprintCallable, Category="Game|Cinematic|Bindings")
+	void SetBindingByTag(FName BindingTag, const TArray<AActor*>& Actors, bool bAllowBindingsFromAsset = false);
+
 	/**
 	 * Adds the specified actor to the overridden bindings for the specified binding ID, optionally still allowing the bindings defined in the Level Sequence asset
 	 *
@@ -193,12 +203,28 @@ public:
 	UFUNCTION(BlueprintCallable, Category="Game|Cinematic|Bindings")
 	void AddBinding(FMovieSceneObjectBindingID Binding, AActor* Actor, bool bAllowBindingsFromAsset = false);
 
+	/**
+	 * Binds an actor to all the bindings tagged with the specified name in this sequence. Does not remove any exising bindings that have been set up through this API. Object Bindings can be tagged within the sequence UI by RMB -> Tags... on the object binding in the tree.
+	 *
+	 * @param BindingTag   The unique tag name to lookup bindings with
+	 * @param Actor        The actor to assign to all the tagged bindings
+	 * @param bAllowBindingsFromAsset Whether to continue to allow bindings from the level sequence asset (true) or not (false)
+	 */
+	UFUNCTION(BlueprintCallable, Category="Game|Cinematic|Bindings")
+	void AddBindingByTag(FName BindingTag, AActor* Actor, bool bAllowBindingsFromAsset = false);
+
 	/**
 	 * Removes the specified actor from the specified binding's actor array
 	 */
 	UFUNCTION(BlueprintCallable, Category="Game|Cinematic|Bindings")
 	void RemoveBinding(FMovieSceneObjectBindingID Binding, AActor* Actor);
 
+	/**
+	 * Removes the specified actor from the specified binding's actor array
+	 */
+	UFUNCTION(BlueprintCallable, Category="Game|Cinematic|Bindings")
+	void RemoveBindingByTag(FName Tag, AActor* Actor);
+
 	/**
 	 * Resets the specified binding back to the defaults defined by the Level Sequence asset
 	 */
@@ -212,16 +238,19 @@ public:
 	void ResetBindings();
 
 	/**
-	 * Find the first object binding ID associated with the specified name (set up through RMB->Expose on Object bindings from within sequencer)
+	 * Retrieve the first object binding that has been tagged with the specified name
 	 */
-	UFUNCTION(BlueprintCallable, Category = "Game|Cinematic|Bindings")
-	FMovieSceneObjectBindingID FindNamedBinding(FName InBindingName) const;
+	UFUNCTION(BlueprintCallable, Category = "Game|Cinematic|Bindings", DisplayName="FindBindingByTag")
+	FMovieSceneObjectBindingID FindNamedBinding(FName Tag) const;
 
 	/**
-	 * Find all object binding IDs associated with the specified name (set up through RMB->Expose on Object bindings from within sequencer)
+	 * Retrieve all the bindings that have been tagged with the specified name
+	 *
+	 * @param Tag  The unique tag name to lookup bindings with. Object Bindings can be tagged within the sequence UI by RMB -> Tags... on the object binding in the tree.
+	 * @return An array containing all the bindings that are tagged with this name, potentially empty.
 	 */
-	UFUNCTION(BlueprintCallable, Category = "Game|Cinematic|Bindings")
-	const TArray<FMovieSceneObjectBindingID>& FindNamedBindings(FName InBindingName) const;
+	UFUNCTION(BlueprintCallable, Category = "Game|Cinematic|Bindings", DisplayName="FindBindingByTag")
+	const TArray<FMovieSceneObjectBindingID>& FindNamedBindings(FName Tag) const;
 
 protected:
 
diff --git a/Engine/Source/Runtime/MovieScene/Private/MovieScene.cpp b/Engine/Source/Runtime/MovieScene/Private/MovieScene.cpp
--- a/Engine/Source/Runtime/MovieScene/Private/MovieScene.cpp
+++ b/Engine/Source/Runtime/MovieScene/Private/MovieScene.cpp
@@ -410,27 +410,31 @@ FText UMovieScene::GetObjectDisplayName(const FGuid& ObjectId)
 }
 
 
-void UMovieScene::ExposeBinding(FName ExposeAs)
+void UMovieScene::AddNewBindingTag(const FName& NewTag)
 {
-	BindingGroups.Add(ExposeAs);
+	BindingGroups.Add(NewTag);
 }
 
-void UMovieScene::ExposeBinding(FName ExposeAs, FMovieSceneObjectBindingID BindingToExpose)
+void UMovieScene::TagBinding(const FName& NewTag, FMovieSceneObjectBindingID BindingToTag)
 {
-	BindingGroups.FindOrAdd(ExposeAs).IDs.AddUnique(BindingToExpose);
+	BindingGroups.FindOrAdd(NewTag).IDs.AddUnique(BindingToTag);
 }
 
-void UMovieScene::RemoveExposedBinding(FName ExposedAs, FMovieSceneObjectBindingID BindingToExpose)
+void UMovieScene::UntagBinding(const FName& Tag, FMovieSceneObjectBindingID Binding)
 {
-	if (FMovieSceneObjectBindingIDs* Array = BindingGroups.Find(ExposedAs))
+	if (FMovieSceneObjectBindingIDs* Array = BindingGroups.Find(Tag))
 	{
-		Array->IDs.Remove(BindingToExpose);
+		Array->IDs.Remove(Binding);
+		if (Array->IDs.Num() == 0)
+		{
+			BindingGroups.Remove(Tag);
+		}
 	}
 }
 
-void UMovieScene::RemoveExposedBinding(FName ExposedAs)
+void UMovieScene::RemoveTag(const FName& TagToRemove)
 {
-	BindingGroups.Remove(ExposedAs);
+	BindingGroups.Remove(TagToRemove);
 }
 
 #if WITH_EDITORONLY_DATA
diff --git a/Engine/Source/Runtime/MovieScene/Private/MovieSceneSequence.cpp b/Engine/Source/Runtime/MovieScene/Private/MovieSceneSequence.cpp
--- a/Engine/Source/Runtime/MovieScene/Private/MovieSceneSequence.cpp
+++ b/Engine/Source/Runtime/MovieScene/Private/MovieSceneSequence.cpp
@@ -96,9 +96,9 @@ FGuid UMovieSceneSequence::FindPossessableObjectId(UObject& Object, UObject* Con
 	return FGuid();
 }
 
-FMovieSceneObjectBindingID UMovieSceneSequence::FindNamedBinding(FName InBindingName) const
+FMovieSceneObjectBindingID UMovieSceneSequence::FindBindingByTag(FName InBindingName) const
 {
-	for (FMovieSceneObjectBindingID ID : FindNamedBindings(InBindingName))
+	for (FMovieSceneObjectBindingID ID : FindBindingsByTag(InBindingName))
 	{
 		return ID;
 	}
@@ -110,9 +110,9 @@ FMovieSceneObjectBindingID UMovieSceneSequence::FindNamedBinding(FName InBinding
 	return FMovieSceneObjectBindingID();
 }
 
-const TArray<FMovieSceneObjectBindingID>& UMovieSceneSequence::FindNamedBindings(FName InBindingName) const
+const TArray<FMovieSceneObjectBindingID>& UMovieSceneSequence::FindBindingsByTag(FName InBindingName) const
 {
-	const FMovieSceneObjectBindingIDs* BindingIDs = GetMovieScene()->AllBindingGroups().Find(InBindingName);
+	const FMovieSceneObjectBindingIDs* BindingIDs = GetMovieScene()->AllTaggedBindings().Find(InBindingName);
 	if (BindingIDs)
 	{
 		return BindingIDs->IDs;
diff --git a/Engine/Source/Runtime/MovieScene/Public/MovieScene.h b/Engine/Source/Runtime/MovieScene/Public/MovieScene.h
--- a/Engine/Source/Runtime/MovieScene/Public/MovieScene.h
+++ b/Engine/Source/Runtime/MovieScene/Public/MovieScene.h
@@ -908,9 +908,9 @@ public:
 	int32 FindNextMarkedFrame(FFrameNumber InFrameNumber, bool bForward);
 
 	/*
-	 * Retrieve all the exposed binding groups for this movie scene
+	 * Retrieve all the tagged binding groups for this movie scene
 	 */
-	const TMap<FName, FMovieSceneObjectBindingIDs>& AllBindingGroups() const
+	const TMap<FName, FMovieSceneObjectBindingIDs>& AllTaggedBindings() const
 	{
 		return BindingGroups;
 	}
@@ -918,22 +918,22 @@ public:
 	/*
 	 * Add a new binding group for the specified name
 	 */
-	void ExposeBinding(FName ExposeAs);
+	void AddNewBindingTag(const FName& NewTag);
 
 	/*
-	 * Expose the specified binding ID using the specified name
+	 * Tag the specified binding ID with the specified name
 	 */
-	void ExposeBinding(FName ExposeAs, FMovieSceneObjectBindingID BindingToExpose);
+	void TagBinding(const FName& NewTag, FMovieSceneObjectBindingID BindingToTag);
 
 	/*
-	 * Remove the specified binding ID from the specified expose group name
+	 * Remove a tag from the specified object binding
 	 */
-	void RemoveExposedBinding(FName ExposedAs, FMovieSceneObjectBindingID BindingToExpose);
+	void UntagBinding(const FName& Tag, FMovieSceneObjectBindingID Binding);
 
 	/*
-	 * Remove the entire group with the specified name
+	 * Remove the specified tag from any binding and forget about it completely
 	 */
-	void RemoveExposedBinding(FName ExposedAs);
+	void RemoveTag(const FName& TagToRemove);
 
 protected:
 
@@ -979,7 +979,7 @@ private:
 	UPROPERTY()
 	TArray<FMovieSceneBinding> ObjectBindings;
 
-	/** Map of persistent named bindings for this sequence */
+	/** Map of persistent tagged bindings for this sequence */
 	UPROPERTY()
 	TMap<FName, FMovieSceneObjectBindingIDs> BindingGroups;
 
diff --git a/Engine/Source/Runtime/MovieScene/Public/MovieSceneSequence.h b/Engine/Source/Runtime/MovieScene/Public/MovieSceneSequence.h
--- a/Engine/Source/Runtime/MovieScene/Public/MovieSceneSequence.h
+++ b/Engine/Source/Runtime/MovieScene/Public/MovieSceneSequence.h
@@ -180,16 +180,16 @@ public:
 	virtual UObject* CreateDirectorInstance(IMovieScenePlayer& Player) { return nullptr; }
 
 	/**
-	 * Find the first object binding ID associated with the specified name (set up through RMB->Expose on Object bindings from within sequencer)
+	 * Find the first object binding ID associated with the specified tag name (set up through RMB->Expose on Object bindings from within sequencer)
 	 */
 	UFUNCTION(BlueprintCallable, Category = "Game|Cinematic|Bindings")
-	MOVIESCENE_API FMovieSceneObjectBindingID FindNamedBinding(FName InBindingName) const;
+	MOVIESCENE_API FMovieSceneObjectBindingID FindBindingByTag(FName InBindingName) const;
 
 	/**
-	 * Find all object binding IDs associated with the specified name (set up through RMB->Expose on Object bindings from within sequencer)
+	 * Find all object binding IDs associated with the specified tag name (set up through RMB->Expose on Object bindings from within sequencer)
 	 */
 	UFUNCTION(BlueprintCallable, Category = "Game|Cinematic|Bindings")
-	MOVIESCENE_API const TArray<FMovieSceneObjectBindingID>& FindNamedBindings(FName InBindingName) const;
+	MOVIESCENE_API const TArray<FMovieSceneObjectBindingID>& FindBindingsByTag(FName InBindingName) const;
 
 public:
 
