Sequencer: Paste sections onto nodes with selected keys or sections if there are no nodes selected.

diff --git a/Engine/Source/Editor/Sequencer/Private/Sequencer.cpp b/Engine/Source/Editor/Sequencer/Private/Sequencer.cpp
--- a/Engine/Source/Editor/Sequencer/Private/Sequencer.cpp
+++ b/Engine/Source/Editor/Sequencer/Private/Sequencer.cpp
@@ -7417,6 +7417,32 @@ bool FSequencer::PasteTracks(const FString& TextToImport, TArray<FNotificationIn
 	return true;
 }
 
+void GetSupportedTracks(TSharedRef<FSequencerDisplayNode> DisplayNode, const TArray<UMovieSceneSection*>& ImportedSections, TSet<UMovieSceneTrack*>& TracksToPasteOnto)
+{
+	if (DisplayNode->GetType() != ESequencerNode::Track)
+	{
+		return;
+	}
+
+	TSharedPtr<FSequencerTrackNode> TrackNode = StaticCastSharedRef<FSequencerTrackNode>(DisplayNode);
+	if (TrackNode.IsValid())
+	{
+		UMovieSceneTrack* Track = TrackNode->GetTrack();
+		if (Track)
+		{
+			for (int32 SectionIndex = 0; SectionIndex < ImportedSections.Num(); ++SectionIndex)
+			{
+				UMovieSceneSection* Section = ImportedSections[SectionIndex];
+
+				if (Track->SupportsType(Section->GetClass()))
+				{
+					TracksToPasteOnto.Add(Track);
+				}
+			}
+		}
+	}
+}
+
 
 bool FSequencer::PasteSections(const FString& TextToImport, TArray<FNotificationInfo>& PasteErrors)
 {
@@ -7432,6 +7458,11 @@ bool FSequencer::PasteSections(const FString& TextToImport, TArray<FNotification
 
 	TSet<TSharedRef<FSequencerDisplayNode>> SelectedNodes = Selection.GetSelectedOutlinerNodes();
 
+	if (SelectedNodes.Num() == 0)
+	{
+		SelectedNodes = Selection.GetNodesWithSelectedKeysOrSections();
+	}
+
 	if (SelectedNodes.Num() == 0)
 	{
 		FNotificationInfo Info(LOCTEXT("PasteSections_NoSelectedTracks", "Can't paste section. No selected tracks to paste sections onto"));
@@ -7460,51 +7491,62 @@ bool FSequencer::PasteSections(const FString& TextToImport, TArray<FNotification
 		}
 	}
 
-	TArray<UMovieSceneSection*> NewSections;
-	TArray<int32> SectionIndicesImported;
-
+	// Check if any of the selected nodes supports pasting this type of section
+	TSet<UMovieSceneTrack*> TracksToPasteOnto;
 	for (TSharedRef<FSequencerDisplayNode> Node : SelectedNodes)
 	{
-		if (Node->GetType() != ESequencerNode::Track)
-		{
-			continue;
-		}
+		GetSupportedTracks(Node, ImportedSections, TracksToPasteOnto);
+	}
 
-		TSharedPtr<FSequencerTrackNode> TrackNode = StaticCastSharedRef<FSequencerTrackNode>(Node);
-		if (TrackNode.IsValid())
+	// Otherwise, look at all child nodes for supported tracks
+	if (TracksToPasteOnto.Num() == 0)
+	{
+		for (TSharedRef<FSequencerDisplayNode> Node : SelectedNodes)
 		{
-			UMovieSceneTrack* Track = TrackNode->GetTrack();
-			for (int32 SectionIndex = 0; SectionIndex < ImportedSections.Num(); ++SectionIndex)
+			TSet<TSharedRef<FSequencerDisplayNode> > DescendantNodes;
+			SequencerHelpers::GetDescendantNodes(Node, DescendantNodes);
+
+			for (TSharedRef<FSequencerDisplayNode> DescendantNode : DescendantNodes)
 			{
-				UMovieSceneSection* Section = ImportedSections[SectionIndex];
+				GetSupportedTracks(DescendantNode, ImportedSections, TracksToPasteOnto);
+			}
+		}
+	}
 
-				if (!Track->SupportsType(Section->GetClass()))
-				{
-					continue;
-				}
+	TArray<UMovieSceneSection*> NewSections;
+	TArray<int32> SectionIndicesImported;
 
-				SectionIndicesImported.AddUnique(SectionIndex);
+	for (UMovieSceneTrack* Track : TracksToPasteOnto)
+	{
+		for (int32 SectionIndex = 0; SectionIndex < ImportedSections.Num(); ++SectionIndex)
+		{
+			UMovieSceneSection* Section = ImportedSections[SectionIndex];
 
-				Track->Modify();
+			if (!Track->SupportsType(Section->GetClass()))
+			{
+				continue;
+			}
 
-				Section->Rename(nullptr, Track);
-				Track->AddSection(*Section);
-				Section->SetRowIndex(MovieSceneToolHelpers::FindAvailableRowIndex(Track, Section));
-				TrackNode->SetExpansionState(true);
+			SectionIndicesImported.AddUnique(SectionIndex);
 
-				if (Section->HasStartFrame())
-				{
-					FFrameNumber NewStartFrame = LocalTime + (Section->GetInclusiveStartFrame() - FirstFrame.GetValue());
-					Section->MoveSection(NewStartFrame - Section->GetInclusiveStartFrame());
-				}
+			Track->Modify();
+
+			Section->Rename(nullptr, Track);
+			Track->AddSection(*Section);
+			Section->SetRowIndex(MovieSceneToolHelpers::FindAvailableRowIndex(Track, Section));
 
-				NewSections.Add(Section);
+			if (Section->HasStartFrame())
+			{
+				FFrameNumber NewStartFrame = LocalTime + (Section->GetInclusiveStartFrame() - FirstFrame.GetValue());
+				Section->MoveSection(NewStartFrame - Section->GetInclusiveStartFrame());
 			}
 
-			// Regenerate for pasting onto the next track 
-			ImportedSections.Empty();
-			FSequencer::ImportSectionsFromText(TextToImport, ImportedSections);
+			NewSections.Add(Section);
 		}
+
+		// Regenerate for pasting onto the next track 
+		ImportedSections.Empty();
+		FSequencer::ImportSectionsFromText(TextToImport, ImportedSections);
 	}
 
 	for (int32 SectionIndex = 0; SectionIndex < ImportedSections.Num(); ++SectionIndex)
