Sequencer: Playback and selection range exlusive drawing fix. Change to access the bound value directly since Discrete gives +1.

diff --git a/Engine/Source/Editor/Sequencer/Private/SequencerTimeSliderController.cpp b/Engine/Source/Editor/Sequencer/Private/SequencerTimeSliderController.cpp
--- a/Engine/Source/Editor/Sequencer/Private/SequencerTimeSliderController.cpp
+++ b/Engine/Source/Editor/Sequencer/Private/SequencerTimeSliderController.cpp
@@ -590,8 +590,8 @@ int32 FSequencerTimeSliderController::DrawPlaybackRange(const FGeometry& Allotte
 
 	TRange<FFrameNumber> PlaybackRange = TimeSliderArgs.PlaybackRange.Get();
 	FFrameRate TickResolution = GetTickResolution();
-	const float PlaybackRangeL = RangeToScreen.InputToLocalX(MovieScene::DiscreteInclusiveLower(PlaybackRange) / TickResolution);
-	const float PlaybackRangeR = RangeToScreen.InputToLocalX(MovieScene::DiscreteExclusiveUpper(PlaybackRange) / TickResolution) - 1;
+	const float PlaybackRangeL = RangeToScreen.InputToLocalX(PlaybackRange.GetLowerBoundValue() / TickResolution);
+	const float PlaybackRangeR = RangeToScreen.InputToLocalX(PlaybackRange.GetUpperBoundValue() / TickResolution) - 1;
 
 	FSlateDrawElement::MakeBox(
 		OutDrawElements,
@@ -650,8 +650,8 @@ int32 FSequencerTimeSliderController::DrawSubSequenceRange(const FGeometry& Allo
 	}
 
 	const FFrameRate   Resolution = GetTickResolution();
-	const FFrameNumber LowerFrame = MovieScene::DiscreteInclusiveLower(RangeValue.GetValue());
-	const FFrameNumber UpperFrame = MovieScene::DiscreteExclusiveUpper(RangeValue.GetValue());
+	const FFrameNumber LowerFrame = RangeValue.GetValue().GetLowerBoundValue();
+	const FFrameNumber UpperFrame = RangeValue.GetValue().GetUpperBoundValue();
 
 	const float SubSequenceRangeL = RangeToScreen.InputToLocalX(LowerFrame / Resolution) - 1;
 	const float SubSequenceRangeR = RangeToScreen.InputToLocalX(UpperFrame / Resolution) + 1;
