Sequencer: Playback and selection range exlusive drawing fix. Change to access the bound value directly since Discrete gives +1.

#rnx

diff --git a/Engine/Source/Editor/Sequencer/Private/SequencerTimeSliderController.cpp b/Engine/Source/Editor/Sequencer/Private/SequencerTimeSliderController.cpp
--- a/Engine/Source/Editor/Sequencer/Private/SequencerTimeSliderController.cpp
+++ b/Engine/Source/Editor/Sequencer/Private/SequencerTimeSliderController.cpp
@@ -534,8 +534,8 @@ int32 FSequencerTimeSliderController::DrawSelectionRange(const FGeometry& Allott
 
 	if (!SelectionRange.IsEmpty())
 	{
-		const float SelectionRangeL = RangeToScreen.InputToLocalX(SelectionRange.GetLowerBoundValue()) - 1;
-		const float SelectionRangeR = RangeToScreen.InputToLocalX(SelectionRange.GetUpperBoundValue()) + 1;
+		const float SelectionRangeL = RangeToScreen.InputToLocalX(SelectionRange.GetLowerBoundValue());
+		const float SelectionRangeR = RangeToScreen.InputToLocalX(SelectionRange.GetUpperBoundValue()) - 1;
 		const auto DrawColor = FEditorStyle::GetSlateColor("SelectionColor").GetColor(FWidgetStyle());
 
 		if (Args.SolidFillOpacity > 0.f)
