Sequencer: Increase stability of track/section sorting and placement, especially when user is dragging sections between rows

#rnx

diff --git a/Engine/Source/Editor/Sequencer/Private/DisplayNodes/SequencerDisplayNode.cpp b/Engine/Source/Editor/Sequencer/Private/DisplayNodes/SequencerDisplayNode.cpp
--- a/Engine/Source/Editor/Sequencer/Private/DisplayNodes/SequencerDisplayNode.cpp
+++ b/Engine/Source/Editor/Sequencer/Private/DisplayNodes/SequencerDisplayNode.cpp
@@ -73,7 +73,31 @@ namespace SequencerNodeConstants
 		// For nodes of the same bias, sort by name
 		if (BiasA == BiasB)
 		{
-			return A->GetDisplayName().CompareToCaseIgnored(B->GetDisplayName()) < 0;
+			const int32 Compare = A->GetDisplayName().CompareToCaseIgnored(B->GetDisplayName());
+
+			if (Compare != 0)
+			{
+				return Compare < 0;
+			}
+
+			// If the nodes have the same name, try to maintain current sorting order
+			const int32 SortA = A->GetSortingOrder();
+			const int32 SortB = B->GetSortingOrder();
+
+			if (SortA >= 0 && SortB >= 0)
+			{
+				// Both nodes have persistent sort orders, use those
+				return SortA < SortB;
+			}
+			else if (SortA >= 0 || SortB >= 0)
+			{
+				// Only one nodes has a persistent sort order, list it first
+				return SortA > SortB;
+			}
+			
+			// If same name and neither has a persistent sort order, then report them as equal
+			return false;
+
 		}
 		return BiasA < BiasB;
 	}
diff --git a/Engine/Source/Editor/Sequencer/Private/Tools/EditToolDragOperations.cpp b/Engine/Source/Editor/Sequencer/Private/Tools/EditToolDragOperations.cpp
--- a/Engine/Source/Editor/Sequencer/Private/Tools/EditToolDragOperations.cpp
+++ b/Engine/Source/Editor/Sequencer/Private/Tools/EditToolDragOperations.cpp
@@ -1101,6 +1101,7 @@ bool FMoveKeysAndSections::HandleSectionMovement(FFrameTime MouseTime, FVector2D
 						}
 					}
 				
+					// Track if we're expanding a parent track so we can unexpand it if we stop targeting it
 					if (TargetRowIndex > 0)
 					{
 						if (!ParentTrack->IsExpanded() && ParentTrack != ExpandedParentTrack)
@@ -1125,6 +1126,19 @@ bool FMoveKeysAndSections::HandleSectionMovement(FFrameTime MouseTime, FVector2D
 
 		bool bDeltaX = DesiredDeltaX != 0;
 		bool bDeltaY = TargetRowIndex != Section->GetRowIndex();
+		const int32 TargetRowDelta = Section->GetRowIndex() - TargetRowIndex;
+
+		// Prevent flickering by only moving sections if the user has actually made an effort to do so
+		if (bDeltaY && PrevMousePosY.IsSet())
+		{
+			// Check mouse has been moved in the direction of intended move
+			if ((TargetRowDelta > 0 && LocalMousePos.Y - PrevMousePosY.GetValue() > 1.0f) || (TargetRowDelta < 0 && LocalMousePos.Y - PrevMousePosY.GetValue() < 1.0f))
+			{
+				// Mouse was not moved in the direction the section wants to swap
+				// Assume offset is due to UI relayout and block moving the section
+				bDeltaY = false;
+			}
+		}
 
 		// Horizontal movement
 		if (bDeltaX)
@@ -1182,6 +1196,11 @@ bool FMoveKeysAndSections::HandleSectionMovement(FFrameTime MouseTime, FVector2D
 		}
 	}
 
+	if (bRowIndexChanged)
+	{
+		PrevMousePosY = LocalMousePos.Y;
+	}
+
 	return bRowIndexChanged;
 }
 
diff --git a/Engine/Source/Editor/Sequencer/Private/Tools/EditToolDragOperations.h b/Engine/Source/Editor/Sequencer/Private/Tools/EditToolDragOperations.h
--- a/Engine/Source/Editor/Sequencer/Private/Tools/EditToolDragOperations.h
+++ b/Engine/Source/Editor/Sequencer/Private/Tools/EditToolDragOperations.h
@@ -165,6 +165,9 @@ protected:
 	/** What was the time of the mouse for the previous frame? Used to calculate a per-frame delta. */
 	FFrameTime MouseTimePrev;
 
+	/** The position of the mouse when the last section move occurred */
+	TOptional<float> PrevMousePosY;
+
 	struct FRelativeOffset
 	{
 		FRelativeOffset()
