Sequencer: Clamp mouse position to prevent sections from jumpng to the bottom when dragging up

diff --git a/Engine/Source/Editor/Sequencer/Private/SSequencerTreeView.cpp b/Engine/Source/Editor/Sequencer/Private/SSequencerTreeView.cpp
--- a/Engine/Source/Editor/Sequencer/Private/SSequencerTreeView.cpp
+++ b/Engine/Source/Editor/Sequencer/Private/SSequencerTreeView.cpp
@@ -363,8 +363,16 @@ float SSequencerTreeView::PhysicalToVirtual(float InPhysical) const
 
 	if (PhysicalNodes.Num())
 	{
-		const FCachedGeometry& Last = PhysicalNodes.Last();
-		return Last.Node->GetVirtualTop() + (InPhysical - Last.PhysicalTop);
+		const FCachedGeometry& First = PhysicalNodes[0];
+		if (InPhysical < First.PhysicalTop)
+		{
+			return First.Node->GetVirtualTop() + (InPhysical - First.PhysicalTop);
+		}
+		else
+		{
+			const FCachedGeometry& Last = PhysicalNodes.Last();
+			return Last.Node->GetVirtualTop() + (InPhysical - Last.PhysicalTop);
+		}
 	}
 
 	return InPhysical;
diff --git a/Engine/Source/Editor/Sequencer/Private/Tools/EditToolDragOperations.cpp b/Engine/Source/Editor/Sequencer/Private/Tools/EditToolDragOperations.cpp
--- a/Engine/Source/Editor/Sequencer/Private/Tools/EditToolDragOperations.cpp
+++ b/Engine/Source/Editor/Sequencer/Private/Tools/EditToolDragOperations.cpp
@@ -1090,7 +1090,20 @@ bool FMoveKeysAndSections::HandleSectionMovement(FFrameTime MouseTime, FVector2D
 						float VirtualSectionBottom = 0.f;
 						ChildNode->TraverseVisible_ParentFirst([&](FSequencerDisplayNode& Node) { VirtualSectionBottom = Node.GetVirtualBottom(); return true; }, true);
 
-						if (VirtualMousePos.Y < VirtualSectionBottom)
+						if (ChildIndex == 0 && (VirtualMousePos.Y <= VirtualSectionTop || LocalMousePos.Y <= 0))
+						{
+							TargetRowIndex = 0;
+							for (TSharedRef<ISequencerSection> TrackSection : TrackNode->GetSections())
+							{
+								if (!Sections.Contains(TrackSection->GetSectionObject()))
+								{
+									TargetRowIndex = -1;
+									break;
+								}
+							}
+							break;
+						}
+						else if (VirtualMousePos.Y < VirtualSectionBottom)
 						{
 							TargetRowIndex = ChildIndex;
 							break;
@@ -1126,13 +1139,13 @@ bool FMoveKeysAndSections::HandleSectionMovement(FFrameTime MouseTime, FVector2D
 
 		bool bDeltaX = DesiredDeltaX != 0;
 		bool bDeltaY = TargetRowIndex != Section->GetRowIndex();
-		const int32 TargetRowDelta = Section->GetRowIndex() - TargetRowIndex;
+		const int32 TargetRowDelta = TargetRowIndex - Section->GetRowIndex();
 
 		// Prevent flickering by only moving sections if the user has actually made an effort to do so
 		if (bDeltaY && PrevMousePosY.IsSet())
 		{
 			// Check mouse has been moved in the direction of intended move
-			if ((TargetRowDelta > 0 && LocalMousePos.Y - PrevMousePosY.GetValue() > 1.0f) || (TargetRowDelta < 0 && LocalMousePos.Y - PrevMousePosY.GetValue() < 1.0f))
+			if ((TargetRowDelta < 0 && LocalMousePos.Y - PrevMousePosY.GetValue() > 1.0f) || (TargetRowDelta > 0 && LocalMousePos.Y - PrevMousePosY.GetValue() < 1.0f))
 			{
 				// Mouse was not moved in the direction the section wants to swap
 				// Assume offset is due to UI relayout and block moving the section
@@ -1146,6 +1159,7 @@ bool FMoveKeysAndSections::HandleSectionMovement(FFrameTime MouseTime, FVector2D
 			Section->MoveSection(MaxDeltaX.Get(DesiredDeltaX));
 		}
 
+
 		// Vertical movement
 		if (bDeltaY && !bSectionsAreOnDifferentRows &&
 			(
