Sequencer: Suppress undo notification after export fbx

#rnx

diff --git a/Engine/Source/Editor/Sequencer/Private/Sequencer.cpp b/Engine/Source/Editor/Sequencer/Private/Sequencer.cpp
--- a/Engine/Source/Editor/Sequencer/Private/Sequencer.cpp
+++ b/Engine/Source/Editor/Sequencer/Private/Sequencer.cpp
@@ -9937,7 +9937,10 @@ void FSequencer::ExportFBXInternal(const FString& ExportFilename, TArray<FGuid>&
 					}
 
 					// Restore modified spawned sections
+					bool bOrigSquelchTransactionNotification = GEditor->bSquelchTransactionNotification;
+					GEditor->bSquelchTransactionNotification = true;
 					GEditor->UndoTransaction(false);
+					GEditor->bSquelchTransactionNotification = bOrigSquelchTransactionNotification;
 				}
 
 				bool bWasChanged;
