Sequencer: Added toast for fbx export

#rnx

diff --git a/Engine/Source/Editor/Sequencer/Private/Sequencer.cpp b/Engine/Source/Editor/Sequencer/Private/Sequencer.cpp
--- a/Engine/Source/Editor/Sequencer/Private/Sequencer.cpp
+++ b/Engine/Source/Editor/Sequencer/Private/Sequencer.cpp
@@ -9952,8 +9952,20 @@ void FSequencer::ExportFBXInternal(const FString& ExportFilename, TArray<FGuid>&
 
 			{
 				FSpawnableRestoreState SpawnableRestoreState(*this);
-
-				MovieSceneToolHelpers::ExportFBX(World, MovieScene, this, Bindings, NodeNameAdapter, Template, ExportFilename, RootToLocalTransform);
+				if (MovieSceneToolHelpers::ExportFBX(World, MovieScene, this, Bindings, NodeNameAdapter, Template, ExportFilename, RootToLocalTransform))
+				{
+					FNotificationInfo Info(NSLOCTEXT("Sequencer", "ExportFBXSucceeded", "FBX Export Succeeded."));
+					Info.Hyperlink = FSimpleDelegate::CreateStatic([](FString InFilename) { FPlatformProcess::ExploreFolder(*InFilename); }, ExportFilename);
+					Info.HyperlinkText = FText::FromString(ExportFilename);
+					Info.ExpireDuration = 5.0f;
+					FSlateNotificationManager::Get().AddNotification(Info)->SetCompletionState(SNotificationItem::CS_Success);
+				}
+				else
+				{
+					FNotificationInfo Info(NSLOCTEXT("Sequencer", "ExportFBXFailed", "FBX Export Failed."));
+					Info.ExpireDuration = 5.0f;
+					FSlateNotificationManager::Get().AddNotification(Info)->SetCompletionState(SNotificationItem::CS_Fail);
+				}
 			}
 
 			ForceEvaluate();
