Sequencer: Traverse parent hierachy to correctly test if a display node is part of a pinned group

#rnx

diff --git a/Engine/Source/Editor/Sequencer/Private/DisplayNodes/SequencerDisplayNode.cpp b/Engine/Source/Editor/Sequencer/Private/DisplayNodes/SequencerDisplayNode.cpp
--- a/Engine/Source/Editor/Sequencer/Private/DisplayNodes/SequencerDisplayNode.cpp
+++ b/Engine/Source/Editor/Sequencer/Private/DisplayNodes/SequencerDisplayNode.cpp
@@ -423,7 +423,8 @@ FSequencerDisplayNode::FSequencerDisplayNode( FName InNodeName, FSequencerNodeTr
 	, ParentTree( InParentTree )
 	, NodeName( InNodeName )
 	, bExpanded( false )
-	, bPinned(false)
+	, bPinned( false )
+	, bInPinnedBranch( false )
 	, bHasBeenInitialized( false )
 {
 	SortType = EDisplayNodeSortType::Undefined;
@@ -1223,9 +1224,19 @@ FSequencerDisplayNode* FSequencerDisplayNode::GetBaseNode() const
 	return GetParentOrRoot()->GetBaseNode();
 }
 
+void FSequencerDisplayNode::UpdateCachedPinnedState(bool bParentIsPinned)
+{
+	bInPinnedBranch = bPinned || bParentIsPinned;
+
+	for (TSharedPtr<FSequencerDisplayNode> Child : ChildNodes)
+	{
+		Child->UpdateCachedPinnedState(bInPinnedBranch);
+	}
+}
+
 bool FSequencerDisplayNode::IsPinned() const
 {
-	return GetBaseNode()->bPinned;
+	return bInPinnedBranch;
 }
 
 void FSequencerDisplayNode::TogglePinned()
diff --git a/Engine/Source/Editor/Sequencer/Private/DisplayNodes/SequencerDisplayNode.h b/Engine/Source/Editor/Sequencer/Private/DisplayNodes/SequencerDisplayNode.h
--- a/Engine/Source/Editor/Sequencer/Private/DisplayNodes/SequencerDisplayNode.h
+++ b/Engine/Source/Editor/Sequencer/Private/DisplayNodes/SequencerDisplayNode.h
@@ -443,6 +443,11 @@ public:
 	 */
 	bool IsExpanded() const;
 
+	/**
+	 * Called by FSequencer to update the cached pinned state
+	 */
+	void UpdateCachedPinnedState(bool bParentIsPinned = false);
+
 	/**
 	 * @return Whether or not this node is pinned
 	 */
@@ -573,6 +578,9 @@ protected:
 	/** Whether or not the node is pinned */
 	bool bPinned;
 
+	/** Cached value of whether this node or one of it's parents is pinned */
+	bool bInPinnedBranch;
+
 	/** Event that is triggered when rename is requested */
 	FRequestRenameEvent RenameRequestedEvent;
 
diff --git a/Engine/Source/Editor/Sequencer/Private/SequencerNodeTree.cpp b/Engine/Source/Editor/Sequencer/Private/SequencerNodeTree.cpp
--- a/Engine/Source/Editor/Sequencer/Private/SequencerNodeTree.cpp
+++ b/Engine/Source/Editor/Sequencer/Private/SequencerNodeTree.cpp
@@ -347,6 +347,9 @@ void FSequencerNodeTree::Update()
 	const bool bIncludeRootNode = false;
 	RootNode->Traverse_ParentFirst(Traverse_OnTreeRefreshed, bIncludeRootNode);
 
+	// Cache pinned state of nodes, needs to happen after OnTreeRefreshed
+	RootNode->UpdateCachedPinnedState();
+
 	// Ensure that the curve editor tree is up to date for our tree layout
 	UpdateCurveEditorTree();
 
