Sequencer: Show "Sequence read only" error message when attempting to modify a readonly sequence.

diff --git a/Engine/Source/Editor/Sequencer/Private/Sequencer.cpp b/Engine/Source/Editor/Sequencer/Private/Sequencer.cpp
--- a/Engine/Source/Editor/Sequencer/Private/Sequencer.cpp
+++ b/Engine/Source/Editor/Sequencer/Private/Sequencer.cpp
@@ -1501,6 +1501,7 @@ void FSequencer::StretchTime(FFrameTime InDeltaTime)
 	UMovieScene* FocusedMovieScene = GetFocusedMovieSceneSequence()->GetMovieScene();
 	if (FocusedMovieScene->IsReadOnly())
 	{
+		ShowReadOnlyError();
 		return;
 	}
 
@@ -1536,6 +1537,7 @@ void FSequencer::ShrinkTime(FFrameTime InDeltaTime)
 	UMovieScene* FocusedMovieScene = GetFocusedMovieSceneSequence()->GetMovieScene();
 	if (FocusedMovieScene->IsReadOnly())
 	{
+		ShowReadOnlyError();
 		return;
 	}
 
@@ -1612,6 +1614,7 @@ void FSequencer::BakeTransform()
 	UMovieScene* FocusedMovieScene = GetFocusedMovieSceneSequence()->GetMovieScene();
 	if (FocusedMovieScene->IsReadOnly())
 	{
+		ShowReadOnlyError();
 		return;
 	}
 
@@ -2253,6 +2256,7 @@ void FSequencer::TogglePlaybackRangeLocked()
 
 		if (MovieScene->IsReadOnly())
 		{
+			ShowReadOnlyError();
 			return;
 		}
 
@@ -4324,6 +4328,13 @@ bool FSequencer::IsReadOnly() const
 	return bReadOnly || (GetFocusedMovieSceneSequence() && GetFocusedMovieSceneSequence()->GetMovieScene()->IsReadOnly());
 }
 
+void FSequencer::ShowReadOnlyError() const
+{
+	FNotificationInfo Info(NSLOCTEXT("Sequencer", "SequenceReadOnly", "Sequence is read only."));
+	Info.ExpireDuration = 5.0f;
+	FSlateNotificationManager::Get().AddNotification(Info)->SetCompletionState(SNotificationItem::CS_Fail);
+}
+
 void FSequencer::VerticalScroll(float ScrollAmountUnits)
 {
 	SequencerWidget->GetTreeView()->ScrollByDelta(ScrollAmountUnits);
@@ -4380,6 +4391,7 @@ FGuid FSequencer::MakeNewSpawnable( UObject& Object, UActorFactory* ActorFactory
 
 	if (MovieScene->IsReadOnly())
 	{
+		ShowReadOnlyError();
 		return FGuid();
 	}
 
@@ -4434,6 +4446,7 @@ void FSequencer::AddSubSequence(UMovieSceneSequence* Sequence)
 
 	if (OwnerMovieScene->IsReadOnly())
 	{
+		ShowReadOnlyError();
 		return;
 	}
 
@@ -4506,6 +4519,7 @@ bool FSequencer::OnRequestNodeDeleted( TSharedRef<const FSequencerDisplayNode> N
 
 	if (OwnerMovieScene->IsReadOnly())
 	{
+		ShowReadOnlyError();
 		return bAnythingRemoved;
 	}
 
@@ -4702,6 +4716,7 @@ void FSequencer::OnNewActorsDropped(const TArray<UObject*>& DroppedObjects, cons
 
 		if (OwnerMovieScene->IsReadOnly())
 		{
+			ShowReadOnlyError();
 			return;
 		}
 
@@ -5045,6 +5060,7 @@ TArray<FGuid> FSequencer::AddActors(const TArray<TWeakObjectPtr<AActor> >& InAct
 
 	if (GetFocusedMovieSceneSequence()->GetMovieScene()->IsReadOnly())
 	{
+		ShowReadOnlyError();
 		return PossessableGuids;
 	}
 
@@ -5764,6 +5780,7 @@ FGuid FSequencer::DoAssignActor(AActor*const* InActors, int32 NumActors, FGuid I
 
 	if (OwnerMovieScene->IsReadOnly())
 	{
+		ShowReadOnlyError();
 		return FGuid();
 	}
 
@@ -6309,6 +6326,7 @@ void FSequencer::MoveSelectedNodesToNewFolder()
 
 	if (FocusedMovieScene->IsReadOnly())
 	{
+		ShowReadOnlyError();
 		return;
 	}
 
@@ -6706,6 +6724,12 @@ void FSequencer::ExportObjectsToText(TArray<UObject*> ObjectsToExport, FString&
 
 void FSequencer::DoPaste()
 {
+	if (IsReadOnly())
+	{
+		ShowReadOnlyError();
+		return;
+	}
+
 	// Grab the text to paste from the clipboard
 	FString TextToImport;
 	FPlatformApplicationMisc::ClipboardPaste(TextToImport);
@@ -7379,6 +7403,7 @@ void FSequencer::SaveSelectedNodesSpawnableState()
 
 	if (MovieScene->IsReadOnly())
 	{
+		ShowReadOnlyError();
 		return;
 	}
 
@@ -7424,7 +7449,8 @@ void FSequencer::SetSelectedNodesSpawnableLevel(FName InLevelName)
 	UMovieScene* MovieScene = GetFocusedMovieSceneSequence()->GetMovieScene();
 
 	if (MovieScene->IsReadOnly())
-	{	
+	{
+		ShowReadOnlyError();
 		return;
 	}
 
@@ -7451,6 +7477,7 @@ void FSequencer::ConvertToSpawnable(TSharedRef<FSequencerObjectBindingNode> Node
 {
 	if (GetFocusedMovieSceneSequence()->GetMovieScene()->IsReadOnly())
 	{
+		ShowReadOnlyError();
 		return;
 	}
 
@@ -7473,6 +7500,7 @@ void FSequencer::ConvertSelectedNodesToSpawnables()
 
 	if (MovieScene->IsReadOnly())
 	{
+		ShowReadOnlyError();
 		return;
 	}
 
@@ -7561,6 +7589,7 @@ TArray<FGuid> FSequencer::ExpandMultiplePossessableBindings(FGuid PossessableGui
 
 	if (MovieScene->IsReadOnly())
 	{
+		ShowReadOnlyError();
 		return TArray<FGuid>();
 	}
 
@@ -7663,6 +7692,7 @@ TArray<FMovieSceneSpawnable*> FSequencer::ConvertToSpawnableInternal(FGuid Posse
 
 	if (MovieScene->IsReadOnly() || !Sequence->AllowsSpawnableObjects())
 	{
+		ShowReadOnlyError();
 		return TArray<FMovieSceneSpawnable*>();
 	}
 
@@ -7753,6 +7783,7 @@ void FSequencer::ConvertToPossessable(TSharedRef<FSequencerObjectBindingNode> No
 {
 	if (GetFocusedMovieSceneSequence()->GetMovieScene()->IsReadOnly())
 	{
+		ShowReadOnlyError();
 		return;
 	}
 
@@ -7775,6 +7806,7 @@ void FSequencer::ConvertSelectedNodesToPossessables()
 
 	if (MovieScene->IsReadOnly())
 	{
+		ShowReadOnlyError();
 		return;
 	}
 
@@ -7858,6 +7890,7 @@ FMovieScenePossessable* FSequencer::ConvertToPossessableInternal(FGuid Spawnable
 
 	if (MovieScene->IsReadOnly())
 	{
+		ShowReadOnlyError();
 		return nullptr;
 	}
 
@@ -8013,6 +8046,7 @@ bool FSequencer::ReplaceFolderBindingGUID(UMovieSceneFolder* Folder, FGuid Origi
 
 	if (MovieScene->IsReadOnly())
 	{
+		ShowReadOnlyError();
 		return true;
 	}
 
@@ -8043,6 +8077,7 @@ void FSequencer::OnAddFolder()
 
 	if (FocusedMovieScene->IsReadOnly())
 	{
+		ShowReadOnlyError();
 		return;
 	}
 
@@ -9104,6 +9139,7 @@ void FSequencer::CreateCamera()
 	UMovieScene* FocusedMovieScene = GetFocusedMovieSceneSequence()->GetMovieScene();
 	if (FocusedMovieScene->IsReadOnly())
 	{
+		ShowReadOnlyError();
 		return;
 	}
 
@@ -9283,6 +9319,7 @@ void FSequencer::RebindPossessableReferences()
 
 	if (FocusedMovieScene->IsReadOnly())
 	{
+		ShowReadOnlyError();
 		return;
 	}
 
diff --git a/Engine/Source/Editor/Sequencer/Private/Sequencer.h b/Engine/Source/Editor/Sequencer/Private/Sequencer.h
--- a/Engine/Source/Editor/Sequencer/Private/Sequencer.h
+++ b/Engine/Source/Editor/Sequencer/Private/Sequencer.h
@@ -610,6 +610,9 @@ public:
 	/** Exports the animation to a camera anim asset. */
 	void ExportToCameraAnim();
 
+	/** */
+	void ShowReadOnlyError() const;
+
 public:
 	
 	/** Access the currently active track area edit tool */
