Sequencer: Add synchronization of renaming of spawnable actors in sequencer to outliner, and fix sync from outliner to sequencer.

diff --git a/Engine/Source/Editor/Sequencer/Private/DisplayNodes/SequencerObjectBindingNode.cpp b/Engine/Source/Editor/Sequencer/Private/DisplayNodes/SequencerObjectBindingNode.cpp
--- a/Engine/Source/Editor/Sequencer/Private/DisplayNodes/SequencerObjectBindingNode.cpp
+++ b/Engine/Source/Editor/Sequencer/Private/DisplayNodes/SequencerObjectBindingNode.cpp
@@ -675,19 +675,33 @@ void FSequencerObjectBindingNode::SetDisplayName(const FText& NewDisplayName)
 
 		// Modify the movie scene so that it gets marked dirty and renames are saved consistently.
 		MovieScene->Modify();
-		MovieScene->SetObjectDisplayName(ObjectBinding, NewDisplayName);
 
 		FMovieSceneSpawnable* Spawnable = MovieScene->FindSpawnable(GetObjectBinding());
+		FMovieScenePossessable* Possessable = MovieScene->FindPossessable(GetObjectBinding());
+
 		if (Spawnable)
 		{
-			Spawnable->SetName(NewDisplayName.ToString());
+			TArrayView<TWeakObjectPtr<>> Objects = GetSequencer().FindObjectsInCurrentSequence(GetObjectBinding());
+			// If there is only one binding, set the name of the bound actor
+			if (Objects.Num() == 1)
+			{
+				AActor* Actor = Cast<AActor>(Objects[0].Get());
+				Actor->SetActorLabel(NewDisplayName.ToString());
+			}
+			else
+			{
+				// Otherwise set our display name
+				Spawnable->SetName(NewDisplayName.ToString());
+			}
 		}
-
-		FMovieScenePossessable* Possessable = MovieScene->FindPossessable(GetObjectBinding());
-		if (Possessable)
+		else if (Possessable)
 		{
 			Possessable->SetName(NewDisplayName.ToString());
 		}
+		else
+		{
+			MovieScene->SetObjectDisplayName(ObjectBinding, NewDisplayName);
+		}
 	}
 }
 
diff --git a/Engine/Source/Editor/Sequencer/Private/LevelEditorSequencerIntegration.cpp b/Engine/Source/Editor/Sequencer/Private/LevelEditorSequencerIntegration.cpp
--- a/Engine/Source/Editor/Sequencer/Private/LevelEditorSequencerIntegration.cpp
+++ b/Engine/Source/Editor/Sequencer/Private/LevelEditorSequencerIntegration.cpp
@@ -277,6 +277,7 @@ void RenameSpawnableRecursive(FSequencer* Sequencer, FMovieSceneSequenceIDRef Se
 				AActor* Actor = Cast<AActor>(WeakObject.Get());
 				if (Actor && Actor == ChangedActor)
 				{
+					MovieScene->Modify();
 					MovieScene->GetSpawnable(Index).SetName(ChangedActor->GetActorLabel());
 				}
 			}
