Sequencer: Don't automatically expand track nodes when dragging

diff --git a/Engine/Source/Editor/Sequencer/Private/DisplayNodes/SequencerTrackNode.cpp b/Engine/Source/Editor/Sequencer/Private/DisplayNodes/SequencerTrackNode.cpp
--- a/Engine/Source/Editor/Sequencer/Private/DisplayNodes/SequencerTrackNode.cpp
+++ b/Engine/Source/Editor/Sequencer/Private/DisplayNodes/SequencerTrackNode.cpp
@@ -163,8 +163,6 @@ void FSequencerTrackNode::UpdateInnerHierarchy()
 			// SetParent adds the track to our ChildNodes
 			NewSubTrack->SetParent(AsShared());
 
-			NewSubTrack->SetExpansionState(IsExpanded());
-
 			NewSubTrack->TreeSerialNumber = TreeSerialNumber;
 			NewSubTrack->UpdateSections();
 
diff --git a/Engine/Source/Editor/Sequencer/Private/Tools/EditToolDragOperations.cpp b/Engine/Source/Editor/Sequencer/Private/Tools/EditToolDragOperations.cpp
--- a/Engine/Source/Editor/Sequencer/Private/Tools/EditToolDragOperations.cpp
+++ b/Engine/Source/Editor/Sequencer/Private/Tools/EditToolDragOperations.cpp
@@ -1058,9 +1058,6 @@ bool FMoveKeysAndSections::HandleSectionMovement(FFrameTime MouseTime, FVector2D
 				{
 					TargetRowIndex = -1;
 				}
-
-				// Ensure the track path is expanded
-				TrackNode->SetExpansionState(true);
 			}
 			else if (TrackNode->GetSubTrackMode() == FSequencerTrackNode::ESubTrackMode::SubTrack)
 			{
