Sequencer: Only update outliner selection set if the selection has changed, prevents UI flickering and other side effects.

diff --git a/Engine/Source/Editor/Sequencer/Private/Sequencer.cpp b/Engine/Source/Editor/Sequencer/Private/Sequencer.cpp
--- a/Engine/Source/Editor/Sequencer/Private/Sequencer.cpp
+++ b/Engine/Source/Editor/Sequencer/Private/Sequencer.cpp
@@ -5239,18 +5239,25 @@ void FSequencer::SynchronizeExternalSelectionWithSequencerSelection()
 				AActor* Actor = Cast<AActor>(RuntimeObject.Get());
 				if ( Actor != nullptr )
 				{
-					SelectedSequencerActors.Add( Actor );
+					ULevel* ActorLevel = Actor->GetLevel();
+					if (!FLevelUtils::IsLevelLocked(ActorLevel))
+					{
+						SelectedSequencerActors.Add(Actor);
+					}
 				}
 
 				UActorComponent* ActorComponent = Cast<UActorComponent>(RuntimeObject.Get());
 				if ( ActorComponent != nullptr )
 				{
-					SelectedSequencerComponents.Add( ActorComponent );
-
-					Actor = ActorComponent->GetOwner();
-					if ( Actor != nullptr )
+					if (!FLevelUtils::IsLevelLocked(ActorComponent->GetOwner()->GetLevel()))
 					{
-						SelectedSequencerActors.Add( Actor );
+						SelectedSequencerComponents.Add(ActorComponent);
+
+						Actor = ActorComponent->GetOwner();
+						if (Actor != nullptr)
+						{
+							SelectedSequencerActors.Add(Actor);
+						}
 					}
 				}
 			}
@@ -5273,6 +5280,69 @@ void FSequencer::SynchronizeExternalSelectionWithSequencerSelection()
 		return;
 	}
 
+	// We need to check if the selection has changed. Rebuilding the selection set if it hasn't changed can cause unwanted side effects.
+	bool bIsSelectionChanged = false;
+
+	// Check if any actors have been added to the selection
+	for (AActor* SelectedSequencerActor : SelectedSequencerActors)
+	{
+		if (!GEditor->GetSelectedActors()->IsSelected(SelectedSequencerActor))
+		{
+			bIsSelectionChanged = true;
+			break;
+		}
+	}
+
+	// Check if any actors have been removed from the selection
+	if (!bIsSelectionChanged)
+	{
+		for (FSelectionIterator It(GEditor->GetSelectedActorIterator()); It; ++It)
+		{
+			if (AActor* CurrentlySelectedActor = Cast<AActor>(*It))
+			{
+				if (!SelectedSequencerActors.Contains(CurrentlySelectedActor))
+				{
+					bIsSelectionChanged = true;
+					break;
+				}
+			}
+		}
+	}
+
+	// Check if any components have been added to the selection
+	if (!bIsSelectionChanged)
+	{
+		for (UActorComponent* SelectedSequencerComponent : SelectedSequencerComponents)
+		{
+			if (!GEditor->GetSelectedComponents()->IsSelected(SelectedSequencerComponent))
+			{
+				bIsSelectionChanged = true;
+				break;
+			}
+		}
+	}
+
+	// Check if any components have been removed from the selection
+	if (!bIsSelectionChanged)
+	{
+		for (FSelectionIterator It(GEditor->GetSelectedComponentIterator()); It; ++It)
+		{
+			if (UActorComponent* CurrentlySelectedComponent = Cast<UActorComponent>(*It))
+			{
+				if (!SelectedSequencerComponents.Contains(CurrentlySelectedComponent))
+				{
+					bIsSelectionChanged = true;
+					break;
+				}
+			}
+		}
+	}
+
+	if (!bIsSelectionChanged)
+	{
+		return;
+	}
+
 	const FScopedTransaction Transaction( NSLOCTEXT( "Sequencer", "UpdatingActorComponentSelection", "Select Actors/Components" ) );
 
 
@@ -5283,11 +5353,7 @@ void FSequencer::SynchronizeExternalSelectionWithSequencerSelection()
 
 	for (AActor* SelectedSequencerActor : SelectedSequencerActors)
 	{
-		ULevel* ActorLevel = SelectedSequencerActor->GetLevel();
-		if (!FLevelUtils::IsLevelLocked(ActorLevel))
-		{
-			GEditor->SelectActor(SelectedSequencerActor, true, bNotifySelectionChanged, bSelectEvenIfHidden);
-		}
+		GEditor->SelectActor(SelectedSequencerActor, true, bNotifySelectionChanged, bSelectEvenIfHidden);
 	}
 
 	GEditor->GetSelectedActors()->EndBatchSelectOperation();
@@ -5299,10 +5365,7 @@ void FSequencer::SynchronizeExternalSelectionWithSequencerSelection()
 
 		for (UActorComponent* SelectedSequencerComponent : SelectedSequencerComponents)
 		{
-			if (!FLevelUtils::IsLevelLocked(SelectedSequencerComponent->GetOwner()->GetLevel()))
-			{
-				GEditor->SelectComponent(SelectedSequencerComponent, true, bNotifySelectionChanged, bSelectEvenIfHidden);
-			}
+			GEditor->SelectComponent(SelectedSequencerComponent, true, bNotifySelectionChanged, bSelectEvenIfHidden);
 		}
 
 		GEditor->GetSelectedComponents()->EndBatchSelectOperation();
