Sequencer: Validate renaming

diff --git a/Engine/Source/Editor/Sequencer/Private/DisplayNodes/SequencerDisplayNode.cpp b/Engine/Source/Editor/Sequencer/Private/DisplayNodes/SequencerDisplayNode.cpp
--- a/Engine/Source/Editor/Sequencer/Private/DisplayNodes/SequencerDisplayNode.cpp
+++ b/Engine/Source/Editor/Sequencer/Private/DisplayNodes/SequencerDisplayNode.cpp
@@ -721,6 +721,16 @@ FText FSequencerDisplayNode::GetDisplayNameToolTipText() const
 	return FText();
 }
 
+bool FSequencerDisplayNode::ValidateDisplayName(const FText& NewDisplayName, FText& OutErrorMessage) const
+{
+	if (NewDisplayName.IsEmpty())
+	{
+		OutErrorMessage = NSLOCTEXT("Sequencer", "RenameFailed_LeftBlank", "Labels cannot be left blank");
+		return false;
+	}
+	return true;
+}
+
 TSharedRef<SWidget> FSequencerDisplayNode::GenerateContainerWidgetForOutliner(const TSharedRef<SSequencerTreeViewRow>& InRow)
 {
 	auto NewWidget = SNew(SAnimationOutlinerTreeNode, SharedThis(this), InRow)
diff --git a/Engine/Source/Editor/Sequencer/Private/DisplayNodes/SequencerDisplayNode.h b/Engine/Source/Editor/Sequencer/Private/DisplayNodes/SequencerDisplayNode.h
--- a/Engine/Source/Editor/Sequencer/Private/DisplayNodes/SequencerDisplayNode.h
+++ b/Engine/Source/Editor/Sequencer/Private/DisplayNodes/SequencerDisplayNode.h
@@ -175,6 +175,11 @@ public:
 	 */
 	virtual FText GetDisplayNameToolTipText() const;
 
+	/**
+	 * Return whether the new display name is valid for this node.
+	 */
+	virtual bool ValidateDisplayName(const FText& NewDisplayName, FText& OutErrorMessage) const;
+
 	/**
 	 * Set the node's display name.
 	 *
diff --git a/Engine/Source/Editor/Sequencer/Private/SAnimationOutlinerTreeNode.cpp b/Engine/Source/Editor/Sequencer/Private/SAnimationOutlinerTreeNode.cpp
--- a/Engine/Source/Editor/Sequencer/Private/SAnimationOutlinerTreeNode.cpp
+++ b/Engine/Source/Editor/Sequencer/Private/SAnimationOutlinerTreeNode.cpp
@@ -66,6 +66,7 @@ void SAnimationOutlinerTreeNode::Construct( const FArguments& InArgs, TSharedRef
 		.IsReadOnly(this, &SAnimationOutlinerTreeNode::IsNodeLabelReadOnly)
 		.Font(NodeFont)
 		.ColorAndOpacity(this, &SAnimationOutlinerTreeNode::GetDisplayNameColor)
+		.OnVerifyTextChanged(this, &SAnimationOutlinerTreeNode::VerifyNodeTextChanged)
 		.OnTextCommitted(this, &SAnimationOutlinerTreeNode::HandleNodeLabelTextCommitted)
 		.Text(this, &SAnimationOutlinerTreeNode::GetDisplayName)
 		.ToolTipText(this, &SAnimationOutlinerTreeNode::GetDisplayNameToolTipText)
@@ -425,6 +426,10 @@ bool SAnimationOutlinerTreeNode::IsNodeLabelReadOnly() const
 	return DisplayNode->GetSequencer().IsReadOnly() || !DisplayNode->CanRenameNode();
 }
 
+bool SAnimationOutlinerTreeNode::VerifyNodeTextChanged(const FText& NewLabel, FText& OutErrorMessage)
+{
+	return DisplayNode->ValidateDisplayName(NewLabel, OutErrorMessage);
+}
 
 void SAnimationOutlinerTreeNode::HandleNodeLabelTextCommitted(const FText& NewLabel, ETextCommit::Type CommitType)
 {
diff --git a/Engine/Source/Editor/Sequencer/Private/SAnimationOutlinerTreeNode.h b/Engine/Source/Editor/Sequencer/Private/SAnimationOutlinerTreeNode.h
--- a/Engine/Source/Editor/Sequencer/Private/SAnimationOutlinerTreeNode.h
+++ b/Engine/Source/Editor/Sequencer/Private/SAnimationOutlinerTreeNode.h
@@ -102,6 +102,9 @@ private:
 	/** Callback for checking whether the node label can be edited. */
 	bool IsNodeLabelReadOnly() const;
 
+	/** Callback to verify whether the text is valid for renaming */
+	bool VerifyNodeTextChanged(const FText& NewLabel, FText& OutErrorMessage);
+
 	/** Callback for when the node label text has changed. */
 	void HandleNodeLabelTextCommitted(const FText& NewLabel, ETextCommit::Type CommitType);
 
