Sequencer: Only expand to show filtered tracks if they are a direct match to a text search

diff --git a/Engine/Source/Editor/Sequencer/Private/SequencerNodeTree.cpp b/Engine/Source/Editor/Sequencer/Private/SequencerNodeTree.cpp
--- a/Engine/Source/Editor/Sequencer/Private/SequencerNodeTree.cpp
+++ b/Engine/Source/Editor/Sequencer/Private/SequencerNodeTree.cpp
@@ -631,42 +631,17 @@ static void AddChildNodes(const TSharedRef<FSequencerDisplayNode>& StartNode, TS
  */
 static void AddFilteredNode(const TSharedRef<FSequencerDisplayNode>& StartNode, TSet<TSharedRef<FSequencerDisplayNode>>& OutFilteredNodes)
 {
-	if (!StartNode->IsExpanded())
-	{
-		StartNode->SetExpansionState(true);
-	}
-
 	AddChildNodes(StartNode, OutFilteredNodes);
 
 	// Gather parent folders up the chain
 	TSharedPtr<FSequencerDisplayNode> ParentNode = StartNode->GetParent();
 	while (ParentNode.IsValid())
 	{
-		if (!ParentNode.Get()->IsExpanded())
-		{
-			ParentNode.Get()->SetExpansionState(true);
-		}
-
 		OutFilteredNodes.Add(ParentNode.ToSharedRef());
 		ParentNode = ParentNode->GetParent();
 	}
 }
 
-static void AddParentNodes(const TSharedRef<FSequencerDisplayNode>& StartNode, TSet<TSharedRef<FSequencerDisplayNode>>& OutFilteredNodes)
-{
-	TSharedPtr<FSequencerDisplayNode> ParentNode = StartNode->GetParent();
-	if (ParentNode.IsValid())
-	{
-		if (!ParentNode.Get()->IsExpanded())
-		{
-			ParentNode.Get()->SetExpansionState(true);
-		}
-
-		OutFilteredNodes.Add(ParentNode.ToSharedRef());
-		AddParentNodes(ParentNode.ToSharedRef(), OutFilteredNodes);
-	}
-}
-
 /**
  * Recursively filters nodes
  *
@@ -808,6 +783,17 @@ static void FilterNodesRecursive( FSequencer& Sequencer, const TSharedRef<FSeque
 					return;
 				}
 			}
+
+			TSharedPtr<FSequencerDisplayNode> ParentNode = StartNode->GetParent();
+			while(ParentNode.IsValid())
+			{
+				if (!ParentNode.Get()->IsExpanded())
+				{
+					ParentNode.Get()->SetExpansionState(true);
+				}
+
+				ParentNode = ParentNode->GetParent();
+			}
 		}
 		AddFilteredNode(StartNode, OutFilteredNodes);
 	}
