Sequencer: Improvements to accessing breadcrumbs when sequencer toolbar is too narrow

diff --git a/Engine/Plugins/MovieScene/LevelSequenceEditor/Source/LevelSequenceEditor/Private/Misc/LevelSequencePlaybackContext.cpp b/Engine/Plugins/MovieScene/LevelSequenceEditor/Source/LevelSequenceEditor/Private/Misc/LevelSequencePlaybackContext.cpp
--- a/Engine/Plugins/MovieScene/LevelSequenceEditor/Source/LevelSequenceEditor/Private/Misc/LevelSequencePlaybackContext.cpp
+++ b/Engine/Plugins/MovieScene/LevelSequenceEditor/Source/LevelSequenceEditor/Private/Misc/LevelSequencePlaybackContext.cpp
@@ -239,27 +239,11 @@ void SLevelSequenceContextPicker::Construct(const FArguments& InArgs)
 			.ForegroundColor(FSlateColor::UseForeground())
 			.ButtonStyle(FEditorStyle::Get(), "ToggleButton")
 			.OnGetMenuContent(this, &SLevelSequenceContextPicker::BuildWorldPickerMenu)
-			.ToolTipText(LOCTEXT("WorldPickerText", "The world context that sequencer should be bound to, and playback within."))
+			.ToolTipText(FText::Format(LOCTEXT("WorldPickerTextFomrat", "'{0}': The world context that sequencer should be bound to, and playback within."), GetCurrentContextText()))
 			.ButtonContent()
 			[
-				SNew(SHorizontalBox)
-
-				+SHorizontalBox::Slot()
-				.AutoWidth()
-				.VAlign(VAlign_Center)
-				[
-					SNew(SImage)
-					.Image(FEditorStyle::GetBrush("SceneOutliner.World"))
-				]
-
-				+SHorizontalBox::Slot()
-				.AutoWidth()
-				.Padding(2, 0, 0, 0)
-				.VAlign(VAlign_Center)
-				[
-					SNew(STextBlock)
-					.Text(this, &SLevelSequenceContextPicker::GetCurrentContextText)
-				]
+				SNew(SImage)
+				.Image(FEditorStyle::GetBrush("SceneOutliner.World"))
 			]
 		]
 	];
diff --git a/Engine/Source/Editor/Sequencer/Private/SSequencer.cpp b/Engine/Source/Editor/Sequencer/Private/SSequencer.cpp
--- a/Engine/Source/Editor/Sequencer/Private/SSequencer.cpp
+++ b/Engine/Source/Editor/Sequencer/Private/SSequencer.cpp
@@ -21,6 +21,7 @@
 #include "Widgets/Layout/SSpacer.h"
 #include "Widgets/Images/SImage.h"
 #include "Widgets/Layout/SBox.h"
+#include "Widgets/Layout/SWrapBox.h"
 #include "Widgets/Layout/SGridPanel.h"
 #include "Widgets/Layout/SScrollBar.h"
 #include "Widgets/Layout/SScrollBorder.h"
@@ -522,54 +523,81 @@ void SSequencer::Construct(const FArguments& InArgs, TSharedRef<FSequencer> InSe
 						.BorderImage(FEditorStyle::GetBrush("ToolPanel.GroupBorder"))
 						.Padding(FMargin(CommonPadding, 0.f))
 						[
-							SNew(SHorizontalBox)
+							SNew(SWrapBox)
+							.UseAllottedWidth(true)
+							.InnerSlotPadding(FVector2D(5, 2))
 
-							// Left aligned Toolbar Icons
-							+SHorizontalBox::Slot()
-							.AutoWidth()
+							+ SWrapBox::Slot()
+							.FillEmptySpace(true)
+							.FillLineWhenWidthLessThan(600)
 							[
 								MakeToolBar()
 							]
 
-							// Right Aligned Breadcrumbs
-							+ SHorizontalBox::Slot()
-							.HAlign(HAlign_Right)
-							.VAlign(VAlign_Center)
+							+ SWrapBox::Slot()
+							.FillEmptySpace(true)
 							[
-								SNew(SSpacer)
-							]
+								SNew(SBorder)
+								.Padding(FMargin(3))
+								.BorderImage(FEditorStyle::GetBrush("ToolPanel.GroupBorder"))
+								[
+									SNew(SHorizontalBox)
 
-							+ SHorizontalBox::Slot()
-							.HAlign(HAlign_Right)
-							.VAlign(VAlign_Center)
-							[
-								SAssignNew(BreadcrumbTrail, SBreadcrumbTrail<FSequencerBreadcrumb>)
-								.Visibility(this, &SSequencer::GetBreadcrumbTrailVisibility)
-								.OnCrumbClicked(this, &SSequencer::OnCrumbClicked)
-								.ButtonStyle(FEditorStyle::Get(), "FlatButton")
-								.DelimiterImage(FEditorStyle::GetBrush("Sequencer.BreadcrumbIcon"))
-								.TextStyle(FEditorStyle::Get(), "Sequencer.BreadcrumbText")
-							]
+									+ SHorizontalBox::Slot()
+									.HAlign(HAlign_Right)
+									.VAlign(VAlign_Center)
+									[
+										SAssignNew(BreadcrumbPickerButton, SComboButton)
+										.Visibility(this, &SSequencer::GetBreadcrumbTrailVisibility)
+										.ButtonStyle(FEditorStyle::Get(), "FlatButton")
+										.ForegroundColor(FLinearColor::White)
+										.OnGetMenuContent(this, &SSequencer::GetBreadcrumbPickerContent)
+										.HasDownArrow(false)
+										.ContentPadding(FMargin(3, 3))
+										.ButtonContent()
+										[
+											SNew(STextBlock)
+											.TextStyle(FEditorStyle::Get(), "Sequencer.BreadcrumbText")
+											.Font(FEditorStyle::Get().GetFontStyle("FontAwesome.11"))
+											.Text(FText::FromString(FString(TEXT("\xf07c"))) /*fa-folder-open*/)
+										]
+									]
 
-							// Sequence Locking symbol
-							+SHorizontalBox::Slot()
-							.HAlign(HAlign_Right)
-							.VAlign(VAlign_Center)
-							.AutoWidth()
-							.Padding(2, 0, 0, 0)
-							[
-								SNew(SCheckBox)
-								.IsFocusable(false)		
-								.IsChecked_Lambda([this] { return GetIsSequenceReadOnly() ? ECheckBoxState::Checked : ECheckBoxState::Unchecked; } )
-								.OnCheckStateChanged(this, &SSequencer::OnSetSequenceReadOnly)
-								.ToolTipText_Lambda([this] { return GetIsSequenceReadOnly() ? LOCTEXT("UnlockSequence", "Unlock the animation so that it is editable") : LOCTEXT("LockSequence", "Lock the animation so that it is not editable"); } )
-								.ForegroundColor(FLinearColor::White)
-								.CheckedImage(FEditorStyle::GetBrush("Sequencer.LockSequence"))
-								.CheckedHoveredImage(FEditorStyle::GetBrush("Sequencer.LockSequence"))
-								.CheckedPressedImage(FEditorStyle::GetBrush("Sequencer.LockSequence"))
-								.UncheckedImage(FEditorStyle::GetBrush("Sequencer.UnlockSequence"))
-								.UncheckedHoveredImage(FEditorStyle::GetBrush("Sequencer.UnlockSequence"))
-								.UncheckedPressedImage(FEditorStyle::GetBrush("Sequencer.UnlockSequence"))
+									// Right Aligned Breadcrumbs
+									+ SHorizontalBox::Slot()
+									.HAlign(HAlign_Right)
+									.VAlign(VAlign_Center)
+									.AutoWidth()
+									[
+										SAssignNew(BreadcrumbTrail, SBreadcrumbTrail<FSequencerBreadcrumb>)
+										.Visibility(this, &SSequencer::GetBreadcrumbTrailVisibility)
+										.OnCrumbClicked(this, &SSequencer::OnCrumbClicked)
+										.ButtonStyle(FEditorStyle::Get(), "FlatButton")
+										.DelimiterImage(FEditorStyle::GetBrush("Sequencer.BreadcrumbIcon"))
+										.TextStyle(FEditorStyle::Get(), "Sequencer.BreadcrumbText")
+									]
+
+									// Sequence Locking symbol
+									+SHorizontalBox::Slot()
+									.HAlign(HAlign_Right)
+									.VAlign(VAlign_Center)
+									.AutoWidth()
+									.Padding(2, 0, 0, 0)
+									[
+										SNew(SCheckBox)
+										.IsFocusable(false)		
+										.IsChecked_Lambda([this] { return GetIsSequenceReadOnly() ? ECheckBoxState::Checked : ECheckBoxState::Unchecked; } )
+										.OnCheckStateChanged(this, &SSequencer::OnSetSequenceReadOnly)
+										.ToolTipText_Lambda([this] { return GetIsSequenceReadOnly() ? LOCTEXT("UnlockSequence", "Unlock the animation so that it is editable") : LOCTEXT("LockSequence", "Lock the animation so that it is not editable"); } )
+										.ForegroundColor(FLinearColor::White)
+										.CheckedImage(FEditorStyle::GetBrush("Sequencer.LockSequence"))
+										.CheckedHoveredImage(FEditorStyle::GetBrush("Sequencer.LockSequence"))
+										.CheckedPressedImage(FEditorStyle::GetBrush("Sequencer.LockSequence"))
+										.UncheckedImage(FEditorStyle::GetBrush("Sequencer.UnlockSequence"))
+										.UncheckedHoveredImage(FEditorStyle::GetBrush("Sequencer.UnlockSequence"))
+										.UncheckedPressedImage(FEditorStyle::GetBrush("Sequencer.UnlockSequence"))
+									]
+								]
 							]
 						]
 					]
@@ -2273,7 +2301,7 @@ void SSequencer::UpdateBreadcrumbs()
 		TAttribute<FText> CrumbNameAttribute = MakeAttributeSP(this, &SSequencer::GetBreadcrumbTextForSection, SubSection);
 
 		// The current breadcrumb is not a moviescene so we need to make a new breadcrumb in order return to the parent moviescene later
-		BreadcrumbTrail->PushCrumb( CrumbNameAttribute, FSequencerBreadcrumb( FocusedID ) );
+		BreadcrumbTrail->PushCrumb( CrumbNameAttribute, FSequencerBreadcrumb( FocusedID, CrumbNameAttribute.Get()) );
 	}
 }
 
@@ -2283,7 +2311,7 @@ void SSequencer::ResetBreadcrumbs()
 	BreadcrumbTrail->ClearCrumbs();
 
 	TAttribute<FText> CrumbNameAttribute = MakeAttributeSP(this, &SSequencer::GetBreadcrumbTextForSequence, MakeWeakObjectPtr(SequencerPtr.Pin()->GetRootMovieSceneSequence()), true);
-	BreadcrumbTrail->PushCrumb(CrumbNameAttribute, FSequencerBreadcrumb(MovieSceneSequenceID::Root));
+	BreadcrumbTrail->PushCrumb(CrumbNameAttribute, FSequencerBreadcrumb(MovieSceneSequenceID::Root, CrumbNameAttribute.Get()));
 }
 
 void SSequencer::PopBreadcrumb()
@@ -2546,6 +2574,38 @@ void SSequencer::OnCrumbClicked(const FSequencerBreadcrumb& Item)
 	}
 }
 
+void SSequencer::OnBreadcrumbPickerContentClicked(const FSequencerBreadcrumb& Breadcrumb)
+{
+	while (BreadcrumbTrail->NumCrumbs() > 1 && BreadcrumbTrail->PeekCrumb().SequenceID != Breadcrumb.SequenceID)
+	{
+		BreadcrumbTrail->PopCrumb();
+	}
+	OnCrumbClicked(Breadcrumb);
+}
+
+TSharedRef<SWidget> SSequencer::GetBreadcrumbPickerContent()
+{
+	TArray<FSequencerBreadcrumb> CrumbData;
+	BreadcrumbTrail->GetAllCrumbData(CrumbData);
+
+	FMenuBuilder MenuBuilder(true, nullptr);
+
+	MenuBuilder.BeginSection("SequencerBreadcrumbPicker");
+
+	for(FSequencerBreadcrumb& Breadcrumb : CrumbData)
+	{
+		MenuBuilder.AddMenuEntry(
+			Breadcrumb.BreadcrumbName,
+			FText::GetEmpty(),
+			FSlateIcon(),
+			FUIAction(FExecuteAction::CreateLambda(([this, Breadcrumb]() { return this->OnBreadcrumbPickerContentClicked(Breadcrumb); } )))
+		);
+	}
+	
+	MenuBuilder.EndSection();
+
+	return MenuBuilder.MakeWidget();
+}
 
 FText SSequencer::GetRootAnimationName() const
 {
diff --git a/Engine/Source/Editor/Sequencer/Private/SSequencer.h b/Engine/Source/Editor/Sequencer/Private/SSequencer.h
--- a/Engine/Source/Editor/Sequencer/Private/SSequencer.h
+++ b/Engine/Source/Editor/Sequencer/Private/SSequencer.h
@@ -77,13 +77,18 @@ struct FSequencerBreadcrumb
 	/** The movie scene this may point to */
 	FMovieSceneSequenceID SequenceID;
 
-	FSequencerBreadcrumb(FMovieSceneSequenceIDRef InSequenceID)
+	/** The display name of this breadcrumb */
+	FText BreadcrumbName;
+
+	FSequencerBreadcrumb(FMovieSceneSequenceIDRef InSequenceID, FText CrumbName)
 		: BreadcrumbType(FSequencerBreadcrumb::MovieSceneType)
 		, SequenceID(InSequenceID)
+		, BreadcrumbName(CrumbName)
 	{ }
 
-	FSequencerBreadcrumb()
+	FSequencerBreadcrumb(FText CrumbName)
 		: BreadcrumbType(FSequencerBreadcrumb::ShotType)
+		, BreadcrumbName(CrumbName)
 	{ }
 };
 
@@ -451,6 +456,11 @@ private:
 	/** Called when a breadcrumb is clicked on in the sequencer */
 	void OnCrumbClicked(const FSequencerBreadcrumb& Item);
 
+	void OnBreadcrumbPickerContentClicked(const FSequencerBreadcrumb& Breadcrumb);
+
+	/** Called when the user opens the breadcrumb dropdown */
+	TSharedRef<SWidget> GetBreadcrumbPickerContent();
+
 	/** Gets the root movie scene name */
 	FText GetRootAnimationName() const;
 
@@ -550,6 +560,9 @@ private:
 	/** The sequencer tree view responsible for the outliner and track areas */
 	TSharedPtr<SSequencerTreeView> TreeView;
 
+	/** Dropdown for selecting breadcrumbs */
+	TSharedPtr<class SComboButton> BreadcrumbPickerButton;
+
 	/** The main sequencer interface */
 	TWeakPtr<FSequencer> SequencerPtr;
 
