Sequencer: Disallow setting parent on itself

diff --git a/Engine/Source/Editor/Sequencer/Private/Sequencer.cpp b/Engine/Source/Editor/Sequencer/Private/Sequencer.cpp
--- a/Engine/Source/Editor/Sequencer/Private/Sequencer.cpp
+++ b/Engine/Source/Editor/Sequencer/Private/Sequencer.cpp
@@ -6949,7 +6949,7 @@ bool FSequencer::PasteObjectBindings(const FString& TextToImport, TArray<FNotifi
 	for (auto PossessableGuid : PossessableGuids)
 	{
 		FMovieScenePossessable* Possessable = MovieScene->FindPossessable(PossessableGuid);
-		if (Possessable && OldToNewGuidMap.Contains(Possessable->GetParent()))
+		if (Possessable && OldToNewGuidMap.Contains(Possessable->GetParent()) && PossessableGuid != OldToNewGuidMap[Possessable->GetParent()])
 		{
 			Possessable->SetParent(OldToNewGuidMap[Possessable->GetParent()]);
 		}
