Sequencer: Add support for pinning a track

diff --git a/Engine/Source/Editor/Sequencer/Private/DisplayNodes/SequencerDisplayNode.cpp b/Engine/Source/Editor/Sequencer/Private/DisplayNodes/SequencerDisplayNode.cpp
--- a/Engine/Source/Editor/Sequencer/Private/DisplayNodes/SequencerDisplayNode.cpp
+++ b/Engine/Source/Editor/Sequencer/Private/DisplayNodes/SequencerDisplayNode.cpp
@@ -399,6 +399,7 @@ FSequencerDisplayNode::FSequencerDisplayNode( FName InNodeName, FSequencerNodeTr
 	, ParentTree( InParentTree )
 	, NodeName( InNodeName )
 	, bExpanded( false )
+	, bPinned(false)
 	, bHasBeenInitialized( false )
 {
 	SortType = EDisplayNodeSortType::Undefined;
@@ -966,6 +967,22 @@ void FSequencerDisplayNode::BuildContextMenu(FMenuBuilder& MenuBuilder)
 			EUserInterfaceActionType::ToggleButton
 		);
 
+		// Only support pinning root nodes
+		if (GetBaseNode()->IsRootNode())
+		{
+			MenuBuilder.AddMenuEntry(
+				LOCTEXT("ToggleNodePin", "Pinned"),
+				LOCTEXT("ToggleNodePinTooltip", "Pin or unpin this node or selected tracks"),
+				FSlateIcon(),
+				FUIAction(
+					FExecuteAction::CreateSP(this, &FSequencerDisplayNode::TogglePinned),
+					FCanExecuteAction(),
+					FIsActionChecked::CreateSP(this, &FSequencerDisplayNode::IsPinned)
+				),
+				NAME_None,
+				EUserInterfaceActionType::ToggleButton
+			);
+		}
 
 		// Add cut, copy and paste functions to the tracks
 		MenuBuilder.AddMenuEntry(FGenericCommands::Get().Cut);
@@ -1098,6 +1115,42 @@ bool FSequencerDisplayNode::IsExpanded() const
 	return bExpanded;
 }
 
+FSequencerDisplayNode* FSequencerDisplayNode::GetBaseNode() const
+{
+	ESequencerNode::Type Type = GetType();
+
+	if (IsRootNode() || Type == ESequencerNode::Folder || Type == ESequencerNode::Object || GetParentOrRoot()->GetType() == ESequencerNode::Folder)
+	{
+		return (FSequencerDisplayNode*)this;
+	}
+
+	return GetParentOrRoot()->GetBaseNode();
+}
+
+bool FSequencerDisplayNode::IsPinned() const
+{
+	return GetBaseNode()->bPinned;
+}
+
+void FSequencerDisplayNode::TogglePinned()
+{
+	FSequencerDisplayNode* BaseNode = GetBaseNode();
+	bool bShouldPin = !BaseNode->bPinned;
+	ParentTree.UnpinAllNodes();
+	BaseNode->bPinned = bShouldPin;
+	ParentTree.GetSequencer().RefreshTree();
+}
+
+void FSequencerDisplayNode::Unpin()
+{
+	FSequencerDisplayNode* BaseNode = GetBaseNode();
+	if (BaseNode->bPinned)
+	{
+		BaseNode->bPinned = false;
+		ParentTree.GetSequencer().RefreshTree();
+	}
+}
+
 
 bool FSequencerDisplayNode::IsHidden() const
 {
diff --git a/Engine/Source/Editor/Sequencer/Private/DisplayNodes/SequencerDisplayNode.h b/Engine/Source/Editor/Sequencer/Private/DisplayNodes/SequencerDisplayNode.h
--- a/Engine/Source/Editor/Sequencer/Private/DisplayNodes/SequencerDisplayNode.h
+++ b/Engine/Source/Editor/Sequencer/Private/DisplayNodes/SequencerDisplayNode.h
@@ -426,6 +426,21 @@ public:
 	 */
 	bool IsExpanded() const;
 
+	/**
+	 * @return Whether or not this node is pinned
+	 */
+	bool IsPinned() const;
+
+	/**
+	 * Toggle whether or not this node is pinned
+	 */
+	void TogglePinned();
+
+	/**
+	 * If this node is pinned, unpin it.
+	 */
+	void Unpin();
+
 	/**
 	 * @return Whether this node is explicitly hidden from the view or not
 	 */
@@ -507,6 +522,11 @@ public:
 
 private:
 
+	/**
+	 * @return The base node this node belongs to, for collections of tracks that are part of an object
+	 */
+	FSequencerDisplayNode* GetBaseNode() const;
+
 	/** Callback for executing a "Rename Node" context menu action. */
 	void HandleContextMenuRenameNodeExecute();
 
@@ -538,6 +558,9 @@ protected:
 	/** Whether or not the node is expanded */
 	bool bExpanded;
 
+	/** Whether or not the node is pinned */
+	bool bPinned;
+
 	/** Event that is triggered when rename is requested */
 	FRequestRenameEvent RenameRequestedEvent;
 
diff --git a/Engine/Source/Editor/Sequencer/Private/SSequencer.cpp b/Engine/Source/Editor/Sequencer/Private/SSequencer.cpp
--- a/Engine/Source/Editor/Sequencer/Private/SSequencer.cpp
+++ b/Engine/Source/Editor/Sequencer/Private/SSequencer.cpp
@@ -399,6 +399,16 @@ void SSequencer::Construct(const FArguments& InArgs, TSharedRef<FSequencer> InSe
 	TSharedRef<SScrollBar> ScrollBar = SNew(SScrollBar)
 		.Thickness(FVector2D(9.0f, 9.0f));
 	SAssignNew(TrackOutliner, SSequencerTrackOutliner);
+
+	SAssignNew(PinnedTrackArea, SSequencerTrackArea, TimeSliderControllerRef, InSequencer);
+	SAssignNew(PinnedTreeView, SSequencerTreeView, InSequencer->GetNodeTree(), PinnedTrackArea.ToSharedRef())
+		.Clipping(EWidgetClipping::ClipToBounds)
+		.OnGetContextMenuContent(FOnGetContextMenuContent::CreateSP(this, &SSequencer::GetContextMenuContent));
+
+	PinnedTrackArea->SetTreeView(PinnedTreeView);
+	PinnedTrackArea->SetShowPinned(true);
+	PinnedTreeView->SetShowPinned(true);
+
 	SAssignNew(TrackArea, SSequencerTrackArea, TimeSliderControllerRef, InSequencer);
 	SAssignNew(TreeView, SSequencerTreeView, InSequencer->GetNodeTree(), TrackArea.ToSharedRef())
 		.ExternalScrollbar(ScrollBar)
@@ -407,6 +417,8 @@ void SSequencer::Construct(const FArguments& InArgs, TSharedRef<FSequencer> InSe
 
 	TrackArea->SetTreeView(TreeView);
 
+	TreeView->AddSlaveTreeView(PinnedTreeView);
+
 	TAttribute<FAnimatedRange> ViewRangeAttribute = InArgs._ViewRange;
 	
 	// We create a custom Time Slider Controller which is just a wrapper around the actual one, but is aware of our custom bounds logic. Currently the range the
@@ -688,29 +700,67 @@ void SSequencer::Construct(const FArguments& InArgs, TSharedRef<FSequencer> InSe
 
 							+ SOverlay::Slot()
 							[
-								SNew(SScrollBorder, TreeView.ToSharedRef())
-								[
-									SNew(SHorizontalBox)
+								SNew(SVerticalBox)
 
-									// outliner tree
-									+ SHorizontalBox::Slot()
-									.FillWidth( FillCoefficient_0 )
+								+SVerticalBox::Slot()
+								.AutoHeight()
+								[
+									SNew(SBorder)
+									.Padding(FMargin(0.0f, 0.0f, 0.0f, CommonPadding))
 									[
-										SNew(SBox)
+										SNew(SHorizontalBox)
+
+										// outliner tree
+										+ SHorizontalBox::Slot()
+										.FillWidth( FillCoefficient_0 )
+										[
+											SNew(SBox)
+											[
+												PinnedTreeView.ToSharedRef()
+											]
+										]
+
+										// track area
+										+ SHorizontalBox::Slot()
+										.FillWidth( FillCoefficient_1 )
 										[
-											TreeView.ToSharedRef()
+											SNew(SBox)
+											.Padding(ResizeBarPadding)
+											.Clipping(EWidgetClipping::ClipToBounds)
+											[
+												PinnedTrackArea.ToSharedRef()
+											]
 										]
 									]
 
-									// track area
-									+ SHorizontalBox::Slot()
-									.FillWidth( FillCoefficient_1 )
+								]
+
+								+SVerticalBox::Slot()
+								[
+									SNew(SScrollBorder, TreeView.ToSharedRef())
 									[
-										SNew(SBox)
-										.Padding(ResizeBarPadding)
-										.Clipping(EWidgetClipping::ClipToBounds)
+										SNew(SHorizontalBox)
+
+										// outliner tree
+										+ SHorizontalBox::Slot()
+										.FillWidth( FillCoefficient_0 )
 										[
-											TrackArea.ToSharedRef()
+											SNew(SBox)
+											[
+												TreeView.ToSharedRef()
+											]
+										]
+
+										// track area
+										+ SHorizontalBox::Slot()
+										.FillWidth( FillCoefficient_1 )
+										[
+											SNew(SBox)
+											.Padding(ResizeBarPadding)
+											.Clipping(EWidgetClipping::ClipToBounds)
+											[
+												TrackArea.ToSharedRef()
+											]
 										]
 									]
 								]
@@ -2276,6 +2326,7 @@ void RestoreKeySelection(const TSet<FSequencerSelectedKey>& OldKeys, FSequencerS
 void SSequencer::UpdateLayoutTree()
 {
 	TrackArea->Empty();
+	PinnedTrackArea->Empty();
 
 	TSharedPtr<FSequencer> Sequencer = SequencerPtr.Pin();
 	if ( Sequencer.IsValid() )
@@ -2330,6 +2381,7 @@ void SSequencer::UpdateLayoutTree()
 					while (Parent.IsValid())
 					{
 						TreeView->SetItemExpansion(Parent->AsShared(), true);
+						PinnedTreeView->SetItemExpansion(Parent->AsShared(), true);
 						Parent = Parent->GetParent();
 					}
 
@@ -2979,9 +3031,18 @@ FPaintPlaybackRangeArgs SSequencer::GetSectionPlaybackRangeArgs() const
 }
 
 
-FVirtualTrackArea SSequencer::GetVirtualTrackArea() const
+FVirtualTrackArea SSequencer::GetVirtualTrackArea(const SSequencerTrackArea* InTrackArea) const
 {
-	return FVirtualTrackArea(*SequencerPtr.Pin(), *TreeView.Get(), TrackArea->GetCachedGeometry());
+	const SSequencerTrackArea* TargetTrackArea = TrackArea.Get();
+	TSharedPtr<SSequencerTreeView> TargetTreeView = TreeView;
+	
+	if (InTrackArea != nullptr)
+	{
+		TargetTrackArea = InTrackArea;
+		TargetTreeView = TargetTrackArea->GetTreeView().Pin();
+	}
+
+	return FVirtualTrackArea(*SequencerPtr.Pin(), *TargetTreeView.Get(), TargetTrackArea->GetCachedGeometry());
 }
 
 FPasteContextMenuArgs SSequencer::GeneratePasteArgs(FFrameNumber PasteAtTime, TSharedPtr<FMovieSceneClipboard> Clipboard)
diff --git a/Engine/Source/Editor/Sequencer/Private/SSequencer.h b/Engine/Source/Editor/Sequencer/Private/SSequencer.h
--- a/Engine/Source/Editor/Sequencer/Private/SSequencer.h
+++ b/Engine/Source/Editor/Sequencer/Private/SSequencer.h
@@ -277,8 +277,12 @@ public:
 	/** Access the tree view for this sequencer */
 	TSharedPtr<SSequencerTreeView> GetTreeView() const;
 
-	/** Generate a helper structure that can be used to transform between phsyical space and virtual space in the track area */
-	FVirtualTrackArea GetVirtualTrackArea() const;
+	/** 
+	 * Generate a helper structure that can be used to transform between phsyical space and virtual space in the track area
+	 *
+	 * @param InTrackArea	(optional) The track area to generate helper structure for, if not specified the main track area will be used.
+	 */
+	FVirtualTrackArea GetVirtualTrackArea(const SSequencerTrackArea* InTrackArea = nullptr) const;
 
 	/** Access this widget's track area widget */
 	TSharedPtr<SSequencerTrackArea> GetTrackAreaWidget() const { return TrackArea; }
@@ -537,6 +541,9 @@ private:
 	/** Section area widget */
 	TSharedPtr<SSequencerTrackArea> TrackArea;
 
+	/** Section area widget for pinned tracks*/
+	TSharedPtr<SSequencerTrackArea> PinnedTrackArea;
+
 	/** Outliner widget */
 	TSharedPtr<SSequencerTrackOutliner> TrackOutliner;
 
@@ -564,6 +571,9 @@ private:
 	/** The sequencer tree view responsible for the outliner and track areas */
 	TSharedPtr<SSequencerTreeView> TreeView;
 
+	/** The sequencer tree view for pinned tracks */
+	TSharedPtr<SSequencerTreeView> PinnedTreeView;
+
 	/** Dropdown for selecting breadcrumbs */
 	TSharedPtr<class SComboButton> BreadcrumbPickerButton;
 
diff --git a/Engine/Source/Editor/Sequencer/Private/SSequencerTrackArea.cpp b/Engine/Source/Editor/Sequencer/Private/SSequencerTrackArea.cpp
--- a/Engine/Source/Editor/Sequencer/Private/SSequencerTrackArea.cpp
+++ b/Engine/Source/Editor/Sequencer/Private/SSequencerTrackArea.cpp
@@ -43,6 +43,8 @@ void SSequencerTrackArea::Construct(const FArguments& InArgs, TSharedRef<FSequen
 	Sequencer = InSequencer;
 	TimeSliderController = InTimeSliderController;
 
+	bShowPinnedNodes = false;
+
 	// Input stack in order or priority
 
 	// Space for the edit tool
@@ -50,7 +52,7 @@ void SSequencerTrackArea::Construct(const FArguments& InArgs, TSharedRef<FSequen
 	// The time slider controller
 	InputStack.AddHandler(TimeSliderController.Get());
 
-	EditTools.Add(MakeShared<FSequencerEditTool_Selection>(*InSequencer));
+	EditTools.Add(MakeShared<FSequencerEditTool_Selection>(*InSequencer, *this));
 	EditTools.Add(MakeShared<FSequencerEditTool_Movement>(*InSequencer));
 }
 
@@ -91,6 +93,12 @@ void SSequencerTrackArea::OnArrangeChildren( const FGeometry& AllottedGeometry,
 			continue;
 		}
 
+		TSharedPtr<SSequencerTrackLane> PinnedTrackLane = CurChild.TrackLane.Pin();
+		if (PinnedTrackLane->IsPinned() != bShowPinnedNodes)
+		{
+			continue;
+		}
+
 		const FMargin Padding(0, CurChild.GetVerticalOffset(), 0, 0);
 
 		AlignmentArrangeResult XResult = AlignChild<Orient_Horizontal>(AllottedGeometry.GetLocalSize().X, CurChild, Padding, 1.0f, false);
diff --git a/Engine/Source/Editor/Sequencer/Private/SSequencerTrackArea.h b/Engine/Source/Editor/Sequencer/Private/SSequencerTrackArea.h
--- a/Engine/Source/Editor/Sequencer/Private/SSequencerTrackArea.h
+++ b/Engine/Source/Editor/Sequencer/Private/SSequencerTrackArea.h
@@ -79,6 +79,7 @@ public:
 
 	/** Assign a tree view to this track area. */
 	void SetTreeView(const TSharedPtr<SSequencerTreeView>& InTreeView);
+	TWeakPtr<SSequencerTreeView> GetTreeView() const { return TreeView; }
 
 	/** Access the currently active track area edit tool */
 	const ISequencerEditTool* GetEditTool() const { return EditTool.IsValid() ? EditTool.Get() : nullptr; }
@@ -86,6 +87,10 @@ public:
 	/** Attempt to activate the tool specified by the identifier */
 	bool AttemptToActivateTool(FName Identifier);
 
+	/** Set whether this TrackArea should show only pinned nodes or only non-pinned nodes  */
+	void SetShowPinned(bool bShowPinned) { bShowPinnedNodes = bShowPinned; }
+	bool ShowPinned() const { return bShowPinnedNodes; }
+
 public:
 
 	// SWidget interface
@@ -150,6 +155,9 @@ private:
 	/** Whether the dropped node is allowed to be dropped onto */
 	bool bAllowDrop;
 
+	/** Whether this TrackArea is for pinned nodes or non-pinned nodes */
+	bool bShowPinnedNodes;
+
 private:
 
 	/** Input handler stack responsible for routing input to the different handlers */
diff --git a/Engine/Source/Editor/Sequencer/Private/SSequencerTrackLane.h b/Engine/Source/Editor/Sequencer/Private/SSequencerTrackLane.h
--- a/Engine/Source/Editor/Sequencer/Private/SSequencerTrackLane.h
+++ b/Engine/Source/Editor/Sequencer/Private/SSequencerTrackLane.h
@@ -31,6 +31,8 @@ public:
 	/** Get the desired physical vertical position of this track lane */
 	float GetPhysicalPosition() const;
 
+	bool IsPinned() const { return DisplayNode->IsPinned(); }
+
 protected:
 	virtual FVector2D ComputeDesiredSize(float LayoutScale) const override;
 
diff --git a/Engine/Source/Editor/Sequencer/Private/SSequencerTreeView.cpp b/Engine/Source/Editor/Sequencer/Private/SSequencerTreeView.cpp
--- a/Engine/Source/Editor/Sequencer/Private/SSequencerTreeView.cpp
+++ b/Engine/Source/Editor/Sequencer/Private/SSequencerTreeView.cpp
@@ -9,6 +9,7 @@
 #include "Framework/Commands/GenericCommands.h"
 #include "Framework/MultiBox/MultiBoxBuilder.h"
 #include "Framework/Application/SlateApplication.h"
+#include "SSequencerTreeView.h"
 
 static FName TrackAreaName = "TrackArea";
 
@@ -145,6 +146,7 @@ void SSequencerTreeView::Construct(const FArguments& InArgs, const TSharedRef<FS
 	bSequencerSelectionChangeBroadcastWasSupressed = false;
 	bPhysicalNodesNeedUpdate = false;
 	bRightMouseButtonDown = false;
+	bShowPinnedNodes = false;
 
 	// We 'leak' these delegates (they'll get cleaned up automatically when the invocation list changes)
 	// It's not safe to attempt their removal in ~SSequencerTreeView because SequencerNodeTree->GetSequencer() may not be valid
@@ -314,8 +316,21 @@ void SSequencerTreeView::ReportChildRowGeometry(const FDisplayNodeRef& InNode, c
 		FVector2D(0,0)
 	).Y;
 
-	CachedRowGeometry.Add(InNode, FCachedGeometry(InNode, ChildOffset, InGeometry.Size.Y));
+	if (InNode->IsPinned() != bShowPinnedNodes)
+	{
+		CachedRowGeometry.Remove(InNode);
+	}
+	else
+	{
+		CachedRowGeometry.Add(InNode, FCachedGeometry(InNode, ChildOffset, InGeometry.Size.Y));
+	}
+
 	bPhysicalNodesNeedUpdate = true;
+
+	for (TSharedPtr<SSequencerTreeView> SlaveTreeView : SlaveTreeViews)
+	{
+		SlaveTreeView->ReportChildRowGeometry(InNode, InGeometry);
+	}
 }
 
 void SSequencerTreeView::OnChildRowRemoved(const FDisplayNodeRef& InNode)
@@ -419,6 +434,11 @@ void SSequencerTreeView::UpdateTrackArea()
 	}
 }
 
+void SSequencerTreeView::AddSlaveTreeView(TSharedPtr<SSequencerTreeView> SlaveTreeView)
+{
+	SlaveTreeViews.Add(SlaveTreeView);
+}
+
 void SSequencerTreeView::OnRightMouseButtonDown(const FPointerEvent& MouseEvent)
 {
 	STreeView::OnRightMouseButtonDown(MouseEvent);
@@ -635,6 +655,11 @@ void SSequencerTreeView::SynchronizeTreeSelectionWithSequencerSelection()
 		}
 		bUpdatingTreeSelection = false;
 	}
+
+	for (TSharedPtr<SSequencerTreeView> SlaveTreeView : SlaveTreeViews)
+	{
+		SlaveTreeView->SynchronizeTreeSelectionWithSequencerSelection();
+	}
 }
 
 void SSequencerTreeView::Private_SetItemSelection( FDisplayNodeRef TheItem, bool bShouldBeSelected, bool bWasUserDirected )
@@ -758,8 +783,15 @@ TSharedPtr<SWidget> SSequencerTreeView::OnContextMenuOpening()
 void SSequencerTreeView::Refresh()
 {
 	RootNodes.Reset();
-	Algo::CopyIf(SequencerNodeTree->GetRootNodes(), RootNodes, &FSequencerDisplayNode::IsVisible);
 
+	for (auto& Item : SequencerNodeTree->GetRootNodes())
+	{
+		if (Item->IsVisible() && Item->IsPinned() == bShowPinnedNodes)
+		{
+			RootNodes.Add(Item);
+		}
+	}
+	
 	// Reset item expansion since we don't know if any expansion states may have changed in-between refreshes
 	{
 		STreeView::OnExpansionChanged.Unbind();
@@ -783,6 +815,11 @@ void SSequencerTreeView::Refresh()
 	bUpdatingTreeSelection = false;
 
 	RebuildList();
+
+	for (TSharedPtr<SSequencerTreeView> SlaveTreeView : SlaveTreeViews)
+	{
+		SlaveTreeView->Refresh();
+	}
 }
 
 void SSequencerTreeView::ScrollByDelta(float DeltaInSlateUnits)
diff --git a/Engine/Source/Editor/Sequencer/Private/SSequencerTreeView.h b/Engine/Source/Editor/Sequencer/Private/SSequencerTreeView.h
--- a/Engine/Source/Editor/Sequencer/Private/SSequencerTreeView.h
+++ b/Engine/Source/Editor/Sequencer/Private/SSequencerTreeView.h
@@ -183,6 +183,12 @@ public:
 	/** Ensure that the track area column is either show or hidden, depending on the visibility of the curve editor */
 	void UpdateTrackArea();
 
+	/** Add a SSequencerTreeView object that should be modified or updated when this Treeview is updated */
+	void AddSlaveTreeView(TSharedPtr<SSequencerTreeView> SlaveTreeView);
+	
+	/** Set whether this TreeView should show only pinned nodes or only non-pinned nodes  */
+	void SetShowPinned(bool bShowPinned) { bShowPinnedNodes = bShowPinned; }
+
 protected:
 	
 	/** Linear, sorted array of nodes that we currently have generated widgets for */
@@ -219,6 +225,9 @@ private:
 	/** A global highlight for the currently hovered tree node hierarchy */
 	TOptional<FHighlightRegion> HighlightRegion;
 
+	/** SSequencerTreeView objects that should be modified or updated when this Treeview is updated */
+	TArray<TSharedPtr<SSequencerTreeView>> SlaveTreeViews;
+
 	/** When true, the sequencer selection is being updated from a change in the tree seleciton. */
 	bool bUpdatingSequencerSelection;
 
@@ -228,6 +237,9 @@ private:
 	/** Right mouse button is down, don't update sequencer selection. */
 	bool bRightMouseButtonDown;
 
+	/** Whether this tree is for pinned nodes or non-pinned nodes */
+	bool bShowPinnedNodes;
+
 	/**
 	 * When true a sequencer selection change broadcast was suppressed when updating sequencer selection
 	 * due to the tree selection changing.
diff --git a/Engine/Source/Editor/Sequencer/Private/SequencerNodeTree.cpp b/Engine/Source/Editor/Sequencer/Private/SequencerNodeTree.cpp
--- a/Engine/Source/Editor/Sequencer/Private/SequencerNodeTree.cpp
+++ b/Engine/Source/Editor/Sequencer/Private/SequencerNodeTree.cpp
@@ -935,4 +935,14 @@ FCurveEditorTreeItemID FSequencerNodeTree::FindCurveEditorTreeItem(TSharedRef<FS
 		return *FoundID;
 	}
 	return FCurveEditorTreeItemID::Invalid();
-}
\ No newline at end of file
+}
+
+void FSequencerNodeTree::UnpinAllNodes()
+{
+	const bool bIncludeRootNode = false;
+	RootNode->Traverse_ParentFirst([](FSequencerDisplayNode& InNode)
+	{
+		InNode.Unpin();
+		return true;
+	}, bIncludeRootNode);
+}
diff --git a/Engine/Source/Editor/Sequencer/Private/SequencerNodeTree.h b/Engine/Source/Editor/Sequencer/Private/SequencerNodeTree.h
--- a/Engine/Source/Editor/Sequencer/Private/SequencerNodeTree.h
+++ b/Engine/Source/Editor/Sequencer/Private/SequencerNodeTree.h
@@ -82,6 +82,12 @@ public:
 	 */
 	void FilterNodes( const FString& InFilter );
 
+
+	/**
+	 * Unpins any pinned nodes in this tree
+	 */
+	void UnpinAllNodes();
+
 	/**
 	 * @return All nodes in a flat array
 	 */
diff --git a/Engine/Source/Editor/Sequencer/Private/Tools/SequencerEditTool_Selection.cpp b/Engine/Source/Editor/Sequencer/Private/Tools/SequencerEditTool_Selection.cpp
--- a/Engine/Source/Editor/Sequencer/Private/Tools/SequencerEditTool_Selection.cpp
+++ b/Engine/Source/Editor/Sequencer/Private/Tools/SequencerEditTool_Selection.cpp
@@ -13,14 +13,20 @@
 struct FSelectionPreviewVisitor final
 	: ISequencerEntityVisitor
 {
-	FSelectionPreviewVisitor(FSequencerSelectionPreview& InSelectionPreview, FSequencerSelection& InSelection, ESelectionPreviewState InSetStateTo)
+	FSelectionPreviewVisitor(FSequencerSelectionPreview& InSelectionPreview, FSequencerSelection& InSelection, ESelectionPreviewState InSetStateTo, bool bInPinned)
 		: SelectionPreview(InSelectionPreview)
 		, ExistingSelection(InSelection)
 		, SetStateTo(InSetStateTo)
+		, bPinned(bInPinned)
 	{}
 
 	virtual void VisitKey(FKeyHandle KeyHandle, FFrameNumber KeyTime, const TSharedPtr<IKeyArea>& KeyArea, UMovieSceneSection* Section, TSharedRef<FSequencerDisplayNode> Node) const override
 	{
+		if (Node->IsPinned() != bPinned)
+		{
+			return;
+		}
+
 		FSequencerSelectedKey Key(*Section, KeyArea, KeyHandle);
 
 		// Under default behavior keys have priority, so if a key is changing selection state then we remove any sections from the selection. The user can bypass this
@@ -72,6 +78,7 @@ private:
 	FSequencerSelectionPreview& SelectionPreview;
 	FSequencerSelection& ExistingSelection;
 	ESelectionPreviewState SetStateTo;
+	bool bPinned;
 };
 
 
@@ -80,8 +87,9 @@ class FMarqueeDragOperation
 {
 public:
 
-	FMarqueeDragOperation(FSequencer& InSequencer)
+	FMarqueeDragOperation(FSequencer& InSequencer, SSequencerTrackArea& InTrackArea)
 		: Sequencer(InSequencer)
+		, TrackArea(InTrackArea)
 		, SequencerWidget(StaticCastSharedRef<SSequencer>(InSequencer.GetSequencerWidget()))
 		, PreviewState(ESelectionPreviewState::Selected)
 	{}
@@ -127,16 +135,18 @@ public:
 		{
 			const float ScrollThresholdV = VirtualTrackArea.GetPhysicalSize().Y * 0.025f;
 
+			TSharedPtr<SSequencerTreeView> TreeView = TrackArea.GetTreeView().Pin();
+
 			float Difference = LocalMousePos.Y - ScrollThresholdV;
 			if (Difference < 0 && MouseDelta.Y < 0)
 			{
-				Sequencer.VerticalScroll( Difference * 0.1f );
+				TreeView->ScrollByDelta( Difference * 0.1f );
 			}
 
 			Difference = LocalMousePos.Y - (VirtualTrackArea.GetPhysicalSize().Y - ScrollThresholdV);
 			if (Difference > 0 && MouseDelta.Y > 0)
 			{
-				Sequencer.VerticalScroll( Difference * 0.1f );
+				TreeView->ScrollByDelta( Difference * 0.1f) ;
 			}
 		}
 
@@ -180,7 +190,7 @@ public:
 		VirtualKeySize.Y = SequencerSectionConstants::KeySize.Y;
 
 		// Visit everything using the preview selection primarily as well as the 
-		auto& SelectionPreview = Sequencer.GetSelectionPreview();
+		FSequencerSelectionPreview& SelectionPreview = Sequencer.GetSelectionPreview();
 
 		// Ensure the preview is empty before calculating the intersection
 		SelectionPreview.Empty();
@@ -189,7 +199,7 @@ public:
 
 		// Now walk everything within the current marquee range, setting preview selection states as we go
 		FSequencerEntityWalker Walker(FSequencerEntityRange(TopLeft(), BottomRight(), VirtualTrackArea.GetTickResolution()), VirtualKeySize);
-		Walker.Traverse(FSelectionPreviewVisitor(SelectionPreview, Sequencer.GetSelection(), PreviewState), RootNodes);
+		Walker.Traverse(FSelectionPreviewVisitor(SelectionPreview, Sequencer.GetSelection(), PreviewState, TrackArea.ShowPinned()), RootNodes);
 	}
 
 	virtual void OnEndDrag(const FPointerEvent& MouseEvent, FVector2D LocalMousePos, const FVirtualTrackArea& VirtualTrackArea) override
@@ -263,7 +273,7 @@ public:
 	virtual int32 OnPaint(const FGeometry& AllottedGeometry, const FSlateRect& MyCullingRect, FSlateWindowElementList& OutDrawElements, int32 LayerId) const
 	{
 		// convert to physical space for rendering
-		const FVirtualTrackArea VirtualTrackArea = SequencerWidget->GetVirtualTrackArea();
+		const FVirtualTrackArea VirtualTrackArea = SequencerWidget->GetVirtualTrackArea(&TrackArea);
 
 		FVector2D SelectionTopLeft = VirtualTrackArea.VirtualToPhysical(TopLeft());
 		FVector2D SelectionBottomRight = VirtualTrackArea.VirtualToPhysical(BottomRight());
@@ -299,6 +309,8 @@ private:
 	/** The sequencer itself */
 	FSequencer& Sequencer;
 
+	SSequencerTrackArea& TrackArea;
+
 	/** Sequencer widget */
 	TSharedRef<SSequencer> SequencerWidget;
 
@@ -313,8 +325,9 @@ private:
 
 const FName FSequencerEditTool_Selection::Identifier = "Selection";
 
-FSequencerEditTool_Selection::FSequencerEditTool_Selection(FSequencer& InSequencer)
+FSequencerEditTool_Selection::FSequencerEditTool_Selection(FSequencer& InSequencer, SSequencerTrackArea& InTrackArea)
 	: FSequencerEditTool(InSequencer)
+	, TrackArea(InTrackArea)
 	, CursorDecorator(nullptr)
 { }
 
@@ -370,7 +383,7 @@ FReply FSequencerEditTool_Selection::OnMouseMove(SWidget& OwnerWidget, const FGe
 		FReply Reply = FReply::Handled();
 
 		TSharedRef<SSequencer> SequencerWidget = StaticCastSharedRef<SSequencer>(Sequencer.GetSequencerWidget());
-		const FVirtualTrackArea VirtualTrackArea = SequencerWidget->GetVirtualTrackArea();
+		const FVirtualTrackArea VirtualTrackArea = SequencerWidget->GetVirtualTrackArea(&TrackArea);
 
 		if (DragOperation.IsValid())
 		{
@@ -391,7 +404,7 @@ FReply FSequencerEditTool_Selection::OnMouseMove(SWidget& OwnerWidget, const FGe
 
 			if (!DragOperation.IsValid())
 			{
-				DragOperation = MakeShareable( new FMarqueeDragOperation(Sequencer) );
+				DragOperation = MakeShareable( new FMarqueeDragOperation(Sequencer, TrackArea) );
 			}
 
 			if (DragOperation.IsValid())
@@ -418,7 +431,7 @@ FReply FSequencerEditTool_Selection::OnMouseButtonUp(SWidget& OwnerWidget, const
 	{
 		TSharedRef<SSequencer> SequencerWidget = StaticCastSharedRef<SSequencer>(Sequencer.GetSequencerWidget());
 		FVector2D LocalPosition = MyGeometry.AbsoluteToLocal(MouseEvent.GetScreenSpacePosition());
-		DragOperation->OnEndDrag(MouseEvent, LocalPosition, SequencerWidget->GetVirtualTrackArea() );
+		DragOperation->OnEndDrag(MouseEvent, LocalPosition, SequencerWidget->GetVirtualTrackArea(&TrackArea));
 		DragOperation = nullptr;
 
 		CursorDecorator = nullptr;
diff --git a/Engine/Source/Editor/Sequencer/Private/Tools/SequencerEditTool_Selection.h b/Engine/Source/Editor/Sequencer/Private/Tools/SequencerEditTool_Selection.h
--- a/Engine/Source/Editor/Sequencer/Private/Tools/SequencerEditTool_Selection.h
+++ b/Engine/Source/Editor/Sequencer/Private/Tools/SequencerEditTool_Selection.h
@@ -10,6 +10,7 @@
 #include "Framework/DelayedDrag.h"
 
 class SSequencer;
+class SSequencerTrackArea;
 struct ISequencerHotspot;
 
 class FSequencerEditTool_Selection
@@ -21,7 +22,7 @@ public:
 	static const FName Identifier;
 
 	/** Create and initialize a new instance. */
-	FSequencerEditTool_Selection(FSequencer& InSequencer);
+	FSequencerEditTool_Selection(FSequencer& InSequencer, SSequencerTrackArea& InTrackArea);
 
 public:
 
@@ -65,6 +66,9 @@ private:
 	/** Cached mouse position for software cursor rendering */
 	FVector2D MousePosition;
 
+	/** TrackArea this object belongs to */
+	SSequencerTrackArea& TrackArea;
+
 	/** Software cursor decorator brush */
 	const FSlateBrush* CursorDecorator;
 };
