Sequencer: Add option to enable all track filters at once

diff --git a/Engine/Source/Editor/Sequencer/Private/SSequencer.cpp b/Engine/Source/Editor/Sequencer/Private/SSequencer.cpp
--- a/Engine/Source/Editor/Sequencer/Private/SSequencer.cpp
+++ b/Engine/Source/Editor/Sequencer/Private/SSequencer.cpp
@@ -1502,6 +1502,12 @@ TSharedRef<SWidget> SSequencer::MakeFilterMenu()
 
 	MenuBuilder.BeginSection("SequencerTracksResetFilters");
 	{
+		MenuBuilder.AddMenuEntry(
+			LOCTEXT("FilterListEnableAll", "Enable All"),
+			LOCTEXT("FilterListEnableAllToolTip", "Selects all filters"),
+			FSlateIcon(),
+			FUIAction(FExecuteAction::CreateSP(this, &SSequencer::OnEnableAllFilters)));
+
 		MenuBuilder.AddMenuEntry(
 			LOCTEXT("FilterListResetFilters", "Reset Filters"),
 			LOCTEXT("FilterListResetToolTip", "Resets current filter selection"),
@@ -1601,6 +1607,22 @@ void SSequencer::OnResetFilters()
 	Sequencer->GetNodeTree()->RemoveAllFilters();
 }
 
+void SSequencer::OnEnableAllFilters()
+{
+	TSharedPtr<FSequencer> Sequencer = SequencerPtr.Pin();
+
+	for (TSharedRef<FSequencerTrackFilter> TrackFilter : AllTrackFilters)
+	{
+		if (TrackFilter->SupportsSequence(Sequencer->GetFocusedMovieSceneSequence()))
+		{
+			if (!Sequencer->GetNodeTree()->IsTrackFilterActive(TrackFilter))
+			{
+				Sequencer->GetNodeTree()->AddFilter(TrackFilter);
+			}
+		}
+	}
+}
+
 void SSequencer::OnTrackFilterClicked(TSharedRef<FSequencerTrackFilter> TrackFilter)
 {
 	TSharedPtr<FSequencer> Sequencer = SequencerPtr.Pin();
diff --git a/Engine/Source/Editor/Sequencer/Private/SSequencer.h b/Engine/Source/Editor/Sequencer/Private/SSequencer.h
--- a/Engine/Source/Editor/Sequencer/Private/SSequencer.h
+++ b/Engine/Source/Editor/Sequencer/Private/SSequencer.h
@@ -396,6 +396,7 @@ public:
 private:
 
 	void OnResetFilters();
+	void OnEnableAllFilters();
 	void OnTrackFilterClicked(TSharedRef<FSequencerTrackFilter> TrackFilter);
 	bool IsTrackFilterActive(TSharedRef<FSequencerTrackFilter> TrackFilter) const;
 
