Fix for crash when scrolling the sequence node list while the the track context menu is open.

diff --git a/Engine/Source/Editor/Sequencer/Private/SSequencerTrackArea.cpp b/Engine/Source/Editor/Sequencer/Private/SSequencerTrackArea.cpp
--- a/Engine/Source/Editor/Sequencer/Private/SSequencerTrackArea.cpp
+++ b/Engine/Source/Editor/Sequencer/Private/SSequencerTrackArea.cpp
@@ -94,7 +94,7 @@ void SSequencerTrackArea::OnArrangeChildren( const FGeometry& AllottedGeometry,
 		}
 
 		TSharedPtr<SSequencerTrackLane> PinnedTrackLane = CurChild.TrackLane.Pin();
-		if (PinnedTrackLane->IsPinned() != bShowPinnedNodes)
+		if (!PinnedTrackLane.IsValid() || PinnedTrackLane->IsPinned() != bShowPinnedNodes)
 		{
 			continue;
 		}
