Sequencer: Fix "show selected nodes only" regression and issues selecting tracks caused by new pinning system

diff --git a/Engine/Source/Editor/Sequencer/Private/SSequencerTreeView.cpp b/Engine/Source/Editor/Sequencer/Private/SSequencerTreeView.cpp
--- a/Engine/Source/Editor/Sequencer/Private/SSequencerTreeView.cpp
+++ b/Engine/Source/Editor/Sequencer/Private/SSequencerTreeView.cpp
@@ -437,6 +437,7 @@ void SSequencerTreeView::UpdateTrackArea()
 void SSequencerTreeView::AddSlaveTreeView(TSharedPtr<SSequencerTreeView> SlaveTreeView)
 {
 	SlaveTreeViews.Add(SlaveTreeView);
+	SlaveTreeView->SetMasterTreeView(SharedThis(this));
 }
 
 void SSequencerTreeView::OnRightMouseButtonDown(const FPointerEvent& MouseEvent)
@@ -645,7 +646,7 @@ void SSequencerTreeView::SynchronizeTreeSelectionWithSequencerSelection()
 			FSequencer& Sequencer = SequencerNodeTree->GetSequencer();
 			for ( const TSharedRef<FSequencerDisplayNode>& Node : Sequencer.GetSelection().GetSelectedOutlinerNodes() )
 			{
-				if (Node->IsSelectable())
+				if (Node->IsSelectable() && Node->GetOutermostParent()->IsPinned() == bShowPinnedNodes)
 				{
 					Private_SetItemSelection( Node, true, false );
 				}
@@ -739,14 +740,28 @@ void SSequencerTreeView::Private_SignalSelectionChanged(ESelectInfo::Type Select
 
 bool SSequencerTreeView::SynchronizeSequencerSelectionWithTreeSelection()
 {
+	// If this is a slave SSequencerTreeView it only has a partial view of what is selected. The master should handle syncing the entire selection instead.
+	if (MasterTreeView.IsValid())
+	{
+		return MasterTreeView->SynchronizeSequencerSelectionWithTreeSelection();
+	}
+
 	bool bSelectionChanged = false;
 	const TSet<TSharedRef<FSequencerDisplayNode>>& SequencerSelection = SequencerNodeTree->GetSequencer().GetSelection().GetSelectedOutlinerNodes();
-	if ( SelectedItems.Num() != SequencerSelection.Num() || SelectedItems.Difference( SequencerSelection ).Num() != 0 )
+	TSet<TSharedRef<FSequencerDisplayNode> > AllSelectedItems(SelectedItems);
+
+	// If we have slave SSequencerTreeViews, combine their selected items as well
+	for (TSharedPtr<SSequencerTreeView> SlaveTreeView : SlaveTreeViews)
+	{
+		AllSelectedItems.Append(SlaveTreeView->GetSelectedItems());
+	}
+
+	if (AllSelectedItems.Num() != SequencerSelection.Num() || AllSelectedItems.Difference(SequencerSelection).Num() != 0 )
 	{
 		FSequencer& Sequencer = SequencerNodeTree->GetSequencer();
 		FSequencerSelection& Selection = Sequencer.GetSelection();
 		Selection.EmptySelectedOutlinerNodes();
-		for ( const TSharedRef<FSequencerDisplayNode>& Item : GetSelectedItems() )
+		for ( const TSharedRef<FSequencerDisplayNode>& Item : AllSelectedItems)
 		{
 			Selection.AddToSelection( Item );
 		}
diff --git a/Engine/Source/Editor/Sequencer/Private/SSequencerTreeView.h b/Engine/Source/Editor/Sequencer/Private/SSequencerTreeView.h
--- a/Engine/Source/Editor/Sequencer/Private/SSequencerTreeView.h
+++ b/Engine/Source/Editor/Sequencer/Private/SSequencerTreeView.h
@@ -185,7 +185,10 @@ public:
 
 	/** Add a SSequencerTreeView object that should be modified or updated when this Treeview is updated */
 	void AddSlaveTreeView(TSharedPtr<SSequencerTreeView> SlaveTreeView);
-	
+
+	/** Set a SSequencerTreeView object this Treeview is slaved to, for operations that should happen on the master */
+	void SetMasterTreeView(TSharedPtr<SSequencerTreeView> InMasterTreeView) { MasterTreeView = InMasterTreeView; }
+
 	/** Set whether this TreeView should show only pinned nodes or only non-pinned nodes  */
 	void SetShowPinned(bool bShowPinned) { bShowPinnedNodes = bShowPinned; }
 
@@ -228,6 +231,9 @@ private:
 	/** SSequencerTreeView objects that should be modified or updated when this Treeview is updated */
 	TArray<TSharedPtr<SSequencerTreeView>> SlaveTreeViews;
 
+	/** The SSequencerTreeView object this SSequencerTreeView is slave to, or nullptr if not a slave */
+	TSharedPtr<SSequencerTreeView> MasterTreeView;
+
 	/** When true, the sequencer selection is being updated from a change in the tree seleciton. */
 	bool bUpdatingSequencerSelection;
 
diff --git a/Engine/Source/Editor/Sequencer/Private/SequencerNodeTree.cpp b/Engine/Source/Editor/Sequencer/Private/SequencerNodeTree.cpp
--- a/Engine/Source/Editor/Sequencer/Private/SequencerNodeTree.cpp
+++ b/Engine/Source/Editor/Sequencer/Private/SequencerNodeTree.cpp
@@ -693,6 +693,12 @@ static void FilterNodesRecursive( FSequencer& Sequencer, const TSharedRef<FSeque
 				TSharedPtr<const FSequencerDisplayNode> ParentNode = StartNode->GetParent();
 				while (ParentNode.IsValid())
 				{
+					// Pinned tracks should be visible whether selected or not
+					if (ParentNode->IsPinned())
+					{
+						break;
+					}
+
 					if (ParentNode->GetType() == ESequencerNode::Object)
 					{
 						const FSequencerObjectBindingNode* ObjectNode = static_cast<const FSequencerObjectBindingNode*>(ParentNode.Get());
@@ -722,7 +728,7 @@ static void FilterNodesRecursive( FSequencer& Sequencer, const TSharedRef<FSeque
 			}
 		}
 
-		if (bPasssedAnyFilters && Sequencer.GetSequencerSettings()->GetShowSelectedNodesOnly())
+		if (bPasssedAnyFilters && Sequencer.GetSequencerSettings()->GetShowSelectedNodesOnly() && !StartNode->IsPinned())
 		{
 			UMovieScene* MovieScene = Sequencer.GetFocusedMovieSceneSequence()->GetMovieScene();
 			const FMovieSceneBinding* Binding = Sequencer.GetFocusedMovieSceneSequence()->GetMovieScene()->FindBinding(ObjectNode.GetObjectBinding());
diff --git a/Engine/Source/Editor/Sequencer/Private/SequencerNodeTree.h b/Engine/Source/Editor/Sequencer/Private/SequencerNodeTree.h
--- a/Engine/Source/Editor/Sequencer/Private/SequencerNodeTree.h
+++ b/Engine/Source/Editor/Sequencer/Private/SequencerNodeTree.h
@@ -171,6 +171,9 @@ public:
 
 private:
 
+	/** Returns whether this NodeTree should only display seleected nodes */
+	bool ShowSelectedNodesOnly() const;
+
 	/** Population algorithm utilities */
 	void RefreshNodes(UMovieScene* MovieScene);
 
