UMG: Fix right aligned lock button in timeline.

diff --git a/Engine/Source/Editor/Sequencer/Private/SSequencer.cpp b/Engine/Source/Editor/Sequencer/Private/SSequencer.cpp
--- a/Engine/Source/Editor/Sequencer/Private/SSequencer.cpp
+++ b/Engine/Source/Editor/Sequencer/Private/SSequencer.cpp
@@ -555,6 +555,14 @@ void SSequencer::Construct(const FArguments& InArgs, TSharedRef<FSequencer> InSe
 								[
 									SNew(SHorizontalBox)
 
+									// Right Aligned Breadcrumbs
+									+ SHorizontalBox::Slot()
+									.HAlign(HAlign_Right)
+									.VAlign(VAlign_Center)
+									[
+										SNew(SSpacer)
+									]
+
 									+ SHorizontalBox::Slot()
 									.HAlign(HAlign_Right)
 									.VAlign(VAlign_Center)
