Sequencer: Add support for soloing and muting sequencer tracks

diff --git a/Engine/Source/Editor/Sequencer/Private/DisplayNodes/SequencerDisplayNode.cpp b/Engine/Source/Editor/Sequencer/Private/DisplayNodes/SequencerDisplayNode.cpp
--- a/Engine/Source/Editor/Sequencer/Private/DisplayNodes/SequencerDisplayNode.cpp
+++ b/Engine/Source/Editor/Sequencer/Private/DisplayNodes/SequencerDisplayNode.cpp
@@ -703,10 +703,32 @@ FLinearColor FSequencerDisplayNode::GetDisplayNameColor() const
 	FSequencerDisplayNode *This = const_cast<FSequencerDisplayNode*>(this);
 	//if empty with no key areas or sections then it's active, otherwise
 	//find first child with active section, then it's active, else inactive.
-	const bool bFoundInActiveSection = ChildNodes.Num() > 0 ? This->Traverse_ParentFirst(FindInActiveSection) :
+	bool bDimLabel = ChildNodes.Num() > 0 ? This->Traverse_ParentFirst(FindInActiveSection) :
 		((this->GetType() == ESequencerNode::Track || this->GetType() == ESequencerNode::KeyArea) && FindInActiveSection(*(This), false))
 		||false;
-	return bFoundInActiveSection ? FLinearColor(0.6f, 0.6f, 0.6f, 0.6f) : FLinearColor(1.0f, 1.0f, 1.0f, 1.0f);
+
+	if (!bDimLabel)
+	{
+		// If the node is a track node, we can use the cached value in UMovieSceneTrack
+		if (GetType() == ESequencerNode::Track)
+		{
+			UMovieSceneTrack* Track = static_cast<const FSequencerTrackNode*>(this)->GetTrack();
+			if (Track && Track->IsEvalDisabled())
+			{
+				bDimLabel = true;
+			}
+		}
+		else
+		{
+			if (ParentTree.IsNodeMute(this) || (ParentTree.HasSoloNodes() && !ParentTree.IsNodeSolo(this)))
+			{
+				bDimLabel = true;
+			}
+		}
+
+	}
+
+	return bDimLabel ? FLinearColor(0.6f, 0.6f, 0.6f, 0.6f) : FLinearColor(1.0f, 1.0f, 1.0f, 1.0f);
 }
 
 FText FSequencerDisplayNode::GetDisplayNameToolTipText() const
@@ -935,25 +957,16 @@ namespace
 void FSequencerDisplayNode::BuildContextMenu(FMenuBuilder& MenuBuilder)
 {
 	TSharedRef<FSequencerDisplayNode> ThisNode = SharedThis(this);
+	FSequencerDisplayNode* BaseNode = GetBaseNode();
+
+	ESequencerNode::Type BaseNodeType = BaseNode->GetType();
 
+	bool bCanSolo = (BaseNodeType == ESequencerNode::Track || BaseNodeType == ESequencerNode::Object || BaseNodeType == ESequencerNode::Folder);
 	bool bIsReadOnly = !GetSequencer().IsReadOnly();
 	FCanExecuteAction CanExecute = FCanExecuteAction::CreateLambda([bIsReadOnly]{ return bIsReadOnly; });
 
 	MenuBuilder.BeginSection("Edit", LOCTEXT("EditContextMenuSectionName", "Edit"));
 	{
-		MenuBuilder.AddMenuEntry(
-			LOCTEXT("ToggleNodeActive", "Active"),
-			LOCTEXT("ToggleNodeActiveTooltip", "Set this track or selected tracks active/inactive"),
-			FSlateIcon(),
-			FUIAction(
-				FExecuteAction::CreateSP(&GetSequencer(), &FSequencer::ToggleNodeActive),
-				CanExecute,
-				FIsActionChecked::CreateSP(&GetSequencer(), &FSequencer::IsNodeActive)
-			),
-			NAME_None,
-			EUserInterfaceActionType::ToggleButton
-		);
-
 		MenuBuilder.AddMenuEntry(
 			LOCTEXT("ToggleNodeLock", "Locked"),
 			LOCTEXT("ToggleNodeLockTooltip", "Lock or unlock this node or selected tracks"),
@@ -968,7 +981,7 @@ void FSequencerDisplayNode::BuildContextMenu(FMenuBuilder& MenuBuilder)
 		);
 
 		// Only support pinning root nodes
-		if (GetBaseNode()->IsRootNode())
+		if (BaseNode->IsRootNode())
 		{
 			MenuBuilder.AddMenuEntry(
 				LOCTEXT("ToggleNodePin", "Pinned"),
@@ -984,6 +997,35 @@ void FSequencerDisplayNode::BuildContextMenu(FMenuBuilder& MenuBuilder)
 			);
 		}
 
+		if (bCanSolo)
+		{
+			MenuBuilder.AddMenuEntry(
+				LOCTEXT("ToggleNodeSolo", "Solo"),
+				LOCTEXT("ToggleNodeLockTooltip", "Solo or unsolo this node or selected tracks"),
+				FSlateIcon(),
+				FUIAction(
+					FExecuteAction::CreateRaw(&ParentTree, &FSequencerNodeTree::ToggleSelectedNodesSolo),
+					CanExecute,
+					FIsActionChecked::CreateSP(&ParentTree, &FSequencerNodeTree::IsSelectedNodesSolo)
+				),
+				NAME_None,
+				EUserInterfaceActionType::ToggleButton
+			);
+
+			MenuBuilder.AddMenuEntry(
+				LOCTEXT("ToggleNodeMute", "Mute"),
+				LOCTEXT("ToggleNodeLockTooltip", "Mute or unmute this node or selected tracks"),
+				FSlateIcon(),
+				FUIAction(
+					FExecuteAction::CreateRaw(&ParentTree, &FSequencerNodeTree::ToggleSelectedNodesMute),
+					CanExecute,
+					FIsActionChecked::CreateSP(&ParentTree, &FSequencerNodeTree::IsSelectedNodesMute)
+				),
+				NAME_None,
+				EUserInterfaceActionType::ToggleButton
+			);
+		}
+
 		// Add cut, copy and paste functions to the tracks
 		MenuBuilder.AddMenuEntry(FGenericCommands::Get().Cut);
 
diff --git a/Engine/Source/Editor/Sequencer/Private/DisplayNodes/SequencerFolderNode.cpp b/Engine/Source/Editor/Sequencer/Private/DisplayNodes/SequencerFolderNode.cpp
--- a/Engine/Source/Editor/Sequencer/Private/DisplayNodes/SequencerFolderNode.cpp
+++ b/Engine/Source/Editor/Sequencer/Private/DisplayNodes/SequencerFolderNode.cpp
@@ -116,6 +116,11 @@ FSlateColor FSequencerFolderNode::GetIconColor() const
 
 FLinearColor FSequencerFolderNode::GetDisplayNameColor() const
 {
+	if (ParentTree.IsNodeMute(this))
+	{
+		return FLinearColor(0.6f, 0.6f, 0.6f, 0.6f);
+	}
+
 	return FLinearColor::White;
 }
 
diff --git a/Engine/Source/Editor/Sequencer/Private/DisplayNodes/SequencerTrackNode.cpp b/Engine/Source/Editor/Sequencer/Private/DisplayNodes/SequencerTrackNode.cpp
--- a/Engine/Source/Editor/Sequencer/Private/DisplayNodes/SequencerTrackNode.cpp
+++ b/Engine/Source/Editor/Sequencer/Private/DisplayNodes/SequencerTrackNode.cpp
@@ -688,13 +688,15 @@ FLinearColor FSequencerTrackNode::GetDisplayNameColor() const
 {
 	UMovieSceneTrack* Track = GetTrack();
 
+	const bool bIsEvalDisabled = Track->IsEvalDisabled();
+
 	// Display track node is red if the property track is not bound to valid property
 	if (UMovieScenePropertyTrack* PropertyTrack = Cast<UMovieScenePropertyTrack>(Track))
 	{
 		// 3D transform tracks don't map to property bindings as below
 		if (Track->IsA<UMovieScene3DTransformTrack>() || Track->IsA<UMovieScenePrimitiveMaterialTrack>())
 		{
-			return FLinearColor::White;
+			return bIsEvalDisabled ? FLinearColor(0.6f, 0.6f, 0.6f, 0.6f) : FLinearColor::White;
 		}
 
 		FGuid ObjectBinding;
@@ -713,15 +715,15 @@ FLinearColor FSequencerTrackNode::GetDisplayNameColor() const
 
 				if (PropertyBinding.GetProperty(*RuntimeObject))
 				{
-					return FLinearColor::White;
+					return bIsEvalDisabled ? FLinearColor(0.6f, 0.6f, 0.6f, 0.6f) : FLinearColor::White;
 				}
 			}
 
-			return FLinearColor::Red;
+			return bIsEvalDisabled ? FLinearColor(0.6f, 0.0f, 0.0f, 0.6f) : FLinearColor::Red;
 		}
 	}
 
-	return FLinearColor::White;
+	return bIsEvalDisabled ? FLinearColor(0.6f, 0.6f, 0.6f, 0.6f) : FLinearColor::White;
 }
 
 float FSequencerTrackNode::GetNodeHeight() const
diff --git a/Engine/Source/Editor/Sequencer/Private/SSequencer.cpp b/Engine/Source/Editor/Sequencer/Private/SSequencer.cpp
--- a/Engine/Source/Editor/Sequencer/Private/SSequencer.cpp
+++ b/Engine/Source/Editor/Sequencer/Private/SSequencer.cpp
@@ -2341,6 +2341,33 @@ void SSequencer::UpdateLayoutTree()
 
 		AdditionalSelectionsToAdd.Empty();
 
+		if (Sequencer->GetFocusedMovieSceneSequence())
+		{
+			bool bAnyChanged = false;
+
+			TSharedPtr<FSequencerNodeTree> NodeTree = Sequencer->GetNodeTree();
+			const bool bHasSoloNodes = NodeTree->HasSoloNodes();
+			for (TSharedRef<FSequencerDisplayNode> Node : NodeTree->GetAllNodes())
+			{
+				if (Node->GetType() == ESequencerNode::Track)
+				{
+					UMovieSceneTrack* Track = static_cast<FSequencerTrackNode&>(Node.Get()).GetTrack();
+					bool bDisableEval = NodeTree->IsNodeMute(&Node.Get()) || (bHasSoloNodes && !NodeTree->IsNodeSolo(&Node.Get()));
+					if (bDisableEval != Track->IsEvalDisabled())
+					{
+						Track->Modify();
+						Track->SetEvalDisabled(bDisableEval);
+						bAnyChanged = true;
+					}
+					
+				}
+			}
+			if (bAnyChanged)
+			{
+				Sequencer->NotifyMovieSceneDataChanged(EMovieSceneDataChangeType::TrackValueChanged);
+			}
+		}
+
 		// Continue broadcasting selection changes
 		Sequencer->GetSelection().ResumeBroadcast();
 	}
diff --git a/Engine/Source/Editor/Sequencer/Private/SSequencerSection.cpp b/Engine/Source/Editor/Sequencer/Private/SSequencerSection.cpp
--- a/Engine/Source/Editor/Sequencer/Private/SSequencerSection.cpp
+++ b/Engine/Source/Editor/Sequencer/Private/SSequencerSection.cpp
@@ -1039,9 +1039,9 @@ int32 SSequencerSection::OnPaint( const FPaintArgs& Args, const FGeometry& Allot
 		Hotspot = GetSequencer().GetHotspot().Get();
 	}
 
-	const bool bEnabled = bParentEnabled && SectionObject->IsActive();
-	const bool bLocked = SectionObject->IsLocked();
 	UMovieScenePropertyTrack* Track = SectionObject->GetTypedOuter<UMovieScenePropertyTrack>();
+	const bool bEnabled = bParentEnabled && SectionObject->IsActive() && !(Track && Track->IsEvalDisabled());
+	const bool bLocked = SectionObject->IsLocked();
 	bool bSetSectionToKey = false;
 	if (Track)
 	{
diff --git a/Engine/Source/Editor/Sequencer/Private/SequencerNodeTree.cpp b/Engine/Source/Editor/Sequencer/Private/SequencerNodeTree.cpp
--- a/Engine/Source/Editor/Sequencer/Private/SequencerNodeTree.cpp
+++ b/Engine/Source/Editor/Sequencer/Private/SequencerNodeTree.cpp
@@ -346,9 +346,58 @@ void FSequencerNodeTree::Update()
 	// Ensure that the curve editor tree is up to date for our tree layout
 	UpdateCurveEditorTree();
 
+	bHasSoloNodes = false;
+	
+	TArray<FString>& SoloNodes = Sequencer.GetFocusedMovieSceneSequence()->GetMovieScene()->GetSoloNodes();
+
+	// Muting overrides soloing, so a solo node only counts if it's not muted.
+	for (const FString& NodePath : SoloNodes)
+	{
+		FSequencerDisplayNode* Node = GetNodeAtPath(NodePath);
+		if (Node)
+		{
+			if (!IsNodeMute(Node))
+			{
+				bHasSoloNodes = true;
+				break;
+			}
+		}
+	}
+
 	OnUpdatedDelegate.Broadcast();
 }
 
+FSequencerDisplayNode* FSequencerNodeTree::GetNodeAtPath(const FString& NodePath) const
+{
+	TArray<FString> NodePathParts;
+	int32 PathLen = NodePath.ParseIntoArray(NodePathParts, TEXT("."));
+
+	FSequencerDisplayNode* Node = &RootNode.Get();
+
+	int32 PathIdx = 0;
+	while (PathIdx < PathLen)
+	{
+		FString& PathPart = NodePathParts[PathIdx];
+		bool bChildFound = false;
+		for (auto ChildNode : Node->GetChildNodes())
+		{
+			if (ChildNode->GetNodeName().ToString().Equals(PathPart))
+			{
+				bChildFound = true;
+				++PathIdx;
+				Node = &ChildNode.Get();
+				break;
+			}
+		}
+		
+		if (!bChildFound)
+		{
+			return nullptr;
+		}
+	}
+
+	return Node;
+}
 
 TSharedRef<ISequencerTrackEditor> FSequencerNodeTree::FindOrAddTypeEditor( UMovieSceneTrack* InTrack )
 {
@@ -512,6 +561,196 @@ bool FSequencerNodeTree::IsTrackLevelFilterActive(const FString& LevelName) cons
 	return !TrackFilterLevelFilter->IsLevelHidden(LevelName);
 }
 
+bool FSequencerNodeTree::IsNodeSolo(const FSequencerDisplayNode* InNode) const
+{
+	const TArray<FString>& SoloNodes = Sequencer.GetFocusedMovieSceneSequence()->GetMovieScene()->GetSoloNodes();
+	const FString NodePath = InNode->GetBaseNode()->GetPathName();
+
+	if (SoloNodes.Contains(NodePath))
+	{
+		return true;
+	}
+
+	// Children should follow their parent's behavior unless told otherwise.
+	TSharedPtr<FSequencerDisplayNode> ParentNode = InNode->GetParent();
+	if (ParentNode.IsValid())
+	{
+		return IsNodeSolo(ParentNode.Get());
+	}
+
+	return false;
+}
+
+bool FSequencerNodeTree::HasSoloNodes() const
+{
+	return bHasSoloNodes;
+}
+
+bool FSequencerNodeTree::IsSelectedNodesSolo() const
+{
+	const TSet<TSharedRef<FSequencerDisplayNode> > SelectedNodes = Sequencer.GetSelection().GetSelectedOutlinerNodes();
+	
+	if (SelectedNodes.Num() == 0)
+	{
+		return false;
+	}
+
+	const TArray<FString>& SoloNodes = Sequencer.GetFocusedMovieSceneSequence()->GetMovieScene()->GetSoloNodes();
+
+	bool bIsSolo = true;
+	for (const TSharedRef<const FSequencerDisplayNode> Node : SelectedNodes)
+	{
+		if (!SoloNodes.Contains(Node->GetBaseNode()->GetPathName()))
+		{
+			bIsSolo = false;
+			break;
+		}
+	}
+
+	return bIsSolo;
+}
+
+void FSequencerNodeTree::ToggleSelectedNodesSolo()
+{
+	UMovieScene* MovieScene = Sequencer.GetFocusedMovieSceneSequence()->GetMovieScene();
+	
+	if (MovieScene->IsReadOnly())
+	{
+		return;
+	}
+	
+	TArray<FString>& SoloNodes = MovieScene->GetSoloNodes();
+
+	const TSet<TSharedRef<FSequencerDisplayNode> >& SelectedNodes = Sequencer.GetSelection().GetSelectedOutlinerNodes();
+	if (SelectedNodes.Num() == 0)
+	{
+		return;
+	}
+
+	// First, determine if any of the selected nodes are not marked as solo
+	// If we have a mix, we should default to setting them all as solo
+	bool bIsSolo = true;
+	for (const TSharedRef<const FSequencerDisplayNode> Node : Sequencer.GetSelection().GetSelectedOutlinerNodes())
+	{
+		if (!SoloNodes.Contains(Node->GetBaseNode()->GetPathName()))
+		{
+			bIsSolo = false;
+			break;
+		}
+	}
+
+	MovieScene->Modify();
+
+	for (const TSharedRef<const FSequencerDisplayNode> Node : Sequencer.GetSelection().GetSelectedOutlinerNodes())
+	{
+		FString NodePath = Node->GetBaseNode()->GetPathName();
+		if (bIsSolo)
+		{
+			// If we're currently solo, unsolo
+			SoloNodes.Remove(NodePath);
+		}
+		else
+		{
+			// Mark solo, being careful as we might be re-marking an already solo node
+			SoloNodes.AddUnique(NodePath);
+		}
+	}
+	
+	GetSequencer().RefreshTree();
+}
+
+bool FSequencerNodeTree::IsNodeMute(const FSequencerDisplayNode* InNode) const
+{
+	const TArray<FString>& MuteNodes = Sequencer.GetFocusedMovieSceneSequence()->GetMovieScene()->GetMuteNodes();
+	const FString NodePath = InNode->GetBaseNode()->GetPathName();
+
+	if (MuteNodes.Contains(NodePath))
+	{
+		return true;
+	}
+	
+	// Children should follow their parent's behavior unless told otherwise.
+	const TSharedPtr<FSequencerDisplayNode> ParentNode = InNode->GetParent();
+	if (ParentNode.IsValid())
+	{
+		return IsNodeMute(ParentNode.Get());
+	}
+	
+	return false;
+}
+
+bool FSequencerNodeTree::IsSelectedNodesMute() const
+{
+	const TSet<TSharedRef<FSequencerDisplayNode> > SelectedNodes = Sequencer.GetSelection().GetSelectedOutlinerNodes();
+
+	if (SelectedNodes.Num() == 0)
+	{
+		return false;
+	}
+
+	const TArray<FString>& MuteNodes = Sequencer.GetFocusedMovieSceneSequence()->GetMovieScene()->GetMuteNodes();
+
+	bool bIsMute = true;
+	for (const TSharedRef<const FSequencerDisplayNode> Node : SelectedNodes)
+	{
+		if (!MuteNodes.Contains(Node->GetBaseNode()->GetPathName()))
+		{
+			bIsMute = false;
+			break;
+		}
+	}
+
+	return bIsMute;
+}
+
+void FSequencerNodeTree::ToggleSelectedNodesMute()
+{
+	UMovieScene* MovieScene = Sequencer.GetFocusedMovieSceneSequence()->GetMovieScene();
+	
+	if (MovieScene->IsReadOnly())
+	{
+		return;
+	}
+
+	TArray<FString>& MuteNodes = MovieScene->GetMuteNodes();
+
+	const TSet<TSharedRef<FSequencerDisplayNode> >& SelectedNodes = Sequencer.GetSelection().GetSelectedOutlinerNodes();
+	if (SelectedNodes.Num() == 0)
+	{
+		return;
+	}
+
+	// First, determine if any of the selected nodes are not marked as Mute
+	// If we have a mix, we should default to setting them all as Mute
+	bool bIsMute = true;
+	for (const TSharedRef<const FSequencerDisplayNode> Node : SelectedNodes)
+	{
+		if (!MuteNodes.Contains(Node->GetBaseNode()->GetPathName()))
+		{
+			bIsMute = false;
+			break;
+		}
+	}
+
+	MovieScene->Modify();
+
+	for (const TSharedRef<const FSequencerDisplayNode> Node : SelectedNodes)
+	{
+		FString NodePath = Node->GetBaseNode()->GetPathName();
+		if (bIsMute)
+		{
+			// If we're currently Mute, unMute
+			MuteNodes.Remove(NodePath);
+		}
+		else
+		{
+			// Mark Mute, being careful as we might be re-marking an already Mute node
+			MuteNodes.AddUnique(NodePath);
+		}
+	}
+
+	GetSequencer().RefreshTree();
+}
 
 void FSequencerNodeTree::SaveExpansionState(const FSequencerDisplayNode& Node, bool bExpanded)
 {	
diff --git a/Engine/Source/Editor/Sequencer/Private/SequencerNodeTree.h b/Engine/Source/Editor/Sequencer/Private/SequencerNodeTree.h
--- a/Engine/Source/Editor/Sequencer/Private/SequencerNodeTree.h
+++ b/Engine/Source/Editor/Sequencer/Private/SequencerNodeTree.h
@@ -153,6 +153,14 @@ public:
 	 */
 	FCurveEditorTreeItemID FindCurveEditorTreeItem(TSharedRef<FSequencerDisplayNode> DisplayNode) const;
 
+	/**
+	 * Attempt to get a node at the specified path
+	 *
+	 * @param NodePath The path of the node to search for
+	 * @return The node located at NodePath, or nullptr if not found
+	 */
+	FSequencerDisplayNode* GetNodeAtPath(const FString& NodePath) const;
+
 public:
 
 	/**
@@ -169,6 +177,21 @@ public:
 	void RemoveLevelFilter(const FString& LevelName);
 	bool IsTrackLevelFilterActive(const FString& LevelName) const;
 
+	void ToggleSelectedNodesSolo();
+	bool IsNodeSolo(const FSequencerDisplayNode* Node) const;
+	
+	/** Returns whether any of the nodes in this tree are marked solo */
+	bool HasSoloNodes() const;
+
+	void ToggleSelectedNodesMute();
+	bool IsNodeMute(const FSequencerDisplayNode* Node) const;
+
+	/** Returns whether any of the currently selected nodes are marked solo */
+	bool IsSelectedNodesSolo() const;
+	
+	/** Returns whether any of the currently selected nodes are muted */
+	bool IsSelectedNodesMute() const;
+
 private:
 
 	/** Returns whether this NodeTree should only display seleected nodes */
@@ -286,4 +309,7 @@ private:
 	TSharedPtr<FSequencerTrackFilter_LevelFilter> TrackFilterLevelFilter;
 
 	bool bFilterUpdateRequested;
+	
+	/** Cached value of whether we have any nodes that should be treated as soloing */
+	bool bHasSoloNodes;
 };
diff --git a/Engine/Source/Runtime/MovieScene/Private/Compilation/MovieSceneEvaluationTemplateGenerator.cpp b/Engine/Source/Runtime/MovieScene/Private/Compilation/MovieSceneEvaluationTemplateGenerator.cpp
--- a/Engine/Source/Runtime/MovieScene/Private/Compilation/MovieSceneEvaluationTemplateGenerator.cpp
+++ b/Engine/Source/Runtime/MovieScene/Private/Compilation/MovieSceneEvaluationTemplateGenerator.cpp
@@ -61,6 +61,11 @@ void FMovieSceneEvaluationTemplateGenerator::ProcessTrack(const UMovieSceneTrack
 {
 	FMovieSceneTrackSegmentBlenderPtr TrackBlender = Track.GetTrackSegmentBlender();
 
+	if (Track.IsEvalDisabled())
+	{
+		return;
+	}
+
 	// Deal with sub tracks specifically
 	if (const UMovieSceneSubTrack* SubTrack = Cast<const UMovieSceneSubTrack>(&Track))
 	{
diff --git a/Engine/Source/Runtime/MovieScene/Public/MovieScene.h b/Engine/Source/Runtime/MovieScene/Public/MovieScene.h
--- a/Engine/Source/Runtime/MovieScene/Public/MovieScene.h
+++ b/Engine/Source/Runtime/MovieScene/Public/MovieScene.h
@@ -663,6 +663,17 @@ public:
 	 * Gets the root folders for this movie scene.
 	 */
 	TArray<UMovieSceneFolder*>& GetRootFolders();
+
+	/**
+	 * Gets the nodes marked as solo in the editor, as node tree paths
+	 */
+	TArray<FString>& GetSoloNodes() { return SoloNodes; }
+
+	/**
+	 * Gets the nodes marked as muted in the editor, as node tree paths
+	 */
+	TArray<FString>& GetMuteNodes() { return MuteNodes; }
+	
 #endif
 
 	/**
@@ -906,6 +917,14 @@ private:
 	UPROPERTY()
 	TArray<UMovieSceneFolder*> RootFolders;
 
+	/** Nodes currently marked Solo, stored as node tree paths */
+	UPROPERTY()
+	TArray<FString> SoloNodes;
+	
+	/** Nodes currently marked Mute, stored as node tree paths */
+	UPROPERTY()
+	TArray<FString> MuteNodes;
+
 private:
 
 	UPROPERTY()
diff --git a/Engine/Source/Runtime/MovieScene/Public/MovieSceneTrack.h b/Engine/Source/Runtime/MovieScene/Public/MovieSceneTrack.h
--- a/Engine/Source/Runtime/MovieScene/Public/MovieSceneTrack.h
+++ b/Engine/Source/Runtime/MovieScene/Public/MovieSceneTrack.h
@@ -221,6 +221,10 @@ protected:
 	/** Intentionally not a UPROPERTY so this isn't serialized */
 	FMovieSceneBlendTypeField SupportedBlendTypes;
 
+	/** Whether evaluation of this track has been disabled via mute/solo */
+	UPROPERTY()
+	bool bIsEvalDisabled;
+
 public:
 
 	/**
@@ -255,6 +259,16 @@ public:
 	 */
 	MOVIESCENE_API bool FixRowIndices();
 
+	/**
+	* @return Whether evaluation of this track should be disabled due to mute/solo settings
+	*/
+	MOVIESCENE_API bool IsEvalDisabled() const { return bIsEvalDisabled; };
+
+	/**
+	* Called by Sequencer to set whether evaluation of this track should be disabled due to mute/solo settings
+	*/
+	MOVIESCENE_API void SetEvalDisabled(bool bEvalDisabled) { bIsEvalDisabled = bEvalDisabled; }
+
 public:
 
 	/*
@@ -397,5 +411,6 @@ public:
 	 * @param Section The section that moved.
 	 */
 	virtual void OnSectionMoved(UMovieSceneSection& Section) { }
+
 #endif
 };
