Sequencer: Fix object binding label so that it reports correctly when there are invalid bindings

diff --git a/Engine/Source/Editor/Sequencer/Private/DisplayNodes/SequencerObjectBindingNode.cpp b/Engine/Source/Editor/Sequencer/Private/DisplayNodes/SequencerObjectBindingNode.cpp
--- a/Engine/Source/Editor/Sequencer/Private/DisplayNodes/SequencerObjectBindingNode.cpp
+++ b/Engine/Source/Editor/Sequencer/Private/DisplayNodes/SequencerObjectBindingNode.cpp
@@ -534,8 +534,11 @@ FLinearColor FSequencerObjectBindingNode::GetDisplayNameColor() const
 		{
 			return FSequencerDisplayNode::GetDisplayNameColor();
 		}
-	
-		return FLinearColor::Yellow;
+
+		if (NumValidObjects > 0)
+		{
+			return FLinearColor::Yellow;
+		}
 	}
 
 	// Spawnables don't have valid object bindings when their track hasn't spawned them yet,
@@ -564,13 +567,13 @@ FText FSequencerObjectBindingNode::GetDisplayNameToolTipText() const
 {
 	FSequencer& Sequencer = ParentTree.GetSequencer();
 
-	TArrayView<TWeakObjectPtr<>> BoundObjects = Sequencer.FindObjectsInCurrentSequence(ObjectBinding);
+	TArrayView<TWeakObjectPtr<> > BoundObjects = Sequencer.FindBoundObjects(ObjectBinding, Sequencer.GetFocusedTemplateID());
 
 	if ( BoundObjects.Num() == 0 )
 	{
 		return LOCTEXT("InvalidBoundObjectToolTip", "The object bound to this track is missing.");
 	}
-	else if (BoundObjects.Num() > 1)
+	else
 	{
 		TArray<FString> ValidBoundObjectLabels;
 		bool bAddEllipsis = false;
@@ -601,6 +604,16 @@ FText FSequencerObjectBindingNode::GetDisplayNameToolTipText() const
 			}
 		}
 
+		// If only 1 bound object, no need to display tooltip
+		if (ValidBoundObjectLabels.Num() == 1 && NumMissing == 0)
+		{
+			return FText();
+		}
+		else if (ValidBoundObjectLabels.Num() == 0 && NumMissing == 1)
+		{
+			return LOCTEXT("InvalidBoundObjectToolTip", "The object bound to this track is missing.");
+		}
+
 		FString MultipleBoundObjectLabel = FString::Join(ValidBoundObjectLabels, TEXT(", "));
 		if (bAddEllipsis)
 		{
@@ -614,10 +627,6 @@ FText FSequencerObjectBindingNode::GetDisplayNameToolTipText() const
 
 		return FText::FromString(MultipleBoundObjectLabel);
 	}
-	else
-	{
-		return FText();
-	}
 }
 
 const FSlateBrush* FSequencerObjectBindingNode::GetIconBrush() const
