Set section row index when pasting.

diff --git a/Engine/Source/Editor/Sequencer/Private/Sequencer.cpp b/Engine/Source/Editor/Sequencer/Private/Sequencer.cpp
--- a/Engine/Source/Editor/Sequencer/Private/Sequencer.cpp
+++ b/Engine/Source/Editor/Sequencer/Private/Sequencer.cpp
@@ -7278,6 +7278,7 @@ bool FSequencer::PasteSections(const FString& TextToImport, TArray<FNotification
 
 				Section->Rename(nullptr, Track);
 				Track->AddSection(*Section);
+				Section->SetRowIndex(MovieSceneToolHelpers::FindAvailableRowIndex(Track, Section));
 
 				if (Section->HasStartFrame())
 				{
