Sequencer: Add option in content browser context menu for level sequences to open maps containing the level sequence

diff --git a/Engine/Plugins/VirtualProduction/Takes/Source/TakeRecorder/Private/TakeRecorderModule.cpp b/Engine/Plugins/VirtualProduction/Takes/Source/TakeRecorder/Private/TakeRecorderModule.cpp
--- a/Engine/Plugins/VirtualProduction/Takes/Source/TakeRecorder/Private/TakeRecorderModule.cpp
+++ b/Engine/Plugins/VirtualProduction/Takes/Source/TakeRecorder/Private/TakeRecorderModule.cpp
@@ -218,53 +218,6 @@ void FTakeRecorderModule::RegisterMenus()
 					}
 				)
 			);
-
-			// if this LevelSequence has an associated map, offer to load the map
-
-			bool bHasMapTag = false;
-			FString LSOriginMapPath;
-			TArray<UObject::FAssetRegistryTag> ObjectTags;
-			LevelSequence->GetAssetRegistryTags(ObjectTags);
-			for (UObject::FAssetRegistryTag& AssetRegistryTag : ObjectTags)
-			{
-				if (AssetRegistryTag.Name == UTakeMetaData::AssetRegistryTag_LevelPath && !AssetRegistryTag.Value.IsEmpty())
-				{
-					LSOriginMapPath = AssetRegistryTag.Value;
-					FString MapFilePath;
-					if (FEditorFileUtils::IsMapPackageAsset(LSOriginMapPath, MapFilePath))
-					{
-						bHasMapTag = true;
-					}
-					break;
-				}
-			}
-
-			if (bHasMapTag) {
-				InSection.AddMenuEntry(
-					"TakeRecorderOpenMap_Label",
-					LOCTEXT("TakeRecorderOpenMap_Label", "Open Map"),
-					LOCTEXT("TakeRecorderOpenMap_Tooltip", "Opens the map used to create this Level Sequence Asset"),
-					FSlateIcon(FEditorStyle::GetStyleSetName(), "LevelEditor.Tabs.Levels" ),
-					FExecuteAction::CreateLambda(
-						[LSOriginMapPath]
-						{
-							// If there are any unsaved changes to the current level, see if the user wants to save those first.
-							if (!GIsDemoMode)
-							{
-								bool bPromptUserToSave = true;
-								bool bSaveMapPackages = true;
-								bool bSaveContentPackages = true;
-								if (FEditorFileUtils::SaveDirtyPackages(bPromptUserToSave, bSaveMapPackages, bSaveContentPackages) == false)
-								{
-									return;
-								}
-							}
-
-							FEditorFileUtils::LoadMap(LSOriginMapPath);
-						}
-					)
-				);
-			}
 		}
 	}));
 #endif // WITH_EDITOR
diff --git a/Engine/Plugins/VirtualProduction/Takes/Source/TakesCore/Private/TakeMetaData.cpp b/Engine/Plugins/VirtualProduction/Takes/Source/TakesCore/Private/TakeMetaData.cpp
--- a/Engine/Plugins/VirtualProduction/Takes/Source/TakesCore/Private/TakeMetaData.cpp
+++ b/Engine/Plugins/VirtualProduction/Takes/Source/TakesCore/Private/TakeMetaData.cpp
@@ -211,7 +211,6 @@ void UTakeMetaData::ExtendAssetRegistryTags(TArray<FAssetRegistryTag>& OutTags)
 	OutTags.Emplace(AssetRegistryTag_TakeNumber,  LexToString(TakeNumber),   FAssetRegistryTag::ETagType::TT_Numerical,     FAssetRegistryTag::TD_None);
 	OutTags.Emplace(AssetRegistryTag_Timestamp,   Timestamp.ToString(),      FAssetRegistryTag::ETagType::TT_Chronological, FAssetRegistryTag::TD_Date | FAssetRegistryTag::TD_Time);
 	OutTags.Emplace(AssetRegistryTag_Description, Description,               FAssetRegistryTag::ETagType::TT_Alphabetical,  FAssetRegistryTag::TD_None);
-	OutTags.Emplace(AssetRegistryTag_LevelPath, GetLevelPath(), FAssetRegistryTag::ETagType::TT_Alphabetical,  FAssetRegistryTag::TD_None);
 }
 
 void UTakeMetaData::ExtendAssetRegistryTagMetaData(TMap<FName, FAssetRegistryTagMetadata>& OutMetadata) const
diff --git a/Engine/Source/Editor/Sequencer/Private/SequencerModule.cpp b/Engine/Source/Editor/Sequencer/Private/SequencerModule.cpp
--- a/Engine/Source/Editor/Sequencer/Private/SequencerModule.cpp
+++ b/Engine/Source/Editor/Sequencer/Private/SequencerModule.cpp
@@ -16,6 +16,13 @@
 #include "Tree/CurveEditorTreeFilter.h"
 #include "AnimatedPropertyKey.h"
 
+#include "ToolMenus.h"
+#include "ContentBrowserMenuContexts.h"
+#include "FileHelpers.h"
+#include "LevelSequence.h"
+#include "AssetRegistryModule.h"
+
+
 #define LOCTEXT_NAMESPACE "SequencerEditor"
 
 // Destructor defined in CPP to avoid having to #include SequencerChannelInterface.h in the main module definition
@@ -120,6 +127,85 @@ public:
 		EditorObjectBindingDelegates.RemoveAll([=](const FOnCreateEditorObjectBinding& Delegate) { return Delegate.GetHandle() == InHandle; });
 	}
 
+	void RegisterMenus()
+	{
+		UToolMenus* ToolMenus = UToolMenus::Get();
+		UToolMenu* Menu = ToolMenus->ExtendMenu("ContentBrowser.AssetContextMenu.LevelSequence");
+		if (!Menu)
+		{
+			return;
+		}
+
+		FToolMenuSection& Section = Menu->FindOrAddSection("GetAssetActions");
+		Section.AddDynamicEntry("SequencerActions", FNewToolMenuSectionDelegate::CreateLambda([](FToolMenuSection& InSection)
+		{
+			UContentBrowserAssetContextMenuContext* Context = InSection.FindContext<UContentBrowserAssetContextMenuContext>();
+			if (!Context)
+			{
+				return;
+			}
+
+			ULevelSequence* LevelSequence = Context->SelectedObjects.Num() == 1 ? Cast<ULevelSequence>(Context->SelectedObjects[0]) : nullptr;
+			if (LevelSequence)
+			{
+				// if this LevelSequence has associated maps, offer to load them
+
+				FAssetRegistryModule& AssetRegistryModule = FModuleManager::LoadModuleChecked<FAssetRegistryModule>("AssetRegistry");
+				const FName LSMapPathName = *LevelSequence->GetOutermost()->GetPathName();
+
+				TArray<FString> AssociatedMaps;
+
+				TArray<FAssetIdentifier> AssociatedAssets;
+
+				// This makes the assumption these functions will append the array, and not clear it.
+				AssetRegistryModule.Get().GetReferencers(LSMapPathName, AssociatedAssets);
+				AssetRegistryModule.Get().GetDependencies(LSMapPathName, AssociatedAssets);
+
+				for (FAssetIdentifier& AssociatedMap : AssociatedAssets)
+				{
+					FString MapFilePath;
+					FString LevelPath = AssociatedMap.PackageName.ToString();
+					if (FEditorFileUtils::IsMapPackageAsset(LevelPath, MapFilePath))
+					{
+						AssociatedMaps.AddUnique(LevelPath);
+					}
+				}
+
+				AssociatedMaps.Sort([](const FString& One, const FString& Two){ return FPaths::GetBaseFilename(One) < FPaths::GetBaseFilename(Two); });
+
+				if(AssociatedMaps.Num()>0)
+				{
+					InSection.AddSubMenu(
+						"SequencerOpenMap_Label",
+						LOCTEXT("SequencerOpenMap_Label", "Open Map"),
+						LOCTEXT("SequencerOpenMap_Tooltip", "Open a map associated with this Level Sequence Asset"),
+						FNewMenuDelegate::CreateLambda(
+							[AssociatedMaps](FMenuBuilder& SubMenuBuilder)
+							{
+								for (const FString& AssociatedMap : AssociatedMaps)
+								{
+									SubMenuBuilder.AddMenuEntry(
+										FText::FromString(FPaths::GetBaseFilename(AssociatedMap)),
+										FText(),
+										FSlateIcon(FEditorStyle::GetStyleSetName(), "LevelEditor.Tabs.Levels"),
+										FExecuteAction::CreateLambda(
+											[AssociatedMap]
+											{
+												FEditorFileUtils::LoadMap(AssociatedMap);
+											}
+										)
+									);
+								}
+							}
+						),
+						false,
+						FSlateIcon(FEditorStyle::GetStyleSetName(), "LevelEditor.Tabs.Levels")
+					);
+				}
+			}
+		}));
+	}
+
 	virtual void StartupModule() override
 	{
 		if (GIsEditor)
@@ -131,6 +217,15 @@ public:
 				NSLOCTEXT("Sequencer", "SequencerEditMode", "Sequencer Mode"),
 				FSlateIcon(),
 				false);
+
+			if (UToolMenus::TryGet())
+			{
+				RegisterMenus();
+			}
+			else
+			{
+				FCoreDelegates::OnPostEngineInit.AddRaw(this, &FSequencerModule::RegisterMenus);
+			}
 		}
 
 		ObjectBindingContextMenuExtensibilityManager = MakeShareable( new FExtensibilityManager );
diff --git a/Engine/Source/Editor/Sequencer/Sequencer.Build.cs b/Engine/Source/Editor/Sequencer/Sequencer.Build.cs
--- a/Engine/Source/Editor/Sequencer/Sequencer.Build.cs
+++ b/Engine/Source/Editor/Sequencer/Sequencer.Build.cs
@@ -24,7 +24,9 @@ public class Sequencer : ModuleRules
 			new string[] {
 				"AppFramework", 
 				"ApplicationCore",
-                "CinematicCamera",
+				"AssetRegistry",
+				"CinematicCamera",
+				"ContentBrowser",
 				"Core", 
 				"CoreUObject", 
                 "CurveEditor",
@@ -45,8 +47,9 @@ public class Sequencer : ModuleRules
 				"LevelSequence",
 				"GraphEditor",
                 "ViewportInteraction",
-                "SerializedRecorderInterface"
-            }
+                "SerializedRecorderInterface",
+				"ToolMenus"
+			}
 		);
 
 		CircularlyReferencedDependentModules.AddRange(
