Sequencer: Add option to delete a node and keep the animated state

diff --git a/Engine/Source/Editor/Sequencer/Private/DisplayNodes/SequencerDisplayNode.cpp b/Engine/Source/Editor/Sequencer/Private/DisplayNodes/SequencerDisplayNode.cpp
--- a/Engine/Source/Editor/Sequencer/Private/DisplayNodes/SequencerDisplayNode.cpp
+++ b/Engine/Source/Editor/Sequencer/Private/DisplayNodes/SequencerDisplayNode.cpp
@@ -1039,9 +1039,19 @@ void FSequencerDisplayNode::BuildContextMenu(FMenuBuilder& MenuBuilder)
 			LOCTEXT("DeleteNode", "Delete"),
 			LOCTEXT("DeleteNodeTooltip", "Delete this or selected tracks"),
 			FSlateIcon(FEditorStyle::GetStyleSetName(), "ContentBrowser.AssetActions.Delete"),
-			FUIAction(FExecuteAction::CreateSP(&GetSequencer(), &FSequencer::DeleteNode, ThisNode), CanExecute)
+			FUIAction(FExecuteAction::CreateSP(&GetSequencer(), &FSequencer::DeleteNode, ThisNode, false), CanExecute)
 		);
 
+		if (ThisNode->GetType() == ESequencerNode::Object)
+		{
+			MenuBuilder.AddMenuEntry(
+				LOCTEXT("DeleteNodeAndKeepState", "Delete and Keep State"),
+				LOCTEXT("DeleteNodeAndKeepStateTooltip", "Delete this object's tracks and keep its current animated state"),
+				FSlateIcon(FEditorStyle::GetStyleSetName(), "ContentBrowser.AssetActions.Delete"),
+				FUIAction(FExecuteAction::CreateSP(&GetSequencer(), &FSequencer::DeleteNode, ThisNode, true), CanExecute)
+			);
+		}
+
 		MenuBuilder.AddMenuEntry(
 			LOCTEXT("RenameNode", "Rename"),
 			LOCTEXT("RenameNodeTooltip", "Rename this track"),
diff --git a/Engine/Source/Editor/Sequencer/Private/Sequencer.cpp b/Engine/Source/Editor/Sequencer/Private/Sequencer.cpp
--- a/Engine/Source/Editor/Sequencer/Private/Sequencer.cpp
+++ b/Engine/Source/Editor/Sequencer/Private/Sequencer.cpp
@@ -4534,7 +4534,7 @@ void GetParentTrackNodeAndNamePath(TSharedRef<const FSequencerDisplayNode> Displ
 }
 
 
-bool FSequencer::OnRequestNodeDeleted( TSharedRef<const FSequencerDisplayNode> NodeToBeDeleted )
+bool FSequencer::OnRequestNodeDeleted( TSharedRef<const FSequencerDisplayNode> NodeToBeDeleted, const bool bKeepState )
 {
 	bool bAnythingRemoved = false;
 	
@@ -4557,7 +4557,7 @@ bool FSequencer::OnRequestNodeDeleted( TSharedRef<const FSequencerDisplayNode> N
 		// Delete Children
 		for ( const TSharedRef<FSequencerDisplayNode>& ChildNode : NodeToBeDeleted->GetChildNodes() )
 		{
-			OnRequestNodeDeleted( ChildNode );
+			OnRequestNodeDeleted( ChildNode, bKeepState );
 		}
 
 		// Delete from parent, or root.
@@ -4583,7 +4583,7 @@ bool FSequencer::OnRequestNodeDeleted( TSharedRef<const FSequencerDisplayNode> N
 		{
 			if (ChildNode->GetType() == ESequencerNode::Object)
 			{
-				OnRequestNodeDeleted(ChildNode);
+				OnRequestNodeDeleted(ChildNode, bKeepState);
 			}
 		}
 
@@ -4596,6 +4596,23 @@ bool FSequencer::OnRequestNodeDeleted( TSharedRef<const FSequencerDisplayNode> N
 			ParentFolder->GetFolder().RemoveChildObjectBinding( BindingToRemove );
 		}
 		
+		if (bKeepState)
+		{
+			for (TWeakObjectPtr<> WeakObject : FindBoundObjects(BindingToRemove, ActiveTemplateIDs.Top()))
+			{
+				TArray<UObject*> SubObjects;
+				GetObjectsWithOuter(WeakObject.Get(), SubObjects);
+				PreAnimatedState.DiscardAndRemoveEntityTokensForObject(*WeakObject.Get());
+				for (UObject* SubObject : SubObjects)
+				{
+					if (SubObject)
+					{
+						PreAnimatedState.DiscardAndRemoveEntityTokensForObject(*SubObject);
+					}
+				}
+			}
+		}
+				
 		// Try to remove as a spawnable first
 		if (OwnerMovieScene->RemoveSpawnable(BindingToRemove))
 		{
@@ -5789,7 +5806,7 @@ void FSequencer::DeleteSelectedItems()
 	}
 	else if (Selection.GetSelectedOutlinerNodes().Num())
 	{
-		DeleteSelectedNodes();
+		DeleteSelectedNodes(false);
 	}
 }
 
@@ -6247,17 +6264,17 @@ void FSequencer::RemoveInvalidBindings(FGuid InObjectBinding)
 	NotifyMovieSceneDataChanged(EMovieSceneDataChangeType::MovieSceneStructureItemsChanged);
 }
 
-void FSequencer::DeleteNode(TSharedRef<FSequencerDisplayNode> NodeToBeDeleted)
+void FSequencer::DeleteNode(TSharedRef<FSequencerDisplayNode> NodeToBeDeleted, const bool bKeepState)
 {
 	// If this node is selected, delete all selected nodes
 	if (GetSelection().IsSelected(NodeToBeDeleted))
 	{
-		DeleteSelectedNodes();
+		DeleteSelectedNodes(bKeepState);
 	}
 	else
 	{
 		const FScopedTransaction Transaction( NSLOCTEXT("Sequencer", "UndoDeletingObject", "Delete Node") );
-		bool bAnythingDeleted = OnRequestNodeDeleted(NodeToBeDeleted);
+		bool bAnythingDeleted = OnRequestNodeDeleted(NodeToBeDeleted, bKeepState);
 		if ( bAnythingDeleted )
 		{
 			NotifyMovieSceneDataChanged( EMovieSceneDataChangeType::MovieSceneStructureItemRemoved );
@@ -6266,7 +6283,7 @@ void FSequencer::DeleteNode(TSharedRef<FSequencerDisplayNode> NodeToBeDeleted)
 }
 
 
-void FSequencer::DeleteSelectedNodes()
+void FSequencer::DeleteSelectedNodes(const bool bKeepState)
 {
 	TSet< TSharedRef<FSequencerDisplayNode> > SelectedNodesCopy = GetSelection().GetSelectedOutlinerNodes();
 
@@ -6285,7 +6302,7 @@ void FSequencer::DeleteSelectedNodes()
 		{
 			// Delete everything in the entire node
 			TSharedRef<const FSequencerDisplayNode> NodeToBeDeleted = StaticCastSharedRef<const FSequencerDisplayNode>(SelectedNode);
-			bAnythingDeleted |= OnRequestNodeDeleted( NodeToBeDeleted );
+			bAnythingDeleted |= OnRequestNodeDeleted( NodeToBeDeleted, bKeepState );
 		}
 	}
 
diff --git a/Engine/Source/Editor/Sequencer/Private/Sequencer.h b/Engine/Source/Editor/Sequencer/Private/Sequencer.h
--- a/Engine/Source/Editor/Sequencer/Private/Sequencer.h
+++ b/Engine/Source/Editor/Sequencer/Private/Sequencer.h
@@ -369,7 +369,7 @@ public:
 	 * @param NodeToBeDeleted	Node with data that should be deleted
 	 * @return true if anything was deleted, otherwise false.
 	 */
-	virtual bool OnRequestNodeDeleted( TSharedRef<const FSequencerDisplayNode> NodeToBeDeleted );
+	virtual bool OnRequestNodeDeleted( TSharedRef<const FSequencerDisplayNode> NodeToBeDeleted, const bool bKeepState );
 
 	/** Zooms to the edges of all currently selected sections. */
 	void ZoomToSelectedSections();
@@ -488,8 +488,8 @@ public:
 	void RemoveInvalidBindings(FGuid ObjectBinding);
 
 	/** Called when a user executes the delete node menu item */
-	void DeleteNode(TSharedRef<FSequencerDisplayNode> NodeToBeDeleted);
-	void DeleteSelectedNodes();
+	void DeleteNode(TSharedRef<FSequencerDisplayNode> NodeToBeDeleted, const bool bKeepState);
+	void DeleteSelectedNodes(const bool bKeepState);
 
 	/** @return The list of nodes which must be moved to move the current selected nodes */
 	TArray<TSharedRef<FSequencerDisplayNode> > GetSelectedNodesToMove();
diff --git a/Engine/Source/Runtime/MovieScene/Private/Evaluation/MovieScenePreAnimatedState.cpp b/Engine/Source/Runtime/MovieScene/Private/Evaluation/MovieScenePreAnimatedState.cpp
--- a/Engine/Source/Runtime/MovieScene/Private/Evaluation/MovieScenePreAnimatedState.cpp
+++ b/Engine/Source/Runtime/MovieScene/Private/Evaluation/MovieScenePreAnimatedState.cpp
@@ -371,6 +371,25 @@ void FMovieScenePreAnimatedState::DiscardEntityTokens()
 	MasterTokens.DiscardEntityTokens();
 }
 
+void FMovieScenePreAnimatedState::DiscardAndRemoveEntityTokensForObject(UObject& Object)
+{
+	FObjectKey ObjectKey(&Object);
+
+	auto* FoundObjectTokens = ObjectTokens.Find(ObjectKey);
+	if (FoundObjectTokens)
+	{
+		FoundObjectTokens->DiscardEntityTokens();
+
+		ObjectTokens.Remove(ObjectKey);
+	}
+
+	for (auto& Pair : EntityToAnimatedObjects)
+	{
+		Pair.Value.Remove(ObjectKey);
+	}
+}
+
+
 /** Explicit, exported template instantiations */
 template struct MOVIESCENE_API TMovieSceneSavedTokens<IMovieScenePreAnimatedTokenPtr>;
 template struct MOVIESCENE_API TMovieSceneSavedTokens<IMovieScenePreAnimatedGlobalTokenPtr>;
diff --git a/Engine/Source/Runtime/MovieScene/Public/Evaluation/MovieScenePreAnimatedState.h b/Engine/Source/Runtime/MovieScene/Public/Evaluation/MovieScenePreAnimatedState.h
--- a/Engine/Source/Runtime/MovieScene/Public/Evaluation/MovieScenePreAnimatedState.h
+++ b/Engine/Source/Runtime/MovieScene/Public/Evaluation/MovieScenePreAnimatedState.h
@@ -277,6 +277,12 @@ public:
 	 */
 	MOVIESCENE_API void DiscardEntityTokens();
 
+	/**
+	 * Discard any tokens that relate to the requested object (ie sections or tracks) without restoring the values.
+	 * Any global pre-animated state tokens for this object will be removed.
+	 */
+	MOVIESCENE_API void DiscardAndRemoveEntityTokensForObject(UObject& Object);
+
 public:
 
 	void SetCaptureEntity(FMovieSceneEvaluationKey InEntity, EMovieSceneCompletionMode InCompletionMode)
