Sequencer: Add MovieSceneChannel::DeleteKeysFrom

Changed trim left/right and split to remove keys outside of the trim and split ranges.

MovieSceneFloatChannel implements DeleteKeysFrom by inserting a key at the split point so that the shape of the curve is retained.

diff --git a/Engine/Plugins/Experimental/GeometryCache/Source/GeometryCacheTracks/Classes/MovieSceneGeometryCacheSection.h b/Engine/Plugins/Experimental/GeometryCache/Source/GeometryCacheTracks/Classes/MovieSceneGeometryCacheSection.h
--- a/Engine/Plugins/Experimental/GeometryCache/Source/GeometryCacheTracks/Classes/MovieSceneGeometryCacheSection.h
+++ b/Engine/Plugins/Experimental/GeometryCache/Source/GeometryCacheTracks/Classes/MovieSceneGeometryCacheSection.h
@@ -74,8 +74,8 @@ public:
 protected:
 	//~ UMovieSceneSection interface
 	virtual TOptional<TRange<FFrameNumber> > GetAutoSizeRange() const override;
-	virtual void TrimSection(FQualifiedFrameTime TrimTime, bool bTrimLeft) override;
-	virtual UMovieSceneSection* SplitSection(FQualifiedFrameTime SplitTime) override;
+	virtual void TrimSection(FQualifiedFrameTime TrimTime, bool bTrimLeft, bool bDeleteKeys) override;
+	virtual UMovieSceneSection* SplitSection(FQualifiedFrameTime SplitTime, bool bDeleteKeys) override;
 	virtual void GetSnapTimes(TArray<FFrameNumber>& OutSnapTimes, bool bGetSectionBorders) const override;
 	virtual TOptional<FFrameTime> GetOffsetTime() const override;
 	virtual FMovieSceneEvalTemplatePtr GenerateTemplate() const override;
diff --git a/Engine/Plugins/Experimental/GeometryCache/Source/GeometryCacheTracks/Private/MovieSceneGeometryCacheSection.cpp b/Engine/Plugins/Experimental/GeometryCache/Source/GeometryCacheTracks/Private/MovieSceneGeometryCacheSection.cpp
--- a/Engine/Plugins/Experimental/GeometryCache/Source/GeometryCacheTracks/Private/MovieSceneGeometryCacheSection.cpp
+++ b/Engine/Plugins/Experimental/GeometryCache/Source/GeometryCacheTracks/Private/MovieSceneGeometryCacheSection.cpp
@@ -112,7 +112,7 @@ TOptional<TRange<FFrameNumber> > UMovieSceneGeometryCacheSection::GetAutoSizeRan
 }
 
 
-void UMovieSceneGeometryCacheSection::TrimSection(FQualifiedFrameTime TrimTime, bool bTrimLeft)
+void UMovieSceneGeometryCacheSection::TrimSection(FQualifiedFrameTime TrimTime, bool bTrimLeft, bool bDeleteKeys)
 {
 	SetFlags(RF_Transactional);
 
@@ -125,17 +125,17 @@ void UMovieSceneGeometryCacheSection::TrimSection(FQualifiedFrameTime TrimTime,
 			Params.FirstLoopStartFrameOffset = HasStartFrame() ? GetFirstLoopStartOffsetAtTrimTime(TrimTime, Params, GetInclusiveStartFrame(), FrameRate) : 0;
 		}
 
-		Super::TrimSection(TrimTime, bTrimLeft);
+		Super::TrimSection(TrimTime, bTrimLeft, bDeleteKeys);
 	}
 }
 
-UMovieSceneSection* UMovieSceneGeometryCacheSection::SplitSection(FQualifiedFrameTime SplitTime)
+UMovieSceneSection* UMovieSceneGeometryCacheSection::SplitSection(FQualifiedFrameTime SplitTime, bool bDeleteKeys)
 {
 	FFrameRate FrameRate = GetTypedOuter<UMovieScene>()->GetTickResolution();
 
 	const FFrameNumber NewOffset = HasStartFrame() ? GetFirstLoopStartOffsetAtTrimTime(SplitTime, Params, GetInclusiveStartFrame(), FrameRate) : 0;
 
-	UMovieSceneSection* NewSection = Super::SplitSection(SplitTime);
+	UMovieSceneSection* NewSection = Super::SplitSection(SplitTime, bDeleteKeys);
 	if (NewSection != nullptr)
 	{
 		UMovieSceneGeometryCacheSection* NewGeometrySection = Cast<UMovieSceneGeometryCacheSection>(NewSection);
diff --git a/Engine/Plugins/Experimental/GeometryCollectionPlugin/Source/GeometryCollectionTracks/Private/MovieSceneGeometryCollectionSection.cpp b/Engine/Plugins/Experimental/GeometryCollectionPlugin/Source/GeometryCollectionTracks/Private/MovieSceneGeometryCollectionSection.cpp
--- a/Engine/Plugins/Experimental/GeometryCollectionPlugin/Source/GeometryCollectionTracks/Private/MovieSceneGeometryCollectionSection.cpp
+++ b/Engine/Plugins/Experimental/GeometryCollectionPlugin/Source/GeometryCollectionTracks/Private/MovieSceneGeometryCollectionSection.cpp
@@ -65,7 +65,7 @@ TOptional<TRange<FFrameNumber> > UMovieSceneGeometryCollectionSection::GetAutoSi
 }
 
 
-void UMovieSceneGeometryCollectionSection::TrimSection(FQualifiedFrameTime TrimTime, bool bTrimLeft)
+void UMovieSceneGeometryCollectionSection::TrimSection(FQualifiedFrameTime TrimTime, bool bTrimLeft, bool bDeleteKeys)
 {
 	SetFlags(RF_Transactional);
 
@@ -78,17 +78,17 @@ void UMovieSceneGeometryCollectionSection::TrimSection(FQualifiedFrameTime TrimT
 			Params.StartFrameOffset = HasStartFrame() ? GetStartOffsetAtTrimTime(TrimTime, Params, GetInclusiveStartFrame(), FrameRate) : 0;
 		}
 
-		Super::TrimSection(TrimTime, bTrimLeft);
+		Super::TrimSection(TrimTime, bTrimLeft, bDeleteKeys);
 	}
 }
 
-UMovieSceneSection* UMovieSceneGeometryCollectionSection::SplitSection(FQualifiedFrameTime SplitTime)
+UMovieSceneSection* UMovieSceneGeometryCollectionSection::SplitSection(FQualifiedFrameTime SplitTime, bool bDeleteKeys)
 {
 	FFrameRate FrameRate = GetTypedOuter<UMovieScene>()->GetTickResolution();
 
 	const FFrameNumber NewOffset = HasStartFrame() ? GetStartOffsetAtTrimTime(SplitTime, Params, GetInclusiveStartFrame(), FrameRate) : 0;
 
-	UMovieSceneSection* NewSection = Super::SplitSection(SplitTime);
+	UMovieSceneSection* NewSection = Super::SplitSection(SplitTime, bDeleteKeys);
 	if (NewSection != nullptr)
 	{
 		UMovieSceneGeometryCollectionSection* NewGeometrySection = Cast<UMovieSceneGeometryCollectionSection>(NewSection);
diff --git a/Engine/Plugins/Experimental/GeometryCollectionPlugin/Source/GeometryCollectionTracks/Public/MovieSceneGeometryCollectionSection.h b/Engine/Plugins/Experimental/GeometryCollectionPlugin/Source/GeometryCollectionTracks/Public/MovieSceneGeometryCollectionSection.h
--- a/Engine/Plugins/Experimental/GeometryCollectionPlugin/Source/GeometryCollectionTracks/Public/MovieSceneGeometryCollectionSection.h
+++ b/Engine/Plugins/Experimental/GeometryCollectionPlugin/Source/GeometryCollectionTracks/Public/MovieSceneGeometryCollectionSection.h
@@ -84,8 +84,8 @@ public:
 protected:
 	//~ UMovieSceneSection interface
 	virtual TOptional<TRange<FFrameNumber> > GetAutoSizeRange() const override;
-	virtual void TrimSection(FQualifiedFrameTime TrimTime, bool bTrimLeft) override;
-	virtual UMovieSceneSection* SplitSection(FQualifiedFrameTime SplitTime) override;
+	virtual void TrimSection(FQualifiedFrameTime TrimTime, bool bTrimLeft, bool bDeleteKeys) override;
+	virtual UMovieSceneSection* SplitSection(FQualifiedFrameTime SplitTime, bool bDeleteKeys) override;
 	virtual void GetSnapTimes(TArray<FFrameNumber>& OutSnapTimes, bool bGetSectionBorders) const override;
 	virtual TOptional<FFrameTime> GetOffsetTime() const override;
 	virtual FMovieSceneEvalTemplatePtr GenerateTemplate() const override;
diff --git a/Engine/Plugins/FX/Niagara/Source/NiagaraEditor/Private/Sequencer/NiagaraSequence/Sections/MovieSceneNiagaraEmitterSection.cpp b/Engine/Plugins/FX/Niagara/Source/NiagaraEditor/Private/Sequencer/NiagaraSequence/Sections/MovieSceneNiagaraEmitterSection.cpp
--- a/Engine/Plugins/FX/Niagara/Source/NiagaraEditor/Private/Sequencer/NiagaraSequence/Sections/MovieSceneNiagaraEmitterSection.cpp
+++ b/Engine/Plugins/FX/Niagara/Source/NiagaraEditor/Private/Sequencer/NiagaraSequence/Sections/MovieSceneNiagaraEmitterSection.cpp
@@ -67,6 +67,21 @@ void FMovieSceneNiagaraEmitterChannel::DeleteKeys(TArrayView<const FKeyHandle> I
 	GetData().DeleteKeys(InHandles);
 }
 
+void FMovieSceneNiagaraEmitterChannel::DeleteKeysFrom(FFrameNumber InTime, bool bDeleteKeysBefore)
+{
+	// Insert a key at the current time to maintain evaluation
+	if (GetData().GetTimes().Num() > 0)
+	{
+		FNiagaraEmitterSectionKey Value;
+		if (Evaluate(InTime, Value))
+		{
+			GetData().UpdateOrAddKey(InTime, Value);
+		}
+	}
+
+	GetData().DeleteKeysFrom(InTime, bDeleteKeysBefore);
+}
+
 void FMovieSceneNiagaraEmitterChannel::ChangeFrameResolution(FFrameRate SourceRate, FFrameRate DestinationRate)
 {
 	GetData().ChangeFrameResolution(SourceRate, DestinationRate);
diff --git a/Engine/Plugins/FX/Niagara/Source/NiagaraEditor/Private/Sequencer/NiagaraSequence/Sections/MovieSceneNiagaraEmitterSection.h b/Engine/Plugins/FX/Niagara/Source/NiagaraEditor/Private/Sequencer/NiagaraSequence/Sections/MovieSceneNiagaraEmitterSection.h
--- a/Engine/Plugins/FX/Niagara/Source/NiagaraEditor/Private/Sequencer/NiagaraSequence/Sections/MovieSceneNiagaraEmitterSection.h
+++ b/Engine/Plugins/FX/Niagara/Source/NiagaraEditor/Private/Sequencer/NiagaraSequence/Sections/MovieSceneNiagaraEmitterSection.h
@@ -95,6 +95,7 @@ public:
 	virtual void SetKeyTimes(TArrayView<const FKeyHandle> InHandles, TArrayView<const FFrameNumber> InKeyTimes) override;
 	virtual void DuplicateKeys(TArrayView<const FKeyHandle> InHandles, TArrayView<FKeyHandle> OutNewHandles) override;
 	virtual void DeleteKeys(TArrayView<const FKeyHandle> InHandles) override;
+	virtual void DeleteKeysFrom(FFrameNumber InTime, bool bDeleteKeysBefore) override;
 	virtual void ChangeFrameResolution(FFrameRate SourceRate, FFrameRate DestinationRate) override;
 	virtual TRange<FFrameNumber> ComputeEffectiveRange() const override;
 	virtual int32 GetNumKeys() const override;
diff --git a/Engine/Source/Editor/MovieSceneTools/Private/MovieSceneToolHelpers.cpp b/Engine/Source/Editor/MovieSceneTools/Private/MovieSceneToolHelpers.cpp
--- a/Engine/Source/Editor/MovieSceneTools/Private/MovieSceneToolHelpers.cpp
+++ b/Engine/Source/Editor/MovieSceneTools/Private/MovieSceneToolHelpers.cpp
@@ -63,25 +63,25 @@
 /* MovieSceneToolHelpers
  *****************************************************************************/
 
-void MovieSceneToolHelpers::TrimSection(const TSet<TWeakObjectPtr<UMovieSceneSection>>& Sections, FQualifiedFrameTime Time, bool bTrimLeft)
+void MovieSceneToolHelpers::TrimSection(const TSet<TWeakObjectPtr<UMovieSceneSection>>& Sections, FQualifiedFrameTime Time, bool bTrimLeft, bool bDeleteKeys)
 {
 	for (auto Section : Sections)
 	{
 		if (Section.IsValid())
 		{
-			Section->TrimSection(Time, bTrimLeft);
+			Section->TrimSection(Time, bTrimLeft, bDeleteKeys);
 		}
 	}
 }
 
 
-void MovieSceneToolHelpers::SplitSection(const TSet<TWeakObjectPtr<UMovieSceneSection>>& Sections, FQualifiedFrameTime Time)
+void MovieSceneToolHelpers::SplitSection(const TSet<TWeakObjectPtr<UMovieSceneSection>>& Sections, FQualifiedFrameTime Time, bool bDeleteKeys)
 {
 	for (auto Section : Sections)
 	{
 		if (Section.IsValid())
 		{
-			Section->SplitSection(Time);
+			Section->SplitSection(Time, bDeleteKeys);
 		}
 	}
 }
diff --git a/Engine/Source/Editor/MovieSceneTools/Private/TrackEditors/AttachTrackEditor.cpp b/Engine/Source/Editor/MovieSceneTools/Private/TrackEditors/AttachTrackEditor.cpp
--- a/Engine/Source/Editor/MovieSceneTools/Private/TrackEditors/AttachTrackEditor.cpp
+++ b/Engine/Source/Editor/MovieSceneTools/Private/TrackEditors/AttachTrackEditor.cpp
@@ -884,7 +884,7 @@ void F3DAttachTrackEditor::TrimAndPreserve(FGuid InObjectBinding, UMovieSceneSec
 		UObject* Object = BoundObjects[0].Get();
 
 		// Trim the section and find the range of the cut
-		InSection->TrimSection(QualifiedNewDetachTime, bInTrimLeft);
+		InSection->TrimSection(QualifiedNewDetachTime, bInTrimLeft, true);
 		TArray<TRange<FFrameNumber>> ExcludedRanges = TRange<FFrameNumber>::Difference(BeforeTrimRange, InSection->GetRange());
 		if (ExcludedRanges.Num() == 0)
 		{
diff --git a/Engine/Source/Editor/MovieSceneTools/Public/MovieSceneToolHelpers.h b/Engine/Source/Editor/MovieSceneTools/Public/MovieSceneToolHelpers.h
--- a/Engine/Source/Editor/MovieSceneTools/Public/MovieSceneToolHelpers.h
+++ b/Engine/Source/Editor/MovieSceneTools/Public/MovieSceneToolHelpers.h
@@ -58,16 +58,18 @@ public:
 	 * @param Sections The sections to trim
 	 * @param Time	The time at which to trim
 	 * @param bTrimLeft Trim left or trim right
+	 * @param bDeleteKeys Delete keys outside the split ranges
 	 */
-	static void TrimSection(const TSet<TWeakObjectPtr<UMovieSceneSection>>& Sections, FQualifiedFrameTime Time, bool bTrimLeft);
+	static void TrimSection(const TSet<TWeakObjectPtr<UMovieSceneSection>>& Sections, FQualifiedFrameTime Time, bool bTrimLeft, bool bDeleteKeys);
 
 	/**
 	 * Splits sections at the given time
 	 *
 	 * @param Sections The sections to split
 	 * @param Time	The time at which to split
+	 * @param bDeleteKeys Delete keys outside the split ranges
 	 */
-	static void SplitSection(const TSet<TWeakObjectPtr<UMovieSceneSection>>& Sections, FQualifiedFrameTime Time);
+	static void SplitSection(const TSet<TWeakObjectPtr<UMovieSceneSection>>& Sections, FQualifiedFrameTime Time, bool bDeleteKeys);
 
 	/**
 	 * Parse a shot name into its components.
diff --git a/Engine/Source/Editor/Sequencer/Private/Sequencer.cpp b/Engine/Source/Editor/Sequencer/Private/Sequencer.cpp
--- a/Engine/Source/Editor/Sequencer/Private/Sequencer.cpp
+++ b/Engine/Source/Editor/Sequencer/Private/Sequencer.cpp
@@ -9735,7 +9735,7 @@ void FSequencer::CloseEntryPopupMenu()
 void FSequencer::TrimSection(bool bTrimLeft)
 {
 	FScopedTransaction TrimSectionTransaction( NSLOCTEXT("Sequencer", "TrimSection_Transaction", "Trim Section") );
-	MovieSceneToolHelpers::TrimSection(Selection.GetSelectedSections(), GetLocalTime(), bTrimLeft);
+	MovieSceneToolHelpers::TrimSection(Selection.GetSelectedSections(), GetLocalTime(), bTrimLeft, Settings->GetDeleteKeysWhenTrimming());
 	NotifyMovieSceneDataChanged( EMovieSceneDataChangeType::TrackValueChanged );
 }
 
@@ -9743,7 +9743,7 @@ void FSequencer::TrimSection(bool bTrimLeft)
 void FSequencer::SplitSection()
 {
 	FScopedTransaction SplitSectionTransaction( NSLOCTEXT("Sequencer", "SplitSection_Transaction", "Split Section") );
-	MovieSceneToolHelpers::SplitSection(Selection.GetSelectedSections(), GetLocalTime());
+	MovieSceneToolHelpers::SplitSection(Selection.GetSelectedSections(), GetLocalTime(), Settings->GetDeleteKeysWhenTrimming());
 	NotifyMovieSceneDataChanged( EMovieSceneDataChangeType::MovieSceneStructureItemAdded );
 }
 
diff --git a/Engine/Source/Editor/Sequencer/Private/SequencerContextMenus.cpp b/Engine/Source/Editor/Sequencer/Private/SequencerContextMenus.cpp
--- a/Engine/Source/Editor/Sequencer/Private/SequencerContextMenus.cpp
+++ b/Engine/Source/Editor/Sequencer/Private/SequencerContextMenus.cpp
@@ -415,6 +415,8 @@ void FSectionContextMenu::AddEditMenu(FMenuBuilder& MenuBuilder)
 	// Copy a reference to the context menu by value into each lambda handler to ensure the type stays alive until the menu is closed
 	TSharedRef<FSectionContextMenu> Shared = AsShared();
 
+	MenuBuilder.BeginSection("Trimming", LOCTEXT("TrimmingSectionMenu", "Trimming"));
+
 	MenuBuilder.AddMenuEntry(
 		LOCTEXT("TrimSectionLeft", "Trim Left"),
 		LOCTEXT("TrimSectionLeftTooltip", "Trim section at current MouseDownTime to the left"),
@@ -438,10 +440,26 @@ void FSectionContextMenu::AddEditMenu(FMenuBuilder& MenuBuilder)
 		LOCTEXT("SplitSectionTooltip", "Split section at current MouseDownTime"),
 		FSlateIcon(),
 		FUIAction(
-			FExecuteAction::CreateLambda([=]{ Shared->SplitSection(); }),
-			FCanExecuteAction::CreateLambda([=]{ return Shared->IsTrimmable(); }))
+			FExecuteAction::CreateLambda([=] { Shared->SplitSection(); }),
+			FCanExecuteAction::CreateLambda([=] { return Shared->IsTrimmable(); }))
 	);
+
+	MenuBuilder.AddMenuEntry(
+		LOCTEXT("DeleteKeysWhenTrimming", "Delete Keys"),
+		LOCTEXT("DeleteKeysWhenTrimmingTooltip", "Delete keys outside of the trimmed range"),
+		FSlateIcon(),
+		FUIAction(
+			FExecuteAction::CreateLambda([=] { Sequencer->GetSequencerSettings()->SetDeleteKeysWhenTrimming(!Sequencer->GetSequencerSettings()->GetDeleteKeysWhenTrimming()); }),
+			FCanExecuteAction(),
+			FIsActionChecked::CreateLambda([=] { return Sequencer->GetSequencerSettings()->GetDeleteKeysWhenTrimming(); })),
+		NAME_None,
+		EUserInterfaceActionType::ToggleButton
+	);
+
+	MenuBuilder.EndSection();
 		
+	MenuBuilder.AddMenuSeparator();
+
 	MenuBuilder.AddMenuEntry(
 		LOCTEXT("AutoSizeSection", "Auto Size"),
 		LOCTEXT("AutoSizeSectionTooltip", "Auto size the section length to the duration of the source of this section (ie. audio, animation or shot length)"),
@@ -763,7 +781,7 @@ void FSectionContextMenu::TrimSection(bool bTrimLeft)
 {
 	FScopedTransaction TrimSectionTransaction(LOCTEXT("TrimSection_Transaction", "Trim Section"));
 
-	MovieSceneToolHelpers::TrimSection(Sequencer->GetSelection().GetSelectedSections(), Sequencer->GetLocalTime(), bTrimLeft);
+	MovieSceneToolHelpers::TrimSection(Sequencer->GetSelection().GetSelectedSections(), Sequencer->GetLocalTime(), bTrimLeft, Sequencer->GetSequencerSettings()->GetDeleteKeysWhenTrimming());
 	Sequencer->NotifyMovieSceneDataChanged( EMovieSceneDataChangeType::TrackValueChanged );
 }
 
@@ -775,7 +793,7 @@ void FSectionContextMenu::SplitSection()
 	FFrameNumber CurrentFrame = Sequencer->GetLocalTime().Time.FrameNumber;
 	FQualifiedFrameTime SplitFrame = FQualifiedFrameTime(CurrentFrame, Sequencer->GetFocusedTickResolution());
 
-	MovieSceneToolHelpers::SplitSection(Sequencer->GetSelection().GetSelectedSections(), SplitFrame);
+	MovieSceneToolHelpers::SplitSection(Sequencer->GetSelection().GetSelectedSections(), SplitFrame, Sequencer->GetSequencerSettings()->GetDeleteKeysWhenTrimming());
 	Sequencer->NotifyMovieSceneDataChanged( EMovieSceneDataChangeType::RefreshAllImmediately );
 }
 
diff --git a/Engine/Source/Editor/Sequencer/Private/SequencerSettings.cpp b/Engine/Source/Editor/Sequencer/Private/SequencerSettings.cpp
--- a/Engine/Source/Editor/Sequencer/Private/SequencerSettings.cpp
+++ b/Engine/Source/Editor/Sequencer/Private/SequencerSettings.cpp
@@ -46,6 +46,7 @@ USequencerSettings::USequencerSettings( const FObjectInitializer& ObjectInitiali
 	bShowCombinedKeyframes = true;
 	bInfiniteKeyAreas = false;
 	bShowChannelColors = false;
+	bDeleteKeysWhenTrimming = true;
 	bActivateRealtimeViewports = true;
 	bEvaluateSubSequencesInIsolation = false;
 	bRerunConstructionScripts = true;
@@ -576,6 +577,20 @@ void USequencerSettings::SetShowChannelColors(bool InbShowChannelColors)
 	}
 }
 
+bool USequencerSettings::GetDeleteKeysWhenTrimming() const
+{
+	return bDeleteKeysWhenTrimming;
+}
+
+void USequencerSettings::SetDeleteKeysWhenTrimming(bool bInDeleteKeysWhenTrimming)
+{
+	if (bDeleteKeysWhenTrimming != bInDeleteKeysWhenTrimming)
+	{
+		bDeleteKeysWhenTrimming = bInDeleteKeysWhenTrimming;
+		SaveConfig();
+	}
+}
+
 bool USequencerSettings::ShouldActivateRealtimeViewports() const
 {
 	return bActivateRealtimeViewports;
diff --git a/Engine/Source/Editor/Sequencer/Public/SequencerSettings.h b/Engine/Source/Editor/Sequencer/Public/SequencerSettings.h
--- a/Engine/Source/Editor/Sequencer/Public/SequencerSettings.h
+++ b/Engine/Source/Editor/Sequencer/Public/SequencerSettings.h
@@ -284,6 +284,11 @@ public:
 	/** Set whether to show channel colors */
 	void SetShowChannelColors(bool bInShowChannelColors);
 
+	/** @return true if deleting keys that fall beyond the section range when trimming */
+	bool GetDeleteKeysWhenTrimming() const;
+	/** Set whether to delete keys that fall beyond the section range when trimming */
+	void SetDeleteKeysWhenTrimming(bool bInDeleteKeysWhenTrimming);
+
 	/** @return Whether to activate realtime viewports when in sequencer */
 	bool ShouldActivateRealtimeViewports() const;
 	/** Toggle whether to allow possession of PIE viewports */
@@ -486,9 +491,13 @@ protected:
 	bool bInfiniteKeyAreas;
 
 	/** Enable or disable displaying channel bar colors for vector properties. */
-	UPROPERTY( config, EditAnywhere, Category=Timeline )
+	UPROPERTY(config, EditAnywhere, Category = Timeline)
 	bool bShowChannelColors;
 
+	/** Enable or disable deleting keys that fall beyond the section range when trimming. */
+	UPROPERTY(config, EditAnywhere, Category = Timeline)
+	bool bDeleteKeysWhenTrimming;
+
 	/** When enabled, sequencer will activate 'Realtime' in viewports */
 	UPROPERTY(config, EditAnywhere, Category=General)
 	bool bActivateRealtimeViewports;
diff --git a/Engine/Source/Runtime/MovieScene/Private/Channels/MovieSceneBoolChannel.cpp b/Engine/Source/Runtime/MovieScene/Private/Channels/MovieSceneBoolChannel.cpp
--- a/Engine/Source/Runtime/MovieScene/Private/Channels/MovieSceneBoolChannel.cpp
+++ b/Engine/Source/Runtime/MovieScene/Private/Channels/MovieSceneBoolChannel.cpp
@@ -82,6 +82,21 @@ void FMovieSceneBoolChannel::DeleteKeys(TArrayView<const FKeyHandle> InHandles)
 	GetData().DeleteKeys(InHandles);
 }
 
+void FMovieSceneBoolChannel::DeleteKeysFrom(FFrameNumber InTime, bool bDeleteKeysBefore)
+{
+	// Insert a key at the current time to maintain evaluation
+	if (GetData().GetTimes().Num() > 0)
+	{
+		bool Value = false;
+		if (Evaluate(InTime, Value))
+		{
+			GetData().UpdateOrAddKey(InTime, Value);
+		}
+	}
+
+	GetData().DeleteKeysFrom(InTime, bDeleteKeysBefore);
+}
+
 void FMovieSceneBoolChannel::ChangeFrameResolution(FFrameRate SourceRate, FFrameRate DestinationRate)
 {
 	GetData().ChangeFrameResolution(SourceRate, DestinationRate);
diff --git a/Engine/Source/Runtime/MovieScene/Private/Channels/MovieSceneByteChannel.cpp b/Engine/Source/Runtime/MovieScene/Private/Channels/MovieSceneByteChannel.cpp
--- a/Engine/Source/Runtime/MovieScene/Private/Channels/MovieSceneByteChannel.cpp
+++ b/Engine/Source/Runtime/MovieScene/Private/Channels/MovieSceneByteChannel.cpp
@@ -84,6 +84,21 @@ void FMovieSceneByteChannel::DeleteKeys(TArrayView<const FKeyHandle> InHandles)
 	GetData().DeleteKeys(InHandles);
 }
 
+void FMovieSceneByteChannel::DeleteKeysFrom(FFrameNumber InTime, bool bDeleteKeysBefore)
+{
+	// Insert a key at the current time to maintain evaluation
+	if (GetData().GetTimes().Num() > 0)
+	{
+		uint8 Value = 0;
+		if (Evaluate(InTime, Value))
+		{
+			GetData().UpdateOrAddKey(InTime, Value);
+		}
+	}
+
+	GetData().DeleteKeysFrom(InTime, bDeleteKeysBefore);
+}
+
 void FMovieSceneByteChannel::ChangeFrameResolution(FFrameRate SourceRate, FFrameRate DestinationRate)
 {
 	GetData().ChangeFrameResolution(SourceRate, DestinationRate);
diff --git a/Engine/Source/Runtime/MovieScene/Private/Channels/MovieSceneFloatChannel.cpp b/Engine/Source/Runtime/MovieScene/Private/Channels/MovieSceneFloatChannel.cpp
--- a/Engine/Source/Runtime/MovieScene/Private/Channels/MovieSceneFloatChannel.cpp
+++ b/Engine/Source/Runtime/MovieScene/Private/Channels/MovieSceneFloatChannel.cpp
@@ -896,6 +896,25 @@ void FMovieSceneFloatChannel::DeleteKeys(TArrayView<const FKeyHandle> InHandles)
 	GetData().DeleteKeys(InHandles);
 }
 
+void FMovieSceneFloatChannel::DeleteKeysFrom(FFrameNumber InTime, bool bDeleteKeysBefore)
+{
+	// Insert a key at the current time to maintain evaluation
+	if (GetData().GetTimes().Num() > 0)
+	{
+		int32 KeyHandleIndex = GetData().FindKey(InTime);
+		if (KeyHandleIndex == INDEX_NONE)
+		{
+			float Value = 0.f;
+			if (Evaluate(InTime, Value))
+			{
+				AddCubicKey(InTime, Value);
+			}
+		}
+	}
+
+	GetData().DeleteKeysFrom(InTime, bDeleteKeysBefore);
+}
+
 void FMovieSceneFloatChannel::ChangeFrameResolution(FFrameRate SourceRate, FFrameRate DestinationRate)
 {
 	check(Times.Num() == Values.Num());
diff --git a/Engine/Source/Runtime/MovieScene/Private/Channels/MovieSceneIntegerChannel.cpp b/Engine/Source/Runtime/MovieScene/Private/Channels/MovieSceneIntegerChannel.cpp
--- a/Engine/Source/Runtime/MovieScene/Private/Channels/MovieSceneIntegerChannel.cpp
+++ b/Engine/Source/Runtime/MovieScene/Private/Channels/MovieSceneIntegerChannel.cpp
@@ -83,6 +83,21 @@ void FMovieSceneIntegerChannel::DeleteKeys(TArrayView<const FKeyHandle> InHandle
 	GetData().DeleteKeys(InHandles);
 }
 
+void FMovieSceneIntegerChannel::DeleteKeysFrom(FFrameNumber InTime, bool bDeleteKeysBefore)
+{
+	// Insert a key at the current time to maintain evaluation
+	if (GetData().GetTimes().Num() > 0)
+	{
+		int32 Value = 0;
+		if (Evaluate(InTime, Value))
+		{
+			GetData().UpdateOrAddKey(InTime, Value);
+		}
+	}
+
+	GetData().DeleteKeysFrom(InTime, bDeleteKeysBefore);
+}
+
 void FMovieSceneIntegerChannel::ChangeFrameResolution(FFrameRate SourceRate, FFrameRate DestinationRate)
 {
 	GetData().ChangeFrameResolution(SourceRate, DestinationRate);
diff --git a/Engine/Source/Runtime/MovieScene/Private/Channels/MovieSceneObjectPathChannel.cpp b/Engine/Source/Runtime/MovieScene/Private/Channels/MovieSceneObjectPathChannel.cpp
--- a/Engine/Source/Runtime/MovieScene/Private/Channels/MovieSceneObjectPathChannel.cpp
+++ b/Engine/Source/Runtime/MovieScene/Private/Channels/MovieSceneObjectPathChannel.cpp
@@ -108,6 +108,21 @@ void FMovieSceneObjectPathChannel::DuplicateKeys(TArrayView<const FKeyHandle> In
 	GetData().DuplicateKeys(InHandles, OutNewHandles);
 }
 
+void FMovieSceneObjectPathChannel::DeleteKeysFrom(FFrameNumber InTime, bool bDeleteKeysBefore)
+{
+	// Insert a key at the current time to maintain evaluation
+	if (GetData().GetTimes().Num() > 0)
+	{
+		UObject* Value = nullptr;
+		if (Evaluate(InTime, Value))
+		{
+			GetData().UpdateOrAddKey(InTime, Value);
+		}
+	}
+
+	GetData().DeleteKeysFrom(InTime, bDeleteKeysBefore);
+}
+
 void FMovieSceneObjectPathChannel::DeleteKeys(TArrayView<const FKeyHandle> InHandles)
 {
 	GetData().DeleteKeys(InHandles);
diff --git a/Engine/Source/Runtime/MovieScene/Private/MovieSceneSection.cpp b/Engine/Source/Runtime/MovieScene/Private/MovieSceneSection.cpp
--- a/Engine/Source/Runtime/MovieScene/Private/MovieSceneSection.cpp
+++ b/Engine/Source/Runtime/MovieScene/Private/MovieSceneSection.cpp
@@ -426,7 +426,7 @@ void UMovieSceneSection::InitialPlacementOnRow(const TArray<UMovieSceneSection*>
 	}
 }
 
-UMovieSceneSection* UMovieSceneSection::SplitSection(FQualifiedFrameTime SplitTime)
+UMovieSceneSection* UMovieSceneSection::SplitSection(FQualifiedFrameTime SplitTime, bool bDeleteKeys)
 {
 	if (!SectionRange.Value.Contains(SplitTime.Time.GetFrame()))
 	{
@@ -437,23 +437,18 @@ UMovieSceneSection* UMovieSceneSection::SplitSection(FQualifiedFrameTime SplitTi
 
 	if (TryModify())
 	{
-		TRange<FFrameNumber> StartingRange  = SectionRange.Value;
-		TRange<FFrameNumber> LeftHandRange  = TRange<FFrameNumber>(StartingRange.GetLowerBound(), TRangeBound<FFrameNumber>::Exclusive(SplitTime.Time.GetFrame()));
-		TRange<FFrameNumber> RightHandRange = TRange<FFrameNumber>(TRangeBound<FFrameNumber>::Inclusive(SplitTime.Time.GetFrame()), StartingRange.GetUpperBound());
-
-		// Trim off the right
-		SectionRange = LeftHandRange;
-
-		// Create a new section
+		// Duplicate the current section to be the section on the right side of the trim point
 		UMovieSceneTrack* Track = CastChecked<UMovieSceneTrack>(GetOuter());
 		Track->Modify();
 
 		UMovieSceneSection* NewSection = DuplicateObject<UMovieSceneSection>(this, Track);
 		check(NewSection);
 
-		NewSection->SetRange(RightHandRange);
 		Track->AddSection(*NewSection);
 
+		NewSection->TrimSection(SplitTime, false, bDeleteKeys);
+		TrimSection(SplitTime, true, bDeleteKeys);
+
 		return NewSection;
 	}
 
@@ -461,7 +456,7 @@ UMovieSceneSection* UMovieSceneSection::SplitSection(FQualifiedFrameTime SplitTi
 }
 
 
-void UMovieSceneSection::TrimSection(FQualifiedFrameTime TrimTime, bool bTrimLeft)
+void UMovieSceneSection::TrimSection(FQualifiedFrameTime TrimTime, bool bTrimLeft, bool bDeleteKeys)
 {
 	if (SectionRange.Value.Contains(TrimTime.Time.GetFrame()))
 	{
@@ -476,6 +471,20 @@ void UMovieSceneSection::TrimSection(FQualifiedFrameTime TrimTime, bool bTrimLef
 			{
 				SectionRange.Value.SetUpperBound(TRangeBound<FFrameNumber>::Exclusive(TrimTime.Time.GetFrame()));
 			}
+
+			if (bDeleteKeys)
+			{
+				if (ChannelProxy.IsValid())
+				{
+					for (const FMovieSceneChannelEntry& Entry : ChannelProxy->GetAllEntries())
+					{
+						for (FMovieSceneChannel* Channel : Entry.GetChannels())
+						{
+							Channel->DeleteKeysFrom(TrimTime.Time.GetFrame(), bTrimLeft);
+						}
+					}
+				}
+			}
 		}
 	}
 }
diff --git a/Engine/Source/Runtime/MovieScene/Private/Sections/MovieSceneSubSection.cpp b/Engine/Source/Runtime/MovieScene/Private/Sections/MovieSceneSubSection.cpp
--- a/Engine/Source/Runtime/MovieScene/Private/Sections/MovieSceneSubSection.cpp
+++ b/Engine/Source/Runtime/MovieScene/Private/Sections/MovieSceneSubSection.cpp
@@ -208,7 +208,7 @@ void UMovieSceneSubSection::PostEditChangeProperty(FPropertyChangedEvent& Proper
 }
 #endif
 
-UMovieSceneSection* UMovieSceneSubSection::SplitSection( FQualifiedFrameTime SplitTime )
+UMovieSceneSection* UMovieSceneSubSection::SplitSection( FQualifiedFrameTime SplitTime, bool bDeleteKeys)
 {
 	// GetRange is in owning sequence resolution so we check against the incoming SplitTime without converting it.
 	TRange<FFrameNumber> InitialRange = GetRange();
@@ -219,7 +219,7 @@ UMovieSceneSection* UMovieSceneSubSection::SplitSection( FQualifiedFrameTime Spl
 
 	FFrameNumber InitialStartOffset = Parameters.StartFrameOffset;
 
-	UMovieSceneSubSection* NewSection = Cast<UMovieSceneSubSection>( UMovieSceneSection::SplitSection( SplitTime ) );
+	UMovieSceneSubSection* NewSection = Cast<UMovieSceneSubSection>( UMovieSceneSection::SplitSection( SplitTime, bDeleteKeys ) );
 	if ( NewSection )
 	{
 		if (InitialRange.GetLowerBound().IsClosed())
@@ -273,7 +273,7 @@ TOptional<TRange<FFrameNumber> > UMovieSceneSubSection::GetAutoSizeRange() const
 	return Super::GetAutoSizeRange();
 }
 
-void UMovieSceneSubSection::TrimSection( FQualifiedFrameTime TrimTime, bool bTrimLeft )
+void UMovieSceneSubSection::TrimSection( FQualifiedFrameTime TrimTime, bool bTrimLeft, bool bDeleteKeys)
 {
 	TRange<FFrameNumber> InitialRange = GetRange();
 	if ( !InitialRange.Contains( TrimTime.Time.GetFrame() ) )
@@ -283,7 +283,7 @@ void UMovieSceneSubSection::TrimSection( FQualifiedFrameTime TrimTime, bool bTri
 
 	FFrameNumber InitialStartOffset = Parameters.StartFrameOffset;
 
-	UMovieSceneSection::TrimSection( TrimTime, bTrimLeft );
+	UMovieSceneSection::TrimSection( TrimTime, bTrimLeft, bDeleteKeys );
 
 	// If trimming off the left, set the offset of the shot
 	if ( bTrimLeft && InitialRange.GetLowerBound().IsClosed() )
diff --git a/Engine/Source/Runtime/MovieScene/Public/Channels/MovieSceneBoolChannel.h b/Engine/Source/Runtime/MovieScene/Public/Channels/MovieSceneBoolChannel.h
--- a/Engine/Source/Runtime/MovieScene/Public/Channels/MovieSceneBoolChannel.h
+++ b/Engine/Source/Runtime/MovieScene/Public/Channels/MovieSceneBoolChannel.h
@@ -77,6 +77,7 @@ public:
 	virtual void SetKeyTimes(TArrayView<const FKeyHandle> InHandles, TArrayView<const FFrameNumber> InKeyTimes) override;
 	virtual void DuplicateKeys(TArrayView<const FKeyHandle> InHandles, TArrayView<FKeyHandle> OutNewHandles) override;
 	virtual void DeleteKeys(TArrayView<const FKeyHandle> InHandles) override;
+	virtual void DeleteKeysFrom(FFrameNumber InTime, bool bDeleteKeysBefore) override;
 	virtual void ChangeFrameResolution(FFrameRate SourceRate, FFrameRate DestinationRate) override;
 	virtual TRange<FFrameNumber> ComputeEffectiveRange() const override;
 	virtual int32 GetNumKeys() const override;
diff --git a/Engine/Source/Runtime/MovieScene/Public/Channels/MovieSceneByteChannel.h b/Engine/Source/Runtime/MovieScene/Public/Channels/MovieSceneByteChannel.h
--- a/Engine/Source/Runtime/MovieScene/Public/Channels/MovieSceneByteChannel.h
+++ b/Engine/Source/Runtime/MovieScene/Public/Channels/MovieSceneByteChannel.h
@@ -80,6 +80,7 @@ public:
 	virtual void SetKeyTimes(TArrayView<const FKeyHandle> InHandles, TArrayView<const FFrameNumber> InKeyTimes) override;
 	virtual void DuplicateKeys(TArrayView<const FKeyHandle> InHandles, TArrayView<FKeyHandle> OutNewHandles) override;
 	virtual void DeleteKeys(TArrayView<const FKeyHandle> InHandles) override;
+	virtual void DeleteKeysFrom(FFrameNumber InTime, bool bDeleteKeysBefore) override;
 	virtual void ChangeFrameResolution(FFrameRate SourceRate, FFrameRate DestinationRate) override;
 	virtual TRange<FFrameNumber> ComputeEffectiveRange() const override;
 	virtual int32 GetNumKeys() const override;
diff --git a/Engine/Source/Runtime/MovieScene/Public/Channels/MovieSceneChannel.h b/Engine/Source/Runtime/MovieScene/Public/Channels/MovieSceneChannel.h
--- a/Engine/Source/Runtime/MovieScene/Public/Channels/MovieSceneChannel.h
+++ b/Engine/Source/Runtime/MovieScene/Public/Channels/MovieSceneChannel.h
@@ -84,6 +84,15 @@ struct MOVIESCENE_API FMovieSceneChannel
 	virtual void DeleteKeys(TArrayView<const FKeyHandle> InHandles)
 	{}
 
+	/**
+	 * Delete keys before or after a specified time
+	 *
+	 * @param InTime				Delete keys after this time
+	 * @param bDeleteKeysBefore     Whether to delete keys before the specified time
+	 */
+	virtual void DeleteKeysFrom(FFrameNumber InTime, bool bDeleteKeysBefore)
+	{}
+
 	/**
 	 * Called when the frame resolution of this channel is to be changed.
 	 *
diff --git a/Engine/Source/Runtime/MovieScene/Public/Channels/MovieSceneChannelData.h b/Engine/Source/Runtime/MovieScene/Public/Channels/MovieSceneChannelData.h
--- a/Engine/Source/Runtime/MovieScene/Public/Channels/MovieSceneChannelData.h
+++ b/Engine/Source/Runtime/MovieScene/Public/Channels/MovieSceneChannelData.h
@@ -406,6 +406,41 @@ struct TMovieSceneChannelData : FMovieSceneChannelData
 		}
 	}
 
+
+	/**
+	 * Delete keys before or after a specified time
+	 *
+	 * @param InTime				Delete keys after this time
+	 * @param bDeleteKeysBefore     Whether to delete keys before the specified time
+	 */
+	void DeleteKeysFrom(FFrameNumber InTime, bool bDeleteKeysBefore)
+	{
+		TArray<FFrameNumber> OutKeyTimes;
+		TArray<FKeyHandle> OutKeyHandles;
+		GetKeys(TRange<FFrameNumber>::All(), &OutKeyTimes, &OutKeyHandles);
+
+		TArray<FKeyHandle> KeysToRemove;
+		for (int32 Index = 0; Index < OutKeyTimes.Num(); ++Index)
+		{
+			if (bDeleteKeysBefore)
+			{
+				if (OutKeyTimes[Index] < InTime)
+				{
+					KeysToRemove.Add(OutKeyHandles[Index]);
+				}
+			}
+			else
+			{
+				if (OutKeyTimes[Index] > InTime)
+				{
+					KeysToRemove.Add(OutKeyHandles[Index]);
+				}
+			}
+		}
+
+		DeleteKeys(KeysToRemove);
+	}
+
 	/**
 	 * Remove all the keys from this channel
 	 */
diff --git a/Engine/Source/Runtime/MovieScene/Public/Channels/MovieSceneFloatChannel.h b/Engine/Source/Runtime/MovieScene/Public/Channels/MovieSceneFloatChannel.h
--- a/Engine/Source/Runtime/MovieScene/Public/Channels/MovieSceneFloatChannel.h
+++ b/Engine/Source/Runtime/MovieScene/Public/Channels/MovieSceneFloatChannel.h
@@ -215,6 +215,7 @@ public:
 	virtual void SetKeyTimes(TArrayView<const FKeyHandle> InHandles, TArrayView<const FFrameNumber> InKeyTimes) override;
 	virtual void DuplicateKeys(TArrayView<const FKeyHandle> InHandles, TArrayView<FKeyHandle> OutNewHandles) override;
 	virtual void DeleteKeys(TArrayView<const FKeyHandle> InHandles) override;
+	virtual void DeleteKeysFrom(FFrameNumber InTime, bool bDeleteKeysBefore) override;
 	virtual void ChangeFrameResolution(FFrameRate SourceRate, FFrameRate DestinationRate) override;
 	virtual TRange<FFrameNumber> ComputeEffectiveRange() const override;
 	virtual int32 GetNumKeys() const override;
diff --git a/Engine/Source/Runtime/MovieScene/Public/Channels/MovieSceneIntegerChannel.h b/Engine/Source/Runtime/MovieScene/Public/Channels/MovieSceneIntegerChannel.h
--- a/Engine/Source/Runtime/MovieScene/Public/Channels/MovieSceneIntegerChannel.h
+++ b/Engine/Source/Runtime/MovieScene/Public/Channels/MovieSceneIntegerChannel.h
@@ -77,6 +77,7 @@ public:
 	virtual void SetKeyTimes(TArrayView<const FKeyHandle> InHandles, TArrayView<const FFrameNumber> InKeyTimes) override;
 	virtual void DuplicateKeys(TArrayView<const FKeyHandle> InHandles, TArrayView<FKeyHandle> OutNewHandles) override;
 	virtual void DeleteKeys(TArrayView<const FKeyHandle> InHandles) override;
+	virtual void DeleteKeysFrom(FFrameNumber InTime, bool bDeleteKeysBefore) override;
 	virtual void ChangeFrameResolution(FFrameRate SourceRate, FFrameRate DestinationRate) override;
 	virtual TRange<FFrameNumber> ComputeEffectiveRange() const override;
 	virtual int32 GetNumKeys() const override;
diff --git a/Engine/Source/Runtime/MovieScene/Public/Channels/MovieSceneObjectPathChannel.h b/Engine/Source/Runtime/MovieScene/Public/Channels/MovieSceneObjectPathChannel.h
--- a/Engine/Source/Runtime/MovieScene/Public/Channels/MovieSceneObjectPathChannel.h
+++ b/Engine/Source/Runtime/MovieScene/Public/Channels/MovieSceneObjectPathChannel.h
@@ -123,6 +123,7 @@ public:
 	virtual void SetKeyTimes(TArrayView<const FKeyHandle> InHandles, TArrayView<const FFrameNumber> InKeyTimes) override;
 	virtual void DuplicateKeys(TArrayView<const FKeyHandle> InHandles, TArrayView<FKeyHandle> OutNewHandles) override;
 	virtual void DeleteKeys(TArrayView<const FKeyHandle> InHandles) override;
+	virtual void DeleteKeysFrom(FFrameNumber InTime, bool bDeleteKeysBefore) override;
 	virtual void ChangeFrameResolution(FFrameRate SourceRate, FFrameRate DestinationRate) override;
 	virtual TRange<FFrameNumber> ComputeEffectiveRange() const override;
 	virtual int32 GetNumKeys() const override;
diff --git a/Engine/Source/Runtime/MovieScene/Public/MovieSceneSection.h b/Engine/Source/Runtime/MovieScene/Public/MovieSceneSection.h
--- a/Engine/Source/Runtime/MovieScene/Public/MovieSceneSection.h
+++ b/Engine/Source/Runtime/MovieScene/Public/MovieSceneSection.h
@@ -328,17 +328,25 @@ public:
 	 * Split a section in two at the split time
 	 *
 	 * @param SplitTime The time at which to split
+	 * @param bDeleteKeys Delete keys outside the split ranges
 	 * @return The newly created split section
 	 */
-	MOVIESCENE_API virtual UMovieSceneSection* SplitSection(FQualifiedFrameTime SplitTime);
+	MOVIESCENE_API virtual UMovieSceneSection* SplitSection(FQualifiedFrameTime SplitTime, bool bDeleteKeys);
+
+	UE_DEPRECATED(4.23, "Please use SplitSection(SplitTime, bDeleteKeys) instead.")
+	virtual UMovieSceneSection* SplitSection(FQualifiedFrameTime SplitTime) { return SplitSection(SplitTime, false); }
 
 	/**
 	 * Trim a section at the trim time
 	 *
 	 * @param TrimTime The time at which to trim
 	 * @param bTrimLeft Whether to trim left or right
+	 * @param bDeleteKeys Delete keys outside the split ranges
 	 */
-	MOVIESCENE_API virtual void TrimSection(FQualifiedFrameTime TrimTime, bool bTrimLeft);
+	MOVIESCENE_API virtual void TrimSection(FQualifiedFrameTime TrimTime, bool bTrimLeft, bool bDeleteKeys);
+
+	UE_DEPRECATED(4.23, "Please use TrimSection(SplitTime, bTrimLeft, bDeleteKeys) instead.")
+	virtual void TrimSection(FQualifiedFrameTime SplitTime, bool bTrimLeft) { TrimSection(SplitTime, bTrimLeft, false); }
 
 	/**
 	 * Get the data structure representing the specified keys.
diff --git a/Engine/Source/Runtime/MovieScene/Public/Sections/MovieSceneSubSection.h b/Engine/Source/Runtime/MovieScene/Public/Sections/MovieSceneSubSection.h
--- a/Engine/Source/Runtime/MovieScene/Public/Sections/MovieSceneSubSection.h
+++ b/Engine/Source/Runtime/MovieScene/Public/Sections/MovieSceneSubSection.h
@@ -138,8 +138,8 @@ public:
 
 	//~ UMovieSceneSection interface
 	virtual TOptional<TRange<FFrameNumber> > GetAutoSizeRange() const override;
-	virtual UMovieSceneSection* SplitSection( FQualifiedFrameTime SplitTime ) override;
-	virtual void TrimSection( FQualifiedFrameTime TrimTime, bool bTrimLeft ) override;
+	virtual UMovieSceneSection* SplitSection( FQualifiedFrameTime SplitTime, bool bDeleteKeys ) override;
+	virtual void TrimSection( FQualifiedFrameTime TrimTime, bool bTrimLeft, bool bDeleteKeys) override;
 	virtual TOptional<FFrameTime> GetOffsetTime() const override { return TOptional<FFrameTime>(FFrameTime(Parameters.StartFrameOffset)); }
 
 public:
diff --git a/Engine/Source/Runtime/MovieSceneTracks/Private/Channels/MovieSceneEventChannel.cpp b/Engine/Source/Runtime/MovieSceneTracks/Private/Channels/MovieSceneEventChannel.cpp
--- a/Engine/Source/Runtime/MovieSceneTracks/Private/Channels/MovieSceneEventChannel.cpp
+++ b/Engine/Source/Runtime/MovieSceneTracks/Private/Channels/MovieSceneEventChannel.cpp
@@ -27,6 +27,11 @@ void FMovieSceneEventChannel::DeleteKeys(TArrayView<const FKeyHandle> InHandles)
 	GetData().DeleteKeys(InHandles);
 }
 
+void FMovieSceneEventChannel::DeleteKeysFrom(FFrameNumber InTime, bool bDeleteKeysBefore)
+{
+	GetData().DeleteKeysFrom(InTime, bDeleteKeysBefore);
+}
+
 void FMovieSceneEventChannel::ChangeFrameResolution(FFrameRate SourceRate, FFrameRate DestinationRate)
 {
 	GetData().ChangeFrameResolution(SourceRate, DestinationRate);
diff --git a/Engine/Source/Runtime/MovieSceneTracks/Private/Sections/MovieSceneActorReferenceSection.cpp b/Engine/Source/Runtime/MovieSceneTracks/Private/Sections/MovieSceneActorReferenceSection.cpp
--- a/Engine/Source/Runtime/MovieSceneTracks/Private/Sections/MovieSceneActorReferenceSection.cpp
+++ b/Engine/Source/Runtime/MovieSceneTracks/Private/Sections/MovieSceneActorReferenceSection.cpp
@@ -40,6 +40,18 @@ void FMovieSceneActorReferenceData::DeleteKeys(TArrayView<const FKeyHandle> InHa
 	GetData().DeleteKeys(InHandles);
 }
 
+void FMovieSceneActorReferenceData::DeleteKeysFrom(FFrameNumber InTime, bool bDeleteKeysBefore)
+{
+	// Insert a key at the current time to maintain evaluation
+	if (GetData().GetTimes().Num() > 0)
+	{
+		FMovieSceneActorReferenceKey Value = Evaluate(InTime);
+		GetData().UpdateOrAddKey(InTime, Value);
+	}
+
+	GetData().DeleteKeysFrom(InTime, bDeleteKeysBefore);
+}
+
 void FMovieSceneActorReferenceData::ChangeFrameResolution(FFrameRate SourceRate, FFrameRate DestinationRate)
 {
 	GetData().ChangeFrameResolution(SourceRate, DestinationRate);
diff --git a/Engine/Source/Runtime/MovieSceneTracks/Private/Sections/MovieSceneAudioSection.cpp b/Engine/Source/Runtime/MovieSceneTracks/Private/Sections/MovieSceneAudioSection.cpp
--- a/Engine/Source/Runtime/MovieSceneTracks/Private/Sections/MovieSceneAudioSection.cpp
+++ b/Engine/Source/Runtime/MovieSceneTracks/Private/Sections/MovieSceneAudioSection.cpp
@@ -156,7 +156,7 @@ TOptional<TRange<FFrameNumber> > UMovieSceneAudioSection::GetAutoSizeRange() con
 }
 
 	
-void UMovieSceneAudioSection::TrimSection(FQualifiedFrameTime TrimTime, bool bTrimLeft)
+void UMovieSceneAudioSection::TrimSection(FQualifiedFrameTime TrimTime, bool bTrimLeft, bool bDeleteKeys)
 {
 	SetFlags(RF_Transactional);
 
@@ -167,15 +167,15 @@ void UMovieSceneAudioSection::TrimSection(FQualifiedFrameTime TrimTime, bool bTr
 			StartFrameOffset = HasStartFrame() ? GetStartOffsetAtTrimTime(TrimTime, StartFrameOffset, GetInclusiveStartFrame()) : 0;
 		}
 
-		Super::TrimSection(TrimTime, bTrimLeft);
+		Super::TrimSection(TrimTime, bTrimLeft, bDeleteKeys);
 	}
 }
 
-UMovieSceneSection* UMovieSceneAudioSection::SplitSection(FQualifiedFrameTime SplitTime)
+UMovieSceneSection* UMovieSceneAudioSection::SplitSection(FQualifiedFrameTime SplitTime, bool bDeleteKeys)
 {
 	const FFrameNumber NewOffset = HasStartFrame() ? GetStartOffsetAtTrimTime(SplitTime, StartFrameOffset, GetInclusiveStartFrame()) : 0;
 
-	UMovieSceneSection* NewSection = Super::SplitSection(SplitTime);
+	UMovieSceneSection* NewSection = Super::SplitSection(SplitTime, bDeleteKeys);
 	if (NewSection != nullptr)
 	{
 		UMovieSceneAudioSection* NewAudioSection = Cast<UMovieSceneAudioSection>(NewSection);
diff --git a/Engine/Source/Runtime/MovieSceneTracks/Private/Sections/MovieSceneEventSection.cpp b/Engine/Source/Runtime/MovieSceneTracks/Private/Sections/MovieSceneEventSection.cpp
--- a/Engine/Source/Runtime/MovieSceneTracks/Private/Sections/MovieSceneEventSection.cpp
+++ b/Engine/Source/Runtime/MovieSceneTracks/Private/Sections/MovieSceneEventSection.cpp
@@ -354,6 +354,11 @@ void FMovieSceneEventSectionData::DeleteKeys(TArrayView<const FKeyHandle> InHand
 	GetData().DeleteKeys(InHandles);
 }
 
+void FMovieSceneEventSectionData::DeleteKeysFrom(FFrameNumber InTime, bool bDeleteKeysBefore)
+{
+	GetData().DeleteKeysFrom(InTime, bDeleteKeysBefore);
+}
+
 void FMovieSceneEventSectionData::ChangeFrameResolution(FFrameRate SourceRate, FFrameRate DestinationRate)
 {
 	GetData().ChangeFrameResolution(SourceRate, DestinationRate);
diff --git a/Engine/Source/Runtime/MovieSceneTracks/Private/Sections/MovieSceneSkeletalAnimationSection.cpp b/Engine/Source/Runtime/MovieSceneTracks/Private/Sections/MovieSceneSkeletalAnimationSection.cpp
--- a/Engine/Source/Runtime/MovieSceneTracks/Private/Sections/MovieSceneSkeletalAnimationSection.cpp
+++ b/Engine/Source/Runtime/MovieSceneTracks/Private/Sections/MovieSceneSkeletalAnimationSection.cpp
@@ -196,7 +196,7 @@ TOptional<TRange<FFrameNumber> > UMovieSceneSkeletalAnimationSection::GetAutoSiz
 }
 
 
-void UMovieSceneSkeletalAnimationSection::TrimSection(FQualifiedFrameTime TrimTime, bool bTrimLeft)
+void UMovieSceneSkeletalAnimationSection::TrimSection(FQualifiedFrameTime TrimTime, bool bTrimLeft, bool bDeleteKeys)
 {
 	SetFlags(RF_Transactional);
 
@@ -209,17 +209,17 @@ void UMovieSceneSkeletalAnimationSection::TrimSection(FQualifiedFrameTime TrimTi
 			Params.FirstLoopStartFrameOffset = HasStartFrame() ? GetFirstLoopStartOffsetAtTrimTime(TrimTime, Params, GetInclusiveStartFrame(), FrameRate) : 0;
 		}
 
-		Super::TrimSection(TrimTime, bTrimLeft);
+		Super::TrimSection(TrimTime, bTrimLeft, bDeleteKeys);
 	}
 }
 
-UMovieSceneSection* UMovieSceneSkeletalAnimationSection::SplitSection(FQualifiedFrameTime SplitTime)
+UMovieSceneSection* UMovieSceneSkeletalAnimationSection::SplitSection(FQualifiedFrameTime SplitTime, bool bDeleteKeys)
 {
 	FFrameRate FrameRate = GetTypedOuter<UMovieScene>()->GetTickResolution();
 
 	const FFrameNumber NewOffset = HasStartFrame() ? GetFirstLoopStartOffsetAtTrimTime(SplitTime, Params, GetInclusiveStartFrame(), FrameRate) : 0;
 
-	UMovieSceneSection* NewSection = Super::SplitSection(SplitTime);
+	UMovieSceneSection* NewSection = Super::SplitSection(SplitTime, bDeleteKeys);
 	if (NewSection != nullptr)
 	{
 		UMovieSceneSkeletalAnimationSection* NewSkeletalSection = Cast<UMovieSceneSkeletalAnimationSection>(NewSection);
diff --git a/Engine/Source/Runtime/MovieSceneTracks/Private/Sections/MovieSceneStringChannel.cpp b/Engine/Source/Runtime/MovieSceneTracks/Private/Sections/MovieSceneStringChannel.cpp
--- a/Engine/Source/Runtime/MovieSceneTracks/Private/Sections/MovieSceneStringChannel.cpp
+++ b/Engine/Source/Runtime/MovieSceneTracks/Private/Sections/MovieSceneStringChannel.cpp
@@ -75,6 +75,21 @@ void FMovieSceneStringChannel::DeleteKeys(TArrayView<const FKeyHandle> InHandles
 	GetData().DeleteKeys(InHandles);
 }
 
+void FMovieSceneStringChannel::DeleteKeysFrom(FFrameNumber InTime, bool bDeleteKeysBefore)
+{
+	// Insert a key at the current time to maintain evaluation
+	if (GetData().GetTimes().Num() > 0)
+	{
+		if (const FString* CurrentValue = Evaluate(InTime))
+		{
+			FString Value(*CurrentValue);
+			GetData().UpdateOrAddKey(InTime, *Value);
+		}
+	}
+
+	GetData().DeleteKeysFrom(InTime, bDeleteKeysBefore);
+}
+
 void FMovieSceneStringChannel::ChangeFrameResolution(FFrameRate SourceRate, FFrameRate DestinationRate)
 {
 	GetData().ChangeFrameResolution(SourceRate, DestinationRate);
diff --git a/Engine/Source/Runtime/MovieSceneTracks/Public/Channels/MovieSceneEventChannel.h b/Engine/Source/Runtime/MovieSceneTracks/Public/Channels/MovieSceneEventChannel.h
--- a/Engine/Source/Runtime/MovieSceneTracks/Public/Channels/MovieSceneEventChannel.h
+++ b/Engine/Source/Runtime/MovieSceneTracks/Public/Channels/MovieSceneEventChannel.h
@@ -44,6 +44,7 @@ public:
 	virtual void SetKeyTimes(TArrayView<const FKeyHandle> InHandles, TArrayView<const FFrameNumber> InKeyTimes) override;
 	virtual void DuplicateKeys(TArrayView<const FKeyHandle> InHandles, TArrayView<FKeyHandle> OutNewHandles) override;
 	virtual void DeleteKeys(TArrayView<const FKeyHandle> InHandles) override;
+	virtual void DeleteKeysFrom(FFrameNumber InTime, bool bDeleteKeysBefore) override;
 	virtual void ChangeFrameResolution(FFrameRate SourceRate, FFrameRate DestinationRate) override;
 	virtual TRange<FFrameNumber> ComputeEffectiveRange() const override;
 	virtual int32 GetNumKeys() const override;
diff --git a/Engine/Source/Runtime/MovieSceneTracks/Public/Sections/MovieSceneActorReferenceSection.h b/Engine/Source/Runtime/MovieSceneTracks/Public/Sections/MovieSceneActorReferenceSection.h
--- a/Engine/Source/Runtime/MovieSceneTracks/Public/Sections/MovieSceneActorReferenceSection.h
+++ b/Engine/Source/Runtime/MovieSceneTracks/Public/Sections/MovieSceneActorReferenceSection.h
@@ -86,6 +86,7 @@ public:
 	virtual void SetKeyTimes(TArrayView<const FKeyHandle> InHandles, TArrayView<const FFrameNumber> InKeyTimes) override;
 	virtual void DuplicateKeys(TArrayView<const FKeyHandle> InHandles, TArrayView<FKeyHandle> OutNewHandles) override;
 	virtual void DeleteKeys(TArrayView<const FKeyHandle> InHandles) override;
+	virtual void DeleteKeysFrom(FFrameNumber InTime, bool bDeleteKeysBefore) override;
 	virtual void ChangeFrameResolution(FFrameRate SourceRate, FFrameRate DestinationRate) override;
 	virtual TRange<FFrameNumber> ComputeEffectiveRange() const override;
 	virtual int32 GetNumKeys() const override;
diff --git a/Engine/Source/Runtime/MovieSceneTracks/Public/Sections/MovieSceneAudioSection.h b/Engine/Source/Runtime/MovieSceneTracks/Public/Sections/MovieSceneAudioSection.h
--- a/Engine/Source/Runtime/MovieSceneTracks/Public/Sections/MovieSceneAudioSection.h
+++ b/Engine/Source/Runtime/MovieSceneTracks/Public/Sections/MovieSceneAudioSection.h
@@ -144,8 +144,8 @@ public:
 
 	//~ UMovieSceneSection interface
 	virtual TOptional<TRange<FFrameNumber> > GetAutoSizeRange() const override;
-	virtual void TrimSection(FQualifiedFrameTime TrimTime, bool bTrimLeft) override;
-	virtual UMovieSceneSection* SplitSection(FQualifiedFrameTime SplitTime) override;
+	virtual void TrimSection(FQualifiedFrameTime TrimTime, bool bTrimLeft, bool bDeleteKeys) override;
+	virtual UMovieSceneSection* SplitSection(FQualifiedFrameTime SplitTime, bool bDeleteKeys) override;
 	virtual TOptional<FFrameTime> GetOffsetTime() const override;
 	virtual FMovieSceneEvalTemplatePtr GenerateTemplate() const override;
 
diff --git a/Engine/Source/Runtime/MovieSceneTracks/Public/Sections/MovieSceneEventSection.h b/Engine/Source/Runtime/MovieSceneTracks/Public/Sections/MovieSceneEventSection.h
--- a/Engine/Source/Runtime/MovieSceneTracks/Public/Sections/MovieSceneEventSection.h
+++ b/Engine/Source/Runtime/MovieSceneTracks/Public/Sections/MovieSceneEventSection.h
@@ -176,6 +176,7 @@ public:
 	virtual void SetKeyTimes(TArrayView<const FKeyHandle> InHandles, TArrayView<const FFrameNumber> InKeyTimes) override;
 	virtual void DuplicateKeys(TArrayView<const FKeyHandle> InHandles, TArrayView<FKeyHandle> OutNewHandles) override;
 	virtual void DeleteKeys(TArrayView<const FKeyHandle> InHandles) override;
+	virtual void DeleteKeysFrom(FFrameNumber InTime, bool bDeleteKeysBefore) override;
 	virtual void ChangeFrameResolution(FFrameRate SourceRate, FFrameRate DestinationRate) override;
 	virtual TRange<FFrameNumber> ComputeEffectiveRange() const override;
 	virtual int32 GetNumKeys() const override;
diff --git a/Engine/Source/Runtime/MovieSceneTracks/Public/Sections/MovieSceneSkeletalAnimationSection.h b/Engine/Source/Runtime/MovieSceneTracks/Public/Sections/MovieSceneSkeletalAnimationSection.h
--- a/Engine/Source/Runtime/MovieSceneTracks/Public/Sections/MovieSceneSkeletalAnimationSection.h
+++ b/Engine/Source/Runtime/MovieSceneTracks/Public/Sections/MovieSceneSkeletalAnimationSection.h
@@ -90,8 +90,8 @@ protected:
 
 	//~ UMovieSceneSection interface
 	virtual TOptional<TRange<FFrameNumber> > GetAutoSizeRange() const override;
-	virtual void TrimSection(FQualifiedFrameTime TrimTime, bool bTrimLeft) override;
-	virtual UMovieSceneSection* SplitSection(FQualifiedFrameTime SplitTime) override;
+	virtual void TrimSection(FQualifiedFrameTime TrimTime, bool bTrimLeft, bool bDeleteKeys) override;
+	virtual UMovieSceneSection* SplitSection(FQualifiedFrameTime SplitTime, bool bDeleteKeys) override;
 	virtual void GetSnapTimes(TArray<FFrameNumber>& OutSnapTimes, bool bGetSectionBorders) const override;
 	virtual TOptional<FFrameTime> GetOffsetTime() const override;
 	virtual FMovieSceneEvalTemplatePtr GenerateTemplate() const override;
diff --git a/Engine/Source/Runtime/MovieSceneTracks/Public/Sections/MovieSceneStringChannel.h b/Engine/Source/Runtime/MovieSceneTracks/Public/Sections/MovieSceneStringChannel.h
--- a/Engine/Source/Runtime/MovieSceneTracks/Public/Sections/MovieSceneStringChannel.h
+++ b/Engine/Source/Runtime/MovieSceneTracks/Public/Sections/MovieSceneStringChannel.h
@@ -68,6 +68,7 @@ public:
 	virtual void SetKeyTimes(TArrayView<const FKeyHandle> InHandles, TArrayView<const FFrameNumber> InKeyTimes) override;
 	virtual void DuplicateKeys(TArrayView<const FKeyHandle> InHandles, TArrayView<FKeyHandle> OutNewHandles) override;
 	virtual void DeleteKeys(TArrayView<const FKeyHandle> InHandles) override;
+	virtual void DeleteKeysFrom(FFrameNumber InTime, bool bDeleteKeysBefore) override;
 	virtual void ChangeFrameResolution(FFrameRate SourceRate, FFrameRate DestinationRate) override;
 	virtual TRange<FFrameNumber> ComputeEffectiveRange() const override;
 	virtual int32 GetNumKeys() const override;
