Sequencer: Added some null checks to prevent some crashes when deleting track data and not updating the UI. Ultimately, this should be fixed with UE-72911.

diff --git a/Engine/Source/Editor/Sequencer/Private/DisplayNodes/SequencerTrackNode.cpp b/Engine/Source/Editor/Sequencer/Private/DisplayNodes/SequencerTrackNode.cpp
--- a/Engine/Source/Editor/Sequencer/Private/DisplayNodes/SequencerTrackNode.cpp
+++ b/Engine/Source/Editor/Sequencer/Private/DisplayNodes/SequencerTrackNode.cpp
@@ -407,7 +407,7 @@ TSharedRef<SWidget> FSequencerTrackNode::GetCustomOutlinerContent()
 	{
 		Params.TrackInsertRowIndex = GetRowIndex();
 	}
-	else if (Track->SupportsMultipleRows())
+	else if (Track && Track->SupportsMultipleRows())
 	{
 		Params.TrackInsertRowIndex = Track->GetMaxRowIndex()+1;
 	}
@@ -688,6 +688,11 @@ FLinearColor FSequencerTrackNode::GetDisplayNameColor() const
 {
 	UMovieSceneTrack* Track = GetTrack();
 
+	if (!Track)
+	{
+		return FLinearColor::White;
+	}
+
 	const bool bIsEvalDisabled = Track->IsEvalDisabled();
 
 	// Display track node is red if the property track is not bound to valid property
diff --git a/Engine/Source/Editor/Sequencer/Private/MovieSceneTrackEditor.cpp b/Engine/Source/Editor/Sequencer/Private/MovieSceneTrackEditor.cpp
--- a/Engine/Source/Editor/Sequencer/Private/MovieSceneTrackEditor.cpp
+++ b/Engine/Source/Editor/Sequencer/Private/MovieSceneTrackEditor.cpp
@@ -168,7 +168,7 @@ void FMovieSceneTrackEditor::BuildObjectBindingTrackMenu(FMenuBuilder& MenuBuild
 
 TSharedPtr<SWidget> FMovieSceneTrackEditor::BuildOutlinerEditWidget(const FGuid& ObjectBinding, UMovieSceneTrack* Track, const FBuildEditWidgetParams& Params)
 {
-	if (Track->GetSupportedBlendTypes().Num() > 0)
+	if (Track && Track->GetSupportedBlendTypes().Num() > 0)
 	{
 		TWeakPtr<ISequencer> WeakSequencer = GetSequencer();
 
diff --git a/Engine/Source/Editor/Sequencer/Private/SSequencerSectionAreaView.cpp b/Engine/Source/Editor/Sequencer/Private/SSequencerSectionAreaView.cpp
--- a/Engine/Source/Editor/Sequencer/Private/SSequencerSectionAreaView.cpp
+++ b/Engine/Source/Editor/Sequencer/Private/SSequencerSectionAreaView.cpp
@@ -200,6 +200,11 @@ void SSequencerSectionAreaView::OnArrangeChildren( const FGeometry& AllottedGeom
 FTimeToPixel SSequencerSectionAreaView::GetTimeToPixel( const FGeometry& AllottedGeometry ) const
 {
 	const UMovieSceneTrack* Track      = SectionAreaNode->GetTrack();
+	if (!Track)
+	{
+		return FTimeToPixel(AllottedGeometry, ViewRange.Get(), FFrameRate());
+	}
+
 	const FFrameRate        Resolution = Track->GetTypedOuter<UMovieScene>()->GetTickResolution();
 
 	return FTimeToPixel( AllottedGeometry, ViewRange.Get(), Resolution );
