Sequencer: Only the master trackarea should update external controls

diff --git a/Engine/Source/Editor/Sequencer/Private/SSequencer.cpp b/Engine/Source/Editor/Sequencer/Private/SSequencer.cpp
--- a/Engine/Source/Editor/Sequencer/Private/SSequencer.cpp
+++ b/Engine/Source/Editor/Sequencer/Private/SSequencer.cpp
@@ -408,6 +408,7 @@ void SSequencer::Construct(const FArguments& InArgs, TSharedRef<FSequencer> InSe
 
 	PinnedTrackArea->SetTreeView(PinnedTreeView);
 	PinnedTrackArea->SetShowPinned(true);
+	PinnedTrackArea->SetIsSlave(true);
 	PinnedTreeView->SetShowPinned(true);
 
 	SAssignNew(TrackArea, SSequencerTrackArea, TimeSliderControllerRef, InSequencer);
diff --git a/Engine/Source/Editor/Sequencer/Private/SSequencerTrackArea.cpp b/Engine/Source/Editor/Sequencer/Private/SSequencerTrackArea.cpp
--- a/Engine/Source/Editor/Sequencer/Private/SSequencerTrackArea.cpp
+++ b/Engine/Source/Editor/Sequencer/Private/SSequencerTrackArea.cpp
@@ -444,7 +444,7 @@ void SSequencerTrackArea::Tick( const FGeometry& AllottedGeometry, const double
 {
 	FVector2D Size = AllottedGeometry.GetLocalSize();
 
-	if (SizeLastFrame.IsSet() && Size.X != SizeLastFrame->X)
+	if (!IsSlave() && SizeLastFrame.IsSet() && Size.X != SizeLastFrame->X)
 	{
 		// Zoom by the difference in horizontal size
 		const float Difference = Size.X - SizeLastFrame->X;
diff --git a/Engine/Source/Editor/Sequencer/Private/SSequencerTrackArea.h b/Engine/Source/Editor/Sequencer/Private/SSequencerTrackArea.h
--- a/Engine/Source/Editor/Sequencer/Private/SSequencerTrackArea.h
+++ b/Engine/Source/Editor/Sequencer/Private/SSequencerTrackArea.h
@@ -91,6 +91,10 @@ public:
 	void SetShowPinned(bool bShowPinned) { bShowPinnedNodes = bShowPinned; }
 	bool ShowPinned() const { return bShowPinnedNodes; }
 
+	/** Set whether this TrackArea is slaved to a master TrackArea and should skip updating external controls */
+	void SetIsSlave(bool bInIsSlave) { bIsSlave = bInIsSlave; }
+	bool IsSlave() const { return bIsSlave; }
+	
 public:
 
 	// SWidget interface
@@ -158,6 +162,9 @@ private:
 	/** Whether this TrackArea is for pinned nodes or non-pinned nodes */
 	bool bShowPinnedNodes;
 
+	/** Whether this TrackArea is slaved to a master TrackArea and should skip updating external controls */
+	bool bIsSlave;
+
 private:
 
 	/** Input handler stack responsible for routing input to the different handlers */
