Sequencer: Ensure that toolkit host exists on SSequencer destructor

diff --git a/Engine/Source/Editor/Sequencer/Private/SSequencer.cpp b/Engine/Source/Editor/Sequencer/Private/SSequencer.cpp
--- a/Engine/Source/Editor/Sequencer/Private/SSequencer.cpp
+++ b/Engine/Source/Editor/Sequencer/Private/SSequencer.cpp
@@ -2240,7 +2240,7 @@ SSequencer::~SSequencer()
 
 
 	TSharedPtr<FSequencer> Sequencer = SequencerPtr.Pin();
-	if(Sequencer)
+	if(Sequencer && Sequencer->GetToolkitHost() && Sequencer->GetToolkitHost()->GetTabManager())
 	{
 		FTabId TabId = FTabId(SSequencer::CurveEditorTabName);
 		TSharedPtr<SDockTab> CurveEditorTab = Sequencer->GetToolkitHost()->GetTabManager()->FindExistingLiveTab(TabId);
