Sequencer: Fix max tick rate for FSequencer so that it only overrides during playback and if frame locked is on. This matches what runtime sequencer does.

diff --git a/Engine/Source/Editor/Sequencer/Private/Sequencer.cpp b/Engine/Source/Editor/Sequencer/Private/Sequencer.cpp
--- a/Engine/Source/Editor/Sequencer/Private/Sequencer.cpp
+++ b/Engine/Source/Editor/Sequencer/Private/Sequencer.cpp
@@ -489,7 +489,6 @@ FSequencer::FSequencer()
 	, NodeTree( MakeShareable( new FSequencerNodeTree( *this ) ) )
 	, bUpdatingSequencerSelection( false )
 	, bUpdatingExternalSelection( false )
-	, OldMaxTickRate(GEngine->GetMaxFPS())
 	, bNeedsEvaluate(false)
 {
 	Selection.GetOnOutlinerNodeSelectionChanged().AddRaw(this, &FSequencer::OnSelectedOutlinerNodesChanged);
@@ -519,6 +518,12 @@ FSequencer::~FSequencer()
 
 void FSequencer::Close()
 {
+	if (OldMaxTickRate.IsSet())
+	{
+		GEngine->SetMaxFPS(OldMaxTickRate.GetValue());
+		OldMaxTickRate.Reset();
+	}
+
 	RootTemplateInstance.Finish(*this);
 	RestorePreAnimatedState();
 
@@ -623,19 +628,6 @@ void FSequencer::Tick(float InDeltaTime)
 		SetLocalTimeDirectly(CurrentTime.Time + Offset);
 	}
 
-	// override max frame rate
-	if (PlaybackState == EMovieScenePlayerStatus::Playing)
-	{
-		if (PlayPosition.GetEvaluationType() == EMovieSceneEvaluationType::FrameLocked)
-		{
-			GEngine->SetMaxFPS(1.f / PlayPosition.GetInputRate().AsInterval());
-		}
-		else
-		{
-			GEngine->SetMaxFPS(OldMaxTickRate);
-		}
-	}
-
 	if (GetSelectionRange().IsEmpty() && GetLoopMode() == SLM_LoopSelectionRange)
 	{
 		Settings->SetLoopMode(SLM_Loop);
@@ -3243,16 +3235,28 @@ void FSequencer::SetPlaybackStatus(EMovieScenePlayerStatus::Type InPlaybackStatu
 		}
 	}
 
-	// backup or restore tick rate
 	if (InPlaybackStatus == EMovieScenePlayerStatus::Playing)
 	{
-		OldMaxTickRate = GEngine->GetMaxFPS();
+		// override max frame rate
+		if (PlayPosition.GetEvaluationType() == EMovieSceneEvaluationType::FrameLocked)
+		{
+			if (!OldMaxTickRate.IsSet())
+			{
+				OldMaxTickRate = GEngine->GetMaxFPS();
+			}
+
+			GEngine->SetMaxFPS(1.f / PlayPosition.GetInputRate().AsInterval());
+		}
 	}
 	else
 	{
 		StopAutoscroll();
 
-		GEngine->SetMaxFPS(OldMaxTickRate);
+		if (OldMaxTickRate.IsSet())
+		{
+			GEngine->SetMaxFPS(OldMaxTickRate.GetValue());
+			OldMaxTickRate.Reset();
+		}
 
 		ShuttleMultiplier = 0;
 	}
diff --git a/Engine/Source/Editor/Sequencer/Private/Sequencer.h b/Engine/Source/Editor/Sequencer/Private/Sequencer.h
--- a/Engine/Source/Editor/Sequencer/Private/Sequencer.h
+++ b/Engine/Source/Editor/Sequencer/Private/Sequencer.h
@@ -1227,7 +1227,7 @@ private:
 	bool bUpdatingExternalSelection;
 
 	/** The maximum tick rate prior to playing (used for overriding delta time during playback). */
-	double OldMaxTickRate;
+	TOptional<double> OldMaxTickRate;
 
 	/** Timing manager that can adjust playback times */
 	TSharedPtr<FMovieSceneTimeController> TimeController;
