Sequencer: Fix spawnable export so that they're temporarily spawned on export

diff --git a/Engine/Source/Editor/Sequencer/Private/Sequencer.cpp b/Engine/Source/Editor/Sequencer/Private/Sequencer.cpp
--- a/Engine/Source/Editor/Sequencer/Private/Sequencer.cpp
+++ b/Engine/Source/Editor/Sequencer/Private/Sequencer.cpp
@@ -72,6 +72,7 @@
 #include "Sections/MovieSceneBoolSection.h"
 #include "Sections/MovieSceneCameraCutSection.h"
 #include "Sections/MovieScene3DTransformSection.h"
+#include "Sections/MovieSceneSpawnSection.h"
 #include "Sections/MovieSceneSubSection.h"
 #include "Tracks/MovieSceneSubTrack.h"
 #include "Sections/MovieSceneCinematicShotSection.h"
@@ -9694,7 +9695,35 @@ void FSequencer::ExportFBXInternal(const FString& ExportFilename, TArray<FGuid>&
 			UWorld* World = Cast<UWorld>(GetPlaybackContext());
 			FMovieSceneSequenceIDRef Template = GetFocusedTemplateID();
 			UnFbx::FFbxExporter::FLevelSequenceNodeNameAdapter NodeNameAdapter(MovieScene, this, Template);
-			MovieSceneToolHelpers::ExportFBX(World,MovieScene, this, Bindings, NodeNameAdapter, Template, ExportFilename, RootToLocalTransform);
+
+			{
+				FScopedTransaction ExportFBXTransaction(NSLOCTEXT("Sequencer", "ExportFBX", "Export FBX"));
+
+				// Temporarily spawn all spawnables so that they can be exported. This is undone after export below with UndoTransaction()
+				for (int32 SpawnableIndex = 0; SpawnableIndex < MovieScene->GetSpawnableCount(); ++SpawnableIndex)
+				{
+					FMovieSceneSpawnable& Spawnable = MovieScene->GetSpawnable(SpawnableIndex);
+
+					UMovieSceneSpawnTrack* SpawnTrack = MovieScene->FindTrack<UMovieSceneSpawnTrack>(Spawnable.GetGuid());
+
+					if (SpawnTrack)
+					{
+						SpawnTrack->Modify();
+						UMovieSceneSpawnSection* SpawnSection = Cast<UMovieSceneSpawnSection>(SpawnTrack->CreateNewSection());
+						SpawnTrack->RemoveAllAnimationData();
+						SpawnSection->GetChannel().SetDefault(true);
+						SpawnTrack->AddSection(*SpawnSection);
+					}
+				}
+
+				ForceEvaluate();
+
+				MovieSceneToolHelpers::ExportFBX(World, MovieScene, this, Bindings, NodeNameAdapter, Template, ExportFilename, RootToLocalTransform);
+			}
+
+			GEditor->UndoTransaction(false);
+
+			ForceEvaluate();
 		}
 	}
 }
