Sequencer: Include any supported blend type if any section supports it

diff --git a/Engine/Source/Editor/Sequencer/Private/SequencerUtilities.cpp b/Engine/Source/Editor/Sequencer/Private/SequencerUtilities.cpp
--- a/Engine/Source/Editor/Sequencer/Private/SequencerUtilities.cpp
+++ b/Engine/Source/Editor/Sequencer/Private/SequencerUtilities.cpp
@@ -187,16 +187,23 @@ void FSequencerUtilities::PopulateMenu_SetBlendType(FMenuBuilder& MenuBuilder, c
 	{
 		EMovieSceneBlendType BlendType = (EMovieSceneBlendType)MovieSceneBlendType->GetValueByIndex(NameIndex);
 
-		// Don't include this if it's not supported
+		// Include this if any section supports it
+		bool bAnySupported = false;
 		for (TWeakObjectPtr<UMovieSceneSection> WeakSection : InSections)
 		{
 			UMovieSceneSection* Section = WeakSection.Get();
-			if (Section && !Section->GetSupportedBlendTypes().Contains(BlendType))
+			if (Section && Section->GetSupportedBlendTypes().Contains(BlendType))
 			{
-				return;
+				bAnySupported = true;
+				break;
 			}
 		}
 
+		if (!bAnySupported)
+		{
+			continue;
+		}
+
 		FName EnumValueName = MovieSceneBlendType->GetNameByIndex(NameIndex);
 		MenuBuilder.AddMenuEntry(
 			MovieSceneBlendType->GetDisplayNameTextByIndex(NameIndex),
