Sequencer: Fix pasting of keys when different things are selected

diff --git a/Engine/Source/Editor/Sequencer/Private/SSequencer.cpp b/Engine/Source/Editor/Sequencer/Private/SSequencer.cpp
--- a/Engine/Source/Editor/Sequencer/Private/SSequencer.cpp
+++ b/Engine/Source/Editor/Sequencer/Private/SSequencer.cpp
@@ -3076,6 +3076,7 @@ FPasteContextMenuArgs SSequencer::GeneratePasteArgs(FFrameNumber PasteAtTime, TS
 	TArray<TSharedRef<FSequencerDisplayNode>> PasteIntoNodes;
 	{
 		TSet<TWeakObjectPtr<UMovieSceneSection>> Sections = Sequencer->GetSelection().GetSelectedSections();
+		
 		for (const FSequencerSelectedKey& Key : Sequencer->GetSelection().GetSelectedKeys())
 		{
 			Sections.Add(Key.Section);
@@ -3088,6 +3089,15 @@ FPasteContextMenuArgs SSequencer::GeneratePasteArgs(FFrameNumber PasteAtTime, TS
 				PasteIntoNodes.Add(Handle->GetTrackNode());
 			}
 		}
+
+		auto SelectedNodes = Sequencer->GetSelection().GetSelectedOutlinerNodes();
+		for (const TSharedRef<FSequencerDisplayNode>& SelectedNode : SelectedNodes)
+		{
+			if (SelectedNode->GetType() == ESequencerNode::Category || SelectedNode->GetType() == ESequencerNode::Track || SelectedNode->GetType() == ESequencerNode::KeyArea)
+			{
+				PasteIntoNodes.Add(SelectedNode);
+			}
+		}
 	}
 
 	if (PasteIntoNodes.Num() == 0)
@@ -3108,13 +3118,20 @@ void SSequencer::OnPaste()
 	TSet<TSharedRef<FSequencerDisplayNode>> SelectedNodes = Sequencer->GetSelection().GetSelectedOutlinerNodes();
 	if (SelectedNodes.Num() == 0)
 	{
-		if (OpenPasteMenu())
+		if (!OpenPasteMenu())
 		{
-			return;
+			DoPaste();
+		}
+	}
+	else
+	{
+		if (!DoPaste())
+		{
+			OpenPasteMenu();
 		}
 	}
 
-	DoPaste();
+	
 }
 
 bool SSequencer::CanPaste()
@@ -3148,11 +3165,11 @@ bool SSequencer::CanPaste()
 	return SequencerPtr.Pin()->GetClipboardStack().Num() != 0;
 }
 
-void SSequencer::DoPaste()
+bool SSequencer::DoPaste()
 {
 	TSharedPtr<FSequencer> Sequencer = SequencerPtr.Pin();
 
-	Sequencer->DoPaste();
+	return Sequencer->DoPaste();
 }
 
 bool SSequencer::OpenPasteMenu()
diff --git a/Engine/Source/Editor/Sequencer/Private/SSequencer.h b/Engine/Source/Editor/Sequencer/Private/SSequencer.h
--- a/Engine/Source/Editor/Sequencer/Private/SSequencer.h
+++ b/Engine/Source/Editor/Sequencer/Private/SSequencer.h
@@ -516,8 +516,11 @@ public:
 	void OnPaste();
 	bool CanPaste();
 
-	/** Handle Track Paste */
-	void DoPaste();
+	/**
+	 * Handle Track Paste
+	 * @return Whether the paste event was handled
+	 */
+	bool DoPaste();
 
 	/** Open the paste menu */
 	bool OpenPasteMenu();
diff --git a/Engine/Source/Editor/Sequencer/Private/Sequencer.cpp b/Engine/Source/Editor/Sequencer/Private/Sequencer.cpp
--- a/Engine/Source/Editor/Sequencer/Private/Sequencer.cpp
+++ b/Engine/Source/Editor/Sequencer/Private/Sequencer.cpp
@@ -6844,12 +6844,13 @@ void FSequencer::ExportObjectsToText(TArray<UObject*> ObjectsToExport, FString&
 	ExportedText = Archive;
 }
 
-void FSequencer::DoPaste()
+bool FSequencer::DoPaste()
 {
 	if (IsReadOnly())
 	{
 		ShowReadOnlyError();
-		return;
+		// If we cancel the paste due to being read-only, count that as having handled the paste operation
+		return true;
 	}
 
 	// Grab the text to paste from the clipboard
@@ -6871,6 +6872,8 @@ void FSequencer::DoPaste()
 		NotificationInfo.bUseLargeFont = false;
 		FSlateNotificationManager::Get().AddNotification(NotificationInfo);
 	}
+
+	return bAnythingPasted;
 }
 
 bool FSequencer::PasteObjectBindings(const FString& TextToImport, TArray<FNotificationInfo>& PasteErrors)
diff --git a/Engine/Source/Editor/Sequencer/Private/Sequencer.h b/Engine/Source/Editor/Sequencer/Private/Sequencer.h
--- a/Engine/Source/Editor/Sequencer/Private/Sequencer.h
+++ b/Engine/Source/Editor/Sequencer/Private/Sequencer.h
@@ -507,7 +507,11 @@ public:
 
 	/** Called when a user executes the paste track menu item */
 	bool CanPaste(const FString& TextToImport);
-	void DoPaste();
+	/**
+	 * Attempts to paste from the clipboard
+	 * @return Whether the paste event was handled
+	 */
+	bool DoPaste();
 	bool PasteTracks(const FString& TextToImport, TArray<FNotificationInfo>& PasteErrors);
 	bool PasteSections(const FString& TextToImport, TArray<FNotificationInfo>& PasteErrors);
 	bool PasteObjectBindings(const FString& TextToImport, TArray<FNotificationInfo>& PasteErrors);
