Sequencer: Improve pinned track UI layout

diff --git a/Engine/Source/Editor/Sequencer/Private/SSequencer.cpp b/Engine/Source/Editor/Sequencer/Private/SSequencer.cpp
--- a/Engine/Source/Editor/Sequencer/Private/SSequencer.cpp
+++ b/Engine/Source/Editor/Sequencer/Private/SSequencer.cpp
@@ -399,10 +399,15 @@ void SSequencer::Construct(const FArguments& InArgs, TSharedRef<FSequencer> InSe
 
 	TSharedRef<SScrollBar> ScrollBar = SNew(SScrollBar)
 		.Thickness(FVector2D(9.0f, 9.0f));
+
+	TSharedRef<SScrollBar> PinnedAreaScrollBar = SNew(SScrollBar)
+		.Thickness(FVector2D(9.0f, 9.0f));
+	
 	SAssignNew(TrackOutliner, SSequencerTrackOutliner);
 
 	SAssignNew(PinnedTrackArea, SSequencerTrackArea, TimeSliderControllerRef, InSequencer);
 	SAssignNew(PinnedTreeView, SSequencerTreeView, InSequencer->GetNodeTree(), PinnedTrackArea.ToSharedRef())
+		.ExternalScrollbar(PinnedAreaScrollBar)
 		.Clipping(EWidgetClipping::ClipToBounds)
 		.OnGetContextMenuContent(FOnGetContextMenuContent::CreateSP(this, &SSequencer::GetContextMenuContent));
 
@@ -704,50 +709,71 @@ void SSequencer::Construct(const FArguments& InArgs, TSharedRef<FSequencer> InSe
 					+ SGridPanel::Slot( Column0, Row2, SGridPanel::Layer(10) )
 					.ColumnSpan(2)
 					[
-						SNew(SHorizontalBox)
+						SAssignNew(MainSequencerArea, SVerticalBox)
 
-						+ SHorizontalBox::Slot()
+						+SVerticalBox::Slot()
+						.AutoHeight()
+						.MaxHeight(TAttribute<float>::Create(TAttribute<float>::FGetter::CreateSP(this, &SSequencer::GetPinnedAreaMaxHeight)))
 						[
-							SNew( SOverlay )
-
-							+ SOverlay::Slot()
+							SNew(SBorder)
+							.Visibility(this, &SSequencer::GetPinnedAreaVisibility)
+							.Padding(FMargin(0.0f, 0.0f, 0.0f, CommonPadding))
 							[
-								SNew(SVerticalBox)
 
-								+SVerticalBox::Slot()
-								.AutoHeight()
+								SNew(SHorizontalBox)
+
+								+ SHorizontalBox::Slot()
 								[
-									SNew(SBorder)
-									.Padding(FMargin(0.0f, 0.0f, 0.0f, CommonPadding))
-									[
-										SNew(SHorizontalBox)
+									SNew(SOverlay)
 
-										// outliner tree
-										+ SHorizontalBox::Slot()
-										.FillWidth( FillCoefficient_0 )
+									+ SOverlay::Slot()
+									[
+										SNew(SScrollBorder, TreeView.ToSharedRef())
 										[
-											SNew(SBox)
+											SNew(SHorizontalBox)
+
+											// outliner tree
+											+ SHorizontalBox::Slot()
+											.FillWidth( FillCoefficient_0 )
 											[
-												PinnedTreeView.ToSharedRef()
+												SNew(SBox)
+												[
+													PinnedTreeView.ToSharedRef()
+												]
 											]
-										]
 
-										// track area
-										+ SHorizontalBox::Slot()
-										.FillWidth( FillCoefficient_1 )
-										[
-											SNew(SBox)
-											.Padding(ResizeBarPadding)
-											.Clipping(EWidgetClipping::ClipToBounds)
+											// track area
+											+ SHorizontalBox::Slot()
+											.FillWidth( FillCoefficient_1 )
 											[
-												PinnedTrackArea.ToSharedRef()
+												SNew(SBox)
+												.Padding(ResizeBarPadding)
+												.Clipping(EWidgetClipping::ClipToBounds)
+												[
+													PinnedTrackArea.ToSharedRef()
+												]
 											]
 										]
 									]
 
+									+ SOverlay::Slot()
+									.HAlign(HAlign_Right)
+									[
+										PinnedAreaScrollBar
+									]
 								]
+							]
+						]
+
+						+SVerticalBox::Slot()
+						[
+							SNew(SHorizontalBox)
+
+							+ SHorizontalBox::Slot()
+							[
+								SNew(SOverlay)
 
-								+SVerticalBox::Slot()
+								+ SOverlay::Slot()
 								[
 									SNew(SScrollBorder, TreeView.ToSharedRef())
 									[
@@ -776,12 +802,12 @@ void SSequencer::Construct(const FArguments& InArgs, TSharedRef<FSequencer> InSe
 										]
 									]
 								]
-							]
 
-							+ SOverlay::Slot()
-							.HAlign( HAlign_Right )
-							[
-								ScrollBar
+								+ SOverlay::Slot()
+								.HAlign(HAlign_Right)
+								[
+									ScrollBar
+								]
 							]
 						]
 					]
@@ -2901,6 +2927,22 @@ void SSequencer::StepToKey(bool bStepToNextKey, bool bCameraOnly)
 }
 
 
+float SSequencer::GetPinnedAreaMaxHeight() const
+{
+	if (!MainSequencerArea.IsValid())
+	{
+		return 0.0f;
+	}
+
+	// Allow the pinned area to use up to 2/3rds of the sequencer area
+	return MainSequencerArea->GetCachedGeometry().GetLocalSize().Y * 0.666f;
+}
+
+EVisibility SSequencer::GetPinnedAreaVisibility() const
+{
+	return PinnedTreeView->GetNumRootNodes() > 0 ? EVisibility::Visible : EVisibility::Collapsed;
+}
+
 FText SSequencer::GetBreadcrumbTextForSection(TWeakObjectPtr<UMovieSceneSubSection> SubSection) const
 {
 	UMovieSceneSubSection* SubSectionPtr = SubSection.Get();
diff --git a/Engine/Source/Editor/Sequencer/Private/SSequencer.h b/Engine/Source/Editor/Sequencer/Private/SSequencer.h
--- a/Engine/Source/Editor/Sequencer/Private/SSequencer.h
+++ b/Engine/Source/Editor/Sequencer/Private/SSequencer.h
@@ -472,6 +472,12 @@ private:
 	/** Gets the root movie scene name */
 	FText GetRootAnimationName() const;
 
+	/** Get the maximum height the pinned track area should be allowed to be.. */
+	float GetPinnedAreaMaxHeight() const;
+	
+	/** Gets whether or not the Pinned track area should be visible. */
+	EVisibility GetPinnedAreaVisibility() const;
+
 	FText GetBreadcrumbTextForSection(TWeakObjectPtr<UMovieSceneSubSection> SubSection) const;
 	FText GetBreadcrumbTextForSequence(TWeakObjectPtr<UMovieSceneSequence> Sequence, bool bIsActive) const;
 
@@ -544,6 +550,9 @@ private:
 	/** Stretch box widget. */
 	TSharedPtr<SSequencerStretchBox> StretchBox;
 
+	/** Main Sequencer Area*/
+	TSharedPtr<SVerticalBox> MainSequencerArea;
+
 	/** Section area widget */
 	TSharedPtr<SSequencerTrackArea> TrackArea;
 
diff --git a/Engine/Source/Editor/Sequencer/Private/SSequencerTreeView.h b/Engine/Source/Editor/Sequencer/Private/SSequencerTreeView.h
--- a/Engine/Source/Editor/Sequencer/Private/SSequencerTreeView.h
+++ b/Engine/Source/Editor/Sequencer/Private/SSequencerTreeView.h
@@ -68,6 +68,9 @@ public:
 	/** Access the underlying tree data */
 	TSharedPtr<FSequencerNodeTree> GetNodeTree() { return SequencerNodeTree; }
 
+	/** @return the number of root nodes this tree contains */
+	int32 GetNumRootNodes() const { return RootNodes.Num(); }
+
 	/** @return an optional region specifying the vertical bounds in which a highlight should be drawn */
 	const TOptional<FHighlightRegion>& GetHighlightRegion() const { return HighlightRegion; }
 
