Sequencer: Fix transform baking for rig rails - the current position on rail is updated per tick, so the actor needs to be ticked on each baked frame.

diff --git a/Engine/Source/Editor/Sequencer/Private/Sequencer.cpp b/Engine/Source/Editor/Sequencer/Private/Sequencer.cpp
--- a/Engine/Source/Editor/Sequencer/Private/Sequencer.cpp
+++ b/Engine/Source/Editor/Sequencer/Private/Sequencer.cpp
@@ -1660,6 +1660,14 @@ void FSequencer::BakeTransform()
 				FFrameNumber KeyTime = FFrameRate::Snap(EvalTime, Resolution, SnapRate).FloorToFrame();
 				FMovieSceneEvaluationRange Range = PlayPosition.JumpTo(KeyTime * RootToLocalTransform.Inverse());
 				EvaluateInternal(Range);
+				
+				// The CameraRig_rail updates the spline position tick, so it needs to be ticked manually while baking the frames
+				USceneComponent* Parent = CameraComponent->GetAttachParent();
+				while (Parent && Parent->GetOwner())
+				{
+					Parent->GetOwner()->Tick(0.3f);
+					Parent = Parent->GetAttachParent();
+				}
 
 				if (CameraComponent)
 				{
