Sequencer: Fix static analysis

diff --git a/Engine/Source/Editor/Sequencer/Private/Sequencer.cpp b/Engine/Source/Editor/Sequencer/Private/Sequencer.cpp
--- a/Engine/Source/Editor/Sequencer/Private/Sequencer.cpp
+++ b/Engine/Source/Editor/Sequencer/Private/Sequencer.cpp
@@ -1661,8 +1661,17 @@ void FSequencer::BakeTransform()
 				FMovieSceneEvaluationRange Range = PlayPosition.JumpTo(KeyTime * RootToLocalTransform.Inverse());
 				EvaluateInternal(Range);
 				
+				USceneComponent* Parent = nullptr;
+				if (CameraComponent)
+				{
+					Parent = CameraComponent->GetAttachParent();
+				} 
+				else if (Actor->GetRootComponent())
+				{
+					Parent = Actor->GetRootComponent()->GetAttachParent();
+				}
+				
 				// The CameraRig_rail updates the spline position tick, so it needs to be ticked manually while baking the frames
-				USceneComponent* Parent = CameraComponent->GetAttachParent();
 				while (Parent && Parent->GetOwner())
 				{
 					Parent->GetOwner()->Tick(0.3f);
