Sequencer: Follow up to CL8614178 to fix when there's no spawn sections to undo

diff --git a/Engine/Source/Editor/Sequencer/Private/Sequencer.cpp b/Engine/Source/Editor/Sequencer/Private/Sequencer.cpp
--- a/Engine/Source/Editor/Sequencer/Private/Sequencer.cpp
+++ b/Engine/Source/Editor/Sequencer/Private/Sequencer.cpp
@@ -9716,6 +9716,7 @@ void FSequencer::ExportFBXInternal(const FString& ExportFilename, TArray<FGuid>&
 			FMovieSceneSequenceIDRef Template = GetFocusedTemplateID();
 			UnFbx::FFbxExporter::FLevelSequenceNodeNameAdapter NodeNameAdapter(MovieScene, this, Template);
 
+			bool bWasTransacted = false;
 			{
 				FScopedTransaction ExportFBXTransaction(NSLOCTEXT("Sequencer", "ExportFBX", "Export FBX"));
 
@@ -9728,6 +9729,8 @@ void FSequencer::ExportFBXInternal(const FString& ExportFilename, TArray<FGuid>&
 
 					if (SpawnTrack)
 					{
+						bWasTransacted = true;
+
 						SpawnTrack->Modify();
 						UMovieSceneSpawnSection* SpawnSection = Cast<UMovieSceneSpawnSection>(SpawnTrack->CreateNewSection());
 						SpawnTrack->RemoveAllAnimationData();
@@ -9741,8 +9744,10 @@ void FSequencer::ExportFBXInternal(const FString& ExportFilename, TArray<FGuid>&
 				MovieSceneToolHelpers::ExportFBX(World, MovieScene, this, Bindings, NodeNameAdapter, Template, ExportFilename, RootToLocalTransform);
 			}
 
-			GEditor->UndoTransaction(false);
-
+			if (bWasTransacted)
+			{
+				GEditor->UndoTransaction(false);
+			}
 			ForceEvaluate();
 		}
 	}
