Sequencer: Make enable/disable filter menu text more consistent.

diff --git a/Engine/Source/Editor/Sequencer/Private/SSequencer.cpp b/Engine/Source/Editor/Sequencer/Private/SSequencer.cpp
--- a/Engine/Source/Editor/Sequencer/Private/SSequencer.cpp
+++ b/Engine/Source/Editor/Sequencer/Private/SSequencer.cpp
@@ -1618,8 +1618,8 @@ void SSequencer::FillLevelFilterMenu(FMenuBuilder& InMenuBarBuilder)
 				FUIAction(FExecuteAction::CreateSP(this, &SSequencer::OnEnableAllLevelFilters, true)));
 
 			InMenuBarBuilder.AddMenuEntry(
-				LOCTEXT("DisableAllLevelFilters", "Disable All"),
-				LOCTEXT("DisableAllLevelFiltersToolTip", "Disable all level filters"),
+				LOCTEXT("ResetLevelFilters", "Reset Filters"),
+				LOCTEXT("ResetLevelFiltersToolTip", "Resets current level filter selection"),
 				FSlateIcon(),
 				FUIAction(FExecuteAction::CreateSP(this, &SSequencer::OnEnableAllLevelFilters, false)));
 
