Sequencer: Fix incorrect transform trajectory rendering when blending multiple sections

diff --git a/Engine/Source/Editor/Sequencer/Private/SequencerEdMode.cpp b/Engine/Source/Editor/Sequencer/Private/SequencerEdMode.cpp
--- a/Engine/Source/Editor/Sequencer/Private/SequencerEdMode.cpp
+++ b/Engine/Source/Editor/Sequencer/Private/SequencerEdMode.cpp
@@ -458,6 +458,11 @@ void FSequencerEdMode::DrawTransformTrack(const TSharedPtr<FSequencer>& Sequence
 			KeyPosRots.Reserve(TrajectoryKeys.Num());
 			for (const FTrajectoryKey& NewTrajectoryKey : TrajectoryKeys)
 			{
+				if (NewTrajectoryKey.KeyData.Num() == 0)
+				{
+					continue;
+				}
+
 				FFrameTime NewKeyTime = NewTrajectoryKey.Time;
 
 				FVector NewKeyPos(0);
