Sequencer: Remove SetPlaybackRangeStart in favor of SetPlaybackStart

diff --git a/Engine/Source/Editor/Sequencer/Private/Sequencer.cpp b/Engine/Source/Editor/Sequencer/Private/Sequencer.cpp
--- a/Engine/Source/Editor/Sequencer/Private/Sequencer.cpp
+++ b/Engine/Source/Editor/Sequencer/Private/Sequencer.cpp
@@ -10575,7 +10575,7 @@ void FSequencer::BindCommands()
 
 	SequencerCommandBindings->MapAction(
 		Commands.SetStartPlaybackRange,
-		FExecuteAction::CreateSP( this, &FSequencer::SetPlaybackRangeStart ),
+		FExecuteAction::CreateLambda([this] { SetPlaybackStart(); }),
 		FCanExecuteAction::CreateSP( this, &FSequencer::IsViewingMasterSequence ) );
 
 	SequencerCommandBindings->MapAction(
@@ -10596,7 +10596,7 @@ void FSequencer::BindCommands()
 
 	SequencerCommandBindings->MapAction(
 		Commands.SetEndPlaybackRange,
-		FExecuteAction::CreateSP( this, &FSequencer::SetPlaybackRangeEnd ),
+		FExecuteAction::CreateLambda([this] { SetPlaybackEnd(); }),
 		FCanExecuteAction::CreateSP( this, &FSequencer::IsViewingMasterSequence ) );
 
 	SequencerCommandBindings->MapAction(
diff --git a/Engine/Source/Editor/Sequencer/Private/Sequencer.h b/Engine/Source/Editor/Sequencer/Private/Sequencer.h
--- a/Engine/Source/Editor/Sequencer/Private/Sequencer.h
+++ b/Engine/Source/Editor/Sequencer/Private/Sequencer.h
@@ -161,28 +161,6 @@ public:
 	 */
 	void SetPlaybackRange(TRange<FFrameNumber> Range);
 	
-	/**
-	 * Set the playback range's end position to the current global time.
-	 *
-	 * @see GetPlaybackRange, SetPlaybackRange, SetPlaybackRangeStart
-	 */
-	void SetPlaybackRangeEnd()
-	{
-		TRange<FFrameNumber> PlayRange = GetPlaybackRange();
-		SetPlaybackRange(TRange<FFrameNumber>(PlayRange.GetLowerBound(), TRangeBound<FFrameNumber>::Exclusive(GetLocalTime().Time.FrameNumber)));
-	}
-
-	/**
-	 * Set the playback range's start position to the current global time.
-	 *
-	 * @see GetPlaybackRange, SetPlaybackRange, SetPlaybackRangeStart
-	 */
-	void SetPlaybackRangeStart()
-	{
-		TRange<FFrameNumber> PlayRange = GetPlaybackRange();
-		SetPlaybackRange(TRange<FFrameNumber>(TRangeBound<FFrameNumber>::Inclusive(GetLocalTime().Time.FrameNumber), PlayRange.GetUpperBound()));
-	}
-
 	/**
 	 * Set the selection range to the next or previous shot's range.
 	 *
