Sequencer: Fix sections not being visible when dragging a section creates additional subtracks.

diff --git a/Engine/Source/Editor/Sequencer/Private/Tools/EditToolDragOperations.cpp b/Engine/Source/Editor/Sequencer/Private/Tools/EditToolDragOperations.cpp
--- a/Engine/Source/Editor/Sequencer/Private/Tools/EditToolDragOperations.cpp
+++ b/Engine/Source/Editor/Sequencer/Private/Tools/EditToolDragOperations.cpp
@@ -1081,6 +1081,25 @@ bool FMoveKeysAndSections::HandleSectionMovement(FFrameTime MouseTime, FVector2D
 							TargetRowIndex = ChildIndex + 1;
 						}
 					}
+				
+					if (TargetRowIndex > 0)
+					{
+						if (!ParentTrack->IsExpanded() && ParentTrack != ExpandedParentTrack)
+						{
+							if (ExpandedParentTrack.IsValid())
+							{
+								ExpandedParentTrack->SetExpansionState(false);
+								ExpandedParentTrack = nullptr;
+							}
+							ExpandedParentTrack = ParentTrack;
+							ParentTrack->SetExpansionState(true);
+						}
+					}
+					else if (ExpandedParentTrack.IsValid())
+					{
+						ExpandedParentTrack->SetExpansionState(false);
+						ExpandedParentTrack = nullptr;
+					}				
 				}
 			}
 		}
diff --git a/Engine/Source/Editor/Sequencer/Private/Tools/EditToolDragOperations.h b/Engine/Source/Editor/Sequencer/Private/Tools/EditToolDragOperations.h
--- a/Engine/Source/Editor/Sequencer/Private/Tools/EditToolDragOperations.h
+++ b/Engine/Source/Editor/Sequencer/Private/Tools/EditToolDragOperations.h
@@ -203,6 +203,9 @@ protected:
 	/** Optional snap field to use when dragging */
 	TOptional<FSequencerSnapField> SnapField;
 
+	/** If we expanded a parent track while dragging, track it here so we can re-collapse it if not dropping on it. */
+	TSharedPtr<FSequencerTrackNode> ExpandedParentTrack;
+
 	/** If the user is moving them via clicking on the Section then we'll allow vertical re-arranging, otherwise not. */
 	bool bHotspotWasSection;
 };
