Sequencer: Sections should be drawn as in-active when the track is muted

diff --git a/Engine/Source/Editor/Sequencer/Private/SSequencerSection.cpp b/Engine/Source/Editor/Sequencer/Private/SSequencerSection.cpp
--- a/Engine/Source/Editor/Sequencer/Private/SSequencerSection.cpp
+++ b/Engine/Source/Editor/Sequencer/Private/SSequencerSection.cpp
@@ -1039,13 +1039,15 @@ int32 SSequencerSection::OnPaint( const FPaintArgs& Args, const FGeometry& Allot
 		Hotspot = GetSequencer().GetHotspot().Get();
 	}
 
-	UMovieScenePropertyTrack* Track = SectionObject->GetTypedOuter<UMovieScenePropertyTrack>();
+	UMovieSceneTrack* Track = SectionObject->GetTypedOuter<UMovieSceneTrack>();
 	const bool bEnabled = bParentEnabled && SectionObject->IsActive() && !(Track && Track->IsEvalDisabled());
 	const bool bLocked = SectionObject->IsLocked();
+
+	UMovieScenePropertyTrack* PropertyTrack = Cast<UMovieScenePropertyTrack>(Track);
 	bool bSetSectionToKey = false;
-	if (Track)
+	if (PropertyTrack)
 	{
-		if (Track->GetSectionToKey() == SectionObject)
+		if (PropertyTrack->GetSectionToKey() == SectionObject)
 		{
 			bSetSectionToKey = true;
 		}
