Sequencer: Save pinned state

diff --git a/Engine/Source/Editor/Sequencer/Private/DisplayNodes/SequencerDisplayNode.cpp b/Engine/Source/Editor/Sequencer/Private/DisplayNodes/SequencerDisplayNode.cpp
--- a/Engine/Source/Editor/Sequencer/Private/DisplayNodes/SequencerDisplayNode.cpp
+++ b/Engine/Source/Editor/Sequencer/Private/DisplayNodes/SequencerDisplayNode.cpp
@@ -446,6 +446,12 @@ void FSequencerDisplayNode::SetParent(TSharedPtr<FSequencerDisplayNode> InParent
 			}
 
 			bExpanded = ParentTree.GetSavedExpansionState( *this );
+
+			if (InParent != ParentTree.GetRootNode())
+			{
+				bPinned = false;
+				ParentTree.SavePinnedState(*this, false);
+			}
 		}
 	}
 
@@ -479,6 +485,10 @@ void FSequencerDisplayNode::OnTreeRefreshed(float InVirtualTop, float InVirtualB
 	{
 		// Assign the saved expansion state when this node is initialized for the first time
 		bExpanded = ParentTree.GetSavedExpansionState( *this );
+		if (IsRootNode())
+		{
+			bPinned = ParentTree.GetSavedPinnedState(*this);
+		}
 	}
 
 	VirtualTop = InVirtualTop;
@@ -1189,7 +1199,10 @@ void FSequencerDisplayNode::TogglePinned()
 	FSequencerDisplayNode* BaseNode = GetBaseNode();
 	bool bShouldPin = !BaseNode->bPinned;
 	ParentTree.UnpinAllNodes();
+
 	BaseNode->bPinned = bShouldPin;
+	ParentTree.SavePinnedState(*this, bShouldPin);
+	
 	ParentTree.GetSequencer().RefreshTree();
 }
 
@@ -1199,6 +1212,8 @@ void FSequencerDisplayNode::Unpin()
 	if (BaseNode->bPinned)
 	{
 		BaseNode->bPinned = false;
+		ParentTree.SavePinnedState(*this, false);
+	
 		ParentTree.GetSequencer().RefreshTree();
 	}
 }
diff --git a/Engine/Source/Editor/Sequencer/Private/SequencerNodeTree.cpp b/Engine/Source/Editor/Sequencer/Private/SequencerNodeTree.cpp
--- a/Engine/Source/Editor/Sequencer/Private/SequencerNodeTree.cpp
+++ b/Engine/Source/Editor/Sequencer/Private/SequencerNodeTree.cpp
@@ -803,6 +803,30 @@ bool FSequencerNodeTree::GetDefaultExpansionState( const FSequencerDisplayNode&
 	return false;
 }
 
+void FSequencerNodeTree::SavePinnedState(const FSequencerDisplayNode& Node, bool bPinned)
+{
+	UMovieScene* MovieScene = Sequencer.GetFocusedMovieSceneSequence()->GetMovieScene();
+	FMovieSceneEditorData& EditorData = MovieScene->GetEditorData();
+
+	if (bPinned)
+	{
+		EditorData.PinnedNodes.AddUnique(Node.GetPathName());
+	}
+	else
+	{
+		EditorData.PinnedNodes.RemoveSingle(Node.GetPathName());
+	}
+}
+
+
+bool FSequencerNodeTree::GetSavedPinnedState(const FSequencerDisplayNode& Node) const
+{
+	UMovieScene* MovieScene = Sequencer.GetFocusedMovieSceneSequence()->GetMovieScene();
+	FMovieSceneEditorData& EditorData = MovieScene->GetEditorData();
+	bool bPinned = EditorData.PinnedNodes.Contains(Node.GetPathName());
+
+	return bPinned;
+}
 
 bool FSequencerNodeTree::IsNodeFiltered(const TSharedRef<const FSequencerDisplayNode> Node) const
 {
diff --git a/Engine/Source/Editor/Sequencer/Private/SequencerNodeTree.h b/Engine/Source/Editor/Sequencer/Private/SequencerNodeTree.h
--- a/Engine/Source/Editor/Sequencer/Private/SequencerNodeTree.h
+++ b/Engine/Source/Editor/Sequencer/Private/SequencerNodeTree.h
@@ -119,6 +119,22 @@ public:
 	 */
 	bool GetDefaultExpansionState( const FSequencerDisplayNode& Node ) const;
 
+	/**
+	 * Saves the pinned state of a display node
+	 *
+	 * @param Node		The node whose pinned state should be saved
+	 * @param bPinned	The new pinned state of the node
+	 */
+	void SavePinnedState( const FSequencerDisplayNode& Node, bool bPinned );
+
+	/**
+	 * Gets the saved pinned state of a display node
+	 *
+	 * @param Node	The node whose pinned state may have been saved
+	 * @return true if the node should be pinned, false otherwise	
+	 */
+	bool GetSavedPinnedState( const FSequencerDisplayNode& Node ) const;
+
 	/**
 	 * Set the single hovered node in the tree
 	 */
diff --git a/Engine/Source/Runtime/MovieScene/Public/MovieScene.h b/Engine/Source/Runtime/MovieScene/Public/MovieScene.h
--- a/Engine/Source/Runtime/MovieScene/Public/MovieScene.h
+++ b/Engine/Source/Runtime/MovieScene/Public/MovieScene.h
@@ -120,6 +120,10 @@ struct FMovieSceneEditorData
 	UPROPERTY()
 	TMap<FString, FMovieSceneExpansionState> ExpansionStates;
 
+	/** List of Pinned nodes. */
+	UPROPERTY()
+	TArray<FString> PinnedNodes;
+
 	/** The last view-range start that the user was observing */
 	UPROPERTY()
 	double ViewStart;
