Sequencer: Only include selected object bindings if this one is selected

diff --git a/Engine/Source/Editor/Sequencer/Private/DisplayNodes/SequencerObjectBindingNode.cpp b/Engine/Source/Editor/Sequencer/Private/DisplayNodes/SequencerObjectBindingNode.cpp
--- a/Engine/Source/Editor/Sequencer/Private/DisplayNodes/SequencerObjectBindingNode.cpp
+++ b/Engine/Source/Editor/Sequencer/Private/DisplayNodes/SequencerObjectBindingNode.cpp
@@ -968,21 +968,25 @@ TSharedRef<SWidget> FSequencerObjectBindingNode::HandleAddTrackComboButtonGetMen
 	TArray<FGuid> ObjectBindings;
 	ObjectBindings.Add(ObjectBinding);
 
-	// Add additionally selected object bindings
-	for (const TSharedRef<FSequencerDisplayNode>& Node : Sequencer.GetSelection().GetSelectedOutlinerNodes())
+	// Only include other selected object bindings if this binding is selected. Otherwise, this will lead to 
+	// confusion with multiple tracks being added to possibly unrelated objects
+	if (Sequencer.GetSelection().IsSelected(SharedThis(this)))
 	{
-		if (Node->GetType() != ESequencerNode::Object)
+		for (const TSharedRef<FSequencerDisplayNode>& Node : Sequencer.GetSelection().GetSelectedOutlinerNodes())
 		{
-			continue;
-		}
+			if (Node->GetType() != ESequencerNode::Object)
+			{
+				continue;
+			}
 
-		auto ObjectBindingNode = StaticCastSharedRef<FSequencerObjectBindingNode>(Node);
+			auto ObjectBindingNode = StaticCastSharedRef<FSequencerObjectBindingNode>(Node);
 
-		FGuid Guid = ObjectBindingNode->GetObjectBinding();
-		for (auto RuntimeObject : Sequencer.FindBoundObjects(Guid, Sequencer.GetFocusedTemplateID()))
-		{
-			ObjectBindings.AddUnique(Guid);
-			continue;
+			FGuid Guid = ObjectBindingNode->GetObjectBinding();
+			for (auto RuntimeObject : Sequencer.FindBoundObjects(Guid, Sequencer.GetFocusedTemplateID()))
+			{
+				ObjectBindings.AddUnique(Guid);
+				continue;
+			}
 		}
 	}
 
