Add extensibility points for drag-and-dropping stuff on the Sequencer editor.

diff --git a/Engine/Source/Editor/Sequencer/Private/SSequencer.cpp b/Engine/Source/Editor/Sequencer/Private/SSequencer.cpp
--- a/Engine/Source/Editor/Sequencer/Private/SSequencer.cpp
+++ b/Engine/Source/Editor/Sequencer/Private/SSequencer.cpp
@@ -292,6 +292,12 @@ void SSequencer::Construct(const FArguments& InArgs, TSharedRef<FSequencer> InSe
 
 	OnReceivedFocus = InArgs._OnReceivedFocus;
 
+	OnReceivedDragOver = InArgs._OnReceivedDragOver;
+	OnReceivedDrop = InArgs._OnReceivedDrop;
+	OnAssetsDrop = InArgs._OnAssetsDrop;
+	OnClassesDrop = InArgs._OnClassesDrop;
+	OnActorsDrop = InArgs._OnActorsDrop;
+
 	USequencerSettings* SequencerSettings = Settings;
 
 	// Get the desired display format from the user's settings each time.
@@ -2503,6 +2509,15 @@ void SSequencer::OnDragLeave( const FDragDropEvent& DragDropEvent )
 
 FReply SSequencer::OnDragOver( const FGeometry& MyGeometry, const FDragDropEvent& DragDropEvent )
 {
+	if (OnReceivedDragOver.IsBound())
+	{
+		FReply DelegateReply = FReply::Unhandled();
+		if (OnReceivedDragOver.Execute(MyGeometry, DragDropEvent, DelegateReply))
+		{
+			return DelegateReply;
+		}
+	}
+
 	bool bIsDragSupported = false;
 
 	TSharedPtr<FDragDropOperation> Operation = DragDropEvent.GetOperation();
@@ -2520,6 +2535,15 @@ FReply SSequencer::OnDragOver( const FGeometry& MyGeometry, const FDragDropEvent
 
 FReply SSequencer::OnDrop( const FGeometry& MyGeometry, const FDragDropEvent& DragDropEvent )
 {
+	if (OnReceivedDrop.IsBound())
+	{
+		FReply DelegateReply = FReply::Unhandled();
+		if (OnReceivedDrop.Execute(MyGeometry, DragDropEvent, DelegateReply))
+		{
+			return DelegateReply;
+		}
+	}
+
 	bool bWasDropHandled = false;
 
 	// @todo sequencer: Get rid of hard-code assumptions about dealing with ACTORS at this level?
@@ -2643,33 +2667,40 @@ void SSequencer::OnAssetsDropped( const FAssetDragDropOp& DragDropOp )
 		}
 	}
 
-	for( auto CurObjectIter = DroppedObjects.CreateConstIterator(); CurObjectIter; ++CurObjectIter )
+	if (OnAssetsDrop.IsBound())
 	{
-		UObject* CurObject = *CurObjectIter;
-
-		if (!SequencerRef.OnHandleAssetDropped(CurObject, TargetObjectGuid))
+		bObjectAdded = OnAssetsDrop.Execute(DroppedObjects, DragDropOp);
+	}
+	else
+	{
+		for (TArray<UObject*>::TConstIterator CurObjectIter = DroppedObjects.CreateConstIterator(); CurObjectIter; ++CurObjectIter)
 		{
-			// Doesn't make sense to drop a level sequence asset into sequencer as a spawnable actor
-			if (CurObject->IsA<ULevelSequence>())
-			{
-				UE_LOG(LogSequencer, Warning, TEXT("Can't add '%s' as a spawnable"), *CurObject->GetName());
-				continue;
-			}
+			UObject* CurObject = *CurObjectIter;
 
-			FGuid NewGuid = SequencerRef.MakeNewSpawnable( *CurObject, DragDropOp.GetActorFactory() );
-
-			UMovieScene* MovieScene = SequencerRef.GetFocusedMovieSceneSequence()->GetMovieScene();
-			if (MovieScene)
+			if (!SequencerRef.OnHandleAssetDropped(CurObject, TargetObjectGuid))
 			{
-				FMovieSceneSpawnable* Spawnable = MovieScene->FindSpawnable(NewGuid);
+				// Doesn't make sense to drop a level sequence asset into sequencer as a spawnable actor
+				if (CurObject->IsA<ULevelSequence>())
+				{
+					UE_LOG(LogSequencer, Warning, TEXT("Can't add '%s' as a spawnable"), *CurObject->GetName());
+					continue;
+				}
+
+				FGuid NewGuid = SequencerRef.MakeNewSpawnable(*CurObject, DragDropOp.GetActorFactory());
 
-				if (Spawnable && Spawnable->GetObjectTemplate()->IsA<ACameraActor>())
+				UMovieScene* MovieScene = SequencerRef.GetFocusedMovieSceneSequence()->GetMovieScene();
+				if (MovieScene)
 				{
-					SequencerRef.NewCameraAdded(NewGuid);
+					FMovieSceneSpawnable* Spawnable = MovieScene->FindSpawnable(NewGuid);
+
+					if (Spawnable && Spawnable->GetObjectTemplate()->IsA<ACameraActor>())
+					{
+						SequencerRef.NewCameraAdded(NewGuid);
+					}
 				}
 			}
+			bObjectAdded = true;
 		}
-		bObjectAdded = true;
 	}
 
 	if( bObjectAdded )
@@ -2687,23 +2718,37 @@ void SSequencer::OnAssetsDropped( const FAssetDragDropOp& DragDropOp )
 
 void SSequencer::OnClassesDropped( const FClassDragDropOp& DragDropOp )
 {
-	FSequencer& SequencerRef = *SequencerPtr.Pin();
-
-	for( auto ClassIter = DragDropOp.ClassesToDrop.CreateConstIterator(); ClassIter; ++ClassIter )
+	if (OnClassesDrop.IsBound())
 	{
-		UClass* Class = ( *ClassIter ).Get();
-		if( Class != nullptr )
+		OnClassesDrop.Execute(DragDropOp.ClassesToDrop, DragDropOp);
+	}
+	else
+	{
+		FSequencer& SequencerRef = *SequencerPtr.Pin();
+
+		for (auto ClassIter = DragDropOp.ClassesToDrop.CreateConstIterator(); ClassIter; ++ClassIter)
 		{
-			UObject* Object = Class->GetDefaultObject();
+			UClass* Class = (*ClassIter).Get();
+			if (Class != nullptr)
+			{
+				UObject* Object = Class->GetDefaultObject();
 
-			FGuid NewGuid = SequencerRef.MakeNewSpawnable( *Object );
+				FGuid NewGuid = SequencerRef.MakeNewSpawnable(*Object);
+			}
 		}
 	}
 }
 
 void SSequencer::OnActorsDropped( FActorDragDropGraphEdOp& DragDropOp )
 {
-	SequencerPtr.Pin()->OnActorsDropped( DragDropOp.Actors );
+	if (OnActorsDrop.IsBound())
+	{
+		OnActorsDrop.Execute(DragDropOp.Actors, DragDropOp);
+	}
+	else
+	{
+		SequencerPtr.Pin()->OnActorsDropped(DragDropOp.Actors);
+	}
 }
 
 void SSequencer::OnCrumbClicked(const FSequencerBreadcrumb& Item)
diff --git a/Engine/Source/Editor/Sequencer/Private/SSequencer.h b/Engine/Source/Editor/Sequencer/Private/SSequencer.h
--- a/Engine/Source/Editor/Sequencer/Private/SSequencer.h
+++ b/Engine/Source/Editor/Sequencer/Private/SSequencer.h
@@ -223,6 +223,21 @@ public:
 		/** Called when any widget contained within sequencer has received focus */
 		SLATE_EVENT( FSimpleDelegate, OnReceivedFocus )
 
+		/** Called when something is dragged over the sequencer. */
+		SLATE_EVENT( FOptionalOnDragDrop, OnReceivedDragOver )
+
+		/** Called when something is dropped onto the sequencer. */
+		SLATE_EVENT( FOptionalOnDragDrop, OnReceivedDrop )
+
+		/** Called when an asset is dropped on the sequencer. Not called if OnReceivedDrop is bound and returned true. */
+		SLATE_EVENT( FOnAssetsDrop, OnAssetsDrop )
+
+		/** Called when a class is dropped on the sequencer. Not called if OnReceivedDrop is bound and returned true. */
+		SLATE_EVENT( FOnClassesDrop, OnClassesDrop )
+		
+		/** Called when an actor is dropped on the sequencer. Not called if OnReceivedDrop is bound and returned true. */
+		SLATE_EVENT( FOnActorsDrop, OnActorsDrop )
+
 		/** Extender to use for the add menu. */
 		SLATE_ARGUMENT( TSharedPtr<FExtender>, AddMenuExtender )
 
@@ -651,6 +666,21 @@ private:
 	/** Called when any widget contained within sequencer has received focus */
 	FSimpleDelegate OnReceivedFocus;
 
+	/** Called when something is dragged over the sequencer. */
+	FOptionalOnDragDrop OnReceivedDragOver;
+
+	/** Called when something is dropped onto the sequencer. */
+	FOptionalOnDragDrop OnReceivedDrop;
+
+	/** Called when an asset is dropped on the sequencer. */
+	FOnAssetsDrop OnAssetsDrop;
+
+	/** Called when a class is dropped on the sequencer. */
+	FOnClassesDrop OnClassesDrop;
+	
+	/** Called when an actor is dropped on the sequencer. */
+	FOnActorsDrop OnActorsDrop;
+
 	/** Cached clamp and view range for unlinking the curve editor time range */
 	TRange<double> CachedClampRange;
 	TRange<double> CachedViewRange;
diff --git a/Engine/Source/Editor/Sequencer/Private/Sequencer.cpp b/Engine/Source/Editor/Sequencer/Private/Sequencer.cpp
--- a/Engine/Source/Editor/Sequencer/Private/Sequencer.cpp
+++ b/Engine/Source/Editor/Sequencer/Private/Sequencer.cpp
@@ -420,6 +420,10 @@ void FSequencer::InitSequencer(const FSequencerInitParams& InitParams, const TSh
 		.OnGetAddMenuContent(InitParams.ViewParams.OnGetAddMenuContent)
 		.OnBuildCustomContextMenuForGuid(InitParams.ViewParams.OnBuildCustomContextMenuForGuid)
 		.OnReceivedFocus(InitParams.ViewParams.OnReceivedFocus)
+		.OnReceivedDragOver(InitParams.ViewParams.OnReceivedDragOver)
+		.OnAssetsDrop(InitParams.ViewParams.OnAssetsDrop)
+		.OnClassesDrop(InitParams.ViewParams.OnClassesDrop)
+		.OnActorsDrop(InitParams.ViewParams.OnActorsDrop)
 		.AddMenuExtender(InitParams.ViewParams.AddMenuExtender)
 		.ToolbarExtender(InitParams.ViewParams.ToolbarExtender);
 
diff --git a/Engine/Source/Editor/Sequencer/Public/ISequencerModule.h b/Engine/Source/Editor/Sequencer/Public/ISequencerModule.h
--- a/Engine/Source/Editor/Sequencer/Public/ISequencerModule.h
+++ b/Engine/Source/Editor/Sequencer/Public/ISequencerModule.h
@@ -19,6 +19,9 @@ class ISequencerTrackEditor;
 class ISequencerEditorObjectBinding;
 class IToolkitHost;
 class UMovieSceneSequence;
+class FAssetDragDropOp;
+class FClassDragDropOp;
+class FActorDragDropGraphEdOp;
 struct FSequencerInitParams;
 
 enum class ECurveEditorTreeFilterType : uint32;
@@ -59,6 +62,18 @@ DECLARE_MULTICAST_DELEGATE_OneParam(FOnSequencerCreated, TSharedRef<ISequencer>)
 /** A delegate that gets executed a sequencer is initialize and allow modification the initialization params. */
 DECLARE_MULTICAST_DELEGATE_ThreeParams(FOnPreSequencerInit, TSharedRef<ISequencer>, TSharedRef<ISequencerObjectChangeListener>, const FSequencerInitParams&);
 
+/** A delegate that gets executed when a drag/drop event happens on the sequencer. The return value determines if the event was handled by the bound delegate. */
+DECLARE_DELEGATE_RetVal_ThreeParams(bool, FOptionalOnDragDrop, const FGeometry&, const FDragDropEvent&, FReply&);
+
+/** A delegate that gets executed when an asset is dropped on the sequencer. The return value determines if the operation was handled by the bound delegate. */
+DECLARE_DELEGATE_RetVal_TwoParams(bool, FOnAssetsDrop, const TArray<UObject*>&, const FAssetDragDropOp&);
+
+/** A delegate that gets executed when a class is dropped on the sequencer. The return value determines if the operation was handled by the bound delegate. */
+DECLARE_DELEGATE_RetVal_TwoParams(bool, FOnClassesDrop, const TArray<TWeakObjectPtr<UClass>>&, const FClassDragDropOp&);
+
+/** A delegate that gets executed when an actor is dropped on the sequencer. The return value determines if the operation was handled by the bound delegate. */
+DECLARE_DELEGATE_RetVal_TwoParams(bool, FOnActorsDrop, const TArray<TWeakObjectPtr<AActor>>&, const FActorDragDropGraphEdOp&);
+
 /**
  * Sequencer view parameters.
  */
@@ -86,6 +101,18 @@ struct FSequencerViewParams
 	/** Style of scrubber to use */
 	ESequencerScrubberStyle ScrubberStyle;
 
+	/** Called when something is dragged over the sequencer. */
+	FOptionalOnDragDrop OnReceivedDragOver;
+
+	/** Called when an asset is dropped on the sequencer. */
+	FOnAssetsDrop OnAssetsDrop;
+
+	/** Called when a class is dropped on the sequencer. */
+	FOnClassesDrop OnClassesDrop;
+
+	/** Called when an actor is dropped on the sequencer. */
+	FOnActorsDrop OnActorsDrop;
+
 	FSequencerViewParams(FString InName = FString())
 		: UniqueName(MoveTemp(InName))
 		, bReadOnly(false)
@@ -399,4 +426,4 @@ inline ISequencerChannelInterface* ISequencerModule::FindChannelEditorInterface(
 	const TUniquePtr<ISequencerChannelInterface>* Found = ChannelToEditorInterfaceMap.Find(ChannelTypeName);
 	ensureMsgf(Found, TEXT("No channel interface found for type ID. Did you call RegisterChannelInterface<> for that type?"));
 	return Found ? Found->Get() : nullptr;
-}
\ No newline at end of file
+}
