Sequencer: Make track icon dim if muted. Track display text is now also dimmed if all the sections are inactive

diff --git a/Engine/Source/Editor/Sequencer/Private/DisplayNodes/SequencerDisplayNode.cpp b/Engine/Source/Editor/Sequencer/Private/DisplayNodes/SequencerDisplayNode.cpp
--- a/Engine/Source/Editor/Sequencer/Private/DisplayNodes/SequencerDisplayNode.cpp
+++ b/Engine/Source/Editor/Sequencer/Private/DisplayNodes/SequencerDisplayNode.cpp
@@ -642,14 +642,14 @@ bool FSequencerDisplayNode::TraverseVisible_ParentFirst(const TFunctionRef<bool(
 	return true;
 }
 
-FLinearColor FSequencerDisplayNode::GetDisplayNameColor() const
+bool FSequencerDisplayNode::IsDimmed() const
 {
 	auto FindInActiveSection = [](FSequencerDisplayNode& InNode, bool EmptyNotActive = true)
 	{
 		if (InNode.GetType() == ESequencerNode::KeyArea)
 		{
 			const FSequencerSectionKeyAreaNode& KeyAreaNode = static_cast<FSequencerSectionKeyAreaNode&>(InNode);
-			auto KeyAreaNodes=  KeyAreaNode.GetAllKeyAreas();
+			auto KeyAreaNodes = KeyAreaNode.GetAllKeyAreas();
 			if (KeyAreaNodes.Num() > 0)
 			{
 				for (const TSharedRef<IKeyArea>& KeyArea : KeyAreaNodes)
@@ -715,7 +715,7 @@ FLinearColor FSequencerDisplayNode::GetDisplayNameColor() const
 	//find first child with active section, then it's active, else inactive.
 	bool bDimLabel = ChildNodes.Num() > 0 ? This->Traverse_ParentFirst(FindInActiveSection) :
 		((this->GetType() == ESequencerNode::Track || this->GetType() == ESequencerNode::KeyArea) && FindInActiveSection(*(This), false))
-		||false;
+		|| false;
 
 	if (!bDimLabel)
 	{
@@ -735,10 +735,14 @@ FLinearColor FSequencerDisplayNode::GetDisplayNameColor() const
 				bDimLabel = true;
 			}
 		}
-
 	}
 
-	return bDimLabel ? FLinearColor(0.6f, 0.6f, 0.6f, 0.6f) : FLinearColor(1.0f, 1.0f, 1.0f, 1.0f);
+	return bDimLabel;
+}
+
+FLinearColor FSequencerDisplayNode::GetDisplayNameColor() const
+{
+	return IsDimmed() ? FLinearColor(0.6f, 0.6f, 0.6f, 0.6f) : FLinearColor(1.0f, 1.0f, 1.0f, 1.0f);
 }
 
 FText FSequencerDisplayNode::GetDisplayNameToolTipText() const
@@ -788,7 +792,7 @@ const FSlateBrush* FSequencerDisplayNode::GetIconOverlayBrush() const
 
 FSlateColor FSequencerDisplayNode::GetIconColor() const
 {
-	return FSlateColor( FLinearColor::White );
+	return GetDisplayNameColor();
 }
 
 FText FSequencerDisplayNode::GetIconToolTipText() const
diff --git a/Engine/Source/Editor/Sequencer/Private/DisplayNodes/SequencerDisplayNode.h b/Engine/Source/Editor/Sequencer/Private/DisplayNodes/SequencerDisplayNode.h
--- a/Engine/Source/Editor/Sequencer/Private/DisplayNodes/SequencerDisplayNode.h
+++ b/Engine/Source/Editor/Sequencer/Private/DisplayNodes/SequencerDisplayNode.h
@@ -187,6 +187,11 @@ public:
 	 */
 	virtual void SetDisplayName(const FText& NewDisplayName) = 0;
 
+	/**
+	 * @return Whether this track should be drawn as dim 
+	 */
+	virtual bool IsDimmed() const;
+
 	/**
 	 * @return Whether this node handles resize events
 	 */
diff --git a/Engine/Source/Editor/Sequencer/Private/DisplayNodes/SequencerFolderNode.cpp b/Engine/Source/Editor/Sequencer/Private/DisplayNodes/SequencerFolderNode.cpp
--- a/Engine/Source/Editor/Sequencer/Private/DisplayNodes/SequencerFolderNode.cpp
+++ b/Engine/Source/Editor/Sequencer/Private/DisplayNodes/SequencerFolderNode.cpp
@@ -111,6 +111,11 @@ const FSlateBrush* FSequencerFolderNode::GetIconBrush() const
 
 FSlateColor FSequencerFolderNode::GetIconColor() const
 {
+	if (ParentTree.IsNodeMute(this))
+	{
+		return FSlateColor(MovieSceneFolder.GetFolderColor().ReinterpretAsLinear() * FLinearColor(0.6f, 0.6f, 0.6f, 0.6f));
+	}
+
 	return FSlateColor(MovieSceneFolder.GetFolderColor());
 }
 
diff --git a/Engine/Source/Editor/Sequencer/Private/DisplayNodes/SequencerTrackNode.cpp b/Engine/Source/Editor/Sequencer/Private/DisplayNodes/SequencerTrackNode.cpp
--- a/Engine/Source/Editor/Sequencer/Private/DisplayNodes/SequencerTrackNode.cpp
+++ b/Engine/Source/Editor/Sequencer/Private/DisplayNodes/SequencerTrackNode.cpp
@@ -693,15 +693,15 @@ FLinearColor FSequencerTrackNode::GetDisplayNameColor() const
 		return FLinearColor::White;
 	}
 
-	const bool bIsEvalDisabled = Track->IsEvalDisabled();
-
 	// Display track node is red if the property track is not bound to valid property
 	if (UMovieScenePropertyTrack* PropertyTrack = Cast<UMovieScenePropertyTrack>(Track))
 	{
+		const bool bIsDimmed = IsDimmed();
+
 		// 3D transform tracks don't map to property bindings as below
 		if (Track->IsA<UMovieScene3DTransformTrack>() || Track->IsA<UMovieScenePrimitiveMaterialTrack>())
 		{
-			return bIsEvalDisabled ? FLinearColor(0.6f, 0.6f, 0.6f, 0.6f) : FLinearColor::White;
+			return bIsDimmed ? FLinearColor(0.6f, 0.6f, 0.6f, 0.6f) : FLinearColor::White;
 		}
 
 		FGuid ObjectBinding;
@@ -721,16 +721,16 @@ FLinearColor FSequencerTrackNode::GetDisplayNameColor() const
 					FTrackInstancePropertyBindings PropertyBinding(PropertyTrack->GetPropertyName(), PropertyTrack->GetPropertyPath());
 					if (PropertyBinding.GetProperty(*Object))
 					{
-						return bIsEvalDisabled ? FLinearColor(0.6f, 0.6f, 0.6f, 0.6f) : FLinearColor::White;
+						return bIsDimmed ? FLinearColor(0.6f, 0.6f, 0.6f, 0.6f) : FLinearColor::White;
 					}
 				}
 			}
 
-			return bIsEvalDisabled ? FLinearColor(0.6f, 0.0f, 0.0f, 0.6f) : FLinearColor::Red;
+			return bIsDimmed ? FLinearColor(0.6f, 0.0f, 0.0f, 0.6f) : FLinearColor::Red;
 		}
 	}
 
-	return bIsEvalDisabled ? FLinearColor(0.6f, 0.6f, 0.6f, 0.6f) : FLinearColor::White;
+	return FSequencerDisplayNode::GetDisplayNameColor();
 }
 
 float FSequencerTrackNode::GetNodeHeight() const
