Sequencer: Jump to start of playback range when stepping to next previous shot.

diff --git a/Engine/Source/Editor/Sequencer/Private/Sequencer.cpp b/Engine/Source/Editor/Sequencer/Private/Sequencer.cpp
--- a/Engine/Source/Editor/Sequencer/Private/Sequencer.cpp
+++ b/Engine/Source/Editor/Sequencer/Private/Sequencer.cpp
@@ -5378,7 +5378,8 @@ void FSequencer::SynchronizeExternalSelectionWithSequencerSelection()
 	{
 		if (GEditor->GetSelectedActorCount())
 		{
-			const FScopedTransaction Transaction( NSLOCTEXT( "Sequencer", "UpdatingActorComponentSelectionNone", "Select None" ) );
+			const bool bShouldActuallyTransact = !GIsTransacting;
+			const FScopedTransaction Transaction( NSLOCTEXT( "Sequencer", "UpdatingActorComponentSelectionNone", "Select None" ), bShouldActuallyTransact );
 			GEditor->SelectNone( bNotifySelectionChanged, bDeselectBSP, bWarnAboutTooManyActors );
 			GEditor->NoteSelectionChange();
 		}
@@ -5448,7 +5449,8 @@ void FSequencer::SynchronizeExternalSelectionWithSequencerSelection()
 		return;
 	}
 
-	const FScopedTransaction Transaction( NSLOCTEXT( "Sequencer", "UpdatingActorComponentSelection", "Select Actors/Components" ) );
+	const bool bShouldActuallyTransact = !GIsTransacting;
+	const FScopedTransaction Transaction( NSLOCTEXT( "Sequencer", "UpdatingActorComponentSelection", "Select Actors/Components" ), bShouldActuallyTransact );
 
 
 	GEditor->GetSelectedActors()->Modify();
@@ -5514,6 +5516,11 @@ void FSequencer::SynchronizeSequencerSelectionWithExternalSelection()
 		return;
 	}
 
+	if (!Sequence->GetMovieScene())
+	{
+		return;
+	}
+
 	TGuardValue<bool> Guard(bUpdatingSequencerSelection, true);
 
 	// If all nodes are already selected, do nothing. This ensures that when an undo event happens, 
@@ -8622,18 +8629,19 @@ void FSequencer::StepToNextShot()
 {
 	if (ActiveTemplateIDs.Num() < 2)
 	{
+		UMovieSceneSection* TargetShotSection = FindNextOrPreviousShot(GetFocusedMovieSceneSequence(), GetLocalTime().Time.FloorToFrame(), true);
+
+		if (TargetShotSection)
+		{
+			SetLocalTime(TargetShotSection->GetRange().GetLowerBoundValue(), ESnapTimeMode::STM_None);
+		}
 		return;
 	}
 
 	FMovieSceneSequenceID OuterSequenceID = ActiveTemplateIDs[ActiveTemplateIDs.Num()-2];
 	UMovieSceneSequence* Sequence = RootTemplateInstance.GetSequence(OuterSequenceID);
 
-	FFrameTime StartTime = FFrameTime(0) * RootToLocalTransform.Inverse();
-	FFrameTime CurrentTime = StartTime;
-	if (FMovieSceneSubSequenceData* SubSequenceData = RootTemplateInstance.GetHierarchy().FindSubData(OuterSequenceID))
-	{
-		CurrentTime = StartTime * SubSequenceData->RootToSequenceTransform;
-	}
+	FFrameTime CurrentTime = SubSequenceRange.GetLowerBoundValue() * RootToLocalTransform.Inverse();
 
 	UMovieSceneSubSection* NextShot = Cast<UMovieSceneSubSection>(FindNextOrPreviousShot(Sequence, CurrentTime.FloorToFrame(), true));
 	if (!NextShot)
@@ -8646,7 +8654,7 @@ void FSequencer::StepToNextShot()
 	PopToSequenceInstance(ActiveTemplateIDs[ActiveTemplateIDs.Num()-2]);
 	FocusSequenceInstance(*NextShot);
 
-	SetLocalTime(FFrameTime(0), ESnapTimeMode::STM_None);
+	SetLocalTime(SubSequenceRange.GetLowerBoundValue(), ESnapTimeMode::STM_None);
 }
 
 
@@ -8654,19 +8662,19 @@ void FSequencer::StepToPreviousShot()
 {
 	if (ActiveTemplateIDs.Num() < 2)
 	{
+		UMovieSceneSection* TargetShotSection = FindNextOrPreviousShot(GetFocusedMovieSceneSequence(), GetLocalTime().Time.FloorToFrame(), false);
+
+		if (TargetShotSection)
+		{
+			SetLocalTime(TargetShotSection->GetRange().GetLowerBoundValue(), ESnapTimeMode::STM_None);
+		}
 		return;
 	}
 
-	FMovieSceneSequenceID OuterSequenceID = ActiveTemplateIDs[ActiveTemplateIDs.Num()-2];
+	FMovieSceneSequenceID OuterSequenceID = ActiveTemplateIDs[ActiveTemplateIDs.Num() - 2];
 	UMovieSceneSequence* Sequence = RootTemplateInstance.GetSequence(OuterSequenceID);
 
-	FFrameTime StartTime = FFrameTime(0) * RootToLocalTransform.Inverse();
-	FFrameTime CurrentTime = StartTime;
-	if (FMovieSceneSubSequenceData* SubSequenceData = RootTemplateInstance.GetHierarchy().FindSubData(OuterSequenceID))
-	{
-		CurrentTime = StartTime * SubSequenceData->RootToSequenceTransform;
-	}
-
+	FFrameTime CurrentTime = SubSequenceRange.GetLowerBoundValue() * RootToLocalTransform.Inverse();
 	UMovieSceneSubSection* PreviousShot = Cast<UMovieSceneSubSection>(FindNextOrPreviousShot(Sequence, CurrentTime.FloorToFrame(), false));
 	if (!PreviousShot)
 	{
@@ -8678,7 +8686,7 @@ void FSequencer::StepToPreviousShot()
 	PopToSequenceInstance(ActiveTemplateIDs[ActiveTemplateIDs.Num()-2]);
 	FocusSequenceInstance(*PreviousShot);
 
-	SetLocalTime(0, ESnapTimeMode::STM_None);
+	SetLocalTime(SubSequenceRange.GetLowerBoundValue(), ESnapTimeMode::STM_None);
 }
 
 
