Sequencer: Backout unrelated changes from CL#9118948

diff --git a/Engine/Source/Editor/Sequencer/Private/Sequencer.cpp b/Engine/Source/Editor/Sequencer/Private/Sequencer.cpp
--- a/Engine/Source/Editor/Sequencer/Private/Sequencer.cpp
+++ b/Engine/Source/Editor/Sequencer/Private/Sequencer.cpp
@@ -5378,8 +5378,7 @@ void FSequencer::SynchronizeExternalSelectionWithSequencerSelection()
 	{
 		if (GEditor->GetSelectedActorCount())
 		{
-			const bool bShouldActuallyTransact = !GIsTransacting;
-			const FScopedTransaction Transaction( NSLOCTEXT( "Sequencer", "UpdatingActorComponentSelectionNone", "Select None" ), bShouldActuallyTransact );
+			const FScopedTransaction Transaction( NSLOCTEXT( "Sequencer", "UpdatingActorComponentSelectionNone", "Select None" ) );
 			GEditor->SelectNone( bNotifySelectionChanged, bDeselectBSP, bWarnAboutTooManyActors );
 			GEditor->NoteSelectionChange();
 		}
@@ -5449,8 +5448,7 @@ void FSequencer::SynchronizeExternalSelectionWithSequencerSelection()
 		return;
 	}
 
-	const bool bShouldActuallyTransact = !GIsTransacting;
-	const FScopedTransaction Transaction( NSLOCTEXT( "Sequencer", "UpdatingActorComponentSelection", "Select Actors/Components" ), bShouldActuallyTransact );
+	const FScopedTransaction Transaction( NSLOCTEXT( "Sequencer", "UpdatingActorComponentSelection", "Select Actors/Components" ) );
 
 
 	GEditor->GetSelectedActors()->Modify();
@@ -5516,11 +5514,6 @@ void FSequencer::SynchronizeSequencerSelectionWithExternalSelection()
 		return;
 	}
 
-	if (!Sequence->GetMovieScene())
-	{
-		return;
-	}
-
 	TGuardValue<bool> Guard(bUpdatingSequencerSelection, true);
 
 	// If all nodes are already selected, do nothing. This ensures that when an undo event happens, 
