Sequencer: Toolbar UI tweaks

  - Moved "Restore Animated State" onto "General Options" menu
  - Added "Open Director Blueprint" to the "General Options" menu
  - Saturated opacity for all 3 options menus to be more in-keeping with the rest of the toolbar

diff --git a/Engine/Source/Editor/EditorStyle/Private/SlateEditorStyle.cpp b/Engine/Source/Editor/EditorStyle/Private/SlateEditorStyle.cpp
--- a/Engine/Source/Editor/EditorStyle/Private/SlateEditorStyle.cpp
+++ b/Engine/Source/Editor/EditorStyle/Private/SlateEditorStyle.cpp
@@ -2493,6 +2493,8 @@ void FSlateEditorStyle::FStyle::SetupSequencerStyles()
 		Set("Sequencer.ToggleIsSnapEnabled.Small", new IMAGE_BRUSH("Sequencer/Main_Icons/Icon_Sequencer_Snap_24x", Icon24x24));
 		Set("Sequencer.ToggleShowCurveEditor", new IMAGE_BRUSH("GenericCurveEditor/Icons/GenericCurveEditor_48x", Icon48x48));
 		Set("Sequencer.ToggleShowCurveEditor.Small", new IMAGE_BRUSH("GenericCurveEditor/Icons/GenericCurveEditor_48x", Icon24x24));
+		Set("Sequencer.OpenDirectorBlueprint", new IMAGE_BRUSH("Sequencer/Main_Icons/Icon_Sequencer_OpenDirectorBlueprint_24x", Icon48x48));
+		Set("Sequencer.OpenDirectorBlueprint.Small", new IMAGE_BRUSH("Sequencer/Main_Icons/Icon_Sequencer_OpenDirectorBlueprint_24x", Icon24x24));
 		Set("Sequencer.ToggleAutoScroll", new IMAGE_BRUSH("Icons/icon_Sequencer_ToggleAutoScroll_40x", Icon48x48));
 		Set("Sequencer.ToggleAutoScroll.Small", new IMAGE_BRUSH("Icons/icon_Sequencer_ToggleAutoScroll_16x", Icon16x16));
 		Set("Sequencer.MoveTool.Small", new IMAGE_BRUSH("Icons/SequencerIcons/icon_Sequencer_Move_24x", Icon16x16));
diff --git a/Engine/Source/Editor/Sequencer/Private/SSequencer.cpp b/Engine/Source/Editor/Sequencer/Private/SSequencer.cpp
--- a/Engine/Source/Editor/Sequencer/Private/SSequencer.cpp
+++ b/Engine/Source/Editor/Sequencer/Private/SSequencer.cpp
@@ -31,6 +31,8 @@
 #include "Widgets/Input/SCheckBox.h"
 #include "Widgets/Docking/SDockTab.h"
 #include "EditorStyleSet.h"
+#include "Kismet2/KismetEditorUtilities.h"
+#include "MovieSceneSequenceEditor.h"
 #include "Engine/Selection.h"
 #include "LevelEditorViewport.h"
 #include "Widgets/Navigation/SBreadcrumbTrail.h"
@@ -1022,6 +1024,25 @@ void SSequencer::BindCommands(TSharedRef<FUICommandList> SequencerCommandBinding
 		FSequencerCommands::Get().ToggleShowStretchBox,
 		FExecuteAction::CreateLambda([this] { StretchBox->ToggleVisibility(); })
 	);
+
+	auto OpenDirectorBlueprint = [WeakSequencer = SequencerPtr]
+	{
+		UMovieSceneSequence*       Sequence       = WeakSequencer.Pin()->GetFocusedMovieSceneSequence();
+		FMovieSceneSequenceEditor* SequenceEditor = Sequence ? FMovieSceneSequenceEditor::Find(Sequence) : nullptr;
+		if (SequenceEditor)
+		{
+			UBlueprint* DirectorBP = SequenceEditor->GetOrCreateDirectorBlueprint(Sequence);
+			if (DirectorBP)
+			{
+				FKismetEditorUtilities::BringKismetToFocusAttentionOnObject(DirectorBP);
+			}
+		}
+	};
+
+	SequencerCommandBindings->MapAction(
+		FSequencerCommands::Get().OpenDirectorBlueprint,
+		FExecuteAction::CreateLambda(OpenDirectorBlueprint)
+	);
 }
 
 void SSequencer::ShowTickResolutionOverlay()
@@ -1330,8 +1351,6 @@ TSharedRef<SWidget> SSequencer::MakeToolBar()
 			ToolBarBuilder.AddSeparator("Level Sequence Separator");
 		}
 
-		ToolBarBuilder.AddToolBarButton( FSequencerCommands::Get().RestoreAnimatedState );
-
 		ToolBarBuilder.AddComboButton(
 			FUIAction(),
 			FOnGetContent::CreateSP(this, &SSequencer::MakeGeneralMenu),
@@ -1743,6 +1762,21 @@ TSharedRef<SWidget> SSequencer::MakeGeneralMenu()
 	FMenuBuilder MenuBuilder( true, SequencerPtr.Pin()->GetCommandBindings() );
 	TSharedPtr<FSequencer> Sequencer = SequencerPtr.Pin();
 
+
+	MenuBuilder.BeginSection("SequenceOptions", LOCTEXT("SequenceOptionsHeader", "Sequence"));
+	{
+		UMovieSceneSequence* RootSequence = SequencerPtr.Pin()->GetRootMovieSceneSequence();
+		if (RootSequence->GetTypedOuter<UBlueprint>() == nullptr)
+		{
+			// Only show this button where it makes sense (ie, if the sequence is not contained within a blueprint already)
+			MenuBuilder.AddMenuEntry(FSequencerCommands::Get().OpenDirectorBlueprint);
+		}
+
+		MenuBuilder.AddMenuEntry(FSequencerCommands::Get().RestoreAnimatedState);
+	}
+	MenuBuilder.EndSection();
+
+
 	// view options
 	MenuBuilder.BeginSection( "ViewOptions", LOCTEXT( "ViewMenuHeader", "View" ) );
 	{
diff --git a/Engine/Source/Editor/Sequencer/Private/SequencerCommands.cpp b/Engine/Source/Editor/Sequencer/Private/SequencerCommands.cpp
--- a/Engine/Source/Editor/Sequencer/Private/SequencerCommands.cpp
+++ b/Engine/Source/Editor/Sequencer/Private/SequencerCommands.cpp
@@ -94,6 +94,7 @@ void FSequencerCommands::RegisterCommands()
 	UI_COMMAND( ToggleShowGotoBox, "Go to Time...", "Go to a particular point on the timeline", EUserInterfaceActionType::Button, FInputChord(EModifierKey::Control, EKeys::G) );
 	UI_COMMAND( ToggleShowTransformBox, "Transform Selection...", "Transform the selected keys and sections by a given amount", EUserInterfaceActionType::Button, FInputChord(EModifierKey::Control, EKeys::M));
 	UI_COMMAND( ToggleShowStretchBox, "Stretch/Shrink...", "Stretch or shrink a given amount, moving keys forwards/backwards as necessary", EUserInterfaceActionType::Button, FInputChord() );
+	UI_COMMAND( OpenDirectorBlueprint, "Open Director Blueprint", "Opens the director blueprint for this sequence.", EUserInterfaceActionType::Button, FInputChord() );
 	UI_COMMAND( BakeTransform, "Bake Transform", "Bake transform in world space, removing any existing transform and attach tracks", EUserInterfaceActionType::Button, FInputChord() );
 	UI_COMMAND( SyncSectionsUsingSourceTimecode, "Sync Sections using Source Timecode", "Synchronize sections to the first selected section using source timecode", EUserInterfaceActionType::Button, FInputChord() );
 	UI_COMMAND( ToggleShowRangeSlider, "Show Range Slider", "Enables and disables showing the time range slider", EUserInterfaceActionType::ToggleButton, FInputChord() );
diff --git a/Engine/Source/Editor/Sequencer/Public/SequencerCommands.h b/Engine/Source/Editor/Sequencer/Public/SequencerCommands.h
--- a/Engine/Source/Editor/Sequencer/Public/SequencerCommands.h
+++ b/Engine/Source/Editor/Sequencer/Public/SequencerCommands.h
@@ -228,6 +228,9 @@ public:
 
 	/** Toggle the visibility of the stretch box. */
 	TSharedPtr< FUICommandInfo > ToggleShowStretchBox;
+
+	/** Opens the director blueprint for a sequence. */
+	TSharedPtr< FUICommandInfo > OpenDirectorBlueprint;
 	
 	/** Sets the tree search widget as the focused widget in Slate for easy typing. */
 	TSharedPtr< FUICommandInfo > QuickTreeSearch;
