Sequencer: Remove unused/infrequently used "Discard Changes" and "Save As"

diff --git a/Engine/Plugins/MovieScene/TemplateSequence/Source/TemplateSequenceEditor/Private/TemplateSequenceEditorToolkit.cpp b/Engine/Plugins/MovieScene/TemplateSequence/Source/TemplateSequenceEditor/Private/TemplateSequenceEditorToolkit.cpp
--- a/Engine/Plugins/MovieScene/TemplateSequence/Source/TemplateSequenceEditor/Private/TemplateSequenceEditorToolkit.cpp
+++ b/Engine/Plugins/MovieScene/TemplateSequence/Source/TemplateSequenceEditor/Private/TemplateSequenceEditorToolkit.cpp
@@ -93,7 +93,6 @@ void FTemplateSequenceEditorToolkit::Initialize(const EToolkitMode::Type Mode, c
 		SequencerInitParams.ToolkitHost = InitToolkitHost;
 		SequencerInitParams.SpawnRegister = SpawnRegister;
 		SequencerInitParams.HostCapabilities.bSupportsCurveEditor = true;
-		SequencerInitParams.HostCapabilities.bSupportsDiscardChanges = true;
 		SequencerInitParams.HostCapabilities.bSupportsSaveMovieSceneAsset = true;
 
 		SequencerInitParams.PlaybackContext.Bind(PlaybackContext.ToSharedRef(), &FTemplateSequenceEditorPlaybackContext::GetPlaybackContext);
diff --git a/Engine/Source/Editor/Sequencer/Private/SSequencer.cpp b/Engine/Source/Editor/Sequencer/Private/SSequencer.cpp
--- a/Engine/Source/Editor/Sequencer/Private/SSequencer.cpp
+++ b/Engine/Source/Editor/Sequencer/Private/SSequencer.cpp
@@ -1330,19 +1330,6 @@ TSharedRef<SWidget> SSequencer::MakeToolBar()
 					LOCTEXT("SaveDirtyPackagesTooltip", "Saves the current sequence and any subsequences"),
 					SaveIcon
 				);
-
-				ToolBarBuilder.AddToolBarButton(
-					FUIAction(FExecuteAction::CreateSP(this, &SSequencer::OnSaveMovieSceneAsClicked)),
-					NAME_None,
-					LOCTEXT("SaveAs", "Save As"),
-					LOCTEXT("SaveAsTooltip", "Saves the current sequence under a different name"),
-					FSlateIcon(FEditorStyle::GetStyleSetName(), "Sequencer.SaveAs")
-				);
-			}
-
-			if (SequencerPtr.Pin()->GetHostCapabilities().bSupportsDiscardChanges)
-			{
-				ToolBarBuilder.AddToolBarButton( FSequencerCommands::Get().DiscardChanges );
 			}
 
 			ToolBarBuilder.AddToolBarButton( FSequencerCommands::Get().FindInContentBrowser );
@@ -2851,12 +2838,6 @@ void SSequencer::OnSaveMovieSceneClicked()
 }
 
 
-void SSequencer::OnSaveMovieSceneAsClicked()
-{
-	SequencerPtr.Pin()->SaveCurrentMovieSceneAs();
-}
-
-
 void SSequencer::StepToNextKey()
 {
 	StepToKey(true, false);
diff --git a/Engine/Source/Editor/Sequencer/Private/Sequencer.cpp b/Engine/Source/Editor/Sequencer/Private/Sequencer.cpp
--- a/Engine/Source/Editor/Sequencer/Private/Sequencer.cpp
+++ b/Engine/Source/Editor/Sequencer/Private/Sequencer.cpp
@@ -5129,36 +5129,6 @@ void FSequencer::SaveCurrentMovieScene()
 }
 
 
-void FSequencer::SaveCurrentMovieSceneAs()
-{
-	if (!GetHostCapabilities().bSupportsSaveMovieSceneAsset)
-	{
-		return;
-	}
-
-	TSharedPtr<IToolkitHost> MyToolkitHost = GetToolkitHost();
-	check(MyToolkitHost);
-
-	TArray<UObject*> AssetsToSave;
-	AssetsToSave.Add(GetCurrentAsset());
-
-	TArray<UObject*> SavedAssets;
-	FEditorFileUtils::SaveAssetsAs(AssetsToSave, SavedAssets);
-
-	if (SavedAssets.Num() == 0)
-	{
-		return;
-	}
-
-	if ((SavedAssets[0] != AssetsToSave[0]) && (SavedAssets[0] != nullptr))
-	{
-		UAssetEditorSubsystem* AssetEditorSubsystem = GEditor->GetEditorSubsystem<UAssetEditorSubsystem>();
-		AssetEditorSubsystem->CloseAllEditorsForAsset(AssetsToSave[0]);
-		AssetEditorSubsystem->OpenEditorForAssets_Advanced(SavedAssets, EToolkitMode::Standalone, MyToolkitHost.ToSharedRef());
-	}
-}
-
-
 TArray<FGuid> FSequencer::AddActors(const TArray<TWeakObjectPtr<AActor> >& InActors, bool bSelectActors)
 {
 	TArray<FGuid> PossessableGuids;
@@ -9315,127 +9285,6 @@ void FSequencer::OnClipboardUsed(TSharedPtr<FMovieSceneClipboard> Clipboard)
 	});
 }
 
-
-void FSequencer::DiscardChanges()
-{
-	if (ActiveTemplateIDs.Num() == 0)
-	{
-		return;
-	}
-
-	TSharedPtr<IToolkitHost> MyToolkitHost = GetToolkitHost();
-
-	if (!MyToolkitHost.IsValid())
-	{
-		return;
-	}
-
-	UMovieSceneSequence* EditedSequence = GetFocusedMovieSceneSequence();
-
-	if (EditedSequence == nullptr)
-	{
-		return;
-	}
-
-	if (FMessageDialog::Open(EAppMsgType::YesNo, LOCTEXT("RevertConfirm", "Are you sure you want to discard your current changes?")) != EAppReturnType::Yes)
-	{
-		return;
-	}
-
-	UAssetEditorSubsystem* AssetEditorSubsystem = GEditor->GetEditorSubsystem<UAssetEditorSubsystem>();
-	UClass* SequenceClass = EditedSequence->GetClass();
-	FString SequencePath = EditedSequence->GetPathName();
-	UPackage* SequencePackage = EditedSequence->GetOutermost();
-
-	// close asset editor
-	AssetEditorSubsystem->CloseAllEditorsForAsset(EditedSequence);
-
-	// collect objects to be unloaded
-	TMap<FString, UObject*> MovedObjects;
-
-	ForEachObjectWithOuter(SequencePackage, [&](UObject* Object) {
-		MovedObjects.Add(Object->GetPathName(), Object);
-	}, true);
-
-	// move objects into transient package
-	UPackage* const TransientPackage = GetTransientPackage();
-
-	for (auto MovedObject : MovedObjects)
-	{
-		UObject* Object = MovedObject.Value;
-
-		const FString OldName = Object->GetName();
-		const FString NewName = FString::Printf(TEXT("UNLOADING_%s"), *OldName);
-		const FName UniqueName = MakeUniqueObjectName(TransientPackage, Object->GetClass(), FName(*NewName));
-		UObject* NewOuter = (Object->GetOuter() == SequencePackage) ? TransientPackage : Object->GetOuter();
-	
-		Object->Rename(*UniqueName.ToString(), NewOuter, REN_DontCreateRedirectors | REN_DoNotDirty | REN_NonTransactional);
-		Object->SetFlags(RF_Transient);
-		Object->ClearFlags(RF_Standalone | RF_Transactional);
-	}
-
-	for (auto MovedObject : MovedObjects)
-	{
-		GLog->Logf(TEXT("Moved %s ---------> %s"), *MovedObject.Key, *MovedObject.Value->GetPathName());
-	}
-
-	// unload package
-	SequencePackage->SetDirtyFlag(false);
-
-	TArray<UPackage*> PackagesToUnload;
-	PackagesToUnload.Add(SequencePackage);
-
-	FText PackageUnloadError;
-	UPackageTools::UnloadPackages(PackagesToUnload, PackageUnloadError);
-
-	if (!PackageUnloadError.IsEmpty())
-	{
-		ResetLoaders(SequencePackage);
-		SequencePackage->ClearFlags(RF_WasLoaded);
-		SequencePackage->bHasBeenFullyLoaded = false;
-		SequencePackage->GetMetaData()->RemoveMetaDataOutsidePackage();
-	}
-
-	// reload package
-	TMap<UObject*, UObject*> MovedToReloadedObjectMap;
-
-	for (const auto MovedObject : MovedObjects)
-	{
-		UObject* ReloadedObject = StaticLoadObject(MovedObject.Value->GetClass(), nullptr, *MovedObject.Key, nullptr);//, LOAD_NoWarn);
-		MovedToReloadedObjectMap.Add(MovedObject.Value, ReloadedObject);
-	}
-
-	for (TObjectIterator<UObject> It; It; ++It)
-	{
-		// @todo sequencer: only process objects that actually reference the package?
-		FArchiveReplaceObjectRef<UObject> Ar(*It, MovedToReloadedObjectMap, false, false, false, false);
-	}
-
-	auto ReloadedSequence = Cast<UMovieSceneSequence>(StaticLoadObject(SequenceClass, nullptr, *SequencePath, nullptr));//, LOAD_NoWarn));
-
-	// release transient objects
-	for (auto MovedObject : MovedObjects)
-	{
-		MovedObject.Value->RemoveFromRoot();
-		MovedObject.Value->MarkPendingKill();
-	}
-
-//	CollectGarbage(GARBAGE_COLLECTION_KEEPFLAGS);
-
-	// clear undo buffer
-	if (true) // @todo sequencer: check whether objects are actually referenced in undo buffer
-	{
-		GEditor->Trans->Reset(LOCTEXT("UnloadedSequence", "Unloaded Sequence"));
-	}
-
-	// reopen asset editor
-	TArray<UObject*> AssetsToReopen;
-	AssetsToReopen.Add(ReloadedSequence);
-
-	AssetEditorSubsystem->OpenEditorForAssets_Advanced(AssetsToReopen, EToolkitMode::Standalone, MyToolkitHost.ToSharedRef());
-}
-
-
 void FSequencer::CreateCamera()
 {
 	UMovieScene* FocusedMovieScene = GetFocusedMovieSceneSequence()->GetMovieScene();
@@ -10530,22 +10379,6 @@ void FSequencer::BindCommands()
 		FIsActionButtonVisible::CreateLambda([this] { return ExactCast<ULevelSequence>(GetFocusedMovieSceneSequence()) != nullptr && IVREditorModule::Get().IsVREditorModeActive() == false; }) //@todo VREditor: Creating a camera while in VR mode disrupts the hmd. This is a temporary fix by hiding the button when in VR mode.
 	);
 
-	SequencerCommandBindings->MapAction(
-		Commands.DiscardChanges,
-		FExecuteAction::CreateSP(this, &FSequencer::DiscardChanges),
-		FCanExecuteAction::CreateLambda([this]{
-			UMovieSceneSequence* EditedSequence = GetFocusedMovieSceneSequence();
-			if (!EditedSequence)
-			{
-				return false;
-			}
-
-			UPackage* EditedPackage = EditedSequence->GetOutermost();
-
-			return ((EditedPackage->FileSize != 0) && EditedPackage->IsDirty());
-		})
-	);
-
 	SequencerCommandBindings->MapAction(
 		Commands.BakeTransform,
 		FExecuteAction::CreateSP( this, &FSequencer::BakeTransform ),
diff --git a/Engine/Source/Editor/Sequencer/Private/Sequencer.h b/Engine/Source/Editor/Sequencer/Private/Sequencer.h
--- a/Engine/Source/Editor/Sequencer/Private/Sequencer.h
+++ b/Engine/Source/Editor/Sequencer/Private/Sequencer.h
@@ -453,9 +453,6 @@ public:
 	/** Called to save the current movie scene */
 	void SaveCurrentMovieScene();
 
-	/** Called to save the current movie scene under a new name */
-	void SaveCurrentMovieSceneAs();
-
 	/** Called when a user executes the assign actor to track menu item */
 	void AssignActor(FMenuBuilder& MenuBuilder, FGuid ObjectBinding);
 	FGuid DoAssignActor(AActor*const* InActors, int32 NumActors, FGuid ObjectBinding);
@@ -581,9 +578,6 @@ public:
 	/** Promote a clipboard to the top of the clipboard stack, and update its timestamp */
 	void OnClipboardUsed(TSharedPtr<FMovieSceneClipboard> Clipboard);
 
-	/** Discard all changes to the current movie scene. */
-	void DiscardChanges();
-
 	/** Create camera and set it as the current camera cut. */
 	void CreateCamera();
 
diff --git a/Engine/Source/Editor/Sequencer/Private/SequencerCommands.cpp b/Engine/Source/Editor/Sequencer/Private/SequencerCommands.cpp
--- a/Engine/Source/Editor/Sequencer/Private/SequencerCommands.cpp
+++ b/Engine/Source/Editor/Sequencer/Private/SequencerCommands.cpp
@@ -136,8 +136,6 @@ void FSequencerCommands::RegisterCommands()
 	UI_COMMAND( SaveCurrentSpawnableState, "Save Default State", "Save the current state of this spawnable object as its default properties.", EUserInterfaceActionType::Button, FInputChord() );
 	UI_COMMAND( RestoreAnimatedState, "Restore Animated State", "Restore any objects that have been animated by sequencer back to their original state.", EUserInterfaceActionType::Button, FInputChord(EModifierKey::Control, EKeys::R) );
 
-	UI_COMMAND( DiscardChanges, "Discard All Changes", "Revert the currently edited movie scene to its last saved state.", EUserInterfaceActionType::Button, FInputChord());
-
 	UI_COMMAND( FixActorReferences, "Fix Actor References", "Try to automatically fix up broken actor bindings.", EUserInterfaceActionType::Button, FInputChord() );
 	UI_COMMAND( RebindPossessableReferences, "Rebind Possesable References", "Rebinds all possessables in the current sequence to ensure they're using the most robust referencing mechanism.", EUserInterfaceActionType::Button, FInputChord() );
 
diff --git a/Engine/Source/Editor/Sequencer/Public/ISequencerModule.h b/Engine/Source/Editor/Sequencer/Public/ISequencerModule.h
--- a/Engine/Source/Editor/Sequencer/Public/ISequencerModule.h
+++ b/Engine/Source/Editor/Sequencer/Public/ISequencerModule.h
@@ -130,15 +130,11 @@ struct FSequencerHostCapabilities
 	/** Should we show the Save-As button in the toolbar? */
 	bool bSupportsSaveMovieSceneAsset;
 
-	/** Do we support discarding the changes */
-	bool bSupportsDiscardChanges;
-
 	/** Do we support the curve editor */
 	bool bSupportsCurveEditor;
 
 	FSequencerHostCapabilities()
 		: bSupportsSaveMovieSceneAsset(false)
-		, bSupportsDiscardChanges(false)
 		, bSupportsCurveEditor(false)
 	{}
 };
diff --git a/Engine/Source/Editor/Sequencer/Public/SequencerCommands.h b/Engine/Source/Editor/Sequencer/Public/SequencerCommands.h
--- a/Engine/Source/Editor/Sequencer/Public/SequencerCommands.h
+++ b/Engine/Source/Editor/Sequencer/Public/SequencerCommands.h
@@ -331,9 +331,6 @@ public:
 	/** Restores all animated state for the current sequence. */
 	TSharedPtr< FUICommandInfo > RestoreAnimatedState;
 
-	/** Discard all changes to the current movie scene. */
-	TSharedPtr< FUICommandInfo > DiscardChanges;
-
 	/** Attempts to fix broken actor references. */
 	TSharedPtr< FUICommandInfo > FixActorReferences;
 
