Sequencer: Call BeginSlipSection correctly.

diff --git a/Engine/Source/Editor/Sequencer/Private/Tools/EditToolDragOperations.cpp b/Engine/Source/Editor/Sequencer/Private/Tools/EditToolDragOperations.cpp
--- a/Engine/Source/Editor/Sequencer/Private/Tools/EditToolDragOperations.cpp
+++ b/Engine/Source/Editor/Sequencer/Private/Tools/EditToolDragOperations.cpp
@@ -174,7 +174,14 @@ void FResizeSection::OnBeginDrag(const FPointerEvent& MouseEvent, FVector2D Loca
 		}
 		else if (TOptional<FSectionHandle> SectionHandle = Sequencer.GetNodeTree()->GetSectionHandle(Section))
 		{
-			SectionHandle->GetSectionInterface()->BeginResizeSection();
+			if (bIsSlipping)
+			{
+				SectionHandle->GetSectionInterface()->BeginSlipSection();
+			}
+			else
+			{
+				SectionHandle->GetSectionInterface()->BeginResizeSection();
+			}
 		}
 
 		SectionInitTimes.Add(Section, bDraggingByEnd ? Section->GetExclusiveEndFrame() : Section->GetInclusiveStartFrame());
