Sequencer: Playback End snap location is no longer excluding the last tick twice, which caused snapping to create non-whole-frame aligned sections even with snapping enabled.

diff --git a/Engine/Source/Editor/Sequencer/Private/Tools/SequencerSnapField.cpp b/Engine/Source/Editor/Sequencer/Private/Tools/SequencerSnapField.cpp
--- a/Engine/Source/Editor/Sequencer/Private/Tools/SequencerSnapField.cpp
+++ b/Engine/Source/Editor/Sequencer/Private/Tools/SequencerSnapField.cpp
@@ -76,7 +76,7 @@ FSequencerSnapField::FSequencerSnapField(const ISequencer& InSequencer, ISequenc
 	if(MovieScene::DiscreteSize(PlaybackRange) > 0)
 	{
 		Visitor.Snaps.Add(FSequencerSnapPoint{ FSequencerSnapPoint::PlaybackRange, MovieScene::DiscreteInclusiveLower(PlaybackRange)});
-		Visitor.Snaps.Add(FSequencerSnapPoint{ FSequencerSnapPoint::PlaybackRange, MovieScene::DiscreteExclusiveUpper(PlaybackRange) - 1});
+		Visitor.Snaps.Add(FSequencerSnapPoint{ FSequencerSnapPoint::PlaybackRange, MovieScene::DiscreteExclusiveUpper(PlaybackRange)});
 	}
 
 	// Add the current time as a potential snap candidate
