Sequencer: When only filtering by text string, include organizational nodes in the search

diff --git a/Engine/Source/Editor/Sequencer/Private/SequencerNodeTree.cpp b/Engine/Source/Editor/Sequencer/Private/SequencerNodeTree.cpp
--- a/Engine/Source/Editor/Sequencer/Private/SequencerNodeTree.cpp
+++ b/Engine/Source/Editor/Sequencer/Private/SequencerNodeTree.cpp
@@ -920,92 +920,109 @@ static void FilterNodesRecursive( FSequencer& Sequencer, const TSharedRef<FSeque
 
 	bool bPasssedAnyFilters = false;
 
-	if (StartNode->GetType() == ESequencerNode::Track)
+	switch(StartNode->GetType())
 	{
-		UMovieSceneTrack* Track = static_cast<const FSequencerTrackNode&>(StartNode.Get()).GetTrack();
-		if (Filters->Num() == 0 || Filters->PassesAnyFilters(Track))
+		case ESequencerNode::Track:
 		{
-			bPasssedAnyFilters = true;
-
-			// Track nodes do not belong to a level, but might be a child of an objectbinding node that does
-			if (LevelTrackFilter->IsActive())
+			UMovieSceneTrack* Track = static_cast<const FSequencerTrackNode&>(StartNode.Get()).GetTrack();
+			if (Filters->Num() == 0 || Filters->PassesAnyFilters(Track))
 			{
-				TSharedPtr<const FSequencerDisplayNode> ParentNode = StartNode->GetParent();
-				while (ParentNode.IsValid())
+				bPasssedAnyFilters = true;
+
+				// Track nodes do not belong to a level, but might be a child of an objectbinding node that does
+				if (LevelTrackFilter->IsActive())
 				{
-					if (ParentNode->GetType() == ESequencerNode::Object)
+					TSharedPtr<const FSequencerDisplayNode> ParentNode = StartNode->GetParent();
+					while (ParentNode.IsValid())
 					{
-						// The track belongs to an objectbinding node, start by assuming it doesn't match the level filter
-						bPasssedAnyFilters = false;
-
-						const FSequencerObjectBindingNode* ObjectNode = static_cast<const FSequencerObjectBindingNode*>(ParentNode.Get());
-						for (TWeakObjectPtr<>& Object : Sequencer.FindObjectsInCurrentSequence(ObjectNode->GetObjectBinding()))
+						if (ParentNode->GetType() == ESequencerNode::Object)
 						{
-							if (Object.IsValid() && LevelTrackFilter->PassesFilter(Object.Get()))
+							// The track belongs to an objectbinding node, start by assuming it doesn't match the level filter
+							bPasssedAnyFilters = false;
+
+							const FSequencerObjectBindingNode* ObjectNode = static_cast<const FSequencerObjectBindingNode*>(ParentNode.Get());
+							for (TWeakObjectPtr<>& Object : Sequencer.FindObjectsInCurrentSequence(ObjectNode->GetObjectBinding()))
 							{
-								// If at least one of the objects on the objectbinding node pass the level filter, show the track
-								bPasssedAnyFilters = true;
-								break;
+								if (Object.IsValid() && LevelTrackFilter->PassesFilter(Object.Get()))
+								{
+									// If at least one of the objects on the objectbinding node pass the level filter, show the track
+									bPasssedAnyFilters = true;
+									break;
+								}
 							}
-						}
 
-						break;
+							break;
+						}
+						ParentNode = ParentNode->GetParent();
 					}
-					ParentNode = ParentNode->GetParent();
 				}
-			}
 
-			if (bPasssedAnyFilters && Sequencer.GetSequencerSettings()->GetShowSelectedNodesOnly())
-			{
-				TSharedPtr<const FSequencerDisplayNode> ParentNode = StartNode->GetParent();
-				while (ParentNode.IsValid())
+				if (bPasssedAnyFilters && Sequencer.GetSequencerSettings()->GetShowSelectedNodesOnly())
 				{
-					// Pinned tracks should be visible whether selected or not
-					if (ParentNode->IsPinned())
+					TSharedPtr<const FSequencerDisplayNode> ParentNode = StartNode->GetParent();
+					while (ParentNode.IsValid())
 					{
-						break;
-					}
-
-					if (ParentNode->GetType() == ESequencerNode::Object)
-					{
-						const FSequencerObjectBindingNode* ObjectNode = static_cast<const FSequencerObjectBindingNode*>(ParentNode.Get());
-						const FMovieSceneBinding* Binding = Sequencer.GetFocusedMovieSceneSequence()->GetMovieScene()->FindBinding(ObjectNode->GetObjectBinding());
-						if (!(Binding && Sequencer.IsBindingVisible(*Binding)))
+						// Pinned tracks should be visible whether selected or not
+						if (ParentNode->IsPinned())
 						{
-							bPasssedAnyFilters = false;
+							break;
 						}
 
-						break;
+						if (ParentNode->GetType() == ESequencerNode::Object)
+						{
+							const FSequencerObjectBindingNode* ObjectNode = static_cast<const FSequencerObjectBindingNode*>(ParentNode.Get());
+							const FMovieSceneBinding* Binding = Sequencer.GetFocusedMovieSceneSequence()->GetMovieScene()->FindBinding(ObjectNode->GetObjectBinding());
+							if (!(Binding && Sequencer.IsBindingVisible(*Binding)))
+							{
+								bPasssedAnyFilters = false;
+							}
+
+							break;
+						}
+						ParentNode = ParentNode->GetParent();
 					}
-					ParentNode = ParentNode->GetParent();
 				}
 			}
+			break;
 		}
-	}
-	else if (StartNode->GetType() == ESequencerNode::Object)
-	{
-		const FSequencerObjectBindingNode ObjectNode = static_cast<const FSequencerObjectBindingNode&>(StartNode.Get());
-		for (TWeakObjectPtr<>& Object : Sequencer.FindObjectsInCurrentSequence(ObjectNode.GetObjectBinding()))
+		case ESequencerNode::Object:
 		{
-			if (Object.IsValid() && (Filters->Num() == 0 || Filters->PassesAnyFilters(Object.Get()))
-				&& LevelTrackFilter->PassesFilter(Object.Get()))
+			const FSequencerObjectBindingNode ObjectNode = static_cast<const FSequencerObjectBindingNode&>(StartNode.Get());
+			for (TWeakObjectPtr<>& Object : Sequencer.FindObjectsInCurrentSequence(ObjectNode.GetObjectBinding()))
 			{
-				bPasssedAnyFilters = true;
-				break;
+				if (Object.IsValid() && (Filters->Num() == 0 || Filters->PassesAnyFilters(Object.Get()))
+					&& LevelTrackFilter->PassesFilter(Object.Get()))
+				{
+					bPasssedAnyFilters = true;
+					break;
+				}
 			}
-		}
 
-		if (bPasssedAnyFilters && Sequencer.GetSequencerSettings()->GetShowSelectedNodesOnly() && !StartNode->IsPinned())
+			if (bPasssedAnyFilters && Sequencer.GetSequencerSettings()->GetShowSelectedNodesOnly() && !StartNode->IsPinned())
+			{
+				UMovieScene* MovieScene = Sequencer.GetFocusedMovieSceneSequence()->GetMovieScene();
+				const FMovieSceneBinding* Binding = Sequencer.GetFocusedMovieSceneSequence()->GetMovieScene()->FindBinding(ObjectNode.GetObjectBinding());
+				if (Binding && !Sequencer.IsBindingVisible(*Binding))
+				{
+					bPasssedAnyFilters = false;
+				}
+			}
+			break;
+		}
+		case ESequencerNode::Category:
+		case ESequencerNode::Folder:
+		case ESequencerNode::KeyArea:
 		{
-			UMovieScene* MovieScene = Sequencer.GetFocusedMovieSceneSequence()->GetMovieScene();
-			const FMovieSceneBinding* Binding = Sequencer.GetFocusedMovieSceneSequence()->GetMovieScene()->FindBinding(ObjectNode.GetObjectBinding());
-			if (Binding && !Sequencer.IsBindingVisible(*Binding))
+			// If no active filters, pass organizational nodes to be text filtered
+			if (Filters->Num() == 0)
 			{
-				bPasssedAnyFilters = false;
+				bPasssedAnyFilters = true;
 			}
+			break;
 		}
 	}
 
+
 	if (bPasssedAnyFilters)
 	{
 		// If we have a filter string, make sure we match
