Sequencer: Redraw viewports after restoring state on close(). This fixes a bug where if Realtime is not enabled in the viewport, closing sequencer would leave objects drawn in a sequencer animated state before the next redraw.

diff --git a/Engine/Source/Editor/Sequencer/Private/Sequencer.cpp b/Engine/Source/Editor/Sequencer/Private/Sequencer.cpp
--- a/Engine/Source/Editor/Sequencer/Private/Sequencer.cpp
+++ b/Engine/Source/Editor/Sequencer/Private/Sequencer.cpp
@@ -557,7 +557,10 @@ void FSequencer::Close()
 	TrackEditors.Empty();
 
 	GUnrealEd->UpdatePivotLocationForSelection();
-	
+
+	// Redraw viewports after restoring pre animated state in case viewports are not set to realtime
+	GEditor->RedrawLevelEditingViewports();
+
 	OnCloseEventDelegate.Broadcast(AsShared());
 }
 
