Sequencer: Fix transform selected sections so that keys are also transformed

diff --git a/Engine/Source/Editor/Sequencer/Private/Sequencer.cpp b/Engine/Source/Editor/Sequencer/Private/Sequencer.cpp
--- a/Engine/Source/Editor/Sequencer/Private/Sequencer.cpp
+++ b/Engine/Source/Editor/Sequencer/Private/Sequencer.cpp
@@ -1502,6 +1502,12 @@ void FSequencer::TransformSelectedKeysAndSections(FFrameTime InDeltaTime, float
 			FMovieSceneChannel* Channel = ChannelInfo.Channel.Get();
 			if (Channel)
 			{
+				// Skip any channels whose section is already selected because they'll be handled below (moving the section and the keys together)
+				if (SelectedSectionsArray.Contains(ChannelInfo.OwningSection))
+				{
+					continue;
+				}
+
 				const int32 NumKeys = ChannelInfo.KeyHandles.Num();
 				KeyTimesScratch.Reset(NumKeys);
 				KeyTimesScratch.SetNum(NumKeys);
@@ -1584,6 +1590,26 @@ void FSequencer::TransformSelectedKeysAndSections(FFrameTime InDeltaTime, float
 			// overwrite the (possibly) existing bound, so it's okay to just overwrite the range without a TRange::Hull.
 			*NewSectionBounds = TRange<FFrameNumber>(LowerBound, UpperBound);
 			bAnythingChanged = true;
+
+			// Modify all of the keys of this section
+			for (const FMovieSceneChannelEntry& Entry : Section->GetChannelProxy().GetAllEntries())
+			{
+				for (FMovieSceneChannel* Channel : Entry.GetChannels())
+				{
+					TArray<FFrameNumber> KeyTimes;
+					TArray<FKeyHandle> KeyHandles;
+					TArray<FFrameNumber> NewKeyTimes;
+					Channel->GetKeys(TRange<FFrameNumber>::All(), &KeyTimes, &KeyHandles);
+
+					for (FFrameNumber KeyTime : KeyTimes)
+					{
+						FFrameNumber NewKeyTime = (OriginTime + InDeltaTime + (KeyTime - OriginTime) * InScale).FloorToFrame();
+						NewKeyTimes.Add(NewKeyTime);
+					}
+
+					Channel->SetKeyTimes(KeyHandles, NewKeyTimes);
+				}
+			}
 		}
 	}
 	
