Sequencer: Fix copying a spawnable's object template so that it's not limited to the given transient outer. This fixes a bug where the DefaultSceneRoot for an empty actor would be lost on a copy/paste.

diff --git a/Engine/Source/Editor/Sequencer/Private/Sequencer.cpp b/Engine/Source/Editor/Sequencer/Private/Sequencer.cpp
--- a/Engine/Source/Editor/Sequencer/Private/Sequencer.cpp
+++ b/Engine/Source/Editor/Sequencer/Private/Sequencer.cpp
@@ -6903,7 +6903,7 @@ void ExportObjectBindingsToText(const TArray<UMovieSceneCopyableBinding*>& Objec
 		// We manually export the object template for the same private-ownership reason as above. Templates need to be re-created anyways as each Spawnable contains its own copy of the template.
 		if (ObjectToExport->SpawnableObjectTemplate)
 		{
-			UExporter::ExportToOutputDevice(&Context, ObjectToExport->SpawnableObjectTemplate, nullptr, Archive, TEXT("copy"), 0, PPF_ExportsNotFullyQualified | PPF_Copy | PPF_Delimited, false, ThisOuter);
+			UExporter::ExportToOutputDevice(&Context, ObjectToExport->SpawnableObjectTemplate, nullptr, Archive, TEXT("copy"), 0, PPF_ExportsNotFullyQualified | PPF_Copy | PPF_Delimited);
 		}
 	}
 
