Sequencer: Fix crash when baking an actor component

diff --git a/Engine/Source/Editor/Sequencer/Private/Sequencer.cpp b/Engine/Source/Editor/Sequencer/Private/Sequencer.cpp
--- a/Engine/Source/Editor/Sequencer/Private/Sequencer.cpp
+++ b/Engine/Source/Editor/Sequencer/Private/Sequencer.cpp
@@ -1697,6 +1697,19 @@ void FSequencer::BakeTransform()
 		for (auto RuntimeObject : FindBoundObjects(Guid, ActiveTemplateIDs.Top()) )
 		{
 			AActor* Actor = Cast<AActor>(RuntimeObject.Get());
+			if (!Actor)
+			{
+				UActorComponent* ActorComponent = Cast<UActorComponent>(RuntimeObject.Get());
+				if (ActorComponent)
+				{
+					Actor = ActorComponent->GetOwner();
+				}
+			}
+
+			if (!Actor)
+			{
+				continue;
+			}
 
 			UCameraComponent* CameraComponent = MovieSceneHelpers::CameraComponentFromRuntimeObject(RuntimeObject.Get());
 
