Sequencer: Fix removing marked frames so that it's based on the hit index, rather than finding it by the FFrameNumber. renamed "Clear Marked Frame" to "Delete Marked Frame"

diff --git a/Engine/Plugins/MovieScene/SequencerScripting/Source/SequencerScripting/Private/ExtensionLibraries/MovieSceneSequenceExtensions.cpp b/Engine/Plugins/MovieScene/SequencerScripting/Source/SequencerScripting/Private/ExtensionLibraries/MovieSceneSequenceExtensions.cpp
--- a/Engine/Plugins/MovieScene/SequencerScripting/Source/SequencerScripting/Private/ExtensionLibraries/MovieSceneSequenceExtensions.cpp
+++ b/Engine/Plugins/MovieScene/SequencerScripting/Source/SequencerScripting/Private/ExtensionLibraries/MovieSceneSequenceExtensions.cpp
@@ -547,21 +547,21 @@ int32 UMovieSceneSequenceExtensions::AddMarkedFrame(UMovieSceneSequence* Sequenc
 	return INDEX_NONE;
 }
 
-void UMovieSceneSequenceExtensions::RemoveMarkedFrame(UMovieSceneSequence* Sequence, int32 RemoveIndex)
+void UMovieSceneSequenceExtensions::DeleteMarkedFrame(UMovieSceneSequence* Sequence, int32 DeleteIndex)
 {
 	UMovieScene* MovieScene = Sequence->GetMovieScene();
 	if (MovieScene)
 	{
-		MovieScene->RemoveMarkedFrame(RemoveIndex);
+		MovieScene->DeleteMarkedFrame(DeleteIndex);
 	}
 }
 
-void UMovieSceneSequenceExtensions::ClearMarkedFrames(UMovieSceneSequence* Sequence)
+void UMovieSceneSequenceExtensions::DeleteMarkedFrames(UMovieSceneSequence* Sequence)
 {
 	UMovieScene* MovieScene = Sequence->GetMovieScene();
 	if (MovieScene)
 	{
-		MovieScene->ClearMarkedFrames();
+		MovieScene->DeleteMarkedFrames();
 	}
 }
 
diff --git a/Engine/Plugins/MovieScene/SequencerScripting/Source/SequencerScripting/Public/ExtensionLibraries/MovieSceneSequenceExtensions.h b/Engine/Plugins/MovieScene/SequencerScripting/Source/SequencerScripting/Public/ExtensionLibraries/MovieSceneSequenceExtensions.h
--- a/Engine/Plugins/MovieScene/SequencerScripting/Source/SequencerScripting/Public/ExtensionLibraries/MovieSceneSequenceExtensions.h
+++ b/Engine/Plugins/MovieScene/SequencerScripting/Source/SequencerScripting/Public/ExtensionLibraries/MovieSceneSequenceExtensions.h
@@ -422,18 +422,18 @@ public:
 	static int32 AddMarkedFrame(UMovieSceneSequence* Sequence, const FMovieSceneMarkedFrame& InMarkedFrame);
 
 	/*
-	 * Remove the user marked frame by index.
+	 * Delete the user marked frame by index.
 	 *
-	 * @RemoveIndex The index to the user marked frame to remove
+	 * @DeleteIndex The index to the user marked frame to delete
 	 */
 	UFUNCTION(BlueprintCallable, Category = "Sequence", meta = (ScriptMethod))
-	static void RemoveMarkedFrame(UMovieSceneSequence* Sequence, int32 RemoveIndex);
+	static void DeleteMarkedFrame(UMovieSceneSequence* Sequence, int32 DeleteIndex);
 
 	/*
-	 * Clear all user marked frames
+	 * Delete all user marked frames
 	 */
 	UFUNCTION(BlueprintCallable, Category = "Sequence", meta = (ScriptMethod))
-	static void ClearMarkedFrames(UMovieSceneSequence* Sequence);
+	static void DeleteMarkedFrames(UMovieSceneSequence* Sequence);
 
 	/*
 	 * Find the user marked frame by label
diff --git a/Engine/Source/Editor/Sequencer/Private/SSequencer.cpp b/Engine/Source/Editor/Sequencer/Private/SSequencer.cpp
--- a/Engine/Source/Editor/Sequencer/Private/SSequencer.cpp
+++ b/Engine/Source/Editor/Sequencer/Private/SSequencer.cpp
@@ -363,8 +363,9 @@ void SSequencer::Construct(const FArguments& InArgs, TSharedRef<FSequencer> InSe
 		TimeSliderArgs.VerticalFrames = InArgs._VerticalFrames;
 		TimeSliderArgs.MarkedFrames = InArgs._MarkedFrames;
 		TimeSliderArgs.OnSetMarkedFrame = InArgs._OnSetMarkedFrame;
-		TimeSliderArgs.OnMarkedFrameChanged = InArgs._OnMarkedFrameChanged;
-		TimeSliderArgs.OnClearAllMarkedFrames = InArgs._OnClearAllMarkedFrames;
+		TimeSliderArgs.OnAddMarkedFrame = InArgs._OnAddMarkedFrame;
+		TimeSliderArgs.OnDeleteMarkedFrame = InArgs._OnDeleteMarkedFrame;
+		TimeSliderArgs.OnDeleteAllMarkedFrames = InArgs._OnDeleteAllMarkedFrames;
 
 		TimeSliderArgs.Settings = Settings;
 		TimeSliderArgs.NumericTypeInterface = GetNumericTypeInterface();
diff --git a/Engine/Source/Editor/Sequencer/Private/SSequencer.h b/Engine/Source/Editor/Sequencer/Private/SSequencer.h
--- a/Engine/Source/Editor/Sequencer/Private/SSequencer.h
+++ b/Engine/Source/Editor/Sequencer/Private/SSequencer.h
@@ -194,11 +194,14 @@ public:
 		/** Called when the user sets a marked frame */
 		SLATE_EVENT(FOnSetMarkedFrame, OnSetMarkedFrame)
 
-		/** Called when the user changes on the set of marked frames */
-		SLATE_EVENT(FOnMarkedFrameChanged, OnMarkedFrameChanged)
+		/** Called when the user adds a marked frame */
+		SLATE_EVENT(FOnAddMarkedFrame, OnAddMarkedFrame)
 
-		/** Called when all marked frames should be cleared */
-		SLATE_EVENT( FSimpleDelegate, OnClearAllMarkedFrames)
+		/** Called when the user deletes a marked frame */
+		SLATE_EVENT(FOnDeleteMarkedFrame, OnDeleteMarkedFrame)
+
+		/** Called when all marked frames should be deleted */
+		SLATE_EVENT( FSimpleDelegate, OnDeleteAllMarkedFrames)
 
 		/** Called when the user changes the clamp range */
 		SLATE_EVENT( FOnTimeRangeChanged, OnClampRangeChanged )
diff --git a/Engine/Source/Editor/Sequencer/Private/Sequencer.cpp b/Engine/Source/Editor/Sequencer/Private/Sequencer.cpp
--- a/Engine/Source/Editor/Sequencer/Private/Sequencer.cpp
+++ b/Engine/Source/Editor/Sequencer/Private/Sequencer.cpp
@@ -410,8 +410,9 @@ void FSequencer::InitSequencer(const FSequencerInitParams& InitParams, const TSh
 		.VerticalFrames(this, &FSequencer::GetVerticalFrames)
 		.MarkedFrames(this, &FSequencer::GetMarkedFrames)
 		.OnSetMarkedFrame(this, &FSequencer::SetMarkedFrame)
-		.OnMarkedFrameChanged(this, &FSequencer::OnMarkedFrameChanged)
-		.OnClearAllMarkedFrames(this, &FSequencer::ClearAllMarkedFrames )
+		.OnAddMarkedFrame(this, &FSequencer::AddMarkedFrame)
+		.OnDeleteMarkedFrame(this, &FSequencer::DeleteMarkedFrame)
+		.OnDeleteAllMarkedFrames(this, &FSequencer::DeleteAllMarkedFrames )
 		.SubSequenceRange( this, &FSequencer::GetSubSequenceRange )
 		.OnPlaybackRangeChanged( this, &FSequencer::SetPlaybackRange )
 		.OnPlaybackRangeBeginDrag( this, &FSequencer::OnPlaybackRangeBeginDrag )
@@ -9168,10 +9169,10 @@ void FSequencer::ToggleMarkAtPlayPosition()
 			int32 MarkedFrameIndex = FocusedMovieScene->FindMarkedFrameByFrameNumber(TickFrameNumber);
 			if (MarkedFrameIndex != INDEX_NONE)
 			{
-				FScopedTransaction RemoveMarkedFrameTransaction(LOCTEXT("RemoveMarkedFrames_Transaction", "Remove Marked Frame"));
+				FScopedTransaction DeleteMarkedFrameTransaction(LOCTEXT("DeleteMarkedFrames_Transaction", "Delete Marked Frame"));
 
 				FocusedMovieScene->Modify();
-				FocusedMovieScene->RemoveMarkedFrame(MarkedFrameIndex);
+				FocusedMovieScene->DeleteMarkedFrame(MarkedFrameIndex);
 			}
 			else
 			{
@@ -9198,7 +9199,7 @@ void FSequencer::SetMarkedFrame(int32 InMarkIndex, FFrameNumber InFrameNumber)
 	}
 }
 
-void FSequencer::OnMarkedFrameChanged(FFrameNumber FrameNumber, bool bSetMark)
+void FSequencer::AddMarkedFrame(FFrameNumber FrameNumber)
 {
 	UMovieSceneSequence* FocusedMovieSequence = GetFocusedMovieSceneSequence();
 	if (FocusedMovieSequence != nullptr)
@@ -9206,29 +9207,34 @@ void FSequencer::OnMarkedFrameChanged(FFrameNumber FrameNumber, bool bSetMark)
 		UMovieScene* FocusedMovieScene = FocusedMovieSequence->GetMovieScene();
 		if (FocusedMovieScene != nullptr)
 		{
-			if (bSetMark)
-			{
-				FScopedTransaction AddMarkedFrameTransaction(LOCTEXT("AddMarkedFrame_Transaction", "Add Marked Frame"));
+			FScopedTransaction AddMarkedFrameTransaction(LOCTEXT("AddMarkedFrame_Transaction", "Add Marked Frame"));
 
-				FocusedMovieScene->Modify();
-				FocusedMovieScene->AddMarkedFrame(FMovieSceneMarkedFrame(FrameNumber));
-			}
-			else
+			FocusedMovieScene->Modify();
+			FocusedMovieScene->AddMarkedFrame(FMovieSceneMarkedFrame(FrameNumber));
+		}
+	}
+}
+
+void FSequencer::DeleteMarkedFrame(int32 InMarkIndex)
+{
+	UMovieSceneSequence* FocusedMovieSequence = GetFocusedMovieSceneSequence();
+	if (FocusedMovieSequence != nullptr)
+	{
+		UMovieScene* FocusedMovieScene = FocusedMovieSequence->GetMovieScene();
+		if (FocusedMovieScene != nullptr)
+		{
+			if (InMarkIndex != INDEX_NONE)
 			{
-				int32 MarkedFrameIndex = FocusedMovieScene->FindMarkedFrameByFrameNumber(FrameNumber);
-				if (MarkedFrameIndex != INDEX_NONE)
-				{
-					FScopedTransaction RemoveMarkedFrameTransaction(LOCTEXT("RemoveMarkedFrame_Transaction", "Remove Marked Frame"));
+				FScopedTransaction DeleteMarkedFrameTransaction(LOCTEXT("DeleteMarkedFrame_Transaction", "Delete Marked Frame"));
 
-					FocusedMovieScene->Modify();
-					FocusedMovieScene->RemoveMarkedFrame(MarkedFrameIndex);
-				}
+				FocusedMovieScene->Modify();
+				FocusedMovieScene->DeleteMarkedFrame(InMarkIndex);
 			}
 		}
 	}
 }
 
-void FSequencer::ClearAllMarkedFrames()
+void FSequencer::DeleteAllMarkedFrames()
 {
 	UMovieSceneSequence* FocusedMovieSequence = GetFocusedMovieSceneSequence();
 	if (FocusedMovieSequence != nullptr)
@@ -9236,10 +9242,10 @@ void FSequencer::ClearAllMarkedFrames()
 		UMovieScene* FocusedMovieScene = FocusedMovieSequence->GetMovieScene();
 		if (FocusedMovieScene != nullptr)
 		{
-			FScopedTransaction ClearAllMarkedFramesTransaction(LOCTEXT("ClearAllMarkedFrames_Transaction", "Clear All Marked Frames"));
+			FScopedTransaction DeleteAllMarkedFramesTransaction(LOCTEXT("DeleteAllMarkedFrames_Transaction", "Delete All Marked Frames"));
 
 			FocusedMovieScene->Modify();
-			FocusedMovieScene->ClearMarkedFrames();
+			FocusedMovieScene->DeleteMarkedFrames();
 		}
 	}
 }
diff --git a/Engine/Source/Editor/Sequencer/Private/Sequencer.h b/Engine/Source/Editor/Sequencer/Private/Sequencer.h
--- a/Engine/Source/Editor/Sequencer/Private/Sequencer.h
+++ b/Engine/Source/Editor/Sequencer/Private/Sequencer.h
@@ -277,11 +277,15 @@ protected:
 
 	/**
 	 * @param	FrameNumber The FrameNumber in Ticks
-	 * @param	bSetMark  true to set the mark, false to clear the mark
 	 */
-	void OnMarkedFrameChanged(FFrameNumber FrameNumber, bool bSetMark);
+	void AddMarkedFrame(FFrameNumber FrameNumber);
 
-	void ClearAllMarkedFrames();
+	/**
+	 * @param InMarkIndex The marked frame index to delete
+     */
+	void DeleteMarkedFrame(int32 InMarkIndex);
+
+	void DeleteAllMarkedFrames();
 
 public:
 
diff --git a/Engine/Source/Editor/Sequencer/Private/SequencerTimeSliderController.cpp b/Engine/Source/Editor/Sequencer/Private/SequencerTimeSliderController.cpp
--- a/Engine/Source/Editor/Sequencer/Private/SequencerTimeSliderController.cpp
+++ b/Engine/Source/Editor/Sequencer/Private/SequencerTimeSliderController.cpp
@@ -1350,19 +1350,19 @@ TSharedRef<SWidget> FSequencerTimeSliderController::OpenSetPlaybackRangeMenu(con
 		else 
 		{
 			MenuBuilder.AddMenuEntry( 
-				LOCTEXT("ClearMark", "Clear Mark"),
+				LOCTEXT("DeleteMark", "Delete Mark"),
 				FText(),
 				FSlateIcon(),
-				FUIAction(FExecuteAction::CreateLambda([=]{ ClearMarkAtFrame(FrameNumber); }))
+				FUIAction(FExecuteAction::CreateLambda([=]{ DeleteMarkAtIndex(MarkedIndex); }))
 			);
 		}
 
 		MenuBuilder.AddMenuEntry( 
-			LOCTEXT("Clear All Marks", "Clear All Marks"),
+			LOCTEXT("Delete All Marks", "Delete All Marks"),
 			FText(),
 			FSlateIcon(),
 			FUIAction(
-				FExecuteAction::CreateLambda([=]{ ClearAllMarks(); }),
+				FExecuteAction::CreateLambda([=]{ DeleteAllMarks(); }),
 				FCanExecuteAction::CreateLambda([=]{ return bHasMarks; })
 			)
 		);
@@ -1632,17 +1632,17 @@ void FSequencerTimeSliderController::SetMark(int32 InMarkIndex, FFrameNumber Fra
 
 void FSequencerTimeSliderController::AddMarkAtFrame(FFrameNumber FrameNumber)
 {
-	TimeSliderArgs.OnMarkedFrameChanged.ExecuteIfBound(FrameNumber, true);
+	TimeSliderArgs.OnAddMarkedFrame.ExecuteIfBound(FrameNumber);
 }
 
-void FSequencerTimeSliderController::ClearMarkAtFrame(FFrameNumber FrameNumber)
+void FSequencerTimeSliderController::DeleteMarkAtIndex(int32 InMarkIndex)
 {
-	TimeSliderArgs.OnMarkedFrameChanged.ExecuteIfBound(FrameNumber, false);
+	TimeSliderArgs.OnDeleteMarkedFrame.ExecuteIfBound(InMarkIndex);
 }
 
-void FSequencerTimeSliderController::ClearAllMarks()
+void FSequencerTimeSliderController::DeleteAllMarks()
 {
-	TimeSliderArgs.OnClearAllMarkedFrames.ExecuteIfBound();
+	TimeSliderArgs.OnDeleteAllMarkedFrames.ExecuteIfBound();
 }
 	
 
diff --git a/Engine/Source/Editor/Sequencer/Public/SequencerTimeSliderController.h b/Engine/Source/Editor/Sequencer/Public/SequencerTimeSliderController.h
--- a/Engine/Source/Editor/Sequencer/Public/SequencerTimeSliderController.h
+++ b/Engine/Source/Editor/Sequencer/Public/SequencerTimeSliderController.h
@@ -252,8 +252,8 @@ private:
 	FFrameTime ComputeFrameTimeFromMouse(const FGeometry& Geometry, FVector2D ScreenSpacePosition, FScrubRangeToScreen RangeToScreen, bool CheckSnapping = true) const;
 
 	void AddMarkAtFrame(FFrameNumber FrameNumber);
-	void ClearMarkAtFrame(FFrameNumber FrameNumber);
-	void ClearAllMarks();
+	void DeleteMarkAtIndex(int32 InMarkIndex);
+	void DeleteAllMarks();
 
 private:
 
diff --git a/Engine/Source/Editor/SequencerWidgets/Public/ITimeSlider.h b/Engine/Source/Editor/SequencerWidgets/Public/ITimeSlider.h
--- a/Engine/Source/Editor/SequencerWidgets/Public/ITimeSlider.h
+++ b/Engine/Source/Editor/SequencerWidgets/Public/ITimeSlider.h
@@ -29,7 +29,8 @@ DECLARE_DELEGATE_TwoParams( FOnViewRangeChanged, TRange<double>, EViewRangeInter
 DECLARE_DELEGATE_OneParam( FOnTimeRangeChanged, TRange<double> )
 DECLARE_DELEGATE_OneParam( FOnFrameRangeChanged, TRange<FFrameNumber> )
 DECLARE_DELEGATE_TwoParams(FOnSetMarkedFrame, int32, FFrameNumber)
-DECLARE_DELEGATE_TwoParams(FOnMarkedFrameChanged, FFrameNumber, bool)
+DECLARE_DELEGATE_OneParam(FOnAddMarkedFrame, FFrameNumber)
+DECLARE_DELEGATE_OneParam(FOnDeleteMarkedFrame, int32)
 DECLARE_DELEGATE_RetVal_TwoParams( FFrameNumber, FOnGetNearestKey, FFrameTime, bool )
 
 /** Structure used to wrap up a range, and an optional animation target */
@@ -151,11 +152,14 @@ struct FTimeSliderArgs
 	/** Called when the marked frame needs to be set */
 	FOnSetMarkedFrame OnSetMarkedFrame;
 
-	/** Called when a marked frame is added/removed */
-	FOnMarkedFrameChanged OnMarkedFrameChanged;
+	/** Called when a marked frame is added */
+	FOnAddMarkedFrame OnAddMarkedFrame;
 
-	/** Called when all marked frames should be cleared */
-	FSimpleDelegate OnClearAllMarkedFrames;
+	/** Called when a marked frame is deleted */
+	FOnDeleteMarkedFrame OnDeleteMarkedFrame;
+
+	/** Called when all marked frames should be deleted */
+	FSimpleDelegate OnDeleteAllMarkedFrames;
 
 	/** Round the scrub position to an integer during playback */
 	TAttribute<EMovieScenePlayerStatus::Type> PlaybackStatus;
diff --git a/Engine/Source/Runtime/MovieScene/Private/MovieScene.cpp b/Engine/Source/Runtime/MovieScene/Private/MovieScene.cpp
--- a/Engine/Source/Runtime/MovieScene/Private/MovieScene.cpp
+++ b/Engine/Source/Runtime/MovieScene/Private/MovieScene.cpp
@@ -1304,12 +1304,12 @@ int32 UMovieScene::AddMarkedFrame(const FMovieSceneMarkedFrame &InMarkedFrame)
 	return MarkedIndex;
 }
 
-void UMovieScene::RemoveMarkedFrame(int32 RemoveIndex)
+void UMovieScene::DeleteMarkedFrame(int32 DeleteIndex)
 {
-	MarkedFrames.RemoveAt(RemoveIndex);
+	MarkedFrames.RemoveAt(DeleteIndex);
 }
 
-void UMovieScene::ClearMarkedFrames()
+void UMovieScene::DeleteMarkedFrames()
 {
 	MarkedFrames.Empty();
 }
diff --git a/Engine/Source/Runtime/MovieScene/Public/MovieScene.h b/Engine/Source/Runtime/MovieScene/Public/MovieScene.h
--- a/Engine/Source/Runtime/MovieScene/Public/MovieScene.h
+++ b/Engine/Source/Runtime/MovieScene/Public/MovieScene.h
@@ -895,16 +895,16 @@ public:
 	int32 AddMarkedFrame(const FMovieSceneMarkedFrame& InMarkedFrame);
 
 	/*
-	 * Remove the user marked frame by index.
+	 * Delete the user marked frame by index.
 	 *
-	 * @RemoveIndex The index to the user marked frame to remove
+	 * @DeleteIndex The index to the user marked frame to delete
 	 */
-	void RemoveMarkedFrame(int32 RemoveIndex);
+	void DeleteMarkedFrame(int32 DeleteIndex);
 
 	/*
-	 * Clear all user marked frames
+	 * Delete all user marked frames
 	 */
-	void ClearMarkedFrames();
+	void DeleteMarkedFrames();
 
 	/*
 	 * Find the user marked frame by label
