Sequencer: Duplicating a possessable object binding now duplicates the possessed objects into the current level and assigns them to the duplicated binding.

diff --git a/Engine/Source/Editor/Sequencer/Private/Sequencer.cpp b/Engine/Source/Editor/Sequencer/Private/Sequencer.cpp
--- a/Engine/Source/Editor/Sequencer/Private/Sequencer.cpp
+++ b/Engine/Source/Editor/Sequencer/Private/Sequencer.cpp
@@ -7058,9 +7058,12 @@ bool FSequencer::DoPaste()
 
 	return bAnythingPasted;
 }
+PRAGMA_DISABLE_OPTIMIZATION
 
 bool FSequencer::PasteObjectBindings(const FString& TextToImport, TArray<FNotificationInfo>& PasteErrors)
 {
+	UWorld* World = Cast<UWorld>(GetPlaybackContext());
+
 	TArray<UMovieSceneFolder*> SelectedParentFolders;
 	FString NewNodePath;
 	CalculateSelectedFolderAndPath(SelectedParentFolders, NewNodePath);
@@ -7147,6 +7150,46 @@ bool FSequencer::PasteObjectBindings(const FString& TextToImport, TArray<FNotifi
 						Possessable->SetParent(SelectedParentGuids[TargetIndex]);
 					}
 				}
+
+				TArray<AActor*> ActorsToDuplicate;
+				for (auto RuntimeObject : FindBoundObjects(CopyableBinding->Possessable.GetGuid(), ActiveTemplateIDs.Top()))
+				{
+					AActor* Actor = Cast<AActor>(RuntimeObject.Get());
+					if (Actor)
+					{
+						ActorsToDuplicate.Add(Actor);
+					}
+				}
+
+				TArray<AActor*> DuplicatedActors;
+				if (ActorsToDuplicate.Num() != 0)
+				{
+					GEditor->SelectNone(false, true);
+					for (AActor* ActorToDuplicate : ActorsToDuplicate)
+					{
+						GEditor->SelectActor(ActorToDuplicate, true, false, false);
+					}
+
+					// Duplicate the bound actors
+					GEditor->edactDuplicateSelected(World->GetCurrentLevel(), false);
+
+					USelection* ActorSelection = GEditor->GetSelectedActors();
+					TArray<TWeakObjectPtr<AActor> > SelectedActors;
+					for (FSelectionIterator Iter(*ActorSelection); Iter; ++Iter)
+					{
+						AActor* Actor = Cast<AActor>(*Iter);
+						if (Actor)
+						{
+							DuplicatedActors.Add(Actor);
+						}
+					}
+
+					// Bind the duplicated actors
+					if (DuplicatedActors.Num())
+					{
+						ReplaceBindingWithActors(NewGuid, DuplicatedActors);
+					}
+				}
 			}
 			else if (CopyableBinding->Spawnable.GetGuid().IsValid())
 			{
@@ -7274,6 +7317,8 @@ bool FSequencer::PasteObjectBindings(const FString& TextToImport, TArray<FNotifi
 	return true;
 }
 
+PRAGMA_ENABLE_OPTIMIZATION
+
 bool FSequencer::PasteTracks(const FString& TextToImport, TArray<FNotificationInfo>& PasteErrors)
 {
 	FScopedTransaction Transaction(FGenericCommands::Get().Paste->GetDescription());
@@ -9378,45 +9423,56 @@ void FSequencer::DuplicateSelection()
 {
 	FScopedTransaction DuplicateSelectionTransaction(LOCTEXT("DuplicateSelection_Transaction", "Duplicate Selection"));
 
-	CopySelection();
-	DoPaste();
-
-	// DoPaste doesn't handle keys, and we want duplicated keys to shift by one display rate frame as an overlapping key isn't useful.
-	
-	if (!Selection.GetSelectedKeys().Num())
+	if (Selection.GetSelectedKeys().Num() != 0)
 	{
-		return;
-	}
+		CopySelection();
+		DoPaste();
 
-	// Offset by a visible amount
-	FFrameNumber FrameOffset = FFrameNumber((int32)GetDisplayRateDeltaFrameCount());
+		// Shift duplicated keys by one display rate frame as an overlapping key isn't useful
 
-	TArray<FSequencerSelectedKey> NewSelection;
-	for (const FSequencerSelectedKey& Key : Selection.GetSelectedKeys())
-	{
-		if (Key.IsValid())
+		// Offset by a visible amount
+		FFrameNumber FrameOffset = FFrameNumber((int32)GetDisplayRateDeltaFrameCount());
+
+		TArray<FSequencerSelectedKey> NewSelection;
+		for (const FSequencerSelectedKey& Key : Selection.GetSelectedKeys())
 		{
-			TSharedPtr<IKeyArea> KeyArea = Key.KeyArea;
-			FKeyHandle KeyHandle = Key.KeyHandle.GetValue();
+			if (Key.IsValid())
+			{
+				TSharedPtr<IKeyArea> KeyArea = Key.KeyArea;
+				FKeyHandle KeyHandle = Key.KeyHandle.GetValue();
 
-			FKeyHandle NewKeyHandle = KeyArea->DuplicateKey(KeyHandle);
-			KeyArea->SetKeyTime(NewKeyHandle, KeyArea->GetKeyTime(KeyHandle) + FrameOffset);
+				FKeyHandle NewKeyHandle = KeyArea->DuplicateKey(KeyHandle);
+				KeyArea->SetKeyTime(NewKeyHandle, KeyArea->GetKeyTime(KeyHandle) + FrameOffset);
 
-			NewSelection.Add(FSequencerSelectedKey(*KeyArea->GetOwningSection(), KeyArea, NewKeyHandle));
+				NewSelection.Add(FSequencerSelectedKey(*KeyArea->GetOwningSection(), KeyArea, NewKeyHandle));
+			}
 		}
-	}
-	
-	Selection.SuspendBroadcast();
-	Selection.EmptySelectedKeys();
 
-	for (const FSequencerSelectedKey& Key : NewSelection)
+		Selection.SuspendBroadcast();
+		Selection.EmptySelectedKeys();
+
+		for (const FSequencerSelectedKey& Key : NewSelection)
+		{
+			Selection.AddToSelection(Key);
+		}
+		Selection.ResumeBroadcast();
+		Selection.GetOnKeySelectionChanged().Broadcast();
+
+		NotifyMovieSceneDataChanged(EMovieSceneDataChangeType::TrackValueChanged);
+	}
+	else if (Selection.GetSelectedSections().Num() != 0)
 	{
-		Selection.AddToSelection(Key);
+		CopySelection();
+		DoPaste();
 	}
-	Selection.ResumeBroadcast();
-	Selection.GetOnKeySelectionChanged().Broadcast();
+	else
+	{
+		CopySelection();
 
-	NotifyMovieSceneDataChanged(EMovieSceneDataChangeType::TrackValueChanged);
+		EmptySelection();
+
+		DoPaste();
+	}
 }
 
 void FSequencer::CopySelectedKeys()
