Sequencer: Check for name too long when renaming sequencer folders and object bindings.

diff --git a/Engine/Source/Editor/Sequencer/Private/DisplayNodes/SequencerDisplayNode.cpp b/Engine/Source/Editor/Sequencer/Private/DisplayNodes/SequencerDisplayNode.cpp
--- a/Engine/Source/Editor/Sequencer/Private/DisplayNodes/SequencerDisplayNode.cpp
+++ b/Engine/Source/Editor/Sequencer/Private/DisplayNodes/SequencerDisplayNode.cpp
@@ -787,6 +787,11 @@ bool FSequencerDisplayNode::ValidateDisplayName(const FText& NewDisplayName, FTe
 		OutErrorMessage = NSLOCTEXT("Sequencer", "RenameFailed_LeftBlank", "Labels cannot be left blank");
 		return false;
 	}
+	else if (NewDisplayName.ToString().Len() >= NAME_SIZE)
+	{
+		OutErrorMessage = FText::Format(NSLOCTEXT("Sequencer", "RenameFailed_TooLong", "Names must be less than {0} characters long"), NAME_SIZE);
+		return false;
+	}
 	return true;
 }
 
