Sequencer: Visit key areas from track, key area, and category nodes only. This prevents visiting keys when object binding nodes are collapsed.

diff --git a/Engine/Source/Editor/Sequencer/Private/Tools/SequencerEntityVisitor.cpp b/Engine/Source/Editor/Sequencer/Private/Tools/SequencerEntityVisitor.cpp
--- a/Engine/Source/Editor/Sequencer/Private/Tools/SequencerEntityVisitor.cpp
+++ b/Engine/Source/Editor/Sequencer/Private/Tools/SequencerEntityVisitor.cpp
@@ -79,9 +79,12 @@ void FSequencerEntityWalker::ConditionallyIntersectNode(const ISequencerEntityVi
 			}
 		}
 
-		if (Range.IntersectKeyArea(InNode, VirtualKeySize.Y))
+		if (InNode->GetType() == ESequencerNode::Track || InNode->GetType() == ESequencerNode::KeyArea || InNode->GetType() == ESequencerNode::Category)
 		{
-			VisitKeyAnyAreas(Visitor, InNode, !InNode->IsExpanded());
+			if (Range.IntersectKeyArea(InNode, VirtualKeySize.Y))
+			{
+				VisitKeyAnyAreas(Visitor, InNode, !InNode->IsExpanded());
+			}
 		}
 	}
 
