Sequencer: Add keys to the snap field when resizing sections

diff --git a/Engine/Source/Editor/Sequencer/Private/Tools/EditToolDragOperations.cpp b/Engine/Source/Editor/Sequencer/Private/Tools/EditToolDragOperations.cpp
--- a/Engine/Source/Editor/Sequencer/Private/Tools/EditToolDragOperations.cpp
+++ b/Engine/Source/Editor/Sequencer/Private/Tools/EditToolDragOperations.cpp
@@ -138,7 +138,7 @@ void FResizeSection::OnBeginDrag(const FPointerEvent& MouseEvent, FVector2D Loca
 	// Construct a snap field of unselected sections
 	TSet<FSequencerSelectedKey> EmptyKeySet;
 	FInvalidKeyAndSectionSnappingCandidates SnapCandidates(EmptyKeySet, Sections);
-	SnapField = FSequencerSnapField(Sequencer, SnapCandidates, ESequencerEntity::Section);
+	SnapField = FSequencerSnapField(Sequencer, SnapCandidates, ESequencerEntity::Section | ESequencerEntity::Key);
 
 	SectionInitTimes.Empty();
 
