COPY from //Dev-Anim to //Dev-Main

diff --git a/Engine/Source/Editor/Sequencer/Private/DisplayNodes/SequencerSectionKeyAreaNode.cpp b/Engine/Source/Editor/Sequencer/Private/DisplayNodes/SequencerSectionKeyAreaNode.cpp
--- a/Engine/Source/Editor/Sequencer/Private/DisplayNodes/SequencerSectionKeyAreaNode.cpp
+++ b/Engine/Source/Editor/Sequencer/Private/DisplayNodes/SequencerSectionKeyAreaNode.cpp
@@ -113,6 +113,20 @@ void FSequencerSectionKeyAreaNode::SetDisplayName(const FText& NewDisplayName)
 	check(false);
 }
 
+FSlateFontInfo FSequencerSectionKeyAreaNode::GetDisplayNameFont() const
+{
+	bool bAnyKeyAreaAnimated = false;
+	for (const TSharedRef<IKeyArea>& KeyArea : KeyAreas)
+	{
+		FMovieSceneChannel* Channel = KeyArea->ResolveChannel();
+		if (Channel && Channel->GetNumKeys() > 0)
+		{
+			return FEditorStyle::GetFontStyle("Sequencer.AnimationOutliner.ItalicFont");
+		}
+	}
+	return FSequencerDisplayNode::GetDisplayNameFont();
+}
+
 void FSequencerSectionKeyAreaNode::CreateCurveModels(TArray<TUniquePtr<FCurveModel>>& OutCurveModels)
 {
 	TSharedRef<ISequencer> Sequencer = GetSequencer().AsShared();
diff --git a/Engine/Source/Editor/Sequencer/Private/DisplayNodes/SequencerSectionKeyAreaNode.h b/Engine/Source/Editor/Sequencer/Private/DisplayNodes/SequencerSectionKeyAreaNode.h
--- a/Engine/Source/Editor/Sequencer/Private/DisplayNodes/SequencerSectionKeyAreaNode.h
+++ b/Engine/Source/Editor/Sequencer/Private/DisplayNodes/SequencerSectionKeyAreaNode.h
@@ -74,6 +74,7 @@ public:
 	virtual FNodePadding GetNodePadding() const override;
 	virtual ESequencerNode::Type GetType() const override;
 	virtual void SetDisplayName(const FText& NewDisplayName) override;
+	virtual FSlateFontInfo GetDisplayNameFont() const override;
 
 	// ICurveEditorTreeItem interface
 	virtual void CreateCurveModels(TArray<TUniquePtr<FCurveModel>>& OutCurveModels) override;
diff --git a/Engine/Source/Editor/Sequencer/Private/DisplayNodes/SequencerTrackNode.cpp b/Engine/Source/Editor/Sequencer/Private/DisplayNodes/SequencerTrackNode.cpp
--- a/Engine/Source/Editor/Sequencer/Private/DisplayNodes/SequencerTrackNode.cpp
+++ b/Engine/Source/Editor/Sequencer/Private/DisplayNodes/SequencerTrackNode.cpp
@@ -928,5 +928,31 @@ void FSequencerTrackNode::CreateCurveModels(TArray<TUniquePtr<FCurveModel>>& Out
 	}
 }
 
+FSlateFontInfo FSequencerTrackNode::GetDisplayNameFont() const
+{
+	bool bAllAnimated = false;
+	TSharedPtr<FSequencerSectionKeyAreaNode> TopLevelKeyArea = GetTopLevelKeyNode();
+	if (TopLevelKeyArea.IsValid())
+	{
+		for (const TSharedRef<IKeyArea>& KeyArea : TopLevelKeyArea->GetAllKeyAreas())
+		{
+			FMovieSceneChannel* Channel = KeyArea->ResolveChannel();
+			if (!Channel || Channel->GetNumKeys() == 0)
+			{
+				return FSequencerDisplayNode::GetDisplayNameFont();
+			}
+			else
+			{
+				bAllAnimated = true;
+			}
+		}
+		if (bAllAnimated == true)
+		{
+			return FEditorStyle::GetFontStyle("Sequencer.AnimationOutliner.ItalicFont");
+		}
+	}
+	return FSequencerDisplayNode::GetDisplayNameFont();
+}
+
 
 #undef LOCTEXT_NAMESPACE
\ No newline at end of file
diff --git a/Engine/Source/Editor/Sequencer/Private/SAnimationOutlinerTreeNode.cpp b/Engine/Source/Editor/Sequencer/Private/SAnimationOutlinerTreeNode.cpp
--- a/Engine/Source/Editor/Sequencer/Private/SAnimationOutlinerTreeNode.cpp
+++ b/Engine/Source/Editor/Sequencer/Private/SAnimationOutlinerTreeNode.cpp
@@ -61,11 +61,9 @@ void SAnimationOutlinerTreeNode::Construct( const FArguments& InArgs, TSharedRef
 
 	TableRowStyle = &FEditorStyle::Get().GetWidgetStyle<FTableRowStyle>("TableView.Row");
 
-	FSlateFontInfo NodeFont = FEditorStyle::GetFontStyle("Sequencer.AnimationOutliner.RegularFont");
-
 	EditableLabel = SNew(SInlineEditableTextBlock)
 		.IsReadOnly(this, &SAnimationOutlinerTreeNode::IsNodeLabelReadOnly)
-		.Font(NodeFont)
+		.Font(this, &SAnimationOutlinerTreeNode::GetDisplayNameFont)
 		.ColorAndOpacity(this, &SAnimationOutlinerTreeNode::GetDisplayNameColor)
 		.OnVerifyTextChanged(this, &SAnimationOutlinerTreeNode::VerifyNodeTextChanged)
 		.OnTextCommitted(this, &SAnimationOutlinerTreeNode::HandleNodeLabelTextCommitted)
@@ -429,6 +427,10 @@ FSlateColor SAnimationOutlinerTreeNode::GetDisplayNameColor() const
 	return DisplayNode->GetDisplayNameColor();
 }
 
+FSlateFontInfo SAnimationOutlinerTreeNode::GetDisplayNameFont() const
+{
+	return DisplayNode->GetDisplayNameFont();
+}
 
 FText SAnimationOutlinerTreeNode::GetDisplayNameToolTipText() const
 {
diff --git a/Engine/Source/Editor/Sequencer/Private/SAnimationOutlinerTreeNode.h b/Engine/Source/Editor/Sequencer/Private/SAnimationOutlinerTreeNode.h
--- a/Engine/Source/Editor/Sequencer/Private/SAnimationOutlinerTreeNode.h
+++ b/Engine/Source/Editor/Sequencer/Private/SAnimationOutlinerTreeNode.h
@@ -89,6 +89,11 @@ private:
 	 */
 	FSlateColor GetDisplayNameColor() const;
 
+	/**
+	*@return The font used to draw the display name.
+	*/
+	FSlateFontInfo GetDisplayNameFont() const;
+
 	/**
 	 * @return The text displayed for the tool tip for the diplay name label. 
 	 */
diff --git a/Engine/Source/Editor/Sequencer/Private/SSequencer.cpp b/Engine/Source/Editor/Sequencer/Private/SSequencer.cpp
--- a/Engine/Source/Editor/Sequencer/Private/SSequencer.cpp
+++ b/Engine/Source/Editor/Sequencer/Private/SSequencer.cpp
@@ -1110,6 +1110,7 @@ void SSequencer::InitializeTrackFilters()
 	AllTrackFilters.Add(MakeShared<FSequencerTrackFilter_CameraObjects>());
 	AllTrackFilters.Add(MakeShared<FSequencerTrackFilter_LightObjects>());
 	AllTrackFilters.Add(MakeShared<FSequencerTrackFilter_SkeletalMeshObjects>());
+	AllTrackFilters.Add(MakeShared<FSequencerTrackFilter_Animated>());
 
 	// Add any global user-defined frontend filters
 	for (TObjectIterator<USequencerTrackFilterExtension> ExtensionIt(RF_NoFlags); ExtensionIt; ++ExtensionIt)
@@ -1207,7 +1208,7 @@ void SSequencer::HandleOutlinerNodeSelectionChanged()
 		if (Settings->ShouldSyncCurveEditorSelection())
 		{
 			TSharedRef<FSequencerNodeTree> NodeTree = Sequencer->GetNodeTree();
-
+			CurveEditor->SuspendBroadcast();
 			// Clear the tree selection
 			CurveEditorTree->ClearSelection();
 			for (TSharedRef<FSequencerDisplayNode> Node : SelectedDisplayNodes)
@@ -1218,6 +1219,7 @@ void SSequencer::HandleOutlinerNodeSelectionChanged()
 					CurveEditorTree->SetItemSelection(CurveEditorTreeItem, true);
 				}
 			}
+			CurveEditor->ResumeBroadcast();
 		}
 	}
 }
@@ -1683,6 +1685,31 @@ void SSequencer::FillLevelFilterMenu(FMenuBuilder& InMenuBarBuilder)
 	}
 }
 
+
+void SSequencer::SetFilterOn(const FText& InName, bool bOn)
+{
+	TSharedPtr<FSequencer> Sequencer = SequencerPtr.Pin();
+
+
+	for (TSharedRef<FSequencerTrackFilter> TrackFilter : AllTrackFilters)
+	{
+		if (TrackFilter->SupportsSequence(Sequencer->GetFocusedMovieSceneSequence()))
+		{
+			if (InName.EqualToCaseIgnored(TrackFilter->GetDisplayName()))
+			{
+				if (bOn && !IsTrackFilterActive(TrackFilter))
+				{
+					Sequencer->GetNodeTree()->AddFilter(TrackFilter);
+				}
+				else if (!bOn && IsTrackFilterActive(TrackFilter))
+				{
+					Sequencer->GetNodeTree()->RemoveFilter(TrackFilter);
+				}
+			}
+		}
+	}
+}
+
 void SSequencer::OnResetFilters()
 {
 	TSharedPtr<FSequencer> Sequencer = SequencerPtr.Pin();
@@ -3489,5 +3516,7 @@ void SSequencer::SetPlayTimeClampedByWorkingRange(double Frame)
 	}
 }
 
+
+
 #undef LOCTEXT_NAMESPACE
 
diff --git a/Engine/Source/Editor/Sequencer/Private/SSequencer.h b/Engine/Source/Editor/Sequencer/Private/SSequencer.h
--- a/Engine/Source/Editor/Sequencer/Private/SSequencer.h
+++ b/Engine/Source/Editor/Sequencer/Private/SSequencer.h
@@ -318,6 +318,10 @@ public:
 
 	/** Sets the play time for the sequence but clamped by the working range. This is useful for cases where we can't clamp via the UI control. */
 	void SetPlayTimeClampedByWorkingRange(double Frame);
+
+	/** Set's the specified filter to be on or off*/
+	void SetFilterOn(const FText& InName, bool bOn);
+
 public:
 
 	// FNotifyHook overrides
diff --git a/Engine/Source/Editor/Sequencer/Private/Sequencer.cpp b/Engine/Source/Editor/Sequencer/Private/Sequencer.cpp
--- a/Engine/Source/Editor/Sequencer/Private/Sequencer.cpp
+++ b/Engine/Source/Editor/Sequencer/Private/Sequencer.cpp
@@ -341,6 +341,9 @@ void FSequencer::InitSequencer(const FSequencerInitParams& InitParams, const TSh
 		CurveEditorModel->InputSnapRateAttribute      = MakeAttributeSP(this, &FSequencer::GetFocusedDisplayRate);
 
 		CurveEditorModel->DefaultKeyAttributes        = MakeAttributeSP(this, &FSequencer::GetDefaultKeyAttributes);
+		
+		CurveEditorModel->OnCurveArrayChanged.AddRaw(this, &FSequencer::OnCurveModelDisplayChanged);
+
 	}
 
 	{
@@ -1762,6 +1765,50 @@ void FSequencer::ShrinkTime(FFrameTime InDeltaTime)
 	Selection.Empty(); //todo restore key and section selection
 }
 
+bool FSequencer::CanAddTransformKeysForSelectedObjects() const
+{
+	for (int32 i = 0; i < TrackEditors.Num(); ++i)
+	{
+		if (TrackEditors[i]->HasTransformKeyBindings() && TrackEditors[i]->CanAddTransformKeysForSelectedObjects())
+		{
+			return true;
+		}
+	}
+	return false;
+}
+
+void FSequencer::OnAddTransformKeysForSelectedObjects(EMovieSceneTransformChannel Channel)
+{
+	TArray<TSharedPtr<ISequencerTrackEditor>> PossibleTrackEditors;
+	bool AtLeastOneHasPriority = false;
+	for (int32 i = 0; i < TrackEditors.Num(); ++i)
+	{
+		if (TrackEditors[i]->HasTransformKeyBindings()  && TrackEditors[i]->CanAddTransformKeysForSelectedObjects())
+		{
+			PossibleTrackEditors.Add(TrackEditors[i]);
+			if (TrackEditors[i]->HasTransformKeyOverridePriority())
+			{
+				AtLeastOneHasPriority = true;
+			}
+		}
+	}
+	for (int32 i = 0; i < PossibleTrackEditors.Num(); ++i)
+	{
+		if (AtLeastOneHasPriority)
+		{
+			if (PossibleTrackEditors[i]->HasTransformKeyOverridePriority())
+			{
+				PossibleTrackEditors[i]->OnAddTransformKeysForSelectedObjects(Channel);
+			}
+		}
+		else
+		{
+			PossibleTrackEditors[i]->OnAddTransformKeysForSelectedObjects(Channel);
+		}
+	}
+
+}
+
 void FSequencer::BakeTransform()
 {
 	UMovieScene* FocusedMovieScene = GetFocusedMovieSceneSequence()->GetMovieScene();
@@ -5496,7 +5543,7 @@ void FSequencer::SynchronizeExternalSelectionWithSequencerSelection()
 	{
 		if (GEditor->GetSelectedActorCount())
 		{
-			const FScopedTransaction Transaction( NSLOCTEXT( "Sequencer", "UpdatingActorComponentSelectionNone", "Select None" ) );
+			const FScopedTransaction Transaction(NSLOCTEXT("Sequencer", "UpdatingActorComponentSelectionNone", "Select None"), !GIsTransacting);
 			GEditor->SelectNone( bNotifySelectionChanged, bDeselectBSP, bWarnAboutTooManyActors );
 			GEditor->NoteSelectionChange();
 		}
@@ -5566,7 +5613,7 @@ void FSequencer::SynchronizeExternalSelectionWithSequencerSelection()
 		return;
 	}
 
-	const FScopedTransaction Transaction( NSLOCTEXT( "Sequencer", "UpdatingActorComponentSelection", "Select Actors/Components" ) );
+	const FScopedTransaction Transaction(NSLOCTEXT("Sequencer", "UpdatingActorComponentSelection", "Select Actors/Components"), !GIsTransacting);
 
 
 	GEditor->GetSelectedActors()->Modify();
@@ -5830,6 +5877,16 @@ FSequencerSelectionPreview& FSequencer::GetSelectionPreview()
 	return SelectionPreview;
 }
 
+void FSequencer::SuspendSelectionBroadcast()
+{
+	Selection.SuspendBroadcast();
+}
+
+void FSequencer::ResumeSelectionBroadcast()
+{
+	Selection.ResumeBroadcast();
+}
+
 void FSequencer::GetSelectedTracks(TArray<UMovieSceneTrack*>& OutSelectedTracks)
 {
 	OutSelectedTracks.Append(Selection.GetSelectedTracks());
@@ -5892,49 +5949,72 @@ void FSequencer::GetSelectedKeyAreas(TArray<const IKeyArea*>& OutSelectedKeyArea
 	}
 }
 
-void FSequencer::SelectByKeyAreas(const TArray<IKeyArea>& InKeyAreas, bool bSelectParentInstead, bool bSelect)
+void FSequencer::SelectByNthCategoryNode(UMovieSceneSection* Section, int Index, bool bSelect)
 {
 	TSet<TSharedRef<FSequencerDisplayNode>> Nodes;
 	TArray<TSharedRef<FSequencerDisplayNode>> NodesToSelect;
-	for (const TSharedRef<FSequencerDisplayNode>& Node : NodeTree->GetAllNodes())
+
+	TOptional<FSectionHandle> SectionHandle = NodeTree->GetSectionHandle(Section);
+	int32 Count = 0;
+	if (SectionHandle.IsSet())
 	{
-		if (Node->GetType() == ESequencerNode::Track)
+		TSharedRef<FSequencerTrackNode> TrackNode = SectionHandle->GetTrackNode();
+		for (const TSharedRef<FSequencerDisplayNode>& Node : TrackNode->GetChildNodes())
 		{
-			TSharedRef<FSequencerTrackNode> TrackNode = StaticCastSharedRef<FSequencerTrackNode>(Node);
-			TArray<TSharedRef<FSequencerSectionKeyAreaNode>> KeyAreaNodes;
-			TrackNode->GetChildKeyAreaNodesRecursively(KeyAreaNodes);
-			for (TSharedRef<FSequencerSectionKeyAreaNode> KeyAreaNode : KeyAreaNodes)
+			if (Node->GetType() == ESequencerNode::Category && Count++ == Index)
 			{
-				for (TSharedPtr<IKeyArea> KeyArea : KeyAreaNode->GetAllKeyAreas())
+				NodesToSelect.Add(Node);
+				if (bSelect == false) //make sure all children not selected
 				{
-					for (const IKeyArea& InKeyArea : InKeyAreas)
+					for (const TSharedRef<FSequencerDisplayNode>& ChildNode : Node->GetChildNodes())
 					{
-						if (InKeyArea.GetOwningSection() == KeyArea->GetOwningSection() &&
-							InKeyArea.GetChannel() == KeyArea->GetChannel())
-						{
-							if (bSelectParentInstead || bSelect == false)
-							{
-								Nodes.Add(KeyAreaNode->GetParent()->AsShared());
-							}
-							if (!bSelectParentInstead || bSelect == false)
-							{
-								Nodes.Add(KeyAreaNode);
-							}
-						}
+						NodesToSelect.Add(ChildNode);
 					}
 				}
 			}
 		}
-		else if (Node->GetType() == ESequencerNode::KeyArea)
+	}
+	if (bSelect)
+	{
+		for (const TSharedRef<FSequencerDisplayNode>& DisplayNode : NodesToSelect)
 		{
-			TSharedRef<FSequencerSectionKeyAreaNode> KeyAreaNode = StaticCastSharedRef<FSequencerSectionKeyAreaNode>(Node);
+			if (DisplayNode->GetParent().IsValid() && DisplayNode->GetParent()->GetType() == ESequencerNode::Track && !DisplayNode->GetParent()->IsExpanded())
+			{
+				DisplayNode->GetParent()->SetExpansionState(true);
+				break;
+			}
+		}
+		Selection.AddToSelection(NodesToSelect);
+	}
+	else
+	{
+		for (const TSharedRef<FSequencerDisplayNode>& DisplayNode : NodesToSelect)
+		{
+			Selection.RemoveFromSelection(DisplayNode);
+			Selection.RemoveFromNodesWithSelectedKeysOrSections(DisplayNode);
+		}
+	}
+}
+
 
+void FSequencer::SelectByKeyAreas(UMovieSceneSection* Section, const TArray<IKeyArea>& InKeyAreas, bool bSelectParentInstead, bool bSelect)
+{
+	TSet<TSharedRef<FSequencerDisplayNode>> Nodes;
+	TArray<TSharedRef<FSequencerDisplayNode>> NodesToSelect;
+
+	TOptional<FSectionHandle> SectionHandle = NodeTree->GetSectionHandle(Section);
+	if (SectionHandle.IsSet())
+	{
+		TSharedRef<FSequencerTrackNode> TrackNode = SectionHandle->GetTrackNode();
+		TArray<TSharedRef<FSequencerSectionKeyAreaNode>> KeyAreaNodes;
+		TrackNode->GetChildKeyAreaNodesRecursively(KeyAreaNodes);
+		for (TSharedRef<FSequencerSectionKeyAreaNode> KeyAreaNode : KeyAreaNodes)
+		{
 			for (TSharedPtr<IKeyArea> KeyArea : KeyAreaNode->GetAllKeyAreas())
 			{
 				for (const IKeyArea& InKeyArea : InKeyAreas)
 				{
-					if (InKeyArea.GetOwningSection() == KeyArea->GetOwningSection() &&
-						InKeyArea.GetChannel() == KeyArea->GetChannel())
+					if (InKeyArea.GetChannel() == KeyArea->GetChannel())
 					{
 						if (bSelectParentInstead || bSelect == false)
 						{
@@ -5949,13 +6029,18 @@ void FSequencer::SelectByKeyAreas(const TArray<IKeyArea>& InKeyAreas, bool bSele
 			}
 		}
 	}
+	
 	if (bSelect)
 	{
 		for (const TSharedRef<FSequencerDisplayNode>& DisplayNode : Nodes)
 		{
+			if (DisplayNode->GetParent().IsValid() && DisplayNode->GetParent()->GetType() == ESequencerNode::Track && !DisplayNode->GetParent()->IsExpanded())
+			{
+				DisplayNode->GetParent()->SetExpansionState(true);
+			}
 			NodesToSelect.Add(DisplayNode);
-			Selection.AddToSelection(NodesToSelect);
 		}
+		Selection.AddToSelection(NodesToSelect);
 	}
 	else
 	{
@@ -7864,6 +7949,19 @@ void FSequencer::ImportTracksFromText(const FString& TextToImport, /*out*/ TArra
 	TempPackage->RemoveFromRoot();
 }
 
+void FSequencer::ObjectImplicitlyAdded(UObject* InObject) const
+{
+	for (int32 i = 0; i < TrackEditors.Num(); ++i)
+	{
+		TrackEditors[i]->ObjectImplicitlyAdded(InObject);
+	}
+}
+
+void FSequencer::SetFilterOn(const FText& InName, bool bOn)
+{
+	SequencerWidget->SetFilterOn(InName, bOn);
+}
+
 
 void FSequencer::ImportSectionsFromText(const FString& TextToImport, /*out*/ TArray<UMovieSceneSection*>& ImportedSections)
 {
@@ -10833,6 +10931,24 @@ void FSequencer::BindCommands()
 		TrackEditors[i]->BindCommands(SequencerCommandBindings);
 	}
 
+
+	SequencerCommandBindings->MapAction(
+		Commands.AddTransformKey,
+		FExecuteAction::CreateSP(this, &FSequencer::OnAddTransformKeysForSelectedObjects, EMovieSceneTransformChannel::All),
+		FCanExecuteAction::CreateSP(this, &FSequencer::CanAddTransformKeysForSelectedObjects));
+	SequencerCommandBindings->MapAction(
+		Commands.AddTranslationKey,
+		FExecuteAction::CreateSP(this, &FSequencer::OnAddTransformKeysForSelectedObjects, EMovieSceneTransformChannel::Translation),
+		FCanExecuteAction::CreateSP(this, &FSequencer::CanAddTransformKeysForSelectedObjects));
+	SequencerCommandBindings->MapAction(
+		Commands.AddRotationKey,
+		FExecuteAction::CreateSP(this, &FSequencer::OnAddTransformKeysForSelectedObjects, EMovieSceneTransformChannel::Rotation),
+		FCanExecuteAction::CreateSP(this, &FSequencer::CanAddTransformKeysForSelectedObjects));
+	SequencerCommandBindings->MapAction(
+		Commands.AddScaleKey,
+		FExecuteAction::CreateSP(this, &FSequencer::OnAddTransformKeysForSelectedObjects, EMovieSceneTransformChannel::Scale),
+		FCanExecuteAction::CreateSP(this, &FSequencer::CanAddTransformKeysForSelectedObjects));
+
 	// copy subset of sequencer commands to shared commands
 	*SequencerSharedBindings = *SequencerCommandBindings;
 
@@ -10858,11 +10974,21 @@ void FSequencer::BindCommands()
 		FExecuteAction::CreateSP(this, &FSequencer::StepForward),
 		EUIActionRepeatMode::RepeatEnabled);
 
+	SequencerCommandBindings->MapAction(
+		Commands.StepForward2,
+		FExecuteAction::CreateSP(this, &FSequencer::StepForward),
+		EUIActionRepeatMode::RepeatEnabled);
+
 	SequencerCommandBindings->MapAction(
 		Commands.StepBackward,
 		FExecuteAction::CreateSP(this, &FSequencer::StepBackward),
 		EUIActionRepeatMode::RepeatEnabled);
 
+	SequencerCommandBindings->MapAction(
+		Commands.StepBackward2,
+		FExecuteAction::CreateSP(this, &FSequencer::StepBackward),
+		EUIActionRepeatMode::RepeatEnabled);
+
 	SequencerCommandBindings->MapAction(
 		Commands.SetInterpolationCubicAuto,
 		FExecuteAction::CreateSP(this, &FSequencer::SetInterpTangentMode, ERichCurveInterpMode::RCIM_Cubic, ERichCurveTangentMode::RCTM_Auto));
@@ -10976,6 +11102,7 @@ void FSequencer::BindCommands()
 		FExecuteAction::CreateSP( this, &FSequencer::SetPlaybackRangeToAllShots ),
 		FCanExecuteAction::CreateSP( this, &FSequencer::IsViewingMasterSequence ) );
 
+
 	// We want a subset of the commands to work in the Curve Editor too, but bound to our functions. This minimizes code duplication
 	// while also freeing us up from issues that result from Sequencer already using two lists (for which our commands might be spread
 	// across both lists which makes a direct copy like it already uses difficult).
@@ -10988,9 +11115,16 @@ void FSequencer::BindCommands()
 	CurveEditorSharedBindings->MapAction(Commands.Pause,				*SequencerCommandBindings->GetActionForCommand(Commands.Pause));
 	CurveEditorSharedBindings->MapAction(Commands.StepForward,			*SequencerCommandBindings->GetActionForCommand(Commands.StepForward));
 	CurveEditorSharedBindings->MapAction(Commands.StepBackward,			*SequencerCommandBindings->GetActionForCommand(Commands.StepBackward));
+	CurveEditorSharedBindings->MapAction(Commands.StepForward2,         *SequencerCommandBindings->GetActionForCommand(Commands.StepForward2));
+	CurveEditorSharedBindings->MapAction(Commands.StepBackward2,        *SequencerCommandBindings->GetActionForCommand(Commands.StepBackward2));
 	CurveEditorSharedBindings->MapAction(Commands.StepToNextKey,		*SequencerCommandBindings->GetActionForCommand(Commands.StepToNextKey));
 	CurveEditorSharedBindings->MapAction(Commands.StepToPreviousKey, 	*SequencerCommandBindings->GetActionForCommand(Commands.StepToPreviousKey));
 
+	CurveEditorSharedBindings->MapAction(Commands.AddTransformKey, *SequencerCommandBindings->GetActionForCommand(Commands.AddTransformKey));
+	CurveEditorSharedBindings->MapAction(Commands.AddTranslationKey, *SequencerCommandBindings->GetActionForCommand(Commands.AddTranslationKey));
+	CurveEditorSharedBindings->MapAction(Commands.AddRotationKey, *SequencerCommandBindings->GetActionForCommand(Commands.AddRotationKey));
+	CurveEditorSharedBindings->MapAction(Commands.AddScaleKey, *SequencerCommandBindings->GetActionForCommand(Commands.AddScaleKey));
+
 	GetCurveEditor()->GetCommands()->Append(CurveEditorSharedBindings);
 
 	// bind widget specific commands
@@ -11176,4 +11310,10 @@ void FSequencer::RecompileDirtyDirectors()
 		}
 	}
 }
+
+void FSequencer::OnCurveModelDisplayChanged(FCurveModel *InCurveModel, bool bDisplayed)
+{
+	OnCurveDisplayChanged.Broadcast(InCurveModel, bDisplayed);
+}
+
 #undef LOCTEXT_NAMESPACE
diff --git a/Engine/Source/Editor/Sequencer/Private/Sequencer.h b/Engine/Source/Editor/Sequencer/Private/Sequencer.h
--- a/Engine/Source/Editor/Sequencer/Private/Sequencer.h
+++ b/Engine/Source/Editor/Sequencer/Private/Sequencer.h
@@ -37,6 +37,7 @@
 #include "AcquiredResources.h"
 #include "SequencerSettings.h"
 #include "Curves/RichCurve.h"
+#include "Sections/MovieScene3DTransformSection.h"
 
 class AActor;
 class ACameraActor;
@@ -431,6 +432,9 @@ public:
 	FReply JumpToPreviousKey();
 	FReply JumpToNextKey();
 
+	bool CanAddTransformKeysForSelectedObjects() const;
+	void OnAddTransformKeysForSelectedObjects(EMovieSceneTransformChannel Channel);
+
 	/** Get the visibility of the record button */
 	EVisibility GetRecordButtonVisibility() const;
 
@@ -720,6 +724,8 @@ public:
 	virtual void KeyProperty(FKeyPropertyParams KeyPropertyParams) override;
 	virtual FSequencerSelection& GetSelection() override;
 	virtual FSequencerSelectionPreview& GetSelectionPreview() override;
+	virtual void SuspendSelectionBroadcast() override;
+	virtual void ResumeSelectionBroadcast() override;
 	virtual void GetSelectedTracks(TArray<UMovieSceneTrack*>& OutSelectedTracks) override;
 	virtual void GetSelectedSections(TArray<UMovieSceneSection*>& OutSelectedSections) override;
 	virtual void GetSelectedFolders(TArray<UMovieSceneFolder*>& OutSelectedFolders) override;
@@ -729,7 +735,8 @@ public:
 	virtual void SelectTrack(UMovieSceneTrack* Track) override;
 	virtual void SelectSection(UMovieSceneSection* Section) override;
 	virtual void SelectByPropertyPaths(const TArray<FString>& InPropertyPaths) override;
-	virtual void SelectByKeyAreas(const TArray<IKeyArea>& InKeyAreas, bool bSelectParentInstead, bool bSelect) override;
+	virtual void SelectByKeyAreas(UMovieSceneSection* Section, const TArray<IKeyArea>& InKeyAreas, bool bSelectParentInstead, bool bSelect) override;
+	virtual void SelectByNthCategoryNode(UMovieSceneSection* Section, int Index, bool bSelect) override;
 	virtual void EmptySelection() override;
 	virtual void ThrobKeySelection() override;
 	virtual void ThrobSectionSelection() override;
@@ -743,10 +750,11 @@ public:
 	virtual FOnMovieSceneBindingsPasted& OnMovieSceneBindingsPasted() override { return OnMovieSceneBindingsPastedDelegate; }
 	virtual FOnSelectionChangedObjectGuids& GetSelectionChangedObjectGuids() override { return OnSelectionChangedObjectGuidsDelegate; }
 	virtual FOnSelectionChangedTracks& GetSelectionChangedTracks() override { return OnSelectionChangedTracksDelegate; }
+	virtual FOnCurveDisplayChanged& GetCurveDisplayChanged() override { return OnCurveDisplayChanged; }
 	virtual FOnSelectionChangedSections& GetSelectionChangedSections() override { return OnSelectionChangedSectionsDelegate; }
 	virtual FGuid CreateBinding(UObject& InObject, const FString& InName) override;
 	virtual UObject* GetPlaybackContext() const override;
-	virtual TArray<UObject*> GetEventContexts() const override;
+	virtual TArray<UObject*> GetEventContexts() const override; 
 	virtual FOnActorAddedToSequencer& OnActorAddedToSequencer() override;
 	virtual FOnPreSave& OnPreSave() override;
 	virtual FOnPostSave& OnPostSave() override;
@@ -761,11 +769,14 @@ public:
 	virtual FGuid MakeNewSpawnable(UObject& SourceObject, UActorFactory* ActorFactory = nullptr, bool bSetupDefaults = true) override;
 	virtual bool IsReadOnly() const override;
 	virtual void ExternalSelectionHasChanged() override { SynchronizeSequencerSelectionWithExternalSelection(); }
+	virtual void ObjectImplicitlyAdded(UObject* InObject) const override;
 	/** Access the user-supplied settings object */
 	virtual USequencerSettings* GetSequencerSettings() override { return Settings; }
 	virtual void SetSequencerSettings(USequencerSettings* InSettings) override { Settings = InSettings; }
 	virtual TSharedPtr<class ITimeSlider> GetTopTimeSliderWidget() const override;
 	virtual void ResetTimeController() override;
+	virtual void SetFilterOn(const FText& InName, bool bOn) override;
+
 
 public:
 
@@ -1044,6 +1055,7 @@ public:
 		return ScrubStyle;
 	}
 
+
 private:
 
 	/** Update the time bases for the current movie scene */
@@ -1217,6 +1229,9 @@ private:
 	/** A delegate which is called any time the sequencer selection changes. */
 	FOnSelectionChangedTracks OnSelectionChangedTracksDelegate;
 
+	/** A delegate which is called any time the sequencers curve eidtor selection changes. */
+	FOnCurveDisplayChanged OnCurveDisplayChanged;
+
 	/** A delegate which is called any time the sequencer selection changes. */
 	FOnSelectionChangedSections OnSelectionChangedSectionsDelegate;
 
@@ -1226,6 +1241,9 @@ private:
 	FOnPostSave OnPostSaveEvent;
 	FOnActivateSequence OnActivateSequenceEvent;
 
+	/** Delegate for Curve Display Changed Event from the Curve Editor, which we than pass to the FOnCurveDisplayChanged delegate */
+	void OnCurveModelDisplayChanged(FCurveModel *InCurveModel, bool bDisplayed);
+
 	int32 SilentModeCount;
 
 	/** When true the sequencer selection is being updated from changes to the external selection. */
diff --git a/Engine/Source/Editor/Sequencer/Private/SequencerCommands.cpp b/Engine/Source/Editor/Sequencer/Private/SequencerCommands.cpp
--- a/Engine/Source/Editor/Sequencer/Private/SequencerCommands.cpp
+++ b/Engine/Source/Editor/Sequencer/Private/SequencerCommands.cpp
@@ -15,10 +15,12 @@ void FSequencerCommands::RegisterCommands()
 	UI_COMMAND( Pause, "Pause", "Pause playback", EUserInterfaceActionType::Button, FInputChord(EKeys::K) );
 	UI_COMMAND( StepForward, "Step Forward", "Step the timeline forward", EUserInterfaceActionType::Button, FInputChord(EKeys::Right) );
 	UI_COMMAND( StepBackward, "Step Backward", "Step the timeline backward", EUserInterfaceActionType::Button, FInputChord(EKeys::Left) );
+	UI_COMMAND( StepForward2, "Step Forward", "Step the timeline forward", EUserInterfaceActionType::Button, FInputChord(EModifierKey::Alt, EKeys::Period) );
+	UI_COMMAND( StepBackward2, "Step Backward", "Step the timeline backward", EUserInterfaceActionType::Button, FInputChord(EModifierKey::Alt, EKeys::Comma) );
 	UI_COMMAND( StepToNextKey, "Step to Next Key", "Step to the next key", EUserInterfaceActionType::Button, FInputChord(EKeys::Period) );
 	UI_COMMAND( StepToPreviousKey, "Step to Previous Key", "Step to the previous key", EUserInterfaceActionType::Button, FInputChord(EKeys::Comma) );
-	UI_COMMAND( StepToNextCameraKey, "Step to Next Camera Key", "Step to the next camera key", EUserInterfaceActionType::Button, FInputChord(EModifierKey::Alt, EKeys::Period) );
-	UI_COMMAND( StepToPreviousCameraKey, "Step to Previous Camera Key", "Step to the previous camera key", EUserInterfaceActionType::Button, FInputChord(EModifierKey::Alt, EKeys::Comma) );
+	UI_COMMAND( StepToNextCameraKey, "Step to Next Camera Key", "Step to the next camera key", EUserInterfaceActionType::Button, FInputChord() );
+	UI_COMMAND( StepToPreviousCameraKey, "Step to Previous Camera Key", "Step to the previous camera key", EUserInterfaceActionType::Button, FInputChord() );
 	UI_COMMAND( StepToNextShot, "Step to Next Shot", "Step to the next shot", EUserInterfaceActionType::Button, FInputChord(EModifierKey::Shift, EKeys::Period) );
 	UI_COMMAND( StepToPreviousShot, "Step to Previous Shot", "Step to the previous shot", EUserInterfaceActionType::Button, FInputChord(EModifierKey::Shift, EKeys::Comma) );
 	UI_COMMAND( SetStartPlaybackRange, "Set Start Playback Range", "Set the start playback range", EUserInterfaceActionType::Button, FInputChord(EKeys::LeftBracket) );
@@ -151,6 +153,11 @@ void FSequencerCommands::RegisterCommands()
 
 	UI_COMMAND( QuickTreeSearch, "Quick Tree Search", "Jumps keyboard focus to the tree searchbox to allow searching for tracks in the current Sequence.", EUserInterfaceActionType::Button, FInputChord(EModifierKey::Control, EKeys::F));
 	
+	UI_COMMAND(AddTransformKey, "Add Transform Key", "Add a transform key at the current time for the selected actor.", EUserInterfaceActionType::Button, FInputChord(EKeys::S));
+	UI_COMMAND(AddTranslationKey, "Add Translation Key", "Add a translation key at the current time for the selected actor.", EUserInterfaceActionType::Button, FInputChord(EModifierKey::Shift, EKeys::W));
+	UI_COMMAND(AddRotationKey, "Add Rotation Key", "Add a rotation key at the current time for the selected actor.", EUserInterfaceActionType::Button, FInputChord(EModifierKey::Shift, EKeys::E));
+	UI_COMMAND(AddScaleKey, "Add Scale Key", "Add a scale key at the current time for the selected actor.", EUserInterfaceActionType::Button, FInputChord(EModifierKey::Shift, EKeys::R));
+
 }
 
 #undef LOCTEXT_NAMESPACE
diff --git a/Engine/Source/Editor/Sequencer/Private/SequencerNodeTree.cpp b/Engine/Source/Editor/Sequencer/Private/SequencerNodeTree.cpp
--- a/Engine/Source/Editor/Sequencer/Private/SequencerNodeTree.cpp
+++ b/Engine/Source/Editor/Sequencer/Private/SequencerNodeTree.cpp
@@ -19,6 +19,7 @@
 #include "CurveEditor.h"
 #include "SequencerNodeSortingMethods.h"
 #include "SequencerTrackFilters.h"
+#include "Channels/MovieSceneChannel.h"
 
 FSequencerNodeTree::~FSequencerNodeTree()
 {
@@ -928,7 +929,24 @@ static void FilterNodesRecursive( FSequencer& Sequencer, const TSharedRef<FSeque
 		case ESequencerNode::Track:
 		{
 			UMovieSceneTrack* Track = static_cast<const FSequencerTrackNode&>(StartNode.Get()).GetTrack();
-			if (Filters->Num() == 0 || Filters->PassesAnyFilters(Track))
+			TSharedPtr<FSequencerSectionKeyAreaNode> TopLevelKeyArea = static_cast<const FSequencerTrackNode&>(StartNode.Get()).GetTopLevelKeyNode();
+			if (TopLevelKeyArea.IsValid())
+			{
+				for (const TSharedRef<IKeyArea>& KeyArea : TopLevelKeyArea->GetAllKeyAreas())
+				{
+					FMovieSceneChannel* Channel = KeyArea->ResolveChannel();
+					if (Channel)
+					{
+						if (Filters->Num() == 0 || Filters->PassesAnyFilters(Channel))
+						{
+							bPasssedAnyFilters = true;
+							break;
+						}
+					}
+				}
+			}
+		
+			if (bPasssedAnyFilters || Filters->Num() == 0 || Filters->PassesAnyFilters(Track, StartNode->GetDisplayName()))
 			{
 				bPasssedAnyFilters = true;
 
@@ -993,7 +1011,7 @@ static void FilterNodesRecursive( FSequencer& Sequencer, const TSharedRef<FSeque
 			const FSequencerObjectBindingNode ObjectNode = static_cast<const FSequencerObjectBindingNode&>(StartNode.Get());
 			for (TWeakObjectPtr<>& Object : Sequencer.FindObjectsInCurrentSequence(ObjectNode.GetObjectBinding()))
 			{
-				if (Object.IsValid() && (Filters->Num() == 0 || Filters->PassesAnyFilters(Object.Get()))
+				if (Object.IsValid() && (Filters->Num() == 0 || Filters->PassesAnyFilters(Object.Get(), StartNode->GetDisplayName()))
 					&& LevelTrackFilter->PassesFilter(Object.Get()))
 				{
 					bPasssedAnyFilters = true;
@@ -1013,8 +1031,85 @@ static void FilterNodesRecursive( FSequencer& Sequencer, const TSharedRef<FSeque
 			break;
 		}
 		case ESequencerNode::Category:
-		case ESequencerNode::Folder:
+		{
+			if (TSharedPtr<FSequencerTrackNode> TrackNode = StartNode->FindParentTrackNode())
+			{
+				UMovieSceneTrack* Track = TrackNode->GetTrack();
+				if (Filters->Num() == 0 || Filters->PassesAnyFilters(Track, StartNode->GetDisplayName()))
+				{
+					bPasssedAnyFilters = true;
+				}
+			}
+			if (bPasssedAnyFilters && Sequencer.GetSequencerSettings()->GetShowSelectedNodesOnly())
+			{
+				TSharedPtr<const FSequencerDisplayNode> ParentNode = StartNode->GetParent();
+				while (ParentNode.IsValid())
+				{
+					// Pinned tracks should be visible whether selected or not
+					if (ParentNode->IsPinned())
+					{
+						break;
+					}
+
+					if (ParentNode->GetType() == ESequencerNode::Object)
+					{
+						const FSequencerObjectBindingNode* ObjectNode = static_cast<const FSequencerObjectBindingNode*>(ParentNode.Get());
+						const FMovieSceneBinding* Binding = Sequencer.GetFocusedMovieSceneSequence()->GetMovieScene()->FindBinding(ObjectNode->GetObjectBinding());
+						if (!(Binding && Sequencer.IsBindingVisible(*Binding)))
+						{
+							bPasssedAnyFilters = false;
+						}
+
+						break;
+					}
+					ParentNode = ParentNode->GetParent();
+				}
+			}
+			break;
+		}
 		case ESequencerNode::KeyArea:
+		{
+			const FSequencerSectionKeyAreaNode KeyAreaNode = static_cast<const FSequencerSectionKeyAreaNode&>(StartNode.Get());
+			for (const TSharedRef<IKeyArea>& KeyArea : KeyAreaNode.GetAllKeyAreas())
+			{
+				FMovieSceneChannel* Channel = KeyArea->ResolveChannel();
+				if (Channel)
+				{
+					if (Filters->Num() == 0 || Filters->PassesAnyFilters(Channel))
+					{
+						bPasssedAnyFilters = true;
+						break;
+					}
+				}
+			}
+			if (bPasssedAnyFilters && Sequencer.GetSequencerSettings()->GetShowSelectedNodesOnly())
+			{
+				TSharedPtr<const FSequencerDisplayNode> ParentNode = StartNode->GetParent();
+				while (ParentNode.IsValid())
+				{
+					// Pinned tracks should be visible whether selected or not
+					if (ParentNode->IsPinned())
+					{
+						break;
+					}
+
+					if (ParentNode->GetType() == ESequencerNode::Object)
+					{
+						const FSequencerObjectBindingNode* ObjectNode = static_cast<const FSequencerObjectBindingNode*>(ParentNode.Get());
+						const FMovieSceneBinding* Binding = Sequencer.GetFocusedMovieSceneSequence()->GetMovieScene()->FindBinding(ObjectNode->GetObjectBinding());
+						if (!(Binding && Sequencer.IsBindingVisible(*Binding)))
+						{
+							bPasssedAnyFilters = false;
+						}
+
+						break;
+					}
+					ParentNode = ParentNode->GetParent();
+				}
+			}
+			break;
+		}
+		case ESequencerNode::Folder:
 		{
 			// If no active filters, pass organizational nodes to be text filtered
 			if (Filters->Num() == 0)
@@ -1025,7 +1120,6 @@ static void FilterNodesRecursive( FSequencer& Sequencer, const TSharedRef<FSeque
 		}
 	}
 
-
 	if (bPasssedAnyFilters)
 	{
 		// If we have a filter string, make sure we match
@@ -1175,7 +1269,7 @@ void FSequencerNodeTree::UpdateCurveEditorTree()
 	for (auto It = CurveEditorTreeItemIDs.CreateIterator(); It; ++It)
 	{
 		TSharedPtr<FSequencerDisplayNode> Node = It->Key.Pin();
-		if (!Node.IsValid() || Node->TreeSerialNumber != SerialNumber)
+		if (!Node.IsValid() || Node->TreeSerialNumber != SerialNumber || !Node->IsVisible())
 		{
 			CurveEditor->RemoveTreeItem(It->Value);
 			It.RemoveCurrent();
diff --git a/Engine/Source/Editor/Sequencer/Private/SequencerTrackFilters.h b/Engine/Source/Editor/Sequencer/Private/SequencerTrackFilters.h
--- a/Engine/Source/Editor/Sequencer/Private/SequencerTrackFilters.h
+++ b/Engine/Source/Editor/Sequencer/Private/SequencerTrackFilters.h
@@ -4,7 +4,6 @@
 
 #include "SequencerTrackFilterBase.h"
 #include "Misc/IFilter.h"
-#include "Misc/FilterCollection.h"
 #include "CollectionManagerTypes.h"
 #include "Styling/SlateIconFinder.h"
 #include "EditorStyleSet.h"
@@ -21,25 +20,28 @@
 #include "Components/SkeletalMeshComponent.h"
 #include "Particles/ParticleSystem.h"
 
+#include "Channels/MovieSceneChannel.h"
+
 class UWorld;
 
 #define LOCTEXT_NAMESPACE "Sequencer"
 
-class FSequencerTrackFilterCollection : public TFilterCollection<FTrackFilterType>
+class FSequencerTrackFilterCollection : public TSharedFromThis<FSequencerTrackFilterCollection>
 {
 public:
 	/**
 	 *	Returns whether the specified Item passes any of the filters in the collection
 	 *
 	 *	@param	InItem	The Item to check against all child filter restrictions
+	 *  @param  InText  The Items Displayed Name
 	 *	@return			Whether the Item passed any child filter restrictions
 	 */
 	 // @todo Maybe this should get moved in to TFilterCollection
-	bool PassesAnyFilters(/*ItemType*/ FTrackFilterType InItem) const
+	bool PassesAnyFilters(/*ItemType*/ FTrackFilterType InItem, const FText& InText) const
 	{
-		for (int Index = 0; Index < ChildFilters.Num(); Index++)
+		for (int32 Index = 0; Index < ChildFilters.Num(); Index++)
 		{
-			if (ChildFilters[Index]->PassesFilter(InItem))
+			if (ChildFilters[Index]->PassesFilterWithDisplayName(InItem, InText))
 			{
 				return true;
 			}
@@ -47,11 +49,27 @@ public:
 
 		return false;
 	}
-
+	/**
+ *	Returns whether the specified MovieSceneChannel  passes any of the filters in the collection
+ *
+ *	@param	InChannel	The Item to check against all child filter restrictions
+ *	@return			Whether the Item passed any child filter restrictions
+ */
+	bool PassesAnyFilters(FMovieSceneChannel* InChannel) const
+	{
+		for (int32 Index = 0; Index < ChildFilters.Num(); Index++)
+		{
+			if (ChildFilters[Index]->PassesFilterChannel(InChannel))
+			{
+				return true;
+			}
+		}
+		return false;
+	}
 	// @todo Maybe this should get moved in to TFilterCollection
-	bool Contains(const TSharedPtr< IFilter< FTrackFilterType > >& InItem) const
+	bool Contains(const TSharedPtr< FSequencerTrackFilter >& InItem) const
 	{
-		for (const TSharedPtr<IFilter<FTrackFilterType>>& Filter : ChildFilters)
+		for (const TSharedPtr<FSequencerTrackFilter>& Filter : ChildFilters)
 		{
 			if (InItem == Filter)
 			{
@@ -66,7 +84,7 @@ public:
 	{
 		for (auto Iterator = ChildFilters.CreateIterator(); Iterator; ++Iterator)
 		{
-			const TSharedPtr< IFilter< FTrackFilterType > >& Filter = *Iterator;
+			const TSharedPtr< FSequencerTrackFilter >& Filter = *Iterator;
 
 			if (Filter.IsValid())
 			{
@@ -78,6 +96,112 @@ public:
 
 		ChangedEvent.Broadcast();
 	}
+
+	/**
+	 * TFilterCollection destructor. Unregisters from all child filter changed delegates.
+	 */
+	~FSequencerTrackFilterCollection()
+	{
+		for (auto Iterator = ChildFilters.CreateIterator(); Iterator; ++Iterator)
+		{
+			const TSharedPtr< FSequencerTrackFilter>& Filter = *Iterator;
+
+			if (Filter.IsValid())
+			{
+				Filter->OnChanged().RemoveAll(this);
+			}
+		}
+	}
+
+	/**
+	 *	Adds the specified Filter to the collection
+	 *
+	 *	@param	Filter	The filter object to add to the collection
+	 *	@return			The index in the collection at which the filter was added
+	 */
+	int32 Add(const TSharedPtr<FSequencerTrackFilter >& Filter)
+	{
+		int32 ExistingIdx;
+		if (ChildFilters.Find(Filter, ExistingIdx))
+		{
+			// The filter already exists, don't add a new one but return the index where it was found.
+			return ExistingIdx;
+		}
+
+		if (Filter.IsValid())
+		{
+			Filter->OnChanged().AddSP(this, &FSequencerTrackFilterCollection::OnChildFilterChanged);
+		}
+
+		int32 Result = ChildFilters.Add(Filter);
+		ChangedEvent.Broadcast();
+
+		return Result;
+	}
+
+	/**
+	 *	Removes as many instances of the specified Filter as there are in the collection
+	 *
+	 *	@param	Filter	The filter object to remove from the collection
+	 *	@return			The number of Filters removed from the collection
+	 */
+	int32 Remove(const TSharedPtr<FSequencerTrackFilter >& Filter)
+	{
+		if (Filter.IsValid())
+		{
+			Filter->OnChanged().RemoveAll(this);
+		}
+
+		int32 Result = ChildFilters.Remove(Filter);
+
+		// Don't broadcast if the collection didn't change
+		if (Result > 0)
+		{
+			ChangedEvent.Broadcast();
+		}
+
+		return Result;
+	}
+
+	/**
+	 *	Gets the filter at the specified index
+	 *
+	 *	@param	Index	The index of the requested filter in the ChildFilters array.
+	 *	@return			Filter at the specified index
+	 */
+	TSharedPtr< FSequencerTrackFilter> GetFilterAtIndex(int32 Index)
+	{
+		return ChildFilters[Index];
+	}
+
+	/** Returns the number of Filters in the collection */
+	FORCEINLINE int32 Num() const
+	{
+		return ChildFilters.Num();
+	}
+
+
+	/** Broadcasts anytime the restrictions of any of the child Filters change */
+	DECLARE_EVENT(FSequencerTrackFilterCollection, FChangedEvent);
+	FChangedEvent& OnChanged() { return ChangedEvent; }
+
+protected:
+
+	/**
+	 *	Called when a child Filter restrictions change and broadcasts the FilterChanged delegate
+	 *	for the collection
+	 */
+	void OnChildFilterChanged()
+	{
+		ChangedEvent.Broadcast();
+	}
+
+	/** The array of child filters */
+	TArray< TSharedPtr< FSequencerTrackFilter > > ChildFilters;
+
+	/**	Fires whenever any filter in the collection changes */
+	FChangedEvent ChangedEvent;
+
 };
 
 class FSequencerTrackFilter_AudioTracks : public FSequencerTrackFilter_ClassType< UMovieSceneAudioTrack >
@@ -183,4 +307,23 @@ private:
 	TWeakObjectPtr<UWorld> CachedWorld;
 };
 
+
+class FSequencerTrackFilter_Animated : public FSequencerTrackFilter
+{
+	virtual FString GetName() const override { return TEXT("AnimatedFilter"); }
+	virtual FText GetDisplayName() const override { return LOCTEXT("SequenceTrackFilter_Animated", "Animated Tracks"); }
+	virtual FText GetToolTipText() const override { return LOCTEXT("SequencerTrackFilter_AnimatedTip", "Show Only Animated Tracks."); }
+	virtual FSlateIcon GetIcon() const { return FSlateIcon(FEditorStyle::GetStyleSetName(), "Sequencer.IconKeyUser"); }
+
+	virtual bool PassesFilter(FTrackFilterType InItem) const override
+	{
+		return false;
+	}
+
+	virtual bool PassesFilterChannel(FMovieSceneChannel* InMovieSceneChannel) const override
+	{
+		return (InMovieSceneChannel && InMovieSceneChannel->GetNumKeys() > 0);
+	}
+
+};
 #undef LOCTEXT_NAMESPACE
\ No newline at end of file
diff --git a/Engine/Source/Editor/Sequencer/Public/ISequencer.h b/Engine/Source/Editor/Sequencer/Public/ISequencer.h
--- a/Engine/Source/Editor/Sequencer/Public/ISequencer.h
+++ b/Engine/Source/Editor/Sequencer/Public/ISequencer.h
@@ -30,6 +30,7 @@ class UMovieSceneSubSection;
 class IDetailsView;
 class IKeyArea;
 enum class EMapChangeType : uint8;
+class FCurveModel;
 
 /**
  * Defines auto change modes.
@@ -142,7 +143,6 @@ enum class EMovieSceneDataChangeType
 	Unknown
 };
 
-
 /**
  * Interface for sequencers.
  */
@@ -167,6 +167,9 @@ public:
 
 	DECLARE_MULTICAST_DELEGATE_OneParam(FOnSelectionChangedSections, TArray<UMovieSceneSection*> /*Sections*/);
 
+	DECLARE_MULTICAST_DELEGATE_TwoParams(FOnCurveDisplayChanged, FCurveModel* , bool /*displayed*/);
+
+
 	DECLARE_MULTICAST_DELEGATE_OneParam(FOnCloseEvent, TSharedRef<ISequencer>);
 
 	DECLARE_MULTICAST_DELEGATE_TwoParams(FOnActorAddedToSequencer, AActor*, const FGuid);
@@ -468,6 +471,9 @@ public:
 	virtual FSequencerSelection& GetSelection() = 0;
 	virtual FSequencerSelectionPreview& GetSelectionPreview() = 0;
 
+	virtual void SuspendSelectionBroadcast() = 0;
+	virtual void ResumeSelectionBroadcast() = 0;
+
 	/** Gets the currently selected tracks. */
 	virtual void GetSelectedTracks(TArray<UMovieSceneTrack*>& OutSelectedTracks) = 0;
 
@@ -496,7 +502,10 @@ public:
 	virtual void SelectByPropertyPaths(const TArray<FString>& InPropertyPaths) = 0;
 
 	/** Selects nodes by key areas */
-	virtual void SelectByKeyAreas(const TArray<IKeyArea>& InKeyAreas, bool bSelectParentInstead, bool bSelect) = 0;
+	virtual void SelectByKeyAreas(UMovieSceneSection* Section, const TArray<IKeyArea>& InKeyAreas, bool bSelectParentInstead, bool bSelect) = 0;
+
+	/** Selects nodes by the nth category node under a section */
+	virtual void SelectByNthCategoryNode(UMovieSceneSection* Section, int Index, bool bSelect) = 0;
 
 	/** Empties the current selection. */
 	virtual void EmptySelection() = 0;
@@ -538,6 +547,9 @@ public:
 	/** Gets a multicast delegate with an array of UMovieSceneSections which is called when the outliner node selection changes. */
 	virtual FOnSelectionChangedSections& GetSelectionChangedSections() = 0;
 
+	/** Gets a multicast delegate when the curve edtior associated with this sequencer has it's selection change. */
+	virtual FOnCurveDisplayChanged& GetCurveDisplayChanged() = 0;
+
 	/** @return a numeric type interface that will parse and display numbers as frames and times correctly */
 	virtual TSharedRef<INumericTypeInterface<double>> GetNumericTypeInterface() const = 0;
 
@@ -604,7 +616,21 @@ public:
 	* @see SetSelectionRange, SetSelectionRangeEnd, SetSelectionRangeStart
 	*/
 	virtual TRange<FFrameNumber> GetSelectionRange() const = 0;
+public:
 
+	/**
+	* Specify that an object was implicitly added. We will notify the track editors that it was
+	@InObject Object that was added to be part of a track/binding but not the real binding
+	*/
+	virtual void ObjectImplicitlyAdded(UObject* InObject) const = 0;
+
+public:
+	/**
+	*    Turn on/off the filter with the specified name
+	* @InName The name of the of the filter
+	* @bOn   Whether or not the filter is on or off.
+	*/
+	virtual void SetFilterOn(const FText& InName, bool bOn) = 0;
 public:
 
 	/**
diff --git a/Engine/Source/Editor/Sequencer/Public/ISequencerTrackEditor.h b/Engine/Source/Editor/Sequencer/Public/ISequencerTrackEditor.h
--- a/Engine/Source/Editor/Sequencer/Public/ISequencerTrackEditor.h
+++ b/Engine/Source/Editor/Sequencer/Public/ISequencerTrackEditor.h
@@ -9,6 +9,7 @@
 #include "Framework/Commands/UICommandList.h"
 #include "ISequencerSection.h"
 #include "MovieSceneTrack.h"
+#include "Sections/MovieScene3DTransformSection.h"
 
 class FExtender;
 class FMenuBuilder;
@@ -92,6 +93,17 @@ public:
 	 */
 	virtual void ExtendObjectBindingTrackMenu(TSharedRef<FExtender> Extender, const TArray<FGuid>& ObjectBindings, const UClass* ObjectClass) { }
 
+
+	/**
+	 * Builds up the object binding cpmtext menu for the outliner.
+	 *
+	 * @param MenuBuilder The menu builder to change.
+	 * @param ObjectBinding The object binding this is for.
+	 * @param ObjectClass The class of the object this is for.
+	 */
+	virtual void BuildObjectBindingContextMenu(FMenuBuilder& MenuBuilder, const TArray<FGuid>& ObjectBindings, const UClass* ObjectClass) {}
+
+
 	/**
 	 * Builds an edit widget for the outliner nodes which represent tracks which are edited by this editor.
 	 * @param ObjectBinding The object binding associated with the track being edited by this editor.
@@ -207,6 +219,34 @@ public:
 		return false;
 	}
 
+	/**
+	 * @return If it supports supports transform key bindings for setting keys.
+	 */
+	virtual bool HasTransformKeyBindings() const { return false; }
+
+	/** Whether or not we can add a transform key for a selected object
+	* @return Returns true if we can.
+	**/
+	virtual bool CanAddTransformKeysForSelectedObjects() const { return false; }
+
+	/**
+	* Adds transform tracks and keys to the selected objects in the level.
+	*
+	* @param Channel The transform channel to add keys for.
+	*/
+	virtual void OnAddTransformKeysForSelectedObjects(EMovieSceneTransformChannel Channel) {};
+
+	/**
+	* If true this track has priority when setting transform keys and should be the only one to set them
+	*
+	*/
+	virtual bool HasTransformKeyOverridePriority() const { return false; }
+
+	/**
+	* Handle this object being implicitly added
+	*/
+	virtual void ObjectImplicitlyAdded(UObject* InObject)  {}
+
 public:
 
 	/** Virtual destructor. */
diff --git a/Engine/Source/Editor/Sequencer/Public/SequencerCommands.h b/Engine/Source/Editor/Sequencer/Public/SequencerCommands.h
--- a/Engine/Source/Editor/Sequencer/Public/SequencerCommands.h
+++ b/Engine/Source/Editor/Sequencer/Public/SequencerCommands.h
@@ -46,6 +46,12 @@ public:
 	/** Step backward */
 	TSharedPtr< FUICommandInfo > StepBackward;
 
+	/** Step forward */
+	TSharedPtr< FUICommandInfo > StepForward2;
+
+	/** Step backward */
+	TSharedPtr< FUICommandInfo > StepBackward2;
+
 	/** Step to next key */
 	TSharedPtr< FUICommandInfo > StepToNextKey;
 
@@ -355,6 +361,20 @@ public:
 	/** Toggle whether we should evaluate sub sequences in isolation */
 	TSharedPtr< FUICommandInfo > ToggleEvaluateSubSequencesInIsolation;
 
+	/** Sets a transform key at the current time for the selected actor */
+	TSharedPtr< FUICommandInfo > AddTransformKey;
+
+	/** Sets a translation key at the current time for the selected actor */
+	TSharedPtr< FUICommandInfo > AddTranslationKey;
+
+	/** Sets a rotation key at the current time for the selected actor */
+	TSharedPtr< FUICommandInfo > AddRotationKey;
+
+	/** Sets a scale key at the current time for the selected actor */
+	TSharedPtr< FUICommandInfo > AddScaleKey;
+
+
+
 	/**
 	 * Initialize commands
 	 */
diff --git a/Engine/Source/Editor/Sequencer/Public/SequencerTrackFilterBase.h b/Engine/Source/Editor/Sequencer/Public/SequencerTrackFilterBase.h
--- a/Engine/Source/Editor/Sequencer/Public/SequencerTrackFilterBase.h
+++ b/Engine/Source/Editor/Sequencer/Public/SequencerTrackFilterBase.h
@@ -12,6 +12,8 @@ class FMenuBuilder;
 
 typedef const UObject* FTrackFilterType;
 
+struct FMovieSceneChannel;
+
 class FSequencerTrackFilter : public IFilter<FTrackFilterType>
 {
 public:
@@ -38,6 +40,14 @@ public:
 	DECLARE_DERIVED_EVENT(FSequencerTrackFilter, IFilter<FTrackFilterType>::FChangedEvent, FChangedEvent);
 	virtual FChangedEvent& OnChanged() override { return ChangedEvent; }
 
+	/** Returns whether the specified Item passes the Filter's restrictions */
+	virtual bool PassesFilterWithDisplayName(FTrackFilterType InItem, const FText& InText) const 
+	{
+		return PassesFilter(InItem);
+	}
+
+	virtual bool PassesFilterChannel(FMovieSceneChannel* InMovieSceneChannel) const { return false; }
+
 protected:
 	void BroadcastChangedEvent() const { ChangedEvent.Broadcast(); }
 
diff --git a/Engine/Source/Editor/SequencerWidgets/Private/STimeRange.cpp b/Engine/Source/Editor/SequencerWidgets/Private/STimeRange.cpp
--- a/Engine/Source/Editor/SequencerWidgets/Private/STimeRange.cpp
+++ b/Engine/Source/Editor/SequencerWidgets/Private/STimeRange.cpp
@@ -20,6 +20,7 @@ void STimeRange::Construct( const STimeRange::FArguments& InArgs, TSharedRef<ITi
 	if (InArgs._ShowWorkingRange)
 	{
 		WorkingRangeStart = SNew(SSpinBox<double>)
+		.IsEnabled(InArgs._EnableWorkingRange)
 		.Value(this, &STimeRange::WorkingStartTime)
 		.ToolTipText(LOCTEXT("WorkingRangeStart", "Working Range Start"))
 		.OnValueCommitted(this, &STimeRange::OnWorkingStartTimeCommitted)
@@ -33,6 +34,7 @@ void STimeRange::Construct( const STimeRange::FArguments& InArgs, TSharedRef<ITi
 		.LinearDeltaSensitivity(25);
 
 		WorkingRangeEnd = SNew(SSpinBox<double>)
+		.IsEnabled(InArgs._EnableWorkingRange)
 		.Value(this, &STimeRange::WorkingEndTime)
 		.ToolTipText(LOCTEXT("WorkingRangeEnd", "Working Range End"))
 		.OnValueCommitted( this, &STimeRange::OnWorkingEndTimeCommitted )
@@ -50,6 +52,7 @@ void STimeRange::Construct( const STimeRange::FArguments& InArgs, TSharedRef<ITi
 	if (InArgs._ShowViewRange)
 	{
 		ViewRangeStart = SNew(SSpinBox<double>)
+		.IsEnabled(InArgs._EnableViewRange)
 		.Value(this, &STimeRange::ViewStartTime)
 		.ToolTipText(LOCTEXT("ViewStartTimeTooltip", "View Range Start Time"))
 		.OnValueCommitted( this, &STimeRange::OnViewStartTimeCommitted )
@@ -64,6 +67,7 @@ void STimeRange::Construct( const STimeRange::FArguments& InArgs, TSharedRef<ITi
 
 
 		ViewRangeEnd = SNew(SSpinBox<double>)
+		.IsEnabled(InArgs._EnableViewRange)
 		.Value(this, &STimeRange::ViewEndTime)
 		.ToolTipText(LOCTEXT("ViewEndTimeTooltip", "View Range End Time"))
 		.OnValueCommitted( this, &STimeRange::OnViewEndTimeCommitted )
@@ -78,8 +82,10 @@ void STimeRange::Construct( const STimeRange::FArguments& InArgs, TSharedRef<ITi
 	}
 
 	TSharedRef<SWidget> PlaybackRangeStart = SNullWidget::NullWidget, PlaybackRangeEnd = SNullWidget::NullWidget;
+	if (InArgs._ShowPlaybackRange)
 	{
 		PlaybackRangeStart = SNew(SSpinBox<double>)
+		.IsEnabled(InArgs._EnablePlaybackRange)
 		.Value(this, &STimeRange::PlayStartTime)
 		.ToolTipText(LOCTEXT("PlayStartTimeTooltip", "Playback Range Start Time"))
 		.OnValueCommitted(this, &STimeRange::OnPlayStartTimeCommitted)
@@ -94,6 +100,7 @@ void STimeRange::Construct( const STimeRange::FArguments& InArgs, TSharedRef<ITi
 
 
 		PlaybackRangeEnd = SNew(SSpinBox<double>)
+		.IsEnabled(InArgs._EnablePlaybackRange)
 		.Value(this, &STimeRange::PlayEndTime)
 		.ToolTipText(LOCTEXT("PlayEndTimeTooltip", "Playback Range Stop Time"))
 		.OnValueCommitted( this, &STimeRange::OnPlayEndTimeCommitted )
diff --git a/Engine/Source/Editor/SequencerWidgets/Private/STimeRange.h b/Engine/Source/Editor/SequencerWidgets/Private/STimeRange.h
--- a/Engine/Source/Editor/SequencerWidgets/Private/STimeRange.h
+++ b/Engine/Source/Editor/SequencerWidgets/Private/STimeRange.h
@@ -13,7 +13,7 @@ class STimeRange : public ITimeSlider
 {
 public:
 	SLATE_BEGIN_ARGS(STimeRange)
-		: _ShowWorkingRange(true), _ShowViewRange(false), _ShowPlaybackRange(false)
+		: _ShowWorkingRange(true), _ShowViewRange(false), _ShowPlaybackRange(true)
 	{}
 		/** Whether to show the working range */
 		SLATE_ARGUMENT( bool, ShowWorkingRange )
@@ -21,6 +21,12 @@ public:
 		SLATE_ARGUMENT( bool, ShowViewRange )
 		/** Whether to show the playback range */
 		SLATE_ARGUMENT( bool, ShowPlaybackRange )
+		/** Whether to enable the working range */
+		SLATE_ARGUMENT( bool, EnableWorkingRange )
+		/** Whether to enable the view range */
+		SLATE_ARGUMENT( bool, EnableViewRange )
+		/** Whether to enable the playback range */
+		SLATE_ARGUMENT( bool, EnablePlaybackRange )
 		/* Content to display inside the time range */
 		SLATE_DEFAULT_SLOT( FArguments, CenterContent )
 	SLATE_END_ARGS()
diff --git a/Engine/Source/Editor/SequencerWidgets/Private/SequencerWidgetsModule.cpp b/Engine/Source/Editor/SequencerWidgets/Private/SequencerWidgetsModule.cpp
--- a/Engine/Source/Editor/SequencerWidgets/Private/SequencerWidgetsModule.cpp
+++ b/Engine/Source/Editor/SequencerWidgets/Private/SequencerWidgetsModule.cpp
@@ -47,6 +47,9 @@ public:
 		.ShowWorkingRange(!!(InArgs.ShowRanges & EShowRange::WorkingRange))
 		.ShowViewRange(!!(InArgs.ShowRanges & EShowRange::ViewRange))
 		.ShowPlaybackRange(!!(InArgs.ShowRanges & EShowRange::PlaybackRange))
+		.EnableWorkingRange(!!(InArgs.EnableRanges & EShowRange::WorkingRange))
+		.EnableViewRange(!!(InArgs.EnableRanges & EShowRange::ViewRange))
+		.EnablePlaybackRange(!!(InArgs.EnableRanges & EShowRange::PlaybackRange))
 		[
 			Content
 		];
diff --git a/Engine/Source/Editor/SequencerWidgets/Public/ISequencerWidgetsModule.h b/Engine/Source/Editor/SequencerWidgets/Public/ISequencerWidgetsModule.h
--- a/Engine/Source/Editor/SequencerWidgets/Public/ISequencerWidgetsModule.h
+++ b/Engine/Source/Editor/SequencerWidgets/Public/ISequencerWidgetsModule.h
@@ -32,14 +32,32 @@ struct FTimeRangeArgs
 		TSharedRef<INumericTypeInterface<double>> InNumericTypeInterface
 		)
 		: ShowRanges(InShowRanges)
+		, EnableRanges(InShowRanges)
 		, Controller(InController)
 		, VisibilityDelegate(InVisibilityDelegate)
 		, NumericTypeInterface(InNumericTypeInterface)
 	{
 	}
-	
+
+	FTimeRangeArgs(
+		EShowRange InShowRanges,
+		EShowRange InEnableRanges,
+		TSharedRef<ITimeSliderController> InController,
+		TAttribute<EVisibility> InVisibilityDelegate,
+		TSharedRef<INumericTypeInterface<double>> InNumericTypeInterface
+		)
+		: ShowRanges(InShowRanges)
+		, EnableRanges(InEnableRanges)
+		, Controller(InController)
+		, VisibilityDelegate(InVisibilityDelegate)
+		, NumericTypeInterface(InNumericTypeInterface)
+	{
+	}
+
 	/** Which ranges to show */
 	EShowRange ShowRanges;
+	/** Which ranges to enable editing on */
+	EShowRange EnableRanges;
 	/** A time slider controller */
 	TSharedRef<ITimeSliderController> Controller;
 	/** Visibility delegate */
diff --git a/Engine/Source/Editor/SkeletonEditor/Private/EditableSkeleton.cpp b/Engine/Source/Editor/SkeletonEditor/Private/EditableSkeleton.cpp
--- a/Engine/Source/Editor/SkeletonEditor/Private/EditableSkeleton.cpp
+++ b/Engine/Source/Editor/SkeletonEditor/Private/EditableSkeleton.cpp
@@ -298,8 +298,7 @@ void FEditableSkeleton::RenameSmartname(const FName InContainerName, SmartName::
 						{
 							SequencesToRecompress.Add(Seq);
 
-							Seq->CompressedData.CompressedCurveByteStream.Empty();
-							Seq->CompressedData.CurveCompressionCodec = nullptr;
+							Seq->ClearCompressedCurveData();
 						}
 					}
 				}
@@ -1595,11 +1594,27 @@ void FEditableSkeleton::UpdateSkeletonReferencePose(USkeletalMesh* InSkeletalMes
 	Skeleton->UpdateReferencePoseFromMesh(InSkeletalMesh);
 }
 
+void FEditableSkeleton::RegisterSlotNode(const FName& InSlotName)
+{
+	if(!Skeleton->ContainsSlotName(InSlotName))
+	{
+		const FScopedTransaction Transaction(LOCTEXT("RegisterSlotNode", "Register Slot Node"));
+		Skeleton->Modify();
+
+		if(Skeleton->RegisterSlotNode(InSlotName))
+		{
+			OnSlotsChanged.Broadcast();
+		}
+	}
+}
+
 bool FEditableSkeleton::AddSlotGroupName(const FName& InSlotName)
 {
 	const FScopedTransaction Transaction(LOCTEXT("AddSlotGroupName", "Add Slot Group Name"));
 	Skeleton->Modify();
-	return Skeleton->AddSlotGroupName(InSlotName);
+	bool bResult = Skeleton->AddSlotGroupName(InSlotName);
+	OnSlotsChanged.Broadcast();
+	return bResult;
 }
 
 void FEditableSkeleton::SetSlotGroupName(const FName& InSlotName, const FName& InGroupName)
@@ -1607,6 +1622,8 @@ void FEditableSkeleton::SetSlotGroupName(const FName& InSlotName, const FName& I
 	const FScopedTransaction Transaction(LOCTEXT("SetSlotGroupName", "Set Slot Group Name"));
 	Skeleton->Modify();
 	Skeleton->SetSlotGroupName(InSlotName, InGroupName);
+
+	OnSlotsChanged.Broadcast();
 }
 
 void FEditableSkeleton::DeleteSlotName(const FName& InSlotName)
@@ -1614,6 +1631,8 @@ void FEditableSkeleton::DeleteSlotName(const FName& InSlotName)
 	const FScopedTransaction Transaction(LOCTEXT("DeleteSlotName", "Delete Slot Name"));
 	Skeleton->Modify();
 	Skeleton->RemoveSlotName(InSlotName);
+
+	OnSlotsChanged.Broadcast();
 }
 
 void FEditableSkeleton::DeleteSlotGroup(const FName& InGroupName)
@@ -1621,6 +1640,8 @@ void FEditableSkeleton::DeleteSlotGroup(const FName& InGroupName)
 	const FScopedTransaction Transaction(LOCTEXT("DeleteSlotGroup", "Delete Slot Group"));
 	Skeleton->Modify();
 	Skeleton->RemoveSlotGroup(InGroupName);
+
+	OnSlotsChanged.Broadcast();
 }
 
 void FEditableSkeleton::RenameSlotName(const FName InOldSlotName, const FName InNewSlotName)
@@ -1628,6 +1649,8 @@ void FEditableSkeleton::RenameSlotName(const FName InOldSlotName, const FName In
 	const FScopedTransaction Transaction(LOCTEXT("RenameSlotName", "Rename Slot Name"));
 	Skeleton->Modify();
 	Skeleton->RenameSlotName(InOldSlotName, InNewSlotName);
+
+	OnSlotsChanged.Broadcast();
 }
 
 FDelegateHandle FEditableSkeleton::RegisterOnSmartNameChanged(const FOnSmartNameChanged::FDelegate& InOnSmartNameChanged)
@@ -1640,4 +1663,14 @@ void FEditableSkeleton::UnregisterOnSmartNameChanged(FDelegateHandle InHandle)
 	OnSmartNameChanged.Remove(InHandle);
 }
 
+FDelegateHandle FEditableSkeleton::RegisterOnSlotsChanged(const FSimpleMulticastDelegate::FDelegate& InOnSlotsChanged)
+{
+	return OnSlotsChanged.Add(InOnSlotsChanged);
+}
+
+void FEditableSkeleton::UnregisterOnSlotsChanged(FDelegateHandle InHandle)
+{
+	OnSlotsChanged.Remove(InHandle);
+}
+
 #undef LOCTEXT_NAMESPACE
diff --git a/Engine/Source/Editor/SkeletonEditor/Private/EditableSkeleton.h b/Engine/Source/Editor/SkeletonEditor/Private/EditableSkeleton.h
--- a/Engine/Source/Editor/SkeletonEditor/Private/EditableSkeleton.h
+++ b/Engine/Source/Editor/SkeletonEditor/Private/EditableSkeleton.h
@@ -66,6 +66,7 @@ public:
 	virtual void SetRigBoneMappings(const TMap<FName, FName>& InMappings) override;
 	virtual void RemoveUnusedBones() override;
 	virtual void UpdateSkeletonReferencePose(USkeletalMesh* InSkeletalMesh) override;
+	virtual void RegisterSlotNode(const FName& InSlotName) override;
 	virtual bool AddSlotGroupName(const FName& InSlotName) override;
 	virtual void SetSlotGroupName(const FName& InSlotName, const FName& InGroupName) override;
 	virtual void DeleteSlotName(const FName& InSlotName) override;
@@ -73,8 +74,10 @@ public:
 	virtual void RenameSlotName(const FName InOldSlotName, const FName InNewSlotName) override;
 	virtual FDelegateHandle RegisterOnSmartNameChanged(const FOnSmartNameChanged::FDelegate& InOnSmartNameChanged) override;
 	virtual void UnregisterOnSmartNameChanged(FDelegateHandle InHandle) override;
-	virtual void RegisterOnNotifiesChanged(const FSimpleMulticastDelegate::FDelegate& InDelegate);
-	virtual void UnregisterOnNotifiesChanged(void* Thing);
+	virtual void RegisterOnNotifiesChanged(const FSimpleMulticastDelegate::FDelegate& InDelegate) override;
+	virtual void UnregisterOnNotifiesChanged(void* Thing) override;
+	virtual FDelegateHandle RegisterOnSlotsChanged(const FSimpleMulticastDelegate::FDelegate& InOnSlotsChanged) override;
+	virtual void UnregisterOnSlotsChanged(FDelegateHandle InHandle) override;
 	virtual void SetBoneTranslationRetargetingMode(FName InBoneName, EBoneTranslationRetargetingMode::Type NewRetargetingMode) override;
 	virtual EBoneTranslationRetargetingMode::Type GetBoneTranslationRetargetingMode(FName InBoneName) const override;
 	virtual void RefreshBoneTree() override;
@@ -163,9 +166,12 @@ private:
 	/** Delegate called when trees need refreshing */
 	FSimpleMulticastDelegate OnTreeRefresh;
 
-	/** Delegate called when a smart name is removed */
+	/** Delegate called when a smart name is changed */
 	FOnSmartNameChanged OnSmartNameChanged;
 
 	/** Delegate called when notifies are modified */
 	FSimpleMulticastDelegate OnNotifiesChanged;
+
+	/** Delegate called when slots are modified */
+	FSimpleMulticastDelegate OnSlotsChanged;
 };
diff --git a/Engine/Source/Editor/SkeletonEditor/Public/IEditableSkeleton.h b/Engine/Source/Editor/SkeletonEditor/Public/IEditableSkeleton.h
--- a/Engine/Source/Editor/SkeletonEditor/Public/IEditableSkeleton.h
+++ b/Engine/Source/Editor/SkeletonEditor/Public/IEditableSkeleton.h
@@ -161,6 +161,9 @@ public:
 	/** Create RefLocalPoses from InSkeletalMesh */
 	virtual void UpdateSkeletonReferencePose(class USkeletalMesh* InSkeletalMesh) = 0;
 
+	/** Register a slot node */
+	virtual void RegisterSlotNode(const FName& InSlotName) = 0;
+
 	/** Add a slot group name */
 	virtual bool AddSlotGroupName(const FName& InSlotName) = 0;
 
@@ -188,6 +191,12 @@ public:
 	/** Unregister a delegate to be called when this skeletons notifies are changed */
 	virtual void UnregisterOnNotifiesChanged(void* Thing) = 0;
 
+	/** Register a delegate to be called when this skeletons slots are changed */
+	virtual FDelegateHandle RegisterOnSlotsChanged(const FSimpleMulticastDelegate::FDelegate& InDelegate) = 0;
+
+	/** Unregister a delegate to be called when this skeletons slots are changed */
+	virtual void UnregisterOnSlotsChanged(FDelegateHandle InHandle) = 0;
+
 	/** Wrap USkeleton::SetBoneTranslationRetargetingMode */
 	virtual void SetBoneTranslationRetargetingMode(FName InBoneName, EBoneTranslationRetargetingMode::Type NewRetargetingMode) = 0;
 
diff --git a/Engine/Source/Editor/UnrealEd/Classes/ActorFactories/ActorFactoryAnimationAsset.h b/Engine/Source/Editor/UnrealEd/Classes/ActorFactories/ActorFactoryAnimationAsset.h
--- a/Engine/Source/Editor/UnrealEd/Classes/ActorFactories/ActorFactoryAnimationAsset.h
+++ b/Engine/Source/Editor/UnrealEd/Classes/ActorFactories/ActorFactoryAnimationAsset.h
@@ -24,7 +24,7 @@ protected:
 	virtual bool CanCreateActorFrom( const FAssetData& AssetData, FText& OutErrorMsg ) override;
 	//~ End UActorFactory Interface
 
-	virtual USkeletalMesh* GetSkeletalMeshFromAsset( UObject* Asset ) const override;
+	virtual USkeletalMesh* GetSkeletalMeshFromAsset( UObject* Asset ) override;
 };
 
 
diff --git a/Engine/Source/Editor/UnrealEd/Classes/ActorFactories/ActorFactorySkeletalMesh.h b/Engine/Source/Editor/UnrealEd/Classes/ActorFactories/ActorFactorySkeletalMesh.h
--- a/Engine/Source/Editor/UnrealEd/Classes/ActorFactories/ActorFactorySkeletalMesh.h
+++ b/Engine/Source/Editor/UnrealEd/Classes/ActorFactories/ActorFactorySkeletalMesh.h
@@ -12,12 +12,31 @@ class AActor;
 struct FAssetData;
 class USkeletalMesh;
 
+DECLARE_DELEGATE_RetVal_OneParam(USkeletalMesh*, FGetSkeletalMeshFromAssetDelegate, UObject* /* in asset */);
+DECLARE_DELEGATE_TwoParams(FPostSkeletalMeshActorSpawnedDelegate, AActor* /* spawned actor */, UObject* /* asset source*/);
+
 UCLASS(MinimalAPI, config=Editor)
 class UActorFactorySkeletalMesh : public UActorFactory
 {
 	GENERATED_UCLASS_BODY()
 
+public:
+	/*
+	 * Adds support for customizing the actor spawning given a set of delegates.
+	 */
+	static UNREALED_API void RegisterDelegatesForAssetClass(
+		UClass* InAssetClass,
+		FGetSkeletalMeshFromAssetDelegate GetSkeletalMeshFromAssetDelegate,
+		FPostSkeletalMeshActorSpawnedDelegate PostSkeletalMeshActorSpawnedDelegate
+	);
+
+	/*
+	 * Removes the registered customization for actor spawning
+	 */
+	static UNREALED_API void UnregisterDelegatesForAssetClass(UClass* InAssetClass);
+
 protected:
+
 	//~ Begin UActorFactory Interface
 	virtual void PostSpawnActor( UObject* Asset, AActor* NewActor ) override;
 	virtual void PostCreateBlueprint( UObject* Asset, AActor* CDO ) override;
@@ -25,7 +44,11 @@ protected:
 	virtual FQuat AlignObjectToSurfaceNormal(const FVector& InSurfaceNormal, const FQuat& ActorRotation) const override;
 	//~ End UActorFactory Interface
 
-	virtual USkeletalMesh* GetSkeletalMeshFromAsset( UObject* Asset ) const;
+	virtual USkeletalMesh* GetSkeletalMeshFromAsset( UObject* Asset );
+
+	UClass* ClassUsedForDelegate;
+	static TMap<UClass*, FGetSkeletalMeshFromAssetDelegate> GetSkeletalMeshDelegates;
+	static TMap<UClass*, FPostSkeletalMeshActorSpawnedDelegate> PostSkeletalMeshActorSpawnedDelegates;
 };
 
 
diff --git a/Engine/Source/Editor/UnrealEd/Classes/Animation/DebugSkelMeshComponent.h b/Engine/Source/Editor/UnrealEd/Classes/Animation/DebugSkelMeshComponent.h
--- a/Engine/Source/Editor/UnrealEd/Classes/Animation/DebugSkelMeshComponent.h
+++ b/Engine/Source/Editor/UnrealEd/Classes/Animation/DebugSkelMeshComponent.h
@@ -308,9 +308,9 @@ class UNREALED_API UDebugSkelMeshComponent : public USkeletalMeshComponent
 	virtual bool IsWindEnabled() const override { return true; }
 	virtual void SetAnimClass(class UClass* NewClass) override;
 	//~ End SkeletalMeshComponent Interface
-	// Preview.
-	// @todo document
-	bool IsPreviewOn() const;
+
+	// return true if currently preview animation asset is on
+	virtual bool IsPreviewOn() const;
 
 	// @todo document
 	FString GetPreviewText() const;
diff --git a/Engine/Source/Editor/UnrealEd/Classes/Animation/EditorNotifyObject.h b/Engine/Source/Editor/UnrealEd/Classes/Animation/EditorNotifyObject.h
--- a/Engine/Source/Editor/UnrealEd/Classes/Animation/EditorNotifyObject.h
+++ b/Engine/Source/Editor/UnrealEd/Classes/Animation/EditorNotifyObject.h
@@ -19,20 +19,13 @@ class UEditorNotifyObject : public UEditorAnimBaseObj
 	GENERATED_UCLASS_BODY()
 
 	/** Set up the editor object
-	 *	@param TrackIdx Index of the track the notify belongs to
-	 *	@param NotifyIndex index of the notify within the track
+	 *	@param InNotify		The notify to modify
 	 */
-	virtual void InitialiseNotify(int32 TrackIdx, int32 NotifyIndex);
+	virtual void InitialiseNotify(const FAnimNotifyEvent& InNotify);
 	
 	/** Copy changes made to the event object back to the montage asset */
 	virtual bool ApplyChangesToMontage() override;
 
-	/** Index of the notify within it's track */
-	int32 NotifyIndex;
-
-	/** Index of the track the notify is in */
-	int32 TrackIndex;
-
 	/** The notify event to modify */
 	UPROPERTY(EditAnywhere, Category=Event)
 	FAnimNotifyEvent Event;
diff --git a/Engine/Source/Editor/UnrealEd/Classes/Preferences/PersonaOptions.h b/Engine/Source/Editor/UnrealEd/Classes/Preferences/PersonaOptions.h
--- a/Engine/Source/Editor/UnrealEd/Classes/Preferences/PersonaOptions.h
+++ b/Engine/Source/Editor/UnrealEd/Classes/Preferences/PersonaOptions.h
@@ -15,6 +15,8 @@
 #include "Engine/EngineBaseTypes.h"
 #include "PersonaOptions.generated.h"
 
+enum class EFrameNumberDisplayFormats : uint8;
+
 /** Persisted camera follow mode */
 UENUM()
 enum class EAnimationViewportCameraFollowMode : uint8
@@ -165,6 +167,30 @@ class UNREALED_API UPersonaOptions : public UObject
 	UPROPERTY(config)
 	TArray<FAssetEditorOptions> AssetEditorOptions;
 
+	/** Snap value used to determine scrub resolution of the curve timeline */
+	UPROPERTY(config)
+	float CurveEditorSnapInterval;
+
+	/** Snap value used to determine scrub resolution of the anim timeline */
+	UPROPERTY(config)
+	int32 TimelineScrubSnapValue;
+
+	/** Display format for the anim timeline */
+	UPROPERTY(config)
+	EFrameNumberDisplayFormats TimelineDisplayFormat;
+
+	/** Whether to display percentage in the anim timeline */
+	UPROPERTY(config)
+	bool bTimelineDisplayPercentage;
+
+	/** Whether to display secondary format (times/frames) in the anim timeline */
+	UPROPERTY(config)
+	bool bTimelineDisplayFormatSecondary;
+
+	/** Whether to display keys in the timeline's curve tracks */
+	UPROPERTY(config)
+	bool bTimelineDisplayCurveKeys;
+
 public:
 	void SetShowGrid( bool bInShowGrid );
 	void SetHighlightOrigin( bool bInHighlightOrigin );
diff --git a/Engine/Source/Editor/UnrealEd/Classes/ThumbnailRendering/SkeletalMeshThumbnailRenderer.h b/Engine/Source/Editor/UnrealEd/Classes/ThumbnailRendering/SkeletalMeshThumbnailRenderer.h
--- a/Engine/Source/Editor/UnrealEd/Classes/ThumbnailRendering/SkeletalMeshThumbnailRenderer.h
+++ b/Engine/Source/Editor/UnrealEd/Classes/ThumbnailRendering/SkeletalMeshThumbnailRenderer.h
@@ -15,20 +15,22 @@
 class FCanvas;
 class FRenderTarget;
 
-UCLASS(config=Editor, MinimalAPI)
-class USkeletalMeshThumbnailRenderer : public UDefaultSizedThumbnailRenderer
+UCLASS(config=Editor)
+class UNREALED_API USkeletalMeshThumbnailRenderer : public UDefaultSizedThumbnailRenderer
 {
 	GENERATED_UCLASS_BODY()
 
 
 	// Begin UThumbnailRenderer Object
-	UNREALED_API virtual void Draw(UObject* Object, int32 X, int32 Y, uint32 Width, uint32 Height, FRenderTarget*, FCanvas* Canvas) override;
+	virtual void Draw(UObject* Object, int32 X, int32 Y, uint32 Width, uint32 Height, FRenderTarget*, FCanvas* Canvas) override;
 	// End UThumbnailRenderer Object
 
 	// UObject implementation
-	UNREALED_API virtual void BeginDestroy() override;
+	virtual void BeginDestroy() override;
 
-private:
+	virtual void AddAdditionalPreviewSceneContent(UObject* Object, UWorld* PreviewWorld) {}
+
+protected:
 	class FSkeletalMeshThumbnailScene* ThumbnailScene;
 };
 
diff --git a/Engine/Source/Editor/UnrealEd/Private/Animation/AsyncAnimCompressionUI.cpp b/Engine/Source/Editor/UnrealEd/Private/Animation/AsyncAnimCompressionUI.cpp
new file mode 100644
--- /dev/null
+++ b/Engine/Source/Editor/UnrealEd/Private/Animation/AsyncAnimCompressionUI.cpp
@@ -0,0 +1,43 @@
+// Copyright Epic Games, Inc. All Rights Reserved.
+
+#include "AsyncAnimCompressionUI.h"
+#include "GlobalEditorNotification.h"
+#include "Widgets/Notifications/SNotificationList.h"
+#include "Animation/AnimCompressionDerivedDataPublic.h"
+#include "Animation/AnimSequence.h"
+
+/** Notification class for asynchronous shader compiling. */
+class FAnimCompressionNotificationImpl : public FGlobalEditorNotification
+{
+protected:
+	/** FGlobalEditorNotification interface */
+	virtual bool ShouldShowNotification(const bool bIsNotificationAlreadyActive) const override;
+	virtual void SetNotificationText(const TSharedPtr<SNotificationItem>& InNotificationItem) const override;
+};
+
+/** Global notification object. */
+FAnimCompressionNotificationImpl GAnimCompressionNotification;
+
+bool FAnimCompressionNotificationImpl::ShouldShowNotification(const bool bIsNotificationAlreadyActive) const
+{
+	const uint32 RemainingJobsThreshold = bIsNotificationAlreadyActive ? 0 : 10;
+	const uint32 ActiveJobs = GAsyncCompressedAnimationsTracker ? GAsyncCompressedAnimationsTracker->GetNumRemainingJobs() : 0;
+
+	return ActiveJobs > RemainingJobsThreshold;
+}
+
+void FAnimCompressionNotificationImpl::SetNotificationText(const TSharedPtr<SNotificationItem>& InNotificationItem) const
+{
+	if(GAsyncCompressedAnimationsTracker)
+	{
+		const int32 RemainingJobs = GAsyncCompressedAnimationsTracker->GetNumRemainingJobs();
+		if (RemainingJobs > 0)
+		{
+			FFormatNamedArguments Args;
+			Args.Add(TEXT("AnimsToCompress"), FText::AsNumber(RemainingJobs));
+			const FText ProgressMessage = FText::Format(NSLOCTEXT("AsyncAnimCompression", "AnimCompressionInProgressFormat", "Compressing Animations ({AnimsToCompress})"), Args);
+
+			InNotificationItem->SetText(ProgressMessage);
+		}
+	}
+}
\ No newline at end of file
diff --git a/Engine/Source/Editor/UnrealEd/Private/Animation/AsyncAnimCompressionUI.h b/Engine/Source/Editor/UnrealEd/Private/Animation/AsyncAnimCompressionUI.h
new file mode 100644
--- /dev/null
+++ b/Engine/Source/Editor/UnrealEd/Private/Animation/AsyncAnimCompressionUI.h
@@ -0,0 +1,7 @@
+// Copyright Epic Games, Inc. All Rights Reserved.
+
+#pragma once
+
+#include "TickableEditorObject.h"
+
+class UAnimSequence;
\ No newline at end of file
diff --git a/Engine/Source/Editor/UnrealEd/Private/Animation/EditorNotifyObject.cpp b/Engine/Source/Editor/UnrealEd/Private/Animation/EditorNotifyObject.cpp
--- a/Engine/Source/Editor/UnrealEd/Private/Animation/EditorNotifyObject.cpp
+++ b/Engine/Source/Editor/UnrealEd/Private/Animation/EditorNotifyObject.cpp
@@ -13,35 +13,36 @@ bool UEditorNotifyObject::ApplyChangesToMontage()
 {
 	if(AnimObject)
 	{
-		check(AnimObject->AnimNotifyTracks.IsValidIndex(TrackIndex));
-		FAnimNotifyTrack* Track = &AnimObject->AnimNotifyTracks[TrackIndex];
-		check(Track->Notifies.IsValidIndex(NotifyIndex));
-		FAnimNotifyEvent* ActualNotify = Track->Notifies[NotifyIndex];
-		Event.OnChanged(Event.GetTime());
-
-		// If we have a duration this is a state notify
-		if(Event.GetDuration() > 0.0f)
+		for(FAnimNotifyEvent& Notify : AnimObject->Notifies)
 		{
-			Event.EndLink.OnChanged(Event.EndLink.GetTime());
-
-			// Always keep link methods in sync between notifies and duration links
-			if(Event.GetLinkMethod() != Event.EndLink.GetLinkMethod())
+			if(Notify.Guid == Event.Guid)
 			{
-				Event.EndLink.ChangeLinkMethod(Event.GetLinkMethod());
+				Event.OnChanged(Event.GetTime());
+
+				// If we have a duration this is a state notify
+				if(Event.GetDuration() > 0.0f)
+				{
+					Event.EndLink.OnChanged(Event.EndLink.GetTime());
+
+					// Always keep link methods in sync between notifies and duration links
+					if(Event.GetLinkMethod() != Event.EndLink.GetLinkMethod())
+					{
+						Event.EndLink.ChangeLinkMethod(Event.GetLinkMethod());
+					}
+				}
+				Notify = Event;
+				break;
 			}
 		}
-		*ActualNotify = Event;
 	}
 
 	return true;
 }
 
-void UEditorNotifyObject::InitialiseNotify(int32 InTrackIdx, int32 InNotifyIndex)
+void UEditorNotifyObject::InitialiseNotify(const FAnimNotifyEvent& InNotify)
 {
 	if(AnimObject)
 	{
-		Event = *AnimObject->AnimNotifyTracks[InTrackIdx].Notifies[InNotifyIndex];
-		NotifyIndex = InNotifyIndex;
-		TrackIndex = InTrackIdx;
+		Event = InNotify;
 	}
 }
diff --git a/Engine/Source/Editor/UnrealEd/Private/AnimationEditorUtils.cpp b/Engine/Source/Editor/UnrealEd/Private/AnimationEditorUtils.cpp
--- a/Engine/Source/Editor/UnrealEd/Private/AnimationEditorUtils.cpp
+++ b/Engine/Source/Editor/UnrealEd/Private/AnimationEditorUtils.cpp
@@ -25,6 +25,7 @@
 #include "Factories/AimOffsetBlendSpaceFactoryNew.h"
 #include "Engine/PoseWatch.h"
 #include "Animation/AnimBlueprintGeneratedClass.h"
+#include "Animation/AnimBoneCompressionSettings.h"
 #include "Animation/AnimComposite.h"
 #include "Animation/AnimCompress.h"
 #include "Animation/BlendSpace.h"
@@ -40,8 +41,8 @@
 #include "AnimationStateMachineGraph.h"
 #include "K2Node_Composite.h"
 #include "Kismet2/BlueprintEditorUtils.h"
-#include "Animation/AnimCompress_Automatic.h"
 #include "AssetRegistryModule.h"
+#include "Interfaces/IMainFrameModule.h"
 
 #define LOCTEXT_NAMESPACE "AnimationEditorUtils"
 
@@ -242,12 +243,214 @@ FString SCreateAnimationAssetDlg::GetFullAssetPath()
 /////////////////////////////////////////////////////
 /////////////////////////////////////////////////////
 
+/////////////////////////////////////////////////////
+// Dialog to prompt user to select an animation compression settings asset.
+/////////////////////////////////////////////////////
+
+SAnimationCompressionSelectionDialog::SAnimationCompressionSelectionDialog()
+	: bValidAssetChosen(false)
+{}
+
+SAnimationCompressionSelectionDialog::~SAnimationCompressionSelectionDialog()
+{}
+
+void SAnimationCompressionSelectionDialog::SetOnAssetSelected(const FOnAssetSelected& InHandler)
+{
+	OnAssetSelectedHandler = InHandler;
+}
+
+void SAnimationCompressionSelectionDialog::DoSelectAsset(const FAssetData& SelectedAsset)
+{
+	bValidAssetChosen = true;
+	OnAssetSelectedHandler.ExecuteIfBound(SelectedAsset);
+
+	CloseDialog();
+}
+
+FReply SAnimationCompressionSelectionDialog::OnConfirmClicked()
+{
+	TArray<FAssetData> SelectedAssets = GetCurrentSelectionDelegate.Execute();
+	if (SelectedAssets.Num() > 0)
+	{
+		DoSelectAsset(SelectedAssets[0]);
+	}
+
+	return FReply::Handled();
+}
+
+FReply SAnimationCompressionSelectionDialog::OnCancelClicked()
+{
+	CloseDialog();
+
+	return FReply::Handled();
+}
+
+void SAnimationCompressionSelectionDialog::CloseDialog()
+{
+	TSharedPtr<SWindow> ContainingWindow = FSlateApplication::Get().FindWidgetWindow(AsShared());
+
+	if (ContainingWindow.IsValid())
+	{
+		ContainingWindow->RequestDestroyWindow();
+	}
+}
+
+void SAnimationCompressionSelectionDialog::OnAssetSelected(const FAssetData& AssetData)
+{
+	CurrentlySelectedAssets = GetCurrentSelectionDelegate.Execute();
+}
+
+void SAnimationCompressionSelectionDialog::OnAssetsActivated(const TArray<FAssetData>& SelectedAssets, EAssetTypeActivationMethod::Type ActivationType)
+{
+	const bool bCorrectActivationMethod = ActivationType == EAssetTypeActivationMethod::DoubleClicked || ActivationType == EAssetTypeActivationMethod::Opened;
+	if (SelectedAssets.Num() > 0 && bCorrectActivationMethod)
+	{
+		DoSelectAsset(SelectedAssets[0]);
+	}
+}
+
+bool SAnimationCompressionSelectionDialog::IsConfirmButtonEnabled() const
+{
+	return CurrentlySelectedAssets.Num() > 0;
+}
+
+void SAnimationCompressionSelectionDialog::Construct(const FArguments& InArgs, const FAnimationCompressionSelectionDialogConfig& InConfig)
+{
+	FAssetPickerConfig AssetPickerConfig;
+	AssetPickerConfig.Filter.ClassNames.Push(UAnimBoneCompressionSettings::StaticClass()->GetFName());
+	AssetPickerConfig.Filter.bRecursiveClasses = true;
+	AssetPickerConfig.InitialAssetViewType = EAssetViewType::List;
+	AssetPickerConfig.OnAssetSelected = FOnAssetSelected::CreateSP(this, &SAnimationCompressionSelectionDialog::OnAssetSelected);
+	AssetPickerConfig.OnAssetsActivated = FOnAssetsActivated::CreateSP(this, &SAnimationCompressionSelectionDialog::OnAssetsActivated);
+	AssetPickerConfig.GetCurrentSelectionDelegates.Add(&GetCurrentSelectionDelegate);
+	AssetPickerConfig.SaveSettingsName = TEXT("AnimationCompressionSelectionDialog");
+	AssetPickerConfig.bCanShowFolders = false;
+	AssetPickerConfig.bCanShowDevelopersFolder = true;
+	AssetPickerConfig.bAllowNullSelection = false;
+	AssetPickerConfig.bAllowDragging = false;
+	AssetPickerConfig.SelectionMode = ESelectionMode::Single;
+	AssetPickerConfig.bFocusSearchBoxWhenOpened = true;
+	AssetPickerConfig.InitialAssetSelection = InConfig.DefaultSelectedAsset != nullptr ? InConfig.DefaultSelectedAsset : FAnimationUtils::GetDefaultAnimationBoneCompressionSettings();
+	AssetPickerConfig.bForceShowEngineContent = true;
+
+	if (AssetPickerConfig.InitialAssetSelection.IsValid())
+	{
+		CurrentlySelectedAssets.Add(AssetPickerConfig.InitialAssetSelection);
+	}
+
+	FContentBrowserModule& ContentBrowserModule = FModuleManager::Get().LoadModuleChecked<FContentBrowserModule>(TEXT("ContentBrowser"));
+	AssetPicker = ContentBrowserModule.Get().CreateAssetPicker(AssetPickerConfig);
+
+	// The root widget in this dialog.
+	TSharedRef<SVerticalBox> MainVerticalBox = SNew(SVerticalBox);
+
+	// Asset view
+	MainVerticalBox->AddSlot()
+		.FillHeight(1)
+		.Padding(0, 0, 0, 4)
+		[
+			SNew(SSplitter)
+
+			+ SSplitter::Slot()
+			.Value(0.75f)
+			[
+				SNew(SBorder)
+				.BorderImage(FEditorStyle::GetBrush("ToolPanel.GroupBorder"))
+				[
+					AssetPicker.ToSharedRef()
+				]
+			]
+		];
+
+	// Buttons and asset name
+	TSharedRef<SHorizontalBox> ButtonsAndNameBox = SNew(SHorizontalBox)
+		+ SHorizontalBox::Slot()
+		.AutoWidth()
+		.HAlign(HAlign_Right)
+		.VAlign(VAlign_Bottom)
+		.Padding(4, 3)
+		[
+			SNew(SButton)
+			.Text(LOCTEXT("AnimationCompressionSelectionDialogSelectButton", "Select"))
+			.ContentPadding(FMargin(8, 2, 8, 2))
+			.IsEnabled(this, &SAnimationCompressionSelectionDialog::IsConfirmButtonEnabled)
+			.OnClicked(this, &SAnimationCompressionSelectionDialog::OnConfirmClicked)
+		]
+		+ SHorizontalBox::Slot()
+		.AutoWidth()
+		.VAlign(VAlign_Bottom)
+		.Padding(4, 3)
+		[
+			SNew(SButton)
+			.ContentPadding(FMargin(8, 2, 8, 2))
+			.Text(LOCTEXT("AnimationCompressionSelectionDialogCancelButton", "Cancel"))
+			.OnClicked(this, &SAnimationCompressionSelectionDialog::OnCancelClicked)
+		];
+
+	MainVerticalBox->AddSlot()
+		.AutoHeight()
+		.HAlign(HAlign_Right)
+		.Padding(0)
+		[
+			ButtonsAndNameBox
+		];
+
+	ChildSlot
+		[
+			MainVerticalBox
+		];
+}
+
+/////////////////////////////////////////////////////
+/////////////////////////////////////////////////////
+/////////////////////////////////////////////////////
+/////////////////////////////////////////////////////
+
 /////////////////////////////////////////////////////
 // Animation editor utility functions
 /////////////////////////////////////////////////////
 
 namespace AnimationEditorUtils
 {
+	FAssetData CreateModalAnimationCompressionSelectionDialog(const FAnimationCompressionSelectionDialogConfig& InConfig)
+	{
+		struct FModalResult
+		{
+			void OnAssetSelected(const FAssetData& SelectedAsset)
+			{
+				SavedResult = SelectedAsset;
+			}
+
+			FAssetData SavedResult;
+		};
+
+		FModalResult ModalWindowResult;
+		auto OnAssetSelectedDelegate = SAnimationCompressionSelectionDialog::FOnAssetSelected::CreateRaw(&ModalWindowResult, &FModalResult::OnAssetSelected);
+
+		TSharedRef<SAnimationCompressionSelectionDialog> Dialog = SNew(SAnimationCompressionSelectionDialog, InConfig);
+		Dialog->SetOnAssetSelected(OnAssetSelectedDelegate);
+
+		const FVector2D DefaultWindowSize(400.0f, 500.0f);
+		const FVector2D WindowSize = InConfig.WindowSizeOverride.IsZero() ? DefaultWindowSize : InConfig.WindowSizeOverride;
+		const FText WindowTitle = InConfig.DialogTitleOverride.IsEmpty() ? LOCTEXT("GenericAnimationCompressionSelectionDialogWindowHeader", "Select compression settings") : InConfig.DialogTitleOverride;
+
+		TSharedRef<SWindow> DialogWindow =
+			SNew(SWindow)
+			.Title(WindowTitle)
+			.ClientSize(WindowSize);
+
+		DialogWindow->SetContent(Dialog);
+
+		IMainFrameModule& MainFrameModule = FModuleManager::LoadModuleChecked<IMainFrameModule>(TEXT("MainFrame"));
+		const TSharedPtr<SWindow>& MainFrameParentWindow = MainFrameModule.GetParentWindow();
+		if (MainFrameParentWindow.IsValid())
+		{
+			FSlateApplication::Get().AddModalWindow(DialogWindow, MainFrameParentWindow.ToSharedRef());
+		}
+
+		return ModalWindowResult.SavedResult;
+	}
+
 	/** Creates a unique package and asset name taking the form InBasePackageName+InSuffix */
 	void CreateUniqueAssetName(const FString& InBasePackageName, const FString& InSuffix, FString& OutPackageName, FString& OutAssetName) 
 	{
@@ -521,38 +724,34 @@ namespace AnimationEditorUtils
 		MenuBuilder.EndSection();
 	}
 
-	bool ApplyCompressionAlgorithm(TArray<UAnimSequence*>& AnimSequencePtrs, UAnimCompress* Algorithm)
+	bool ApplyCompressionAlgorithm(TArray<UAnimSequence*>& AnimSequencePtrs, UAnimBoneCompressionSettings* OverrideSettings)
 	{
-		if(Algorithm)
+		const bool bProceed = (AnimSequencePtrs.Num() > 1)? EAppReturnType::Yes == FMessageDialog::Open(EAppMsgType::YesNo,
+			FText::Format(NSLOCTEXT("UnrealEd", "AboutToCompressAnimations_F", "About to compress {0} animations.  Proceed?"), FText::AsNumber(AnimSequencePtrs.Num()))) : true;
+		if(bProceed)
 		{
-			const bool bProceed = (AnimSequencePtrs.Num() > 1)? EAppReturnType::Yes == FMessageDialog::Open(EAppMsgType::YesNo,
-				FText::Format(NSLOCTEXT("UnrealEd", "AboutToCompressAnimations_F", "About to compress {0} animations.  Proceed?"), FText::AsNumber(AnimSequencePtrs.Num()))) : true;
-			if(bProceed)
+			GWarn->BeginSlowTask(LOCTEXT("AnimCompressing", "Compressing"), true);
+
 			{
-				GWarn->BeginSlowTask(LOCTEXT("AnimCompressing", "Compressing"), true);
+				TSharedPtr<FAnimCompressContext> CompressContext = MakeShareable(new FAnimCompressContext(false, true, AnimSequencePtrs.Num()));
 
+				for (UAnimSequence* AnimSeq : AnimSequencePtrs)
 				{
-					TSharedPtr<FAnimCompressContext> CompressContext = MakeShareable(new FAnimCompressContext(false, true, AnimSequencePtrs.Num()));
-
-					for (UAnimSequence* AnimSeq : AnimSequencePtrs)
+					if (OverrideSettings != nullptr)
 					{
-						// If we are not compressing with 'Auto', then clear CompressCommandletVersion
-						// So we can recompress these animations later.
-						const bool bIsAutoCompressor = Algorithm->IsA(UAnimCompress_Automatic::StaticClass());
-						if (!bIsAutoCompressor)
-						{
-							AnimSeq->CompressCommandletVersion = 0;
-						}
-						AnimSeq->CompressionScheme = static_cast<UAnimCompress*>(StaticDuplicateObject(Algorithm, AnimSeq));
-						AnimSeq->RequestAnimCompression(FRequestAnimCompressionParams(false, CompressContext));
-						++CompressContext->AnimIndex;
+						AnimSeq->BoneCompressionSettings = OverrideSettings;
 					}
-				}
-				
-				GWarn->EndSlowTask();
 
-				return true;
+					// Clear CompressCommandletVersion so we can recompress these animations later.
+					AnimSeq->CompressCommandletVersion = 0;
+					AnimSeq->RequestAnimCompression(FRequestAnimCompressionParams(false, CompressContext));
+					++CompressContext->AnimIndex;
+				}
 			}
+
+			GWarn->EndSlowTask();
+
+			return true;
 		}
 
 		return false;
diff --git a/Engine/Source/Editor/UnrealEd/Private/ClassIconFinder.cpp b/Engine/Source/Editor/UnrealEd/Private/ClassIconFinder.cpp
--- a/Engine/Source/Editor/UnrealEd/Private/ClassIconFinder.cpp
+++ b/Engine/Source/Editor/UnrealEd/Private/ClassIconFinder.cpp
@@ -112,7 +112,9 @@ const UClass* FClassIconFinder::GetIconClassForAssetData(const FAssetData& InAss
 		return FindObject<UClass>(ANY_PACKAGE, *InAssetData.AssetName.ToString());
 	}
 	
-	if ( AssetClass->IsChildOf<UBlueprint>() )
+	static const FName IgnoreClassThumbnail(TEXT("IgnoreClassThumbnail"));
+	if ( AssetClass->IsChildOf<UBlueprint>() &&
+		!AssetClass->HasMetaDataHierarchical(IgnoreClassThumbnail))
 	{
 		if ( bOutIsClassType )
 		{
diff --git a/Engine/Source/Editor/UnrealEd/Private/Dialogs/Dialogs.cpp b/Engine/Source/Editor/UnrealEd/Private/Dialogs/Dialogs.cpp
--- a/Engine/Source/Editor/UnrealEd/Private/Dialogs/Dialogs.cpp
+++ b/Engine/Source/Editor/UnrealEd/Private/Dialogs/Dialogs.cpp
@@ -915,7 +915,7 @@ void SGenericDialogWidget::Construct( const FArguments& InArgs )
 	];
 }
 
-void SGenericDialogWidget::OpenDialog(const FText& InDialogTitle, const TSharedRef< SWidget >& DisplayContent, const FArguments& InArgs)
+void SGenericDialogWidget::OpenDialog(const FText& InDialogTitle, const TSharedRef< SWidget >& DisplayContent, const FArguments& InArgs, bool bAsModalDialog)
 {
 	TSharedPtr< SWindow > Window;
 	TSharedPtr< SGenericDialogWidget > GenericDialogWidget;
@@ -941,7 +941,14 @@ void SGenericDialogWidget::OpenDialog(const FText& InDialogTitle, const TSharedR
 
 	GenericDialogWidget->SetWindow(Window);
 
-	FSlateApplication::Get().AddWindow( Window.ToSharedRef() );
+	if(bAsModalDialog)
+	{
+		GEditor->EditorAddModalWindow(Window.ToSharedRef());
+	}
+	else
+	{
+		FSlateApplication::Get().AddWindow( Window.ToSharedRef() );
+	}
 }
 
 FReply SGenericDialogWidget::OnOK_Clicked(void)
diff --git a/Engine/Source/Editor/UnrealEd/Private/EditorModeManager.cpp b/Engine/Source/Editor/UnrealEd/Private/EditorModeManager.cpp
--- a/Engine/Source/Editor/UnrealEd/Private/EditorModeManager.cpp
+++ b/Engine/Source/Editor/UnrealEd/Private/EditorModeManager.cpp
@@ -638,6 +638,7 @@ void FEditorModeTools::SetCoordSystem(ECoordSystem NewCoordSystem)
 		}
 	}
 	CoordSystem = NewCoordSystem;
+	BroadcastCoordSystemChanged(NewCoordSystem);
 }
 
 void FEditorModeTools::SetDefaultMode( const FEditorModeID DefaultModeID )
diff --git a/Engine/Source/Editor/UnrealEd/Private/Factories/ActorFactory.cpp b/Engine/Source/Editor/UnrealEd/Private/Factories/ActorFactory.cpp
--- a/Engine/Source/Editor/UnrealEd/Private/Factories/ActorFactory.cpp
+++ b/Engine/Source/Editor/UnrealEd/Private/Factories/ActorFactory.cpp
@@ -275,8 +275,11 @@ AActor* UActorFactory::CreateActor( UObject* Asset, ULevel* InLevel, FTransform
 		{
 			PostSpawnActor(Asset, NewActor);
 
-			// Only do this if the actor wasn't already given a name
-			if (Name == NAME_None && Asset)
+			// Only do this if the actor wasn't already given a name.
+			// We also want to skip this if a toolchain during PostSpawn
+			// already adjusted the name.
+			if (Name == NAME_None && Asset && 
+				!NewActor->GetName().StartsWith(Asset->GetName()))
 			{
 				FActorLabelUtilities::SetActorLabelUnique(NewActor, Asset->GetName());
 			}
@@ -874,7 +877,7 @@ bool UActorFactoryAnimationAsset::CanCreateActorFrom( const FAssetData& AssetDat
 	return true;
 }
 
-USkeletalMesh* UActorFactoryAnimationAsset::GetSkeletalMeshFromAsset( UObject* Asset ) const
+USkeletalMesh* UActorFactoryAnimationAsset::GetSkeletalMeshFromAsset( UObject* Asset )
 {
 	USkeletalMesh* SkeletalMesh = NULL;
 	
@@ -952,25 +955,22 @@ void UActorFactoryAnimationAsset::PostCreateBlueprint( UObject* Asset,  AActor*
 /*-----------------------------------------------------------------------------
 UActorFactorySkeletalMesh
 -----------------------------------------------------------------------------*/
+
+// static storage
+TMap<UClass*, FGetSkeletalMeshFromAssetDelegate> UActorFactorySkeletalMesh::GetSkeletalMeshDelegates;
+TMap<UClass*, FPostSkeletalMeshActorSpawnedDelegate> UActorFactorySkeletalMesh::PostSkeletalMeshActorSpawnedDelegates;
+
 UActorFactorySkeletalMesh::UActorFactorySkeletalMesh(const FObjectInitializer& ObjectInitializer)
 	: Super(ObjectInitializer)
 { 
 	DisplayName = LOCTEXT("SkeletalMeshDisplayName", "Skeletal Mesh");
 	NewActorClass = ASkeletalMeshActor::StaticClass();
 	bUseSurfaceOrientation = true;
+	ClassUsedForDelegate = nullptr;
 }
 
 bool UActorFactorySkeletalMesh::CanCreateActorFrom( const FAssetData& AssetData, FText& OutErrorMsg )
 {	
-	if ( !AssetData.IsValid() || 
-		( !AssetData.GetClass()->IsChildOf( USkeletalMesh::StaticClass() ) && 
-		  !AssetData.GetClass()->IsChildOf( UAnimBlueprint::StaticClass() ) && 
-		  !AssetData.GetClass()->IsChildOf( USkeleton::StaticClass() ) ) )
-	{
-		OutErrorMsg = NSLOCTEXT("CanCreateActor", "NoAnimSeq", "A valid anim sequence must be specified.");
-		return false;
-	}
-
 	FAssetRegistryModule& AssetRegistryModule = FModuleManager::LoadModuleChecked<FAssetRegistryModule>(TEXT("AssetRegistry"));
 	IAssetRegistry& AssetRegistry = AssetRegistryModule.Get();
 	FAssetData SkeletalMeshData;
@@ -1035,9 +1035,16 @@ bool UActorFactorySkeletalMesh::CanCreateActorFrom( const FAssetData& AssetData,
 				return false;
 			}
 		}
-		else
+	}
+
+	if (!SkeletalMeshData.IsValid())
+	{
+		for (const auto& Pair : UActorFactorySkeletalMesh::GetSkeletalMeshDelegates)
 		{
-			OutErrorMsg = NSLOCTEXT("CanCreateActor", "NoSkelMeshTargetSkeleton", "SkeletalMesh must have a valid Target Skeleton.");
+			if (Pair.Value.Execute(AssetData.GetAsset()) != nullptr)
+			{
+				return true;
+			}
 		}
 	}
 
@@ -1058,7 +1065,7 @@ bool UActorFactorySkeletalMesh::CanCreateActorFrom( const FAssetData& AssetData,
 	return true;
 }
 
-USkeletalMesh* UActorFactorySkeletalMesh::GetSkeletalMeshFromAsset( UObject* Asset ) const
+USkeletalMesh* UActorFactorySkeletalMesh::GetSkeletalMeshFromAsset( UObject* Asset )
 {
 	USkeletalMesh*SkeletalMesh = Cast<USkeletalMesh>( Asset );
 	UAnimBlueprint* AnimBlueprint = Cast<UAnimBlueprint>( Asset );
@@ -1075,6 +1082,19 @@ USkeletalMesh* UActorFactorySkeletalMesh::GetSkeletalMeshFromAsset( UObject* Ass
 		SkeletalMesh = Skeleton->GetPreviewMesh(true);
 	}
 
+	if (SkeletalMesh == NULL)
+	{
+		for (const auto& Pair : UActorFactorySkeletalMesh::GetSkeletalMeshDelegates)
+		{
+			if (USkeletalMesh* SkeletalMeshFromDelegate = Pair.Value.Execute(Asset))
+			{
+				ClassUsedForDelegate = Pair.Key;
+				SkeletalMesh = SkeletalMeshFromDelegate;
+				break;
+			}
+		}
+	}
+
 	check( SkeletalMesh != NULL );
 	return SkeletalMesh;
 }
@@ -1103,6 +1123,12 @@ void UActorFactorySkeletalMesh::PostSpawnActor( UObject* Asset, AActor* NewActor
 	{
 		NewSMActor->GetSkeletalMeshComponent()->SetAnimInstanceClass(AnimBlueprint->GeneratedClass);
 	}
+
+	if (ClassUsedForDelegate != NULL)
+	{
+		UActorFactorySkeletalMesh::PostSkeletalMeshActorSpawnedDelegates.FindChecked(ClassUsedForDelegate).Execute(NewSMActor, Asset);
+	}
+
 }
 
 void UActorFactorySkeletalMesh::PostCreateBlueprint( UObject* Asset, AActor* CDO )
@@ -1124,6 +1150,23 @@ FQuat UActorFactorySkeletalMesh::AlignObjectToSurfaceNormal(const FVector& InSur
 	return FindActorAlignmentRotation(ActorRotation, FVector(0.f, 0.f, 1.f), InSurfaceNormal);
 }
 
+void UActorFactorySkeletalMesh::RegisterDelegatesForAssetClass(
+	UClass* InAssetClass,
+	FGetSkeletalMeshFromAssetDelegate GetSkeletalMeshFromAssetDelegate,
+	FPostSkeletalMeshActorSpawnedDelegate PostSkeletalMeshActorSpawnedDelegate
+)
+{
+	UActorFactorySkeletalMesh::GetSkeletalMeshDelegates.Add(InAssetClass, GetSkeletalMeshFromAssetDelegate);
+	UActorFactorySkeletalMesh::PostSkeletalMeshActorSpawnedDelegates.Add(InAssetClass, PostSkeletalMeshActorSpawnedDelegate);
+}
+
+void UActorFactorySkeletalMesh::UnregisterDelegatesForAssetClass(UClass* InAssetClass)
+{
+	UActorFactorySkeletalMesh::GetSkeletalMeshDelegates.Remove(InAssetClass);
+	UActorFactorySkeletalMesh::PostSkeletalMeshActorSpawnedDelegates.Remove(InAssetClass);
+}
+
+
 /*-----------------------------------------------------------------------------
 UActorFactoryCameraActor
 -----------------------------------------------------------------------------*/
diff --git a/Engine/Source/Editor/UnrealEd/Private/Factories/AnimSequenceFactory.cpp b/Engine/Source/Editor/UnrealEd/Private/Factories/AnimSequenceFactory.cpp
--- a/Engine/Source/Editor/UnrealEd/Private/Factories/AnimSequenceFactory.cpp
+++ b/Engine/Source/Editor/UnrealEd/Private/Factories/AnimSequenceFactory.cpp
@@ -66,24 +66,23 @@ bool UAnimSequenceFactory::ConfigureProperties()
 
 UObject* UAnimSequenceFactory::FactoryCreateNew(UClass* Class, UObject* InParent, FName Name, EObjectFlags Flags, UObject* Context, FFeedbackContext* Warn)
 {
+	//
+	UAnimSequence* AnimSequence = NewObject<UAnimSequence>(InParent, Class, Name, Flags);
+
+	// @todo I think this will crash, we should support differentoptions
+	AnimSequence->SequenceLength = 0.f;
+	AnimSequence->SetRawNumberOfFrame(0);
+	
 	if (TargetSkeleton)
 	{
-		UAnimSequence* AnimSequence = NewObject<UAnimSequence>(InParent, Class, Name, Flags);
-
-		// @todo I think this will crash, we should support differentoptions
-		AnimSequence->SequenceLength = 0.f;
-		AnimSequence->SetRawNumberOfFrame(0);
-		
-		AnimSequence->SetSkeleton( TargetSkeleton );
-		if (PreviewSkeletalMesh)
-		{
-			AnimSequence->SetPreviewMesh(PreviewSkeletalMesh);
-		}
-
-		return AnimSequence;
+		AnimSequence->SetSkeleton(TargetSkeleton);
+	}
+	if (PreviewSkeletalMesh)
+	{
+		AnimSequence->SetPreviewMesh(PreviewSkeletalMesh);
 	}
 
-	return NULL;
+	return AnimSequence;
 }
 
 void UAnimSequenceFactory::OnTargetSkeletonSelected(const FAssetData& SelectedAsset)
diff --git a/Engine/Source/Editor/UnrealEd/Private/Factories/EditorFactories.cpp b/Engine/Source/Editor/UnrealEd/Private/Factories/EditorFactories.cpp
--- a/Engine/Source/Editor/UnrealEd/Private/Factories/EditorFactories.cpp
+++ b/Engine/Source/Editor/UnrealEd/Private/Factories/EditorFactories.cpp
@@ -5398,6 +5398,7 @@ void FCustomizableTextObjectFactory::ProcessBuffer(UObject* InParent, EObjectFla
 				UObject* ObjectParent = InParent ? InParent : GetParentForNewObject(ObjClass);
 
 				// Make sure this name is not used by anything else. Will rename other stuff if necessary
+				UpdateObjectName(ObjClass, ObjName);
 				ClearObjectNameUsage(ObjectParent, ObjName);
 
 				// Spawn the object and reset it's archetype
