Sequencer: Remove stray pragma disable/enable optimization

#rnx

diff --git a/Engine/Source/Editor/Sequencer/Private/Sequencer.cpp b/Engine/Source/Editor/Sequencer/Private/Sequencer.cpp
--- a/Engine/Source/Editor/Sequencer/Private/Sequencer.cpp
+++ b/Engine/Source/Editor/Sequencer/Private/Sequencer.cpp
@@ -7557,7 +7557,6 @@ bool FSequencer::DoPaste()
 
 	return bAnythingPasted;
 }
-PRAGMA_DISABLE_OPTIMIZATION
 
 bool FSequencer::PasteObjectBindings(const FString& TextToImport, TArray<FNotificationInfo>& PasteErrors)
 {
@@ -7816,8 +7815,6 @@ bool FSequencer::PasteObjectBindings(const FString& TextToImport, TArray<FNotifi
 	return true;
 }
 
-PRAGMA_ENABLE_OPTIMIZATION
-
 bool FSequencer::PasteTracks(const FString& TextToImport, TArray<FNotificationInfo>& PasteErrors)
 {
 	FScopedTransaction Transaction(FGenericCommands::Get().Paste->GetDescription());
