Removing unnecessary include

diff --git a/Engine/Source/Editor/Sequencer/Private/SequencerKeyStructGenerator.cpp b/Engine/Source/Editor/Sequencer/Private/SequencerKeyStructGenerator.cpp
--- a/Engine/Source/Editor/Sequencer/Private/SequencerKeyStructGenerator.cpp
+++ b/Engine/Source/Editor/Sequencer/Private/SequencerKeyStructGenerator.cpp
@@ -3,7 +3,6 @@
 #include "SequencerKeyStructGenerator.h"
 #include "Channels/MovieSceneChannelData.h"
 #include "UObject/CoreObjectVersion.h"
-#include "Uobject/UnrealTypePrivate.h" // For converting UProperties to FProperties
 
 UMovieSceneKeyStructType::UMovieSceneKeyStructType(const FObjectInitializer& ObjInit)
 	: Super(ObjInit)
