Sequencer: Autoscrub only if in scrubbing state to fix double evaluations when autoscroll is enabled during playing.

diff --git a/Engine/Source/Editor/Sequencer/Private/Sequencer.cpp b/Engine/Source/Editor/Sequencer/Private/Sequencer.cpp
--- a/Engine/Source/Editor/Sequencer/Private/Sequencer.cpp
+++ b/Engine/Source/Editor/Sequencer/Private/Sequencer.cpp
@@ -283,6 +283,8 @@ public:
 	}
 };
 
+PRAGMA_DISABLE_OPTIMIZATION
+
 void FSequencer::InitSequencer(const FSequencerInitParams& InitParams, const TSharedRef<ISequencerObjectChangeListener>& InObjectChangeListener, const TArray<FOnCreateTrackEditor>& TrackEditorDelegates, const TArray<FOnCreateEditorObjectBinding>& EditorObjectBindingDelegates)
 {
 	bIsEditingWithinLevelEditor = InitParams.bEditWithinLevelEditor;
@@ -534,6 +536,7 @@ void FSequencer::InitSequencer(const FSequencerInitParams& InitParams, const TSh
 	OnActivateSequenceEvent.Broadcast(ActiveTemplateIDs[0]);
 }
 
+PRAGMA_ENABLE_OPTIMIZATION
 
 FSequencer::FSequencer()
 	: SequencerCommandBindings( new FUICommandList )
@@ -691,7 +694,7 @@ void FSequencer::Tick(float InDeltaTime)
 	}
 
 	// Animate the autoscrub offset if it's set
-	if (AutoscrubOffset.IsSet())
+	if (AutoscrubOffset.IsSet() && PlaybackState == EMovieScenePlayerStatus::Scrubbing )
 	{
 		FQualifiedFrameTime CurrentTime = GetLocalTime();
 		FFrameTime Offset = (AutoscrubOffset.GetValue() * AutoScrollFactor) * CurrentTime.Rate;
