Sequencer: Remove stray pragma

diff --git a/Engine/Source/Editor/Sequencer/Private/Sequencer.cpp b/Engine/Source/Editor/Sequencer/Private/Sequencer.cpp
--- a/Engine/Source/Editor/Sequencer/Private/Sequencer.cpp
+++ b/Engine/Source/Editor/Sequencer/Private/Sequencer.cpp
@@ -283,8 +283,6 @@ public:
 	}
 };
 
-PRAGMA_DISABLE_OPTIMIZATION
-
 void FSequencer::InitSequencer(const FSequencerInitParams& InitParams, const TSharedRef<ISequencerObjectChangeListener>& InObjectChangeListener, const TArray<FOnCreateTrackEditor>& TrackEditorDelegates, const TArray<FOnCreateEditorObjectBinding>& EditorObjectBindingDelegates)
 {
 	bIsEditingWithinLevelEditor = InitParams.bEditWithinLevelEditor;
@@ -536,8 +534,6 @@ void FSequencer::InitSequencer(const FSequencerInitParams& InitParams, const TSh
 	OnActivateSequenceEvent.Broadcast(ActiveTemplateIDs[0]);
 }
 
-PRAGMA_ENABLE_OPTIMIZATION
-
 FSequencer::FSequencer()
 	: SequencerCommandBindings( new FUICommandList )
 	, SequencerSharedBindings(new FUICommandList)
