Sequencer: Set local time looped instead of directly so that evaluation is play to instead of jump to when autoscrolling. This allows for swept evaluation when autosrolling.

diff --git a/Engine/Source/Editor/Sequencer/Private/Sequencer.cpp b/Engine/Source/Editor/Sequencer/Private/Sequencer.cpp
--- a/Engine/Source/Editor/Sequencer/Private/Sequencer.cpp
+++ b/Engine/Source/Editor/Sequencer/Private/Sequencer.cpp
@@ -694,7 +694,7 @@ void FSequencer::Tick(float InDeltaTime)
 	{
 		FQualifiedFrameTime CurrentTime = GetLocalTime();
 		FFrameTime Offset = (AutoscrubOffset.GetValue() * AutoScrollFactor) * CurrentTime.Rate;
-		SetLocalTimeDirectly(CurrentTime.Time + Offset);
+		SetLocalTimeLooped(CurrentTime.Time + Offset);
 	}
 
 	if (GetSelectionRange().IsEmpty() && GetLoopMode() == SLM_LoopSelectionRange)
@@ -2914,11 +2914,11 @@ void FSequencer::UpdateAutoScroll(double NewTime, float ThresholdPercentage)
 	{
 		if (AutoscrollOffset.GetValue() < 0 && LocalTime.AsSeconds() > ViewRange.GetLowerBoundValue() + Threshold)
 		{
-			SetLocalTimeDirectly( (ViewRange.GetLowerBoundValue() + Threshold) * LocalTime.Rate );
+			SetLocalTimeLooped( (ViewRange.GetLowerBoundValue() + Threshold) * LocalTime.Rate );
 		}
 		else if (AutoscrollOffset.GetValue() > 0 && LocalTime.AsSeconds() < ViewRange.GetUpperBoundValue() - Threshold)
 		{
-			SetLocalTimeDirectly( (ViewRange.GetUpperBoundValue() - Threshold) * LocalTime.Rate );
+			SetLocalTimeLooped( (ViewRange.GetUpperBoundValue() - Threshold) * LocalTime.Rate );
 		}
 	}
 
