Sequencer: Fix GetUniqueName so that it returns if the name is already unique

diff --git a/Engine/Source/Editor/Sequencer/Private/SequencerUtilities.cpp b/Engine/Source/Editor/Sequencer/Private/SequencerUtilities.cpp
--- a/Engine/Source/Editor/Sequencer/Private/SequencerUtilities.cpp
+++ b/Engine/Source/Editor/Sequencer/Private/SequencerUtilities.cpp
@@ -217,6 +217,11 @@ void FSequencerUtilities::PopulateMenu_SetBlendType(FMenuBuilder& MenuBuilder, c
 
 FName FSequencerUtilities::GetUniqueName( FName CandidateName, const TArray<FName>& ExistingNames )
 {
+	if (!ExistingNames.Contains(CandidateName))
+	{
+		return CandidateName;
+	}
+
 	FString CandidateNameString = CandidateName.ToString();
 	FString BaseNameString = CandidateNameString;
 	if ( CandidateNameString.Len() >= 3 && CandidateNameString.Right(3).IsNumeric() )
