Sequencer Scripting: Export fbx through python now spawns spawnables if necessary

diff --git a/Engine/Plugins/MovieScene/SequencerScripting/Source/SequencerScriptingEditor/Private/SequencerTools.cpp b/Engine/Plugins/MovieScene/SequencerScripting/Source/SequencerScriptingEditor/Private/SequencerTools.cpp
--- a/Engine/Plugins/MovieScene/SequencerScripting/Source/SequencerScriptingEditor/Private/SequencerTools.cpp
+++ b/Engine/Plugins/MovieScene/SequencerScripting/Source/SequencerScriptingEditor/Private/SequencerTools.cpp
@@ -19,6 +19,7 @@
 #include "MovieSceneCommonHelpers.h"
 #include "Components/SkeletalMeshComponent.h"
 #include "Animation/AnimSequenceBase.h"
+#include "ScopedTransaction.h"
 
 #define LOCTEXT_NAMESPACE "SequencerTools"
 
@@ -117,6 +118,7 @@ TArray<FSequencerBoundObjects> USequencerToolsFunctionLibrary::GetBoundObjects(U
 	// Evaluation needs to occur in order to obtain spawnables
 	FFrameRate Resolution = InSequence->GetMovieScene()->GetTickResolution();
 	TRange<FFrameNumber> SpecifiedRange = InRange.ToNative(Resolution);
+	Player->Play();
 	Player->PlayToFrame(SpecifiedRange.GetLowerBoundValue());
 
 	FMovieSceneSequenceID SequenceId = Player->State.FindSequenceId(InSequence);
@@ -128,6 +130,7 @@ TArray<FSequencerBoundObjects> USequencerToolsFunctionLibrary::GetBoundObjects(U
 		BoundObjects.Add(FSequencerBoundObjects(Binding, Player->GetBoundObjects(ObjectBinding)));
 	}
 
+	Player->Stop();
 	InWorld->DestroyActor(OutActor);
 	return BoundObjects;
 }
@@ -146,6 +149,7 @@ TArray<FSequencerBoundObjects> USequencerToolsFunctionLibrary::GetObjectBindings
 
 	FFrameRate Resolution = InSequence->GetMovieScene()->GetTickResolution();
 	TRange<FFrameNumber> SpecifiedRange = InRange.ToNative(Resolution);
+	Player->Play();
 	Player->PlayToFrame(SpecifiedRange.GetLowerBoundValue());
 
 	TArray<FSequencerBoundObjects> BoundObjects;
@@ -160,6 +164,7 @@ TArray<FSequencerBoundObjects> USequencerToolsFunctionLibrary::GetObjectBindings
 		}
 	}
 
+	Player->Stop();
 	InWorld->DestroyActor(OutActor);
 	return BoundObjects;
 }
@@ -188,10 +193,28 @@ bool USequencerToolsFunctionLibrary::ExportFBX(UWorld* World, ULevelSequence* Se
 	Player->Initialize(Sequence, World->PersistentLevel, Settings, CameraSettings);
 	Player->State.AssignSequence(MovieSceneSequenceID::Root, *Sequence, *Player);
 	FMovieSceneSequenceIDRef Template = MovieSceneSequenceID::Root;
+	bool bDidExport = false;
 	FMovieSceneSequenceTransform RootToLocalTransform;
-	bool bDidExport = MovieSceneToolHelpers::ExportFBX(World, MovieScene, Player, Bindings, NodeNameAdapter, Template, InFBXFileName, RootToLocalTransform);
+
+	FScopedTransaction ExportFBXTransaction(NSLOCTEXT("Sequencer", "ExportFBX", "Export FBX"));
+	{
+		FSpawnableRestoreState SpawnableRestoreState(MovieScene);
+
+		if (SpawnableRestoreState.bWasChanged)
+		{
+			// Evaluate at the beginning of the subscene time to ensure that spawnables are created before export
+			Player->Play();
+			Player->PlayToFrame(MovieScene::DiscreteInclusiveLower(MovieScene->GetPlaybackRange()));
+		}
+
+		bDidExport = MovieSceneToolHelpers::ExportFBX(World, MovieScene, Player, Bindings, NodeNameAdapter, Template, InFBXFileName, RootToLocalTransform);
+
+	}
+		
+	Player->Stop();
 	Exporter->SetExportOptionsOverride(nullptr);
 	World->DestroyActor(OutActor);
+	
 	return bDidExport;
 }
 
diff --git a/Engine/Source/Editor/MovieSceneTools/Private/MovieSceneToolHelpers.cpp b/Engine/Source/Editor/MovieSceneTools/Private/MovieSceneToolHelpers.cpp
--- a/Engine/Source/Editor/MovieSceneTools/Private/MovieSceneToolHelpers.cpp
+++ b/Engine/Source/Editor/MovieSceneTools/Private/MovieSceneToolHelpers.cpp
@@ -66,6 +66,9 @@
 #include "LiveLinkPresetTypes.h"
 #include "LiveLinkSourceSettings.h"
 #include "Features/IModularFeatures.h"
+#include "Tracks/MovieSceneSpawnTrack.h"
+#include "Sections/MovieSceneSpawnSection.h"
+
 /* MovieSceneToolHelpers
  *****************************************************************************/
 
@@ -2191,3 +2194,51 @@ bool MovieSceneToolHelpers::ExportToAnimSequence(UAnimSequence* AnimSequence, UM
 	return true;
 }
 
+FSpawnableRestoreState::FSpawnableRestoreState(UMovieScene* MovieScene) :
+	bWasChanged(false)
+{
+	WeakMovieScene = MovieScene;
+
+	for (int32 SpawnableIndex = 0; SpawnableIndex < WeakMovieScene->GetSpawnableCount(); ++SpawnableIndex)
+	{
+		FMovieSceneSpawnable& Spawnable = WeakMovieScene->GetSpawnable(SpawnableIndex);
+
+		UMovieSceneSpawnTrack* SpawnTrack = WeakMovieScene->FindTrack<UMovieSceneSpawnTrack>(Spawnable.GetGuid());
+
+		if (SpawnTrack)
+		{
+			bWasChanged = true;
+
+			// Spawnable could be in a subscene, so temporarily override it to persist throughout
+			SpawnOwnershipMap.Add(Spawnable.GetGuid(), Spawnable.GetSpawnOwnership());
+			Spawnable.SetSpawnOwnership(ESpawnOwnership::MasterSequence);
+
+			// Spawnable could have animated spawned state, so temporarily override it to spawn infinitely
+			UMovieSceneSpawnSection* SpawnSection = Cast<UMovieSceneSpawnSection>(SpawnTrack->CreateNewSection());
+			SpawnSection->Modify();
+			SpawnSection->GetChannel().Reset();
+			SpawnSection->GetChannel().SetDefault(true);
+		}
+	}
+}
+
+FSpawnableRestoreState::~FSpawnableRestoreState()
+{
+	if (!bWasChanged || !WeakMovieScene.IsValid())
+	{
+		return;
+	}
+
+	// Restore spawnable owners
+	for (int32 SpawnableIndex = 0; SpawnableIndex < WeakMovieScene->GetSpawnableCount(); ++SpawnableIndex)
+	{
+		FMovieSceneSpawnable& Spawnable = WeakMovieScene->GetSpawnable(SpawnableIndex);
+		Spawnable.SetSpawnOwnership(SpawnOwnershipMap[Spawnable.GetGuid()]);
+	}
+
+	// Restore modified spawned sections
+	bool bOrigSquelchTransactionNotification = GEditor->bSquelchTransactionNotification;
+	GEditor->bSquelchTransactionNotification = true;
+	GEditor->UndoTransaction(false);
+	GEditor->bSquelchTransactionNotification = bOrigSquelchTransactionNotification;
+}
diff --git a/Engine/Source/Editor/MovieSceneTools/Public/MovieSceneToolHelpers.h b/Engine/Source/Editor/MovieSceneTools/Public/MovieSceneToolHelpers.h
--- a/Engine/Source/Editor/MovieSceneTools/Public/MovieSceneToolHelpers.h
+++ b/Engine/Source/Editor/MovieSceneTools/Public/MovieSceneToolHelpers.h
@@ -13,6 +13,7 @@
 #include "ISequencer.h"
 #include "Logging/TokenizedMessage.h"
 #include "MovieSceneTranslator.h"
+#include "MovieSceneSpawnable.h"
 #include "MovieSceneCaptureSettings.h"
 #include "SEnumCombobox.h"
 #include "Animation/AnimSequence.h"
@@ -375,6 +376,18 @@ public:
 	 * Get the fbx camera name
 	 */
 	static FString GetCameraName(fbxsdk::FbxCamera* InCamera);
+
+};
+
+// Helper to make spawnables persist throughout the export process and then restore properly afterwards
+struct MOVIESCENETOOLS_API FSpawnableRestoreState
+{
+	FSpawnableRestoreState(UMovieScene* MovieScene);
+	~FSpawnableRestoreState();
+
+	bool bWasChanged;
+	TMap<FGuid, ESpawnOwnership> SpawnOwnershipMap;
+	TWeakObjectPtr<UMovieScene> WeakMovieScene;
 };
 
 class FTrackEditorBindingIDPicker : public FMovieSceneObjectBindingIDPicker
diff --git a/Engine/Source/Editor/Sequencer/Private/Sequencer.cpp b/Engine/Source/Editor/Sequencer/Private/Sequencer.cpp
--- a/Engine/Source/Editor/Sequencer/Private/Sequencer.cpp
+++ b/Engine/Source/Editor/Sequencer/Private/Sequencer.cpp
@@ -10475,73 +10475,16 @@ void FSequencer::ExportFBXInternal(const FString& ExportFilename, TArray<FGuid>&
 			FMovieSceneSequenceIDRef Template = GetFocusedTemplateID();
 			UnFbx::FFbxExporter::FLevelSequenceNodeNameAdapter NodeNameAdapter(MovieScene, this, Template);
 
-			// Helper to make spawnables persist throughout the export process and then restore properly afterwards
-			struct FSpawnableRestoreState
-			{
-				FSpawnableRestoreState(FSequencer& InSequencer) :
-					bWasChanged(false)
-				{
-					WeakMovieScene = InSequencer.GetFocusedMovieSceneSequence()->GetMovieScene();
-
-					for (int32 SpawnableIndex = 0; SpawnableIndex < WeakMovieScene->GetSpawnableCount(); ++SpawnableIndex)
-					{
-						FMovieSceneSpawnable& Spawnable = WeakMovieScene->GetSpawnable(SpawnableIndex);
-
-						UMovieSceneSpawnTrack* SpawnTrack = WeakMovieScene->FindTrack<UMovieSceneSpawnTrack>(Spawnable.GetGuid());
-
-						if (SpawnTrack)
-						{
-							bWasChanged = true;
-
-							// Spawnable could be in a subscene, so temporarily override it to persist throughout
-							SpawnOwnershipMap.Add(Spawnable.GetGuid(), Spawnable.GetSpawnOwnership());
-							Spawnable.SetSpawnOwnership(ESpawnOwnership::MasterSequence);
-
-							// Spawnable could have animated spawned state, so temporarily override it to spawn infinitely
-							UMovieSceneSpawnSection* SpawnSection = Cast<UMovieSceneSpawnSection>(SpawnTrack->CreateNewSection());
-							SpawnSection->Modify();
-							SpawnSection->GetChannel().Reset();
-							SpawnSection->GetChannel().SetDefault(true);
-						}
-					}
-
-					if (bWasChanged)
-					{
-						// Evaluate at the beginning of the subscene time to ensure that spawnables are created before export
-						InSequencer.SetLocalTimeDirectly(MovieScene::DiscreteInclusiveLower(InSequencer.GetTimeBounds()));
-					}
-				}
+			FScopedTransaction ExportFBXTransaction(NSLOCTEXT("Sequencer", "ExportFBX", "Export FBX"));
 
-				~FSpawnableRestoreState()
+			{
+				FSpawnableRestoreState SpawnableRestoreState(MovieScene);
+				if (SpawnableRestoreState.bWasChanged)
 				{
-					if (!bWasChanged || !WeakMovieScene.IsValid())
-					{
-						return;
-					}
-
-					// Restore spawnable owners
-					for (int32 SpawnableIndex = 0; SpawnableIndex < WeakMovieScene->GetSpawnableCount(); ++SpawnableIndex)
-					{
-						FMovieSceneSpawnable& Spawnable = WeakMovieScene->GetSpawnable(SpawnableIndex);
-						Spawnable.SetSpawnOwnership(SpawnOwnershipMap[Spawnable.GetGuid()]);
-					}
-
-					// Restore modified spawned sections
-					bool bOrigSquelchTransactionNotification = GEditor->bSquelchTransactionNotification;
-					GEditor->bSquelchTransactionNotification = true;
-					GEditor->UndoTransaction(false);
-					GEditor->bSquelchTransactionNotification = bOrigSquelchTransactionNotification;
+					// Evaluate at the beginning of the subscene time to ensure that spawnables are created before export
+					SetLocalTimeDirectly(MovieScene::DiscreteInclusiveLower(GetTimeBounds()));
 				}
 
-				bool bWasChanged;
-				TMap<FGuid, ESpawnOwnership> SpawnOwnershipMap;
-				TWeakObjectPtr<UMovieScene> WeakMovieScene;
-			};
-
-			FScopedTransaction ExportFBXTransaction(NSLOCTEXT("Sequencer", "ExportFBX", "Export FBX"));
-
-			{
-				FSpawnableRestoreState SpawnableRestoreState(*this);
 				if (MovieSceneToolHelpers::ExportFBX(World, MovieScene, this, Bindings, NodeNameAdapter, Template, ExportFilename, RootToLocalTransform))
 				{
 					FNotificationInfo Info(NSLOCTEXT("Sequencer", "ExportFBXSucceeded", "FBX Export Succeeded."));
