Sequencer: Remove label browser settings

#rnx

diff --git a/Engine/Source/Editor/Sequencer/Private/Sequencer.cpp b/Engine/Source/Editor/Sequencer/Private/Sequencer.cpp
--- a/Engine/Source/Editor/Sequencer/Private/Sequencer.cpp
+++ b/Engine/Source/Editor/Sequencer/Private/Sequencer.cpp
@@ -10839,14 +10839,6 @@ void FSequencer::BindCommands()
 		FCanExecuteAction::CreateLambda( []{ return true; } ),
 		FIsActionChecked::CreateLambda( [this]{ return Settings->GetShowChannelColors(); } ) );
 
-	SequencerCommandBindings->MapAction(
-		Commands.ToggleLabelBrowser,
-		FExecuteAction::CreateLambda( [this]{
-			Settings->SetLabelBrowserVisible( !Settings->GetLabelBrowserVisible() );
-		} ),
-		FCanExecuteAction::CreateLambda( []{ return true; } ),
-		FIsActionChecked::CreateLambda( [this]{ return Settings->GetLabelBrowserVisible(); } ) );
-
 	SequencerCommandBindings->MapAction(
 		Commands.ToggleShowSelectedNodesOnly,
 		FExecuteAction::CreateLambda( [this]{
diff --git a/Engine/Source/Editor/Sequencer/Private/SequencerCommands.cpp b/Engine/Source/Editor/Sequencer/Private/SequencerCommands.cpp
--- a/Engine/Source/Editor/Sequencer/Private/SequencerCommands.cpp
+++ b/Engine/Source/Editor/Sequencer/Private/SequencerCommands.cpp
@@ -120,7 +120,6 @@ void FSequencerCommands::RegisterCommands()
 	UI_COMMAND( FindInContentBrowser, "Find in Content Browser", "Find the viewed sequence asset in the content browser", EUserInterfaceActionType::Button, FInputChord() );
 	UI_COMMAND( ToggleCombinedKeyframes, "Combined Keyframes", "Show/hide the combined keyframes at the top node level", EUserInterfaceActionType::ToggleButton, FInputChord() );
 	UI_COMMAND( ToggleChannelColors, "Channel Colors", "Show/hide the channel colors in the track area", EUserInterfaceActionType::ToggleButton, FInputChord() );
-	UI_COMMAND( ToggleLabelBrowser, "Label Browser", "Show/hide the track label browser", EUserInterfaceActionType::ToggleButton, FInputChord() );
 	UI_COMMAND( ToggleShowSelectedNodesOnly, "Selected Nodes Only", "Show selected nodes only", EUserInterfaceActionType::ToggleButton, FInputChord() );
 	
 	UI_COMMAND( ToggleShowCurveEditor, "Curve Editor", "Show the animation keys in a curve editor", EUserInterfaceActionType::ToggleButton, FInputChord() );
diff --git a/Engine/Source/Editor/Sequencer/Private/SequencerSettings.cpp b/Engine/Source/Editor/Sequencer/Private/SequencerSettings.cpp
--- a/Engine/Source/Editor/Sequencer/Private/SequencerSettings.cpp
+++ b/Engine/Source/Editor/Sequencer/Private/SequencerSettings.cpp
@@ -28,7 +28,6 @@ USequencerSettings::USequencerSettings( const FObjectInitializer& ObjectInitiali
 	CurveValueSnapInterval = 0.1f;
 	GridSpacing = TOptional<float>();
 	bSnapCurveValueToInterval = true;
-	bLabelBrowserVisible = false;
 	bShowSelectedNodesOnly = false;
 	bRewindOnRecord = true;
 	ZoomPosition = ESequencerZoomPosition::SZP_CurrentTime;
@@ -358,20 +357,6 @@ void USequencerSettings::SetSnapCurveValueToInterval( bool InbSnapCurveValueToIn
 	}
 }
 
-bool USequencerSettings::GetLabelBrowserVisible() const
-{
-	return bLabelBrowserVisible;
-}
-
-void USequencerSettings::SetLabelBrowserVisible(bool Visible)
-{
-	if (bLabelBrowserVisible != Visible)
-	{
-		bLabelBrowserVisible = Visible;
-		SaveConfig();
-	}
-}
-
 bool USequencerSettings::GetShowSelectedNodesOnly() const
 {
 	return bShowSelectedNodesOnly;
diff --git a/Engine/Source/Editor/Sequencer/Public/SequencerCommands.h b/Engine/Source/Editor/Sequencer/Public/SequencerCommands.h
--- a/Engine/Source/Editor/Sequencer/Public/SequencerCommands.h
+++ b/Engine/Source/Editor/Sequencer/Public/SequencerCommands.h
@@ -295,9 +295,6 @@ public:
 	/** Toggles whether to show channel colors in the track area. */
 	TSharedPtr< FUICommandInfo > ToggleChannelColors;
 
-	/** Toggles whether the label browser is enabled in the level editor. */
-	TSharedPtr< FUICommandInfo > ToggleLabelBrowser;
-
 	/** Turns auto scroll on and off. */
 	TSharedPtr< FUICommandInfo > ToggleAutoScroll;
 
diff --git a/Engine/Source/Editor/Sequencer/Public/SequencerSettings.h b/Engine/Source/Editor/Sequencer/Public/SequencerSettings.h
--- a/Engine/Source/Editor/Sequencer/Public/SequencerSettings.h
+++ b/Engine/Source/Editor/Sequencer/Public/SequencerSettings.h
@@ -203,11 +203,6 @@ public:
 	/** Sets whether or not to snap curve values to the interval. */
 	void SetSnapCurveValueToInterval(bool InbSnapCurveValueToInterval);
 
-	/** Gets whether or not the label browser is visible. */
-	bool GetLabelBrowserVisible() const;
-	/** Sets whether or not the label browser is visible. */
-	void SetLabelBrowserVisible(bool Visible);
-
 	/** Gets whether or not to show selected nodes only. */
 	bool GetShowSelectedNodesOnly() const;
 	/** Sets whether or not to show selected nodes only. */
@@ -435,10 +430,6 @@ protected:
 	UPROPERTY( config, EditAnywhere, Category=Snapping )
 	bool bSnapCurveValueToInterval;
 
-	/** Enable or disable the label browser. */
-	UPROPERTY( config, EditAnywhere, Category=General )
-	bool bLabelBrowserVisible;
-
 	/** Only show selected nodes in the tree view. */
 	UPROPERTY( config, EditAnywhere, Category=General )
 	bool bShowSelectedNodesOnly;
