Sequencer: Make calls to OnSectionMoved consistent in dragging operations.

Call OnSectionMoved on both sections being moved and resized, both during interactive dragging and at the end of the operation (before this change it was called in one or other situation but not both consistently).
To differentiate between the two situations (in case it's important), we add a new parameter struct that provides more information.

diff --git a/Engine/Source/Editor/Sequencer/Private/Tools/EditToolDragOperations.cpp b/Engine/Source/Editor/Sequencer/Private/Tools/EditToolDragOperations.cpp
--- a/Engine/Source/Editor/Sequencer/Private/Tools/EditToolDragOperations.cpp
+++ b/Engine/Source/Editor/Sequencer/Private/Tools/EditToolDragOperations.cpp
@@ -148,6 +148,10 @@ void FResizeSection::OnBeginDrag(const FPointerEvent& MouseEvent, FVector2D Loca
 	for (TWeakObjectPtr<UMovieSceneSection> WeakSection : Sections)
 	{
 		UMovieSceneSection* Section = WeakSection.Get();
+		if (Section == nullptr)
+		{
+			continue;
+		}
 
 		if (bIsDilating)
 		{
@@ -190,6 +194,24 @@ void FResizeSection::OnBeginDrag(const FPointerEvent& MouseEvent, FVector2D Loca
 
 void FResizeSection::OnEndDrag(const FPointerEvent& MouseEvent, FVector2D LocalMousePos, const FVirtualTrackArea& VirtualTrackArea)
 {
+	bool bIsDilating = MouseEvent.IsControlDown();
+
+	if (!bIsDilating)
+	{
+		FMovieSceneSectionMovedParams SectionMovedParams(EPropertyChangeType::ValueSet);
+		for (const TWeakObjectPtr<UMovieSceneSection>& WeakSection : Sections)
+		{
+			if (UMovieSceneSection* Section = WeakSection.Get())
+			{
+				if (UMovieSceneTrack* OuterTrack = Section->GetTypedOuter<UMovieSceneTrack>())
+				{
+					OuterTrack->Modify();
+					OuterTrack->OnSectionMoved(*Section, SectionMovedParams);
+				}
+			}
+		}
+	}
+
 	EndTransaction();
 }
 
@@ -213,8 +235,10 @@ void FResizeSection::OnDrag(const FPointerEvent& MouseEvent, FVector2D LocalMous
 		TArray<FFrameNumber> SectionTimes;
 		for (const TWeakObjectPtr<UMovieSceneSection>& WeakSection : Sections)
 		{
-			UMovieSceneSection* Section = WeakSection.Get();
-			SectionTimes.Add(SectionInitTimes[Section] + DeltaTime);
+			if (UMovieSceneSection* Section = WeakSection.Get())
+			{
+				SectionTimes.Add(SectionInitTimes[Section] + DeltaTime);
+			}
 		}
 
 		float SnapThresholdPx = VirtualTrackArea.PixelToSeconds(PixelSnapWidth) - VirtualTrackArea.PixelToSeconds(0.f);
@@ -369,7 +393,7 @@ void FResizeSection::OnDrag(const FPointerEvent& MouseEvent, FVector2D LocalMous
 		if (OuterTrack)
 		{
 			OuterTrack->Modify();
-			OuterTrack->OnSectionMoved(*Section);
+			OuterTrack->OnSectionMoved(*Section, EPropertyChangeType::Interactive);
 		}
 	}
 
@@ -377,9 +401,12 @@ void FResizeSection::OnDrag(const FPointerEvent& MouseEvent, FVector2D LocalMous
 		TSet<UMovieSceneTrack*> Tracks;
 		for (TWeakObjectPtr<UMovieSceneSection> WeakSection : Sections)
 		{
-			if (UMovieSceneTrack* Track = WeakSection.Get()->GetTypedOuter<UMovieSceneTrack>())
+			if (UMovieSceneSection* Section = WeakSection.Get())
 			{
-				Tracks.Add(Track);
+				if (UMovieSceneTrack* Track = Section->GetTypedOuter<UMovieSceneTrack>())
+				{
+					Tracks.Add(Track);
+				}
 			}
 		}
 		for (UMovieSceneTrack* Track : Tracks)
@@ -459,6 +486,11 @@ void FManipulateSectionEasing::OnBeginDrag(const FPointerEvent& MouseEvent, FVec
 	Transaction.Reset( new FScopedTransaction(NSLOCTEXT("Sequencer", "DragSectionEasing", "Change Section Easing")) );
 
 	UMovieSceneSection* Section = WeakSection.Get();
+	if (Section == nullptr)
+	{
+		return;
+	}
+
 	Section->SetFlags( RF_Transactional );
 	Section->Modify();
 
@@ -485,6 +517,10 @@ void FManipulateSectionEasing::OnDrag(const FPointerEvent& MouseEvent, FVector2D
 	FFrameTime  DeltaTime = VirtualTrackArea.PixelToFrame(LocalMousePos.X) - MouseDownTime;
 
 	UMovieSceneSection* Section = WeakSection.Get();
+	if (Section == nullptr)
+	{
+		return;
+	}
 
 	// Snapping
 	if (Settings->GetIsSnapEnabled())
@@ -759,10 +795,14 @@ void FMoveKeysAndSections::OnDrag(const FPointerEvent& MouseEvent, FVector2D Loc
 
 	// Get a list of the unique tracks in this selection and update their easing so previews draw interactively as you drag.
 	TSet<UMovieSceneTrack*> Tracks;
+	FMovieSceneSectionMovedParams SectionMovedParams(EPropertyChangeType::Interactive);
 	for (TWeakObjectPtr<UMovieSceneSection> WeakSection : Sections)
 	{
 		if (UMovieSceneTrack* Track = WeakSection.Get()->GetTypedOuter<UMovieSceneTrack>())
 		{
+			Track->Modify();
+			Track->OnSectionMoved(*WeakSection.Get(), SectionMovedParams);
+
 			Tracks.Add(Track);
 		}
 	}
@@ -814,6 +854,7 @@ void FMoveKeysAndSections::OnEndDrag(const FPointerEvent& MouseEvent, FVector2D
 		Sequencer.NotifyMovieSceneDataChanged(EMovieSceneDataChangeType::MovieSceneStructureItemsChanged);
 	}
 
+	FMovieSceneSectionMovedParams SectionMovedParams(EPropertyChangeType::ValueSet);
 	for (const TWeakObjectPtr<UMovieSceneSection>& WeakSection : Sections)
 	{
 		UMovieSceneSection* Section = WeakSection.Get();
@@ -822,7 +863,7 @@ void FMoveKeysAndSections::OnEndDrag(const FPointerEvent& MouseEvent, FVector2D
 		if (OuterTrack)
 		{
 			OuterTrack->Modify();
-			OuterTrack->OnSectionMoved(*Section);
+			OuterTrack->OnSectionMoved(*Section, SectionMovedParams);
 		}
 	}
 
diff --git a/Engine/Source/Runtime/MovieScene/Public/MovieSceneTrack.h b/Engine/Source/Runtime/MovieScene/Public/MovieSceneTrack.h
--- a/Engine/Source/Runtime/MovieScene/Public/MovieSceneTrack.h
+++ b/Engine/Source/Runtime/MovieScene/Public/MovieSceneTrack.h
@@ -131,6 +131,16 @@ struct FMovieSceneSupportsEasingParams
 	FMovieSceneSupportsEasingParams(const UMovieSceneSection* InSection) : ForSection(InSection) {}
 };
 
+#if WITH_EDITOR
+/** Parameters for sections moving in the editor */
+struct FMovieSceneSectionMovedParams
+{
+	EPropertyChangeType::Type MoveType;
+
+	FMovieSceneSectionMovedParams(EPropertyChangeType::Type InMoveType) : MoveType(InMoveType) {}
+};
+#endif
+
 /**
  * Base class for a track in a Movie Scene
  */
@@ -449,7 +459,6 @@ public:
 	 *
 	 * @param Section The section that moved.
 	 */
-	virtual void OnSectionMoved(UMovieSceneSection& Section) { }
-
+	virtual void OnSectionMoved(UMovieSceneSection& Section, const FMovieSceneSectionMovedParams& Params) {}
 #endif
 };
diff --git a/Engine/Source/Runtime/MovieSceneTracks/Private/Sections/MovieSceneCameraCutSection.cpp b/Engine/Source/Runtime/MovieSceneTracks/Private/Sections/MovieSceneCameraCutSection.cpp
--- a/Engine/Source/Runtime/MovieSceneTracks/Private/Sections/MovieSceneCameraCutSection.cpp
+++ b/Engine/Source/Runtime/MovieSceneTracks/Private/Sections/MovieSceneCameraCutSection.cpp
@@ -128,8 +128,7 @@ void UMovieSceneCameraCutSection::PostEditChangeProperty(FPropertyChangedEvent&
 	{
 		if (UMovieSceneCameraCutTrack* Track = GetTypedOuter<UMovieSceneCameraCutTrack>())
 		{
-			TArray<UMovieSceneSection*> Sections = Track->GetAllSections();
-			MovieSceneHelpers::FixupConsecutiveSections(Sections, *this, false, true, Track->bCanBlend);
+			Track->OnSectionMoved(*this, EPropertyChangeType::ValueSet);
 		}
 	}
 }
diff --git a/Engine/Source/Runtime/MovieSceneTracks/Private/Tracks/MovieSceneCameraCutTrack.cpp b/Engine/Source/Runtime/MovieSceneTracks/Private/Tracks/MovieSceneCameraCutTrack.cpp
--- a/Engine/Source/Runtime/MovieSceneTracks/Private/Tracks/MovieSceneCameraCutTrack.cpp
+++ b/Engine/Source/Runtime/MovieSceneTracks/Private/Tracks/MovieSceneCameraCutTrack.cpp
@@ -185,7 +185,7 @@ FText UMovieSceneCameraCutTrack::GetDefaultDisplayName() const
 
 
 #if WITH_EDITOR
-void UMovieSceneCameraCutTrack::OnSectionMoved(UMovieSceneSection& Section)
+void UMovieSceneCameraCutTrack::OnSectionMoved(UMovieSceneSection& Section, const FMovieSceneSectionMovedParams& Params)
 {
 	MovieSceneHelpers::FixupConsecutiveSections(Sections, Section, false, true, bCanBlend);
 }
diff --git a/Engine/Source/Runtime/MovieSceneTracks/Private/Tracks/MovieSceneCinematicShotTrack.cpp b/Engine/Source/Runtime/MovieSceneTracks/Private/Tracks/MovieSceneCinematicShotTrack.cpp
--- a/Engine/Source/Runtime/MovieSceneTracks/Private/Tracks/MovieSceneCinematicShotTrack.cpp
+++ b/Engine/Source/Runtime/MovieSceneTracks/Private/Tracks/MovieSceneCinematicShotTrack.cpp
@@ -157,7 +157,7 @@ FMovieSceneTrackRowSegmentBlenderPtr UMovieSceneCinematicShotTrack::GetRowSegmen
 }
 
 #if WITH_EDITOR
-void UMovieSceneCinematicShotTrack::OnSectionMoved(UMovieSceneSection& Section)
+void UMovieSceneCinematicShotTrack::OnSectionMoved(UMovieSceneSection& Section, const FMovieSceneSectionMovedParams& Params)
 {
 	//MovieSceneHelpers::FixupConsecutiveSections(Sections, Section, false);
 }
diff --git a/Engine/Source/Runtime/MovieSceneTracks/Public/Tracks/MovieSceneCameraCutTrack.h b/Engine/Source/Runtime/MovieSceneTracks/Public/Tracks/MovieSceneCameraCutTrack.h
--- a/Engine/Source/Runtime/MovieSceneTracks/Public/Tracks/MovieSceneCameraCutTrack.h
+++ b/Engine/Source/Runtime/MovieSceneTracks/Public/Tracks/MovieSceneCameraCutTrack.h
@@ -53,7 +53,7 @@ public:
 #endif
 
 #if WITH_EDITOR
-	virtual void OnSectionMoved(UMovieSceneSection& Section) override;
+	virtual void OnSectionMoved(UMovieSceneSection& Section, const FMovieSceneSectionMovedParams& Params) override;
 #endif
 
 protected:
diff --git a/Engine/Source/Runtime/MovieSceneTracks/Public/Tracks/MovieSceneCinematicShotTrack.h b/Engine/Source/Runtime/MovieSceneTracks/Public/Tracks/MovieSceneCinematicShotTrack.h
--- a/Engine/Source/Runtime/MovieSceneTracks/Public/Tracks/MovieSceneCinematicShotTrack.h
+++ b/Engine/Source/Runtime/MovieSceneTracks/Public/Tracks/MovieSceneCinematicShotTrack.h
@@ -44,7 +44,7 @@ public:
 	virtual FMovieSceneTrackSegmentBlenderPtr GetTrackSegmentBlender() const override;
 	
 #if WITH_EDITOR
-	virtual void OnSectionMoved(UMovieSceneSection& Section) override;
+	virtual void OnSectionMoved(UMovieSceneSection& Section, const FMovieSceneSectionMovedParams& Params) override;
 #endif
 
 #if WITH_EDITORONLY_DATA
