Sequencer: Only show easing handles for sections that actually support manual easing.

diff --git a/Engine/Source/Editor/Sequencer/Private/SSequencerSection.cpp b/Engine/Source/Editor/Sequencer/Private/SSequencerSection.cpp
--- a/Engine/Source/Editor/Sequencer/Private/SSequencerSection.cpp
+++ b/Engine/Source/Editor/Sequencer/Private/SSequencerSection.cpp
@@ -1675,7 +1675,11 @@ void SSequencerSection::PaintEasingHandles( FSequencerSectionPainter& InPainter,
 			}
 		}
 
-		if (!bDrawThisSectionsHandles)
+		const UMovieSceneTrack* Track = UnderlappingSectionObj->GetTypedOuter<UMovieSceneTrack>();
+		FMovieSceneSupportsEasingParams SupportsEasingParams(UnderlappingSectionObj);
+		EMovieSceneTrackEasingSupportFlags EasingSupportFlags = Track->SupportsEasing(SupportsEasingParams);
+
+		if (!bDrawThisSectionsHandles || !EnumHasAnyFlags(EasingSupportFlags, EMovieSceneTrackEasingSupportFlags::ManualEasing))
 		{
 			continue;
 		}
@@ -1685,7 +1689,7 @@ void SSequencerSection::PaintEasingHandles( FSequencerSectionPainter& InPainter,
 		const FSlateBrush* EasingHandle = FEditorStyle::GetBrush("Sequencer.Section.EasingHandle");
 		FVector2D HandleSize(10.f, 10.f);
 
-		if (UnderlappingSectionObj->HasStartFrame())
+		if (UnderlappingSectionObj->HasStartFrame() && EnumHasAllFlags(EasingSupportFlags, EMovieSceneTrackEasingSupportFlags::ManualEaseIn))
 		{
 			TRange<FFrameNumber> EaseInRange = UnderlappingSectionObj->GetEaseInRange();
 			// Always draw handles if the section is highlighted, even if there is no range (to allow manual adjustment)
@@ -1706,7 +1710,7 @@ void SSequencerSection::PaintEasingHandles( FSequencerSectionPainter& InPainter,
 			);
 		}
 
-		if (UnderlappingSectionObj->HasEndFrame())
+		if (UnderlappingSectionObj->HasEndFrame() && EnumHasAllFlags(EasingSupportFlags, EMovieSceneTrackEasingSupportFlags::ManualEaseOut))
 		{
 			TRange<FFrameNumber> EaseOutRange = UnderlappingSectionObj->GetEaseOutRange();
 
