Sequencer: When dropping onto or below the last object binding, allow the drop to happen below so that an item can be moved to the last in the tree.

diff --git a/Engine/Source/Editor/Sequencer/Private/DisplayNodes/SequencerObjectBindingNode.cpp b/Engine/Source/Editor/Sequencer/Private/DisplayNodes/SequencerObjectBindingNode.cpp
--- a/Engine/Source/Editor/Sequencer/Private/DisplayNodes/SequencerObjectBindingNode.cpp
+++ b/Engine/Source/Editor/Sequencer/Private/DisplayNodes/SequencerObjectBindingNode.cpp
@@ -918,7 +918,15 @@ TOptional<EItemDropZone> FSequencerObjectBindingNode::CanDrop(FSequencerDisplayN
 	// This removes a confusing "above" -> "blocked" -> "above/below" transition.
 	if (ItemDropZone == EItemDropZone::OntoItem || ItemDropZone == EItemDropZone::BelowItem)
 	{
-		ItemDropZone = EItemDropZone::AboveItem;
+		// Except when dropping onto the last item so that we can drop to the end of the tree
+		if (ParentTree.GetRootNode()->GetNumChildren() > 0 && this == &ParentTree.GetRootNode()->GetChildNodes().Last().Get())
+		{
+			ItemDropZone = EItemDropZone::BelowItem;
+		}
+		else
+		{
+			ItemDropZone = EItemDropZone::AboveItem;
+		}
 	}
 
 	if (GetParent().IsValid() && GetParent()->GetType() != ESequencerNode::Folder)
