Sequencer: Invalidate when sequence is updated. This fixes a regression when redrawing all viewports was changed to requets realtime frames because hit proxies weren't invalidated.

diff --git a/Engine/Source/Editor/Sequencer/Private/LevelEditorSequencerIntegration.cpp b/Engine/Source/Editor/Sequencer/Private/LevelEditorSequencerIntegration.cpp
--- a/Engine/Source/Editor/Sequencer/Private/LevelEditorSequencerIntegration.cpp
+++ b/Engine/Source/Editor/Sequencer/Private/LevelEditorSequencerIntegration.cpp
@@ -416,9 +416,13 @@ void FLevelEditorSequencerIntegration::OnSequencerEvaluated()
 	// Blueprint, and other editors that have a 3d viewport.
 	for (FEditorViewportClient* LevelVC : GEditor->GetLevelViewportClients())
 	{
-		if (LevelVC && !LevelVC->IsRealtime())
+		if (LevelVC)
 		{
-			LevelVC->RequestRealTimeFrames(1);
+			if (!LevelVC->IsRealtime())
+			{
+				LevelVC->RequestRealTimeFrames(1);
+			}
+			LevelVC->Invalidate();
 		}
 	}
 
