Sequencer: Fix static analysis

#rnx

diff --git a/Engine/Source/Editor/Sequencer/Private/Sequencer.cpp b/Engine/Source/Editor/Sequencer/Private/Sequencer.cpp
--- a/Engine/Source/Editor/Sequencer/Private/Sequencer.cpp
+++ b/Engine/Source/Editor/Sequencer/Private/Sequencer.cpp
@@ -1720,6 +1720,11 @@ void FSequencer::StretchTime(FFrameTime InDeltaTime)
 {
 	// From the current time, find all the keys and sections to the right and move them by InDeltaTime
 	UMovieScene* FocusedMovieScene = GetFocusedMovieSceneSequence()->GetMovieScene();
+	if (!FocusedMovieScene)
+	{
+		return;
+	}
+
 	if (FocusedMovieScene->IsReadOnly())
 	{
 		ShowReadOnlyError();
@@ -1756,6 +1761,11 @@ void FSequencer::ShrinkTime(FFrameTime InDeltaTime)
 {
 	// From the current time, find all the keys and sections to the right and move them by -InDeltaTime
 	UMovieScene* FocusedMovieScene = GetFocusedMovieSceneSequence()->GetMovieScene();
+	if (!FocusedMovieScene)
+	{
+		return;
+	}
+
 	if (FocusedMovieScene->IsReadOnly())
 	{
 		ShowReadOnlyError();
@@ -1877,6 +1887,11 @@ void FSequencer::OnAddTransformKeysForSelectedObjects(EMovieSceneTransformChanne
 void FSequencer::BakeTransform()
 {
 	UMovieScene* FocusedMovieScene = GetFocusedMovieSceneSequence()->GetMovieScene();
+	if (!FocusedMovieScene)
+	{
+		return;
+	}
+
 	if (FocusedMovieScene->IsReadOnly())
 	{
 		ShowReadOnlyError();
@@ -2277,9 +2292,13 @@ TRange<FFrameNumber> FSequencer::GetSelectionRange() const
 
 void FSequencer::SetSelectionRange(TRange<FFrameNumber> Range)
 {
-	const FScopedTransaction Transaction(LOCTEXT("SetSelectionRange_Transaction", "Set Selection Range"));
 	UMovieScene* FocusedMovieScene = GetFocusedMovieSceneSequence()->GetMovieScene();
+	if (!FocusedMovieScene)
+	{
+		return;
+	}
 
+	const FScopedTransaction Transaction(LOCTEXT("SetSelectionRange_Transaction", "Set Selection Range"));
 	FocusedMovieScene->Modify();
 	FocusedMovieScene->SetSelectionRange(Range);
 }
@@ -2396,14 +2415,15 @@ void FSequencer::SetPlaybackRange(TRange<FFrameNumber> Range)
 	{
 		if (!IsPlaybackRangeLocked())
 		{
-			const FScopedTransaction Transaction(LOCTEXT("SetPlaybackRange_Transaction", "Set Playback Range"));
-
 			UMovieScene* FocusedMovieScene = GetFocusedMovieSceneSequence()->GetMovieScene();
+			if (FocusedMovieScene)
+			{
+				const FScopedTransaction Transaction(LOCTEXT("SetPlaybackRange_Transaction", "Set Playback Range"));
+				FocusedMovieScene->SetPlaybackRange(Range);
 
-			FocusedMovieScene->SetPlaybackRange(Range);
-
-			bNeedsEvaluate = true;
-			NotifyMovieSceneDataChanged(EMovieSceneDataChangeType::TrackValueChanged);
+				bNeedsEvaluate = true;
+				NotifyMovieSceneDataChanged(EMovieSceneDataChangeType::TrackValueChanged);
+			}
 		}
 	}
 }
@@ -2616,25 +2636,32 @@ void FSequencer::ZoomOutViewRange()
 
 void FSequencer::UpdatePlaybackRange()
 {
-	if (Settings->ShouldKeepPlayRangeInSectionBounds())
+	if (!Settings->ShouldKeepPlayRangeInSectionBounds())
+	{
+		return;
+	}
+
+	UMovieScene* FocusedMovieScene = GetFocusedMovieSceneSequence()->GetMovieScene();
+	if (!FocusedMovieScene)
 	{
-		UMovieScene* FocusedMovieScene = GetFocusedMovieSceneSequence()->GetMovieScene();
-		TArray<UMovieSceneSection*> AllSections = FocusedMovieScene->GetAllSections();
+		return;
+	}
 
-		if (AllSections.Num() > 0 && !IsPlaybackRangeLocked())
+	TArray<UMovieSceneSection*> AllSections = FocusedMovieScene->GetAllSections();
+
+	if (AllSections.Num() > 0 && !IsPlaybackRangeLocked())
+	{
+		TRange<FFrameNumber> NewBounds = TRange<FFrameNumber>::Empty();
+		for (UMovieSceneSection* Section : AllSections)
 		{
-			TRange<FFrameNumber> NewBounds = TRange<FFrameNumber>::Empty();
-			for (UMovieSceneSection* Section : AllSections)
-			{
-				NewBounds = TRange<FFrameNumber>::Hull(Section->ComputeEffectiveRange(), NewBounds);
-			}
+			NewBounds = TRange<FFrameNumber>::Hull(Section->ComputeEffectiveRange(), NewBounds);
+		}
 
-			// When the playback range is determined by the section bounds, don't mark the change in the playback range otherwise the scene will be marked dirty
-			if (!NewBounds.IsDegenerate())
-			{
-				const bool bAlwaysMarkDirty = false;
-				FocusedMovieScene->SetPlaybackRange(NewBounds, bAlwaysMarkDirty);
-			}
+		// When the playback range is determined by the section bounds, don't mark the change in the playback range otherwise the scene will be marked dirty
+		if (!NewBounds.IsDegenerate())
+		{
+			const bool bAlwaysMarkDirty = false;
+			FocusedMovieScene->SetPlaybackRange(NewBounds, bAlwaysMarkDirty);
 		}
 	}
 }
@@ -3256,6 +3283,11 @@ FGuid FSequencer::GetHandleToObject( UObject* Object, bool bCreateHandleIfMissin
 	UMovieSceneSequence* FocusedMovieSceneSequence = GetFocusedMovieSceneSequence();
 	UMovieScene* FocusedMovieScene = FocusedMovieSceneSequence->GetMovieScene();
 	
+	if (!FocusedMovieScene)
+	{
+		return FGuid();
+	}
+
 	if (FocusedMovieScene->IsReadOnly())
 	{
 		return FGuid();
@@ -3961,6 +3993,10 @@ UObject* FSequencer::FindSpawnedObjectOrTemplate(const FGuid& BindingId)
 	}
 
 	UMovieScene* FocusedMovieScene = Sequence->GetMovieScene();
+	if (!FocusedMovieScene)
+	{
+		return nullptr;
+	}
 
 	FMovieScenePossessable* Possessable = FocusedMovieScene->FindPossessable(BindingId);
 	// If we're a possessable with a parent spawnable and we don't have the object, we look the object up within the default object of the spawnable
@@ -4497,6 +4533,11 @@ EPlaybackMode::Type FSequencer::GetPlaybackMode() const
 void FSequencer::UpdateTimeBoundsToFocusedMovieScene()
 {
 	UMovieScene* FocusedMovieScene = GetFocusedMovieSceneSequence()->GetMovieScene();
+	if (!FocusedMovieScene)
+	{
+		return;
+	}
+
 	FQualifiedFrameTime CurrentTime = GetLocalTime();
 
 	// Set the view range to:
@@ -5012,8 +5053,11 @@ bool FSequencer::OnRequestNodeDeleted( TSharedRef<const FSequencerDisplayNode> N
 		else
 		{
 			UMovieScene* FocusedMovieScene = GetFocusedMovieSceneSequence()->GetMovieScene();
-			FocusedMovieScene->Modify();
-			FocusedMovieScene->GetRootFolders().Remove( &FolderToBeDeleted->GetFolder() );
+			if (FocusedMovieScene)
+			{
+				FocusedMovieScene->Modify();
+				FocusedMovieScene->GetRootFolders().Remove(&FolderToBeDeleted->GetFolder());
+			}
 		}
 
 		bAnythingRemoved = true;
@@ -7032,9 +7076,11 @@ void FSequencer::MoveNodeToFolder(TSharedRef<FSequencerDisplayNode> NodeToMove,
 			else
 			{
 				UMovieScene* FocusedMovieScene = GetFocusedMovieSceneSequence()->GetMovieScene();
-
-				FocusedMovieScene->Modify();
-				FocusedMovieScene->GetRootFolders().Remove(&FolderNode->GetFolder());
+				if (FocusedMovieScene)
+				{
+					FocusedMovieScene->Modify();
+					FocusedMovieScene->GetRootFolders().Remove(&FolderNode->GetFolder());
+				}
 			}
 
 			DestinationFolder->AddChildFolder(&FolderNode->GetFolder());
@@ -7127,6 +7173,11 @@ void FSequencer::MoveSelectedNodesToNewFolder()
 {
 	UMovieScene* FocusedMovieScene = GetFocusedMovieSceneSequence()->GetMovieScene();
 
+	if (!FocusedMovieScene)
+	{
+		return;
+	}
+
 	if (FocusedMovieScene->IsReadOnly())
 	{
 		ShowReadOnlyError();
@@ -9035,6 +9086,10 @@ void FSequencer::OnLoadRecordedData()
 		return;
 	}
 	UMovieScene* FocusedMovieScene = FocusedMovieSceneSequence->GetMovieScene();
+	if (!FocusedMovieScene)
+	{
+		return;
+	}
 	if (FocusedMovieScene->IsReadOnly())
 	{
 		return;
@@ -9119,6 +9174,10 @@ bool FSequencer::ReplaceFolderBindingGUID(UMovieSceneFolder* Folder, FGuid Origi
 void FSequencer::OnAddFolder()
 {
 	UMovieScene* FocusedMovieScene = GetFocusedMovieSceneSequence()->GetMovieScene();
+	if (!FocusedMovieScene)
+	{
+		return;
+	}
 
 	if (FocusedMovieScene->IsReadOnly())
 	{
@@ -9678,22 +9737,22 @@ TSet<FFrameNumber> FSequencer::GetVerticalFrames() const
 	if (FocusedMovieSequence != nullptr)
 	{
 		UMovieScene* FocusedMovieScene = FocusedMovieSequence->GetMovieScene();
-		for (UMovieSceneTrack* MasterTrack : FocusedMovieScene->GetMasterTracks())
+		if (FocusedMovieScene != nullptr)
 		{
-			if (FocusedMovieScene != nullptr)
+			for (UMovieSceneTrack* MasterTrack : FocusedMovieScene->GetMasterTracks())
 			{
-				if (MasterTrack->DisplayOptions.bShowVerticalFrames)
+				if (MasterTrack && MasterTrack->DisplayOptions.bShowVerticalFrames)
 				{
 					AddVerticalFrames(VerticalFrames, MasterTrack);
 				}
 			}
-		}
 
-		if (UMovieSceneTrack* CameraCutTrack = FocusedMovieScene->GetCameraCutTrack())
-		{
-			if (CameraCutTrack->DisplayOptions.bShowVerticalFrames)
+			if (UMovieSceneTrack* CameraCutTrack = FocusedMovieScene->GetCameraCutTrack())
 			{
-				AddVerticalFrames(VerticalFrames, CameraCutTrack);
+				if (CameraCutTrack->DisplayOptions.bShowVerticalFrames)
+				{
+					AddVerticalFrames(VerticalFrames, CameraCutTrack);
+				}
 			}
 		}
 	}
@@ -10083,6 +10142,11 @@ void FSequencer::OnClipboardUsed(TSharedPtr<FMovieSceneClipboard> Clipboard)
 void FSequencer::CreateCamera()
 {
 	UMovieScene* FocusedMovieScene = GetFocusedMovieSceneSequence()->GetMovieScene();
+	if (!FocusedMovieScene)
+	{
+		return;
+	}
+
 	if (FocusedMovieScene->IsReadOnly())
 	{
 		ShowReadOnlyError();
@@ -10202,6 +10266,10 @@ void FSequencer::FixActorReferences()
 	FScopedTransaction FixActorReferencesTransaction( NSLOCTEXT( "Sequencer", "FixActorReferences", "Fix Actor References" ) );
 
 	UMovieScene* FocusedMovieScene = GetFocusedMovieSceneSequence()->GetMovieScene();
+	if (!FocusedMovieScene)
+	{
+		return;
+	}
 
 	TMap<FString, AActor*> ActorNameToActorMap;
 
@@ -10262,6 +10330,10 @@ void FSequencer::RebindPossessableReferences()
 {
 	UMovieSceneSequence* FocusedSequence = GetFocusedMovieSceneSequence();
 	UMovieScene* FocusedMovieScene = FocusedSequence->GetMovieScene();
+	if (!FocusedMovieScene)
+	{
+		return;
+	}
 
 	if (FocusedMovieScene->IsReadOnly())
 	{
