Sequencer - Cache the stored playback state before restoring it in tick.

diff --git a/Engine/Source/Editor/Sequencer/Private/Sequencer.cpp b/Engine/Source/Editor/Sequencer/Private/Sequencer.cpp
--- a/Engine/Source/Editor/Sequencer/Private/Sequencer.cpp
+++ b/Engine/Source/Editor/Sequencer/Private/Sequencer.cpp
@@ -548,7 +548,6 @@ FSequencer::FSequencer()
 	, bPerspectiveViewportCameraCutEnabled( false )
 	, bIsEditingWithinLevelEditor( false )
 	, bNeedTreeRefresh( false )
-	, StoredPlaybackState( EMovieScenePlayerStatus::Stopped )
 	, NodeTree( MakeShareable( new FSequencerNodeTree( *this ) ) )
 	, bUpdatingSequencerSelection( false )
 	, bUpdatingExternalSelection( false )
@@ -664,8 +663,10 @@ void FSequencer::Tick(float InDeltaTime)
 
 	if (bNeedTreeRefresh || NodeTree->NeedsFilterUpdate())
 	{
-		SelectionPreview.Empty();
+		EMovieScenePlayerStatus::Type StoredPlaybackState = GetPlaybackStatus();
+		SetPlaybackStatus(EMovieScenePlayerStatus::Stopped);
 
+		SelectionPreview.Empty();
 		RefreshTree();
 
 		SetPlaybackStatus(StoredPlaybackState);
@@ -2186,14 +2187,13 @@ void FSequencer::NotifyMovieSceneDataChanged( EMovieSceneDataChangeType DataChan
 		}
 	}
 
-	StoredPlaybackState = GetPlaybackStatus();
-
 	if ( DataChangeType == EMovieSceneDataChangeType::MovieSceneStructureItemRemoved ||
 		DataChangeType == EMovieSceneDataChangeType::MovieSceneStructureItemsChanged ||
 		DataChangeType == EMovieSceneDataChangeType::Unknown )
 	{
 		// When structure items are removed, or we don't know what may have changed, refresh the tree and instances immediately so that the data
 		// is in a consistent state when the UI is updated during the next tick.
+		EMovieScenePlayerStatus::Type StoredPlaybackState = GetPlaybackStatus();
 		SetPlaybackStatus( EMovieScenePlayerStatus::Stopped );
 		SelectionPreview.Empty();
 		RefreshTree();
diff --git a/Engine/Source/Editor/Sequencer/Private/Sequencer.h b/Engine/Source/Editor/Sequencer/Private/Sequencer.h
--- a/Engine/Source/Editor/Sequencer/Private/Sequencer.h
+++ b/Engine/Source/Editor/Sequencer/Private/Sequencer.h
@@ -1193,9 +1193,6 @@ private:
 		the MovieScene data can change many times per frame.) */
 	bool bNeedTreeRefresh;
 
-	/** Stores the playback status to be restored on refresh. */
-	EMovieScenePlayerStatus::Type StoredPlaybackState;
-
 	FSequencerSelection Selection;
 	FSequencerSelectionPreview SelectionPreview;
 
