Sequencer: Changed keyframing options to radio buttons for clarity

#rnx

diff --git a/Engine/Source/Editor/Sequencer/Private/SSequencer.cpp b/Engine/Source/Editor/Sequencer/Private/SSequencer.cpp
--- a/Engine/Source/Editor/Sequencer/Private/SSequencer.cpp
+++ b/Engine/Source/Editor/Sequencer/Private/SSequencer.cpp
@@ -2269,7 +2269,7 @@ TSharedRef<SWidget> SSequencer::MakeKeyGroupMenu()
 				FCanExecuteAction(),
 				FIsActionChecked::CreateLambda([this] { return SequencerPtr.Pin()->GetKeyInterpolation() == EMovieSceneKeyInterpolation::Auto; })),
 			NAME_None,
-			EUserInterfaceActionType::ToggleButton
+			EUserInterfaceActionType::RadioButton
 		);
 
 		MenuBuilder.AddMenuEntry(
@@ -2281,7 +2281,7 @@ TSharedRef<SWidget> SSequencer::MakeKeyGroupMenu()
 				FCanExecuteAction(),
 				FIsActionChecked::CreateLambda([this] { return SequencerPtr.Pin()->GetKeyInterpolation() == EMovieSceneKeyInterpolation::User; })),
 			NAME_None,
-			EUserInterfaceActionType::ToggleButton
+			EUserInterfaceActionType::RadioButton
 		);
 
 		MenuBuilder.AddMenuEntry(
@@ -2293,7 +2293,7 @@ TSharedRef<SWidget> SSequencer::MakeKeyGroupMenu()
 				FCanExecuteAction(),
 				FIsActionChecked::CreateLambda([this] { return SequencerPtr.Pin()->GetKeyInterpolation() == EMovieSceneKeyInterpolation::Break; })),
 			NAME_None,
-			EUserInterfaceActionType::ToggleButton
+			EUserInterfaceActionType::RadioButton
 		);
 
 		MenuBuilder.AddMenuEntry(
@@ -2305,7 +2305,7 @@ TSharedRef<SWidget> SSequencer::MakeKeyGroupMenu()
 				FCanExecuteAction(),
 				FIsActionChecked::CreateLambda([this] { return SequencerPtr.Pin()->GetKeyInterpolation() == EMovieSceneKeyInterpolation::Linear; })),
 			NAME_None,
-			EUserInterfaceActionType::ToggleButton
+			EUserInterfaceActionType::RadioButton
 		);
 
 		MenuBuilder.AddMenuEntry(
@@ -2317,7 +2317,7 @@ TSharedRef<SWidget> SSequencer::MakeKeyGroupMenu()
 				FCanExecuteAction(),
 				FIsActionChecked::CreateLambda([this] { return SequencerPtr.Pin()->GetKeyInterpolation() == EMovieSceneKeyInterpolation::Constant; })),
 			NAME_None,
-			EUserInterfaceActionType::ToggleButton
+			EUserInterfaceActionType::RadioButton
 		);
 	}
 	MenuBuilder.EndSection(); // SequencerInterpolation
diff --git a/Engine/Source/Editor/Sequencer/Private/SequencerCommands.cpp b/Engine/Source/Editor/Sequencer/Private/SequencerCommands.cpp
--- a/Engine/Source/Editor/Sequencer/Private/SequencerCommands.cpp
+++ b/Engine/Source/Editor/Sequencer/Private/SequencerCommands.cpp
@@ -76,15 +76,15 @@ void FSequencerCommands::RegisterCommands()
 	UI_COMMAND( SetAutoChangeAll, "All", "Create a key and a track if it doesn't exist when channels/properties change.", EUserInterfaceActionType::ToggleButton, FInputChord());
 	UI_COMMAND( SetAutoChangeNone, "None", "Disable auto-keying and auto-tracking.", EUserInterfaceActionType::ToggleButton, FInputChord());
 
-	UI_COMMAND( AllowAllEdits, "Allow All Edits", "Allow any edits to occur, some of which may produce tracks/keys or modify default properties.", EUserInterfaceActionType::ToggleButton, FInputChord() );
-	UI_COMMAND( AllowSequencerEditsOnly, "Allow Sequencer Edits Only", "All edits will produce either a track or a key.", EUserInterfaceActionType::ToggleButton, FInputChord());
-	UI_COMMAND( AllowLevelEditsOnly, "Allow Level Edits Only", "Properties in the details panel will be disabled if they have a track.", EUserInterfaceActionType::ToggleButton, FInputChord());
+	UI_COMMAND( AllowAllEdits, "Allow All Edits", "Allow any edits to occur, some of which may produce tracks/keys or modify default properties.", EUserInterfaceActionType::RadioButton, FInputChord() );
+	UI_COMMAND( AllowSequencerEditsOnly, "Allow Sequencer Edits Only", "All edits will produce either a track or a key.", EUserInterfaceActionType::RadioButton, FInputChord());
+	UI_COMMAND( AllowLevelEditsOnly, "Allow Level Edits Only", "Properties in the details panel will be disabled if they have a track.", EUserInterfaceActionType::RadioButton, FInputChord());
 
 	UI_COMMAND(ToggleAutoKeyEnabled, "Auto-key", "Create a key when channels/properties change. Only automatically adds a key when there's already a track and at least one key.", EUserInterfaceActionType::ToggleButton, FInputChord());
 
-	UI_COMMAND( SetKeyChanged, "Key Changed", "Key just the changed channel when it changes.", EUserInterfaceActionType::ToggleButton, FInputChord());
-	UI_COMMAND( SetKeyGroup, "Key Group", "Key the groups channels/properties when only one of them changes. ie. Keys all three translation channels when only translation Y changes", EUserInterfaceActionType::ToggleButton, FInputChord());
-	UI_COMMAND( SetKeyAll, "Key All", "Key all channels/properties when only one of them changes. ie. Keys all translation, rotation, scale channels when only translation Y changes", EUserInterfaceActionType::ToggleButton, FInputChord());
+	UI_COMMAND( SetKeyChanged, "Key Changed", "Key just the changed channel when it changes.", EUserInterfaceActionType::RadioButton, FInputChord());
+	UI_COMMAND( SetKeyGroup, "Key Group", "Key the groups channels/properties when only one of them changes. ie. Keys all three translation channels when only translation Y changes", EUserInterfaceActionType::RadioButton, FInputChord());
+	UI_COMMAND( SetKeyAll, "Key All", "Key all channels/properties when only one of them changes. ie. Keys all translation, rotation, scale channels when only translation Y changes", EUserInterfaceActionType::RadioButton, FInputChord());
 
 	UI_COMMAND( ToggleMarkAtPlayPosition, "Toggle Mark", "Sets or clears a mark at the current play position.", EUserInterfaceActionType::Button, FInputChord(EKeys::M) );
 	UI_COMMAND( StepToNextMark, "Step to Next Marked Frame", "Step to the next marked frame", EUserInterfaceActionType::Button, FInputChord(EModifierKey::Control | EModifierKey::Shift, EKeys::Period) );
