Sequencer: Added missing transactions for mute and solo

#rnx

diff --git a/Engine/Source/Editor/Sequencer/Private/SequencerNodeTree.cpp b/Engine/Source/Editor/Sequencer/Private/SequencerNodeTree.cpp
--- a/Engine/Source/Editor/Sequencer/Private/SequencerNodeTree.cpp
+++ b/Engine/Source/Editor/Sequencer/Private/SequencerNodeTree.cpp
@@ -20,6 +20,7 @@
 #include "SequencerNodeSortingMethods.h"
 #include "SequencerTrackFilters.h"
 #include "Channels/MovieSceneChannel.h"
+#include "ScopedTransaction.h"
 
 FSequencerNodeTree::~FSequencerNodeTree()
 {
@@ -647,6 +648,8 @@ void FSequencerNodeTree::ToggleSelectedNodesSolo()
 		}
 	}
 
+	const FScopedTransaction Transaction(NSLOCTEXT("Sequencer", "ToggleSolo", "Toggle Solo"));
+
 	MovieScene->Modify();
 
 	for (const TSharedRef<const FSequencerDisplayNode> Node : Sequencer.GetSelection().GetSelectedOutlinerNodes())
@@ -740,6 +743,8 @@ void FSequencerNodeTree::ToggleSelectedNodesMute()
 		}
 	}
 
+	const FScopedTransaction Transaction(NSLOCTEXT("Sequencer", "ToggleMute", "Toggle Mute"));
+
 	MovieScene->Modify();
 
 	for (const TSharedRef<const FSequencerDisplayNode> Node : SelectedNodes)
