Sequencer: Step to next/previous key should wrap around

#rnx

diff --git a/Engine/Source/Editor/Sequencer/Private/SKeyNavigationButtons.h b/Engine/Source/Editor/Sequencer/Private/SKeyNavigationButtons.h
--- a/Engine/Source/Editor/Sequencer/Private/SKeyNavigationButtons.h
+++ b/Engine/Source/Editor/Sequencer/Private/SKeyNavigationButtons.h
@@ -179,6 +179,13 @@ public:
 			}
 		}
 
+		if (!PreviousTime.IsSet() && AllTimes.Num() > 0)
+		{
+			AllTimes.Sort();
+
+			PreviousTime = AllTimes.Last();
+		}
+
 		if (PreviousTime.IsSet())
 		{
 			Sequencer.SetLocalTime(PreviousTime.GetValue());
@@ -230,6 +237,13 @@ public:
 			}
 		}
 
+		if (!NextTime.IsSet() && AllTimes.Num() > 0)
+		{
+			AllTimes.Sort();
+
+			NextTime = AllTimes[0];
+		}
+
 		if (NextTime.IsSet())
 		{
 			Sequencer.SetLocalTime(NextTime.GetValue());
