Sequencer: Save button in curve editor

diff --git a/Engine/Source/Editor/Sequencer/Private/SSequencer.cpp b/Engine/Source/Editor/Sequencer/Private/SSequencer.cpp
--- a/Engine/Source/Editor/Sequencer/Private/SSequencer.cpp
+++ b/Engine/Source/Editor/Sequencer/Private/SSequencer.cpp
@@ -108,8 +108,10 @@ class SSequencerCurveEditor : public SCompoundWidget
 	{}
 	SLATE_END_ARGS()
 public:
-	void Construct(const FArguments& InArgs, TSharedRef<SCurveEditorPanel> InEditorPanel)
+	void Construct(const FArguments& InArgs, TSharedRef<SCurveEditorPanel> InEditorPanel, TSharedPtr<FSequencer> InSequencer)
 	{
+		WeakSequencer = InSequencer;
+
 		ChildSlot
 		[
 			SNew(SVerticalBox)
@@ -132,10 +134,44 @@ public:
 		FToolBarBuilder ToolBarBuilder(InEditorPanel->GetCommands(), FMultiBoxCustomization::None, InEditorPanel->GetToolbarExtender(), EOrientation::Orient_Horizontal, true);
 		ToolBarBuilder.SetStyle(&FEditorStyle::Get(), "Sequencer.ToolBar");
 		ToolBarBuilder.BeginSection("Asset");
+
+		{
+			TAttribute<FSlateIcon> SaveIcon;
+			SaveIcon.Bind(TAttribute<FSlateIcon>::FGetter::CreateLambda([&] {
+
+				TArray<UMovieScene*> MovieScenesToSave;
+				MovieSceneHelpers::GetDescendantMovieScenes(WeakSequencer.Pin()->GetRootMovieSceneSequence(), MovieScenesToSave);
+				for (UMovieScene* MovieSceneToSave : MovieScenesToSave)
+				{
+					UPackage* MovieScenePackageToSave = MovieSceneToSave->GetOuter()->GetOutermost();
+					if (MovieScenePackageToSave->IsDirty())
+					{
+						return FSlateIcon(FEditorStyle::GetStyleSetName(), "Sequencer.SaveAsterisk");
+					}
+				}
+
+				return FSlateIcon(FEditorStyle::GetStyleSetName(), "Sequencer.Save");
+				}));
+
+			if (WeakSequencer.Pin()->GetHostCapabilities().bSupportsSaveMovieSceneAsset)
+			{
+				ToolBarBuilder.AddToolBarButton(
+					FUIAction(FExecuteAction::CreateLambda([&] { WeakSequencer.Pin()->SaveCurrentMovieScene(); })),
+					NAME_None,
+					LOCTEXT("SaveDirtyPackages", "Save"),
+					LOCTEXT("SaveDirtyPackagesTooltip", "Saves the current sequence and any subsequences"),
+					SaveIcon
+				);
+			}		
+		}
+		
 		ToolBarBuilder.EndSection();
 		// We just use all of the extenders as our toolbar, we don't have a need to create a separate toolbar.
 		return ToolBarBuilder.MakeWidget();
 	}
+
+private:
+	TWeakPtr<FSequencer> WeakSequencer;
 };
 
 class FSequencerCurveEditorTimeSliderController : public FSequencerTimeSliderController
@@ -488,7 +524,7 @@ void SSequencer::Construct(const FArguments& InArgs, TSharedRef<FSequencer> InSe
 		CurveEditorWidget->GetKeyDetailsView()->GetPropertyRowGenerator()->RegisterInstancedCustomPropertyTypeLayout("FrameNumber", FOnGetPropertyTypeCustomizationInstance::CreateStatic(CreateFrameNumberCustomization, SequencerPtr));
 		TAttribute<bool> IsEnabledAttribute = TAttribute<bool>::Create(TAttribute<bool>::FGetter::CreateSP(this, &SSequencer::GetIsCurveEditorEnabled));
 	
-		CurveEditorPanel = SNew(SSequencerCurveEditor, CurveEditorWidget);
+		CurveEditorPanel = SNew(SSequencerCurveEditor, CurveEditorWidget, InSequencer);
 		CurveEditorPanel->SetEnabled(IsEnabledAttribute);
 		CurveEditorWidget->SetEnabled(IsEnabledAttribute);
 
@@ -1280,28 +1316,17 @@ TSharedRef<SWidget> SSequencer::MakeToolBar()
 			TAttribute<FSlateIcon> SaveIcon;
 			SaveIcon.Bind(TAttribute<FSlateIcon>::FGetter::CreateLambda([&] {
 
-				bool bAnyMovieSceneDirty = false;
-
 				TArray<UMovieScene*> MovieScenesToSave;
 				MovieSceneHelpers::GetDescendantMovieScenes(SequencerPtr.Pin()->GetRootMovieSceneSequence(), MovieScenesToSave);
-				for (auto MovieSceneToSave : MovieScenesToSave)
+				for (UMovieScene* MovieSceneToSave : MovieScenesToSave)
 				{
 					UPackage* MovieScenePackageToSave = MovieSceneToSave->GetOuter()->GetOutermost();
 					if (MovieScenePackageToSave->IsDirty())
 					{
-						bAnyMovieSceneDirty = true;
-						break;
+						return FSlateIcon(FEditorStyle::GetStyleSetName(), "Sequencer.SaveAsterisk");
 					}
 				}
-
-				if (bAnyMovieSceneDirty)
-				{
-					return FSlateIcon(FEditorStyle::GetStyleSetName(), "Sequencer.SaveAsterisk");
-				}
-				else
-				{
-					return FSlateIcon(FEditorStyle::GetStyleSetName(), "Sequencer.Save");
-				}
+				return FSlateIcon(FEditorStyle::GetStyleSetName(), "Sequencer.Save");
 			}));
 
 			if (SequencerPtr.Pin()->GetHostCapabilities().bSupportsSaveMovieSceneAsset)
