Sequencer: Fix copy/paste for tracks with multiple rows (subtracks). When a subtrack is copied, copy only the sections otherwise all the sections across all the rows will get pasted. When pasting a sectioini onto a subtrack, set the row index to the subtrack's row index.

diff --git a/Engine/Source/Editor/Sequencer/Private/Sequencer.cpp b/Engine/Source/Editor/Sequencer/Private/Sequencer.cpp
--- a/Engine/Source/Editor/Sequencer/Private/Sequencer.cpp
+++ b/Engine/Source/Editor/Sequencer/Private/Sequencer.cpp
@@ -7518,7 +7518,7 @@ void FSequencer::CopySelectedTracks(TArray<TSharedPtr<FSequencerTrackNode>>& Tra
 {
 	UMovieScene* MovieScene = GetFocusedMovieSceneSequence()->GetMovieScene();
 
-	TArray<UObject*> CopyableTracks;
+	TArray<UObject*> CopyableObjects;
 	for (TSharedPtr<FSequencerTrackNode> TrackNode : TrackNodes)
 	{
 		bool bIsParentSelected = false;
@@ -7535,16 +7535,33 @@ void FSequencer::CopySelectedTracks(TArray<TSharedPtr<FSequencerTrackNode>>& Tra
 
 		if (!bIsParentSelected)
 		{
-			UMovieSceneCopyableTrack *CopyableTrack = NewObject<UMovieSceneCopyableTrack>(GetTransientPackage(), UMovieSceneCopyableTrack::StaticClass(), NAME_None, RF_Transient);
-			CopyableTracks.Add(CopyableTrack);
+			// If this is a subtrack, only copy the sections that belong to this row. otherwise copying the entire track will copy all the sections across all the rows
+			if (TrackNode->GetSubTrackMode() == FSequencerTrackNode::ESubTrackMode::SubTrack)
+			{
+				for (UMovieSceneSection* Section : TrackNode->GetTrack()->GetAllSections())
+				{
+					if (Section && Section->GetRowIndex() == TrackNode->GetRowIndex())
+					{
+						CopyableObjects.Add(Section);
+					}
+				}
+			}
+			else
+			{
+				UMovieSceneCopyableTrack* CopyableTrack = NewObject<UMovieSceneCopyableTrack>(GetTransientPackage(), UMovieSceneCopyableTrack::StaticClass(), NAME_None, RF_Transient);
+				CopyableObjects.Add(CopyableTrack);
 
-			UMovieSceneTrack* DuplicatedTrack = Cast<UMovieSceneTrack>(StaticDuplicateObject(TrackNode->GetTrack(), CopyableTrack));
-			CopyableTrack->Track = DuplicatedTrack;
-			CopyableTrack->bIsAMasterTrack = MovieScene->IsAMasterTrack(*TrackNode->GetTrack());
+				UMovieSceneTrack* DuplicatedTrack = Cast<UMovieSceneTrack>(StaticDuplicateObject(TrackNode->GetTrack(), CopyableTrack));
+				CopyableTrack->Track = DuplicatedTrack;
+				CopyableTrack->bIsAMasterTrack = MovieScene->IsAMasterTrack(*TrackNode->GetTrack());
+			}
 		}
 	}
 
-	ExportObjectsToText(CopyableTracks, /*out*/ ExportedText);
+	if (CopyableObjects.Num())
+	{
+		ExportObjectsToText(CopyableObjects, /*out*/ ExportedText);
+	}
 }
 
 
@@ -8032,27 +8049,24 @@ bool FSequencer::PasteTracks(const FString& TextToImport, TArray<FNotificationIn
 	return true;
 }
 
-void GetSupportedTracks(TSharedRef<FSequencerDisplayNode> DisplayNode, const TArray<UMovieSceneSection*>& ImportedSections, TSet<UMovieSceneTrack*>& TracksToPasteOnto)
+void GetSupportedTracks(TSharedRef<FSequencerDisplayNode> DisplayNode, const TArray<UMovieSceneSection*>& ImportedSections, TArray<TSharedRef<FSequencerTrackNode>>& TracksToPasteOnto)
 {
 	if (DisplayNode->GetType() != ESequencerNode::Track)
 	{
 		return;
 	}
 
-	TSharedPtr<FSequencerTrackNode> TrackNode = StaticCastSharedRef<FSequencerTrackNode>(DisplayNode);
-	if (TrackNode.IsValid())
+	TSharedRef<FSequencerTrackNode> TrackNode = StaticCastSharedRef<FSequencerTrackNode>(DisplayNode);
+	UMovieSceneTrack* Track = TrackNode->GetTrack();
+	if (Track)
 	{
-		UMovieSceneTrack* Track = TrackNode->GetTrack();
-		if (Track)
+		for (int32 SectionIndex = 0; SectionIndex < ImportedSections.Num(); ++SectionIndex)
 		{
-			for (int32 SectionIndex = 0; SectionIndex < ImportedSections.Num(); ++SectionIndex)
-			{
-				UMovieSceneSection* Section = ImportedSections[SectionIndex];
+			UMovieSceneSection* Section = ImportedSections[SectionIndex];
 
-				if (Track->SupportsType(Section->GetClass()))
-				{
-					TracksToPasteOnto.Add(Track);
-				}
+			if (Track->SupportsType(Section->GetClass()))
+			{
+				TracksToPasteOnto.Add(TrackNode);
 			}
 		}
 	}
@@ -8107,7 +8121,7 @@ bool FSequencer::PasteSections(const FString& TextToImport, TArray<FNotification
 	}
 
 	// Check if any of the selected nodes supports pasting this type of section
-	TSet<UMovieSceneTrack*> TracksToPasteOnto;
+	TArray<TSharedRef<FSequencerTrackNode>> TracksToPasteOnto;
 	for (TSharedRef<FSequencerDisplayNode> Node : SelectedNodes)
 	{
 		GetSupportedTracks(Node, ImportedSections, TracksToPasteOnto);
@@ -8131,12 +8145,12 @@ bool FSequencer::PasteSections(const FString& TextToImport, TArray<FNotification
 	TArray<UMovieSceneSection*> NewSections;
 	TArray<int32> SectionIndicesImported;
 
-	for (UMovieSceneTrack* Track : TracksToPasteOnto)
+	for (TSharedRef<FSequencerTrackNode> TrackNode : TracksToPasteOnto)
 	{
 		for (int32 SectionIndex = 0; SectionIndex < ImportedSections.Num(); ++SectionIndex)
 		{
 			UMovieSceneSection* Section = ImportedSections[SectionIndex];
-
+			UMovieSceneTrack* Track = TrackNode->GetTrack();
 			if (!Track->SupportsType(Section->GetClass()))
 			{
 				continue;
@@ -8148,7 +8162,14 @@ bool FSequencer::PasteSections(const FString& TextToImport, TArray<FNotification
 
 			Section->Rename(nullptr, Track);
 			Track->AddSection(*Section);
-			Section->SetRowIndex(MovieSceneToolHelpers::FindAvailableRowIndex(Track, Section));
+			if (TrackNode->GetSubTrackMode() == FSequencerTrackNode::ESubTrackMode::SubTrack)
+			{
+				Section->SetRowIndex(TrackNode->GetRowIndex());
+			}
+			else
+			{
+				Section->SetRowIndex(MovieSceneToolHelpers::FindAvailableRowIndex(Track, Section));
+			}
 
 			if (Section->HasStartFrame())
 			{
