Sequencer: Add support for toggling markers in a subsequence to be globally visible in other subsequences

diff --git a/Engine/Source/Editor/Sequencer/Private/SSequencer.cpp b/Engine/Source/Editor/Sequencer/Private/SSequencer.cpp
--- a/Engine/Source/Editor/Sequencer/Private/SSequencer.cpp
+++ b/Engine/Source/Editor/Sequencer/Private/SSequencer.cpp
@@ -394,6 +394,7 @@ void SSequencer::Construct(const FArguments& InArgs, TSharedRef<FSequencer> InSe
 		TimeSliderArgs.SubSequenceRange = InArgs._SubSequenceRange;
 		TimeSliderArgs.VerticalFrames = InArgs._VerticalFrames;
 		TimeSliderArgs.MarkedFrames = InArgs._MarkedFrames;
+		TimeSliderArgs.GlobalMarkedFrames = InArgs._GlobalMarkedFrames;
 		TimeSliderArgs.OnSetMarkedFrame = InArgs._OnSetMarkedFrame;
 		TimeSliderArgs.OnAddMarkedFrame = InArgs._OnAddMarkedFrame;
 		TimeSliderArgs.OnDeleteMarkedFrame = InArgs._OnDeleteMarkedFrame;
@@ -1904,6 +1905,47 @@ TSharedRef<SWidget> SSequencer::MakeViewMenu()
 	}
 	MenuBuilder.EndSection();
 
+	MenuBuilder.BeginSection("MarkedFrames", LOCTEXT("MarkedFramesHeader", "Marked Frames"));
+	{
+		UMovieSceneSequence* FocusedMovieSequence = Sequencer->GetFocusedMovieSceneSequence();
+		UMovieScene* FocusedMovieScene = nullptr;
+		if (FocusedMovieSequence != nullptr)
+		{
+			FocusedMovieScene = FocusedMovieSequence->GetMovieScene();
+		}
+
+		MenuBuilder.AddMenuEntry(
+			LOCTEXT("ToggleShowMarkedFrameGlobal", "Show Marked Frames Globally"),
+			LOCTEXT("ToggleShowMarkedFrameGlobalTooltip", "Makes marked frames in this sub-sequence visible in parent/sibling sequences for the current editor session"),
+			FSlateIcon(),
+			FUIAction(
+				FExecuteAction::CreateLambda([FocusedMovieScene, this]() {
+					if (FocusedMovieScene != nullptr)
+					{
+						FocusedMovieScene->ToggleGloballyShowMarkedFrames();
+						SequencerPtr.Pin()->InvalidateGlobalMarkedFramesCache();
+					}
+				}),
+				FCanExecuteAction::CreateLambda([FocusedMovieScene]() { return FocusedMovieScene != nullptr; }),
+				FIsActionChecked::CreateLambda([FocusedMovieScene]() {
+					if (FocusedMovieScene != nullptr)
+					{
+						return FocusedMovieScene->GetGloballyShowMarkedFrames();
+					}
+					return false;
+				})),
+			NAME_None,
+			EUserInterfaceActionType::ToggleButton
+		);
+		MenuBuilder.AddMenuEntry(
+			LOCTEXT("ClearMarkedFramesGlobal", "Clear Global Marked Frames"),
+			LOCTEXT("ClearMarkedFramesGlobalTooltip", "Set all marked frames in all sub-sequences to not be globally displayed."),
+			FSlateIcon(),
+			FUIAction(FExecuteAction::CreateLambda([this]() { SequencerPtr.Pin()->ClearGlobalMarkedFrames(); }))
+		);
+	}
+	MenuBuilder.EndSection();
+
 	return MenuBuilder.MakeWidget();
 }
 
diff --git a/Engine/Source/Editor/Sequencer/Private/SSequencer.h b/Engine/Source/Editor/Sequencer/Private/SSequencer.h
--- a/Engine/Source/Editor/Sequencer/Private/SSequencer.h
+++ b/Engine/Source/Editor/Sequencer/Private/SSequencer.h
@@ -150,6 +150,9 @@ public:
 		/** The Marked Frames */
 		SLATE_ATTRIBUTE(TArray<FMovieSceneMarkedFrame>, MarkedFrames)
 
+		/** The Global Marked Frames */
+		SLATE_ATTRIBUTE(TArray<FMovieSceneMarkedFrame>, GlobalMarkedFrames)
+
 		/** The current sub sequence range */
 		SLATE_ATTRIBUTE( TOptional<TRange<FFrameNumber>>, SubSequenceRange)
 
@@ -429,6 +432,8 @@ private:
 
 	void FillLevelFilterMenu(FMenuBuilder& InMenuBarBuilder);
 
+	void ClearGlobalMarkedFrames();
+
 	/**
 	* Called when the time snap interval changes.
 	*
diff --git a/Engine/Source/Editor/Sequencer/Private/Sequencer.cpp b/Engine/Source/Editor/Sequencer/Private/Sequencer.cpp
--- a/Engine/Source/Editor/Sequencer/Private/Sequencer.cpp
+++ b/Engine/Source/Editor/Sequencer/Private/Sequencer.cpp
@@ -444,6 +444,7 @@ void FSequencer::InitSequencer(const FSequencerInitParams& InitParams, const TSh
 		.SelectionRange( this, &FSequencer::GetSelectionRange )
 		.VerticalFrames(this, &FSequencer::GetVerticalFrames)
 		.MarkedFrames(this, &FSequencer::GetMarkedFrames)
+		.GlobalMarkedFrames(this, &FSequencer::GetGlobalMarkedFrames)
 		.OnSetMarkedFrame(this, &FSequencer::SetMarkedFrame)
 		.OnAddMarkedFrame(this, &FSequencer::AddMarkedFrame)
 		.OnDeleteMarkedFrame(this, &FSequencer::DeleteMarkedFrame)
@@ -798,6 +799,11 @@ void FSequencer::Tick(float InDeltaTime)
 		}
 		PrePossessionViewTargets.Reset();
 	}
+
+	if (!bGlobalMarkedFramesCached)
+	{
+		UpdateGlobalMarkedFramesCache();
+	}
 }
 
 
@@ -944,6 +950,7 @@ void FSequencer::FocusSequenceInstance(UMovieSceneSubSection& InSubSection)
 	OnActivateSequenceEvent.Broadcast(ActiveTemplateIDs.Top());
 
 	bNeedsEvaluate = true;
+	bGlobalMarkedFramesCached = false;
 }
 
 
@@ -1105,6 +1112,7 @@ void FSequencer::PopToSequenceInstance(FMovieSceneSequenceIDRef SequenceID)
 		OnActivateSequenceEvent.Broadcast(ActiveTemplateIDs.Top());
 
 		bNeedsEvaluate = true;
+		bGlobalMarkedFramesCached = false;
 	}
 }
 
@@ -2233,6 +2241,7 @@ void FSequencer::NotifyMovieSceneDataChanged( EMovieSceneDataChangeType DataChan
 		}
 	}
 
+	bGlobalMarkedFramesCached = false;
 	bNeedsEvaluate = true;
 	State.ClearObjectCaches(*this);
 
@@ -9796,6 +9805,117 @@ TArray<FMovieSceneMarkedFrame> FSequencer::GetMarkedFrames() const
 	return TArray<FMovieSceneMarkedFrame>();
 }
 
+void FindGlobalMarkedFrames(TArray<FMovieSceneMarkedFrame>& GlobalMarkedFrames, const UMovieScene* MovieScene, const FFrameNumber FrameOffset, const UMovieScene* FocusedMovieScene)
+{
+	if (MovieScene == nullptr)
+	{
+		return;
+	}
+
+	if (MovieScene != FocusedMovieScene && MovieScene->GetGloballyShowMarkedFrames())
+	{
+		TArray<FMovieSceneMarkedFrame> MarkedFrames = MovieScene->GetMarkedFrames();
+		for (FMovieSceneMarkedFrame MarkedFrame : MarkedFrames)
+		{
+			MarkedFrame.FrameNumber += FrameOffset;
+			GlobalMarkedFrames.Add(MarkedFrame);
+		}
+	}
+
+	const TArray<UMovieSceneTrack*>& Tracks = MovieScene->GetMasterTracks();
+	for (const UMovieSceneTrack* Track : Tracks)
+	{
+		if (Track->SupportsType(UMovieSceneSubSection::StaticClass()))
+		{
+			for (const UMovieSceneSection* Section : Track->GetAllSections())
+			{
+				const UMovieSceneSubSection* SubSection = Cast<UMovieSceneSubSection>(Section);
+				if (SubSection != nullptr)
+				{
+					const UMovieSceneSequence* SubSequence = SubSection->GetSequence();
+					if (SubSequence != nullptr)
+					{
+						const UMovieScene* SubSequenceScene = SubSequence->GetMovieScene();
+						if (SubSequenceScene != nullptr)
+						{
+							FindGlobalMarkedFrames(GlobalMarkedFrames, SubSequenceScene, SubSection->GetInclusiveStartFrame() + FrameOffset, FocusedMovieScene);
+						}
+					}
+				}
+			}
+		}
+		else if(Track->SupportsType(UMovieSceneCinematicShotSection::StaticClass()))
+		{
+			for (const UMovieSceneSection* Section : Track->GetAllSections())
+			{
+				const UMovieSceneCinematicShotSection* ShotSection = Cast<UMovieSceneCinematicShotSection>(Section);
+				if (ShotSection != nullptr)
+				{
+					const UMovieSceneSequence* SubSequence = ShotSection->GetSequence();
+					if (SubSequence != nullptr)
+					{
+						const UMovieScene* SubSequenceScene = SubSequence->GetMovieScene();
+						if (SubSequenceScene != nullptr)
+						{
+							FindGlobalMarkedFrames(GlobalMarkedFrames, SubSequenceScene, ShotSection->GetInclusiveStartFrame() + FrameOffset, FocusedMovieScene);
+						}
+					}
+				}
+			}
+		}
+	}
+}
+
+TArray<FMovieSceneMarkedFrame> FSequencer::GetGlobalMarkedFrames() const
+{
+	return GlobalMarkedFramesCache;
+}
+
+void FSequencer::UpdateGlobalMarkedFramesCache()
+{
+	GlobalMarkedFramesCache.Empty();
+
+	UMovieSceneSequence* FocusedMovieSequence = GetFocusedMovieSceneSequence();
+	UMovieScene* FocusedMovieScene = nullptr;
+	if (FocusedMovieSequence != nullptr)
+	{
+		FocusedMovieScene = FocusedMovieSequence->GetMovieScene();
+	}
+
+	UMovieSceneSequence* RootMovieSequence = GetRootMovieSceneSequence();
+	if (RootMovieSequence != nullptr)
+	{
+		UMovieScene* RootMovieScene = RootMovieSequence->GetMovieScene();
+		if (RootMovieScene != nullptr)
+		{
+
+			TRange<FFrameTime> Range(0,0);
+			Range = RootToLocalTransform.TransformRangeUnwarped(Range);
+			FFrameNumber FocusedMovieSceneFrameOffset = Range.GetLowerBoundValue().FrameNumber;
+
+			FindGlobalMarkedFrames(GlobalMarkedFramesCache, RootMovieScene, FocusedMovieSceneFrameOffset, FocusedMovieScene);
+		}
+	}
+	
+	bGlobalMarkedFramesCached = true;
+}
+
+void FSequencer::ClearGlobalMarkedFrames()
+{
+	bGlobalMarkedFramesCached = false;
+
+	for (const TTuple<FMovieSceneSequenceID, FMovieSceneSubSequenceData>& Pair : RootTemplateInstance.GetHierarchy().AllSubSequenceData())
+	{
+		if (UMovieSceneSequence* Sequence = Pair.Value.GetSequence())
+		{
+			if (UMovieScene* MovieScene = Sequence->GetMovieScene())
+			{
+				MovieScene->SetGloballyShowMarkedFrames(false);
+			}
+		}
+	}
+}
+
 void FSequencer::ToggleMarkAtPlayPosition()
 {
 	UMovieSceneSequence* FocusedMovieSequence = GetFocusedMovieSceneSequence();
diff --git a/Engine/Source/Editor/Sequencer/Private/Sequencer.h b/Engine/Source/Editor/Sequencer/Private/Sequencer.h
--- a/Engine/Source/Editor/Sequencer/Private/Sequencer.h
+++ b/Engine/Source/Editor/Sequencer/Private/Sequencer.h
@@ -263,6 +263,15 @@ public:
 	/** @return The set of marked frames */
 	TArray<FMovieSceneMarkedFrame> GetMarkedFrames() const;
 
+	TArray<FMovieSceneMarkedFrame> GetGlobalMarkedFrames() const;
+	void InvalidateGlobalMarkedFramesCache() { bGlobalMarkedFramesCached = false; }
+	void UpdateGlobalMarkedFramesCache();
+
+	/** 
+	  * Disables all global marked frames from all sub-sequences
+	  */
+	void ClearGlobalMarkedFrames();
+
 protected:
 
 	/** Set/Clear a Mark at the current time */
@@ -1288,6 +1297,9 @@ private:
 
 	FAcquiredResources AcquiredResources;
 
+	bool bGlobalMarkedFramesCached;
+	TArray<FMovieSceneMarkedFrame> GlobalMarkedFramesCache;
+
 	/** The range of the currently displayed sub sequence in relation to its parent section, in the resolution of the current sub sequence */
 	TRange<FFrameNumber> SubSequenceRange;
 
diff --git a/Engine/Source/Editor/Sequencer/Private/SequencerTimeSliderController.cpp b/Engine/Source/Editor/Sequencer/Private/SequencerTimeSliderController.cpp
--- a/Engine/Source/Editor/Sequencer/Private/SequencerTimeSliderController.cpp
+++ b/Engine/Source/Editor/Sequencer/Private/SequencerTimeSliderController.cpp
@@ -300,7 +300,8 @@ void FSequencerTimeSliderController::DrawTicks( FSlateWindowElementList& OutDraw
 int32 FSequencerTimeSliderController::DrawMarkedFrames( const FGeometry& AllottedGeometry, const FScrubRangeToScreen& RangeToScreen, FSlateWindowElementList& OutDrawElements, int32 LayerId, const ESlateDrawEffect& DrawEffects, bool bDrawLabels ) const
 {
 	const TArray<FMovieSceneMarkedFrame> & MarkedFrames = TimeSliderArgs.MarkedFrames.Get();
-	if (MarkedFrames.Num() < 1)
+	const TArray<FMovieSceneMarkedFrame> & GlobalMarkedFrames = TimeSliderArgs.GlobalMarkedFrames.Get();
+	if (MarkedFrames.Num() < 1 && GlobalMarkedFrames.Num() < 1)
 	{
 		return LayerId;
 	}
@@ -310,50 +311,57 @@ int32 FSequencerTimeSliderController::DrawMarkedFrames( const FGeometry& Allotte
 
 	FQualifiedFrameTime ScrubPosition = FQualifiedFrameTime(TimeSliderArgs.ScrubPosition.Get(), GetTickResolution());
 	FScrubberMetrics    ScrubMetrics = GetScrubPixelMetrics(ScrubPosition, RangeToScreen);
+	auto DrawFrameMarkers = ([=](const TArray<FMovieSceneMarkedFrame> & InMarkedFrames, FSlateWindowElementList& DrawElements, bool bFade) {
 
-	for (const FMovieSceneMarkedFrame& MarkedFrame : MarkedFrames)
-	{
-		double Seconds = MarkedFrame.FrameNumber / GetTickResolution();
-
-		const float  LinePos = RangeToScreen.InputToLocalX( Seconds );
-		TArray<FVector2D> LinePoints;
-		LinePoints.AddUninitialized(2);
-		LinePoints[0] = FVector2D( LinePos, 0.0f );
-		LinePoints[1] = FVector2D( LinePos, FMath::FloorToFloat( AllottedGeometry.Size.Y ) );
-
-		FSlateDrawElement::MakeLines(
-			OutDrawElements,
-			LayerId+1,
-			AllottedGeometry.ToPaintGeometry(),
-			LinePoints,
-			DrawEffects,
-			MarkedFrame.Color,
-			false
-		);
-
-		FString LabelString = MarkedFrame.Label;
-		if (bDrawLabels && !LabelString.IsEmpty())
+		for (const FMovieSceneMarkedFrame& MarkedFrame : InMarkedFrames)
 		{
-			// Draw the label next to the marked frame line
-			FVector2D TextSize = FontMeasureService->Measure(LabelString, SmallLayoutFont);
+			double Seconds = MarkedFrame.FrameNumber / GetTickResolution();
 
-			// Flip the text position if getting near the end of the view range
-			static const float TextOffsetPx = 2.f;
-			bool  bDrawLeft = (AllottedGeometry.Size.X - LinePos) < (TextSize.X + 14.f) - TextOffsetPx;
-			float TextPosition = bDrawLeft ? LinePos - TextSize.X - TextOffsetPx : LinePos + TextOffsetPx;
+			FLinearColor DrawColor = bFade ? MarkedFrame.Color.Desaturate(0.25f) : MarkedFrame.Color;
+			const float  LinePos = RangeToScreen.InputToLocalX(Seconds);
+			TArray<FVector2D> LinePoints;
+			LinePoints.AddUninitialized(2);
+			LinePoints[0] = FVector2D(LinePos, 0.0f);
+			LinePoints[1] = FVector2D(LinePos, FMath::FloorToFloat(AllottedGeometry.Size.Y));
 
-			FSlateDrawElement::MakeText(
-				OutDrawElements,
+			FSlateDrawElement::MakeLines(
+				DrawElements,
 				LayerId + 1,
-				AllottedGeometry.ToPaintGeometry(FVector2D(TextPosition, 0.f), TextSize),
-				LabelString,
-				SmallLayoutFont,
+				AllottedGeometry.ToPaintGeometry(),
+				LinePoints,
 				DrawEffects,
-				MarkedFrame.Color
+				DrawColor,
+				false
 			);
+
+			FString LabelString = MarkedFrame.Label;
+			if (bDrawLabels && !LabelString.IsEmpty())
+			{
+				// Draw the label next to the marked frame line
+				FVector2D TextSize = FontMeasureService->Measure(LabelString, SmallLayoutFont);
+
+				// Flip the text position if getting near the end of the view range
+				static const float TextOffsetPx = 2.f;
+				bool  bDrawLeft = (AllottedGeometry.Size.X - LinePos) < (TextSize.X + 14.f) - TextOffsetPx;
+				float TextPosition = bDrawLeft ? LinePos - TextSize.X - TextOffsetPx : LinePos + TextOffsetPx;
+
+				FSlateDrawElement::MakeText(
+					DrawElements,
+					LayerId + 1,
+					AllottedGeometry.ToPaintGeometry(FVector2D(TextPosition, 0.f), TextSize),
+					LabelString,
+					SmallLayoutFont,
+					DrawEffects,
+					DrawColor
+				);
+			}
 		}
-	}
 
+	});
+	
+	DrawFrameMarkers(GlobalMarkedFrames, OutDrawElements, true);
+	DrawFrameMarkers(MarkedFrames, OutDrawElements, false);
+	
 	return LayerId + 1;
 }
 
diff --git a/Engine/Source/Editor/Sequencer/Private/Tools/SequencerSnapField.cpp b/Engine/Source/Editor/Sequencer/Private/Tools/SequencerSnapField.cpp
--- a/Engine/Source/Editor/Sequencer/Private/Tools/SequencerSnapField.cpp
+++ b/Engine/Source/Editor/Sequencer/Private/Tools/SequencerSnapField.cpp
@@ -52,7 +52,7 @@ struct FSnapGridVisitor : ISequencerEntityVisitor
 	mutable TArray<FSequencerSnapPoint> Snaps;
 };
 
-FSequencerSnapField::FSequencerSnapField(const ISequencer& InSequencer, ISequencerSnapCandidate& Candidate, uint32 EntityMask)
+FSequencerSnapField::FSequencerSnapField(const FSequencer& InSequencer, ISequencerSnapCandidate& Candidate, uint32 EntityMask)
 {
 	TSharedPtr<SSequencerTreeView> TreeView = StaticCastSharedRef<SSequencer>(InSequencer.GetSequencerWidget())->GetTreeView();
 
@@ -96,6 +96,12 @@ FSequencerSnapField::FSequencerSnapField(const ISequencer& InSequencer, ISequenc
 		Visitor.Snaps.Add( FSequencerSnapPoint{ FSequencerSnapPoint::Mark, MarkedFrame.FrameNumber } );
 	}
 
+	// Add in the global marked frames
+	for (const FMovieSceneMarkedFrame& MarkedFrame : InSequencer.GetGlobalMarkedFrames())
+	{
+		Visitor.Snaps.Add(FSequencerSnapPoint{ FSequencerSnapPoint::Mark, MarkedFrame.FrameNumber });
+	}
+
 	// Sort
 	Visitor.Snaps.Sort([](const FSequencerSnapPoint& A, const FSequencerSnapPoint& B){
 		return A.Time < B.Time;
diff --git a/Engine/Source/Editor/Sequencer/Private/Tools/SequencerSnapField.h b/Engine/Source/Editor/Sequencer/Private/Tools/SequencerSnapField.h
--- a/Engine/Source/Editor/Sequencer/Private/Tools/SequencerSnapField.h
+++ b/Engine/Source/Editor/Sequencer/Private/Tools/SequencerSnapField.h
@@ -51,7 +51,7 @@ public:
 	};
 
 	/** Construction from a sequencer and a snap canidate implementation. Optionally provide an entity mask to completely ignore some entity types */
-	FSequencerSnapField(const ISequencer& InSequencer, ISequencerSnapCandidate& Candidate, uint32 EntityMask = ESequencerEntity::Everything);
+	FSequencerSnapField(const FSequencer& InSequencer, ISequencerSnapCandidate& Candidate, uint32 EntityMask = ESequencerEntity::Everything);
 
 	/** Move construction / assignment */
 	FSequencerSnapField(FSequencerSnapField&& In) : SortedSnaps(MoveTemp(In.SortedSnaps)) {}
diff --git a/Engine/Source/Editor/SequencerWidgets/Public/ITimeSlider.h b/Engine/Source/Editor/SequencerWidgets/Public/ITimeSlider.h
--- a/Engine/Source/Editor/SequencerWidgets/Public/ITimeSlider.h
+++ b/Engine/Source/Editor/SequencerWidgets/Public/ITimeSlider.h
@@ -147,6 +147,9 @@ struct FTimeSliderArgs
 	/** Attribute for the current sequence's marked frames */
 	TAttribute<TArray<FMovieSceneMarkedFrame>> MarkedFrames;
 
+	/** Attribute for the marked frames that might need to be shown, but do not belong to the current sequence*/
+	TAttribute<TArray<FMovieSceneMarkedFrame>> GlobalMarkedFrames;
+
 	/** Called when the marked frame needs to be set */
 	FOnSetMarkedFrame OnSetMarkedFrame;
 
diff --git a/Engine/Source/Runtime/MovieScene/Public/MovieScene.h b/Engine/Source/Runtime/MovieScene/Public/MovieScene.h
--- a/Engine/Source/Runtime/MovieScene/Public/MovieScene.h
+++ b/Engine/Source/Runtime/MovieScene/Public/MovieScene.h
@@ -937,6 +937,23 @@ public:
 	 */
 	int32 FindNextMarkedFrame(FFrameNumber InFrameNumber, bool bForward);
 
+#if WITH_EDITORONLY_DATA
+	/*
+	 * Set whether this scene's marked frames should be shown globally
+	 */
+	void SetGloballyShowMarkedFrames(bool bShowMarkedFrames) { bGloballyShowMarkedFrames = bShowMarkedFrames; }
+
+	/*
+	 * Toggle whether this scene's marked frames should be shown globally
+	 */
+	void ToggleGloballyShowMarkedFrames() { bGloballyShowMarkedFrames = !bGloballyShowMarkedFrames; }
+
+	/*
+	 * Returns whether this scene's marked frames should be shown globally
+	 */
+	bool GetGloballyShowMarkedFrames() const { return bGloballyShowMarkedFrames; }
+#endif
+
 	/*
 	 * Retrieve all the tagged binding groups for this movie scene
 	 */
@@ -1083,6 +1100,9 @@ private:
 	UPROPERTY()
 	TArray<FMovieSceneSectionGroup> SectionGroups;
 
+	/** Whether this scene's marked frames should be shown globally */
+	bool bGloballyShowMarkedFrames;
+
 private:
 
 	UPROPERTY()
