Sequencer: Implemented NodeGroups

diff --git a/Engine/Source/Editor/EditorStyle/Private/SlateEditorStyle.cpp b/Engine/Source/Editor/EditorStyle/Private/SlateEditorStyle.cpp
--- a/Engine/Source/Editor/EditorStyle/Private/SlateEditorStyle.cpp
+++ b/Engine/Source/Editor/EditorStyle/Private/SlateEditorStyle.cpp
@@ -2624,6 +2624,8 @@ void FSlateEditorStyle::FStyle::SetupSequencerStyles()
 		Set("Sequencer.OpenDirectorBlueprint.Small", new IMAGE_BRUSH("Sequencer/Main_Icons/Icon_Sequencer_OpenDirectorBlueprint_24x", Icon24x24));
 		Set("Sequencer.OpenTaggedBindingManager", new IMAGE_BRUSH("Sequencer/Main_Icons/Icon_Sequencer_OpenTaggedBindingManager_16x", Icon48x48));
 		Set("Sequencer.OpenTaggedBindingManager.Small", new IMAGE_BRUSH("Sequencer/Main_Icons/Icon_Sequencer_OpenTaggedBindingManager_16x", Icon24x24));
+		Set("Sequencer.OpenNodeGroupsManager", new IMAGE_BRUSH("Sequencer/Main_Icons/Icon_Sequencer_OpenGroupManager_16x", Icon48x48));
+		Set("Sequencer.OpenNodeGroupsManager.Small", new IMAGE_BRUSH("Sequencer/Main_Icons/Icon_Sequencer_OpenGroupManager_16x", Icon24x24));
 		Set("Sequencer.BreadcrumbText", FTextBlockStyle(NormalText)
 			.SetFont(DEFAULT_FONT("Bold", 11))
 			.SetColorAndOpacity(FLinearColor(1.0f, 1.0f, 1.0f))
diff --git a/Engine/Source/Editor/Sequencer/Private/DisplayNodes/SequencerDisplayNode.cpp b/Engine/Source/Editor/Sequencer/Private/DisplayNodes/SequencerDisplayNode.cpp
--- a/Engine/Source/Editor/Sequencer/Private/DisplayNodes/SequencerDisplayNode.cpp
+++ b/Engine/Source/Editor/Sequencer/Private/DisplayNodes/SequencerDisplayNode.cpp
@@ -1032,9 +1032,9 @@ void FSequencerDisplayNode::BuildContextMenu(FMenuBuilder& MenuBuilder)
 
 	ESequencerNode::Type BaseNodeType = BaseNode->GetType();
 
-	bool bCanSolo = (BaseNodeType == ESequencerNode::Track || BaseNodeType == ESequencerNode::Object || BaseNodeType == ESequencerNode::Folder);
-	bool bIsReadOnly = !GetSequencer().IsReadOnly();
-	FCanExecuteAction CanExecute = FCanExecuteAction::CreateLambda([bIsReadOnly]{ return bIsReadOnly; });
+	bool bFilterableNode = (BaseNodeType == ESequencerNode::Track || BaseNodeType == ESequencerNode::Object || BaseNodeType == ESequencerNode::Folder);
+	bool bIsReadOnly = GetSequencer().IsReadOnly();
+	FCanExecuteAction CanExecute = FCanExecuteAction::CreateLambda([bIsReadOnly]{ return !bIsReadOnly; });
 
 	MenuBuilder.BeginSection("Edit", LOCTEXT("EditContextMenuSectionName", "Edit"));
 	{
@@ -1068,7 +1068,7 @@ void FSequencerDisplayNode::BuildContextMenu(FMenuBuilder& MenuBuilder)
 			);
 		}
 
-		if (bCanSolo)
+		if (bFilterableNode)
 		{
 			MenuBuilder.AddMenuEntry(
 				LOCTEXT("ToggleNodeSolo", "Solo"),
@@ -1164,6 +1164,15 @@ void FSequencerDisplayNode::BuildContextMenu(FMenuBuilder& MenuBuilder)
 				FSlateIcon(FEditorStyle::GetStyleSetName(), "ContentBrowser.AssetTreeFolderOpen"),
 				FUIAction(FExecuteAction::CreateSP(&GetSequencer(), &FSequencer::MoveSelectedNodesToNewFolder)));
 		}
+
+		if (bFilterableNode && !bIsReadOnly)
+		{
+			MenuBuilder.AddSubMenu(
+				LOCTEXT("AddNodesToNodeGroup", "Add to Group"),
+				LOCTEXT("AddNodesToNodeGroupTooltip", "Add selected nodes to a group"),
+				FNewMenuDelegate::CreateSP(&GetSequencer(), &FSequencer::BuildAddSelectedToNodeGroupMenu));
+
+		}
 	}
 	MenuBuilder.EndSection();
 
diff --git a/Engine/Source/Editor/Sequencer/Private/SSequencer.cpp b/Engine/Source/Editor/Sequencer/Private/SSequencer.cpp
--- a/Engine/Source/Editor/Sequencer/Private/SSequencer.cpp
+++ b/Engine/Source/Editor/Sequencer/Private/SSequencer.cpp
@@ -85,6 +85,7 @@
 #include "SequencerLog.h"
 #include "MovieSceneCopyableBinding.h"
 #include "SObjectBindingTagManager.h"
+#include "SSequencerGroupManager.h"
 #include "MovieSceneCopyableTrack.h"
 #include "IPropertyRowGenerator.h"
 #include "Fonts/FontMeasure.h"
@@ -1089,6 +1090,11 @@ void SSequencer::BindCommands(TSharedRef<FUICommandList> SequencerCommandBinding
 		FSequencerCommands::Get().OpenTaggedBindingManager,
 		FExecuteAction::CreateSP(this, &SSequencer::OpenTaggedBindingManager)
 	);
+
+	SequencerCommandBindings->MapAction(
+		FSequencerCommands::Get().OpenNodeGroupsManager,
+		FExecuteAction::CreateSP(this, &SSequencer::OpenNodeGroupsManager)
+	);
 }
 
 void SSequencer::ShowTickResolutionOverlay()
@@ -1215,6 +1221,11 @@ void SSequencer::HandleOutlinerNodeSelectionChanged()
 			CurveEditor->ResumeBroadcast();
 		}
 	}
+
+	if (NodeGroupManager.IsValid())
+	{
+		NodeGroupManager->SelectItemsSelectedInSequencer();
+	}
 }
 
 TSharedRef<SWidget> SSequencer::MakeAddButton()
@@ -1585,6 +1596,22 @@ TSharedRef<SWidget> SSequencer::MakeFilterMenu()
 	}
 	MenuBuilder.EndSection();
 
+	UMovieSceneSequence* FocusedMovieSequence = Sequencer->GetFocusedMovieSceneSequence();
+	UMovieScene* FocusedMovieScene = nullptr;
+	if (IsValid(FocusedMovieSequence))
+	{
+		FocusedMovieScene = FocusedMovieSequence->GetMovieScene();
+		if (IsValid(FocusedMovieScene))
+		{
+			if (FocusedMovieScene->GetNodeGroups().Num() > 0)
+			{
+				MenuBuilder.BeginSection("NodeGroupFilters");
+				MenuBuilder.AddSubMenu(LOCTEXT("NodeGroupFilters", "Group Filters"), LOCTEXT("NodeGroupFiltersToolTip", "Enables using this group as a filter for all items"), FNewMenuDelegate::CreateRaw(this, &SSequencer::FillNodeGroupsFilterMenu), false);
+				MenuBuilder.EndSection();
+			}
+		}
+	}
+
 	UObject* PlaybackContext = Sequencer->GetPlaybackContext();
 	UWorld* World = PlaybackContext ? PlaybackContext->GetWorld() : nullptr;
 
@@ -1620,6 +1647,76 @@ TSharedRef<SWidget> SSequencer::MakeFilterMenu()
 	return MenuBuilder.MakeWidget();
 }
 
+void SSequencer::FillNodeGroupsFilterMenu(FMenuBuilder& InMenuBarBuilder)
+{
+	TSharedPtr<FSequencer> Sequencer = SequencerPtr.Pin();
+	UMovieSceneSequence* FocusedMovieSequence = Sequencer->GetFocusedMovieSceneSequence();
+	UMovieScene* FocusedMovieScene = nullptr;
+
+	if (IsValid(FocusedMovieSequence))
+	{
+		FocusedMovieScene = FocusedMovieSequence->GetMovieScene();
+		if (IsValid(FocusedMovieScene))
+		{
+			InMenuBarBuilder.BeginSection("SequencerTracksResetFilters");
+			InMenuBarBuilder.AddMenuEntry(
+				LOCTEXT("FilterListResetNodeGroupFilters", "Reset Group Filters"),
+				LOCTEXT("FilterListResetNodeGroupFiltersToolTip", "Disables all group filters"),
+				FSlateIcon(),
+				FUIAction(FExecuteAction::CreateSP(this, &SSequencer::OnResetNodeGroupFilters))
+			);
+			InMenuBarBuilder.EndSection();
+
+			bool bIsReadOnly = GetIsSequenceReadOnly();
+
+			InMenuBarBuilder.BeginSection("NodeGroupFilters");
+			for (UMovieSceneNodeGroup* NodeGroup : FocusedMovieScene->GetNodeGroups())
+			{
+				InMenuBarBuilder.AddMenuEntry(
+					FText::FromName(NodeGroup->GetName()),
+					FText::FromName(NodeGroup->GetName()),
+					FSlateIcon(),
+					FUIAction(
+						FExecuteAction::CreateSP(this, &SSequencer::OnNodeGroupFilterClicked, NodeGroup),
+						FCanExecuteAction::CreateLambda([bIsReadOnly]() { return !bIsReadOnly; }),
+						FIsActionChecked::CreateLambda([NodeGroup]() { return NodeGroup->GetEnableFilter(); })),
+					NAME_None,
+					EUserInterfaceActionType::ToggleButton
+				);
+
+			}
+			InMenuBarBuilder.EndSection();
+		}
+	}
+}
+
+void SSequencer::OnResetNodeGroupFilters()
+{
+	TSharedPtr<FSequencer> Sequencer = SequencerPtr.Pin();
+	UMovieSceneSequence* FocusedMovieSequence = Sequencer->GetFocusedMovieSceneSequence();
+	UMovieScene* FocusedMovieScene = nullptr;
+	if (IsValid(FocusedMovieSequence))
+	{
+		FocusedMovieScene = FocusedMovieSequence->GetMovieScene();
+		if (IsValid(FocusedMovieScene))
+		{
+			for (UMovieSceneNodeGroup* NodeGroup : FocusedMovieScene->GetNodeGroups())
+			{
+				NodeGroup->SetEnableFilter(false);
+			}
+		}
+	}
+}
+
+void SSequencer::OnNodeGroupFilterClicked(UMovieSceneNodeGroup* NodeGroup)
+{
+	if (NodeGroup)
+	{
+		NodeGroup->SetEnableFilter(!NodeGroup->GetEnableFilter());
+	}
+	
+}
+
 void SSequencer::FillLevelFilterMenu(FMenuBuilder& InMenuBarBuilder)
 {
 	TSharedPtr<FSequencer> Sequencer = SequencerPtr.Pin();
@@ -1808,6 +1905,7 @@ TSharedRef<SWidget> SSequencer::MakeActionsMenu()
 		}
 
 		MenuBuilder.AddMenuEntry(FSequencerCommands::Get().OpenTaggedBindingManager);
+		MenuBuilder.AddMenuEntry(FSequencerCommands::Get().OpenNodeGroupsManager);
 
 		MenuBuilder.AddMenuEntry(FSequencerCommands::Get().RestoreAnimatedState);
 
@@ -1979,6 +2077,36 @@ void SSequencer::OpenTaggedBindingManager()
 	WeakExposedBindingsWindow = ExposedBindingsWindow;
 }
 
+void SSequencer::OpenNodeGroupsManager()
+{
+	if (TSharedPtr<SWindow> Window = WeakNodeGroupWindow.Pin())
+	{
+		Window->DrawAttention(FWindowDrawAttentionParameters());
+		return;
+	}
+
+	TSharedRef<SWindow> NodeGroupManagerWindow = SNew(SWindow)
+		.Title(FText::Format(LOCTEXT("NodeGroup_Title", "Groups in {0}"), FText::FromName(SequencerPtr.Pin()->GetRootMovieSceneSequence()->GetFName())))
+		.SupportsMaximize(false)
+		.ClientSize(FVector2D(600.f, 500.f))
+		.Content()
+		[
+			SAssignNew(NodeGroupManager, SSequencerGroupManager, SequencerPtr)
+		];
+
+	TSharedPtr<SWindow> ParentWindow = FSlateApplication::Get().FindWidgetWindow(AsShared());
+	if (ParentWindow)
+	{
+		FSlateApplication::Get().AddWindowAsNativeChild(NodeGroupManagerWindow, ParentWindow.ToSharedRef());
+	}
+	else
+	{
+		FSlateApplication::Get().AddWindow(NodeGroupManagerWindow);
+	}
+
+	WeakNodeGroupWindow = NodeGroupManagerWindow;
+}
+
 void SSequencer::FillPlaybackSpeedMenu(FMenuBuilder& InMenuBarBuilder)
 {
 	const int32 NumPlaybackSpeeds = 7;
@@ -2447,6 +2575,12 @@ SSequencer::~SSequencer()
 	{
 		Window->DestroyWindowImmediately();
 	}
+
+	if (TSharedPtr<SWindow> Window = WeakNodeGroupWindow.Pin())
+	{
+		Window->DestroyWindowImmediately();
+		NodeGroupManager.Reset();
+	}
 }
 
 
@@ -2584,6 +2718,11 @@ void SSequencer::UpdateLayoutTree()
 
 		// Continue broadcasting selection changes
 		Sequencer->GetSelection().ResumeBroadcast();
+		
+		if (NodeGroupManager.IsValid())
+		{
+			NodeGroupManager->RefreshNodeGroups();
+		}
 	}
 }
 
diff --git a/Engine/Source/Editor/Sequencer/Private/SSequencer.h b/Engine/Source/Editor/Sequencer/Private/SSequencer.h
--- a/Engine/Source/Editor/Sequencer/Private/SSequencer.h
+++ b/Engine/Source/Editor/Sequencer/Private/SSequencer.h
@@ -36,6 +36,7 @@ class SDockTab;
 class SWindow;
 class USequencerSettings;
 class FSequencerTrackFilter;
+class SSequencerGroupManager;
 struct FPaintPlaybackRangeArgs;
 struct FSequencerCustomizationInfo;
 struct FSequencerSelectionCurveFilter;
@@ -412,6 +413,10 @@ public:
 	/** Makes the time display format menu for the toolbar and the play rate menu. */
 	void FillTimeDisplayFormatMenu(FMenuBuilder& MenuBuilder);
 
+	void OpenNodeGroupsManager();
+
+	TSharedPtr<SSequencerGroupManager> GetNodeGroupsManager() const { return NodeGroupManager; }
+
 public:	
 	/** Makes a time range widget with the specified inner content */
 	TSharedRef<SWidget> MakeTimeRange(const TSharedRef<SWidget>& InnerContent, bool bShowWorkingRange, bool bShowViewRange, bool bShowPlaybackRange);
@@ -431,8 +436,10 @@ private:
 	bool IsTrackLevelFilterActive(const FString LevelName) const;
 
 	void FillLevelFilterMenu(FMenuBuilder& InMenuBarBuilder);
+	void FillNodeGroupsFilterMenu(FMenuBuilder& InMenuBarBuilder);
 
-	void ClearGlobalMarkedFrames();
+	void OnResetNodeGroupFilters();
+	void OnNodeGroupFilterClicked(UMovieSceneNodeGroup* NodeGroup);
 
 	/**
 	* Called when the time snap interval changes.
@@ -722,6 +729,10 @@ private:
 
 	TWeakPtr<SWindow> WeakExposedBindingsWindow;
 
+	TWeakPtr<SWindow> WeakNodeGroupWindow;
+
+	TSharedPtr<SSequencerGroupManager> NodeGroupManager;
+
 public:
 	static const FName CurveEditorTabName;
 
diff --git a/Engine/Source/Editor/Sequencer/Private/SSequencerGroupManager.cpp b/Engine/Source/Editor/Sequencer/Private/SSequencerGroupManager.cpp
new file mode 100644
--- /dev/null
+++ b/Engine/Source/Editor/Sequencer/Private/SSequencerGroupManager.cpp
@@ -0,0 +1,715 @@
+// Copyright Epic Games, Inc. All Rights Reserved.
+
+#include "SSequencerGroupManager.h"
+
+#include "Sequencer.h"
+#include "MovieSceneSequence.h"
+#include "MovieScene.h"
+#include "DisplayNodes/SequencerDisplayNode.h"
+#include "SequencerDisplayNodeDragDropOp.h"
+#include "SlateOptMacros.h"
+#include "Widgets/SNullWidget.h"
+#include "Widgets/SBoxPanel.h"
+#include "Widgets/Layout/SBox.h"
+#include "Widgets/Layout/SScrollBorder.h"
+#include "Widgets/Images/SImage.h"
+#include "Widgets/SOverlay.h"
+#include "Widgets/Text/STextBlock.h"
+#include "Widgets/Input/SEditableTextBox.h"
+#include "Widgets/Input/SButton.h"
+#include "Widgets/Input/SComboButton.h"
+#include "Widgets/Text/SInlineEditableTextBlock.h"
+#include "EditorFontGlyphs.h"
+#include "ScopedTransaction.h"
+#include "Framework/MultiBox/MultiBoxBuilder.h"
+
+
+#define LOCTEXT_NAMESPACE "SSequencerGroupManager"
+
+class SSequencerGroupNodeRow;
+
+struct FSequencerNodeGroupTreeNode
+{
+	enum Type
+	{
+		/* Base Node Type */
+		BaseNode,
+		/* Group */
+		GroupNode,
+		/* Item Node*/
+		ItemNode
+	};
+
+	FSequencerNodeGroupTreeNode(const FText& InDisplayText)
+		: DisplayText(InDisplayText)
+	{
+	}
+	virtual ~FSequencerNodeGroupTreeNode() {}
+
+	virtual Type GetType() const { return BaseNode; }
+
+	FText DisplayText;
+
+	TArray<TSharedPtr<FSequencerNodeGroupTreeNode>> Children;
+};
+
+struct FSequencerGroupItemNode : public FSequencerNodeGroupTreeNode
+{
+	FSequencerGroupItemNode(const FText& InDisplayText, const FString& InPath, TSharedPtr<FSequencerNodeGroupNode> InGroup)
+		: FSequencerNodeGroupTreeNode(InDisplayText), Path(InPath), Group(InGroup)
+	{
+	}
+	virtual ~FSequencerGroupItemNode() override {}
+
+	virtual Type GetType() const override { return ItemNode; }
+
+	FString Path;
+	TSharedPtr<FSequencerNodeGroupNode> Group;
+};
+
+struct FSequencerNodeGroupNode : public FSequencerNodeGroupTreeNode
+{
+	FSequencerNodeGroupNode(const FText& InDisplayText, UMovieSceneNodeGroup* InGroup, TWeakPtr<SSequencerGroupManager> InGroupManager)
+		: FSequencerNodeGroupTreeNode(InDisplayText), WeakGroupManager(InGroupManager), Group(InGroup)
+	{
+	}
+
+	virtual ~FSequencerNodeGroupNode() override {}
+
+	virtual Type GetType() const override { return GroupNode; }
+
+	FReply OnEnableFilterClicked()
+	{
+		if (IsValid(Group))
+		{
+			Group->SetEnableFilter(!Group->GetEnableFilter());
+		}
+		return FReply::Handled();
+	}
+
+	bool IsFilterEnabled() const
+	{
+		return Group->GetEnableFilter();
+	}
+
+	bool VerifyNodeTextChanged(const FText& NewLabel, FText& OutErrorMessage)
+	{
+		return !NewLabel.IsEmptyOrWhitespace();
+	}
+
+	void HandleNodeLabelTextCommitted(const FText& NewLabel, ETextCommit::Type CommitType)
+	{
+		TSharedPtr<SSequencerGroupManager> GroupManager = WeakGroupManager.Pin();
+		UMovieScene* MovieScene = GroupManager ? GroupManager->GetMovieScene() : nullptr;
+		if (MovieScene)
+		{
+			MovieScene->Modify();
+			Group->SetName(FName(*FText::TrimPrecedingAndTrailing(NewLabel).ToString()));
+		}
+	}
+
+	void OnRenameRequested()
+	{
+		if (InlineEditableTextBlock)
+		{
+			InlineEditableTextBlock->EnterEditingMode();
+		}
+	}
+
+	TWeakPtr<SSequencerGroupManager> WeakGroupManager;
+	UMovieSceneNodeGroup* Group;
+	TSharedPtr<SInlineEditableTextBlock> InlineEditableTextBlock;
+};
+
+class SSequencerGroupNodeRow : public STableRow<TSharedPtr<FSequencerNodeGroupTreeNode>>
+{
+	SLATE_BEGIN_ARGS(SSequencerGroupNodeRow) {}
+	SLATE_END_ARGS()
+
+public:
+
+	void Construct(const FArguments& InArgs, const TSharedRef<STableViewBase>& InOwnerTableView, TWeakPtr<FSequencerNodeGroupTreeNode> InWeakSequencerGroupTreeNode, TWeakPtr<SSequencerGroupManager> InWeakSequencerGroupManager)
+	{
+		WeakSequencerGroupManager = InWeakSequencerGroupManager;
+		WeakSequencerGroupTreeNode = InWeakSequencerGroupTreeNode;
+
+		STableRow<TSharedPtr<FSequencerNodeGroupTreeNode>>::ConstructInternal(STableRow::FArguments()
+			.Padding(5.f)
+			.OnCanAcceptDrop(this, &SSequencerGroupNodeRow::OnCanAcceptDrop)
+			.OnAcceptDrop(this, &SSequencerGroupNodeRow::OnAcceptDrop)
+			, InOwnerTableView);
+
+		TSharedPtr<FSequencerNodeGroupTreeNode> SequencerGroupTreeNode = WeakSequencerGroupTreeNode.Pin();
+
+		if (!SequencerGroupTreeNode)
+		{
+			return;
+		}
+
+		TSharedPtr<SSequencerGroupManager> SequencerGroupManager = WeakSequencerGroupManager.Pin();
+
+		const FSlateBrush* IconBrush = SequencerGroupManager ? SequencerGroupManager->GetIconBrush(SequencerGroupTreeNode) : nullptr;
+
+		if(SequencerGroupTreeNode)
+		{
+			if(SequencerGroupTreeNode->GetType() == FSequencerNodeGroupTreeNode::Type::ItemNode)
+			{
+				TSharedPtr<FSequencerGroupItemNode> SequencerGroupItemNode = StaticCastSharedPtr<FSequencerGroupItemNode>(SequencerGroupTreeNode);
+				this->ChildSlot
+				[
+					SNew(SHorizontalBox)
+					+ SHorizontalBox::Slot()
+					.VAlign(VAlign_Center)
+					.AutoWidth()
+					[
+						SNew(SExpanderArrow, SharedThis(this))
+					]
+
+					+ SHorizontalBox::Slot()
+					.VAlign(VAlign_Center)
+					.Padding(FMargin(5.f, 0.f, 5.f, 0.f))
+					.AutoWidth()
+					[
+						SNew(SOverlay)
+						+ SOverlay::Slot()
+						[
+							SNew(SImage)
+								.Image(IconBrush ? IconBrush : FCoreStyle::Get().GetDefaultBrush())
+								.ColorAndOpacity(IconBrush ? FLinearColor::White : FLinearColor::Transparent)
+						]
+					]
+
+					+ SHorizontalBox::Slot()
+					.VAlign(VAlign_Center)
+					[
+						SNew(STextBlock)
+						.Text(SequencerGroupTreeNode->DisplayText)
+						.ToolTipText(FText::FromString(SequencerGroupItemNode->Path))
+					]
+				];
+			}
+			else if (SequencerGroupTreeNode->GetType() == FSequencerNodeGroupTreeNode::Type::GroupNode)
+			{
+				TSharedPtr<FSequencerNodeGroupNode> NodeGroupNode = StaticCastSharedPtr<FSequencerNodeGroupNode>(SequencerGroupTreeNode);
+
+				this->ChildSlot
+				[
+					SNew(SHorizontalBox)
+					+ SHorizontalBox::Slot()
+					.VAlign(VAlign_Center)
+					.AutoWidth()
+					[
+						SNew(SExpanderArrow, SharedThis(this))
+					]
+
+					+ SHorizontalBox::Slot()
+					.VAlign(VAlign_Center)
+					.Padding(FMargin(5.f, 0.f, 5.f, 0.f))
+					.AutoWidth()
+					[
+						SNew(SOverlay)
+						+ SOverlay::Slot()
+						[
+							SNew(SButton)
+							.OnClicked(FOnClicked::CreateSP(NodeGroupNode.ToSharedRef(), &FSequencerNodeGroupNode::OnEnableFilterClicked))
+							.ButtonStyle( FEditorStyle::Get(), "NoBorder" )
+							.Content()
+							[
+								SNew(STextBlock)
+								.TextStyle(FEditorStyle::Get(), "GenericFilters.TextStyle")
+								.Font(FEditorStyle::Get().GetFontStyle("FontAwesome.11"))
+								.Text(FEditorFontGlyphs::Filter)
+								.ColorAndOpacity(NodeGroupNode->Group->GetEnableFilter() ? FLinearColor::White : FLinearColor(0.66f, 0.66f, 0.66f, 0.66f))
+							]
+						]
+					]
+
+					+ SHorizontalBox::Slot()
+					.VAlign(VAlign_Center)
+					[
+						SAssignNew(NodeGroupNode->InlineEditableTextBlock, SInlineEditableTextBlock)
+						.OnVerifyTextChanged(NodeGroupNode.ToSharedRef(), &FSequencerNodeGroupNode::VerifyNodeTextChanged)
+						.OnTextCommitted(NodeGroupNode.ToSharedRef(), &FSequencerNodeGroupNode::HandleNodeLabelTextCommitted)
+						.Text(NodeGroupNode->DisplayText)
+						.Clipping(EWidgetClipping::ClipToBounds)
+					]
+				];
+
+			}
+		}
+	}
+
+private:
+
+	TOptional<EItemDropZone> OnCanAcceptDrop(const FDragDropEvent& DragDropEvent, EItemDropZone InItemDropZone, TSharedPtr<FSequencerNodeGroupTreeNode> SequencerGroupTreeNode)
+	{
+		TSharedPtr<FSequencerDisplayNodeDragDropOp> DragDropOp = DragDropEvent.GetOperationAs<FSequencerDisplayNodeDragDropOp>();
+		if (DragDropOp.IsValid())
+		{
+			DragDropOp->ResetToDefaultToolTip();
+			
+			TOptional<EItemDropZone> AllowedDropZone;
+			if (SequencerGroupTreeNode->GetType() == FSequencerNodeGroupTreeNode::Type::GroupNode)
+			{
+				AllowedDropZone = EItemDropZone::OntoItem;
+				DragDropOp->CurrentHoverText = FText::Format(LOCTEXT("DragDropAddItemsFormat", "Add {0} item(s)"), FText::AsNumber(DragDropOp->GetDraggedNodes().Num()));
+			}
+
+			if (AllowedDropZone.IsSet() == false)
+			{
+				DragDropOp->CurrentIconBrush = FEditorStyle::GetBrush(TEXT("Graph.ConnectorFeedback.Error"));
+			}
+			return AllowedDropZone;
+		}
+		return TOptional<EItemDropZone>();
+	}
+
+	FReply OnAcceptDrop(const FDragDropEvent& DragDropEvent, EItemDropZone InItemDropZone, TSharedPtr<FSequencerNodeGroupTreeNode> SequencerGroupTreeNode)
+	{
+		TSharedPtr<FSequencerDisplayNodeDragDropOp> DragDropOp = DragDropEvent.GetOperationAs<FSequencerDisplayNodeDragDropOp>();
+		if (DragDropOp.IsValid())
+		{
+			TSharedPtr<SSequencerGroupManager> SequencerGroupManager = WeakSequencerGroupManager.Pin();
+			TSharedPtr<FSequencer> Sequencer = SequencerGroupManager ? SequencerGroupManager->GetSequencer() : nullptr;
+			if (Sequencer.IsValid())
+			{
+				if (SequencerGroupTreeNode->GetType() == FSequencerNodeGroupTreeNode::Type::GroupNode)
+				{
+					TSharedPtr<FSequencerNodeGroupNode> NodeGroupNode = StaticCastSharedPtr<FSequencerNodeGroupNode>(SequencerGroupTreeNode);
+					Sequencer->AddNodesToExistingNodeGroup(DragDropOp->GetDraggedNodes(), NodeGroupNode->Group);
+
+					return FReply::Handled();
+				}
+			}
+		}
+		return FReply::Unhandled();
+	}
+
+	TWeakPtr<SSequencerGroupManager> WeakSequencerGroupManager;
+	TWeakPtr<FSequencerNodeGroupTreeNode> WeakSequencerGroupTreeNode;
+};
+
+UMovieScene* SSequencerGroupManager::GetMovieScene() const
+{
+	TSharedPtr<FSequencer> Sequencer = WeakSequencer.Pin();
+	UMovieSceneSequence* Sequence = Sequencer ? Sequencer->GetFocusedMovieSceneSequence() : nullptr;
+	UMovieScene* MovieScene = Sequence ? Sequence->GetMovieScene() : nullptr;
+
+	return MovieScene;
+}
+
+void SSequencerGroupManager::Construct(const FArguments& InArgs, TWeakPtr<FSequencer> InWeakSequencer)
+{
+	WeakSequencer = InWeakSequencer;
+
+	UMovieScene* MovieScene = GetMovieScene();
+
+	if (!ensure(MovieScene))
+	{
+		return;
+	}
+
+	TWeakPtr<SSequencerGroupManager> WeakTabManager = SharedThis(this);
+	auto HandleGenerateRow = [WeakTabManager](TSharedPtr<FSequencerNodeGroupTreeNode> InNode, const TSharedRef<STableViewBase>& InOwnerTableView) -> TSharedRef<ITableRow>
+	{
+		return SNew(SSequencerGroupNodeRow, InOwnerTableView, InNode, WeakTabManager);
+	};
+
+	auto HandleGetChildren = [](TSharedPtr<FSequencerNodeGroupTreeNode> InParent, TArray<TSharedPtr<FSequencerNodeGroupTreeNode>>& OutChildren)
+	{
+		OutChildren.Append(InParent->Children);
+	};
+
+	TreeView = SNew(STreeView<TSharedPtr<FSequencerNodeGroupTreeNode>>)
+		.OnGenerateRow_Lambda(HandleGenerateRow)
+		.OnGetChildren_Lambda(HandleGetChildren)
+		.TreeItemsSource(&NodeGroupsTree)
+		.OnSelectionChanged(this, &SSequencerGroupManager::HandleTreeSelectionChanged)
+		.OnContextMenuOpening(this, &SSequencerGroupManager::OnContextMenuOpening);
+
+	ChildSlot
+	[
+		SNew(SBorder)
+		.BorderImage(FEditorStyle::GetBrush("ToolPanel.GroupBorder"))
+		[
+			SNew(SVerticalBox)
+
+			+ SVerticalBox::Slot()
+			[
+				SNew(SScrollBorder, TreeView.ToSharedRef())
+				[
+					TreeView.ToSharedRef()
+				]
+			]
+		]
+	];
+
+	UpdateTree();
+}
+
+SSequencerGroupManager::~SSequencerGroupManager()
+{
+
+}
+
+void SSequencerGroupManager::UpdateTree()
+{
+	TSharedPtr<FSequencer> Sequencer = WeakSequencer.Pin(); 
+	UMovieScene* MovieScene = GetMovieScene();
+
+	if (!ensure(Sequencer) || !ensure(MovieScene))
+	{
+		return;
+	}
+
+	TSharedRef<FSequencerNodeTree> NodeTree = Sequencer->GetNodeTree();
+
+	TArray<UMovieSceneNodeGroup*> ExpandedSets;
+	TSet<TSharedPtr<FSequencerNodeGroupTreeNode>> ExpandedTreeItems;
+	TreeView->GetExpandedItems(ExpandedTreeItems);
+	for(TSharedPtr<FSequencerNodeGroupTreeNode> ExpandedTreeItem : ExpandedTreeItems)
+	{
+		if (ExpandedTreeItem->GetType() == FSequencerNodeGroupTreeNode::Type::GroupNode)
+		{
+			TSharedPtr<FSequencerNodeGroupNode> NodeGroupNode = StaticCastSharedPtr<FSequencerNodeGroupNode>(ExpandedTreeItem);
+			ExpandedSets.Add(NodeGroupNode->Group);
+		}
+	}
+
+	NodeGroupsTree.Empty();
+	AllNodeGroupItems.Empty();
+
+	for (UMovieSceneNodeGroup* NodeGroup : MovieScene->GetNodeGroups())
+	{
+		TSharedPtr<FSequencerNodeGroupNode> SequencerGroupNode = MakeShared<FSequencerNodeGroupNode>(FText::FromName(NodeGroup->GetName()), NodeGroup, SharedThis(this));
+		NodeGroupsTree.Add(SequencerGroupNode);
+		for (const FString& NodePath : NodeGroup->GetNodes())
+		{
+			const FSequencerDisplayNode* Node = NodeTree->GetNodeAtPath(NodePath);
+			if (Node)
+			{
+				TSharedPtr<FSequencerGroupItemNode> SequencerGroupItemNode = MakeShared<FSequencerGroupItemNode>(Node->GetDisplayName(), NodePath, SequencerGroupNode);
+				SequencerGroupNode->Children.Add(SequencerGroupItemNode);
+				AllNodeGroupItems.Add(NodePath);
+			}
+		}
+
+		SequencerGroupNode->Children.Sort([](const TSharedPtr<FSequencerNodeGroupTreeNode>& A, const TSharedPtr<FSequencerNodeGroupTreeNode>& B) {
+			return A->DisplayText.CompareTo(B->DisplayText) < 0;
+		});
+	}
+
+	NodeGroupsTree.Sort([](const TSharedPtr<FSequencerNodeGroupTreeNode>& A, const TSharedPtr<FSequencerNodeGroupTreeNode>& B) {
+		return A->DisplayText.CompareTo(B->DisplayText) < 0;
+	});
+
+	TreeView->SetTreeItemsSource(&NodeGroupsTree);
+
+	for (TSharedPtr<FSequencerNodeGroupTreeNode> NodeGroupTreeNode : NodeGroupsTree)
+	{
+		if (NodeGroupTreeNode->GetType() == FSequencerNodeGroupTreeNode::Type::GroupNode)
+		{
+			TSharedPtr<FSequencerNodeGroupNode> NodeGroupNode = StaticCastSharedPtr<FSequencerNodeGroupNode>(NodeGroupTreeNode);
+
+			if (ExpandedSets.Contains(NodeGroupNode->Group))
+			{
+				TreeView->SetItemExpansion(NodeGroupTreeNode,true);
+			}
+		}
+	}
+
+	bNodeGroupsDirty = false;
+}
+
+void SSequencerGroupManager::HandleTreeSelectionChanged(TSharedPtr<FSequencerNodeGroupTreeNode> InSelectedNode, ESelectInfo::Type SelectionType)
+{
+	SelectSelectedItemsInSequencer();
+}
+
+void SSequencerGroupManager::Tick(const FGeometry& AllottedGeometry, const double InCurrentTime, const float InDeltaTime)
+{
+	if (bNodeGroupsDirty)
+	{
+		UpdateTree();
+	}
+
+	if (RequestedRenameNodeGroup && !TreeView->IsPendingRefresh())
+	{
+		for (const TSharedPtr<FSequencerNodeGroupTreeNode> Node : NodeGroupsTree)
+		{
+			if (Node->GetType() == FSequencerNodeGroupTreeNode::Type::GroupNode)
+			{
+				TSharedPtr<FSequencerNodeGroupNode> SequencerGroupItemNode = StaticCastSharedPtr<FSequencerNodeGroupNode>(Node);
+				if (SequencerGroupItemNode->Group == RequestedRenameNodeGroup)
+				{
+					SequencerGroupItemNode->OnRenameRequested();
+					break;
+				}
+			}
+		}
+
+		RequestedRenameNodeGroup = nullptr;
+	}
+}
+
+const FSlateBrush* SSequencerGroupManager::GetIconBrush(TSharedPtr<FSequencerNodeGroupTreeNode> NodeGroupTreeNode) const
+{
+	if (NodeGroupTreeNode->GetType() == FSequencerNodeGroupTreeNode::Type::ItemNode)
+	{
+		TSharedPtr<FSequencerGroupItemNode> SequencerGroupItemNode = StaticCastSharedPtr<FSequencerGroupItemNode>(NodeGroupTreeNode);
+		TSharedPtr<FSequencer> Sequencer = WeakSequencer.Pin();
+		if (!Sequencer)
+		{
+			return nullptr;
+		}
+		TSharedRef<FSequencerNodeTree> NodeTree = Sequencer->GetNodeTree();
+		const FSequencerDisplayNode* Node = NodeTree->GetNodeAtPath(SequencerGroupItemNode->Path);
+		
+		return Node ? Node->GetIconBrush() : nullptr;
+	}
+
+	return nullptr;
+}
+
+void SSequencerGroupManager::RequestDeleteNodeGroup(FSequencerNodeGroupNode* NodeGroupNode)
+{
+	UMovieScene* MovieScene = GetMovieScene();
+	if (!ensure(MovieScene) || !ensure(NodeGroupNode))
+	{
+		return;
+	}
+
+	if (MovieScene->IsReadOnly())
+	{
+		return;
+	}
+
+	const FScopedTransaction Transaction(LOCTEXT("DeleteGroupTransaction", "Delete Group"));
+
+	MovieScene->Modify();
+	MovieScene->GetNodeGroups().RemoveNodeGroup(NodeGroupNode->Group);
+}
+
+void SSequencerGroupManager::RemoveSelectedItemsFromNodeGroup()
+{
+	UMovieScene* MovieScene = GetMovieScene();
+
+	if (!ensure(MovieScene))
+	{
+		return;
+	}
+
+	if (MovieScene->IsReadOnly())
+	{
+		return;
+	}
+
+	TArray<TPair<UMovieSceneNodeGroup*,FString>> ItemsToRemove;
+	TArray<TSharedPtr<FSequencerNodeGroupTreeNode>> SelectedNodes = TreeView->GetSelectedItems();
+	for (const TSharedPtr<FSequencerNodeGroupTreeNode> Node : SelectedNodes)
+	{
+		if (Node->GetType() == FSequencerNodeGroupTreeNode::Type::ItemNode)
+		{
+			TSharedPtr<FSequencerGroupItemNode> ItemNode = StaticCastSharedPtr<FSequencerGroupItemNode>(Node);
+			ItemsToRemove.Add(TPair<UMovieSceneNodeGroup*, FString>(ItemNode->Group->Group,ItemNode->Path));
+		}
+	}
+
+	if (ItemsToRemove.Num() < 1)
+	{
+		return;
+	}
+
+	const FScopedTransaction Transaction(LOCTEXT("RemoveItemFromGroupTransaction", "Remove Items From Group"));
+
+	MovieScene->Modify();
+	for (const TPair<UMovieSceneNodeGroup*, FString> Item : ItemsToRemove)
+	{
+		Item.Key->RemoveNode(Item.Value);
+	}
+	
+}
+
+void SSequencerGroupManager::GetSelectedItemsNodePaths(TSet<FString>& OutSelectedNodePaths) const
+{
+	TArray<TSharedPtr<FSequencerNodeGroupTreeNode>> SelectedNodes = TreeView->GetSelectedItems();
+	for (const TSharedPtr<FSequencerNodeGroupTreeNode> Node : SelectedNodes)
+	{
+		if (Node->GetType() == FSequencerNodeGroupTreeNode::Type::ItemNode)
+		{
+			TSharedPtr<FSequencerGroupItemNode> ItemNode = StaticCastSharedPtr<FSequencerGroupItemNode>(Node);
+			OutSelectedNodePaths.Add(ItemNode->Path);
+		}
+		else if (Node->GetType() == FSequencerNodeGroupTreeNode::Type::GroupNode)
+		{
+			for (TSharedPtr<FSequencerNodeGroupTreeNode> ChildNode : Node->Children)
+			{
+				// Note: Currently, children of a set can only be item nodes, but that may change in the future.
+				if (ChildNode->GetType() == FSequencerNodeGroupTreeNode::Type::ItemNode)
+				{
+					TSharedPtr<FSequencerGroupItemNode> ItemNode = StaticCastSharedPtr<FSequencerGroupItemNode>(ChildNode);
+					OutSelectedNodePaths.Add(ItemNode->Path);
+				}
+			}
+		}
+	}
+}
+
+void SSequencerGroupManager::SelectSelectedItemsInSequencer()
+{
+	if (bSynchronizingSelection)
+	{
+		return;
+	}
+
+	TGuardValue<bool> Guard(bSynchronizingSelection, true);
+
+	TSharedPtr<FSequencer> Sequencer = WeakSequencer.Pin();
+	if (!ensure(Sequencer))
+	{
+		return;
+	}
+	
+	TSet<FString> SelectedNodePaths;
+	GetSelectedItemsNodePaths(SelectedNodePaths);
+	
+	if (SelectedNodePaths.Num() < 1)
+	{
+		return;
+	}
+
+	Sequencer->SelectNodesByPath(SelectedNodePaths);
+}
+
+void SSequencerGroupManager::SelectItemsSelectedInSequencer()
+{
+	if (bSynchronizingSelection)
+	{
+		return;
+	}
+
+	TGuardValue<bool> Guard(bSynchronizingSelection, true);
+
+	TSharedPtr<FSequencer> Sequencer = WeakSequencer.Pin();
+	if (!ensure(Sequencer))
+	{
+		return;
+	}
+
+	const TSet<TSharedRef<FSequencerDisplayNode>>& SequencerSelectedNodes = Sequencer->GetSelection().GetSelectedOutlinerNodes();
+
+	// Build a list of the nodepaths that we want to consider for selection
+	TSet<FString> NodesPathsToSelect;
+	for (const TSharedRef<const FSequencerDisplayNode> Node : SequencerSelectedNodes)
+	{
+		const FSequencerDisplayNode* BaseNode = Node->GetBaseNode();
+		ESequencerNode::Type NodeType = BaseNode->GetType();
+
+		if (NodeType == ESequencerNode::Track || NodeType == ESequencerNode::Object || NodeType == ESequencerNode::Folder)
+		{
+			FString NodePath = BaseNode->GetPathName();
+			if (AllNodeGroupItems.Contains(NodePath))
+			{
+				NodesPathsToSelect.Add(NodePath);
+			}
+		}
+	}
+
+	TreeView->ClearSelection();
+
+	// Build a list of the treenodes which match a nodepath we want to select
+	TArray<TSharedPtr<FSequencerGroupItemNode>> TreeNodesToSelect;
+	for (const TSharedPtr<FSequencerNodeGroupTreeNode> Node : NodeGroupsTree)
+	{
+		if (Node->GetType() == FSequencerNodeGroupTreeNode::Type::ItemNode)
+		{
+			TSharedPtr<FSequencerGroupItemNode> ItemNode = StaticCastSharedPtr<FSequencerGroupItemNode>(Node);
+			if (NodesPathsToSelect.Contains(ItemNode->Path))
+			{
+				TreeView->SetItemSelection(Node, true);
+			}
+		}
+		else if (Node->GetType() == FSequencerNodeGroupTreeNode::Type::GroupNode)
+		{
+			for (TSharedPtr<FSequencerNodeGroupTreeNode> ChildNode : Node->Children)
+			{
+				// Note: Currently, children of a set can only be item nodes, but that may change in the future.
+				if (ChildNode->GetType() == FSequencerNodeGroupTreeNode::Type::ItemNode)
+				{
+					TSharedPtr<FSequencerGroupItemNode> ItemNode = StaticCastSharedPtr<FSequencerGroupItemNode>(ChildNode);
+					if (NodesPathsToSelect.Contains(ItemNode->Path))
+					{
+						TreeView->SetItemSelection(ChildNode, true);
+					}
+				}
+			}
+		}
+	}	
+}
+
+TSharedPtr<SWidget> SSequencerGroupManager::OnContextMenuOpening()
+{
+	TArray<TSharedPtr<FSequencerNodeGroupTreeNode>> SelectedNodes = TreeView->GetSelectedItems();
+
+	FMenuBuilder MenuBuilder(true, nullptr);
+
+	UMovieScene* MovieScene = GetMovieScene();
+	bool bIsReadOnly = MovieScene? MovieScene->IsReadOnly() : true;
+
+	for (const TSharedPtr<FSequencerNodeGroupTreeNode> Node : SelectedNodes)
+	{
+		if (Node->GetType() == FSequencerNodeGroupTreeNode::Type::GroupNode)
+		{
+			TSharedPtr<FSequencerNodeGroupNode> NodeGroupNode = StaticCastSharedPtr<FSequencerNodeGroupNode>(Node);
+
+			MenuBuilder.AddMenuEntry(
+				LOCTEXT("NodeGroupToggleFilter", "Toggle Filter"),
+				LOCTEXT("NodeGroupToggleFilterTooltip", "Toggle whether this group should be used to filter items"),
+				FSlateIcon(),
+				FUIAction(
+					FExecuteAction::CreateLambda([NodeGroupNode]() { NodeGroupNode->OnEnableFilterClicked(); }),
+					FCanExecuteAction::CreateLambda([bIsReadOnly]() { return !bIsReadOnly; }),
+					FIsActionChecked::CreateLambda([NodeGroupNode]() { return NodeGroupNode->IsFilterEnabled(); })),
+				NAME_None,
+				EUserInterfaceActionType::ToggleButton
+			);
+
+			MenuBuilder.AddMenuEntry(
+				FText::Format(LOCTEXT("RenameNodeGroupFormat", "Rename {0}"), NodeGroupNode->DisplayText),
+				FText(),
+				FSlateIcon(),
+				FUIAction(FExecuteAction::CreateLambda([NodeGroupNode]() { NodeGroupNode->OnRenameRequested(); })));
+
+			MenuBuilder.AddMenuEntry(
+				LOCTEXT("DeleteNodeGroup", "Delete Group"),
+				FText(),
+				FSlateIcon(),
+				FUIAction(FExecuteAction::CreateSP(this, &SSequencerGroupManager::RequestDeleteNodeGroup, NodeGroupNode.Get())));
+
+			break;
+		}
+	}
+
+	bool bAnyItemSelected = false;
+	for (const TSharedPtr<FSequencerNodeGroupTreeNode> Node : SelectedNodes)
+	{
+		if (Node->GetType() == FSequencerNodeGroupTreeNode::Type::ItemNode)
+		{
+			bAnyItemSelected = true;
+			break;
+		}
+	}
+
+	if (bAnyItemSelected)
+	{
+		MenuBuilder.AddMenuEntry(
+			LOCTEXT("RemoveItemsFromNodeGropu", "Remove Items From Group"),
+			FText(),
+			FSlateIcon(),
+			FUIAction(FExecuteAction::CreateSP(this, &SSequencerGroupManager::RemoveSelectedItemsFromNodeGroup)));
+	}
+
+	return MenuBuilder.MakeWidget();
+}
+
+#undef LOCTEXT_NAMESPACE
diff --git a/Engine/Source/Editor/Sequencer/Private/SSequencerGroupManager.h b/Engine/Source/Editor/Sequencer/Private/SSequencerGroupManager.h
new file mode 100644
--- /dev/null
+++ b/Engine/Source/Editor/Sequencer/Private/SSequencerGroupManager.h
@@ -0,0 +1,77 @@
+// Copyright Epic Games, Inc. All Rights Reserved.
+
+#pragma once
+
+#include "Widgets/DeclarativeSyntaxSupport.h"
+#include "Widgets/SCompoundWidget.h"
+#include "Widgets/Views/STreeView.h"
+
+class FSequencer;
+class FMenuBuilder;
+class UMovieScene;
+class UMovieSceneNodeGroup;
+
+struct FSequencerNodeGroupTreeNode;
+struct FSequencerNodeGroupNode;
+
+class SSequencerGroupManager : public SCompoundWidget
+{
+public:
+
+	SLATE_BEGIN_ARGS(SSequencerGroupManager) {}
+	SLATE_END_ARGS()
+
+	void Construct(const FArguments& InArgs, TWeakPtr<FSequencer> InWeakSequencer);
+	
+	~SSequencerGroupManager();
+
+	TSharedPtr<FSequencer> GetSequencer() const
+	{
+		return WeakSequencer.Pin();
+	}
+
+	const FSlateBrush* GetIconBrush(TSharedPtr<FSequencerNodeGroupTreeNode> NodeGroupTreeNode) const;
+
+	void RefreshNodeGroups() { bNodeGroupsDirty = true; }
+
+	UMovieScene* GetMovieScene() const;
+
+private:
+
+	virtual void Tick(const FGeometry& AllottedGeometry, const double InCurrentTime, const float InDeltaTime) override;
+	void UpdateTree();
+	void HandleTreeSelectionChanged(TSharedPtr<FSequencerNodeGroupTreeNode> InSelectedNode, ESelectInfo::Type SelectionType);
+	void RequestDeleteNodeGroup(FSequencerNodeGroupNode* NodeGroupNode);
+	void RemoveSelectedItemsFromNodeGroup();
+	
+	void GetSelectedItemsNodePaths(TSet<FString>& OutSelectedNodePaths) const;
+
+public:
+	void SelectSelectedItemsInSequencer();
+	void SelectItemsSelectedInSequencer();
+
+	/** Request entering rename editing for a selection set, after setup and any refresh have completed */
+	void RequestRenameNodeGroup(UMovieSceneNodeGroup* NodeGroup) { RequestedRenameNodeGroup = NodeGroup; }
+
+private:
+	TSharedPtr<SWidget> OnContextMenuOpening();
+	
+private:
+
+	/** The sequencer UI instance that is currently open */
+	TWeakPtr<FSequencer> WeakSequencer;
+	
+	TSharedPtr<STreeView<TSharedPtr<FSequencerNodeGroupTreeNode>>> TreeView;
+
+	TArray<TSharedPtr<FSequencerNodeGroupTreeNode>> NodeGroupsTree;
+
+	/** A collection of NodePaths for all items which are in all node groups, for quick processing */
+	TSet<FString> AllNodeGroupItems;
+	
+	/** If not null, enter rename editing for this set after any setup and refresh have completed */
+	UMovieSceneNodeGroup* RequestedRenameNodeGroup;
+
+	bool bNodeGroupsDirty;
+	bool bSynchronizingSelection;
+
+};
\ No newline at end of file
diff --git a/Engine/Source/Editor/Sequencer/Private/Sequencer.cpp b/Engine/Source/Editor/Sequencer/Private/Sequencer.cpp
--- a/Engine/Source/Editor/Sequencer/Private/Sequencer.cpp
+++ b/Engine/Source/Editor/Sequencer/Private/Sequencer.cpp
@@ -138,6 +138,7 @@
 #include "Interfaces/IAnalyticsProvider.h"
 #include "Components/SkeletalMeshComponent.h"
 #include "SequencerCustomizationManager.h"
+#include "SSequencerGroupManager.h"
 
 #define LOCTEXT_NAMESPACE "Sequencer"
 
@@ -532,6 +533,8 @@ void FSequencer::InitSequencer(const FSequencerInitParams& InitParams, const TSh
 
 	UpdateSequencerCustomizations();
 
+	AddNodeGroupsCollectionChangedDelegate();
+
 	OnActivateSequenceEvent.Broadcast(ActiveTemplateIDs[0]);
 }
 
@@ -858,6 +861,8 @@ UMovieSceneSubSection* FSequencer::FindSubSection(FMovieSceneSequenceID Sequence
 
 void FSequencer::ResetToNewRootSequence(UMovieSceneSequence& NewSequence)
 {
+	RemoveNodeGroupsCollectionChangedDelegate();
+
 	RootSequence = &NewSequence;
 	RestorePreAnimatedState();
 
@@ -893,12 +898,16 @@ void FSequencer::ResetToNewRootSequence(UMovieSceneSequence& NewSequence)
 
 	UpdateSequencerCustomizations();
 
+	AddNodeGroupsCollectionChangedDelegate();
+
 	OnActivateSequenceEvent.Broadcast(ActiveTemplateIDs.Top());
 }
 
 
 void FSequencer::FocusSequenceInstance(UMovieSceneSubSection& InSubSection)
 {
+	RemoveNodeGroupsCollectionChangedDelegate();
+
 	FMovieSceneRootOverridePath Path;
 	Path.Set(ActiveTemplateIDs.Last(), RootTemplateInstance.GetHierarchy());
 
@@ -949,6 +958,8 @@ void FSequencer::FocusSequenceInstance(UMovieSceneSubSection& InSubSection)
 
 	OnActivateSequenceEvent.Broadcast(ActiveTemplateIDs.Top());
 
+	AddNodeGroupsCollectionChangedDelegate();
+
 	bNeedsEvaluate = true;
 	bGlobalMarkedFramesCached = false;
 }
@@ -1074,6 +1085,8 @@ void FSequencer::PopToSequenceInstance(FMovieSceneSequenceIDRef SequenceID)
 {
 	if( ActiveTemplateIDs.Num() > 1 )
 	{
+		RemoveNodeGroupsCollectionChangedDelegate();
+
 		// Pop until we find the movie scene to focus
 		while( SequenceID != ActiveTemplateIDs.Last() )
 		{
@@ -1109,6 +1122,8 @@ void FSequencer::PopToSequenceInstance(FMovieSceneSequenceIDRef SequenceID)
 
 		UpdateSequencerCustomizations();
 
+		AddNodeGroupsCollectionChangedDelegate();
+
 		OnActivateSequenceEvent.Broadcast(ActiveTemplateIDs.Top());
 
 		bNeedsEvaluate = true;
@@ -5807,6 +5822,147 @@ void FSequencer::OnSelectedOutlinerNodesChanged()
 	OnSelectionChangedSectionsDelegate.Broadcast(SelectedSections);
 }
 
+void FSequencer::AddNodeGroupsCollectionChangedDelegate()
+{
+	UMovieSceneSequence* MovieSceneSequence = GetFocusedMovieSceneSequence();
+	UMovieScene* MovieScene = MovieSceneSequence ? MovieSceneSequence->GetMovieScene() : nullptr;
+	if (ensure(MovieScene))
+	{
+		if (!MovieScene->GetNodeGroups().OnNodeGroupCollectionChanged().IsBoundToObject(this))
+		{
+			MovieScene->GetNodeGroups().OnNodeGroupCollectionChanged().AddSP(this, &FSequencer::OnNodeGroupsCollectionChanged);
+		}
+	}
+}
+
+void FSequencer::RemoveNodeGroupsCollectionChangedDelegate()
+{
+	UMovieSceneSequence* MovieSceneSequence = GetFocusedMovieSceneSequence();
+	UMovieScene* MovieScene = MovieSceneSequence ? MovieSceneSequence->GetMovieScene() : nullptr;
+	if (ensure(MovieScene))
+	{
+		MovieScene->GetNodeGroups().OnNodeGroupCollectionChanged().RemoveAll(this);
+	}
+}
+
+void FSequencer::OnNodeGroupsCollectionChanged()
+{
+	TSharedPtr<SSequencerGroupManager> NodeGroupManager = SequencerWidget->GetNodeGroupsManager();
+	if (NodeGroupManager)
+	{
+		NodeGroupManager->RefreshNodeGroups();
+	}
+
+	NodeTree->NodeGroupsCollectionChanged();
+}
+
+void FSequencer::AddSelectedNodesToNewNodeGroup()
+{
+	UMovieScene* MovieScene = GetFocusedMovieSceneSequence()->GetMovieScene();
+
+	if (MovieScene->IsReadOnly())
+	{
+		return;
+	}
+
+	const TSet<TSharedRef<FSequencerDisplayNode> >& SelectedNodes = GetSelection().GetSelectedOutlinerNodes();
+	if (SelectedNodes.Num() == 0)
+	{
+		return;
+	}
+
+	TSet<FString> NodesToAdd;
+	for (const TSharedRef<const FSequencerDisplayNode> Node : SelectedNodes)
+	{
+		const FSequencerDisplayNode* BaseNode = Node->GetBaseNode();
+		ESequencerNode::Type NodeType = BaseNode->GetType();
+
+		if (NodeType == ESequencerNode::Track || NodeType == ESequencerNode::Object || NodeType == ESequencerNode::Folder)
+		{
+			FString NodePath = BaseNode->GetPathName();
+			NodesToAdd.Add(NodePath);
+		}
+	}
+
+	if (NodesToAdd.Num() == 0)
+	{
+		return;
+	}
+
+	TArray<FName> ExistingGroupNames;
+	for (const UMovieSceneNodeGroup* NodeGroup : MovieScene->GetNodeGroups())
+	{
+		ExistingGroupNames.Add(NodeGroup->GetName());
+	}
+
+	UMovieSceneNodeGroup* NewNodeGroup = NewObject<UMovieSceneNodeGroup>(&MovieScene->GetNodeGroups());
+	NewNodeGroup->SetName(FSequencerUtilities::GetUniqueName(FName("Group"), ExistingGroupNames));
+
+	for (const FString& NodeToAdd : NodesToAdd)
+	{
+		NewNodeGroup->AddNode(NodeToAdd);
+	}
+
+	const FScopedTransaction Transaction(LOCTEXT("CreateNewGroupTransaction", "Create New Group"));
+
+	MovieScene->Modify();
+
+	MovieScene->GetNodeGroups().AddNodeGroup(NewNodeGroup);
+
+	SequencerWidget->OpenNodeGroupsManager();
+	SequencerWidget->GetNodeGroupsManager()->RequestRenameNodeGroup(NewNodeGroup);
+}
+
+void FSequencer::AddSelectedNodesToExistingNodeGroup(UMovieSceneNodeGroup* NodeGroup)
+{
+	AddNodesToExistingNodeGroup(GetSelection().GetSelectedOutlinerNodes().Array(), NodeGroup);
+}
+
+void FSequencer::AddNodesToExistingNodeGroup(const TArray<TSharedRef<FSequencerDisplayNode>>& InNodes, UMovieSceneNodeGroup* InNodeGroup)
+{
+	UMovieScene* MovieScene = GetFocusedMovieSceneSequence()->GetMovieScene();
+
+	if (MovieScene->IsReadOnly())
+	{
+		return;
+	}
+
+	if (!MovieScene->GetNodeGroups().Contains(InNodeGroup))
+	{
+		return;
+	}
+
+	TSet<FString> NodesToAdd;
+	for (const TSharedRef<const FSequencerDisplayNode> Node : InNodes)
+	{
+		const FSequencerDisplayNode* BaseNode = Node->GetBaseNode();
+		ESequencerNode::Type NodeType = BaseNode->GetType();
+
+		if (NodeType == ESequencerNode::Track || NodeType == ESequencerNode::Object || NodeType == ESequencerNode::Folder)
+		{
+			FString NodePath = BaseNode->GetPathName();
+			NodesToAdd.Add(NodePath);
+		}
+	}
+
+	if (NodesToAdd.Num() == 0)
+	{
+		return;
+	}
+
+	const FScopedTransaction Transaction(LOCTEXT("AddNodesToGroupTransaction", "Add Nodes to Group"));
+
+	MovieScene->Modify();
+
+	for (const FString& NodeToAdd : NodesToAdd)
+	{
+		if (!InNodeGroup->ContainsNode(NodeToAdd))
+		{
+			InNodeGroup->AddNode(NodeToAdd);
+		}
+	}
+
+}
 
 void FSequencer::SynchronizeExternalSelectionWithSequencerSelection()
 {
@@ -6128,6 +6284,67 @@ void FSequencer::SynchronizeSequencerSelectionWithExternalSelection()
 	}
 }
 
+void FSequencer::SelectNodesByPath(const TSet<FString>& NodePaths)
+{
+	if (bUpdatingExternalSelection)
+	{
+		return;
+	}
+
+	UMovieSceneSequence* Sequence = GetFocusedMovieSceneSequence();
+	if (!Sequence->GetMovieScene())
+	{
+		return;
+	}
+
+	// If all nodes are already selected, do nothing. This ensures that when an undo event happens, 
+	// nodes are not cleared and reselected, which can cause issues with the curve editor auto-fitting 
+	// based on selection.
+	bool bAllAlreadySelected = true;
+	const TSet<TSharedRef<FSequencerDisplayNode>>& CurrentSelection = GetSelection().GetSelectedOutlinerNodes();
+
+	TSet<TSharedRef<FSequencerDisplayNode>> NodesToSelect;
+	for (TSharedRef<FSequencerDisplayNode> DisplayNode : NodeTree->GetAllNodes())
+	{
+		if (NodePaths.Contains(DisplayNode->GetPathName()))
+		{
+			NodesToSelect.Add(DisplayNode);
+			if (bAllAlreadySelected && !CurrentSelection.Contains(DisplayNode))
+			{
+				bAllAlreadySelected = false;
+			}
+		}
+	}
+
+	if (!bAllAlreadySelected || (NodesToSelect.Num() != CurrentSelection.Num()))
+	{
+		Selection.SuspendBroadcast();
+		Selection.EmptySelectedOutlinerNodes();
+		for (TSharedRef<FSequencerDisplayNode> NodeToSelect : NodesToSelect)
+		{
+			Selection.AddToSelection(NodeToSelect);
+		}
+
+		TSharedPtr<SSequencerTreeView> TreeView = SequencerWidget->GetTreeView();
+		const TSet<TSharedRef<FSequencerDisplayNode>>& OutlinerSelection = GetSelection().GetSelectedOutlinerNodes();
+		for (auto& Node : OutlinerSelection)
+		{
+			auto Parent = Node->GetParent();
+			while (Parent.IsValid())
+			{
+				TreeView->SetItemExpansion(Parent->AsShared(), true);
+				Parent = Parent->GetParent();
+			}
+
+			TreeView->RequestScrollIntoView(Node);
+			break;
+		}
+
+		Selection.ResumeBroadcast();
+		Selection.GetOnOutlinerNodeSelectionChanged().Broadcast();
+	}
+}
+
 bool FSequencer::IsBindingVisible(const FMovieSceneBinding& InBinding)
 {
 	if (Settings->GetShowSelectedNodesOnly() && OnGetIsBindingVisible().IsBound())
@@ -11636,6 +11853,34 @@ void FSequencer::BuildObjectBindingEditButtons(TSharedPtr<SHorizontalBox> EditBo
 	}
 }
 
+void FSequencer::BuildAddSelectedToNodeGroupMenu(FMenuBuilder& MenuBuilder)
+{
+	UMovieSceneSequence* FocusedMovieSceneSequence = GetFocusedMovieSceneSequence();
+	UMovieScene* MovieScene = FocusedMovieSceneSequence ? FocusedMovieSceneSequence->GetMovieScene() : nullptr;
+	if (MovieScene)
+	{
+		MenuBuilder.AddMenuEntry(
+			LOCTEXT("NewNodeGroup", "New Group"),
+			LOCTEXT("AddNodesToNewNodeGroupTooltip", "Creates a new group and adds the selected nodes"),
+			FSlateIcon(),
+			FUIAction(FExecuteAction::CreateSP(this, &FSequencer::AddSelectedNodesToNewNodeGroup)));
+
+		if (MovieScene->GetNodeGroups().Num() > 0)
+		{
+			MenuBuilder.AddMenuSeparator();
+
+			for (UMovieSceneNodeGroup* NodeGroup : MovieScene->GetNodeGroups())
+			{
+				MenuBuilder.AddMenuEntry(
+					FText::FromName(NodeGroup->GetName()),
+					LOCTEXT("AddNodesToNodeGroupFormatTooltip", "Adds the selected nodes to this existing group"),
+					FSlateIcon(),
+					FUIAction(FExecuteAction::CreateSP(this, &FSequencer::AddSelectedNodesToExistingNodeGroup, NodeGroup)));
+			}
+		}
+	}
+}
+
 void FSequencer::UpdateTimeBases()
 {
 	UMovieSceneSequence* RootSequencePtr = GetRootMovieSceneSequence();
diff --git a/Engine/Source/Editor/Sequencer/Private/Sequencer.h b/Engine/Source/Editor/Sequencer/Private/Sequencer.h
--- a/Engine/Source/Editor/Sequencer/Private/Sequencer.h
+++ b/Engine/Source/Editor/Sequencer/Private/Sequencer.h
@@ -422,6 +422,13 @@ public:
 	 */
 	void BuildObjectBindingEditButtons(TSharedPtr<SHorizontalBox> EditBox, const FGuid& ObjectBinding, const UClass* ObjectClass);
 
+	/**
+	 * Builds up the menu of node groups to add selected nodes to
+	 *
+	 * @param MenuBuilder The menu builder to add things to.
+	 */
+	void BuildAddSelectedToNodeGroupMenu(FMenuBuilder& MenuBuilder);
+
 	/** Called when an actor is dropped into Sequencer */
 	void OnActorsDropped( const TArray<TWeakObjectPtr<AActor> >& Actors );
 
@@ -554,6 +561,9 @@ public:
 	/** Updates the sequencer selection to match the current external selection. */
 	void SynchronizeSequencerSelectionWithExternalSelection();
 		
+	/** Updates the sequencer selection to match the list of node paths. */
+	void SelectNodesByPath(const TSet<FString>& NodePaths);
+
 	/** Whether the binding is visible in the tree view */
 	bool IsBindingVisible(const FMovieSceneBinding& InBinding);
 
@@ -975,6 +985,18 @@ protected:
 
 	void OnSelectedOutlinerNodesChanged();
 
+	void AddNodeGroupsCollectionChangedDelegate();
+	void RemoveNodeGroupsCollectionChangedDelegate();
+
+	void OnNodeGroupsCollectionChanged();
+
+public:
+	void AddSelectedNodesToNewNodeGroup();
+	void AddSelectedNodesToExistingNodeGroup(UMovieSceneNodeGroup* NodeGroup);
+	void AddNodesToExistingNodeGroup(const TArray<TSharedRef<FSequencerDisplayNode>>& Nodes, UMovieSceneNodeGroup* NodeGroup);
+
+private:
+
 	/** Updates a viewport client from camera cut data */
 	void UpdatePreviewLevelViewportClientFromCameraCut(FLevelEditorViewportClient& InViewportClient, UObject* InCameraObject, const EMovieSceneCameraCutParams& CameraCutParams);
 
diff --git a/Engine/Source/Editor/Sequencer/Private/SequencerCommands.cpp b/Engine/Source/Editor/Sequencer/Private/SequencerCommands.cpp
--- a/Engine/Source/Editor/Sequencer/Private/SequencerCommands.cpp
+++ b/Engine/Source/Editor/Sequencer/Private/SequencerCommands.cpp
@@ -98,6 +98,7 @@ void FSequencerCommands::RegisterCommands()
 	UI_COMMAND( ToggleShowStretchBox, "Stretch/Shrink...", "Stretch or shrink a given amount, moving keys forwards/backwards as necessary", EUserInterfaceActionType::Button, FInputChord() );
 	UI_COMMAND( OpenDirectorBlueprint, "Open Director Blueprint", "Opens the director blueprint for this sequence.", EUserInterfaceActionType::Button, FInputChord() );
 	UI_COMMAND( OpenTaggedBindingManager, "Open Binding Tag Manager", "Specifies options for tagging bindings within this sequence for external systems to reference as a persistent name.", EUserInterfaceActionType::Button, FInputChord() );
+	UI_COMMAND( OpenNodeGroupsManager, "Open Sequencer Group Manager", "Manage groups within this sequence.", EUserInterfaceActionType::Button, FInputChord());
 	UI_COMMAND( BakeTransform, "Bake Transform", "Bake transform in world space, removing any existing transform and attach tracks", EUserInterfaceActionType::Button, FInputChord() );
 	UI_COMMAND( SyncSectionsUsingSourceTimecode, "Sync Sections using Source Timecode", "Synchronize sections to the first selected section using source timecode", EUserInterfaceActionType::Button, FInputChord() );
 	UI_COMMAND( ToggleShowRangeSlider, "Range Slider", "Enables and disables showing the time range slider", EUserInterfaceActionType::ToggleButton, FInputChord() );
diff --git a/Engine/Source/Editor/Sequencer/Private/SequencerNodeTree.cpp b/Engine/Source/Editor/Sequencer/Private/SequencerNodeTree.cpp
--- a/Engine/Source/Editor/Sequencer/Private/SequencerNodeTree.cpp
+++ b/Engine/Source/Editor/Sequencer/Private/SequencerNodeTree.cpp
@@ -20,6 +20,8 @@
 #include "SequencerNodeSortingMethods.h"
 #include "SequencerTrackFilters.h"
 #include "Channels/MovieSceneChannel.h"
+#include "SequencerUtilities.h"
+#include "SequencerLog.h"
 #include "ScopedTransaction.h"
 
 FSequencerNodeTree::~FSequencerNodeTree()
@@ -275,7 +277,11 @@ TSharedRef<FSequencerFolderNode> FSequencerNodeTree::CreateOrUpdateFolder(UMovie
 
 bool FSequencerNodeTree::HasActiveFilter() const
 {
-	return (!FilterString.IsEmpty() || TrackFilters->Num() > 0 || TrackFilterLevelFilter->IsActive() || Sequencer.GetSequencerSettings()->GetShowSelectedNodesOnly());
+	return (!FilterString.IsEmpty()
+		|| TrackFilters->Num() > 0
+		|| TrackFilterLevelFilter->IsActive()
+		|| Sequencer.GetSequencerSettings()->GetShowSelectedNodesOnly()
+		|| Sequencer.GetFocusedMovieSceneSequence()->GetMovieScene()->GetNodeGroups().HasAnyActiveFilter());
 }
 
 TSharedPtr<FSequencerObjectBindingNode> FSequencerNodeTree::FindOrCreateObjectBinding(const FGuid& BindingID, const TSortedMap<FGuid, const FMovieSceneBinding*>& AllBindings, const TSortedMap<FGuid, FGuid>& ChildToParentBinding, TMap<FSequencerDisplayNode*, TSharedPtr<FSequencerDisplayNode>>* OutChildToParentMap)
@@ -970,7 +976,7 @@ static bool FilterNodesRecursive(FSequencer& Sequencer, const TSharedRef<FSequen
 		return bAnyChildPassed;
 	}
 
-	bool bPasssedAnyFilters = false;
+	bool bPassedAnyFilters = false;
 	bool bIsTrackOrObjectBinding = false;
 
 	switch(StartNode->GetType())
@@ -990,16 +996,16 @@ static bool FilterNodesRecursive(FSequencer& Sequencer, const TSharedRef<FSequen
 					{
 						if (Filters->Num() == 0 || Filters->PassesAnyFilters(Channel))
 						{
-							bPasssedAnyFilters = true;
+							bPassedAnyFilters = true;
 							break;
 						}
 					}
 				}
 			}
 		
-			if (bPasssedAnyFilters || Filters->Num() == 0 || Filters->PassesAnyFilters(Track, StartNode->GetDisplayName()))
+			if (bPassedAnyFilters || Filters->Num() == 0 || Filters->PassesAnyFilters(Track, StartNode->GetDisplayName()))
 			{
-				bPasssedAnyFilters = true;
+				bPassedAnyFilters = true;
 
 				// Track nodes do not belong to a level, but might be a child of an objectbinding node that does
 				if (LevelTrackFilter->IsActive())
@@ -1010,7 +1016,7 @@ static bool FilterNodesRecursive(FSequencer& Sequencer, const TSharedRef<FSequen
 						if (ParentNode->GetType() == ESequencerNode::Object)
 						{
 							// The track belongs to an objectbinding node, start by assuming it doesn't match the level filter
-							bPasssedAnyFilters = false;
+							bPassedAnyFilters = false;
 
 							const FSequencerObjectBindingNode* ObjectNode = static_cast<const FSequencerObjectBindingNode*>(ParentNode.Get());
 							for (TWeakObjectPtr<>& Object : Sequencer.FindObjectsInCurrentSequence(ObjectNode->GetObjectBinding()))
@@ -1018,7 +1024,7 @@ static bool FilterNodesRecursive(FSequencer& Sequencer, const TSharedRef<FSequen
 								if (Object.IsValid() && LevelTrackFilter->PassesFilter(Object.Get()))
 								{
 									// If at least one of the objects on the objectbinding node pass the level filter, show the track
-									bPasssedAnyFilters = true;
+									bPassedAnyFilters = true;
 									break;
 								}
 							}
@@ -1029,7 +1035,7 @@ static bool FilterNodesRecursive(FSequencer& Sequencer, const TSharedRef<FSequen
 					}
 				}
 
-				if (bPasssedAnyFilters && Sequencer.GetSequencerSettings()->GetShowSelectedNodesOnly())
+				if (bPassedAnyFilters && Sequencer.GetSequencerSettings()->GetShowSelectedNodesOnly())
 				{
 					TSharedPtr<const FSequencerDisplayNode> ParentNode = StartNode->GetParent();
 					while (ParentNode.IsValid())
@@ -1067,12 +1073,12 @@ static bool FilterNodesRecursive(FSequencer& Sequencer, const TSharedRef<FSequen
 				if (Object.IsValid() && (Filters->Num() == 0 || Filters->PassesAnyFilters(Object.Get(), StartNode->GetDisplayName()))
 					&& LevelTrackFilter->PassesFilter(Object.Get()))
 				{
-					bPasssedAnyFilters = true;
+					bPassedAnyFilters = true;
 					break;
 				}
 			}
 
-			if (bPasssedAnyFilters && Sequencer.GetSequencerSettings()->GetShowSelectedNodesOnly() && !StartNode->IsPinned())
+			if (bPassedAnyFilters && Sequencer.GetSequencerSettings()->GetShowSelectedNodesOnly() && !StartNode->IsPinned())
 			{
 				UMovieScene* MovieScene = Sequencer.GetFocusedMovieSceneSequence()->GetMovieScene();
 				const FMovieSceneBinding* Binding = Sequencer.GetFocusedMovieSceneSequence()->GetMovieScene()->FindBinding(ObjectNode.GetObjectBinding());
@@ -1090,10 +1096,10 @@ static bool FilterNodesRecursive(FSequencer& Sequencer, const TSharedRef<FSequen
 				UMovieSceneTrack* Track = TrackNode->GetTrack();
 				if (Filters->Num() == 0 || Filters->PassesAnyFilters(Track, StartNode->GetDisplayName()))
 				{
-					bPasssedAnyFilters = true;
+					bPassedAnyFilters = true;
 				}
 			}
-			if (bPasssedAnyFilters && Sequencer.GetSequencerSettings()->GetShowSelectedNodesOnly())
+			if (bPassedAnyFilters && Sequencer.GetSequencerSettings()->GetShowSelectedNodesOnly())
 			{
 				TSharedPtr<const FSequencerDisplayNode> ParentNode = StartNode->GetParent();
 				while (ParentNode.IsValid())
@@ -1110,7 +1116,7 @@ static bool FilterNodesRecursive(FSequencer& Sequencer, const TSharedRef<FSequen
 						const FMovieSceneBinding* Binding = Sequencer.GetFocusedMovieSceneSequence()->GetMovieScene()->FindBinding(ObjectNode->GetObjectBinding());
 						if (!(Binding && Sequencer.IsBindingVisible(*Binding)))
 						{
-							bPasssedAnyFilters = false;
+							bPassedAnyFilters = false;
 						}
 
 						break;
@@ -1130,12 +1136,12 @@ static bool FilterNodesRecursive(FSequencer& Sequencer, const TSharedRef<FSequen
 				{
 					if (Filters->Num() == 0 || Filters->PassesAnyFilters(Channel))
 					{
-						bPasssedAnyFilters = true;
+						bPassedAnyFilters = true;
 						break;
 					}
 				}
 			}
-			if (bPasssedAnyFilters && Sequencer.GetSequencerSettings()->GetShowSelectedNodesOnly())
+			if (bPassedAnyFilters && Sequencer.GetSequencerSettings()->GetShowSelectedNodesOnly())
 			{
 				TSharedPtr<const FSequencerDisplayNode> ParentNode = StartNode->GetParent();
 				while (ParentNode.IsValid())
@@ -1152,7 +1158,7 @@ static bool FilterNodesRecursive(FSequencer& Sequencer, const TSharedRef<FSequen
 						const FMovieSceneBinding* Binding = Sequencer.GetFocusedMovieSceneSequence()->GetMovieScene()->FindBinding(ObjectNode->GetObjectBinding());
 						if (!(Binding && Sequencer.IsBindingVisible(*Binding)))
 						{
-							bPasssedAnyFilters = false;
+							bPassedAnyFilters = false;
 						}
 
 						break;
@@ -1167,19 +1173,19 @@ static bool FilterNodesRecursive(FSequencer& Sequencer, const TSharedRef<FSequen
 			// Special case: If we're pinned, then we should pass regardless
 			if (StartNode->IsPinned())
 			{
-				bPasssedAnyFilters = true;
+				bPassedAnyFilters = true;
 				break;
 			}
 
 			// Special case: If we're only filtering on text search, include folders and key areas in the search
 			if (Filters->Num() == 0 && FilterStrings.Num() > 0)
 			{
-				bPasssedAnyFilters = true;
+				bPassedAnyFilters = true;
 
 				// Special case: but don't include if only showing selected bindings and we don't have child that passed
 				if (Sequencer.GetSequencerSettings()->GetShowSelectedNodesOnly() && !bAnyChildPassed)
 				{
-					bPasssedAnyFilters = false;
+					bPassedAnyFilters = false;
 
 					// Special case: unless we're the child of a node that is a selected binding
 					TSharedPtr<const FSequencerDisplayNode> ParentNode = StartNode->GetParent();
@@ -1188,6 +1194,7 @@ static bool FilterNodesRecursive(FSequencer& Sequencer, const TSharedRef<FSequen
 						// Pinned tracks should be visible whether selected or not
 						if (ParentNode->IsPinned())
 						{
+							bPassedAnyFilters = true;
 							break;
 						}
 
@@ -1197,7 +1204,7 @@ static bool FilterNodesRecursive(FSequencer& Sequencer, const TSharedRef<FSequen
 							const FMovieSceneBinding* Binding = Sequencer.GetFocusedMovieSceneSequence()->GetMovieScene()->FindBinding(ObjectNode->GetObjectBinding());
 							if (Binding && Sequencer.IsBindingVisible(*Binding))
 							{
-								bPasssedAnyFilters = true;
+								bPassedAnyFilters = true;
 							}
 
 							break;
@@ -1210,8 +1217,39 @@ static bool FilterNodesRecursive(FSequencer& Sequencer, const TSharedRef<FSequen
 		}
 	}
 
+	if (bPassedAnyFilters)
+	{
+		// If filtering on selection set is enabled, we need to run another pass to verify we're in an enabled node group
+		UMovieSceneNodeGroupCollection& NodeGroups = Sequencer.GetFocusedMovieSceneSequence()->GetMovieScene()->GetNodeGroups();
+		if (NodeGroups.HasAnyActiveFilter())
+		{
+			bPassedAnyFilters = false;
+
+			TSharedPtr<const FSequencerDisplayNode> CurrentNode = StartNode;
+			while (CurrentNode.IsValid() && !bPassedAnyFilters)
+			{
+				// Special case: Pinned tracks should be visible whether in the node group or not
+				if (CurrentNode->IsPinned())
+				{
+					bPassedAnyFilters = true;
+					break;
+				}
+
+				for (const UMovieSceneNodeGroup* NodeGroup : NodeGroups)
+				{
+					if (NodeGroup->GetEnableFilter() && NodeGroup->ContainsNode(CurrentNode->GetPathName()))
+					{
+						bPassedAnyFilters = true;
+						break;
+					}
+				}
+
+				CurrentNode = CurrentNode->GetParent();
+			}
+		}
+	}
 
-	if (bPasssedAnyFilters)
+	if (bPassedAnyFilters)
 	{
 		OutFilteredNodes.Add(StartNode);
 		AddParentNodes(StartNode, OutFilteredNodes);
@@ -1243,7 +1281,7 @@ void FSequencerNodeTree::UpdateFilters()
 	UWorld* World = PlaybackContext ? PlaybackContext->GetWorld() : nullptr;
 	TrackFilterLevelFilter->UpdateWorld(World);
 
-	if (TrackFilters->Num() > 0 || !FilterString.IsEmpty() || TrackFilterLevelFilter->IsActive() || Sequencer.GetSequencerSettings()->GetShowSelectedNodesOnly())
+	if (HasActiveFilter())
 	{
 		// Build a list of strings that must be matched
 		TArray<FString> FilterStrings;
@@ -1259,6 +1297,7 @@ void FSequencerNodeTree::UpdateFilters()
 		}
 	}
 
+	bFilteringOnNodeGroups = Sequencer.GetFocusedMovieSceneSequence()->GetMovieScene()->GetNodeGroups().HasAnyActiveFilter();
 	bFilterUpdateRequested = false;
 }
 
@@ -1271,6 +1310,14 @@ void FSequencerNodeTree::FilterNodes(const FString& InFilter)
 	}
 }
 
+void FSequencerNodeTree::NodeGroupsCollectionChanged()
+{
+	if (Sequencer.GetFocusedMovieSceneSequence()->GetMovieScene()->GetNodeGroups().HasAnyActiveFilter() || bFilteringOnNodeGroups)
+	{
+		RequestFilterUpdate();
+	}
+}
+
 TArray< TSharedRef<FSequencerDisplayNode> > FSequencerNodeTree::GetAllNodes() const
 {
 	TArray< TSharedRef<FSequencerDisplayNode> > AllNodes;
diff --git a/Engine/Source/Editor/Sequencer/Private/SequencerNodeTree.h b/Engine/Source/Editor/Sequencer/Private/SequencerNodeTree.h
--- a/Engine/Source/Editor/Sequencer/Private/SequencerNodeTree.h
+++ b/Engine/Source/Editor/Sequencer/Private/SequencerNodeTree.h
@@ -82,6 +82,8 @@ public:
 	 */
 	void FilterNodes( const FString& InFilter );
 
+	/** Called when the active MovieScene's node group colletion has been modifed */
+	void NodeGroupsCollectionChanged();
 
 	/**
 	 * Unpins any pinned nodes in this tree
@@ -210,7 +212,7 @@ public:
 
 private:
 
-	/** Returns whether this NodeTree should only display seleected nodes */
+	/** Returns whether this NodeTree should only display selected nodes */
 	bool ShowSelectedNodesOnly() const;
 
 	/** Population algorithm utilities */
@@ -325,7 +327,8 @@ private:
 	TSharedPtr<FSequencerTrackFilter_LevelFilter> TrackFilterLevelFilter;
 
 	bool bFilterUpdateRequested;
-	
+	bool bFilteringOnNodeGroups;
+
 	/** Cached value of whether we have any nodes that should be treated as soloing */
 	bool bHasSoloNodes;
 };
diff --git a/Engine/Source/Editor/Sequencer/Public/SequencerCommands.h b/Engine/Source/Editor/Sequencer/Public/SequencerCommands.h
--- a/Engine/Source/Editor/Sequencer/Public/SequencerCommands.h
+++ b/Engine/Source/Editor/Sequencer/Public/SequencerCommands.h
@@ -240,7 +240,10 @@ public:
 
 	/** Opens the tagged binding manager. */
 	TSharedPtr< FUICommandInfo > OpenTaggedBindingManager;
-	
+
+	/** Opens the node group manager. */
+	TSharedPtr< FUICommandInfo > OpenNodeGroupsManager;
+
 	/** Sets the tree search widget as the focused widget in Slate for easy typing. */
 	TSharedPtr< FUICommandInfo > QuickTreeSearch;
 
diff --git a/Engine/Source/Runtime/MovieScene/Private/MovieScene.cpp b/Engine/Source/Runtime/MovieScene/Private/MovieScene.cpp
--- a/Engine/Source/Runtime/MovieScene/Private/MovieScene.cpp
+++ b/Engine/Source/Runtime/MovieScene/Private/MovieScene.cpp
@@ -52,6 +52,8 @@ UMovieScene::UMovieScene(const FObjectInitializer& ObjectInitializer)
 	bForceFixedFrameIntervalPlayback_DEPRECATED = false;
 	FixedFrameInterval_DEPRECATED = 0.0f;
 
+	NodeGroupCollection = CreateEditorOnlyDefaultSubobject<UMovieSceneNodeGroupCollection>("NodeGroupCollection");
+
 	InTime_DEPRECATED    =  FLT_MAX;
 	OutTime_DEPRECATED   = -FLT_MAX;
 	StartTime_DEPRECATED =  FLT_MAX;
@@ -532,6 +534,89 @@ void FMovieSceneSectionGroup::Clean()
 	Sections.RemoveAll([](TWeakObjectPtr<UMovieSceneSection>& Section) { return !(Section.IsValid()); });
 }
 
+void UMovieSceneNodeGroup::AddNode(const FString& Path)
+{
+	Nodes.AddUnique(Path);
+	OnNodeGroupChangedEvent.Broadcast();
+}
+
+void UMovieSceneNodeGroup::RemoveNode(const FString& Path)
+{
+	Nodes.Remove(Path);
+	OnNodeGroupChangedEvent.Broadcast();
+}
+
+bool UMovieSceneNodeGroup::ContainsNode(const FString& Path) const
+{
+	return Nodes.Contains(Path);
+}
+
+void UMovieSceneNodeGroup::SetName(const FName& InName)
+{
+	Name = InName;
+	OnNodeGroupChangedEvent.Broadcast();
+}
+
+void UMovieSceneNodeGroup::SetEnableFilter(bool bInEnableFilter)
+{
+	if (bEnableFilter != bInEnableFilter)
+	{
+		bEnableFilter = bInEnableFilter;
+		OnNodeGroupChangedEvent.Broadcast();
+	}
+}
+
+void UMovieSceneNodeGroupCollection::PostLoad()
+{
+	bAnyActiveFilter = false;
+	for (UMovieSceneNodeGroup* NodeGroup : NodeGroups)
+	{
+		NodeGroup->OnNodeGroupChanged().AddUObject(this, &UMovieSceneNodeGroupCollection::OnNodeGroupChanged);
+		
+		if (NodeGroup->GetEnableFilter())
+		{
+			bAnyActiveFilter = true;
+		}
+	}
+	
+	Super::PostLoad();
+}
+
+void UMovieSceneNodeGroupCollection::AddNodeGroup(UMovieSceneNodeGroup* NodeGroup)
+{
+	if (!NodeGroups.Contains(NodeGroup))
+	{
+		NodeGroups.Add(NodeGroup);
+		NodeGroup->OnNodeGroupChanged().AddUObject(this, &UMovieSceneNodeGroupCollection::OnNodeGroupChanged);
+		
+		OnNodeGroupCollectionChangedEvent.Broadcast();
+	}
+}
+
+void UMovieSceneNodeGroupCollection::RemoveNodeGroup(UMovieSceneNodeGroup* NodeGroup)
+{
+	NodeGroup->OnNodeGroupChanged().RemoveAll(this);
+
+	if (NodeGroups.RemoveSingle(NodeGroup))
+	{
+		OnNodeGroupCollectionChangedEvent.Broadcast();
+	}
+}
+
+void UMovieSceneNodeGroupCollection::OnNodeGroupChanged()
+{
+	bAnyActiveFilter = false;
+	for (const UMovieSceneNodeGroup* NodeGroup : NodeGroups)
+	{
+		if (NodeGroup->GetEnableFilter())
+		{
+			bAnyActiveFilter = true;
+			break;
+		}
+	}
+
+	OnNodeGroupCollectionChangedEvent.Broadcast();
+}
 
 bool UMovieScene::IsSectionInGroup(const UMovieSceneSection& InSection) const
 {
diff --git a/Engine/Source/Runtime/MovieScene/Public/MovieScene.h b/Engine/Source/Runtime/MovieScene/Public/MovieScene.h
--- a/Engine/Source/Runtime/MovieScene/Public/MovieScene.h
+++ b/Engine/Source/Runtime/MovieScene/Public/MovieScene.h
@@ -235,7 +235,7 @@ public:
 	}
 	FORCEINLINE bool operator!=(const FMovieSceneSectionGroup& Other) const
 	{
-		return &Sections != &Other.Sections;;
+		return &Sections != &Other.Sections;
 	}
 
 	/**
@@ -248,6 +248,109 @@ public:
 	FORCEINLINE TArray<TWeakObjectPtr<UMovieSceneSection> >::RangedForConstIteratorType end() const { return Sections.end(); }
 };
 
+/**
+ * Structure that represents a group of nodes
+ */
+UCLASS()
+class MOVIESCENE_API UMovieSceneNodeGroup : public UObject
+{
+	GENERATED_BODY()
+
+#if WITH_EDITORONLY_DATA
+public:
+	const FName GetName() const { return Name; }
+	void SetName(const FName& Name);
+
+	void AddNode(const FString& Path);
+	void RemoveNode(const FString& Path);
+	TArrayView<FString> GetNodes() { return Nodes; }
+	bool ContainsNode(const FString& Path) const;
+
+	bool GetEnableFilter() const { return bEnableFilter; }
+	void SetEnableFilter(bool bInEnableFilter);
+
+	/** Event that is triggered whenever this node group has changed */
+	DECLARE_EVENT(UMovieSceneNodeGroup, FOnNodeGroupChanged)
+	FOnNodeGroupChanged& OnNodeGroupChanged() { return OnNodeGroupChangedEvent; }
+
+private:
+	UPROPERTY()
+	FName Name;
+
+	/** Nodes that are part of this node group, stored as node tree paths */
+	UPROPERTY()
+	TArray<FString> Nodes;
+
+	/** Whether sequencer should filter to only show this nodes in this group */
+	bool bEnableFilter;
+
+	/** Event that is triggered whenever this node group has changed */
+	FOnNodeGroupChanged OnNodeGroupChangedEvent;
+
+	/**
+	 * Comparison operators
+	 * We only need these for being stored in a container, to check if it's the same object.
+	 * Not intended for direct use.
+	 */
+	FORCEINLINE bool operator==(const UMovieSceneNodeGroup& Other) const
+	{
+		return &Nodes == &Other.Nodes;
+	}
+	FORCEINLINE bool operator!=(const UMovieSceneNodeGroup& Other) const
+	{
+		return &Nodes != &Other.Nodes;
+	}
+#endif
+};
+
+/**
+ * Structure that represents a collection of NodeGroups
+ */
+UCLASS()
+class MOVIESCENE_API UMovieSceneNodeGroupCollection : public UObject
+{
+	GENERATED_BODY()
+
+#if WITH_EDITORONLY_DATA
+public:
+	/** Called after this object has been deserialized */
+	virtual void PostLoad() override;
+
+	void AddNodeGroup(UMovieSceneNodeGroup* NodeGroup);
+	void RemoveNodeGroup(UMovieSceneNodeGroup* NodeGroup);
+
+	bool Contains(UMovieSceneNodeGroup* NodeGroup) const { return NodeGroups.Contains(NodeGroup); }
+	const int32 Num() const { return NodeGroups.Num(); }
+	bool HasAnyActiveFilter() const { return bAnyActiveFilter; }
+
+	/** Event that is triggered whenever this collection of node groups, or an included node group, has changed */
+	DECLARE_EVENT(UMovieSceneNodeGroupCollection, FOnNodeGroupCollectionChanged)
+	FOnNodeGroupCollectionChanged& OnNodeGroupCollectionChanged() { return OnNodeGroupCollectionChangedEvent; }
+
+private:
+	UPROPERTY()
+	TArray<UMovieSceneNodeGroup*> NodeGroups;
+
+	bool bAnyActiveFilter;
+
+	void OnNodeGroupChanged();
+
+	/** Event that is triggered whenever this collection of node groups, or an included node group, has changed */
+	FOnNodeGroupCollectionChanged OnNodeGroupCollectionChangedEvent;
+
+public:
+
+	/**
+	 * DO NOT USE DIRECTLY
+	 * STL-like iterators to enable range-based for loop support.
+	 */
+	FORCEINLINE TArray<UMovieSceneNodeGroup*>::RangedForIteratorType      begin()	{ return NodeGroups.begin(); }
+	FORCEINLINE TArray<UMovieSceneNodeGroup*>::RangedForConstIteratorType begin()	const { return NodeGroups.begin(); }
+	FORCEINLINE TArray<UMovieSceneNodeGroup*>::RangedForIteratorType      end()	{ return NodeGroups.end(); }
+	FORCEINLINE TArray<UMovieSceneNodeGroup*>::RangedForConstIteratorType end()	const { return NodeGroups.end(); }
+#endif
+};
+
 /**
  * Implements a movie scene asset.
  */
@@ -873,6 +976,8 @@ public:
 	 */
 	void CleanSectionGroups();
 
+	UMovieSceneNodeGroupCollection& GetNodeGroups() { return *NodeGroupCollection; }
+
 	/** The timecode at which this movie scene section is based (ie. when it was recorded) */
 	UPROPERTY()
 	FMovieSceneTimecodeSource TimecodeSource;
@@ -1100,6 +1205,12 @@ private:
 	UPROPERTY()
 	TArray<FMovieSceneSectionGroup> SectionGroups;
 
+#if WITH_EDITORONLY_DATA
+	/** Collection of user-defined groups */
+	UPROPERTY()
+	UMovieSceneNodeGroupCollection* NodeGroupCollection;
+#endif
+
 	/** Whether this scene's marked frames should be shown globally */
 	bool bGloballyShowMarkedFrames;
 
