Sequencer: When renaming or moving a node, update the nodepath to fix identifying by nodepath

diff --git a/Engine/Source/Editor/Sequencer/Private/DisplayNodes/SequencerDisplayNode.cpp b/Engine/Source/Editor/Sequencer/Private/DisplayNodes/SequencerDisplayNode.cpp
--- a/Engine/Source/Editor/Sequencer/Private/DisplayNodes/SequencerDisplayNode.cpp
+++ b/Engine/Source/Editor/Sequencer/Private/DisplayNodes/SequencerDisplayNode.cpp
@@ -451,6 +451,8 @@ void FSequencerDisplayNode::SetParent(TSharedPtr<FSequencerDisplayNode> InParent
 	TSharedPtr<FSequencerDisplayNode> CurrentParent = ParentNode.Pin();
 	if (CurrentParent != InParent)
 	{
+		const FString OldPath = GetPathName();
+
 		TSharedRef<FSequencerDisplayNode> ThisNode = AsShared();
 		if (CurrentParent)
 		{
@@ -478,9 +480,11 @@ void FSequencerDisplayNode::SetParent(TSharedPtr<FSequencerDisplayNode> InParent
 				ParentTree.SavePinnedState(*this, false);
 			}
 		}
-	}
 
-	ParentNode = InParent;
+		ParentNode = InParent;
+
+		ParentTree.GetSequencer().OnNodePathChanged(OldPath, GetPathName());
+	}
 }
 
 void FSequencerDisplayNode::MoveChild(int32 InChildIndex, int32 InDesiredNewIndex)
@@ -925,6 +929,12 @@ FString FSequencerDisplayNode::GetPathName() const
 	return PathName;
 }
 
+void FSequencerDisplayNode::SetNodeName(const FName& InName)
+{
+	const FString OldPath = GetPathName();
+	NodeName = InName;
+	ParentTree.GetSequencer().OnNodePathChanged(OldPath, GetPathName());
+}
 
 TSharedPtr<SWidget> FSequencerDisplayNode::OnSummonContextMenu()
 {
diff --git a/Engine/Source/Editor/Sequencer/Private/DisplayNodes/SequencerDisplayNode.h b/Engine/Source/Editor/Sequencer/Private/DisplayNodes/SequencerDisplayNode.h
--- a/Engine/Source/Editor/Sequencer/Private/DisplayNodes/SequencerDisplayNode.h
+++ b/Engine/Source/Editor/Sequencer/Private/DisplayNodes/SequencerDisplayNode.h
@@ -341,6 +341,11 @@ public:
 		return NodeName;
 	}
 
+	/**
+	 * Set the name of the node (for identification purposes)
+	 */
+	void SetNodeName(const FName& InName);
+
 	/**
 	 * @return The number of child nodes belonging to this node
 	 */
diff --git a/Engine/Source/Editor/Sequencer/Private/DisplayNodes/SequencerFolderNode.cpp b/Engine/Source/Editor/Sequencer/Private/DisplayNodes/SequencerFolderNode.cpp
--- a/Engine/Source/Editor/Sequencer/Private/DisplayNodes/SequencerFolderNode.cpp
+++ b/Engine/Source/Editor/Sequencer/Private/DisplayNodes/SequencerFolderNode.cpp
@@ -97,6 +97,8 @@ void FSequencerFolderNode::SetDisplayName( const FText& NewDisplayName )
 
 		const FScopedTransaction Transaction( NSLOCTEXT( "SequencerFolderNode", "RenameFolder", "Rename folder." ) );
 		MovieSceneFolder.SetFolderName( UniqueName );
+
+		SetNodeName( UniqueName );
 	}
 }
 
diff --git a/Engine/Source/Editor/Sequencer/Private/DisplayNodes/SequencerObjectBindingNode.cpp b/Engine/Source/Editor/Sequencer/Private/DisplayNodes/SequencerObjectBindingNode.cpp
--- a/Engine/Source/Editor/Sequencer/Private/DisplayNodes/SequencerObjectBindingNode.cpp
+++ b/Engine/Source/Editor/Sequencer/Private/DisplayNodes/SequencerObjectBindingNode.cpp
@@ -885,6 +885,8 @@ void FSequencerObjectBindingNode::SetDisplayName(const FText& NewDisplayName)
 		{
 			MovieScene->SetObjectDisplayName(ObjectBinding, NewDisplayName);
 		}
+
+		SetNodeName(FName(*NewDisplayName.ToString()));
 	}
 }
 
diff --git a/Engine/Source/Editor/Sequencer/Private/DisplayNodes/SequencerTrackNode.cpp b/Engine/Source/Editor/Sequencer/Private/DisplayNodes/SequencerTrackNode.cpp
--- a/Engine/Source/Editor/Sequencer/Private/DisplayNodes/SequencerTrackNode.cpp
+++ b/Engine/Source/Editor/Sequencer/Private/DisplayNodes/SequencerTrackNode.cpp
@@ -783,6 +783,8 @@ void FSequencerTrackNode::SetDisplayName(const FText& NewDisplayName)
 
 		NameableTrack->SetDisplayName(NewDisplayName);
 		GetSequencer().NotifyMovieSceneDataChanged(EMovieSceneDataChangeType::TrackValueChanged);
+
+		SetNodeName(FName(*NewDisplayName.ToString()));
 	}
 }
 
diff --git a/Engine/Source/Editor/Sequencer/Private/Sequencer.cpp b/Engine/Source/Editor/Sequencer/Private/Sequencer.cpp
--- a/Engine/Source/Editor/Sequencer/Private/Sequencer.cpp
+++ b/Engine/Source/Editor/Sequencer/Private/Sequencer.cpp
@@ -6365,6 +6365,77 @@ bool FSequencer::IsTrackVisible(const UMovieSceneTrack* InTrack)
 	return true;
 }
 
+void FSequencer::OnNodePathChanged(const FString& OldPath, const FString& NewPath)
+{
+	if (!OldPath.Equals(NewPath))
+	{
+		UMovieScene* MovieScene = GetFocusedMovieSceneSequence()->GetMovieScene();
+
+		if (!MovieScene->IsReadOnly())
+		{
+			MovieScene->Modify();
+
+			MovieScene->GetNodeGroups().UpdateNodePath(OldPath, NewPath);
+
+			// If the node is in the solo list, replace it with it's new path
+			if (MovieScene->GetSoloNodes().Remove(OldPath))
+			{
+				MovieScene->GetSoloNodes().Add(NewPath);
+			}
+
+			// If the node is in the mute list, replace it with it's new path
+			if (MovieScene->GetMuteNodes().Remove(OldPath))
+			{
+				MovieScene->GetMuteNodes().Add(NewPath);
+			}
+
+			// Find any solo/muted nodes with a path that is a child of the renamed node, and rename their paths as well
+			FString PathPrefix = OldPath + '.';
+
+			TArray<FString> PathsToRename;
+			for (const FString& NodePath : MovieScene->GetSoloNodes())
+			{
+				if (NodePath.StartsWith(PathPrefix))
+				{
+					PathsToRename.Add(NodePath);
+				}
+			}
+
+			for (const FString& NodePath : PathsToRename)
+			{
+				FString NewNodePath = NodePath;
+				if (NewNodePath.RemoveFromStart(PathPrefix))
+				{
+					NewNodePath = NewPath + '.' + NewNodePath;
+					MovieScene->GetSoloNodes().Remove(NodePath);
+					MovieScene->GetSoloNodes().Add(NewNodePath);
+				}
+			}
+
+			PathsToRename.Empty();
+			for (const FString& NodePath : MovieScene->GetMuteNodes())
+			{
+				if (NodePath.StartsWith(PathPrefix))
+				{
+					PathsToRename.Add(NodePath);
+				}
+			}
+
+			for (const FString& NodePath : PathsToRename)
+			{
+				FString NewNodePath = NodePath;
+				if (NewNodePath.RemoveFromStart(PathPrefix))
+				{
+					NewNodePath = NewPath + '.' + NewNodePath;
+					MovieScene->GetMuteNodes().Remove(NodePath);
+					MovieScene->GetMuteNodes().Add(NewNodePath);
+				}
+			}
+
+		}
+	}
+}
+
 void FSequencer::OnSelectedNodesOnlyChanged()
 {
 	RefreshTree();
diff --git a/Engine/Source/Editor/Sequencer/Private/Sequencer.h b/Engine/Source/Editor/Sequencer/Private/Sequencer.h
--- a/Engine/Source/Editor/Sequencer/Private/Sequencer.h
+++ b/Engine/Source/Editor/Sequencer/Private/Sequencer.h
@@ -570,6 +570,9 @@ public:
 	/** Whether the track is visible in the tree view */
 	bool IsTrackVisible(const UMovieSceneTrack* InTrack);
 
+	/** Call when the path to a display node changes, to update anything tracking the node via path */
+	void OnNodePathChanged(const FString& OldPath, const FString& NewPath);
+
 	void OnSelectedNodesOnlyChanged();
 
 	TSharedPtr<FCurveEditor> GetCurveEditor() const
diff --git a/Engine/Source/Runtime/MovieScene/Private/MovieScene.cpp b/Engine/Source/Runtime/MovieScene/Private/MovieScene.cpp
--- a/Engine/Source/Runtime/MovieScene/Private/MovieScene.cpp
+++ b/Engine/Source/Runtime/MovieScene/Private/MovieScene.cpp
@@ -551,6 +551,42 @@ bool UMovieSceneNodeGroup::ContainsNode(const FString& Path) const
 	return Nodes.Contains(Path);
 }
 
+void UMovieSceneNodeGroup::UpdateNodePath(const FString& OldPath, const FString& NewPath)
+{
+	if (!OldPath.Equals(NewPath))
+	{
+		// If the node is in this group, replace it with it's new path
+		if (Nodes.Remove(OldPath))
+		{
+			Nodes.Add(NewPath);
+		}
+
+		// Find any nodes with a path that is a child of the changed node, and rename their paths as well
+		FString PathPrefix = OldPath + '.';
+
+		TArray<FString> PathsToRename;
+		for (const FString& NodePath : Nodes)
+		{
+			if (NodePath.StartsWith(PathPrefix))
+			{
+				PathsToRename.Add(NodePath);
+			}
+		}
+
+		for (const FString& NodePath : PathsToRename)
+		{
+			FString NewNodePath = NodePath;
+			if (NewNodePath.RemoveFromStart(PathPrefix))
+			{
+				NewNodePath = NewPath + '.' + NewNodePath;
+				Nodes.Remove(NodePath);
+				Nodes.Add(NewNodePath);
+			}
+		}
+
+	}
+}
+
 void UMovieSceneNodeGroup::SetName(const FName& InName)
 {
 	Name = InName;
@@ -603,6 +639,14 @@ void UMovieSceneNodeGroupCollection::RemoveNodeGroup(UMovieSceneNodeGroup* NodeG
 	}
 }
 
+void UMovieSceneNodeGroupCollection::UpdateNodePath(const FString& OldPath, const FString& NewPath)
+{
+	for (UMovieSceneNodeGroup* NodeGroup : NodeGroups)
+	{
+		NodeGroup->UpdateNodePath(OldPath, NewPath);
+	}
+}
+
 void UMovieSceneNodeGroupCollection::OnNodeGroupChanged()
 {
 	bAnyActiveFilter = false;
diff --git a/Engine/Source/Runtime/MovieScene/Public/MovieScene.h b/Engine/Source/Runtime/MovieScene/Public/MovieScene.h
--- a/Engine/Source/Runtime/MovieScene/Public/MovieScene.h
+++ b/Engine/Source/Runtime/MovieScene/Public/MovieScene.h
@@ -266,6 +266,8 @@ public:
 	TArrayView<FString> GetNodes() { return Nodes; }
 	bool ContainsNode(const FString& Path) const;
 
+	void UpdateNodePath(const FString& OldPath, const FString& NewPath);
+
 	bool GetEnableFilter() const { return bEnableFilter; }
 	void SetEnableFilter(bool bInEnableFilter);
 
@@ -323,6 +325,8 @@ public:
 	const int32 Num() const { return NodeGroups.Num(); }
 	bool HasAnyActiveFilter() const { return bAnyActiveFilter; }
 
+	void UpdateNodePath(const FString& OldPath, const FString& NewPath);
+
 	/** Event that is triggered whenever this collection of node groups, or an included node group, has changed */
 	DECLARE_EVENT(UMovieSceneNodeGroupCollection, FOnNodeGroupCollectionChanged)
 	FOnNodeGroupCollectionChanged& OnNodeGroupCollectionChanged() { return OnNodeGroupCollectionChangedEvent; }
