Sequencer: Fix bake transform so that it iterates through the selected node guides, rather than through the selected nodes since the selection can change during the bake process.

diff --git a/Engine/Source/Editor/Sequencer/Private/Sequencer.cpp b/Engine/Source/Editor/Sequencer/Private/Sequencer.cpp
--- a/Engine/Source/Editor/Sequencer/Private/Sequencer.cpp
+++ b/Engine/Source/Editor/Sequencer/Private/Sequencer.cpp
@@ -1926,6 +1926,7 @@ void FSequencer::BakeTransform()
 
 	FocusedMovieScene->Modify();
 
+	TArray<FGuid> Guids;
 	for (const TSharedRef<FSequencerDisplayNode>& Node : Selection.GetSelectedOutlinerNodes())
 	{
 		if (Node->GetType() != ESequencerNode::Object)
@@ -1933,10 +1934,15 @@ void FSequencer::BakeTransform()
 			continue;
 		}
 
-		FFrameTime ResetTime = PlayPosition.GetCurrentPosition();
 		auto ObjectBindingNode = StaticCastSharedRef<FSequencerObjectBindingNode>(Node);
-
 		FGuid Guid = ObjectBindingNode->GetObjectBinding();
+
+		Guids.Add(Guid);
+	}
+
+	FFrameTime ResetTime = PlayPosition.GetCurrentPosition();
+	for (FGuid Guid : Guids)
+	{
 		for (auto RuntimeObject : FindBoundObjects(Guid, ActiveTemplateIDs.Top()) )
 		{
 			AActor* Actor = Cast<AActor>(RuntimeObject.Get());
