Sequencer: Added missing tick animation in bake

#rnx

diff --git a/Engine/Source/Editor/Sequencer/Private/Sequencer.cpp b/Engine/Source/Editor/Sequencer/Private/Sequencer.cpp
--- a/Engine/Source/Editor/Sequencer/Private/Sequencer.cpp
+++ b/Engine/Source/Editor/Sequencer/Private/Sequencer.cpp
@@ -2001,6 +2001,8 @@ void FSequencer::BakeTransform()
 
 					if (USkeletalMeshComponent* SkeletalMeshComponent = Cast<USkeletalMeshComponent>(Parent))
 					{
+						SkeletalMeshComponent->TickAnimation(0.f, false);
+
 						SkeletalMeshComponent->RefreshBoneTransforms();
 						SkeletalMeshComponent->RefreshSlaveComponents();
 						SkeletalMeshComponent->UpdateComponentToWorld();
