Sequencer: Fix getting key positions when duplicate threshold is large. Clamp threshold by the max time that the tick resolution can represent.

#rnx

diff --git a/Engine/Source/Editor/Sequencer/Private/Sequencer.cpp b/Engine/Source/Editor/Sequencer/Private/Sequencer.cpp
--- a/Engine/Source/Editor/Sequencer/Private/Sequencer.cpp
+++ b/Engine/Source/Editor/Sequencer/Private/Sequencer.cpp
@@ -1058,11 +1058,21 @@ void FSequencer::GetKeysFromSelection(TUniquePtr<FSequencerKeyCollection>& KeyCo
 		SelectedNodes.Add(&Node.Get());
 	}
 
+	int64 TotalMaxSeconds = static_cast<int64>(TNumericLimits<int32>::Max() / GetFocusedTickResolution().AsDecimal());
+
 	FFrameNumber ThresholdFrames = (DuplicateThresholdSeconds * GetFocusedTickResolution()).FloorToFrame();
+	if (ThresholdFrames.Value < -TotalMaxSeconds)
+	{
+		ThresholdFrames.Value = TotalMaxSeconds;
+	}
+	else if (ThresholdFrames.Value > TotalMaxSeconds)
+	{
+		ThresholdFrames.Value = TotalMaxSeconds;
+	}
+
 	KeyCollection->Update(FSequencerKeyCollectionSignature::FromNodesRecursive(SelectedNodes, ThresholdFrames));
 }
 
-
 void FSequencer::GetAllKeys(TUniquePtr<FSequencerKeyCollection>& KeyCollection, float DuplicateThresholdSeconds)
 {
 	if (!KeyCollection.IsValid())
