Sequencer: Update bindings when converting from spawnable to possessable and vice versa.

diff --git a/Engine/Source/Editor/Sequencer/Private/Sequencer.cpp b/Engine/Source/Editor/Sequencer/Private/Sequencer.cpp
--- a/Engine/Source/Editor/Sequencer/Private/Sequencer.cpp
+++ b/Engine/Source/Editor/Sequencer/Private/Sequencer.cpp
@@ -9202,6 +9202,23 @@ TArray<FMovieSceneSpawnable*> FSequencer::ConvertToSpawnableInternal(FGuid Posse
 			SpawnRegister->HandleConvertPossessableToSpawnable(FoundObject, *this, TransformData);
 			SpawnRegister->SetupDefaultsForSpawnable(nullptr, Spawnable->GetGuid(), TransformData, AsShared(), Settings);
 
+			TMap<FGuid, FGuid> OldGuidToNewGuidMap;
+			OldGuidToNewGuidMap.Add(PossessableGuid, Spawnable->GetGuid());
+			
+			// Fixup any section bindings
+			TArray<UMovieScene*> MovieScenesToUpdate;
+			MovieSceneHelpers::GetDescendantMovieScenes(GetRootMovieSceneSequence(), MovieScenesToUpdate);
+			for (UMovieScene* MovieSceneToUpdate : MovieScenesToUpdate)
+			{
+				for (UMovieSceneSection* Section : MovieSceneToUpdate->GetAllSections())
+				{
+					if (Section)
+					{
+						Section->OnBindingsUpdated(OldGuidToNewGuidMap);
+					}
+				}
+			}
+
 			ForceEvaluate();
 
 			NotifyMovieSceneDataChanged(EMovieSceneDataChangeType::MovieSceneStructureItemsChanged);
@@ -9404,6 +9421,23 @@ FMovieScenePossessable* FSequencer::ConvertToPossessableInternal(FGuid Spawnable
 		static const FName SequencerActorTag(TEXT("SequencerActor"));
 		PossessedActor->Tags.Remove(SequencerActorTag);
 
+		TMap<FGuid, FGuid> OldGuidToNewGuidMap;
+		OldGuidToNewGuidMap.Add(OldSpawnableGuid, NewPossessableGuid);
+		
+		// Fixup any section bindings
+		TArray<UMovieScene*> MovieScenesToUpdate;
+		MovieSceneHelpers::GetDescendantMovieScenes(GetRootMovieSceneSequence(), MovieScenesToUpdate);
+		for (UMovieScene* MovieSceneToUpdate : MovieScenesToUpdate)
+		{
+			for (UMovieSceneSection* Section : MovieSceneToUpdate->GetAllSections())
+			{
+				if (Section)
+				{
+					Section->OnBindingsUpdated(OldGuidToNewGuidMap);
+				}
+			}
+		}
+
 		GEditor->SelectActor(PossessedActor, false, true);
 
 		ForceEvaluate();
@@ -10767,9 +10801,17 @@ void FSequencer::FixActorReferences()
 	}
 
 	// Fixup any section bindings
-	for (UMovieSceneSection* Section : FocusedMovieScene->GetAllSections())
+	TArray<UMovieScene*> MovieScenesToUpdate;
+	MovieSceneHelpers::GetDescendantMovieScenes(GetRootMovieSceneSequence(), MovieScenesToUpdate);
+	for (UMovieScene* MovieSceneToUpdate : MovieScenesToUpdate)
 	{
-		Section->OnBindingsUpdated(OldGuidToNewGuidMap);
+		for (UMovieSceneSection* Section : MovieSceneToUpdate->GetAllSections())
+		{
+			if (Section)
+			{
+				Section->OnBindingsUpdated(OldGuidToNewGuidMap);
+			}
+		}
 	}
 }
 
diff --git a/Engine/Source/Runtime/MovieSceneTracks/Private/Sections/MovieScene3DConstraintSection.cpp b/Engine/Source/Runtime/MovieSceneTracks/Private/Sections/MovieScene3DConstraintSection.cpp
--- a/Engine/Source/Runtime/MovieSceneTracks/Private/Sections/MovieScene3DConstraintSection.cpp
+++ b/Engine/Source/Runtime/MovieSceneTracks/Private/Sections/MovieScene3DConstraintSection.cpp
@@ -30,6 +30,8 @@ void UMovieScene3DConstraintSection::OnBindingsUpdated(const TMap<FGuid, FGuid>&
 {
 	if (OldGuidToNewGuidMap.Contains(ConstraintBindingID.GetGuid()))
 	{
+		Modify();
+
 		ConstraintBindingID.SetGuid(OldGuidToNewGuidMap[ConstraintBindingID.GetGuid()]);
 	}
 }
diff --git a/Engine/Source/Runtime/MovieSceneTracks/Private/Sections/MovieSceneActorReferenceSection.cpp b/Engine/Source/Runtime/MovieSceneTracks/Private/Sections/MovieSceneActorReferenceSection.cpp
--- a/Engine/Source/Runtime/MovieSceneTracks/Private/Sections/MovieSceneActorReferenceSection.cpp
+++ b/Engine/Source/Runtime/MovieSceneTracks/Private/Sections/MovieSceneActorReferenceSection.cpp
@@ -132,3 +132,24 @@ void UMovieSceneActorReferenceSection::PostLoad()
 		}
 	}
 }
+
+void UMovieSceneActorReferenceSection::OnBindingsUpdated(const TMap<FGuid, FGuid>& OldGuidToNewGuidMap)
+{
+	if (OldGuidToNewGuidMap.Contains(ActorReferenceData.GetDefault().Object.GetGuid()))
+	{
+		Modify();
+
+		ActorReferenceData.GetDefault().Object.SetGuid(OldGuidToNewGuidMap[ActorReferenceData.GetDefault().Object.GetGuid()]);
+	}
+
+	for (FMovieSceneActorReferenceKey Key : ActorReferenceData.GetData().GetValues())
+	{
+		if (OldGuidToNewGuidMap.Contains(Key.Object.GetGuid()))
+		{
+			Modify();
+			
+			Key.Object.SetGuid(OldGuidToNewGuidMap[Key.Object.GetGuid()]);
+		}	
+	}
+}
+
diff --git a/Engine/Source/Runtime/MovieSceneTracks/Private/Sections/MovieSceneCameraCutSection.cpp b/Engine/Source/Runtime/MovieSceneTracks/Private/Sections/MovieSceneCameraCutSection.cpp
--- a/Engine/Source/Runtime/MovieSceneTracks/Private/Sections/MovieSceneCameraCutSection.cpp
+++ b/Engine/Source/Runtime/MovieSceneTracks/Private/Sections/MovieSceneCameraCutSection.cpp
@@ -68,6 +68,8 @@ void UMovieSceneCameraCutSection::OnBindingsUpdated(const TMap<FGuid, FGuid>& Ol
 {
 	if (OldGuidToNewGuidMap.Contains(CameraBindingID.GetGuid()))
 	{
+		Modify();
+
 		CameraBindingID.SetGuid(OldGuidToNewGuidMap[CameraBindingID.GetGuid()]);
 	}
 }
diff --git a/Engine/Source/Runtime/MovieSceneTracks/Public/Sections/MovieSceneActorReferenceSection.h b/Engine/Source/Runtime/MovieSceneTracks/Public/Sections/MovieSceneActorReferenceSection.h
--- a/Engine/Source/Runtime/MovieSceneTracks/Public/Sections/MovieSceneActorReferenceSection.h
+++ b/Engine/Source/Runtime/MovieSceneTracks/Public/Sections/MovieSceneActorReferenceSection.h
@@ -167,6 +167,9 @@ public:
 	//~ UObject interface
 	virtual void PostLoad() override;
 
+	//~ UMovieSceneSection interface
+	virtual void OnBindingsUpdated(const TMap<FGuid, FGuid>& OldGuidToNewGuidMap) override;
+
 	const FMovieSceneActorReferenceData& GetActorReferenceData() const { return ActorReferenceData; }
 
 private:
