Sequencer: Rework bake transform so that it iterates through the bake times first and then gathers bound objects. This fixes not being able to bake unspawned actors. Also, handles scale properly now.

diff --git a/Engine/Source/Editor/Sequencer/Private/Sequencer.cpp b/Engine/Source/Editor/Sequencer/Private/Sequencer.cpp
--- a/Engine/Source/Editor/Sequencer/Private/Sequencer.cpp
+++ b/Engine/Source/Editor/Sequencer/Private/Sequencer.cpp
@@ -1951,49 +1951,57 @@ void FSequencer::BakeTransform()
 	}
 
 	FFrameTime ResetTime = PlayPosition.GetCurrentPosition();
+
+	FFrameRate   Resolution = FocusedMovieScene->GetTickResolution();
+	FFrameRate   SnapRate = FocusedMovieScene->GetDisplayRate();
+
+	FFrameNumber InFrame = MovieScene::DiscreteInclusiveLower(GetPlaybackRange());
+	FFrameNumber OutFrame = MovieScene::DiscreteExclusiveUpper(GetPlaybackRange());
+
+	struct FBakeData
+	{
+		TArray<FVector> Locations;
+		TArray<FRotator> Rotations;
+		TArray<FVector> Scales;
+		TArray<FFrameNumber> KeyTimes;
+	};
+
+	TMap<FGuid, FBakeData> BakeDataMap;
 	for (FGuid Guid : Guids)
 	{
-		for (auto RuntimeObject : FindBoundObjects(Guid, ActiveTemplateIDs.Top()) )
+		BakeDataMap.Add(Guid);
+	}
+
+	FFrameTime Interval = FFrameRate::TransformTime(1, SnapRate, Resolution);
+	for (FFrameTime EvalTime = InFrame; EvalTime < OutFrame; EvalTime += Interval)
+	{
+		FFrameNumber KeyTime = FFrameRate::Snap(EvalTime, Resolution, SnapRate).FloorToFrame();
+		FMovieSceneEvaluationRange Range = PlayPosition.JumpTo(KeyTime * RootToLocalTransform.InverseLinearOnly());
+
+		EvaluateInternal(Range);
+
+		for (FGuid Guid : Guids)
 		{
-			AActor* Actor = Cast<AActor>(RuntimeObject.Get());
-			if (!Actor)
+			for (auto RuntimeObject : FindBoundObjects(Guid, ActiveTemplateIDs.Top()) )
 			{
-				UActorComponent* ActorComponent = Cast<UActorComponent>(RuntimeObject.Get());
-				if (ActorComponent)
+				AActor* Actor = Cast<AActor>(RuntimeObject.Get());
+				if (!Actor)
 				{
-					Actor = ActorComponent->GetOwner();
+					UActorComponent* ActorComponent = Cast<UActorComponent>(RuntimeObject.Get());
+					if (ActorComponent)
+					{
+						Actor = ActorComponent->GetOwner();
+					}
 				}
-			}
-
-			if (!Actor)
-			{
-				continue;
-			}
-
-			UCameraComponent* CameraComponent = MovieSceneHelpers::CameraComponentFromRuntimeObject(RuntimeObject.Get());
-
-			FVector Location = Actor->GetActorLocation();
-			FVector Rotation = Actor->GetActorRotation().Euler();
-			FVector Scale = Actor->GetActorScale();
-
-			// Cache transforms
-			TArray<FVector> Locations;
-			TArray<FRotator> Rotations;
-			TArray<FFrameNumber> KeyTimes;
 
-			FFrameRate   Resolution  = FocusedMovieScene->GetTickResolution();
-			FFrameRate   SnapRate    = FocusedMovieScene->GetDisplayRate();
+				if (!Actor)
+				{
+					continue;
+				}
 
-			FFrameNumber InFrame     = MovieScene::DiscreteInclusiveLower(GetPlaybackRange());
-			FFrameNumber OutFrame    = MovieScene::DiscreteExclusiveUpper(GetPlaybackRange());
+				UCameraComponent* CameraComponent = MovieSceneHelpers::CameraComponentFromRuntimeObject(RuntimeObject.Get());
 
-			FFrameTime Interval = FFrameRate::TransformTime(1, SnapRate, Resolution);
-			for (FFrameTime EvalTime = InFrame; EvalTime < OutFrame; EvalTime += Interval)
-			{
-				FFrameNumber KeyTime = FFrameRate::Snap(EvalTime, Resolution, SnapRate).FloorToFrame();
-				FMovieSceneEvaluationRange Range = PlayPosition.JumpTo(KeyTime * RootToLocalTransform.InverseLinearOnly());
-				EvaluateInternal(Range);
-				
+				// Cache transforms
 				USceneComponent* Parent = nullptr;
 				if (CameraComponent)
 				{
@@ -2008,7 +2016,6 @@ void FSequencer::BakeTransform()
 				while (Parent && Parent->GetOwner())
 				{
 					Parent->GetOwner()->Tick(0.03f);
-
 					if (USkeletalMeshComponent* SkeletalMeshComponent = Cast<USkeletalMeshComponent>(Parent))
 					{
 						SkeletalMeshComponent->TickAnimation(0.f, false);
@@ -2020,7 +2027,6 @@ void FSequencer::BakeTransform()
 						SkeletalMeshComponent->MarkRenderTransformDirty();
 						SkeletalMeshComponent->MarkRenderDynamicDataDirty();
 					}
-
 					Parent = Parent->GetAttachParent();
 				}
 
@@ -2035,117 +2041,156 @@ void FSequencer::BakeTransform()
 					FVector LocalTranslation = TransformWithAdditiveOffset.GetTranslation();
 					FRotator LocalRotation = TransformWithAdditiveOffset.GetRotation().Rotator();
 
-					Locations.Add(LocalTranslation);
-					Rotations.Add(LocalRotation);
+					BakeDataMap[Guid].Locations.Add(LocalTranslation);
+					BakeDataMap[Guid].Rotations.Add(LocalRotation);
+					BakeDataMap[Guid].Scales.Add(FVector::OneVector);
 				}
 				else
 				{
-					Locations.Add(Actor->GetActorLocation());
-					Rotations.Add(Actor->GetActorRotation());
+					BakeDataMap[Guid].Locations.Add(Actor->GetActorLocation());
+					BakeDataMap[Guid].Rotations.Add(Actor->GetActorRotation());
+					BakeDataMap[Guid].Scales.Add(Actor->GetActorScale());
 				}
 
-				KeyTimes.Add(KeyTime);
+				BakeDataMap[Guid].KeyTimes.Add(KeyTime);
 			}
+		}
+	}
+
+	for (auto& BakeData : BakeDataMap)
+	{
+		FGuid Guid = BakeData.Key;
 
-			// Delete any attach tracks
-			// cbb: this only operates on a single attach section.
-			AActor* AttachParentActor = nullptr;
-			UMovieScene3DAttachTrack* AttachTrack = Cast<UMovieScene3DAttachTrack>(FocusedMovieScene->FindTrack(UMovieScene3DAttachTrack::StaticClass(), Guid));
-			if (AttachTrack)
+		// Delete any attach tracks
+		// cbb: this only operates on a single attach section.
+		AActor* AttachParentActor = nullptr;
+		UMovieScene3DAttachTrack* AttachTrack = Cast<UMovieScene3DAttachTrack>(FocusedMovieScene->FindTrack(UMovieScene3DAttachTrack::StaticClass(), Guid));
+		if (AttachTrack)
+		{
+			for (auto AttachSection : AttachTrack->GetAllSections())
 			{
-				for (auto AttachSection : AttachTrack->GetAllSections())
+				FMovieSceneObjectBindingID ConstraintBindingID = (Cast<UMovieScene3DAttachSection>(AttachSection))->GetConstraintBindingID();
+				for (auto ParentObject : FindBoundObjects(ConstraintBindingID.GetGuid(), ConstraintBindingID.GetSequenceID()) )
 				{
-					FMovieSceneObjectBindingID ConstraintBindingID = (Cast<UMovieScene3DAttachSection>(AttachSection))->GetConstraintBindingID();
-					for (auto ParentObject : FindBoundObjects(ConstraintBindingID.GetGuid(), ConstraintBindingID.GetSequenceID()) )
-					{
-						AttachParentActor = Cast<AActor>(ParentObject.Get());
-						break;
-					}
+					AttachParentActor = Cast<AActor>(ParentObject.Get());
+					break;
 				}
-
-				FocusedMovieScene->RemoveTrack(*AttachTrack);
 			}
 
-			// Delete any transform tracks
-			UMovieScene3DTransformTrack* TransformTrack = Cast<UMovieScene3DTransformTrack>(FocusedMovieScene->FindTrack(UMovieScene3DTransformTrack::StaticClass(), Guid, "Transform"));
-			if (TransformTrack)
-			{
-				FocusedMovieScene->RemoveTrack(*TransformTrack);
-			}
+			FocusedMovieScene->RemoveTrack(*AttachTrack);
+		}
+
+		// Delete any transform tracks
+		UMovieScene3DTransformTrack* TransformTrack = Cast<UMovieScene3DTransformTrack>(FocusedMovieScene->FindTrack(UMovieScene3DTransformTrack::StaticClass(), Guid, "Transform"));
+		if (TransformTrack)
+		{
+			FocusedMovieScene->RemoveTrack(*TransformTrack);
+		}
+
+		// Delete any camera anim tracks
+		UMovieSceneCameraAnimTrack* CameraAnimTrack = Cast<UMovieSceneCameraAnimTrack>(FocusedMovieScene->FindTrack(UMovieSceneCameraAnimTrack::StaticClass(), Guid));
+		if (CameraAnimTrack)
+		{
+			FocusedMovieScene->RemoveTrack(*CameraAnimTrack);
+		}
+
+		// Delete any camera shake tracks
+		UMovieSceneCameraShakeTrack* CameraShakeTrack = Cast<UMovieSceneCameraShakeTrack>(FocusedMovieScene->FindTrack(UMovieSceneCameraShakeTrack::StaticClass(), Guid));
+		if (CameraShakeTrack)
+		{
+			FocusedMovieScene->RemoveTrack(*CameraShakeTrack);
+		}
+
+		// Reset position
+		EvaluateInternal(PlayPosition.JumpTo(ResetTime));
 
-			// Delete any camera anim tracks
-			UMovieSceneCameraAnimTrack* CameraAnimTrack = Cast<UMovieSceneCameraAnimTrack>(FocusedMovieScene->FindTrack(UMovieSceneCameraAnimTrack::StaticClass(), Guid));
-			if (CameraAnimTrack)
+		FVector DefaultLocation = FVector::ZeroVector;
+		FVector DefaultRotation = FVector::ZeroVector;
+		FVector DefaultScale = FVector::OneVector;
+		FTransform ParentInverseTransform;
+		ParentInverseTransform.SetIdentity();
+
+		for (auto RuntimeObject : FindBoundObjects(Guid, ActiveTemplateIDs.Top()))
+		{
+			AActor* Actor = Cast<AActor>(RuntimeObject.Get());
+			if (!Actor)
 			{
-				FocusedMovieScene->RemoveTrack(*CameraAnimTrack);
+				UActorComponent* ActorComponent = Cast<UActorComponent>(RuntimeObject.Get());
+				if (ActorComponent)
+				{
+					Actor = ActorComponent->GetOwner();
+				}
 			}
 
-			// Delete any camera shake tracks
-			UMovieSceneCameraShakeTrack* CameraShakeTrack = Cast<UMovieSceneCameraShakeTrack>(FocusedMovieScene->FindTrack(UMovieSceneCameraShakeTrack::StaticClass(), Guid));
-			if (CameraShakeTrack)
+			if (!Actor)
 			{
-				FocusedMovieScene->RemoveTrack(*CameraShakeTrack);
+				continue;
 			}
 
-			// Reset position
-			EvaluateInternal(PlayPosition.JumpTo(ResetTime));
+			DefaultLocation = Actor->GetActorLocation();
+			DefaultRotation = Actor->GetActorRotation().Euler();
+			DefaultScale = Actor->GetActorScale();
 
 			// Always detach from any existing parent
 			Actor->DetachFromActor(FDetachmentTransformRules::KeepRelativeTransform);
 
 			// If there was an attach track that was the parent, detach and attach to that actor's parent if it exists
-			FTransform ParentInverseTransform;
-			ParentInverseTransform.SetIdentity();
 			if (AttachParentActor)
 			{
 				AActor* ExistingParentActor = AttachParentActor->GetAttachParentActor();
 				if (ExistingParentActor)
 				{
 					Actor->AttachToActor(ExistingParentActor, FAttachmentTransformRules::KeepRelativeTransform);
-					ParentInverseTransform = ExistingParentActor->GetActorTransform().Inverse();
+
+					// It's possible the attachment did not succeed, in which case, we need to check the parent actor again
+					if (Actor->GetParentActor())
+					{
+						ParentInverseTransform = Actor->GetParentActor()->GetActorTransform().Inverse();
+					}
 				}
 			}
-
-			// Create new transform track and section
-			TransformTrack = Cast<UMovieScene3DTransformTrack>(FocusedMovieScene->AddTrack(UMovieScene3DTransformTrack::StaticClass(), Guid));
-
-			if (TransformTrack)
-			{
-				UMovieScene3DTransformSection* TransformSection = CastChecked<UMovieScene3DTransformSection>(TransformTrack->CreateNewSection());
-				TransformTrack->AddSection(*TransformSection);
+		}
 			
-				TransformSection->SetRange(TRange<FFrameNumber>::All());
-
-				TArrayView<FMovieSceneFloatChannel*> FloatChannels = TransformSection->GetChannelProxy().GetChannels<FMovieSceneFloatChannel>();
-				FloatChannels[0]->SetDefault(Location.X);
-				FloatChannels[1]->SetDefault(Location.Y);
-				FloatChannels[2]->SetDefault(Location.Z);
-				FloatChannels[3]->SetDefault(Rotation.X);
-				FloatChannels[4]->SetDefault(Rotation.Y);
-				FloatChannels[5]->SetDefault(Rotation.Z);
-				FloatChannels[6]->SetDefault(Scale.X);
-				FloatChannels[7]->SetDefault(Scale.Y);
-				FloatChannels[8]->SetDefault(Scale.Z);
-
-				for (int32 Counter = 0; Counter < KeyTimes.Num(); ++Counter)
-				{
-					FFrameNumber KeyTime = KeyTimes[Counter];
-
-					FTransform Transform(Rotations[Counter], Locations[Counter]);
-					FTransform LocalTransform = ParentInverseTransform * Transform;
-					FVector LocalTranslation = LocalTransform.GetTranslation();
-					FVector LocalRotation = LocalTransform.GetRotation().Euler();
-
-					FloatChannels[0]->AddLinearKey(KeyTime, LocalTranslation.X);
-					FloatChannels[1]->AddLinearKey(KeyTime, LocalTranslation.Y);
-					FloatChannels[2]->AddLinearKey(KeyTime, LocalTranslation.Z);
-					FloatChannels[3]->AddLinearKey(KeyTime, LocalRotation.X);
-					FloatChannels[4]->AddLinearKey(KeyTime, LocalRotation.Y);
-					FloatChannels[5]->AddLinearKey(KeyTime, LocalRotation.Z);
-					FloatChannels[6]->AddLinearKey(KeyTime, Scale.X);
-					FloatChannels[7]->AddLinearKey(KeyTime, Scale.Y);
-					FloatChannels[8]->AddLinearKey(KeyTime, Scale.Z);
-				}
+		// Create new transform track and section
+		TransformTrack = Cast<UMovieScene3DTransformTrack>(FocusedMovieScene->AddTrack(UMovieScene3DTransformTrack::StaticClass(), Guid));
+
+		if (TransformTrack)
+		{
+			UMovieScene3DTransformSection* TransformSection = CastChecked<UMovieScene3DTransformSection>(TransformTrack->CreateNewSection());
+			TransformTrack->AddSection(*TransformSection);
+
+			TransformSection->SetRange(TRange<FFrameNumber>::All());
+
+			TArrayView<FMovieSceneFloatChannel*> FloatChannels = TransformSection->GetChannelProxy().GetChannels<FMovieSceneFloatChannel>();
+			FloatChannels[0]->SetDefault(DefaultLocation.X);
+			FloatChannels[1]->SetDefault(DefaultLocation.Y);
+			FloatChannels[2]->SetDefault(DefaultLocation.Z);
+			FloatChannels[3]->SetDefault(DefaultRotation.X);
+			FloatChannels[4]->SetDefault(DefaultRotation.Y);
+			FloatChannels[5]->SetDefault(DefaultRotation.Z);
+			FloatChannels[6]->SetDefault(DefaultScale.X);
+			FloatChannels[7]->SetDefault(DefaultScale.Y);
+			FloatChannels[8]->SetDefault(DefaultScale.Z);
+
+			for (int32 Counter = 0; Counter < BakeData.Value.KeyTimes.Num(); ++Counter)
+			{
+				FFrameNumber KeyTime = BakeData.Value.KeyTimes[Counter];
+
+				FTransform Transform(BakeData.Value.Rotations[Counter], BakeData.Value.Locations[Counter], BakeData.Value.Scales[Counter]);
+				FTransform LocalTransform = ParentInverseTransform * Transform;
+				FVector LocalTranslation = LocalTransform.GetTranslation();
+				FVector LocalRotation = LocalTransform.GetRotation().Euler();
+				FVector LocalScale = LocalTransform.GetScale3D();
+
+				FloatChannels[0]->AddLinearKey(KeyTime, LocalTranslation.X);
+				FloatChannels[1]->AddLinearKey(KeyTime, LocalTranslation.Y);
+				FloatChannels[2]->AddLinearKey(KeyTime, LocalTranslation.Z);
+				FloatChannels[3]->AddLinearKey(KeyTime, LocalRotation.X);
+				FloatChannels[4]->AddLinearKey(KeyTime, LocalRotation.Y);
+				FloatChannels[5]->AddLinearKey(KeyTime, LocalRotation.Z);
+				FloatChannels[6]->AddLinearKey(KeyTime, LocalScale.X);
+				FloatChannels[7]->AddLinearKey(KeyTime, LocalScale.Y);
+				FloatChannels[8]->AddLinearKey(KeyTime, LocalScale.Z);
 			}
 		}
 	}
