Sequencer: Don't set viewports to realtime if they're already running PIE.

diff --git a/Engine/Source/Editor/Sequencer/Private/LevelEditorSequencerIntegration.cpp b/Engine/Source/Editor/Sequencer/Private/LevelEditorSequencerIntegration.cpp
--- a/Engine/Source/Editor/Sequencer/Private/LevelEditorSequencerIntegration.cpp
+++ b/Engine/Source/Editor/Sequencer/Private/LevelEditorSequencerIntegration.cpp
@@ -1034,6 +1034,13 @@ void FLevelEditorSequencerIntegration::DetachOutlinerColumn()
 
 void FLevelEditorSequencerIntegration::ActivateRealtimeViewports()
 {
+	// If PIE is running, the viewport will already be rendering the scene in realtime as part of the 
+	// normal game loop. If we set it to realtime, the editor would render it a second time each frame.
+	if (GEditor->IsPlaySessionInProgress())
+	{
+		return;
+	}
+
 	for (const FSequencerAndOptions& SequencerAndOptions : BoundSequencers)
 	{
 		TSharedPtr<FSequencer> Pinned = SequencerAndOptions.Sequencer.Pin();
