Fixed build and bad merge in Sequencer.cpp

diff --git a/Engine/Source/Editor/SceneOutliner/Private/ActorBrowsingMode.cpp b/Engine/Source/Editor/SceneOutliner/Private/ActorBrowsingMode.cpp
--- a/Engine/Source/Editor/SceneOutliner/Private/ActorBrowsingMode.cpp
+++ b/Engine/Source/Editor/SceneOutliner/Private/ActorBrowsingMode.cpp
@@ -27,8 +27,7 @@
 #include "ActorBrowsingModeSettings.h"
 #include "ScopedTransaction.h"
 
-//#todo_Outliner: remove
-DEFINE_LOG_CATEGORY_STATIC(LogSceneOutliner, Log, All);
+DEFINE_LOG_CATEGORY_STATIC(LogActorBrowser, Log, All);
 
 #define LOCTEXT_NAMESPACE "SceneOutliner_ActorBrowsingMode"
 
@@ -478,7 +477,7 @@ void FActorBrowsingMode::OnItemSelectionChanged(FSceneOutlinerTreeItemPtr TreeIt
 		const bool bSelectEvenIfHidden = true;	// @todo outliner: Is this actually OK?
 		for (auto* Actor : SelectedActors)
 		{
-			UE_LOG(LogSceneOutliner, Verbose, TEXT("Clicking on Actor (world outliner): %s (%s)"), *Actor->GetClass()->GetName(), *Actor->GetActorLabel());
+			UE_LOG(LogActorBrowser, Verbose, TEXT("Clicking on Actor (world outliner): %s (%s)"), *Actor->GetClass()->GetName(), *Actor->GetActorLabel());
 			GEditor->SelectActor(Actor, bShouldSelect, bNotifyAfterSelect, bSelectEvenIfHidden);
 		}
 
diff --git a/Engine/Source/Editor/Sequencer/Private/Sequencer.cpp b/Engine/Source/Editor/Sequencer/Private/Sequencer.cpp
--- a/Engine/Source/Editor/Sequencer/Private/Sequencer.cpp
+++ b/Engine/Source/Editor/Sequencer/Private/Sequencer.cpp
@@ -5609,9 +5609,11 @@ void FSequencer::UpdatePreviewLevelViewportClientFromCameraCut(FLevelEditorViewp
 	AActor* CameraActor = Cast<AActor>(InCameraObject);
 	AActor* PreviousCameraActor = Cast<AActor>(CameraCutParams.PreviousCameraObject);
 
-	const bool bIsBlending = (CameraCutParams.BlendTime > 0.f &&
-			(CameraActor != nullptr || PreviousCameraActor != nullptr));
 	const float BlendFactor = FMath::Clamp(CameraCutParams.PreviewBlendFactor, 0.f, 1.f);
+	const bool bIsBlending = (
+			CameraCutParams.BlendTime > 0.f &&
+			BlendFactor < 1.f - SMALL_NUMBER &&
+			(CameraActor != nullptr || PreviousCameraActor != nullptr));
 
 	// To preview blending we'll have to offset the viewport camera using the view modifiers API.
 	bApplyViewModifier = bIsBlending;
@@ -7070,7 +7072,7 @@ void FSequencer::AssignActor(FMenuBuilder& MenuBuilder, FGuid InObjectBinding)
 		InitOptions.bFocusSearchBoxWhenOpened = true;
 		// Only want the actor label column
 		InitOptions.ColumnMap.Add(FSceneOutlinerBuiltInColumnTypes::Label(), FSceneOutlinerColumnInfo(ESceneOutlinerColumnVisibility::Visible, 0));
-
+		
 		// Only display actors that are not possessed already
 		InitOptions.Filters->AddFilterPredicate<FActorTreeItem>(FActorTreeItem::FFilterPredicate::CreateLambda( IsActorValidForAssignment ) );
 	}
