Sequencer: Pass in actor component to GetComponents()->IsSelected since it's now more strict

diff --git a/Engine/Source/Editor/Sequencer/Private/Sequencer.cpp b/Engine/Source/Editor/Sequencer/Private/Sequencer.cpp
--- a/Engine/Source/Editor/Sequencer/Private/Sequencer.cpp
+++ b/Engine/Source/Editor/Sequencer/Private/Sequencer.cpp
@@ -6525,7 +6525,8 @@ void FSequencer::SynchronizeSequencerSelectionWithExternalSelection()
 			}
 
 			const bool bActorSelected = ActorSelection->IsSelected( RuntimeObject );
-			const bool bComponentSelected = GEditor->GetSelectedComponents()->IsSelected( RuntimeObject);
+			UActorComponent* ActorComponent = Cast<UActorComponent>(RuntimeObject);
+			const bool bComponentSelected = ActorComponent ? GEditor->GetSelectedComponents()->IsSelected(ActorComponent) : false;
 
 			if (bActorSelected || bComponentSelected)
 			{
