Sequencer: Pass in actor to ActorSelection->IsSelected since it's now more strict

diff --git a/Engine/Source/Editor/Sequencer/Private/Sequencer.cpp b/Engine/Source/Editor/Sequencer/Private/Sequencer.cpp
--- a/Engine/Source/Editor/Sequencer/Private/Sequencer.cpp
+++ b/Engine/Source/Editor/Sequencer/Private/Sequencer.cpp
@@ -6524,7 +6524,9 @@ void FSequencer::SynchronizeSequencerSelectionWithExternalSelection()
 				}
 			}
 
-			const bool bActorSelected = ActorSelection->IsSelected( RuntimeObject );
+			AActor* Actor = Cast<AActor>(RuntimeObject);
+			const bool bActorSelected = Actor ? ActorSelection->IsSelected(Actor) : false;
+						
 			UActorComponent* ActorComponent = Cast<UActorComponent>(RuntimeObject);
 			const bool bComponentSelected = ActorComponent ? GEditor->GetSelectedComponents()->IsSelected(ActorComponent) : false;
 
