Sequencer: New drag drop behavior for composite drag drop. Fix regresion in UE5 from dropping an actor from the scene outliner.

diff --git a/Engine/Source/Editor/Sequencer/Private/SSequencer.cpp b/Engine/Source/Editor/Sequencer/Private/SSequencer.cpp
--- a/Engine/Source/Editor/Sequencer/Private/SSequencer.cpp
+++ b/Engine/Source/Editor/Sequencer/Private/SSequencer.cpp
@@ -49,8 +49,11 @@
 #include "SSequencerTrackArea.h"
 #include "SSequencerTrackOutliner.h"
 #include "DragAndDrop/AssetDragDropOp.h"
+#include "DragAndDrop/ActorDragDropOp.h"
 #include "DragAndDrop/ActorDragDropGraphEdOp.h"
 #include "DragAndDrop/ClassDragDropOp.h"
+#include "DragAndDrop/CompositeDragDropOp.h"
+#include "DragAndDrop/FolderDragDropOp.h"
 #include "Widgets/Input/SSearchBox.h"
 #include "SSequencerTreeView.h"
 #include "MovieSceneTrackEditor.h"
@@ -95,6 +98,7 @@
 #include "SequencerTrackFilters.h"
 #include "SequencerTrackFilterExtension.h"
 #include "SequencerCustomizationManager.h"
+#include "EngineUtils.h"
 
 #define LOCTEXT_NAMESPACE "Sequencer"
 
@@ -338,6 +342,7 @@ void SSequencer::Construct(const FArguments& InArgs, TSharedRef<FSequencer> InSe
 	RootCustomization.OnAssetsDrop = InArgs._OnAssetsDrop;
 	RootCustomization.OnClassesDrop = InArgs._OnClassesDrop;
 	RootCustomization.OnActorsDrop = InArgs._OnActorsDrop;
+	RootCustomization.OnCompositeDrop = InArgs._OnCompositeDrop;
 
 	// Get the desired display format from the user's settings each time.
 	TAttribute<EFrameNumberDisplayFormats> GetDisplayFormatAttr = MakeAttributeLambda(
@@ -2914,7 +2919,8 @@ FReply SSequencer::OnDragOver( const FGeometry& MyGeometry, const FDragDropEvent
 	if (Operation.IsValid() && (
 		Operation->IsOfType<FAssetDragDropOp>() ||
 		Operation->IsOfType<FClassDragDropOp>() ||
-		Operation->IsOfType<FActorDragDropGraphEdOp>() ) )
+		Operation->IsOfType<FActorDragDropGraphEdOp>() ||
+		Operation->IsOfType<FCompositeDragDropOp>() ) )
 	{
 		bIsDragSupported = true;
 	}
@@ -2975,6 +2981,13 @@ FReply SSequencer::OnDrop( const FGeometry& MyGeometry, const FDragDropEvent& Dr
 			OnActorsDropped( *DragDropOp );
 			bWasDropHandled = true;
 		}
+		else if (Operation->IsOfType<FCompositeDragDropOp>())
+		{
+			const auto& DragDropOp = StaticCastSharedPtr<FCompositeDragDropOp>(Operation);
+
+			OnCompositeDropped(*DragDropOp);
+			bWasDropHandled = true;
+		}
 	}
 
 	return bWasDropHandled ? FReply::Handled() : FReply::Unhandled();
@@ -3186,6 +3199,70 @@ void SSequencer::OnActorsDropped( FActorDragDropGraphEdOp& DragDropOp )
 	}
 }
 
+void SSequencer::OnCompositeDropped( FCompositeDragDropOp& DragDropOp )
+{
+	const FScopedTransaction Transaction(LOCTEXT("DropActors", "Drop Actors"));
+
+	ESequencerDropResult DropResult = ESequencerDropResult::Unhandled;
+
+	TArray<TWeakObjectPtr<AActor>> DraggedActors;
+	FFolderDragDropOp* FolderDrag = DragDropOp.GetSubOp<FFolderDragDropOp>().Get();
+	FActorDragDropOp* ActorDrag = DragDropOp.GetSubOp<FActorDragDropOp>().Get();
+	if (ActorDrag)
+	{
+		DraggedActors.Reserve(ActorDrag->Actors.Num());
+		for (TWeakObjectPtr<AActor> WeakActor : ActorDrag->Actors)
+		{
+			if (AActor* Actor = WeakActor.Get())
+			{
+				DraggedActors.Add(Actor);
+			}
+		}
+	}
+
+	if (FolderDrag)
+	{
+		// Copy the array onto the stack if it's within a reasonable size
+		TArray<FName, TInlineAllocator<16>> DraggedFolders(FolderDrag->Folders);
+
+		// Find any actors in the global editor world that have any of the dragged paths.
+		// WARNING: Actor iteration can be very slow, so this needs to be optimized
+		TSharedPtr<FSequencer> Sequencer = SequencerPtr.Pin();
+		UObject* PlaybackContext = Sequencer->GetPlaybackContext();
+		UWorld* World = PlaybackContext ? PlaybackContext->GetWorld() : nullptr;
+		if (World)
+		{
+			for (FActorIterator ActorIt(World); ActorIt; ++ActorIt)
+			{
+				FName ActorPath = ActorIt->GetFolderPath();
+				if (ActorPath.IsNone() || !DraggedFolders.Contains(ActorPath))
+				{
+					continue;
+				}
+
+				DraggedActors.Add(*ActorIt);
+			}		
+		}
+	}	
+	
+	for (FOnCompositeDrop Delegate : OnCompositeDrop)
+	{
+		if (Delegate.IsBound())
+		{
+			DropResult = Delegate.Execute(DraggedActors, DragDropOp);
+			if (DropResult != ESequencerDropResult::Unhandled)
+			{
+				break;
+			}
+		}
+	}
+
+	if (DropResult == ESequencerDropResult::Unhandled)
+	{
+		SequencerPtr.Pin()->OnActorsDropped(DraggedActors);
+	}
+}
+
 void SSequencer::OnCrumbClicked(const FSequencerBreadcrumb& Item)
 {
 	if (Item.BreadcrumbType != FSequencerBreadcrumb::ShotType)
@@ -3894,6 +3971,7 @@ void SSequencer::ApplySequencerCustomizations(const TArray<FSequencerCustomizati
 	OnAssetsDrop.Reset();
 	OnActorsDrop.Reset();
 	OnClassesDrop.Reset();
+	OnCompositeDrop.Reset();
 
 	ApplySequencerCustomization(RootCustomization);
 	for (const FSequencerCustomizationInfo& Info : Customizations)
@@ -3935,6 +4013,10 @@ void SSequencer::ApplySequencerCustomization(const FSequencerCustomizationInfo&
 	{
 		OnClassesDrop.Add(Customization.OnClassesDrop);
 	}
+	if (Customization.OnCompositeDrop.IsBound())
+	{
+		OnCompositeDrop.Add(Customization.OnCompositeDrop);
+	}
 }
 
 USequencerSettings* SSequencer::GetSequencerSettings() const
diff --git a/Engine/Source/Editor/Sequencer/Private/SSequencer.h b/Engine/Source/Editor/Sequencer/Private/SSequencer.h
--- a/Engine/Source/Editor/Sequencer/Private/SSequencer.h
+++ b/Engine/Source/Editor/Sequencer/Private/SSequencer.h
@@ -21,6 +21,7 @@
 class FActorDragDropGraphEdOp;
 class FAssetDragDropOp;
 class FClassDragDropOp;
+class FCompositeDragDropOp;
 class FMovieSceneClipboard;
 class FSequencerTimeSliderController;
 class FVirtualTrackArea;
@@ -260,6 +261,9 @@ public:
 		/** Called when an actor is dropped on the sequencer. Not called if OnReceivedDrop is bound and returned true. */
 		SLATE_EVENT( FOnActorsDrop, OnActorsDrop )
 
+		/** Called when an actor is dropped on the sequencer as a composite. Not called if OnReceivedDrop is bound and returned true. */
+		SLATE_EVENT( FOnCompositeDrop, OnCompositeDrop )
+
 		/** Extender to use for the add menu. */
 		SLATE_ARGUMENT( TSharedPtr<FExtender>, AddMenuExtender )
 
@@ -510,6 +514,13 @@ private:
 	 */
 	void OnActorsDropped(FActorDragDropGraphEdOp& DragDropOp); 
 
+	/**
+	 * Called when one or more objects are dropped into the widget as a composite
+	 *
+	 * @param	DragDropOp	Information about the objects(s) that was dropped
+	 */
+	void OnCompositeDropped(FCompositeDragDropOp& DragDropOp); 
+
 	/** Called when a breadcrumb is clicked on in the sequencer */
 	void OnCrumbClicked(const FSequencerBreadcrumb& Item);
 
@@ -720,6 +731,9 @@ private:
 	/** Called when an actor is dropped on the sequencer. */
 	TArray<FOnActorsDrop> OnActorsDrop;
 
+	/** Called when an actor is dropped on the sequencer as a composite. */
+	TArray<FOnCompositeDrop> OnCompositeDrop;
+
 	/** Stores the callbacks and extenders provided to the constructor. */
 	FSequencerCustomizationInfo RootCustomization;
 
diff --git a/Engine/Source/Editor/Sequencer/Public/SequencerCustomizationManager.h b/Engine/Source/Editor/Sequencer/Public/SequencerCustomizationManager.h
--- a/Engine/Source/Editor/Sequencer/Public/SequencerCustomizationManager.h
+++ b/Engine/Source/Editor/Sequencer/Public/SequencerCustomizationManager.h
@@ -8,6 +8,7 @@
 class FAssetDragDropOp;
 class FClassDragDropOp;
 class FActorDragDropGraphEdOp;
+class FCompositeDragDropOp;
 class UMovieSceneSequence;
 
 /** A delegate that gets executed when a drag/drop event happens on the sequencer. The return value determines if the event was handled by the bound delegate. */
@@ -29,6 +30,9 @@ DECLARE_DELEGATE_RetVal_TwoParams(ESequencerDropResult, FOnClassesDrop, const TA
 /** A delegate that gets executed when an actor is dropped on the sequencer. The return value determines if the operation was handled by the bound delegate. */
 DECLARE_DELEGATE_RetVal_TwoParams(ESequencerDropResult, FOnActorsDrop, const TArray<TWeakObjectPtr<AActor>>&, const FActorDragDropGraphEdOp&);
 
+/** A delegate that gets executed when an actor is dropped on the sequencer as a composite. The return value determines if the operation was handled by the bound delegate. */
+DECLARE_DELEGATE_RetVal_TwoParams(ESequencerDropResult, FOnCompositeDrop, const TArray<TWeakObjectPtr<AActor>>&, const FCompositeDragDropOp&);
+
 /** Class for specifying customizations to apply to a sequence editor. */
 struct FSequencerCustomizationInfo
 {
@@ -52,6 +56,9 @@ struct FSequencerCustomizationInfo
 
 	/** Called when an actor is dropped on the sequencer. Only if OnReceivedDrop hasn't handled the event. */
 	FOnActorsDrop OnActorsDrop;
+
+	/** Called when an actor is dropped on the sequencer as a composite. Only if OnReceivedDrop hasn't handled the event. */
+	FOnCompositeDrop OnCompositeDrop;
 };
 
 /** Class to pass to ISequencerCustomization for building a customization. */
