Sequencer: Make the status info text hit test invisible so that the last row can be clicked on

diff --git a/Engine/Source/Editor/Sequencer/Private/SSequencerTreeFilterStatusBar.cpp b/Engine/Source/Editor/Sequencer/Private/SSequencerTreeFilterStatusBar.cpp
--- a/Engine/Source/Editor/Sequencer/Private/SSequencerTreeFilterStatusBar.cpp
+++ b/Engine/Source/Editor/Sequencer/Private/SSequencerTreeFilterStatusBar.cpp
@@ -30,6 +30,7 @@ void SSequencerTreeFilterStatusBar::Construct(const FArguments& InArgs, TSharedP
 			SNew(SBorder)
 			.BorderImage(FEditorStyle::GetBrush("ToolPanel.GroupBorder"))
 			.Padding(FMargin(3.0f, 3.0f))
+			.Visibility(EVisibility::HitTestInvisible)
 			[
 				SAssignNew(TextBlock, STextBlock)
 			]
