Sequencer: Add a spacer node at the bottom of the tree and ensure that it's not selectable and doesn't highlight on hover. This allows for scrolling past the last node so that the last node is visible above the info text.

diff --git a/Engine/Source/Editor/Sequencer/Private/SSequencerTreeView.cpp b/Engine/Source/Editor/Sequencer/Private/SSequencerTreeView.cpp
--- a/Engine/Source/Editor/Sequencer/Private/SSequencerTreeView.cpp
+++ b/Engine/Source/Editor/Sequencer/Private/SSequencerTreeView.cpp
@@ -138,6 +138,19 @@ void SSequencerTreeViewRow::Tick(const FGeometry& AllottedGeometry, const double
 	StaticCastSharedPtr<SSequencerTreeView>(OwnerTablePtr.Pin())->ReportChildRowGeometry(Node.Pin().ToSharedRef(), AllottedGeometry);
 }
 
+const FSlateBrush* SSequencerTreeViewRow::GetBorder() const 
+{
+	TSharedPtr<FSequencerDisplayNode> PinnedNode = Node.Pin();
+	bool bIsSelectable = PinnedNode->IsSelectable();
+
+	if (!bIsSelectable)
+	{
+		return nullptr;
+	}
+
+	return SMultiColumnTableRow::GetBorder();
+}
+
 void SSequencerTreeView::Construct(const FArguments& InArgs, const TSharedRef<FSequencerNodeTree>& InNodeTree, const TSharedRef<SSequencerTrackArea>& InTrackArea)
 {
 	SequencerNodeTree = InNodeTree;
@@ -714,6 +727,11 @@ void SSequencerTreeView::SynchronizeTreeSelectionWithSequencerSelection()
 
 void SSequencerTreeView::Private_SetItemSelection( FDisplayNodeRef TheItem, bool bShouldBeSelected, bool bWasUserDirected )
 {
+	if (!TheItem->IsSelectable())
+	{
+		return;
+	}
+
 	STreeView::Private_SetItemSelection( TheItem, bShouldBeSelected, bWasUserDirected );
 	if ( bUpdatingTreeSelection == false )
 	{
diff --git a/Engine/Source/Editor/Sequencer/Private/SSequencerTreeView.h b/Engine/Source/Editor/Sequencer/Private/SSequencerTreeView.h
--- a/Engine/Source/Editor/Sequencer/Private/SSequencerTreeView.h
+++ b/Engine/Source/Editor/Sequencer/Private/SSequencerTreeView.h
@@ -284,6 +284,7 @@ public:
 	void AddTrackAreaReference(const TSharedPtr<SSequencerTrackLane>& Lane);
 
 	virtual void Tick(const FGeometry& AllottedGeometry, const double InCurrentTime, const float InDeltaTime) override;
+	virtual const FSlateBrush* GetBorder() const override;
 
 	/** Overridden from SMultiColumnTableRow.  Generates a widget for this column of the tree row. */
 	virtual TSharedRef<SWidget> GenerateWidgetForColumn(const FName& ColumnId) override;
diff --git a/Engine/Source/Editor/Sequencer/Private/SequencerNodeTree.cpp b/Engine/Source/Editor/Sequencer/Private/SequencerNodeTree.cpp
--- a/Engine/Source/Editor/Sequencer/Private/SequencerNodeTree.cpp
+++ b/Engine/Source/Editor/Sequencer/Private/SequencerNodeTree.cpp
@@ -26,6 +26,47 @@
 #include "SequencerLog.h"
 #include "SequencerCommonHelpers.h"
 
+
+/**
+ * A spacer node that is unselectable
+ */
+class FSequencerSpacerNode : public FSequencerDisplayNode
+{
+public:
+
+	/**
+	 * Create and initialize a new instance.
+	 * 
+	 * @param InParentTree The tree this node is in.
+	 */
+	FSequencerSpacerNode(float InSize, int32 InSortingOrder, FSequencerNodeTree& InParentTree)
+		: FSequencerDisplayNode(NAME_None, InParentTree)
+		, Size(InSize)
+		, SortingOrder(InSortingOrder)
+	{ }
+
+public:
+
+	// FSequencerDisplayNode interface
+	virtual bool CanRenameNode() const override { return false; }
+	virtual int32 GetSortingOrder() const override { return SortingOrder; }
+	virtual FText GetDisplayName() const override { return FText(); }
+	virtual float GetNodeHeight() const override { return Size; }
+	virtual FNodePadding GetNodePadding() const override { return FNodePadding(0.f); }
+	virtual ESequencerNode::Type GetType() const override { return ESequencerNode::Object; }
+	virtual void SetDisplayName(const FText& NewDisplayName) override { }
+	virtual TSharedRef<SWidget> GenerateContainerWidgetForOutliner(const TSharedRef<SSequencerTreeViewRow>& InRow) override { return SNew(SBox).HeightOverride(Size); }
+	virtual bool IsSelectable() const override { return false; }
+
+private:
+
+	/** The size of the spacer */
+	float Size;
+
+	/** The sorting order for the spacer */
+	int32 SortingOrder;
+};
+
 FSequencerNodeTree::~FSequencerNodeTree()
 {
 	if (TrackFilters.IsValid())
@@ -41,6 +82,7 @@ FSequencerNodeTree::~FSequencerNodeTree()
 
 FSequencerNodeTree::FSequencerNodeTree(FSequencer& InSequencer)
 	: RootNode(MakeShared<FSequencerRootNode>(*this))
+	, BottomSpacerNode(MakeShared<FSequencerSpacerNode>(30.f, INT_MAX, *this))
 	, SerialNumber(0)
 	, Sequencer(InSequencer)
 	, bFilterUpdateRequested(false)
@@ -215,6 +257,8 @@ void FSequencerNodeTree::RefreshNodes(UMovieScene* MovieScene)
 	// @todo sequencer: Newly added sections may need to be visible even when there is a filter
 	bFilterUpdateRequested = true;
 	UpdateFilters();
+
+	BottomSpacerNode->SetParent(RootNode);
 }
 
 TSharedPtr<FSequencerTrackNode> FSequencerNodeTree::CreateOrUpdateTrack(UMovieSceneTrack* Track, ETrackType TrackType)
diff --git a/Engine/Source/Editor/Sequencer/Private/SequencerNodeTree.h b/Engine/Source/Editor/Sequencer/Private/SequencerNodeTree.h
--- a/Engine/Source/Editor/Sequencer/Private/SequencerNodeTree.h
+++ b/Engine/Source/Editor/Sequencer/Private/SequencerNodeTree.h
@@ -312,6 +312,8 @@ private:
 	/** Symbolic root node that contains the actual displayed root nodes as children */
 	TSharedRef<FSequencerDisplayNode> RootNode;
 
+	TSharedRef<FSequencerDisplayNode> BottomSpacerNode;
+
 	/** A serially incrementing integer that is increased each time the tree is refreshed to track node relevance */
 	uint32 SerialNumber;
 	/** Map from FMovieSceneBinding::GetObjectGuid to display node */
