Sequencer: Add spacer node after filtering

diff --git a/Engine/Source/Editor/Sequencer/Private/SequencerNodeTree.cpp b/Engine/Source/Editor/Sequencer/Private/SequencerNodeTree.cpp
--- a/Engine/Source/Editor/Sequencer/Private/SequencerNodeTree.cpp
+++ b/Engine/Source/Editor/Sequencer/Private/SequencerNodeTree.cpp
@@ -258,7 +258,9 @@ void FSequencerNodeTree::RefreshNodes(UMovieScene* MovieScene)
 	bFilterUpdateRequested = true;
 	UpdateFilters();
 
+	// Always add the bottom spacer node
 	BottomSpacerNode->SetParent(RootNode);
+	FilteredNodes.Add(BottomSpacerNode);
 }
 
 TSharedPtr<FSequencerTrackNode> FSequencerNodeTree::CreateOrUpdateTrack(UMovieSceneTrack* Track, ETrackType TrackType)
