Sequencer: Partially revert 14807766 because it was causing menus to be defocused.

diff --git a/Engine/Source/Editor/Sequencer/Private/SequencerUtilities.cpp b/Engine/Source/Editor/Sequencer/Private/SequencerUtilities.cpp
--- a/Engine/Source/Editor/Sequencer/Private/SequencerUtilities.cpp
+++ b/Engine/Source/Editor/Sequencer/Private/SequencerUtilities.cpp
@@ -52,7 +52,7 @@ TSharedRef<SWidget> FSequencerUtilities::MakeAddButton(FText HoverText, FOnGetCo
 
 		SNew(SComboButton)
 		.HasDownArrow(false)
-		.IsFocusable(false) // Intentionally false so that it's easier to tab to the next numeric input
+		.IsFocusable(true)
 		.ButtonStyle(FEditorStyle::Get(), "HoverHintOnly")
 		.ForegroundColor( FSlateColor::UseForeground() )
 		.IsEnabled_Lambda([=]() { return InSequencer.IsValid() ? !InSequencer.Pin()->IsReadOnly() : false; })
