Fix sequencer toolbar so that it has consistent sizing and styling with the rest of the editor toolbars

diff --git a/Engine/Source/Editor/EditorStyle/Private/StarshipStyle.cpp b/Engine/Source/Editor/EditorStyle/Private/StarshipStyle.cpp
--- a/Engine/Source/Editor/EditorStyle/Private/StarshipStyle.cpp
+++ b/Engine/Source/Editor/EditorStyle/Private/StarshipStyle.cpp
@@ -2322,17 +2322,6 @@ void FStarshipEditorStyle::FStyle::SetupSequencerStyles()
 		Set("EMovieSceneBlendType::Absolute", new IMAGE_BRUSH("Sequencer/EMovieSceneBlendType_Absolute", FVector2D(32, 16)));
 		Set("EMovieSceneBlendType::Relative", new IMAGE_BRUSH("Sequencer/EMovieSceneBlendType_Relative", FVector2D(32, 16)));
 		Set("EMovieSceneBlendType::Additive", new IMAGE_BRUSH("Sequencer/EMovieSceneBlendType_Additive", FVector2D(32, 16)));
-
-		FToolBarStyle SequencerToolbarStyle = FToolBarStyle(FStarshipCoreStyle::GetCoreStyle().GetWidgetStyle<FToolBarStyle>("ToolBar"));
-
-		// Sequencer & Curve Editor Toolbar Style ( Grabs core default, copies brushes, changes separator and block spacing )
-		SequencerToolbarStyle.SetBackground(FSlateColorBrush(FStyleColors::Background));
-		SequencerToolbarStyle.SetBlockPadding(FMargin(0.0f, 4.0f));
-
-		SequencerToolbarStyle.SetSeparatorBrush(FSlateColorBrush(FLinearColor(FColor(48, 48, 48))));
-		SequencerToolbarStyle.SetSeparatorPadding(FMargin( 1.f, 1.f ));
-
-		Set("Sequencer.ToolBar", SequencerToolbarStyle);
 	}
 
 
diff --git a/Engine/Source/Editor/Sequencer/Private/SSequencer.cpp b/Engine/Source/Editor/Sequencer/Private/SSequencer.cpp
--- a/Engine/Source/Editor/Sequencer/Private/SSequencer.cpp
+++ b/Engine/Source/Editor/Sequencer/Private/SSequencer.cpp
@@ -139,8 +139,7 @@ public:
 
 	TSharedRef<SWidget> MakeToolbar(TSharedRef<SCurveEditorPanel> InEditorPanel)
 	{
-		FToolBarBuilder ToolBarBuilder(InEditorPanel->GetCommands(), FMultiBoxCustomization::None, InEditorPanel->GetToolbarExtender(), true);
-		ToolBarBuilder.SetStyle(&FEditorStyle::Get(), "Sequencer.ToolBar");
+		FSlimHorizontalToolBarBuilder ToolBarBuilder(InEditorPanel->GetCommands(), FMultiBoxCustomization::None, InEditorPanel->GetToolbarExtender(), true);
 		ToolBarBuilder.BeginSection("Asset");
 
 		{
@@ -580,14 +579,12 @@ void SSequencer::Construct(const FArguments& InArgs, TSharedRef<FSequencer> InSe
 					+ SGridPanel::Slot( Column0, Row0, SGridPanel::Layer(10) )
 					.ColumnSpan(2)
 					[
-						SNew(SBorder)
-						.BorderImage(FEditorStyle::GetBrush("ToolPanel.GroupBorder"))
-						.Padding(FMargin(CommonPadding, 0.f))
+						SNew(SBox)
+						.Padding(FMargin(CommonPadding,0.0f,0.0f,0.f))
 						[
 							SNew(SWrapBox)
 							.UseAllottedSize(true)
-							.InnerSlotPadding(FVector2D(5, 2))
-
+							.InnerSlotPadding(FVector2D(5, 0))
 							+ SWrapBox::Slot()
 							.FillEmptySpace(true)
 							.FillLineWhenSizeLessThan(600)
@@ -598,143 +595,137 @@ void SSequencer::Construct(const FArguments& InArgs, TSharedRef<FSequencer> InSe
 							+ SWrapBox::Slot()
 							.FillEmptySpace(true)
 							[
-								SNew(SBorder)
-								.Padding(FMargin(3.0f, 0.0f))
-								.BorderImage(FEditorStyle::GetBrush("ToolPanel.GroupBorder"))
+								SNew(SHorizontalBox)
+								// Right Aligned Breadcrumbs
+								+ SHorizontalBox::Slot()
+								.HAlign(HAlign_Right)
+								.VAlign(VAlign_Center)
 								[
-									SNew(SHorizontalBox)
+									SNew(SSpacer)
+								]
 
-									// Right Aligned Breadcrumbs
-									+ SHorizontalBox::Slot()
-									.HAlign(HAlign_Right)
-									.VAlign(VAlign_Center)
-									[
-										SNew(SSpacer)
-									]
+								// History Back Button
+								+ SHorizontalBox::Slot()
+								.HAlign(HAlign_Right)
+								.VAlign(VAlign_Center)
+								.AutoWidth()
+								[
+									SNew(SVerticalBox)
 
-									// History Back Button
-									+ SHorizontalBox::Slot()
-									.HAlign(HAlign_Right)
-									.VAlign(VAlign_Center)
-									.AutoWidth()
+									+ SVerticalBox::Slot()
+									.FillHeight(1.0f)
 									[
-										SNew(SVerticalBox)
-
-										+ SVerticalBox::Slot()
-										.FillHeight(1.0f)
+										SNew(SButton)
+										.Visibility(this, &SSequencer::GetBreadcrumbTrailVisibility)
+										.VAlign(EVerticalAlignment::VAlign_Center)
+										.ButtonStyle(FEditorStyle::Get(), "FlatButton")
+										.ForegroundColor(FLinearColor::White)
+										.ToolTipText_Lambda([this] { return SequencerPtr.Pin()->GetNavigateBackwardTooltip(); })
+										.ContentPadding(FMargin(1, 0))
+										.OnClicked_Lambda([this] { return SequencerPtr.Pin()->NavigateBackward(); })
+										.IsEnabled_Lambda([this] { return SequencerPtr.Pin()->CanNavigateBackward(); })
 										[
-											SNew(SButton)
-											.Visibility(this, &SSequencer::GetBreadcrumbTrailVisibility)
-											.VAlign(EVerticalAlignment::VAlign_Center)
-											.ButtonStyle(FEditorStyle::Get(), "FlatButton")
-											.ForegroundColor(FLinearColor::White)
-											.ToolTipText_Lambda([this] { return SequencerPtr.Pin()->GetNavigateBackwardTooltip(); })
-											.ContentPadding(FMargin(1, 0))
-											.OnClicked_Lambda([this] { return SequencerPtr.Pin()->NavigateBackward(); })
-											.IsEnabled_Lambda([this] { return SequencerPtr.Pin()->CanNavigateBackward(); })
-											[
-												SNew(STextBlock)
-												.TextStyle(FEditorStyle::Get(), "ContentBrowser.TopBar.Font")
-												.Font(FEditorStyle::Get().GetFontStyle("FontAwesome.11"))
-												.Text(FText::FromString(FString(TEXT("\xf060"))) /*fa-arrow-left*/)
-											]
+											SNew(STextBlock)
+											.TextStyle(FEditorStyle::Get(), "ContentBrowser.TopBar.Font")
+											.Font(FEditorStyle::Get().GetFontStyle("FontAwesome.11"))
+											.Text(FText::FromString(FString(TEXT("\xf060"))) /*fa-arrow-left*/)
 										]
 									]
+								]
 
-									// History Forward Button
-									+ SHorizontalBox::Slot()
-									.HAlign(HAlign_Right)
-									.VAlign(VAlign_Center)
-									.AutoWidth()
-									[
-										SNew(SVerticalBox)
-
-										+ SVerticalBox::Slot()
-										.FillHeight(1.0f)
-										[
-											SNew(SButton)
-											.Visibility(this, &SSequencer::GetBreadcrumbTrailVisibility)
-											.VAlign(EVerticalAlignment::VAlign_Center)
-											.ButtonStyle(FEditorStyle::Get(), "FlatButton")
-											.ForegroundColor(FLinearColor::White)
-											.ToolTipText_Lambda([this] { return SequencerPtr.Pin()->GetNavigateForwardTooltip(); })
-											.ContentPadding(FMargin(1, 0))
-											.OnClicked_Lambda([this] { return SequencerPtr.Pin()->NavigateForward(); })
-											.IsEnabled_Lambda([this] { return SequencerPtr.Pin()->CanNavigateForward(); })
-											[
-												SNew(STextBlock)
-												.TextStyle(FEditorStyle::Get(), "ContentBrowser.TopBar.Font")
-												.Font(FEditorStyle::Get().GetFontStyle("FontAwesome.11"))
-												.Text(FText::FromString(FString(TEXT("\xf061"))) /*fa-arrow-right*/)
-											]
-										]
-									]
+								// History Forward Button
+								+ SHorizontalBox::Slot()
+								.HAlign(HAlign_Right)
+								.VAlign(VAlign_Center)
+								.AutoWidth()
+								[
+									SNew(SVerticalBox)
 
-									// Separator
-									+ SHorizontalBox::Slot()
-									.AutoWidth()
-									.Padding(3, 0)
-									[
-										SNew(SSeparator)
-										.Visibility(this, &SSequencer::GetBreadcrumbTrailVisibility)
-										.Orientation(Orient_Vertical)
-									]
-							
-									+ SHorizontalBox::Slot()
-									.HAlign(HAlign_Right)
-									.VAlign(VAlign_Center)
-									.AutoWidth()
+									+ SVerticalBox::Slot()
+									.FillHeight(1.0f)
 									[
-										SAssignNew(BreadcrumbPickerButton, SComboButton)
+										SNew(SButton)
 										.Visibility(this, &SSequencer::GetBreadcrumbTrailVisibility)
+										.VAlign(EVerticalAlignment::VAlign_Center)
 										.ButtonStyle(FEditorStyle::Get(), "FlatButton")
 										.ForegroundColor(FLinearColor::White)
-										.OnGetMenuContent_Lambda([this] { return SNew(SSequencerHierarchyBrowser, SequencerPtr); })
-										.HasDownArrow(false)
-										.ContentPadding(FMargin(3, 3))
-										.ButtonContent()
+										.ToolTipText_Lambda([this] { return SequencerPtr.Pin()->GetNavigateForwardTooltip(); })
+										.ContentPadding(FMargin(1, 0))
+										.OnClicked_Lambda([this] { return SequencerPtr.Pin()->NavigateForward(); })
+										.IsEnabled_Lambda([this] { return SequencerPtr.Pin()->CanNavigateForward(); })
 										[
 											SNew(STextBlock)
-											.TextStyle(FEditorStyle::Get(), "Sequencer.BreadcrumbText")
+											.TextStyle(FEditorStyle::Get(), "ContentBrowser.TopBar.Font")
 											.Font(FEditorStyle::Get().GetFontStyle("FontAwesome.11"))
-											.Text(FText::FromString(FString(TEXT("\xf07c"))) /*fa-folder-open*/)
+											.Text(FText::FromString(FString(TEXT("\xf061"))) /*fa-arrow-right*/)
 										]
 									]
+								]
 
-									// Right Aligned Breadcrumbs
-									+ SHorizontalBox::Slot()
-									.HAlign(HAlign_Right)
-									.VAlign(VAlign_Center)
-									.AutoWidth()
+								// Separator
+								+ SHorizontalBox::Slot()
+								.AutoWidth()
+								.Padding(3, 0)
+								[
+									SNew(SSeparator)
+									.Visibility(this, &SSequencer::GetBreadcrumbTrailVisibility)
+									.Orientation(Orient_Vertical)
+								]
+							
+								+ SHorizontalBox::Slot()
+								.HAlign(HAlign_Right)
+								.VAlign(VAlign_Center)
+								.AutoWidth()
+								[
+									SAssignNew(BreadcrumbPickerButton, SComboButton)
+									.Visibility(this, &SSequencer::GetBreadcrumbTrailVisibility)
+									.ButtonStyle(FEditorStyle::Get(), "FlatButton")
+									.ForegroundColor(FLinearColor::White)
+									.OnGetMenuContent_Lambda([this] { return SNew(SSequencerHierarchyBrowser, SequencerPtr); })
+									.HasDownArrow(false)
+									.ContentPadding(FMargin(3, 3))
+									.ButtonContent()
 									[
-										SAssignNew(BreadcrumbTrail, SBreadcrumbTrail<FSequencerBreadcrumb>)
-										.Visibility(this, &SSequencer::GetBreadcrumbTrailVisibility)
-										.OnCrumbClicked(this, &SSequencer::OnCrumbClicked)
-										.ButtonStyle(FEditorStyle::Get(), "FlatButton")
-										.DelimiterImage(FEditorStyle::GetBrush("Sequencer.BreadcrumbIcon"))
+										SNew(STextBlock)
 										.TextStyle(FEditorStyle::Get(), "Sequencer.BreadcrumbText")
+										.Font(FEditorStyle::Get().GetFontStyle("FontAwesome.11"))
+										.Text(FText::FromString(FString(TEXT("\xf07c"))) /*fa-folder-open*/)
 									]
+								]
 
-									// Sequence Locking symbol
-									+SHorizontalBox::Slot()
-									.HAlign(HAlign_Right)
-									.VAlign(VAlign_Center)
-									.AutoWidth()
-									.Padding(2, 0, 0, 0)
-									[
-										SNew(SCheckBox)
-										.IsFocusable(false)		
-										.IsChecked_Lambda([this] { return GetIsSequenceReadOnly() ? ECheckBoxState::Checked : ECheckBoxState::Unchecked; } )
-										.OnCheckStateChanged(this, &SSequencer::OnSetSequenceReadOnly)
-										.ToolTipText_Lambda([this] { return GetIsSequenceReadOnly() ? LOCTEXT("UnlockSequence", "Unlock the animation so that it is editable") : LOCTEXT("LockSequence", "Lock the animation so that it is not editable"); } )
-										.ForegroundColor(FLinearColor::White)
-										.CheckedImage(FEditorStyle::GetBrush("Sequencer.LockSequence"))
-										.CheckedHoveredImage(FEditorStyle::GetBrush("Sequencer.LockSequence"))
-										.CheckedPressedImage(FEditorStyle::GetBrush("Sequencer.LockSequence"))
-										.UncheckedImage(FEditorStyle::GetBrush("Sequencer.UnlockSequence"))
-										.UncheckedHoveredImage(FEditorStyle::GetBrush("Sequencer.UnlockSequence"))
-										.UncheckedPressedImage(FEditorStyle::GetBrush("Sequencer.UnlockSequence"))
-									]
+								// Right Aligned Breadcrumbs
+								+ SHorizontalBox::Slot()
+								.HAlign(HAlign_Right)
+								.VAlign(VAlign_Center)
+								.AutoWidth()
+								[
+									SAssignNew(BreadcrumbTrail, SBreadcrumbTrail<FSequencerBreadcrumb>)
+									.Visibility(this, &SSequencer::GetBreadcrumbTrailVisibility)
+									.OnCrumbClicked(this, &SSequencer::OnCrumbClicked)
+									.ButtonStyle(FEditorStyle::Get(), "FlatButton")
+									.DelimiterImage(FEditorStyle::GetBrush("Sequencer.BreadcrumbIcon"))
+									.TextStyle(FEditorStyle::Get(), "Sequencer.BreadcrumbText")
+								]
+
+								// Sequence Locking symbol
+								+SHorizontalBox::Slot()
+								.HAlign(HAlign_Right)
+								.VAlign(VAlign_Center)
+								.AutoWidth()
+								.Padding(2, 0, 0, 0)
+								[
+									SNew(SCheckBox)
+									.IsFocusable(false)		
+									.IsChecked_Lambda([this] { return GetIsSequenceReadOnly() ? ECheckBoxState::Checked : ECheckBoxState::Unchecked; } )
+									.OnCheckStateChanged(this, &SSequencer::OnSetSequenceReadOnly)
+									.ToolTipText_Lambda([this] { return GetIsSequenceReadOnly() ? LOCTEXT("UnlockSequence", "Unlock the animation so that it is editable") : LOCTEXT("LockSequence", "Lock the animation so that it is not editable"); } )
+									.ForegroundColor(FLinearColor::White)
+									.CheckedImage(FEditorStyle::GetBrush("Sequencer.LockSequence"))
+									.CheckedHoveredImage(FEditorStyle::GetBrush("Sequencer.LockSequence"))
+									.CheckedPressedImage(FEditorStyle::GetBrush("Sequencer.LockSequence"))
+									.UncheckedImage(FEditorStyle::GetBrush("Sequencer.UnlockSequence"))
+									.UncheckedHoveredImage(FEditorStyle::GetBrush("Sequencer.UnlockSequence"))
+									.UncheckedPressedImage(FEditorStyle::GetBrush("Sequencer.UnlockSequence"))
 								]
 							]
 						]
@@ -1411,9 +1402,8 @@ TSharedRef<SWidget> SSequencer::MakeToolBar()
 
 	TSharedPtr<FExtender> Extender = FExtender::Combine(AllExtenders);
 
-	FToolBarBuilder ToolBarBuilder( SequencerPtr.Pin()->GetCommandBindings(), FMultiBoxCustomization::None, Extender, true);
-	ToolBarBuilder.SetStyle(&FEditorStyle::Get(), "Sequencer.ToolBar");
-
+	FSlimHorizontalToolBarBuilder ToolBarBuilder( SequencerPtr.Pin()->GetCommandBindings(), FMultiBoxCustomization::None, Extender, true);
+	
 	ToolBarBuilder.BeginSection("Base Commands");
 	{
 		// General 
diff --git a/Engine/Source/Editor/Sequencer/Private/SSequencerPlayRateCombo.cpp b/Engine/Source/Editor/Sequencer/Private/SSequencerPlayRateCombo.cpp
--- a/Engine/Source/Editor/Sequencer/Private/SSequencerPlayRateCombo.cpp
+++ b/Engine/Source/Editor/Sequencer/Private/SSequencerPlayRateCombo.cpp
@@ -50,7 +50,7 @@ void SSequencerPlayRateCombo::Construct(const FArguments& InArgs, TWeakPtr<FSequ
 
 			+ SHorizontalBox::Slot()
 			.HAlign(HAlign_Center)
-			.VAlign(VAlign_Fill)
+			.VAlign(VAlign_Center)
 			.AutoWidth()
 			[
 				SNew(STextBlock)
diff --git a/Engine/Source/Runtime/SlateCore/Private/Styling/StarshipCoreStyle.cpp b/Engine/Source/Runtime/SlateCore/Private/Styling/StarshipCoreStyle.cpp
--- a/Engine/Source/Runtime/SlateCore/Private/Styling/StarshipCoreStyle.cpp
+++ b/Engine/Source/Runtime/SlateCore/Private/Styling/StarshipCoreStyle.cpp
@@ -1802,7 +1802,7 @@ void FStarshipCoreStyle::SetupMultiboxStyles(TSharedRef<FStyle>& Style)
 			.SetBackgroundPadding(FMargin(0,4,0,4))
 			.SetExpandBrush(IMAGE_BRUSH("Icons/toolbar_expand_16x", Icon16x16))
 			.SetComboButtonPadding(FMargin(6.0f, 0.0f))
-			.SetButtonPadding(FMargin(4.0f, 0.0f))
+			.SetButtonPadding(FMargin(4.0f, 4.0f))
 			.SetCheckBoxPadding(FMargin(8.0f, 0.0f))
 			.SetSeparatorBrush(FSlateColorBrush(FStyleColors::Recessed))
 			.SetSeparatorPadding(FMargin(8.f, -4.0f))
