Control Rig: Sequencer: Selection Fixes from 4.26

diff --git a/Engine/Plugins/Experimental/ControlRig/Source/ControlRigEditor/Private/Sequencer/ControlRigParameterTrackEditor.cpp b/Engine/Plugins/Experimental/ControlRig/Source/ControlRigEditor/Private/Sequencer/ControlRigParameterTrackEditor.cpp
--- a/Engine/Plugins/Experimental/ControlRig/Source/ControlRigEditor/Private/Sequencer/ControlRigParameterTrackEditor.cpp
+++ b/Engine/Plugins/Experimental/ControlRig/Source/ControlRigEditor/Private/Sequencer/ControlRigParameterTrackEditor.cpp
@@ -1277,7 +1277,8 @@ void FControlRigParameterTrackEditor::OnCurveDisplayChanged(FCurveModel* CurveMo
 			bool bSync = GetSequencer()->GetSequencerSettings()->ShouldSyncCurveEditorSelection();
 			GetSequencer()->SuspendSelectionBroadcast();
 			GetSequencer()->GetSequencerSettings()->SyncCurveEditorSelection(false);
-			GetSequencer()->SelectByChannels(MovieSection, Channels, true, bDisplayed);
+			GetSequencer()->SelectByChannels(MovieSection, Channels, false, bDisplayed);
+			GetSequencer()->RefreshTree();
 			GetSequencer()->GetSequencerSettings()->SyncCurveEditorSelection(bSync);
 			GetSequencer()->ResumeSelectionBroadcast();
 		}
diff --git a/Engine/Source/Editor/Sequencer/Private/Sequencer.cpp b/Engine/Source/Editor/Sequencer/Private/Sequencer.cpp
--- a/Engine/Source/Editor/Sequencer/Private/Sequencer.cpp
+++ b/Engine/Source/Editor/Sequencer/Private/Sequencer.cpp
@@ -2487,14 +2487,20 @@ void FSequencer::NotifyMovieSceneDataChanged( EMovieSceneDataChangeType DataChan
 	OnMovieSceneDataChangedDelegate.Broadcast(DataChangeType);
 }
 
+static bool bRefreshTreeGuard = false;
 void FSequencer::RefreshTree()
 {
-	SequencerWidget->UpdateLayoutTree();
-	bNeedTreeRefresh = false;
-	OnTreeViewChangedDelegate.Broadcast();
+	if (bRefreshTreeGuard == false)
+	{
+		TGuardValue<bool> Guard(bRefreshTreeGuard, true);
 
-	// Force a broadcast of selection changed after the tree view has been updated, in the event that selection was suppressed while the tree was refreshing
-	Selection.Tick();
+		SequencerWidget->UpdateLayoutTree();
+		bNeedTreeRefresh = false;
+		OnTreeViewChangedDelegate.Broadcast();
+
+		// Force a broadcast of selection changed after the tree view has been updated, in the event that selection was suppressed while the tree was refreshing
+		Selection.Tick();
+	}
 }
 
 FAnimatedRange FSequencer::GetViewRange() const
