Level sequences no longer create a new keyframe button for properties in the Actor Details panel every time it is reopened.

diff --git a/Engine/Source/Editor/Sequencer/Private/LevelEditorSequencerIntegration.cpp b/Engine/Source/Editor/Sequencer/Private/LevelEditorSequencerIntegration.cpp
--- a/Engine/Source/Editor/Sequencer/Private/LevelEditorSequencerIntegration.cpp
+++ b/Engine/Source/Editor/Sequencer/Private/LevelEditorSequencerIntegration.cpp
@@ -883,9 +883,10 @@ void FLevelEditorSequencerIntegration::ActivateDetailHandler(const FLevelEditorS
 	}
 
 	FDelegateHandle OnPropertyEditorOpenedHandle = EditModule.OnPropertyEditorOpened().AddRaw(this, &FLevelEditorSequencerIntegration::OnPropertyEditorOpened);
+	FDelegateHandle OnGetGlobalRowExtensionHandle = EditModule.GetGlobalRowExtensionDelegate().AddStatic(&RegisterDetailsExtensionHandler);
 
 	auto DeactivateDetailKeyframeHandler =
-		[this, OnPropertyEditorOpenedHandle]()
+		[this, OnPropertyEditorOpenedHandle, OnGetGlobalRowExtensionHandle]()
 		{
 			FPropertyEditorModule* EditModulePtr = FModuleManager::Get().GetModulePtr<FPropertyEditorModule>("PropertyEditor");
 			if (!EditModulePtr)
@@ -893,6 +894,7 @@ void FLevelEditorSequencerIntegration::ActivateDetailHandler(const FLevelEditorS
 				return;
 			}
 
+			EditModulePtr->GetGlobalRowExtensionDelegate().Remove(OnGetGlobalRowExtensionHandle);
 			EditModulePtr->OnPropertyEditorOpened().Remove(OnPropertyEditorOpenedHandle);
 
 			for (const FName& DetailsTabIdentifier : DetailsTabIdentifiers)
@@ -910,7 +912,7 @@ void FLevelEditorSequencerIntegration::ActivateDetailHandler(const FLevelEditorS
 			}
 		};
 
-	EditModule.GetGlobalRowExtensionDelegate().AddStatic(&RegisterDetailsExtensionHandler);
+	AcquiredResources.Add(DetailHandlerName, DeactivateDetailKeyframeHandler);
 
 	FName DetailHandlerRefreshName("DetailHandlerRefresh");
 	auto RefreshDetailHandler =
@@ -932,7 +934,6 @@ void FLevelEditorSequencerIntegration::ActivateDetailHandler(const FLevelEditorS
 			}
 		};
 
-	AcquiredResources.Add(DetailHandlerName, DeactivateDetailKeyframeHandler);
 
 	if (Options.bForceRefreshDetails)
 	{
