Editgrate of  15302438  from 4.26 release fix. Sequencer/Control Rig Selection getting lost issue.

diff --git a/Engine/Source/Editor/Sequencer/Private/MovieSceneTrackEditor.cpp b/Engine/Source/Editor/Sequencer/Private/MovieSceneTrackEditor.cpp
--- a/Engine/Source/Editor/Sequencer/Private/MovieSceneTrackEditor.cpp
+++ b/Engine/Source/Editor/Sequencer/Private/MovieSceneTrackEditor.cpp
@@ -128,11 +128,6 @@ void FMovieSceneTrackEditor::AnimatablePropertyChanged( FOnKeyProperty OnKeyProp
 
 		UpdatePlaybackRange();
 
-		TSharedPtr<FSequencer> SequencerToUpdate = StaticCastSharedPtr<FSequencer>(Sequencer.Pin());
-		if (SequencerToUpdate.IsValid())
-		{
-			SequencerToUpdate->SynchronizeSequencerSelectionWithExternalSelection();
-		}
 	}
 }
 
