Sequencer: Clear transient flag on pasted section

diff --git a/Engine/Source/Editor/Sequencer/Private/Sequencer.cpp b/Engine/Source/Editor/Sequencer/Private/Sequencer.cpp
--- a/Engine/Source/Editor/Sequencer/Private/Sequencer.cpp
+++ b/Engine/Source/Editor/Sequencer/Private/Sequencer.cpp
@@ -9455,6 +9455,7 @@ bool FSequencer::PasteSections(const FString& TextToImport, TArray<FNotification
 
 			Track->Modify();
 
+			Section->ClearFlags(RF_Transient);
 			Section->Rename(nullptr, Track);
 			Track->AddSection(*Section);
 			if (TrackNode->GetSubTrackMode() == FSequencerTrackNode::ESubTrackMode::SubTrack)
