Sequencer: Control Rig: Undo losing Selection due to Actor Selection change event getting sent always by GEditor on Undo. What we do is if there is no change to Sequencer binding node selections and we are transacting  we don't do anything.

diff --git a/Engine/Source/Editor/Sequencer/Private/Sequencer.cpp b/Engine/Source/Editor/Sequencer/Private/Sequencer.cpp
--- a/Engine/Source/Editor/Sequencer/Private/Sequencer.cpp
+++ b/Engine/Source/Editor/Sequencer/Private/Sequencer.cpp
@@ -6868,8 +6868,9 @@ void FSequencer::SynchronizeSequencerSelectionWithExternalSelection()
 			}
 		}
 	}
+	//Only test if none are selected if we are not transacting, otherwise it will clear out control rig's incorrectly.
 
-	if (!bAllAlreadySelected || (NodesToSelect.Num() == 0 && Selection.GetSelectedOutlinerNodes().Num()))
+	if (!bAllAlreadySelected || (!GIsTransacting && (NodesToSelect.Num() == 0 && Selection.GetSelectedOutlinerNodes().Num())))
 	{
 		Selection.SuspendBroadcast();
 		Selection.EmptySelectedOutlinerNodes();
