Sequencer Update Style Pass for EA

diff --git a/Engine/Plugins/MovieScene/LevelSequenceEditor/Source/LevelSequenceEditor/Private/Misc/LevelSequenceEditorHelpers.cpp b/Engine/Plugins/MovieScene/LevelSequenceEditor/Source/LevelSequenceEditor/Private/Misc/LevelSequenceEditorHelpers.cpp
--- a/Engine/Plugins/MovieScene/LevelSequenceEditor/Source/LevelSequenceEditor/Private/Misc/LevelSequenceEditorHelpers.cpp
+++ b/Engine/Plugins/MovieScene/LevelSequenceEditor/Source/LevelSequenceEditor/Private/Misc/LevelSequenceEditorHelpers.cpp
@@ -79,6 +79,7 @@ class SMasterSequenceSettings : public SCompoundWidget, public FGCObject
 		Details2View = PropertyEditor.CreateDetailView(Details2ViewArgs);
 
 		ChildSlot
+		.Padding(8.f)
 		[
 			SNew(SVerticalBox)
 
@@ -122,7 +123,8 @@ class SMasterSequenceSettings : public SCompoundWidget, public FGCObject
 			.Padding(5.f)
 			[
 				SNew(SButton)
-				.ContentPadding(FMargin(10, 5))
+                .ButtonStyle( &FAppStyle::Get().GetWidgetStyle< FButtonStyle >( "PrimaryButton" ) )
+                .TextStyle( &FAppStyle::Get().GetWidgetStyle< FTextBlockStyle >("DialogButtonText"))
 				.Text(LOCTEXT("CreateMasterSequence", "Create Master Sequence"))
 				.OnClicked(this, &SMasterSequenceSettings::OnCreateMasterSequence)
 			]
diff --git a/Engine/Plugins/MovieScene/LevelSequenceEditor/Source/LevelSequenceEditor/Private/Misc/LevelSequencePlaybackContext.cpp b/Engine/Plugins/MovieScene/LevelSequenceEditor/Source/LevelSequenceEditor/Private/Misc/LevelSequencePlaybackContext.cpp
--- a/Engine/Plugins/MovieScene/LevelSequenceEditor/Source/LevelSequenceEditor/Private/Misc/LevelSequencePlaybackContext.cpp
+++ b/Engine/Plugins/MovieScene/LevelSequenceEditor/Source/LevelSequenceEditor/Private/Misc/LevelSequencePlaybackContext.cpp
@@ -323,17 +323,18 @@ void SLevelSequenceContextPicker::Construct(const FArguments& InArgs)
 	[
 		SNew(SBorder)
 		.BorderImage(this, &SLevelSequenceContextPicker::GetBorderBrush)
-		.Padding(0.0f)
+		.Padding(FMargin(4.f, 0.f))
 		[
 			SNew(SComboButton)
 			.ContentPadding(0)
 			.ForegroundColor(FSlateColor::UseForeground())
-			.ButtonStyle(FEditorStyle::Get(), "ToggleButton")
+			.ComboButtonStyle(FAppStyle::Get(), "SimpleComboButton")
 			.OnGetMenuContent(this, &SLevelSequenceContextPicker::BuildWorldPickerMenu)
 			.ToolTipText(FText::Format(LOCTEXT("WorldPickerTextFomrat", "'{0}': The world context and playback client that sequencer should be bound to, and playback within."), GetCurrentContextAndClientText()))
 			.ButtonContent()
 			[
 				SNew(SImage)
+				.ColorAndOpacity(FSlateColor::UseForeground())
 				.Image(FEditorStyle::GetBrush("SceneOutliner.World"))
 			]
 		]
diff --git a/Engine/Source/Editor/EditorStyle/Private/StarshipStyle.cpp b/Engine/Source/Editor/EditorStyle/Private/StarshipStyle.cpp
--- a/Engine/Source/Editor/EditorStyle/Private/StarshipStyle.cpp
+++ b/Engine/Source/Editor/EditorStyle/Private/StarshipStyle.cpp
@@ -2153,12 +2153,6 @@ void FStarshipEditorStyle::FStyle::SetupSequencerStyles()
 		Set("Sequencer.OpenTaggedBindingManager.Small", new IMAGE_BRUSH("Sequencer/Main_Icons/Icon_Sequencer_OpenTaggedBindingManager_16x", Icon24x24));
 		Set("Sequencer.OpenNodeGroupsManager", new IMAGE_BRUSH("Sequencer/Main_Icons/Icon_Sequencer_OpenGroupManager_16x", Icon48x48));
 		Set("Sequencer.OpenNodeGroupsManager.Small", new IMAGE_BRUSH("Sequencer/Main_Icons/Icon_Sequencer_OpenGroupManager_16x", Icon24x24));
-		Set("Sequencer.BreadcrumbText", FTextBlockStyle(NormalText)
-			.SetFont(DEFAULT_FONT("Bold", 11))
-			.SetColorAndOpacity(FLinearColor(1.0f, 1.0f, 1.0f))
-			.SetHighlightColor(FLinearColor(1.0f, 1.0f, 1.0f))
-			.SetShadowOffset(FVector2D(1, 1))
-			.SetShadowColorAndOpacity(FLinearColor(0, 0, 0, 0.9f)));
 
 		Set("Sequencer.ToggleAutoScroll", new IMAGE_BRUSH("Icons/icon_Sequencer_ToggleAutoScroll_40x", Icon48x48));
 		Set("Sequencer.ToggleAutoScroll.Small", new IMAGE_BRUSH("Icons/icon_Sequencer_ToggleAutoScroll_16x", Icon16x16));
@@ -2207,12 +2201,7 @@ void FStarshipEditorStyle::FStyle::SetupSequencerStyles()
 		Set("Sequencer.CursorDecorator_MarqueeAdd", new IMAGE_BRUSH("Sequencer/CursorDecorator_MarqueeAdd", Icon16x16));
 		Set("Sequencer.CursorDecorator_MarqueeSubtract", new IMAGE_BRUSH("Sequencer/CursorDecorator_MarqueeSubtract", Icon16x16));
 
-		Set("Sequencer.BreadcrumbText", FTextBlockStyle(NormalText)
-			.SetFont(DEFAULT_FONT("Bold", 11))
-			.SetColorAndOpacity(FLinearColor(1.0f, 1.0f, 1.0f))
-			.SetHighlightColor(FLinearColor(1.0f, 1.0f, 1.0f))
-			.SetShadowOffset(FVector2D(1, 1))
-			.SetShadowColorAndOpacity(FLinearColor(0, 0, 0, 0.9f)));
+		Set("Sequencer.BreadcrumbText", FTextBlockStyle(NormalText).SetFont(FStyleFonts::Get().NormalBold));
 		Set("Sequencer.BreadcrumbIcon", new IMAGE_BRUSH("Common/SmallArrowRight", Icon10x10));
 
 		Set("Sequencer.AddKey.Details", new IMAGE_BRUSH("Sequencer/AddKey_Details", FVector2D(11, 11)));
diff --git a/Engine/Source/Editor/MovieSceneTools/Private/TrackEditors/CameraCutTrackEditor.cpp b/Engine/Source/Editor/MovieSceneTools/Private/TrackEditors/CameraCutTrackEditor.cpp
--- a/Engine/Source/Editor/MovieSceneTools/Private/TrackEditors/CameraCutTrackEditor.cpp
+++ b/Engine/Source/Editor/MovieSceneTools/Private/TrackEditors/CameraCutTrackEditor.cpp
@@ -204,11 +204,12 @@ TSharedPtr<SWidget> FCameraCutTrackEditor::BuildOutlinerEditWidget(const FGuid&
 	.Padding(4, 0, 0, 0)
 	[
 		SNew(SCheckBox)
+        .Style( &FAppStyle::Get().GetWidgetStyle<FCheckBoxStyle>("ToggleButtonCheckBoxAlt"))
+		.Type(ESlateCheckBoxType::CheckBox)
 		.IsFocusable(false)
 		.IsChecked(this, &FCameraCutTrackEditor::IsCameraLocked)
 		.OnCheckStateChanged(this, &FCameraCutTrackEditor::OnLockCameraClicked)
 		.ToolTipText(this, &FCameraCutTrackEditor::GetLockCameraToolTip)
-		.ForegroundColor(FLinearColor::White)
 		.CheckedImage(FEditorStyle::GetBrush("Sequencer.LockCamera"))
 		.CheckedHoveredImage(FEditorStyle::GetBrush("Sequencer.LockCamera"))
 		.CheckedPressedImage(FEditorStyle::GetBrush("Sequencer.LockCamera"))
diff --git a/Engine/Source/Editor/MovieSceneTools/Private/TrackEditors/CinematicShotTrackEditor.cpp b/Engine/Source/Editor/MovieSceneTools/Private/TrackEditors/CinematicShotTrackEditor.cpp
--- a/Engine/Source/Editor/MovieSceneTools/Private/TrackEditors/CinematicShotTrackEditor.cpp
+++ b/Engine/Source/Editor/MovieSceneTools/Private/TrackEditors/CinematicShotTrackEditor.cpp
@@ -107,11 +107,12 @@ TSharedPtr<SWidget> FCinematicShotTrackEditor::BuildOutlinerEditWidget(const FGu
 	.Padding(4, 0, 0, 0)
 	[
 		SNew(SCheckBox)
+        .Style( &FAppStyle::Get().GetWidgetStyle<FCheckBoxStyle>("ToggleButtonCheckBoxAlt"))
+		.Type(ESlateCheckBoxType::CheckBox)
 		.IsFocusable(false)
 		.IsChecked(this, &FCinematicShotTrackEditor::AreShotsLocked)
 		.OnCheckStateChanged(this, &FCinematicShotTrackEditor::OnLockShotsClicked)
 		.ToolTipText(this, &FCinematicShotTrackEditor::GetLockShotsToolTip)
-		.ForegroundColor(FLinearColor::White)
 		.CheckedImage(FEditorStyle::GetBrush("Sequencer.LockCamera"))
 		.CheckedHoveredImage(FEditorStyle::GetBrush("Sequencer.LockCamera"))
 		.CheckedPressedImage(FEditorStyle::GetBrush("Sequencer.LockCamera"))
diff --git a/Engine/Source/Editor/MovieSceneTools/Private/TrackEditors/TransformTrackEditor.cpp b/Engine/Source/Editor/MovieSceneTools/Private/TrackEditors/TransformTrackEditor.cpp
--- a/Engine/Source/Editor/MovieSceneTools/Private/TrackEditors/TransformTrackEditor.cpp
+++ b/Engine/Source/Editor/MovieSceneTools/Private/TrackEditors/TransformTrackEditor.cpp
@@ -415,12 +415,13 @@ void F3DTransformTrackEditor::BuildObjectBindingEditButtons(TSharedPtr<SHorizont
 		.Padding(4, 0, 0, 0)
 		[
 			SNew(SCheckBox)		
+                .Style( &FAppStyle::Get().GetWidgetStyle<FCheckBoxStyle>("ToggleButtonCheckBoxAlt"))
+				.Type(ESlateCheckBoxType::CheckBox)
 				.IsFocusable(false)
 				.Visibility(this, &F3DTransformTrackEditor::IsCameraVisible, ObjectGuid)
 				.IsChecked(this, &F3DTransformTrackEditor::IsCameraLocked, ObjectGuid)
 				.OnCheckStateChanged(this, &F3DTransformTrackEditor::OnLockCameraClicked, ObjectGuid)
 				.ToolTipText(this, &F3DTransformTrackEditor::GetLockCameraToolTip, ObjectGuid)
-				.ForegroundColor(FLinearColor::White)
 				.CheckedImage(FEditorStyle::GetBrush("Sequencer.LockCamera"))
 				.CheckedHoveredImage(FEditorStyle::GetBrush("Sequencer.LockCamera"))
 				.CheckedPressedImage(FEditorStyle::GetBrush("Sequencer.LockCamera"))
diff --git a/Engine/Source/Editor/Sequencer/Private/SSequencer.cpp b/Engine/Source/Editor/Sequencer/Private/SSequencer.cpp
--- a/Engine/Source/Editor/Sequencer/Private/SSequencer.cpp
+++ b/Engine/Source/Editor/Sequencer/Private/SSequencer.cpp
@@ -31,6 +31,7 @@
 #include "Widgets/Input/SSpinBox.h"
 #include "Widgets/Input/SCheckBox.h"
 #include "Widgets/Docking/SDockTab.h"
+#include "SEditorHeaderButton.h"
 #include "EditorStyleSet.h"
 #include "Kismet2/KismetEditorUtilities.h"
 #include "MovieSceneSequenceEditor.h"
@@ -580,6 +581,7 @@ void SSequencer::Construct(const FArguments& InArgs, TSharedRef<FSequencer> InSe
 					// Toolbar
 					+ SGridPanel::Slot( Column0, Row0, SGridPanel::Layer(10) )
 					.ColumnSpan(2)
+					.Padding(4.f, 0.f, 8.f, 0.f)
 					[
 						SNew(SBox)
 						.Padding(FMargin(CommonPadding,0.0f,0.0f,0.f))
@@ -620,17 +622,20 @@ void SSequencer::Construct(const FArguments& InArgs, TSharedRef<FSequencer> InSe
 										SNew(SButton)
 										.Visibility(this, &SSequencer::GetBreadcrumbTrailVisibility)
 										.VAlign(EVerticalAlignment::VAlign_Center)
-										.ButtonStyle(FEditorStyle::Get(), "FlatButton")
-										.ForegroundColor(FLinearColor::White)
+										.ButtonStyle(FEditorStyle::Get(), "SimpleButton")
 										.ToolTipText_Lambda([this] { return SequencerPtr.Pin()->GetNavigateBackwardTooltip(); })
 										.ContentPadding(FMargin(1, 0))
 										.OnClicked_Lambda([this] { return SequencerPtr.Pin()->NavigateBackward(); })
 										.IsEnabled_Lambda([this] { return SequencerPtr.Pin()->CanNavigateBackward(); })
 										[
-											SNew(STextBlock)
-											.TextStyle(FEditorStyle::Get(), "ContentBrowser.TopBar.Font")
-											.Font(FEditorStyle::Get().GetFontStyle("FontAwesome.11"))
-											.Text(FText::FromString(FString(TEXT("\xf060"))) /*fa-arrow-left*/)
+											SNew(SBox) // scale up since the default icons are 16x16
+											.WidthOverride(20)
+											.HeightOverride(20)
+											[
+												SNew(SImage)
+												.ColorAndOpacity(FSlateColor::UseForeground())
+												.Image(FAppStyle::Get().GetBrush("Icons.ArrowLeft"))
+											]
 										]
 									]
 								]
@@ -649,17 +654,20 @@ void SSequencer::Construct(const FArguments& InArgs, TSharedRef<FSequencer> InSe
 										SNew(SButton)
 										.Visibility(this, &SSequencer::GetBreadcrumbTrailVisibility)
 										.VAlign(EVerticalAlignment::VAlign_Center)
-										.ButtonStyle(FEditorStyle::Get(), "FlatButton")
-										.ForegroundColor(FLinearColor::White)
+										.ButtonStyle(FEditorStyle::Get(), "SimpleButton")
 										.ToolTipText_Lambda([this] { return SequencerPtr.Pin()->GetNavigateForwardTooltip(); })
 										.ContentPadding(FMargin(1, 0))
 										.OnClicked_Lambda([this] { return SequencerPtr.Pin()->NavigateForward(); })
 										.IsEnabled_Lambda([this] { return SequencerPtr.Pin()->CanNavigateForward(); })
 										[
-											SNew(STextBlock)
-											.TextStyle(FEditorStyle::Get(), "ContentBrowser.TopBar.Font")
-											.Font(FEditorStyle::Get().GetFontStyle("FontAwesome.11"))
-											.Text(FText::FromString(FString(TEXT("\xf061"))) /*fa-arrow-right*/)
+											SNew(SBox) // scale up since the default icons are 16x16
+											.WidthOverride(20)
+											.HeightOverride(20)
+											[
+												SNew(SImage)
+												.ColorAndOpacity(FSlateColor::UseForeground())
+												.Image(FAppStyle::Get().GetBrush("Icons.ArrowRight"))
+											]
 										]
 									]
 								]
@@ -681,17 +689,15 @@ void SSequencer::Construct(const FArguments& InArgs, TSharedRef<FSequencer> InSe
 								[
 									SAssignNew(BreadcrumbPickerButton, SComboButton)
 									.Visibility(this, &SSequencer::GetBreadcrumbTrailVisibility)
-									.ButtonStyle(FEditorStyle::Get(), "FlatButton")
-									.ForegroundColor(FLinearColor::White)
+									.ButtonStyle(FAppStyle::Get(), "SimpleButton")
 									.OnGetMenuContent_Lambda([this] { return SNew(SSequencerHierarchyBrowser, SequencerPtr); })
 									.HasDownArrow(false)
 									.ContentPadding(FMargin(3, 3))
 									.ButtonContent()
 									[
-										SNew(STextBlock)
-										.TextStyle(FEditorStyle::Get(), "Sequencer.BreadcrumbText")
-										.Font(FEditorStyle::Get().GetFontStyle("FontAwesome.11"))
-										.Text(FText::FromString(FString(TEXT("\xf07c"))) /*fa-folder-open*/)
+										SNew(SImage)
+										.ColorAndOpacity(FSlateColor::UseForeground())
+										.Image(FAppStyle::Get().GetBrush("Icons.FolderOpen"))
 									]
 								]
 
@@ -704,9 +710,9 @@ void SSequencer::Construct(const FArguments& InArgs, TSharedRef<FSequencer> InSe
 									SAssignNew(BreadcrumbTrail, SBreadcrumbTrail<FSequencerBreadcrumb>)
 									.Visibility(this, &SSequencer::GetBreadcrumbTrailVisibility)
 									.OnCrumbClicked(this, &SSequencer::OnCrumbClicked)
-									.ButtonStyle(FEditorStyle::Get(), "FlatButton")
-									.DelimiterImage(FEditorStyle::GetBrush("Sequencer.BreadcrumbIcon"))
-									.TextStyle(FEditorStyle::Get(), "Sequencer.BreadcrumbText")
+									.ButtonStyle(FAppStyle::Get(), "SimpleButton")
+									.DelimiterImage(FAppStyle::Get().GetBrush("Sequencer.BreadcrumbIcon"))
+									.TextStyle(FAppStyle::Get(), "Sequencer.BreadcrumbText")
 								]
 
 								// Sequence Locking symbol
@@ -714,20 +720,21 @@ void SSequencer::Construct(const FArguments& InArgs, TSharedRef<FSequencer> InSe
 								.HAlign(HAlign_Right)
 								.VAlign(VAlign_Center)
 								.AutoWidth()
-								.Padding(2, 0, 0, 0)
 								[
 									SNew(SCheckBox)
+									.Style(FAppStyle::Get(),"ToggleButtonCheckBoxAlt")
+									.Type(ESlateCheckBoxType::CheckBox) // Use CheckBox instead of ToggleType since we're not putting ohter widget inside
+									.Padding(FMargin(0.f))
 									.IsFocusable(false)		
 									.IsChecked_Lambda([this] { return GetIsSequenceReadOnly() ? ECheckBoxState::Checked : ECheckBoxState::Unchecked; } )
 									.OnCheckStateChanged(this, &SSequencer::OnSetSequenceReadOnly)
 									.ToolTipText_Lambda([this] { return GetIsSequenceReadOnly() ? LOCTEXT("UnlockSequence", "Unlock the animation so that it is editable") : LOCTEXT("LockSequence", "Lock the animation so that it is not editable"); } )
-									.ForegroundColor(FLinearColor::White)
-									.CheckedImage(FEditorStyle::GetBrush("Sequencer.LockSequence"))
-									.CheckedHoveredImage(FEditorStyle::GetBrush("Sequencer.LockSequence"))
-									.CheckedPressedImage(FEditorStyle::GetBrush("Sequencer.LockSequence"))
-									.UncheckedImage(FEditorStyle::GetBrush("Sequencer.UnlockSequence"))
-									.UncheckedHoveredImage(FEditorStyle::GetBrush("Sequencer.UnlockSequence"))
-									.UncheckedPressedImage(FEditorStyle::GetBrush("Sequencer.UnlockSequence"))
+									.CheckedImage(FAppStyle::Get().GetBrush("Icons.Lock"))
+									.CheckedHoveredImage(FAppStyle::Get().GetBrush("Icons.Lock"))
+									.CheckedPressedImage(FAppStyle::Get().GetBrush("Icons.Lock"))
+									.UncheckedImage(FAppStyle::Get().GetBrush("Icons.Unlock"))
+									.UncheckedHoveredImage(FAppStyle::Get().GetBrush("Icons.Unlock"))
+									.UncheckedPressedImage(FAppStyle::Get().GetBrush("Icons.Unlock"))
 								]
 							]
 						]
@@ -758,22 +765,23 @@ void SSequencer::Construct(const FArguments& InArgs, TSharedRef<FSequencer> InSe
 							[
 								MakeAddButton()
 							]
-							
+						
 							+ SHorizontalBox::Slot()
-							.AutoWidth()
 							.VAlign(VAlign_Center)
-							.Padding(FMargin(0.f, 0.f, CommonPadding, 0.f))
 							[
-								MakeFilterButton()
+								SAssignNew(SearchBox, SSearchBox)
+								.HintText(LOCTEXT("SearchNodesHint", "Search Tracks"))
+								.OnTextChanged( this, &SSequencer::OnOutlinerSearchChanged )
 							]
 
 							+ SHorizontalBox::Slot()
+							.AutoWidth()
 							.VAlign(VAlign_Center)
+							.Padding(FMargin(0.f, 0.f, CommonPadding, 0.f))
 							[
-								SAssignNew(SearchBox, SSearchBox)
-								.HintText(LOCTEXT("SearchNodesHint", "Search Tracks"))
-								.OnTextChanged( this, &SSequencer::OnOutlinerSearchChanged )
+								MakeFilterButton()
 							]
+
 							+ SHorizontalBox::Slot()
 							.AutoWidth()
 							.VAlign(VAlign_Center)
@@ -1067,6 +1075,7 @@ void SSequencer::Construct(const FArguments& InArgs, TSharedRef<FSequencer> InSe
 
 					+ SSplitter::Slot()
 					.Value(FillCoefficient_0)
+					.MinSize(200)
 					.OnSlotResized(SSplitter::FOnSlotResized::CreateSP(this, &SSequencer::OnColumnFillCoefficientChanged, 0))
 					[
 						SNew(SSpacer)
@@ -1317,81 +1326,28 @@ void SSequencer::HandleOutlinerNodeSelectionChanged()
 
 TSharedRef<SWidget> SSequencer::MakeAddButton()
 {
-	return SNew(SComboButton)
+	return SNew(SEditorHeaderButton)
 	.OnGetMenuContent(this, &SSequencer::MakeAddMenu)
-	.ButtonStyle(FEditorStyle::Get(), "FlatButton.Success")
-	.ContentPadding(FMargin(2.0f, 1.0f))
-	.HasDownArrow(false)
-	.ButtonContent()
-	[
-		SNew(SHorizontalBox)
-
-		+ SHorizontalBox::Slot()
-		.VAlign(VAlign_Center)
-		.AutoWidth()
-		[
-			SNew(STextBlock)
-			.TextStyle(FEditorStyle::Get(), "NormalText.Important")
-			.Font(FEditorStyle::Get().GetFontStyle("FontAwesome.10"))
-			.Text(FEditorFontGlyphs::Plus)
-			.IsEnabled_Lambda([=]() { return !SequencerPtr.Pin()->IsReadOnly(); })
-		]
-
-		+ SHorizontalBox::Slot()
-		.AutoWidth()
-		.Padding(4, 0, 0, 0)
-		[
-			SNew(STextBlock)
-			.TextStyle(FEditorStyle::Get(), "NormalText.Important")
-			.Text(LOCTEXT("Track", "Track"))
-			.IsEnabled_Lambda([=]() { return !SequencerPtr.Pin()->IsReadOnly(); })
-		]
-
-		+ SHorizontalBox::Slot()
-		.VAlign(VAlign_Center)
-		.AutoWidth()
-		.Padding(4, 0, 0, 0)
-		[
-			SNew(STextBlock)
-			.TextStyle(FEditorStyle::Get(), "NormalText.Important")
-			.Font(FEditorStyle::Get().GetFontStyle("FontAwesome.10"))
-			.Text(FEditorFontGlyphs::Caret_Down)
-			.IsEnabled_Lambda([=]() { return !SequencerPtr.Pin()->IsReadOnly(); })
-		]
-	];
+	.Icon(FAppStyle::Get().GetBrush("Icons.Plus"))
+	.Text(LOCTEXT("Track", "Track"))
+	.IsEnabled_Lambda([=]() { return !SequencerPtr.Pin()->IsReadOnly(); });
 }
 
 TSharedRef<SWidget> SSequencer::MakeFilterButton()
 {
 	return SNew(SComboButton)
-	.ComboButtonStyle(FEditorStyle::Get(), "GenericFilters.ComboButtonStyle")
-	.ForegroundColor(FLinearColor::White)
+	.ButtonStyle( FAppStyle::Get(), "SimpleButton" )
+	.ForegroundColor(FSlateColor::UseForeground())
 	.ContentPadding(0)
 	.ToolTipText(LOCTEXT("AddTrackFilterToolTip", "Add a track filter."))
 	.OnGetMenuContent(this, &SSequencer::MakeFilterMenu)
-	.HasDownArrow(true)
 	.ContentPadding(FMargin(1, 0))
+	.HasDownArrow(false)
 	.ButtonContent()
 	[
-		SNew(SHorizontalBox)
-
-		+ SHorizontalBox::Slot()
-		.AutoWidth()
-		[
-			SNew(STextBlock)
-			.TextStyle(FEditorStyle::Get(), "GenericFilters.TextStyle")
-			.Font(FEditorStyle::Get().GetFontStyle("FontAwesome.9"))
-			.Text(FEditorFontGlyphs::Filter)
-		]
-
-		+ SHorizontalBox::Slot()
-		.AutoWidth()
-		.Padding(2, 0, 0, 0)
-		[
-			SNew(STextBlock)
-			.TextStyle(FEditorStyle::Get(), "GenericFilters.TextStyle")
-			.Text(LOCTEXT("Filters", "Filters"))
-		]
+		SNew(SImage)
+		.ColorAndOpacity(FSlateColor::UseForeground())
+		.Image(FAppStyle::Get().GetBrush("Icons.Filter"))
 	];
 }
 
