Fixes ui overflow issues when sequencer is squished

diff --git a/Engine/Source/Editor/Sequencer/Private/SSequencer.cpp b/Engine/Source/Editor/Sequencer/Private/SSequencer.cpp
--- a/Engine/Source/Editor/Sequencer/Private/SSequencer.cpp
+++ b/Engine/Source/Editor/Sequencer/Private/SSequencer.cpp
@@ -755,6 +755,7 @@ void SSequencer::Construct(const FArguments& InArgs, TSharedRef<FSequencer> InSe
 						SNew(SBorder)
 						.BorderImage(FEditorStyle::GetBrush("ToolPanel.GroupBorder"))
 						.Padding(FMargin(CommonPadding*2, CommonPadding))
+						.Clipping(EWidgetClipping::ClipToBounds)
 						[
 							SNew(SHorizontalBox)
 
