Sequencer: Fixed Spawnable 'Change Class' assigning a CDO directly

  - Referencing the CDO directly in a persistent way causes issues with class recompilation and reinstancing since CDO's are never reinstanced.
  - We now create a copy of the class instead

[Copied from 15749122]

diff --git a/Engine/Source/Editor/Sequencer/Private/DisplayNodes/SequencerObjectBindingNode.cpp b/Engine/Source/Editor/Sequencer/Private/DisplayNodes/SequencerObjectBindingNode.cpp
--- a/Engine/Source/Editor/Sequencer/Private/DisplayNodes/SequencerObjectBindingNode.cpp
+++ b/Engine/Source/Editor/Sequencer/Private/DisplayNodes/SequencerObjectBindingNode.cpp
@@ -520,9 +520,15 @@ void FSequencerObjectBindingNode::HandleTemplateActorClassPicked(UClass* ChosenC
 	FScopedTransaction Transaction(LOCTEXT("ChangeClass", "Change Class"));
 
 	MovieScene->Modify();
-	Spawnable->SetObjectTemplate(ChosenClass->ClassDefaultObject);
-	GetSequencer().GetSpawnRegister().DestroySpawnedObject(Spawnable->GetGuid(), GetSequencer().GetFocusedTemplateID(), GetSequencer());
-	GetSequencer().ForceEvaluate();
+
+	TValueOrError<FNewSpawnable, FText> Result = GetSequencer().GetSpawnRegister().CreateNewSpawnableType(*ChosenClass, *MovieScene, nullptr);
+	if (Result.IsValid())
+	{
+		Spawnable->SetObjectTemplate(Result.GetValue().ObjectTemplate);
+
+		GetSequencer().GetSpawnRegister().DestroySpawnedObject(Spawnable->GetGuid(), GetSequencer().GetFocusedTemplateID(), GetSequencer());
+		GetSequencer().ForceEvaluate();
+	}
 }
 
 
diff --git a/Engine/Source/Runtime/LevelSequence/Private/LevelSequenceActorSpawner.cpp b/Engine/Source/Runtime/LevelSequence/Private/LevelSequenceActorSpawner.cpp
--- a/Engine/Source/Runtime/LevelSequence/Private/LevelSequenceActorSpawner.cpp
+++ b/Engine/Source/Runtime/LevelSequence/Private/LevelSequenceActorSpawner.cpp
@@ -73,6 +73,11 @@ UObject* FLevelSequenceActorSpawner::SpawnObject(FMovieSceneSpawnable& Spawnable
 		return nullptr;
 	}
 
+	if (!ensure(!ObjectTemplate->GetClass()->HasAnyClassFlags(CLASS_NewerVersionExists)))
+	{
+		return nullptr;
+	}
+
 	const EObjectFlags ObjectFlags = RF_Transient | RF_Transactional;
 
 	// @todo sequencer livecapture: Consider using SetPlayInEditorWorld() and RestoreEditorWorld() here instead
diff --git a/Engine/Source/Runtime/MovieScene/Public/MovieSceneSpawnable.h b/Engine/Source/Runtime/MovieScene/Public/MovieSceneSpawnable.h
--- a/Engine/Source/Runtime/MovieScene/Public/MovieSceneSpawnable.h
+++ b/Engine/Source/Runtime/MovieScene/Public/MovieSceneSpawnable.h
@@ -104,6 +104,7 @@ public:
 	 */
 	void SetObjectTemplate(UObject* InObjectTemplate)
 	{
+		checkf(!InObjectTemplate->HasAnyFlags(RF_ClassDefaultObject), TEXT("Setting CDOs as object templates is not supported. Please use the class directly."));
 		ObjectTemplate = InObjectTemplate;
 	}
 
