Selected sequencer track color is now the style's selection color instead of the previous orange selection color, the text has been adjusted to UseForeground rather than pure white so it is more readable when selected.

diff --git a/Engine/Source/Editor/Sequencer/Private/DisplayNodes/SequencerDisplayNode.cpp b/Engine/Source/Editor/Sequencer/Private/DisplayNodes/SequencerDisplayNode.cpp
--- a/Engine/Source/Editor/Sequencer/Private/DisplayNodes/SequencerDisplayNode.cpp
+++ b/Engine/Source/Editor/Sequencer/Private/DisplayNodes/SequencerDisplayNode.cpp
@@ -798,9 +798,9 @@ FSlateFontInfo FSequencerDisplayNode::GetDisplayNameFont() const
 	return NodeFont;
 }
 
-FLinearColor FSequencerDisplayNode::GetDisplayNameColor() const
+FSlateColor FSequencerDisplayNode::GetDisplayNameColor() const
 {
-	return IsDimmed() ? FLinearColor(0.6f, 0.6f, 0.6f, 0.6f) : FLinearColor(1.0f, 1.0f, 1.0f, 1.0f);
+	return IsDimmed() ? FLinearColor(0.6f, 0.6f, 0.6f, 0.6f) : FSlateColor::UseForeground();
 }
 
 FText FSequencerDisplayNode::GetDisplayNameToolTipText() const
diff --git a/Engine/Source/Editor/Sequencer/Private/DisplayNodes/SequencerDisplayNode.h b/Engine/Source/Editor/Sequencer/Private/DisplayNodes/SequencerDisplayNode.h
--- a/Engine/Source/Editor/Sequencer/Private/DisplayNodes/SequencerDisplayNode.h
+++ b/Engine/Source/Editor/Sequencer/Private/DisplayNodes/SequencerDisplayNode.h
@@ -176,7 +176,7 @@ public:
 	/**
 	* @return the color used to draw the display name.
 	*/
-	virtual FLinearColor GetDisplayNameColor() const;
+	virtual FSlateColor GetDisplayNameColor() const;
 
 	/**
 	*@return The font used to draw the display name.
diff --git a/Engine/Source/Editor/Sequencer/Private/DisplayNodes/SequencerFolderNode.cpp b/Engine/Source/Editor/Sequencer/Private/DisplayNodes/SequencerFolderNode.cpp
--- a/Engine/Source/Editor/Sequencer/Private/DisplayNodes/SequencerFolderNode.cpp
+++ b/Engine/Source/Editor/Sequencer/Private/DisplayNodes/SequencerFolderNode.cpp
@@ -121,14 +121,14 @@ FSlateColor FSequencerFolderNode::GetIconColor() const
 	return FSlateColor(MovieSceneFolder.GetFolderColor());
 }
 
-FLinearColor FSequencerFolderNode::GetDisplayNameColor() const
+FSlateColor FSequencerFolderNode::GetDisplayNameColor() const
 {
 	if (ParentTree.IsNodeMute(this))
 	{
 		return FLinearColor(0.6f, 0.6f, 0.6f, 0.6f);
 	}
 
-	return FLinearColor::White;
+	return FSlateColor::UseForeground();
 }
 
 
diff --git a/Engine/Source/Editor/Sequencer/Private/DisplayNodes/SequencerFolderNode.h b/Engine/Source/Editor/Sequencer/Private/DisplayNodes/SequencerFolderNode.h
--- a/Engine/Source/Editor/Sequencer/Private/DisplayNodes/SequencerFolderNode.h
+++ b/Engine/Source/Editor/Sequencer/Private/DisplayNodes/SequencerFolderNode.h
@@ -27,7 +27,7 @@ public:
 	virtual void SetDisplayName( const FText& NewDisplayName ) override;
 	virtual const FSlateBrush* GetIconBrush() const override;
 	virtual FSlateColor GetIconColor() const override;
-	virtual FLinearColor GetDisplayNameColor() const override;
+	virtual FSlateColor GetDisplayNameColor() const override;
 	virtual bool CanDrag() const override;
 	virtual TOptional<EItemDropZone> CanDrop( FSequencerDisplayNodeDragDropOp& DragDropOp, EItemDropZone ItemDropZone ) const override;
 	virtual void Drop( const TArray<TSharedRef<FSequencerDisplayNode>>& DraggedNodes, EItemDropZone ItemDropZone ) override;
diff --git a/Engine/Source/Editor/Sequencer/Private/DisplayNodes/SequencerObjectBindingNode.cpp b/Engine/Source/Editor/Sequencer/Private/DisplayNodes/SequencerObjectBindingNode.cpp
--- a/Engine/Source/Editor/Sequencer/Private/DisplayNodes/SequencerObjectBindingNode.cpp
+++ b/Engine/Source/Editor/Sequencer/Private/DisplayNodes/SequencerObjectBindingNode.cpp
@@ -794,7 +794,7 @@ FText FSequencerObjectBindingNode::GetDisplayName() const
 	return FText();
 }
 
-FLinearColor FSequencerObjectBindingNode::GetDisplayNameColor() const
+FSlateColor FSequencerObjectBindingNode::GetDisplayNameColor() const
 {
 	FSequencer& Sequencer = ParentTree.GetSequencer();
 
diff --git a/Engine/Source/Editor/Sequencer/Private/DisplayNodes/SequencerObjectBindingNode.h b/Engine/Source/Editor/Sequencer/Private/DisplayNodes/SequencerObjectBindingNode.h
--- a/Engine/Source/Editor/Sequencer/Private/DisplayNodes/SequencerObjectBindingNode.h
+++ b/Engine/Source/Editor/Sequencer/Private/DisplayNodes/SequencerObjectBindingNode.h
@@ -61,7 +61,7 @@ public:
 	virtual TSharedRef<SWidget> GetCustomOutlinerContent() override;
 	virtual TSharedPtr<SWidget> GetAdditionalOutlinerLabel() override;
 	virtual FText GetDisplayName() const override;
-	virtual FLinearColor GetDisplayNameColor() const override;
+	virtual FSlateColor GetDisplayNameColor() const override;
 	virtual FText GetDisplayNameToolTipText() const override;
 	virtual const FSlateBrush* GetIconBrush() const override;
 	virtual const FSlateBrush* GetIconOverlayBrush() const override;
diff --git a/Engine/Source/Editor/Sequencer/Private/DisplayNodes/SequencerTrackNode.cpp b/Engine/Source/Editor/Sequencer/Private/DisplayNodes/SequencerTrackNode.cpp
--- a/Engine/Source/Editor/Sequencer/Private/DisplayNodes/SequencerTrackNode.cpp
+++ b/Engine/Source/Editor/Sequencer/Private/DisplayNodes/SequencerTrackNode.cpp
@@ -727,13 +727,13 @@ FText FSequencerTrackNode::GetDisplayName() const
 	return AssociatedTrack.IsValid() ? AssociatedTrack->GetDisplayName() : FText::GetEmpty();
 }
 
-FLinearColor FSequencerTrackNode::GetDisplayNameColor() const
+FSlateColor FSequencerTrackNode::GetDisplayNameColor() const
 {
 	UMovieSceneTrack* Track = GetTrack();
 
 	if (!Track)
 	{
-		return FLinearColor::White;
+		return FSlateColor::UseForeground();
 	}
 
 	// Display track node is red if the property track is not bound to valid property
@@ -744,7 +744,7 @@ FLinearColor FSequencerTrackNode::GetDisplayNameColor() const
 		// 3D transform tracks don't map to property bindings as below
 		if (Track->IsA<UMovieScene3DTransformTrack>() || Track->IsA<UMovieScenePrimitiveMaterialTrack>())
 		{
-			return bIsDimmed ? FLinearColor(0.6f, 0.6f, 0.6f, 0.6f) : FLinearColor::White;
+			return bIsDimmed ? FLinearColor(0.6f, 0.6f, 0.6f, 0.6f) : FSlateColor::UseForeground();
 		}
 
 		FGuid ObjectBinding;
@@ -764,7 +764,7 @@ FLinearColor FSequencerTrackNode::GetDisplayNameColor() const
 					FTrackInstancePropertyBindings PropertyBinding(PropertyTrack->GetPropertyName(), PropertyTrack->GetPropertyPath().ToString());
 					if (PropertyBinding.GetProperty(*Object))
 					{
-						return bIsDimmed ? FLinearColor(0.6f, 0.6f, 0.6f, 0.6f) : FLinearColor::White;
+						return bIsDimmed ? FLinearColor(0.6f, 0.6f, 0.6f, 0.6f) : FSlateColor::UseForeground();
 					}
 				}
 			}
diff --git a/Engine/Source/Editor/Sequencer/Private/DisplayNodes/SequencerTrackNode.h b/Engine/Source/Editor/Sequencer/Private/DisplayNodes/SequencerTrackNode.h
--- a/Engine/Source/Editor/Sequencer/Private/DisplayNodes/SequencerTrackNode.h
+++ b/Engine/Source/Editor/Sequencer/Private/DisplayNodes/SequencerTrackNode.h
@@ -131,7 +131,7 @@ public:
 	virtual TSharedRef<SWidget> GetCustomOutlinerContent() override;
 	virtual void GetChildKeyAreaNodesRecursively(TArray<TSharedRef<FSequencerSectionKeyAreaNode>>& OutNodes) const override;
 	virtual FText GetDisplayName() const override;
-	virtual FLinearColor GetDisplayNameColor() const override;
+	virtual FSlateColor GetDisplayNameColor() const override;
 	virtual FSlateFontInfo GetDisplayNameFont() const override;
 	virtual float GetNodeHeight() const override;
 	virtual FNodePadding GetNodePadding() const override;
diff --git a/Engine/Source/Editor/Sequencer/Private/SAnimationOutlinerTreeNode.cpp b/Engine/Source/Editor/Sequencer/Private/SAnimationOutlinerTreeNode.cpp
--- a/Engine/Source/Editor/Sequencer/Private/SAnimationOutlinerTreeNode.cpp
+++ b/Engine/Source/Editor/Sequencer/Private/SAnimationOutlinerTreeNode.cpp
@@ -23,6 +23,7 @@
 #include "Widgets/Text/SInlineEditableTextBlock.h"
 #include "Framework/SlateDelegates.h"
 #include "MovieSceneTrack.h"
+#include "Styling/StyleColors.h"
 
 #define LOCTEXT_NAMESPACE "AnimationOutliner"
 
@@ -330,11 +331,11 @@ FSlateColor SAnimationOutlinerTreeNode::GetNodeBackgroundTint() const
 
 	if (bIsSelected)
 	{
-		return FEditorStyle::GetSlateColor("SelectionColor_Pressed");
+		return FStyleColors::Select;
 	}
 	else if (Sequencer.GetSelection().NodeHasSelectedKeysOrSections(DisplayNode.ToSharedRef()))
 	{
-		return FLinearColor(FColor(115, 115, 115, 255));
+		return FStyleColors::Header;
 	}	
 	else if (DisplayNode->IsHovered())
 	{
@@ -355,11 +356,11 @@ FSlateColor SAnimationOutlinerTreeNode::GetNodeInnerBackgroundTint() const
 
 		if ( bIsSelected )
 		{
-			return FEditorStyle::GetSlateColor( "SelectionColor_Pressed" );
+			return FStyleColors::Select;
 		}
 		else if ( Sequencer.GetSelection().NodeHasSelectedKeysOrSections( DisplayNode.ToSharedRef() ) )
 		{
-			return FLinearColor( FColor( 115, 115, 115, 255 ) );
+			return FStyleColors::Header;
 		}
 		else if ( DisplayNode->IsHovered() )
 		{
