Sequencer: Initialize SequencerState time to the InSequencer time if the state doesn't exist in OnSequencerCreated. Remove the SequencerState on close if the session doesn't exist. This fixes a bug where Sequencer should initialize the current playback time to the lower bound of the playback range.

diff --git a/Engine/Plugins/Developer/Concert/ConcertSync/ConcertSyncClient/Source/ConcertSyncClient/Private/ConcertClientSequencerManager.cpp b/Engine/Plugins/Developer/Concert/ConcertSync/ConcertSyncClient/Source/ConcertSyncClient/Private/ConcertClientSequencerManager.cpp
--- a/Engine/Plugins/Developer/Concert/ConcertSync/ConcertSyncClient/Source/ConcertSyncClient/Private/ConcertClientSequencerManager.cpp
+++ b/Engine/Plugins/Developer/Concert/ConcertSync/ConcertSyncClient/Source/ConcertSyncClient/Private/ConcertClientSequencerManager.cpp
@@ -78,6 +78,14 @@ void FConcertClientSequencerManager::OnSequencerCreated(TSharedRef<ISequencer> I
 	// Find a Sequencer state for a newly opened sequencer if we have one.
 	UMovieSceneSequence* Sequence = InSequencer->GetRootMovieSceneSequence();
 	check(Sequence != nullptr);
+
+	if (!SequencerStates.Contains(*Sequence->GetPathName()))
+	{
+		FConcertSequencerState NewState;
+		NewState.Time = InSequencer->GetGlobalTime();
+		SequencerStates.Add(*Sequence->GetPathName(), NewState);
+	}
+
 	FConcertSequencerState& SequencerState = SequencerStates.FindOrAdd(*Sequence->GetPathName());
 
 	// Setup the Sequencer
@@ -223,6 +231,15 @@ void FConcertClientSequencerManager::OnSequencerClosed(TSharedRef<ISequencer> In
 			Session->SendCustomEvent(CloseEvent, Session->GetSessionServerEndpointId(), EConcertMessageFlags::ReliableOrdered);
 		}
 	}
+	else
+	{
+		UMovieSceneSequence* Sequence = InSequencer->GetRootMovieSceneSequence();
+		if (Sequence)
+		{
+			SequencerStates.Remove(*Sequence->GetPathName());
+		}
+	}
+
 	// Removed the closed Sequencer
 	OpenSequencers.RemoveAtSwap(Index);
 }
diff --git a/Engine/Source/Editor/Sequencer/Private/Sequencer.cpp b/Engine/Source/Editor/Sequencer/Private/Sequencer.cpp
--- a/Engine/Source/Editor/Sequencer/Private/Sequencer.cpp
+++ b/Engine/Source/Editor/Sequencer/Private/Sequencer.cpp
@@ -430,9 +430,6 @@ void FSequencer::InitSequencer(const FSequencerInitParams& InitParams, const TSh
 	
 	RootSequence = InitParams.RootSequence;
 
-	UpdateTimeBases();
-	PlayPosition.Reset(FFrameTime(0));
-
 	{
 		CompiledDataManager = FindObject<UMovieSceneCompiledDataManager>(GetTransientPackage(), TEXT("SequencerCompiledDataManager"));
 		if (!CompiledDataManager)
@@ -450,6 +447,9 @@ void FSequencer::InitSequencer(const FSequencerInitParams& InitParams, const TSh
 
 	ResetTimeController();
 
+	UpdateTimeBases();
+	PlayPosition.Reset(GetPlaybackRange().GetLowerBoundValue());
+
 	{
 		static_assert(sizeof(FSequencerViewModifierInfo) <= sizeof(FSequencer::FViewModifierInfo),
 			"We need to stomp the FViewModifierInfo objects with something new as a patch fix");
@@ -537,14 +537,9 @@ void FSequencer::InitSequencer(const FSequencerInitParams& InitParams, const TSh
 
 	// Update initial movie scene data
 	NotifyMovieSceneDataChanged( EMovieSceneDataChangeType::ActiveMovieSceneChanged );
-	UpdateTimeBoundsToFocusedMovieScene();
 
-	FQualifiedFrameTime CurrentTime = GetLocalTime();
-	if (!TargetViewRange.Contains(CurrentTime.AsSeconds()))
-	{
-		SetLocalTimeDirectly(LastViewRange.GetLowerBoundValue() * CurrentTime.Rate);
-		OnGlobalTimeChangedDelegate.Broadcast();
-	}
+	// Update the view range to the new current time
+	UpdateTimeBoundsToFocusedMovieScene();
 
 	// NOTE: Could fill in asset editor commands here!
 
