Adds Spawnable Column in Scene Outliner when Sequencer is Open

diff --git a/Engine/Plugins/MovieScene/LevelSequenceEditor/Source/LevelSequenceEditor/Private/LevelSequenceEditorToolkit.cpp b/Engine/Plugins/MovieScene/LevelSequenceEditor/Source/LevelSequenceEditor/Private/LevelSequenceEditorToolkit.cpp
--- a/Engine/Plugins/MovieScene/LevelSequenceEditor/Source/LevelSequenceEditor/Private/LevelSequenceEditorToolkit.cpp
+++ b/Engine/Plugins/MovieScene/LevelSequenceEditor/Source/LevelSequenceEditor/Private/LevelSequenceEditorToolkit.cpp
@@ -209,8 +209,7 @@ void FLevelSequenceEditorToolkit::Initialize(const EToolkitMode::Type Mode, cons
 	// @todo remove when world-centric mode is added
 	FLevelEditorModule& LevelEditorModule = FModuleManager::LoadModuleChecked<FLevelEditorModule>("LevelEditor");
 
-	// Reopen the scene outliner so that is refreshed with the sequencer info column
-	if (Sequencer->GetSequencerSettings()->GetShowOutlinerInfoColumn())
+	// Reopen the scene outliner so that is refreshed with the sequencer columns
 	{
 		TSharedPtr<FTabManager> LevelEditorTabManager = LevelEditorModule.GetLevelEditorTabManager();
 		if (LevelEditorTabManager->FindExistingLiveTab(FName("LevelEditorSceneOutliner")).IsValid())
diff --git a/Engine/Plugins/MovieScene/TemplateSequence/Source/TemplateSequenceEditor/Private/TemplateSequenceEditorToolkit.cpp b/Engine/Plugins/MovieScene/TemplateSequence/Source/TemplateSequenceEditor/Private/TemplateSequenceEditorToolkit.cpp
--- a/Engine/Plugins/MovieScene/TemplateSequence/Source/TemplateSequenceEditor/Private/TemplateSequenceEditorToolkit.cpp
+++ b/Engine/Plugins/MovieScene/TemplateSequence/Source/TemplateSequenceEditor/Private/TemplateSequenceEditorToolkit.cpp
@@ -121,8 +121,7 @@ void FTemplateSequenceEditorToolkit::Initialize(const EToolkitMode::Type Mode, c
 
 	FLevelEditorModule& LevelEditorModule = FModuleManager::LoadModuleChecked<FLevelEditorModule>("LevelEditor");
 
-	// Reopen the scene outliner so that is refreshed with the sequencer info column
-	if (Sequencer->GetSequencerSettings()->GetShowOutlinerInfoColumn())
+	// Reopen the scene outliner so that is refreshed with the sequencer columns
 	{
 		TSharedPtr<FTabManager> LevelEditorTabManager = LevelEditorModule.GetLevelEditorTabManager();
 		if (LevelEditorTabManager->FindExistingLiveTab(FName("LevelEditorSceneOutliner")).IsValid())
diff --git a/Engine/Source/Editor/Sequencer/Private/LevelEditorSequencerIntegration.cpp b/Engine/Source/Editor/Sequencer/Private/LevelEditorSequencerIntegration.cpp
--- a/Engine/Source/Editor/Sequencer/Private/LevelEditorSequencerIntegration.cpp
+++ b/Engine/Source/Editor/Sequencer/Private/LevelEditorSequencerIntegration.cpp
@@ -30,6 +30,7 @@
 #include "Evaluation/MovieSceneEvaluationTemplateInstance.h"
 #include "SequencerSettings.h"
 #include "SequencerInfoColumn.h"
+#include "SequencerSpawnableColumn.h"
 #include "LevelEditorViewport.h"
 #include "Modules/ModuleManager.h"
 #include "Framework/MultiBox/MultiBoxBuilder.h"
@@ -1091,9 +1092,27 @@ TSharedRef< ISceneOutlinerColumn > FLevelEditorSequencerIntegration::CreateSeque
 	return MakeShareable( new Sequencer::FSequencerInfoColumn( SceneOutliner, *BoundSequencers[0].Sequencer.Pin(), BoundSequencers[0].BindingData.Get() ) );
 }
 
+TSharedRef< ISceneOutlinerColumn > FLevelEditorSequencerIntegration::CreateSequencerSpawnableColumn( ISceneOutliner& SceneOutliner ) const
+{
+	//@todo only supports the first bound sequencer
+	check(BoundSequencers.Num() > 0);
+	check(BoundSequencers[0].Sequencer.IsValid());
+
+	return MakeShareable( new Sequencer::FSequencerSpawnableColumn() );
+}
 
 void FLevelEditorSequencerIntegration::AttachOutlinerColumn()
 {
+	// Register Spawnable Column 
+	FSceneOutlinerModule& SceneOutlinerModule = FModuleManager::LoadModuleChecked< FSceneOutlinerModule >("SceneOutliner");
+
+	FSceneOutlinerColumnInfo SpawnColumnInfo(ESceneOutlinerColumnVisibility::Visible, 11, 
+		FCreateSceneOutlinerColumn::CreateRaw( this, &FLevelEditorSequencerIntegration::CreateSequencerSpawnableColumn));
+
+	SceneOutlinerModule.RegisterDefaultColumnType< Sequencer::FSequencerSpawnableColumn >(SpawnColumnInfo);
+	AcquiredResources.Add([=]{ this->DetachOutlinerColumn(); });
+
+	// Register Info (Level Sequence) Column only if the preference is turned on
 	for (const FSequencerAndOptions& SequencerAndOptions : BoundSequencers)
 	{
 		TSharedPtr<FSequencer> Pinned = SequencerAndOptions.Sequencer.Pin();
@@ -1106,20 +1125,18 @@ void FLevelEditorSequencerIntegration::AttachOutlinerColumn()
 		}
 	}
 
-	FSceneOutlinerModule& SceneOutlinerModule = FModuleManager::LoadModuleChecked< FSceneOutlinerModule >("SceneOutliner");
-
 	FSceneOutlinerColumnInfo ColumnInfo(ESceneOutlinerColumnVisibility::Visible, 15, 
 		FCreateSceneOutlinerColumn::CreateRaw( this, &FLevelEditorSequencerIntegration::CreateSequencerInfoColumn));
 
 	SceneOutlinerModule.RegisterDefaultColumnType< Sequencer::FSequencerInfoColumn >(ColumnInfo);
 
-	AcquiredResources.Add([=]{ this->DetachOutlinerColumn(); });
 }
 
 void FLevelEditorSequencerIntegration::DetachOutlinerColumn()
 {
 	FSceneOutlinerModule& SceneOutlinerModule = FModuleManager::LoadModuleChecked< FSceneOutlinerModule >("SceneOutliner");
 
+	SceneOutlinerModule.UnRegisterColumnType< Sequencer::FSequencerSpawnableColumn >();
 	SceneOutlinerModule.UnRegisterColumnType< Sequencer::FSequencerInfoColumn >();
 
 	FLevelEditorModule& LevelEditorModule = FModuleManager::LoadModuleChecked<FLevelEditorModule>("LevelEditor");
diff --git a/Engine/Source/Editor/Sequencer/Private/SequencerSpawnableColumn.cpp b/Engine/Source/Editor/Sequencer/Private/SequencerSpawnableColumn.cpp
new file mode 100644
--- /dev/null
+++ b/Engine/Source/Editor/Sequencer/Private/SequencerSpawnableColumn.cpp
@@ -0,0 +1,69 @@
+// Copyright Epic Games, Inc. All Rights Reserved.
+
+#include "SequencerSpawnableColumn.h"
+#include "ActorTreeItem.h"
+
+namespace Sequencer
+{
+
+FName FSequencerSpawnableColumn::GetID()
+{
+	static FName IDName("Spawnable");
+	return IDName;
+}
+
+FName FSequencerSpawnableColumn::GetColumnID()
+{
+	return GetID();
+}
+
+SHeaderRow::FColumn::FArguments FSequencerSpawnableColumn::ConstructHeaderRowColumn()
+{
+	return SHeaderRow::Column(GetColumnID())
+		.FixedWidth(12.f)
+		.HAlignHeader(HAlign_Left) // Note the icon has some baked in right alignment, so we are compensating here
+		.VAlignHeader(VAlign_Center)
+		.HAlignCell(HAlign_Center)
+		.VAlignCell(VAlign_Center)
+		.HeaderContentPadding(FMargin(0.0))
+		[
+			SNew(SImage)
+			.Image(FAppStyle::Get().GetBrush("Sequencer.SpawnableIconOverlay"))
+		];
+}
+
+const TSharedRef< SWidget > FSequencerSpawnableColumn::ConstructRowWidget(FSceneOutlinerTreeItemRef TreeItem, const STableRow<FSceneOutlinerTreeItemPtr>& Row)
+{
+	TSharedRef<SHorizontalBox> HorizontalBox = SNew(SHorizontalBox);
+
+	HorizontalBox->AddSlot()
+	.AutoWidth()
+	.VAlign(VAlign_Center)
+	.HAlign(HAlign_Center)
+	[
+		SNew(SImage)
+		.Image(this, &FSequencerSpawnableColumn::GetSpawnableIcon, TreeItem)
+	];
+
+	return HorizontalBox;
+}
+
+const FSlateBrush* FSequencerSpawnableColumn::GetSpawnableIcon( FSceneOutlinerTreeItemRef TreeItem ) const
+{
+	static const FSlateBrush* SpawnedBrush  = FAppStyle::Get().GetBrush("Sequencer.SpawnableIconOverlay");
+	static const FSlateBrush* NoSpawnedBrush  = FAppStyle::Get().GetBrush("NoBrush");
+
+	bool bIsSpawned = false;
+	{
+		static const FName SequencerActorTag(TEXT("SequencerActor"));
+		if (const FActorTreeItem* ActorItem = (*TreeItem).CastTo<FActorTreeItem>())
+		{
+			AActor* Actor = ActorItem->Actor.Get();
+			bIsSpawned = Actor && Actor->ActorHasTag(SequencerActorTag);
+		}
+	}
+
+	return bIsSpawned ? SpawnedBrush : NoSpawnedBrush;
+}
+
+}// End Sequencer namespace
diff --git a/Engine/Source/Editor/Sequencer/Private/SequencerSpawnableColumn.h b/Engine/Source/Editor/Sequencer/Private/SequencerSpawnableColumn.h
new file mode 100644
--- /dev/null
+++ b/Engine/Source/Editor/Sequencer/Private/SequencerSpawnableColumn.h
@@ -0,0 +1,36 @@
+// Copyright Epic Games, Inc. All Rights Reserved.
+
+#pragma once
+
+#include "CoreMinimal.h"
+// #include "SceneOutlinerPublicTypes.h"
+#include "ISceneOutlinerColumn.h"
+
+class ISceneOutliner;
+
+namespace Sequencer
+{
+
+/**
+ * A custom column for the SceneOutliner to display if the Actor row is Spawnable 
+ */
+class FSequencerSpawnableColumn : public ISceneOutlinerColumn
+{
+
+public:
+	FSequencerSpawnableColumn() {}
+	FSequencerSpawnableColumn(ISceneOutliner& InSceneOutliner) {}
+
+	static FName GetID();
+
+	virtual FName GetColumnID() override;
+
+	virtual SHeaderRow::FColumn::FArguments ConstructHeaderRowColumn() override;
+	
+	virtual const TSharedRef< SWidget > ConstructRowWidget( FSceneOutlinerTreeItemRef TreeItem, const STableRow<FSceneOutlinerTreeItemPtr>& Row ) override;
+
+	const FSlateBrush* GetSpawnableIcon( FSceneOutlinerTreeItemRef TreeItem ) const;
+	
+};
+
+}
\ No newline at end of file
diff --git a/Engine/Source/Editor/Sequencer/Public/LevelEditorSequencerIntegration.h b/Engine/Source/Editor/Sequencer/Public/LevelEditorSequencerIntegration.h
--- a/Engine/Source/Editor/Sequencer/Public/LevelEditorSequencerIntegration.h
+++ b/Engine/Source/Editor/Sequencer/Public/LevelEditorSequencerIntegration.h
@@ -191,6 +191,7 @@ private:
 	TArray<FSequencerAndOptions> BoundSequencers;
 
 	TSharedRef< ISceneOutlinerColumn > CreateSequencerInfoColumn( ISceneOutliner& SceneOutliner ) const;
+	TSharedRef< ISceneOutlinerColumn > CreateSequencerSpawnableColumn( ISceneOutliner& SceneOutliner ) const;
 
 private:
 
