Sequencer Spawnable Column - Fixes alignment

diff --git a/Engine/Source/Editor/Sequencer/Private/SequencerSpawnableColumn.cpp b/Engine/Source/Editor/Sequencer/Private/SequencerSpawnableColumn.cpp
--- a/Engine/Source/Editor/Sequencer/Private/SequencerSpawnableColumn.cpp
+++ b/Engine/Source/Editor/Sequencer/Private/SequencerSpawnableColumn.cpp
@@ -20,10 +20,10 @@ FName FSequencerSpawnableColumn::GetColumnID()
 SHeaderRow::FColumn::FArguments FSequencerSpawnableColumn::ConstructHeaderRowColumn()
 {
 	return SHeaderRow::Column(GetColumnID())
-		.FixedWidth(12.f)
+		.FixedWidth(20.f)
 		.HAlignHeader(HAlign_Left) // Note the icon has some baked in right alignment, so we are compensating here
 		.VAlignHeader(VAlign_Center)
-		.HAlignCell(HAlign_Center)
+		.HAlignCell(HAlign_Left) // Note the icon has some baked in right alignment, so we are compensating here
 		.VAlignCell(VAlign_Center)
 		.HeaderContentPadding(FMargin(0.0))
 		[
@@ -34,18 +34,7 @@ SHeaderRow::FColumn::FArguments FSequencerSpawnableColumn::ConstructHeaderRowCol
 
 const TSharedRef< SWidget > FSequencerSpawnableColumn::ConstructRowWidget(FSceneOutlinerTreeItemRef TreeItem, const STableRow<FSceneOutlinerTreeItemPtr>& Row)
 {
-	TSharedRef<SHorizontalBox> HorizontalBox = SNew(SHorizontalBox);
-
-	HorizontalBox->AddSlot()
-	.AutoWidth()
-	.VAlign(VAlign_Center)
-	.HAlign(HAlign_Center)
-	[
-		SNew(SImage)
-		.Image(this, &FSequencerSpawnableColumn::GetSpawnableIcon, TreeItem)
-	];
-
-	return HorizontalBox;
+	return SNew(SImage).Image(this, &FSequencerSpawnableColumn::GetSpawnableIcon, TreeItem);
 }
 
 const FSlateBrush* FSequencerSpawnableColumn::GetSpawnableIcon( FSceneOutlinerTreeItemRef TreeItem ) const
