Non-unity fix

related to

diff --git a/Engine/Source/Editor/Sequencer/Private/SequencerSpawnableColumn.cpp b/Engine/Source/Editor/Sequencer/Private/SequencerSpawnableColumn.cpp
--- a/Engine/Source/Editor/Sequencer/Private/SequencerSpawnableColumn.cpp
+++ b/Engine/Source/Editor/Sequencer/Private/SequencerSpawnableColumn.cpp
@@ -2,6 +2,8 @@
 
 #include "SequencerSpawnableColumn.h"
 #include "ActorTreeItem.h"
+#include "Widgets/Images/SImage.h"
+#include "GameFramework/Actor.h"
 
 namespace Sequencer
 {
