Sequencer: Fix mute row not restoring properly by ensuring only rewriting the mute and solo paths if they actually exist. Also, reparent to the root instead of null

diff --git a/Engine/Source/Editor/Sequencer/Private/Sequencer.cpp b/Engine/Source/Editor/Sequencer/Private/Sequencer.cpp
--- a/Engine/Source/Editor/Sequencer/Private/Sequencer.cpp
+++ b/Engine/Source/Editor/Sequencer/Private/Sequencer.cpp
@@ -6950,8 +6950,11 @@ void FSequencer::OnNodePathChanged(const FString& OldPath, const FString& NewPat
 			if (NewNodePath.RemoveFromStart(PathPrefix))
 			{
 				NewNodePath = NewPath + '.' + NewNodePath;
-				MovieScene->GetSoloNodes().Remove(NodePath);
-				MovieScene->GetSoloNodes().Add(NewNodePath);
+				if (NodeTree->GetNodeAtPath(NewNodePath))
+				{
+					MovieScene->GetSoloNodes().Remove(NodePath);
+					MovieScene->GetSoloNodes().Add(NewNodePath);
+				}
 			}
 		}
 
@@ -6970,8 +6973,11 @@ void FSequencer::OnNodePathChanged(const FString& OldPath, const FString& NewPat
 			if (NewNodePath.RemoveFromStart(PathPrefix))
 			{
 				NewNodePath = NewPath + '.' + NewNodePath;
-				MovieScene->GetMuteNodes().Remove(NodePath);
-				MovieScene->GetMuteNodes().Add(NewNodePath);
+				if (NodeTree->GetNodeAtPath(NewNodePath))
+				{
+					MovieScene->GetMuteNodes().Remove(NodePath);
+					MovieScene->GetMuteNodes().Add(NewNodePath);
+				}
 			}
 		}
 	}
diff --git a/Engine/Source/Editor/Sequencer/Private/SequencerNodeTree.cpp b/Engine/Source/Editor/Sequencer/Private/SequencerNodeTree.cpp
--- a/Engine/Source/Editor/Sequencer/Private/SequencerNodeTree.cpp
+++ b/Engine/Source/Editor/Sequencer/Private/SequencerNodeTree.cpp
@@ -616,8 +616,8 @@ void FSequencerNodeTree::MoveDisplayNodeToRoot(TSharedRef<FSequencerDisplayNode>
 		}
 	}
 
-	// Clear the node's parent so that subsequent calls for GetNodePath correctly indicate that they no longer have a parent.
-	Node->SetParent(nullptr);
+	// Reset to the root node
+	Node->SetParent(RootNode);
 
 	// Our children have changed parents which means that on subsequent creation they will retrieve their expansion state
 	// from the map using their new path. If the new path already exists the object goes to the state stored at that path.
