Fix for wrong capitilization in include file name
#rnx

diff --git a/Engine/Source/Editor/Sequencer/Private/Tools/SMotionTrailOptions.cpp b/Engine/Source/Editor/Sequencer/Private/Tools/SMotionTrailOptions.cpp
--- a/Engine/Source/Editor/Sequencer/Private/Tools/SMotionTrailOptions.cpp
+++ b/Engine/Source/Editor/Sequencer/Private/Tools/SMotionTrailOptions.cpp
@@ -2,7 +2,7 @@
 
 #include "Tools/SMotionTrailOptions.h"
 #include "Tools/MotionTrailOptions.h"
-#include "Modules/moduleManager.h"
+#include "Modules/ModuleManager.h"
 #include "PropertyEditorModule.h"
 #include "Widgets/Layout/SBorder.h"
 
