Sequencer: Fix toolbar so that it's not in the splitter. Otherwise the splitter overlaps the toolbar buttons.

diff --git a/Engine/Source/Editor/Sequencer/Private/SSequencer.cpp b/Engine/Source/Editor/Sequencer/Private/SSequencer.cpp
--- a/Engine/Source/Editor/Sequencer/Private/SSequencer.cpp
+++ b/Engine/Source/Editor/Sequencer/Private/SSequencer.cpp
@@ -562,185 +562,184 @@ void SSequencer::Construct(const FArguments& InArgs, TSharedRef<FSequencer> InSe
 	[
 		SNew(SVerticalBox)
 
+		// Toolbar
 		+ SVerticalBox::Slot()
+		.AutoHeight()
 		[
-			SNew(SSplitter)
-			.Orientation(Orient_Horizontal)
-			
-			+ SSplitter::Slot()
-			.Value(0.9f)
+			SNew(SBox)
+			.Padding(FMargin(CommonPadding,0.0f,0.0f,0.f))
 			[
-				SNew(SOverlay)
+				SNew(SWrapBox)
+				.UseAllottedSize(true)
+				.InnerSlotPadding(FVector2D(5, 0))
+				+ SWrapBox::Slot()
+				.FillEmptySpace(true)
+				.FillLineWhenSizeLessThan(600)
+				[
+					SAssignNew(ToolbarContainer, SBox)
+				]
 
-				+ SOverlay::Slot()
+				+ SWrapBox::Slot()
+				.FillEmptySpace(true)
 				[
-					// track area grid panel
-					SNew( SGridPanel )
-					.FillRow( 2, 1.f )
-					.FillColumn( 0, FillCoefficient_0 )
-					.FillColumn( 1, FillCoefficient_1 )
+					SNew(SHorizontalBox)
+					// Right Aligned Breadcrumbs
+					+ SHorizontalBox::Slot()
+					.HAlign(HAlign_Right)
+					.VAlign(VAlign_Center)
+					[
+						SNew(SSpacer)
+					]
 
-					// Toolbar
-					+ SGridPanel::Slot( Column0, Row0, SGridPanel::Layer(10) )
-					.ColumnSpan(2)
-					.Padding(4.f, 0.f, 8.f, 0.f)
+					// History Back Button
+					+ SHorizontalBox::Slot()
+					.HAlign(HAlign_Right)
+					.VAlign(VAlign_Center)
+					.AutoWidth()
 					[
-						SNew(SBox)
-						.Padding(FMargin(CommonPadding,0.0f,0.0f,0.f))
-						[
-							SNew(SWrapBox)
-							.UseAllottedSize(true)
-							.InnerSlotPadding(FVector2D(5, 0))
-							+ SWrapBox::Slot()
-							.FillEmptySpace(true)
-							.FillLineWhenSizeLessThan(600)
-							[
-								SAssignNew(ToolbarContainer, SBox)
-							]
+						SNew(SVerticalBox)
 
-							+ SWrapBox::Slot()
-							.FillEmptySpace(true)
+						+ SVerticalBox::Slot()
+						.FillHeight(1.0f)
+						[
+							SNew(SButton)
+							.Visibility_Lambda([this] { return CanNavigateBreadcrumbs() ? EVisibility::Visible : EVisibility::Collapsed; } )
+							.VAlign(EVerticalAlignment::VAlign_Center)
+							.ButtonStyle(FEditorStyle::Get(), "SimpleButton")
+							.ToolTipText_Lambda([this] { return SequencerPtr.Pin()->GetNavigateBackwardTooltip(); })
+							.ContentPadding(FMargin(1, 0))
+							.OnClicked_Lambda([this] { return SequencerPtr.Pin()->NavigateBackward(); })
+							.IsEnabled_Lambda([this] { return SequencerPtr.Pin()->CanNavigateBackward(); })
 							[
-								SNew(SHorizontalBox)
-								// Right Aligned Breadcrumbs
-								+ SHorizontalBox::Slot()
-								.HAlign(HAlign_Right)
-								.VAlign(VAlign_Center)
+								SNew(SBox) // scale up since the default icons are 16x16
+								.WidthOverride(20)
+								.HeightOverride(20)
 								[
-									SNew(SSpacer)
+									SNew(SImage)
+									.ColorAndOpacity(FSlateColor::UseForeground())
+									.Image(FAppStyle::Get().GetBrush("Icons.ArrowLeft"))
 								]
+							]
+						]
+					]
 
-								// History Back Button
-								+ SHorizontalBox::Slot()
-								.HAlign(HAlign_Right)
-								.VAlign(VAlign_Center)
-								.AutoWidth()
-								[
-									SNew(SVerticalBox)
-
-									+ SVerticalBox::Slot()
-									.FillHeight(1.0f)
-									[
-										SNew(SButton)
-										.Visibility_Lambda([this] { return CanNavigateBreadcrumbs() ? EVisibility::Visible : EVisibility::Collapsed; } )
-										.VAlign(EVerticalAlignment::VAlign_Center)
-										.ButtonStyle(FEditorStyle::Get(), "SimpleButton")
-										.ToolTipText_Lambda([this] { return SequencerPtr.Pin()->GetNavigateBackwardTooltip(); })
-										.ContentPadding(FMargin(1, 0))
-										.OnClicked_Lambda([this] { return SequencerPtr.Pin()->NavigateBackward(); })
-										.IsEnabled_Lambda([this] { return SequencerPtr.Pin()->CanNavigateBackward(); })
-										[
-											SNew(SBox) // scale up since the default icons are 16x16
-											.WidthOverride(20)
-											.HeightOverride(20)
-											[
-												SNew(SImage)
-												.ColorAndOpacity(FSlateColor::UseForeground())
-												.Image(FAppStyle::Get().GetBrush("Icons.ArrowLeft"))
-											]
-										]
-									]
-								]
+					// History Forward Button
+					+ SHorizontalBox::Slot()
+					.HAlign(HAlign_Right)
+					.VAlign(VAlign_Center)
+					.AutoWidth()
+					[
+						SNew(SVerticalBox)
 
-								// History Forward Button
-								+ SHorizontalBox::Slot()
-								.HAlign(HAlign_Right)
-								.VAlign(VAlign_Center)
-								.AutoWidth()
+						+ SVerticalBox::Slot()
+						.FillHeight(1.0f)
+						[
+							SNew(SButton)
+							.Visibility_Lambda([this] { return CanNavigateBreadcrumbs() ? EVisibility::Visible : EVisibility::Collapsed; } )
+							.VAlign(EVerticalAlignment::VAlign_Center)
+							.ButtonStyle(FEditorStyle::Get(), "SimpleButton")
+							.ToolTipText_Lambda([this] { return SequencerPtr.Pin()->GetNavigateForwardTooltip(); })
+							.ContentPadding(FMargin(1, 0))
+							.OnClicked_Lambda([this] { return SequencerPtr.Pin()->NavigateForward(); })
+							.IsEnabled_Lambda([this] { return SequencerPtr.Pin()->CanNavigateForward(); })
+							[
+								SNew(SBox) // scale up since the default icons are 16x16
+								.WidthOverride(20)
+								.HeightOverride(20)
 								[
-									SNew(SVerticalBox)
-
-									+ SVerticalBox::Slot()
-									.FillHeight(1.0f)
-									[
-										SNew(SButton)
-										.Visibility_Lambda([this] { return CanNavigateBreadcrumbs() ? EVisibility::Visible : EVisibility::Collapsed; } )
-										.VAlign(EVerticalAlignment::VAlign_Center)
-										.ButtonStyle(FEditorStyle::Get(), "SimpleButton")
-										.ToolTipText_Lambda([this] { return SequencerPtr.Pin()->GetNavigateForwardTooltip(); })
-										.ContentPadding(FMargin(1, 0))
-										.OnClicked_Lambda([this] { return SequencerPtr.Pin()->NavigateForward(); })
-										.IsEnabled_Lambda([this] { return SequencerPtr.Pin()->CanNavigateForward(); })
-										[
-											SNew(SBox) // scale up since the default icons are 16x16
-											.WidthOverride(20)
-											.HeightOverride(20)
-											[
-												SNew(SImage)
-												.ColorAndOpacity(FSlateColor::UseForeground())
-												.Image(FAppStyle::Get().GetBrush("Icons.ArrowRight"))
-											]
-										]
-									]
+									SNew(SImage)
+									.ColorAndOpacity(FSlateColor::UseForeground())
+									.Image(FAppStyle::Get().GetBrush("Icons.ArrowRight"))
 								]
+							]
+						]
+					]
 
-								// Separator
-								+ SHorizontalBox::Slot()
-								.AutoWidth()
-								.Padding(3, 0)
-								[
-									SNew(SSeparator)
-									.Visibility_Lambda([this] { return CanNavigateBreadcrumbs() ? EVisibility::Visible : EVisibility::Collapsed; } )
-									.Orientation(Orient_Vertical)
-								]
+					// Separator
+					+ SHorizontalBox::Slot()
+					.AutoWidth()
+					.Padding(3, 0)
+					[
+						SNew(SSeparator)
+						.Visibility_Lambda([this] { return CanNavigateBreadcrumbs() ? EVisibility::Visible : EVisibility::Collapsed; } )
+						.Orientation(Orient_Vertical)
+					]
 							
-								+ SHorizontalBox::Slot()
-								.HAlign(HAlign_Right)
-								.VAlign(VAlign_Center)
-								.AutoWidth()
-								[
-									SAssignNew(BreadcrumbPickerButton, SComboButton)
-									.Visibility_Lambda([this] { return CanNavigateBreadcrumbs() ? EVisibility::Visible : EVisibility::Collapsed; } )
-									.ButtonStyle(FAppStyle::Get(), "SimpleButton")
-									.OnGetMenuContent_Lambda([this] { return SNew(SSequencerHierarchyBrowser, SequencerPtr); })
-									.HasDownArrow(false)
-									.ContentPadding(FMargin(3, 3))
-									.ButtonContent()
-									[
-										SNew(SImage)
-										.ColorAndOpacity(FSlateColor::UseForeground())
-										.Image(FAppStyle::Get().GetBrush("Icons.FolderOpen"))
-									]
-								]
+					+ SHorizontalBox::Slot()
+					.HAlign(HAlign_Right)
+					.VAlign(VAlign_Center)
+					.AutoWidth()
+					[
+						SAssignNew(BreadcrumbPickerButton, SComboButton)
+						.Visibility_Lambda([this] { return CanNavigateBreadcrumbs() ? EVisibility::Visible : EVisibility::Collapsed; } )
+						.ButtonStyle(FAppStyle::Get(), "SimpleButton")
+						.OnGetMenuContent_Lambda([this] { return SNew(SSequencerHierarchyBrowser, SequencerPtr); })
+						.HasDownArrow(false)
+						.ContentPadding(FMargin(3, 3))
+						.ButtonContent()
+						[
+							SNew(SImage)
+							.ColorAndOpacity(FSlateColor::UseForeground())
+							.Image(FAppStyle::Get().GetBrush("Icons.FolderOpen"))
+						]
+					]
 
-								// Right Aligned Breadcrumbs
-								+ SHorizontalBox::Slot()
-								.HAlign(HAlign_Right)
-								.VAlign(VAlign_Center)
-								.AutoWidth()
-								[
-									SAssignNew(BreadcrumbTrail, SBreadcrumbTrail<FSequencerBreadcrumb>)
-									.Visibility(this, &SSequencer::GetBreadcrumbTrailVisibility)
-									.OnCrumbClicked(this, &SSequencer::OnCrumbClicked)
-									.ButtonStyle(FAppStyle::Get(), "SimpleButton")
-									.DelimiterImage(FAppStyle::Get().GetBrush("Sequencer.BreadcrumbIcon"))
-									.TextStyle(FAppStyle::Get(), "Sequencer.BreadcrumbText")
-								]
+					// Right Aligned Breadcrumbs
+					+ SHorizontalBox::Slot()
+					.HAlign(HAlign_Right)
+					.VAlign(VAlign_Center)
+					.AutoWidth()
+					[
+						SAssignNew(BreadcrumbTrail, SBreadcrumbTrail<FSequencerBreadcrumb>)
+						.Visibility(this, &SSequencer::GetBreadcrumbTrailVisibility)
+						.OnCrumbClicked(this, &SSequencer::OnCrumbClicked)
+						.ButtonStyle(FAppStyle::Get(), "SimpleButton")
+						.DelimiterImage(FAppStyle::Get().GetBrush("Sequencer.BreadcrumbIcon"))
+						.TextStyle(FAppStyle::Get(), "Sequencer.BreadcrumbText")
+					]
 
-								// Sequence Locking symbol
-								+SHorizontalBox::Slot()
-								.HAlign(HAlign_Right)
-								.VAlign(VAlign_Center)
-								.AutoWidth()
-								[
-									SNew(SCheckBox)
-									.Style(FAppStyle::Get(),"ToggleButtonCheckBoxAlt")
-									.Type(ESlateCheckBoxType::CheckBox) // Use CheckBox instead of ToggleType since we're not putting ohter widget inside
-									.Padding(FMargin(0.f))
-									.IsFocusable(false)		
-									.IsChecked_Lambda([this] { return GetIsSequenceReadOnly() ? ECheckBoxState::Checked : ECheckBoxState::Unchecked; } )
-									.OnCheckStateChanged(this, &SSequencer::OnSetSequenceReadOnly)
-									.ToolTipText_Lambda([this] { return GetIsSequenceReadOnly() ? LOCTEXT("UnlockSequence", "Unlock the animation so that it is editable") : LOCTEXT("LockSequence", "Lock the animation so that it is not editable"); } )
-									.CheckedImage(FAppStyle::Get().GetBrush("Icons.Lock"))
-									.CheckedHoveredImage(FAppStyle::Get().GetBrush("Icons.Lock"))
-									.CheckedPressedImage(FAppStyle::Get().GetBrush("Icons.Lock"))
-									.UncheckedImage(FAppStyle::Get().GetBrush("Icons.Unlock"))
-									.UncheckedHoveredImage(FAppStyle::Get().GetBrush("Icons.Unlock"))
-									.UncheckedPressedImage(FAppStyle::Get().GetBrush("Icons.Unlock"))
-								]
-							]
-						]
+					// Sequence Locking symbol
+					+SHorizontalBox::Slot()
+					.HAlign(HAlign_Right)
+					.VAlign(VAlign_Center)
+					.AutoWidth()
+					[
+						SNew(SCheckBox)
+						.Style(FAppStyle::Get(),"ToggleButtonCheckBoxAlt")
+						.Type(ESlateCheckBoxType::CheckBox) // Use CheckBox instead of ToggleType since we're not putting ohter widget inside
+						.Padding(FMargin(0.f))
+						.IsFocusable(false)		
+						.IsChecked_Lambda([this] { return GetIsSequenceReadOnly() ? ECheckBoxState::Checked : ECheckBoxState::Unchecked; } )
+						.OnCheckStateChanged(this, &SSequencer::OnSetSequenceReadOnly)
+						.ToolTipText_Lambda([this] { return GetIsSequenceReadOnly() ? LOCTEXT("UnlockSequence", "Unlock the animation so that it is editable") : LOCTEXT("LockSequence", "Lock the animation so that it is not editable"); } )
+						.CheckedImage(FAppStyle::Get().GetBrush("Icons.Lock"))
+						.CheckedHoveredImage(FAppStyle::Get().GetBrush("Icons.Lock"))
+						.CheckedPressedImage(FAppStyle::Get().GetBrush("Icons.Lock"))
+						.UncheckedImage(FAppStyle::Get().GetBrush("Icons.Unlock"))
+						.UncheckedHoveredImage(FAppStyle::Get().GetBrush("Icons.Unlock"))
+						.UncheckedPressedImage(FAppStyle::Get().GetBrush("Icons.Unlock"))
 					]
+				]
+			]
+		]
+
+		+ SVerticalBox::Slot()
+		[
+			SNew(SSplitter)
+			.Orientation(Orient_Horizontal)
+			
+			+ SSplitter::Slot()
+			.Value(0.9f)
+			[
+				SNew(SOverlay)
+
+				+ SOverlay::Slot()
+				[
+					// track area grid panel
+					SNew( SGridPanel )
+					.FillRow( 2, 1.f )
+					.FillColumn( 0, FillCoefficient_0 )
+					.FillColumn( 1, FillCoefficient_1 )
 
 					+ SGridPanel::Slot( Column0, Row1 )
 					[
