Sequencer: Set the end range of newly created sections as the playback end frame rather than a percentage of the view range

diff --git a/Engine/Source/Editor/Sequencer/Private/SequencerUtilities.cpp b/Engine/Source/Editor/Sequencer/Private/SequencerUtilities.cpp
--- a/Engine/Source/Editor/Sequencer/Private/SequencerUtilities.cpp
+++ b/Engine/Source/Editor/Sequencer/Private/SequencerUtilities.cpp
@@ -15,6 +15,7 @@
 #include "SequencerTrackNode.h"
 #include "MovieSceneTrack.h"
 #include "MovieSceneSection.h"
+#include "MovieSceneTimeHelpers.h"
 #include "Framework/MultiBox/MultiBoxBuilder.h"
 #include "ISequencerTrackEditor.h"
 #include "ISequencer.h"
@@ -109,7 +110,7 @@ void FSequencerUtilities::PopulateMenu_CreateNewSection(FMenuBuilder& MenuBuilde
 		}
 
 		FQualifiedFrameTime CurrentTime = Sequencer->GetLocalTime();
-		TRange<double> VisibleRange = Sequencer->GetViewRange();
+		FFrameNumber PlaybackEnd = UE::MovieScene::DiscreteExclusiveUpper(Sequencer->GetFocusedMovieSceneSequence()->GetMovieScene()->GetPlaybackRange());
 
 		FScopedTransaction Transaction(LOCTEXT("AddSectionTransactionText", "Add Section"));
 		if (UMovieSceneSection* NewSection = Track->CreateNewSection())
@@ -128,8 +129,7 @@ void FSequencerUtilities::PopulateMenu_CreateNewSection(FMenuBuilder& MenuBuilde
 
 			Track->Modify();
 
-			int32 DurationFrames = ( (VisibleRange.Size<double>() * 0.75) * CurrentTime.Rate ).FloorToFrame().Value;
-			NewSection->SetRange(TRange<FFrameNumber>(CurrentTime.Time.FrameNumber, CurrentTime.Time.FrameNumber + DurationFrames));
+			NewSection->SetRange(TRange<FFrameNumber>(CurrentTime.Time.FrameNumber, PlaybackEnd));
 			NewSection->SetOverlapPriority(OverlapPriority);
 			NewSection->SetRowIndex(RowIndex);
 			NewSection->SetBlendType(BlendType);
