Sequencer: Fix copy/paste tracks so that it only replaces an existing track if it has the same display name.

diff --git a/Engine/Source/Editor/Sequencer/Private/Sequencer.cpp b/Engine/Source/Editor/Sequencer/Private/Sequencer.cpp
--- a/Engine/Source/Editor/Sequencer/Private/Sequencer.cpp
+++ b/Engine/Source/Editor/Sequencer/Private/Sequencer.cpp
@@ -9516,7 +9516,7 @@ bool FSequencer::PasteTracks(const FString& TextToImport, UMovieSceneFolder* InP
 							// Tracks of the same class should be unique per name.
 							for (UMovieSceneTrack* Track : Binding.GetTracks())
 							{
-								if (Track->GetClass() == NewTrack->GetClass() && Track->GetTrackName() == NewTrack->GetTrackName())
+								if (Track->GetClass() == NewTrack->GetClass() && Track->GetTrackName() == NewTrack->GetTrackName() && Track->GetDisplayName().IdenticalTo(NewTrack->GetDisplayName()))
 								{
 									// If a track of the same class and name exists, remove it so the new track replaces it
 									MovieScene->RemoveTrack(*Track);
