Sequencer: Remove tangent state of section keys for performance reasons. It's not worth querying all the keys.

diff --git a/Engine/Source/Editor/Sequencer/Private/SequencerContextMenus.cpp b/Engine/Source/Editor/Sequencer/Private/SequencerContextMenus.cpp
--- a/Engine/Source/Editor/Sequencer/Private/SequencerContextMenus.cpp
+++ b/Engine/Source/Editor/Sequencer/Private/SequencerContextMenus.cpp
@@ -512,10 +512,7 @@ void FSectionContextMenu::AddEditMenu(FMenuBuilder& MenuBuilder)
 		FSlateIcon(FEditorStyle::GetStyleSetName(), "Sequencer.IconKeyAuto"),
 		FUIAction(
 			FExecuteAction::CreateLambda([=]{ Shared->SetInterpTangentMode(RCIM_Cubic, RCTM_Auto); }),
-			FCanExecuteAction::CreateLambda([=]{ return Shared->CanSetInterpTangentMode(); }),
-			FIsActionChecked::CreateLambda([=]{ return Shared->IsInterpTangentModeSelected(RCIM_Cubic, RCTM_Auto); }) ),
-		NAME_None,
-		EUserInterfaceActionType::ToggleButton
+			FCanExecuteAction::CreateLambda([=]{ return Shared->CanSetInterpTangentMode(); }) )
 	);
 
 	MenuBuilder.AddMenuEntry(
@@ -524,10 +521,7 @@ void FSectionContextMenu::AddEditMenu(FMenuBuilder& MenuBuilder)
 		FSlateIcon(FEditorStyle::GetStyleSetName(), "Sequencer.IconKeyUser"),
 		FUIAction(
 			FExecuteAction::CreateLambda([=]{ Shared->SetInterpTangentMode(RCIM_Cubic, RCTM_User); }),
-			FCanExecuteAction::CreateLambda([=]{ return Shared->CanSetInterpTangentMode(); }),
-			FIsActionChecked::CreateLambda([=]{ return Shared->IsInterpTangentModeSelected(RCIM_Cubic, RCTM_User); }) ),
-		NAME_None,
-		EUserInterfaceActionType::ToggleButton
+			FCanExecuteAction::CreateLambda([=]{ return Shared->CanSetInterpTangentMode(); }) )
 	);
 
 	MenuBuilder.AddMenuEntry(
@@ -536,10 +530,7 @@ void FSectionContextMenu::AddEditMenu(FMenuBuilder& MenuBuilder)
 		FSlateIcon(FEditorStyle::GetStyleSetName(), "Sequencer.IconKeyBreak"),
 		FUIAction(
 			FExecuteAction::CreateLambda([=]{ Shared->SetInterpTangentMode(RCIM_Cubic, RCTM_Break); }),
-			FCanExecuteAction::CreateLambda([=]{ return Shared->CanSetInterpTangentMode(); }),
-			FIsActionChecked::CreateLambda([=]{ return Shared->IsInterpTangentModeSelected(RCIM_Cubic, RCTM_Break); }) ),
-		NAME_None,
-		EUserInterfaceActionType::ToggleButton
+			FCanExecuteAction::CreateLambda([=]{ return Shared->CanSetInterpTangentMode(); }) )
 	);
 
 	MenuBuilder.AddMenuEntry(
@@ -548,10 +539,7 @@ void FSectionContextMenu::AddEditMenu(FMenuBuilder& MenuBuilder)
 		FSlateIcon(FEditorStyle::GetStyleSetName(), "Sequencer.IconKeyLinear"),
 		FUIAction(
 			FExecuteAction::CreateLambda([=]{ Shared->SetInterpTangentMode(RCIM_Linear, RCTM_Auto); }),
-			FCanExecuteAction::CreateLambda([=]{ return Shared->CanSetInterpTangentMode(); }),
-			FIsActionChecked::CreateLambda([=]{ return Shared->IsInterpTangentModeSelected(RCIM_Linear, RCTM_Auto); }) ),
-		NAME_None,
-		EUserInterfaceActionType::ToggleButton
+			FCanExecuteAction::CreateLambda([=]{ return Shared->CanSetInterpTangentMode(); }) )
 	);
 
 	MenuBuilder.AddMenuEntry(
@@ -560,10 +548,7 @@ void FSectionContextMenu::AddEditMenu(FMenuBuilder& MenuBuilder)
 		FSlateIcon(FEditorStyle::GetStyleSetName(), "Sequencer.IconKeyConstant"),
 		FUIAction(
 			FExecuteAction::CreateLambda([=]{ Shared->SetInterpTangentMode(RCIM_Constant, RCTM_Auto); }),
-			FCanExecuteAction::CreateLambda([=]{ return Shared->CanSetInterpTangentMode(); }),
-			FIsActionChecked::CreateLambda([=]{ return Shared->IsInterpTangentModeSelected(RCIM_Constant, RCTM_Auto); }) ),
-		NAME_None,
-		EUserInterfaceActionType::ToggleButton
+			FCanExecuteAction::CreateLambda([=]{ return Shared->CanSetInterpTangentMode(); }) )
 	);
 
 	MenuBuilder.EndSection(); // SequencerInterpolation
@@ -981,54 +966,7 @@ bool FSectionContextMenu::CanSetInterpTangentMode() const
 	}
 
 	return false;
-}
-				
-
-bool FSectionContextMenu::IsInterpTangentModeSelected(ERichCurveInterpMode InterpMode, ERichCurveTangentMode TangentMode) const
-{
-	TSet<TSharedPtr<IKeyArea> > KeyAreas;
-	for (const TSharedRef<FSequencerDisplayNode>& DisplayNode : Sequencer->GetSelection().GetSelectedOutlinerNodes())
-	{
-		SequencerHelpers::GetAllKeyAreas(DisplayNode, KeyAreas);
-	}
-
-	if (KeyAreas.Num() == 0)
-	{
-		const TSet<TSharedRef<FSequencerDisplayNode>>& SelectedNodes = Sequencer->GetSelection().GetNodesWithSelectedKeysOrSections();
-		for (const TSharedRef<FSequencerDisplayNode>& DisplayNode : SelectedNodes)
-		{
-			SequencerHelpers::GetAllKeyAreas(DisplayNode, KeyAreas);
-		}
-	}
-
-	int32 NumKeys = 0;
-	for (TSharedPtr<IKeyArea> KeyArea : KeyAreas)
-	{
-		if (KeyArea.IsValid())
-		{
-			UMovieSceneSection* Section = KeyArea->GetOwningSection();
-			if (Section)
-			{
-				for (FMovieSceneFloatChannel* FloatChannel : Section->GetChannelProxy().GetChannels<FMovieSceneFloatChannel>())
-				{
-					TMovieSceneChannelData<FMovieSceneFloatValue> ChannelData = FloatChannel->GetData();
-					TArrayView<FMovieSceneFloatValue> Values = ChannelData.GetValues();
-
-					NumKeys += FloatChannel->GetNumKeys();
-					for (int32 KeyIndex = 0; KeyIndex < FloatChannel->GetNumKeys(); ++KeyIndex)
-					{
-						if (Values[KeyIndex].InterpMode != InterpMode || Values[KeyIndex].TangentMode != TangentMode)
-						{
-							return false;
-						}
-					}
-				}
-			}
-		}
-	}
-
-	return NumKeys != 0;
-}
+}			
 
 
 void FSectionContextMenu::ToggleSectionActive()
diff --git a/Engine/Source/Editor/Sequencer/Private/SequencerContextMenus.h b/Engine/Source/Editor/Sequencer/Private/SequencerContextMenus.h
--- a/Engine/Source/Editor/Sequencer/Private/SequencerContextMenus.h
+++ b/Engine/Source/Editor/Sequencer/Private/SequencerContextMenus.h
@@ -70,8 +70,6 @@ private:
 
 	bool CanSetInterpTangentMode() const;
 
-	bool IsInterpTangentModeSelected(ERichCurveInterpMode InterpMode, ERichCurveTangentMode TangentMode) const;
-
 	void ToggleSectionActive();
 
 	bool IsSectionActive() const;
