Sequencer: Add possessable to spawnable parent when copy/paste

diff --git a/Engine/Source/Editor/Sequencer/Private/Sequencer.cpp b/Engine/Source/Editor/Sequencer/Private/Sequencer.cpp
--- a/Engine/Source/Editor/Sequencer/Private/Sequencer.cpp
+++ b/Engine/Source/Editor/Sequencer/Private/Sequencer.cpp
@@ -9404,10 +9404,16 @@ bool FSequencer::PasteObjectBindings(const FString& TextToImport, UMovieSceneFol
 				}
 
 				// If the parent doesn't actually exist, clear it.
-				if (!MovieScene->FindPossessable(ParentGuid) && !MovieScene->FindSpawnable(ParentGuid))
+				FMovieScenePossessable* PossessableParent = MovieScene->FindPossessable(ParentGuid);
+				FMovieSceneSpawnable* SpawnableParent = MovieScene->FindSpawnable(ParentGuid);
+				if (!PossessableParent && !SpawnableParent)
 				{
 					Possessable->SetParent(FGuid());
 				}
+				else if (SpawnableParent)
+				{
+					SpawnableParent->AddChildPossessable(PossessableGuid);
+				}
 			}
 		}
 	}
