Sequencer: Speculative fix for UpdateOnHover for null MovieSceneChannel

diff --git a/Engine/Source/Editor/Sequencer/Private/SequencerHotspots.cpp b/Engine/Source/Editor/Sequencer/Private/SequencerHotspots.cpp
--- a/Engine/Source/Editor/Sequencer/Private/SequencerHotspots.cpp
+++ b/Engine/Source/Editor/Sequencer/Private/SequencerHotspots.cpp
@@ -73,7 +73,7 @@ void FSectionHotspot::UpdateOnHover(SSequencerTrackArea& InTrackArea, ISequencer
 		{
 			for (const FMovieSceneChannel* Channel : Entry.GetChannels())
 			{
-				if (Channel->GetNumKeys() != 0)
+				if (Channel && Channel->GetNumKeys() != 0)
 				{
 					InTrackArea.AttemptToActivateTool(FSequencerEditTool_Selection::Identifier);
 					return;
