Sequencer: Update easing when a section is deleted

diff --git a/Engine/Source/Editor/Sequencer/Private/Sequencer.cpp b/Engine/Source/Editor/Sequencer/Private/Sequencer.cpp
--- a/Engine/Source/Editor/Sequencer/Private/Sequencer.cpp
+++ b/Engine/Source/Editor/Sequencer/Private/Sequencer.cpp
@@ -1484,6 +1484,7 @@ void FSequencer::DeleteSections(const TSet<TWeakObjectPtr<UMovieSceneSection>>&
 			Track->SetFlags(RF_Transactional);
 			Track->Modify();
 			Track->RemoveSection(*Section);
+			Track->UpdateEasing();
 		}
 
 		bAnythingRemoved = true;
