Sequencer: Don't allow drawing the drop target on the spacer node's track lane

diff --git a/Engine/Source/Editor/Sequencer/Private/DisplayNodes/SequencerDisplayNode.h b/Engine/Source/Editor/Sequencer/Private/DisplayNodes/SequencerDisplayNode.h
--- a/Engine/Source/Editor/Sequencer/Private/DisplayNodes/SequencerDisplayNode.h
+++ b/Engine/Source/Editor/Sequencer/Private/DisplayNodes/SequencerDisplayNode.h
@@ -61,7 +61,9 @@ namespace ESequencerNode
 		/* Symbolic root node */
 		Root,
 		/* Folder node */
-		Folder
+		Folder,
+		/* Spacer node */
+		Spacer
 	};
 }
 
diff --git a/Engine/Source/Editor/Sequencer/Private/SSequencerTrackArea.cpp b/Engine/Source/Editor/Sequencer/Private/SSequencerTrackArea.cpp
--- a/Engine/Source/Editor/Sequencer/Private/SSequencerTrackArea.cpp
+++ b/Engine/Source/Editor/Sequencer/Private/SSequencerTrackArea.cpp
@@ -198,7 +198,7 @@ int32 SSequencerTrackArea::OnPaint(const FPaintArgs& Args, const FGeometry& Allo
 
 
 		// Draw drop target
-		if (DroppedNode.IsValid() && TrackSlots.Contains(DroppedNode.Pin()))
+		if (DroppedNode.IsValid() && DroppedNode.Pin()->GetType() != ESequencerNode::Spacer && TrackSlots.Contains(DroppedNode.Pin()))
 		{
 			TSharedPtr<SSequencerTrackLane> TrackLane = TrackSlots.FindRef(DroppedNode.Pin()).Pin();
 			
diff --git a/Engine/Source/Editor/Sequencer/Private/SequencerNodeTree.cpp b/Engine/Source/Editor/Sequencer/Private/SequencerNodeTree.cpp
--- a/Engine/Source/Editor/Sequencer/Private/SequencerNodeTree.cpp
+++ b/Engine/Source/Editor/Sequencer/Private/SequencerNodeTree.cpp
@@ -53,7 +53,7 @@ public:
 	virtual FText GetDisplayName() const override { return FText(); }
 	virtual float GetNodeHeight() const override { return Size; }
 	virtual FNodePadding GetNodePadding() const override { return FNodePadding(0.f); }
-	virtual ESequencerNode::Type GetType() const override { return ESequencerNode::Object; }
+	virtual ESequencerNode::Type GetType() const override { return ESequencerNode::Spacer; }
 	virtual void SetDisplayName(const FText& NewDisplayName) override { }
 	virtual TSharedRef<SWidget> GenerateContainerWidgetForOutliner(const TSharedRef<SSequencerTreeViewRow>& InRow) override { return SNew(SBox).HeightOverride(Size); }
 	virtual bool IsSelectable() const override { return false; }
