Sequencer: Move record button to the left and add some space so it's not centered and easily clicked on

diff --git a/Engine/Source/Editor/Sequencer/Private/Sequencer.cpp b/Engine/Source/Editor/Sequencer/Private/Sequencer.cpp
--- a/Engine/Source/Editor/Sequencer/Private/Sequencer.cpp
+++ b/Engine/Source/Editor/Sequencer/Private/Sequencer.cpp
@@ -4456,6 +4456,8 @@ TSharedRef<SWidget> FSequencer::MakeTransportControls(bool bExtended)
 
 	FTransportControlArgs TransportControlArgs;
 	{
+		TransportControlArgs.WidgetsToCreate.Add(FTransportControlWidget(FOnMakeTransportWidget::CreateSP(this, &FSequencer::OnCreateTransportRecord)));
+
 		TransportControlArgs.OnBackwardEnd.BindSP( this, &FSequencer::OnJumpToStart );
 		TransportControlArgs.OnBackwardStep.BindSP( this, &FSequencer::OnStepBackward, FFrameNumber(1) );
 		TransportControlArgs.OnForwardPlay.BindSP( this, &FSequencer::OnPlayForward, true );
@@ -4475,7 +4477,6 @@ TSharedRef<SWidget> FSequencer::MakeTransportControls(bool bExtended)
 		}
 		TransportControlArgs.WidgetsToCreate.Add(FTransportControlWidget(ETransportControlWidgetType::BackwardStep));
 		TransportControlArgs.WidgetsToCreate.Add(FTransportControlWidget(ETransportControlWidgetType::BackwardPlay));
-		TransportControlArgs.WidgetsToCreate.Add(FTransportControlWidget(FOnMakeTransportWidget::CreateSP(this, &FSequencer::OnCreateTransportRecord)));
 		TransportControlArgs.WidgetsToCreate.Add(FTransportControlWidget(ETransportControlWidgetType::ForwardPlay));
 		TransportControlArgs.WidgetsToCreate.Add(FTransportControlWidget(ETransportControlWidgetType::ForwardStep));
 		if(bExtended)
@@ -4648,7 +4649,21 @@ TSharedRef<SWidget> FSequencer::OnCreateTransportRecord()
 		})
 	);
 
-	return RecordButton;
+
+	TSharedRef<SHorizontalBox> RecordBox = SNew(SHorizontalBox);
+	RecordBox->AddSlot()
+	.AutoWidth()
+	[
+		RecordButton
+	];
+	RecordBox->AddSlot()
+	[
+		// Intentionally add some space to separate the record button so it's not easily pressed
+		SNew(SSpacer)
+		.Size(FVector2D(10.0f, 0.0f))
+	];
+
+	return RecordBox;
 }
 
 UObject* FSequencer::FindSpawnedObjectOrTemplate(const FGuid& BindingId)
