Sequencer: Keep actor name and binding label in sync when there is only one binding. This is already done for spawnables, so it makes sense to do this for possessables too.

diff --git a/Engine/Source/Editor/Sequencer/Private/DisplayNodes/SequencerObjectBindingNode.cpp b/Engine/Source/Editor/Sequencer/Private/DisplayNodes/SequencerObjectBindingNode.cpp
--- a/Engine/Source/Editor/Sequencer/Private/DisplayNodes/SequencerObjectBindingNode.cpp
+++ b/Engine/Source/Editor/Sequencer/Private/DisplayNodes/SequencerObjectBindingNode.cpp
@@ -1005,20 +1005,20 @@ void FSequencerObjectBindingNode::SetDisplayName(const FText& NewDisplayName)
 		FMovieSceneSpawnable* Spawnable = MovieScene->FindSpawnable(GetObjectBinding());
 		FMovieScenePossessable* Possessable = MovieScene->FindPossessable(GetObjectBinding());
 
-		if (Spawnable)
+		// If there is only one binding, set the name of the bound actor
+		TArrayView<TWeakObjectPtr<>> Objects = GetSequencer().FindObjectsInCurrentSequence(GetObjectBinding());
+		if (Objects.Num() == 1)
 		{
-			TArrayView<TWeakObjectPtr<>> Objects = GetSequencer().FindObjectsInCurrentSequence(GetObjectBinding());
-			// If there is only one binding, set the name of the bound actor
-			if (Objects.Num() == 1)
+			if (AActor* Actor = Cast<AActor>(Objects[0].Get()))
 			{
-				AActor* Actor = Cast<AActor>(Objects[0].Get());
 				Actor->SetActorLabel(NewDisplayName.ToString());
 			}
-			else
-			{
-				// Otherwise set our display name
-				Spawnable->SetName(NewDisplayName.ToString());
-			}
+		}
+
+		if (Spawnable)
+		{
+			// Otherwise set our display name
+			Spawnable->SetName(NewDisplayName.ToString());
 		}
 		else if (Possessable)
 		{
diff --git a/Engine/Source/Editor/Sequencer/Private/LevelEditorSequencerIntegration.cpp b/Engine/Source/Editor/Sequencer/Private/LevelEditorSequencerIntegration.cpp
--- a/Engine/Source/Editor/Sequencer/Private/LevelEditorSequencerIntegration.cpp
+++ b/Engine/Source/Editor/Sequencer/Private/LevelEditorSequencerIntegration.cpp
@@ -247,7 +247,7 @@ void FLevelEditorSequencerIntegration::Initialize(const FLevelEditorSequencerInt
 	UpdateDetails(bForceRefresh);
 }
 
-void RenameSpawnableRecursive(FSequencer* Sequencer, UMovieScene* MovieScene, FMovieSceneSequenceIDRef SequenceID, const FMovieSceneSequenceHierarchy* Hierarchy, AActor* ChangedActor)
+void RenameBindingRecursive(FSequencer* Sequencer, UMovieScene* MovieScene, FMovieSceneSequenceIDRef SequenceID, const FMovieSceneSequenceHierarchy* Hierarchy, AActor* ChangedActor)
 {
 	check(MovieScene);
 
@@ -266,6 +266,22 @@ void RenameSpawnableRecursive(FSequencer* Sequencer, UMovieScene* MovieScene, FM
 			}
 		}
 	}
+	for (int32 Index = 0; Index < MovieScene->GetPossessableCount(); ++Index)
+	{
+		FGuid ThisGuid = MovieScene->GetPossessable(Index).GetGuid();
+
+		// If there is only one binding, set the name of the possessable
+		TArrayView<TWeakObjectPtr<>> BoundObjects = Sequencer->FindBoundObjects(ThisGuid, SequenceID);
+		if (BoundObjects.Num() == 1)
+		{
+			AActor* Actor = Cast<AActor>(BoundObjects[0].Get());
+			if (Actor && Actor == ChangedActor)
+			{
+				MovieScene->Modify();
+				MovieScene->GetPossessable(Index).SetName(ChangedActor->GetActorLabel());
+			}
+		}
+	}
 
 	if (Hierarchy)
 	{
@@ -282,7 +298,7 @@ void RenameSpawnableRecursive(FSequencer* Sequencer, UMovieScene* MovieScene, FM
 
 					if (SubMovieScene)
 					{
-						RenameSpawnableRecursive(Sequencer, SubMovieScene, ChildID, Hierarchy, ChangedActor);
+						RenameBindingRecursive(Sequencer, SubMovieScene, ChildID, Hierarchy, ChangedActor);
 					}
 				}
 			}
@@ -305,7 +321,7 @@ void FLevelEditorSequencerIntegration::OnActorLabelChanged(AActor* ChangedActor)
 
 			if (MovieScene)
 			{
-				RenameSpawnableRecursive(Pinned.Get(), MovieScene, MovieSceneSequenceID::Root, Hierarchy, ChangedActor);
+				RenameBindingRecursive(Pinned.Get(), MovieScene, MovieSceneSequenceID::Root, Hierarchy, ChangedActor);
 			}
 		}
 	}
