Sequencer: Update a few tooltips with "actor" instead of "objects"

Add selected actors to this track
Replace the object binding with selected actors
Remove selected actors from this track
Remove all bound actors from this track

diff --git a/Engine/Source/Editor/Sequencer/Private/DisplayNodes/SequencerObjectBindingNode.cpp b/Engine/Source/Editor/Sequencer/Private/DisplayNodes/SequencerObjectBindingNode.cpp
--- a/Engine/Source/Editor/Sequencer/Private/DisplayNodes/SequencerObjectBindingNode.cpp
+++ b/Engine/Source/Editor/Sequencer/Private/DisplayNodes/SequencerObjectBindingNode.cpp
@@ -549,7 +549,7 @@ void FSequencerObjectBindingNode::AddAssignActorMenu(FMenuBuilder& MenuBuilder)
 
 	MenuBuilder.AddMenuEntry(
 		LOCTEXT("AddSelectedToBinding", "Add Selected"),
-		LOCTEXT("AddSelectedToBindingTooltip", "Add selected objects to this track"),
+		LOCTEXT("AddSelectedToBindingTooltip", "Add selected actors to this track"),
 		FSlateIcon(),
 		FUIAction(
 			FExecuteAction::CreateLambda([=] { GetSequencer().AddActorsToBinding(ObjectBinding, SelectedActors); }),
@@ -558,7 +558,7 @@ void FSequencerObjectBindingNode::AddAssignActorMenu(FMenuBuilder& MenuBuilder)
 		);
 	MenuBuilder.AddMenuEntry(
 		LOCTEXT("ReplaceBindingWithSelected", "Replace with Selected"),
-		LOCTEXT("ReplaceBindingWithSelectedTooltip", "Replace the object binding with selected objects"),
+		LOCTEXT("ReplaceBindingWithSelectedTooltip", "Replace the object binding with selected actors"),
 		FSlateIcon(),
 		FUIAction(
 			FExecuteAction::CreateLambda([=] { GetSequencer().ReplaceBindingWithActors(ObjectBinding, SelectedActors); }),
@@ -567,7 +567,7 @@ void FSequencerObjectBindingNode::AddAssignActorMenu(FMenuBuilder& MenuBuilder)
 	);
 	MenuBuilder.AddMenuEntry(
 		LOCTEXT("RemoveSelectedFromBinding", "Remove Selected"),
-		LOCTEXT("RemoveSelectedFromBindingTooltip", "Remove selected objects from this track"),
+		LOCTEXT("RemoveSelectedFromBindingTooltip", "Remove selected actors from this track"),
 		FSlateIcon(),
 		FUIAction(
 			FExecuteAction::CreateLambda([=] { GetSequencer().RemoveActorsFromBinding(ObjectBinding, SelectedActors); }),
@@ -576,7 +576,7 @@ void FSequencerObjectBindingNode::AddAssignActorMenu(FMenuBuilder& MenuBuilder)
 	);
 	MenuBuilder.AddMenuEntry(
 		LOCTEXT("RemoveAllBindings", "Remove All"),
-		LOCTEXT("RemoveAllBindingsTooltip", "Remove all bound objects from this track"),
+		LOCTEXT("RemoveAllBindingsTooltip", "Remove all bound actors from this track"),
 		FSlateIcon(),
 		FUIAction(
 			FExecuteAction::CreateLambda([=] { GetSequencer().RemoveAllBindings(ObjectBinding); })
