Sequencer: Invoke rename when adding a new folder

diff --git a/Engine/Source/Editor/Sequencer/Private/SSequencer.cpp b/Engine/Source/Editor/Sequencer/Private/SSequencer.cpp
--- a/Engine/Source/Editor/Sequencer/Private/SSequencer.cpp
+++ b/Engine/Source/Editor/Sequencer/Private/SSequencer.cpp
@@ -2966,6 +2966,23 @@ void SSequencer::UpdateLayoutTree()
 
 		AdditionalSelectionsToAdd.Empty();
 
+		if (!NodePathToRename.IsEmpty())
+		{
+			for (TSharedRef<FSequencerDisplayNode> Node : Sequencer->GetNodeTree()->GetAllNodes())
+			{
+				if (Node->GetPathName() == NodePathToRename)
+				{
+					GEditor->GetTimerManager()->SetTimerForNextTick([Node]()
+					{
+						Node->OnRenameRequested().Broadcast();
+					}
+					);
+					break;
+				}
+			}
+			NodePathToRename.Empty();
+		}
+
 		if (Sequencer->GetFocusedMovieSceneSequence())
 		{
 			bool bAnyChanged = false;
diff --git a/Engine/Source/Editor/Sequencer/Private/SSequencer.h b/Engine/Source/Editor/Sequencer/Private/SSequencer.h
--- a/Engine/Source/Editor/Sequencer/Private/SSequencer.h
+++ b/Engine/Source/Editor/Sequencer/Private/SSequencer.h
@@ -614,6 +614,9 @@ public:
 	/** This adds the specified path to the selection set to be restored the next time the tree view is refreshed. */
 	void AddAdditionalPathToSelectionSet(const FString& Path) { AdditionalSelectionsToAdd.Add(Path); }
 
+	/** Request to rename the given node path. */
+	void RequestRenameNode(const FString& Path) { NodePathToRename = Path; }
+
 	/** Applies dynamic sequencer customizations to this editor. */
 	void ApplySequencerCustomizations(const TArray<FSequencerCustomizationInfo>& Customizations);
 
@@ -768,6 +771,8 @@ private:
 	*/
 	TArray<FString> AdditionalSelectionsToAdd;
 
+	FString NodePathToRename;
+
 	TSharedPtr<SWidget> TickResolutionOverlay;
 
 	/** All possible track filter objects */
diff --git a/Engine/Source/Editor/Sequencer/Private/Sequencer.cpp b/Engine/Source/Editor/Sequencer/Private/Sequencer.cpp
--- a/Engine/Source/Editor/Sequencer/Private/Sequencer.cpp
+++ b/Engine/Source/Editor/Sequencer/Private/Sequencer.cpp
@@ -11023,6 +11023,8 @@ void FSequencer::OnAddFolder()
 	// set, which will cause the newly created node to be selected when the selection is restored post-refresh.
 	SequencerWidget->AddAdditionalPathToSelectionSet(NewNodePath);
 
+	SequencerWidget->RequestRenameNode(NewNodePath);
+
 	NotifyMovieSceneDataChanged( EMovieSceneDataChangeType::MovieSceneStructureItemAdded );
 }
 
