Fix add keyframe missing in UMG editor

diff --git a/Engine/Source/Editor/Sequencer/Private/LevelEditorSequencerIntegration.cpp b/Engine/Source/Editor/Sequencer/Private/LevelEditorSequencerIntegration.cpp
--- a/Engine/Source/Editor/Sequencer/Private/LevelEditorSequencerIntegration.cpp
+++ b/Engine/Source/Editor/Sequencer/Private/LevelEditorSequencerIntegration.cpp
@@ -819,88 +819,6 @@ void FLevelEditorSequencerIntegration::MakeBrowseToSelectedActorSubMenu(FMenuBui
 	}
 }
 
-static TSharedPtr<IDetailKeyframeHandler> GetKeyframeHandler(TWeakPtr<IDetailTreeNode> OwnerTreeNode)
-{
-	TSharedPtr<IDetailTreeNode> OwnerTreeNodePtr = OwnerTreeNode.Pin();
-	if (!OwnerTreeNodePtr.IsValid())
-	{
-		return TSharedPtr<IDetailKeyframeHandler>();
-	}
-
-	IDetailsView* DetailsView = OwnerTreeNodePtr->GetNodeDetailsView();
-	if (DetailsView == nullptr)
-	{
-		return TSharedPtr<IDetailKeyframeHandler>();
-	}
-
-	return DetailsView->GetKeyframeHandler();
-}
-
-static bool IsKeyframeButtonVisible(TWeakPtr<IDetailTreeNode> OwnerTreeNode, TSharedPtr<IPropertyHandle> PropertyHandle)
-{
-	TSharedPtr<IDetailKeyframeHandler> KeyframeHandler = GetKeyframeHandler(OwnerTreeNode);
-	if (!KeyframeHandler.IsValid() || !PropertyHandle.IsValid())
-	{
-		return false;
-	}
-
-	const UClass* ObjectClass = PropertyHandle->GetOuterBaseClass();
-	if (ObjectClass == nullptr)
-	{
-		return false;
-	}
-
-	return KeyframeHandler->IsPropertyKeyable(ObjectClass, *PropertyHandle);
-}
-
-static bool IsKeyframeButtonEnabled(TWeakPtr<IDetailTreeNode> OwnerTreeNode)
-{
-	TSharedPtr<IDetailKeyframeHandler> KeyframeHandler = GetKeyframeHandler(OwnerTreeNode);
-	if (!KeyframeHandler.IsValid())
-	{
-		return false;
-	}
-
-	return KeyframeHandler->IsPropertyKeyingEnabled();
-}
-
-static void OnAddKeyframeClicked(TWeakPtr<IDetailTreeNode> OwnerTreeNode, TSharedPtr<IPropertyHandle> PropertyHandle)
-{
-	TSharedPtr<IDetailKeyframeHandler> KeyframeHandler = GetKeyframeHandler(OwnerTreeNode);
-	if (!KeyframeHandler.IsValid() || !PropertyHandle.IsValid())
-	{
-		return;
-	}
-
-	KeyframeHandler->OnKeyPropertyClicked(*PropertyHandle);
-}
-
-
-static void RegisterDetailsExtensionHandler(const FOnGenerateGlobalRowExtensionArgs& Args, TArray<FPropertyRowExtensionButton>& OutExtensionButtons)
-{
-	// local copy for capturing in handlers below
-	TSharedPtr<IPropertyHandle> PropertyHandle = Args.PropertyHandle;
-	if (!PropertyHandle.IsValid())
-	{
-		return;
-	}
-
-	static FSlateIcon CreateKeyIcon(FAppStyle::Get().GetStyleSetName(), "Sequencer.AddKey.Details");
-
-	TWeakPtr<IDetailTreeNode> OwnerTreeNode = Args.OwnerTreeNode;
-
-	FPropertyRowExtensionButton& CreateKey = OutExtensionButtons.AddDefaulted_GetRef();
-	CreateKey.Icon = CreateKeyIcon;
-	CreateKey.Label = NSLOCTEXT("PropertyEditor", "CreateKey", "Create Key");
-	CreateKey.ToolTip = NSLOCTEXT("PropertyEditor", "CreateKeyToolTip", "Add a keyframe for this property.");
-	CreateKey.UIAction = FUIAction(
-		FExecuteAction::CreateStatic(&OnAddKeyframeClicked, OwnerTreeNode, PropertyHandle),
-		FCanExecuteAction::CreateStatic(&IsKeyframeButtonEnabled, OwnerTreeNode),
-		FGetActionCheckState(),
-		FIsActionButtonVisible::CreateStatic(&IsKeyframeButtonVisible, OwnerTreeNode, PropertyHandle)
-	);
-}
-
 void FLevelEditorSequencerIntegration::ActivateDetailHandler(const FLevelEditorSequencerIntegrationOptions& Options)
 {
 	static FName DetailHandlerName("DetailHandler");
@@ -921,10 +839,9 @@ void FLevelEditorSequencerIntegration::ActivateDetailHandler(const FLevelEditorS
 	}
 
 	FDelegateHandle OnPropertyEditorOpenedHandle = EditModule.OnPropertyEditorOpened().AddRaw(this, &FLevelEditorSequencerIntegration::OnPropertyEditorOpened);
-	FDelegateHandle OnGetGlobalRowExtensionHandle = EditModule.GetGlobalRowExtensionDelegate().AddStatic(&RegisterDetailsExtensionHandler);
 
 	auto DeactivateDetailKeyframeHandler =
-		[this, OnPropertyEditorOpenedHandle, OnGetGlobalRowExtensionHandle]()
+		[this, OnPropertyEditorOpenedHandle]()
 		{
 			FPropertyEditorModule* EditModulePtr = FModuleManager::Get().GetModulePtr<FPropertyEditorModule>("PropertyEditor");
 			if (!EditModulePtr)
@@ -932,7 +849,6 @@ void FLevelEditorSequencerIntegration::ActivateDetailHandler(const FLevelEditorS
 				return;
 			}
 
-			EditModulePtr->GetGlobalRowExtensionDelegate().Remove(OnGetGlobalRowExtensionHandle);
 			EditModulePtr->OnPropertyEditorOpened().Remove(OnPropertyEditorOpenedHandle);
 
 			for (const FName& DetailsTabIdentifier : DetailsTabIdentifiers)
diff --git a/Engine/Source/Editor/Sequencer/Private/SequencerModule.cpp b/Engine/Source/Editor/Sequencer/Private/SequencerModule.cpp
--- a/Engine/Source/Editor/Sequencer/Private/SequencerModule.cpp
+++ b/Engine/Source/Editor/Sequencer/Private/SequencerModule.cpp
@@ -15,6 +15,8 @@
 #include "SequencerEdMode.h"
 #include "SequencerObjectChangeListener.h"
 #include "IDetailKeyframeHandler.h"
+#include "IDetailTreeNode.h"
+#include "IDetailsView.h"
 #include "Tree/CurveEditorTreeFilter.h"
 #include "AnimatedPropertyKey.h"
 
@@ -44,6 +46,87 @@ ECurveEditorTreeFilterType ISequencerModule::GetSequencerSelectionFilterType()
 	return FilterType;
 }
 
+static TSharedPtr<IDetailKeyframeHandler> GetKeyframeHandler(TWeakPtr<IDetailTreeNode> OwnerTreeNode)
+{
+	TSharedPtr<IDetailTreeNode> OwnerTreeNodePtr = OwnerTreeNode.Pin();
+	if (!OwnerTreeNodePtr.IsValid())
+	{
+		return TSharedPtr<IDetailKeyframeHandler>();
+	}
+
+	IDetailsView* DetailsView = OwnerTreeNodePtr->GetNodeDetailsView();
+	if (DetailsView == nullptr)
+	{
+		return TSharedPtr<IDetailKeyframeHandler>();
+	}
+
+	return DetailsView->GetKeyframeHandler();
+}
+
+static bool IsKeyframeButtonVisible(TWeakPtr<IDetailTreeNode> OwnerTreeNode, TSharedPtr<IPropertyHandle> PropertyHandle)
+{
+	TSharedPtr<IDetailKeyframeHandler> KeyframeHandler = GetKeyframeHandler(OwnerTreeNode);
+	if (!KeyframeHandler.IsValid() || !PropertyHandle.IsValid())
+	{
+		return false;
+	}
+
+	const UClass* ObjectClass = PropertyHandle->GetOuterBaseClass();
+	if (ObjectClass == nullptr)
+	{
+		return false;
+	}
+
+	return KeyframeHandler->IsPropertyKeyable(ObjectClass, *PropertyHandle);
+}
+
+static bool IsKeyframeButtonEnabled(TWeakPtr<IDetailTreeNode> OwnerTreeNode)
+{
+	TSharedPtr<IDetailKeyframeHandler> KeyframeHandler = GetKeyframeHandler(OwnerTreeNode);
+	if (!KeyframeHandler.IsValid())
+	{
+		return false;
+	}
+
+	return KeyframeHandler->IsPropertyKeyingEnabled();
+}
+
+static void OnAddKeyframeClicked(TWeakPtr<IDetailTreeNode> OwnerTreeNode, TSharedPtr<IPropertyHandle> PropertyHandle)
+{
+	TSharedPtr<IDetailKeyframeHandler> KeyframeHandler = GetKeyframeHandler(OwnerTreeNode);
+	if (!KeyframeHandler.IsValid() || !PropertyHandle.IsValid())
+	{
+		return;
+	}
+
+	KeyframeHandler->OnKeyPropertyClicked(*PropertyHandle);
+}
+
+static void RegisterKeyframeExtensionHandler(const FOnGenerateGlobalRowExtensionArgs& Args, TArray<FPropertyRowExtensionButton>& OutExtensionButtons)
+{
+	// local copy for capturing in handlers below
+	TSharedPtr<IPropertyHandle> PropertyHandle = Args.PropertyHandle;
+	if (!PropertyHandle.IsValid())
+	{
+		return;
+	}
+
+	static FSlateIcon CreateKeyIcon(FAppStyle::Get().GetStyleSetName(), "Sequencer.AddKey.Details");
+
+	TWeakPtr<IDetailTreeNode> OwnerTreeNode = Args.OwnerTreeNode;
+
+	FPropertyRowExtensionButton& CreateKey = OutExtensionButtons.AddDefaulted_GetRef();
+	CreateKey.Icon = CreateKeyIcon;
+	CreateKey.Label = NSLOCTEXT("PropertyEditor", "CreateKey", "Create Key");
+	CreateKey.ToolTip = NSLOCTEXT("PropertyEditor", "CreateKeyToolTip", "Add a keyframe for this property.");
+	CreateKey.UIAction = FUIAction(
+		FExecuteAction::CreateStatic(&OnAddKeyframeClicked, OwnerTreeNode, PropertyHandle),
+		FCanExecuteAction::CreateStatic(&IsKeyframeButtonEnabled, OwnerTreeNode),
+		FGetActionCheckState(),
+		FIsActionButtonVisible::CreateStatic(&IsKeyframeButtonVisible, OwnerTreeNode, PropertyHandle)
+	);
+}
+
 /**
  * SequencerModule implementation (private)
  */
@@ -214,6 +297,9 @@ public:
 			{
 				FCoreDelegates::OnPostEngineInit.AddRaw(this, &FSequencerModule::RegisterMenus);
 			}
+
+			FPropertyEditorModule& EditModule = FModuleManager::Get().GetModuleChecked<FPropertyEditorModule>("PropertyEditor");
+			OnGetGlobalRowExtensionHandle = EditModule.GetGlobalRowExtensionDelegate().AddStatic(&RegisterKeyframeExtensionHandler);
 		}
 
 		ObjectBindingContextMenuExtensibilityManager = MakeShareable( new FExtensibilityManager );
@@ -229,6 +315,11 @@ public:
 		{
 			FSequencerCommands::Unregister();
 
+			if (FPropertyEditorModule* EditModulePtr = FModuleManager::Get().GetModulePtr<FPropertyEditorModule>("PropertyEditor"))
+			{
+				EditModulePtr->GetGlobalRowExtensionDelegate().Remove(OnGetGlobalRowExtensionHandle);
+			}
+
 			FEditorModeRegistry::Get().UnregisterMode(FSequencerEdMode::EM_SequencerMode);
 		}
 	}
@@ -335,6 +426,9 @@ private:
 	/** List of object binding handler delegates sequencers will execute when they are created */
 	TArray< FOnCreateEditorObjectBinding > EditorObjectBindingDelegates;
 
+	/** Global details row extension delegate; */
+	FDelegateHandle OnGetGlobalRowExtensionHandle;
+
 	/** Multicast delegate used to notify others of sequencer initialization params and allow modification. */
 	FOnPreSequencerInit OnPreSequencerInit;
 
