Sequencer: Reworked binding updates to properly support FMovieSceneBindingID

Deprecated OnBindingsUpdated in favor of OnBindingIDsUpdated which allows the binding id to be compared and resolved as a fixed binding id so that conflicts with identical guids in different sequences can be avoided.

diff --git a/Engine/Source/Editor/Sequencer/Private/Sequencer.cpp b/Engine/Source/Editor/Sequencer/Private/Sequencer.cpp
--- a/Engine/Source/Editor/Sequencer/Private/Sequencer.cpp
+++ b/Engine/Source/Editor/Sequencer/Private/Sequencer.cpp
@@ -81,6 +81,7 @@
 #include "Tracks/MovieSceneSubTrack.h"
 #include "Sections/MovieSceneCinematicShotSection.h"
 #include "MovieSceneObjectBindingIDCustomization.h"
+#include "MovieSceneObjectBindingID.h"
 #include "ISettingsModule.h"
 #include "Framework/Commands/GenericCommands.h"
 #include "Tracks/MovieSceneSpawnTrack.h"
@@ -10619,22 +10620,7 @@ TArray<FMovieSceneSpawnable*> FSequencer::ConvertToSpawnableInternal(FGuid Posse
 			SpawnRegister->HandleConvertPossessableToSpawnable(FoundObject, *this, TransformData);
 			SpawnRegister->SetupDefaultsForSpawnable(nullptr, Spawnable->GetGuid(), TransformData, AsShared(), Settings);
 
-			TMap<FGuid, FGuid> OldGuidToNewGuidMap;
-			OldGuidToNewGuidMap.Add(PossessableGuid, Spawnable->GetGuid());
-			
-			// Fixup any section bindings
-			TArray<UMovieScene*> MovieScenesToUpdate;
-			MovieSceneHelpers::GetDescendantMovieScenes(GetRootMovieSceneSequence(), MovieScenesToUpdate);
-			for (UMovieScene* MovieSceneToUpdate : MovieScenesToUpdate)
-			{
-				for (UMovieSceneSection* Section : MovieSceneToUpdate->GetAllSections())
-				{
-					if (Section)
-					{
-						Section->OnBindingsUpdated(OldGuidToNewGuidMap);
-					}
-				}
-			}
+			UpdateBindingIDs(PossessableGuid, Spawnable->GetGuid());
 
 			ForceEvaluate();
 
@@ -10839,22 +10825,7 @@ FMovieScenePossessable* FSequencer::ConvertToPossessableInternal(FGuid Spawnable
 		static const FName SequencerActorTag(TEXT("SequencerActor"));
 		PossessedActor->Tags.Remove(SequencerActorTag);
 
-		TMap<FGuid, FGuid> OldGuidToNewGuidMap;
-		OldGuidToNewGuidMap.Add(OldSpawnableGuid, NewPossessableGuid);
-		
-		// Fixup any section bindings
-		TArray<UMovieScene*> MovieScenesToUpdate;
-		MovieSceneHelpers::GetDescendantMovieScenes(GetRootMovieSceneSequence(), MovieScenesToUpdate);
-		for (UMovieScene* MovieSceneToUpdate : MovieScenesToUpdate)
-		{
-			for (UMovieSceneSection* Section : MovieSceneToUpdate->GetAllSections())
-			{
-				if (Section)
-				{
-					Section->OnBindingsUpdated(OldGuidToNewGuidMap);
-				}
-			}
-		}
+		UpdateBindingIDs(OldSpawnableGuid, NewPossessableGuid);
 
 		GEditor->SelectActor(PossessedActor, false, true);
 
@@ -10864,7 +10835,34 @@ FMovieScenePossessable* FSequencer::ConvertToPossessableInternal(FGuid Spawnable
 	return Possessable;
 }
 
+void FSequencer::UpdateBindingIDs(FGuid OldGuid, FGuid NewGuid)
+{
+	TMap<UE::MovieScene::FFixedObjectBindingID, UE::MovieScene::FFixedObjectBindingID> OldFixedToNewFixedMap;
+	OldFixedToNewFixedMap.Add(UE::MovieScene::FFixedObjectBindingID(OldGuid, ActiveTemplateIDs.Top()), UE::MovieScene::FFixedObjectBindingID(NewGuid, ActiveTemplateIDs.Top()));
+
+	const FMovieSceneSequenceHierarchy* Hierarchy = CompiledDataManager->FindHierarchy(RootTemplateInstance.GetCompiledDataID());
+	if (!Hierarchy)
+	{
+		return;
+	}
 
+	for (const TTuple<FMovieSceneSequenceID, FMovieSceneSubSequenceData>& Pair : Hierarchy->AllSubSequenceData())
+	{
+		if (UMovieSceneSequence* Sequence = Pair.Value.GetSequence())
+		{
+			if (UMovieScene* MovieScene = Sequence->GetMovieScene())
+			{
+				for (UMovieSceneSection* Section : MovieScene->GetAllSections())
+				{
+					if (Section)
+					{
+						Section->OnBindingIDsUpdated(OldFixedToNewFixedMap, Pair.Key, Hierarchy, *this);
+					}
+				}
+			}
+		}
+	}
+}
 
 void FSequencer::OnLoadRecordedData()
 {
@@ -12279,18 +12277,9 @@ void FSequencer::FixActorReferences()
 		}
 	}
 
-	// Fixup any section bindings
-	TArray<UMovieScene*> MovieScenesToUpdate;
-	MovieSceneHelpers::GetDescendantMovieScenes(GetRootMovieSceneSequence(), MovieScenesToUpdate);
-	for (UMovieScene* MovieSceneToUpdate : MovieScenesToUpdate)
+	for (TPair<FGuid, FGuid> GuidPair : OldGuidToNewGuidMap)
 	{
-		for (UMovieSceneSection* Section : MovieSceneToUpdate->GetAllSections())
-		{
-			if (Section)
-			{
-				Section->OnBindingsUpdated(OldGuidToNewGuidMap);
-			}
-		}
+		UpdateBindingIDs(GuidPair.Key, GuidPair.Value);
 	}
 }
 
diff --git a/Engine/Source/Editor/Sequencer/Private/Sequencer.h b/Engine/Source/Editor/Sequencer/Private/Sequencer.h
--- a/Engine/Source/Editor/Sequencer/Private/Sequencer.h
+++ b/Engine/Source/Editor/Sequencer/Private/Sequencer.h
@@ -1092,6 +1092,8 @@ private:
 
 	/** Internal conversion function that doesn't perform expensive reset/update tasks */
 	FMovieScenePossessable* ConvertToPossessableInternal(FGuid SpawnableGuid);
+	
+	void UpdateBindingIDs(FGuid OldGuid, FGuid NewGuid);
 
 	/** Recurses through a folder to replace converted GUID with new GUID */
 	bool ReplaceFolderBindingGUID(UMovieSceneFolder *Folder, FGuid Original, FGuid Converted);
diff --git a/Engine/Source/Runtime/MovieScene/Public/MovieSceneSection.h b/Engine/Source/Runtime/MovieScene/Public/MovieSceneSection.h
--- a/Engine/Source/Runtime/MovieScene/Public/MovieSceneSection.h
+++ b/Engine/Source/Runtime/MovieScene/Public/MovieSceneSection.h
@@ -13,6 +13,7 @@
 #include "Evaluation/Blending/MovieSceneBlendType.h"
 #include "Evaluation/MovieSceneCompletionMode.h"
 #include "Generators/MovieSceneEasingFunction.h"
+#include "Evaluation/MovieSceneSequenceHierarchy.h"
 #include "MovieSceneFrameMigration.h"
 #include "Misc/QualifiedFrameTime.h"
 #include "Evaluation/MovieSceneEvaluationCustomVersion.h"
@@ -26,9 +27,11 @@ struct FKeyHandle;
 struct FEasingComponentData;
 struct FMovieSceneChannelProxy;
 struct FMovieSceneEvalTemplatePtr;
+struct FMovieSceneSequenceID;
 struct FFrameRate;
 
 class UMovieSceneEntitySystemLinker;
+class IMovieScenePlayer;
 
 namespace UE
 {
@@ -36,6 +39,7 @@ namespace MovieScene
 {
 	struct FEntityImportParams;
 	struct FImportedEntity;
+	struct FFixedObjectBindingID;
 }
 }
 
@@ -507,6 +511,9 @@ public:
 	 * When guid bindings are updated to allow this section to fix-up any internal bindings
 	 *
 	 */
+	virtual void OnBindingIDsUpdated(const TMap<UE::MovieScene::FFixedObjectBindingID, UE::MovieScene::FFixedObjectBindingID>& OldFixedToNewFixedMap, FMovieSceneSequenceID LocalSequenceID, const FMovieSceneSequenceHierarchy* Hierarchy, IMovieScenePlayer& Player) {}
+
+	UE_DEPRECATED(5.0, "OnBindingsUpdated has been deprecated. Use OnBindingIDsUpdated instead")
 	virtual void OnBindingsUpdated(const TMap<FGuid, FGuid>& OldGuidToNewGuidMap) { }
 
 	/** Get the referenced bindings for this section */
diff --git a/Engine/Source/Runtime/MovieSceneTracks/Private/Sections/MovieScene3DConstraintSection.cpp b/Engine/Source/Runtime/MovieSceneTracks/Private/Sections/MovieScene3DConstraintSection.cpp
--- a/Engine/Source/Runtime/MovieSceneTracks/Private/Sections/MovieScene3DConstraintSection.cpp
+++ b/Engine/Source/Runtime/MovieSceneTracks/Private/Sections/MovieScene3DConstraintSection.cpp
@@ -2,7 +2,9 @@
 
 #include "Sections/MovieScene3DConstraintSection.h"
 #include "MovieSceneObjectBindingID.h"
-
+#include "MovieSceneSequenceID.h"
+#include "Evaluation/MovieSceneSequenceHierarchy.h"
+#include "IMovieScenePlayer.h"
 
 UMovieScene3DConstraintSection::UMovieScene3DConstraintSection( const FObjectInitializer& ObjectInitializer )
 	: Super( ObjectInitializer )
@@ -10,13 +12,15 @@ UMovieScene3DConstraintSection::UMovieScene3DConstraintSection( const FObjectIni
 	bSupportsInfiniteRange = true;
 }
 
-void UMovieScene3DConstraintSection::OnBindingsUpdated(const TMap<FGuid, FGuid>& OldGuidToNewGuidMap)
+void UMovieScene3DConstraintSection::OnBindingIDsUpdated(const TMap<UE::MovieScene::FFixedObjectBindingID, UE::MovieScene::FFixedObjectBindingID>& OldFixedToNewFixedMap, FMovieSceneSequenceID LocalSequenceID, const FMovieSceneSequenceHierarchy* Hierarchy, IMovieScenePlayer& Player)
 {
-	if (OldGuidToNewGuidMap.Contains(ConstraintBindingID.GetGuid()))
+	UE::MovieScene::FFixedObjectBindingID FixedBindingID = ConstraintBindingID.ResolveToFixed(LocalSequenceID, Player);
+
+	if (OldFixedToNewFixedMap.Contains(FixedBindingID))
 	{
 		Modify();
 
-		ConstraintBindingID.SetGuid(OldGuidToNewGuidMap[ConstraintBindingID.GetGuid()]);
+		SetConstraintBindingID(OldFixedToNewFixedMap[FixedBindingID].ConvertToRelative(LocalSequenceID, Hierarchy));
 	}
 }
 
diff --git a/Engine/Source/Runtime/MovieSceneTracks/Private/Sections/MovieSceneActorReferenceSection.cpp b/Engine/Source/Runtime/MovieSceneTracks/Private/Sections/MovieSceneActorReferenceSection.cpp
--- a/Engine/Source/Runtime/MovieSceneTracks/Private/Sections/MovieSceneActorReferenceSection.cpp
+++ b/Engine/Source/Runtime/MovieSceneTracks/Private/Sections/MovieSceneActorReferenceSection.cpp
@@ -2,7 +2,10 @@
 
 #include "Sections/MovieSceneActorReferenceSection.h"
 #include "Channels/MovieSceneChannelProxy.h"
-
+#include "MovieSceneObjectBindingID.h"
+#include "MovieSceneSequenceID.h"
+#include "Evaluation/MovieSceneSequenceHierarchy.h"
+#include "IMovieScenePlayer.h"
 
 bool FMovieSceneActorReferenceData::Evaluate(FFrameTime InTime, FMovieSceneActorReferenceKey& OutValue) const
 {
@@ -133,23 +136,27 @@ void UMovieSceneActorReferenceSection::PostLoad()
 	}
 }
 
-void UMovieSceneActorReferenceSection::OnBindingsUpdated(const TMap<FGuid, FGuid>& OldGuidToNewGuidMap)
+void UMovieSceneActorReferenceSection::OnBindingIDsUpdated(const TMap<UE::MovieScene::FFixedObjectBindingID, UE::MovieScene::FFixedObjectBindingID>& OldFixedToNewFixedMap, FMovieSceneSequenceID LocalSequenceID, const FMovieSceneSequenceHierarchy* Hierarchy, IMovieScenePlayer& Player)
 {
-	if (OldGuidToNewGuidMap.Contains(ActorReferenceData.GetDefault().Object.GetGuid()))
+	UE::MovieScene::FFixedObjectBindingID DefaultFixedBindingID = ActorReferenceData.GetDefault().Object.ResolveToFixed(LocalSequenceID, Player);
+
+	if (OldFixedToNewFixedMap.Contains(DefaultFixedBindingID))
 	{
 		Modify();
 
-		ActorReferenceData.GetDefault().Object.SetGuid(OldGuidToNewGuidMap[ActorReferenceData.GetDefault().Object.GetGuid()]);
+		ActorReferenceData.GetDefault().Object = OldFixedToNewFixedMap[DefaultFixedBindingID].ConvertToRelative(LocalSequenceID, Hierarchy);
 	}
 
 	for (FMovieSceneActorReferenceKey& Key : ActorReferenceData.GetData().GetValues())
 	{
-		if (OldGuidToNewGuidMap.Contains(Key.Object.GetGuid()))
+		UE::MovieScene::FFixedObjectBindingID KeyFixedBindingID = Key.Object.ResolveToFixed(LocalSequenceID, Player);
+
+		if (OldFixedToNewFixedMap.Contains(KeyFixedBindingID))
 		{
 			Modify();
-			
-			Key.Object.SetGuid(OldGuidToNewGuidMap[Key.Object.GetGuid()]);
-		}	
+
+			Key.Object = OldFixedToNewFixedMap[KeyFixedBindingID].ConvertToRelative(LocalSequenceID, Hierarchy);
+		}
 	}
 }
 
diff --git a/Engine/Source/Runtime/MovieSceneTracks/Private/Sections/MovieSceneCameraCutSection.cpp b/Engine/Source/Runtime/MovieSceneTracks/Private/Sections/MovieSceneCameraCutSection.cpp
--- a/Engine/Source/Runtime/MovieSceneTracks/Private/Sections/MovieSceneCameraCutSection.cpp
+++ b/Engine/Source/Runtime/MovieSceneTracks/Private/Sections/MovieSceneCameraCutSection.cpp
@@ -4,11 +4,12 @@
 
 #include "MovieScene.h"
 #include "Tracks/MovieSceneCameraCutTrack.h"
-#include "MovieScene.h"
+#include "MovieSceneSequenceID.h"
 #include "IMovieScenePlayer.h"
 #include "Camera/CameraComponent.h"
 #include "Evaluation/MovieSceneEvaluationTemplateInstance.h"
 #include "Evaluation/MovieSceneEvalTemplate.h"
+#include "Evaluation/MovieSceneSequenceHierarchy.h"
 #include "EntitySystem/MovieSceneEntityManager.h"
 #include "EntitySystem/MovieSceneEntityBuilder.h"
 #include "EntitySystem/BuiltInComponentTypes.h"
@@ -35,13 +36,15 @@ UMovieSceneCameraCutSection::UMovieSceneCameraCutSection(const FObjectInitialize
 	SetBlendType(EMovieSceneBlendType::Absolute);
 }
 
-void UMovieSceneCameraCutSection::OnBindingsUpdated(const TMap<FGuid, FGuid>& OldGuidToNewGuidMap)
+void UMovieSceneCameraCutSection::OnBindingIDsUpdated(const TMap<UE::MovieScene::FFixedObjectBindingID, UE::MovieScene::FFixedObjectBindingID>& OldFixedToNewFixedMap, FMovieSceneSequenceID LocalSequenceID, const FMovieSceneSequenceHierarchy* Hierarchy, IMovieScenePlayer& Player)
 {
-	if (OldGuidToNewGuidMap.Contains(CameraBindingID.GetGuid()))
+	UE::MovieScene::FFixedObjectBindingID FixedBindingID = CameraBindingID.ResolveToFixed(LocalSequenceID, Player);
+
+	if (OldFixedToNewFixedMap.Contains(FixedBindingID))
 	{
 		Modify();
 
-		CameraBindingID.SetGuid(OldGuidToNewGuidMap[CameraBindingID.GetGuid()]);
+		CameraBindingID = OldFixedToNewFixedMap[FixedBindingID].ConvertToRelative(LocalSequenceID, Hierarchy);
 	}
 }
 
diff --git a/Engine/Source/Runtime/MovieSceneTracks/Public/Sections/MovieScene3DConstraintSection.h b/Engine/Source/Runtime/MovieSceneTracks/Public/Sections/MovieScene3DConstraintSection.h
--- a/Engine/Source/Runtime/MovieSceneTracks/Public/Sections/MovieScene3DConstraintSection.h
+++ b/Engine/Source/Runtime/MovieSceneTracks/Public/Sections/MovieScene3DConstraintSection.h
@@ -40,7 +40,7 @@ public:
 
 	//~ UMovieSceneSection interface
 
-	virtual void OnBindingsUpdated(const TMap<FGuid, FGuid>& OldGuidToNewGuidMap) override;
+	virtual void OnBindingIDsUpdated(const TMap<UE::MovieScene::FFixedObjectBindingID, UE::MovieScene::FFixedObjectBindingID>& OldFixedToNewFixedMap, FMovieSceneSequenceID LocalSequenceID, const FMovieSceneSequenceHierarchy* Hierarchy, IMovieScenePlayer& Player) override;
 	
 	virtual void GetReferencedBindings(TArray<FGuid>& OutBindings) override;
 
diff --git a/Engine/Source/Runtime/MovieSceneTracks/Public/Sections/MovieSceneActorReferenceSection.h b/Engine/Source/Runtime/MovieSceneTracks/Public/Sections/MovieSceneActorReferenceSection.h
--- a/Engine/Source/Runtime/MovieSceneTracks/Public/Sections/MovieSceneActorReferenceSection.h
+++ b/Engine/Source/Runtime/MovieSceneTracks/Public/Sections/MovieSceneActorReferenceSection.h
@@ -169,7 +169,7 @@ public:
 	virtual void PostLoad() override;
 
 	//~ UMovieSceneSection interface
-	virtual void OnBindingsUpdated(const TMap<FGuid, FGuid>& OldGuidToNewGuidMap) override;
+	virtual void OnBindingIDsUpdated(const TMap<UE::MovieScene::FFixedObjectBindingID, UE::MovieScene::FFixedObjectBindingID>& OldFixedToNewFixedMap, FMovieSceneSequenceID LocalSequenceID, const FMovieSceneSequenceHierarchy* Hierarchy, IMovieScenePlayer& Player) override;
 
 	const FMovieSceneActorReferenceData& GetActorReferenceData() const { return ActorReferenceData; }
 
diff --git a/Engine/Source/Runtime/MovieSceneTracks/Public/Sections/MovieSceneCameraCutSection.h b/Engine/Source/Runtime/MovieSceneTracks/Public/Sections/MovieSceneCameraCutSection.h
--- a/Engine/Source/Runtime/MovieSceneTracks/Public/Sections/MovieSceneCameraCutSection.h
+++ b/Engine/Source/Runtime/MovieSceneTracks/Public/Sections/MovieSceneCameraCutSection.h
@@ -52,7 +52,7 @@ public:
 	}
 
 	//~ UMovieSceneSection interface
-	virtual void OnBindingsUpdated(const TMap<FGuid, FGuid>& OldGuidToNewGuidMap) override;
+	virtual void OnBindingIDsUpdated(const TMap<UE::MovieScene::FFixedObjectBindingID, UE::MovieScene::FFixedObjectBindingID>& OldFixedToNewFixedMap, FMovieSceneSequenceID LocalSequenceID, const FMovieSceneSequenceHierarchy* Hierarchy, IMovieScenePlayer& Player) override;
 	virtual void GetReferencedBindings(TArray<FGuid>& OutBindings) override;
 
 	/** ~UObject interface */
