Sequencer: Use subdued foreground color for muted tracks and muted object binding nodes

diff --git a/Engine/Source/Editor/Sequencer/Private/DisplayNodes/SequencerDisplayNode.cpp b/Engine/Source/Editor/Sequencer/Private/DisplayNodes/SequencerDisplayNode.cpp
--- a/Engine/Source/Editor/Sequencer/Private/DisplayNodes/SequencerDisplayNode.cpp
+++ b/Engine/Source/Editor/Sequencer/Private/DisplayNodes/SequencerDisplayNode.cpp
@@ -829,7 +829,7 @@ FSlateFontInfo FSequencerDisplayNode::GetDisplayNameFont() const
 
 FSlateColor FSequencerDisplayNode::GetDisplayNameColor() const
 {
-	return IsDimmed() ? FLinearColor(0.6f, 0.6f, 0.6f, 0.6f) : FSlateColor::UseForeground();
+	return IsDimmed() ? FSlateColor::UseSubduedForeground() : FSlateColor::UseForeground();
 }
 
 FText FSequencerDisplayNode::GetDisplayNameToolTipText() const
diff --git a/Engine/Source/Editor/Sequencer/Private/DisplayNodes/SequencerFolderNode.cpp b/Engine/Source/Editor/Sequencer/Private/DisplayNodes/SequencerFolderNode.cpp
--- a/Engine/Source/Editor/Sequencer/Private/DisplayNodes/SequencerFolderNode.cpp
+++ b/Engine/Source/Editor/Sequencer/Private/DisplayNodes/SequencerFolderNode.cpp
@@ -115,7 +115,7 @@ FSlateColor FSequencerFolderNode::GetIconColor() const
 {
 	if (ParentTree.IsNodeMute(this))
 	{
-		return FSlateColor(MovieSceneFolder.GetFolderColor().ReinterpretAsLinear() * FLinearColor(0.6f, 0.6f, 0.6f, 0.6f));
+		return FSlateColor(MovieSceneFolder.GetFolderColor().ReinterpretAsLinear().Desaturate(0.6f));
 	}
 
 	return FSlateColor(MovieSceneFolder.GetFolderColor());
@@ -125,7 +125,7 @@ FSlateColor FSequencerFolderNode::GetDisplayNameColor() const
 {
 	if (ParentTree.IsNodeMute(this))
 	{
-		return FLinearColor(0.6f, 0.6f, 0.6f, 0.6f);
+		return FSlateColor::UseSubduedForeground();
 	}
 
 	return FSlateColor::UseForeground();
diff --git a/Engine/Source/Editor/Sequencer/Private/DisplayNodes/SequencerObjectBindingNode.cpp b/Engine/Source/Editor/Sequencer/Private/DisplayNodes/SequencerObjectBindingNode.cpp
--- a/Engine/Source/Editor/Sequencer/Private/DisplayNodes/SequencerObjectBindingNode.cpp
+++ b/Engine/Source/Editor/Sequencer/Private/DisplayNodes/SequencerObjectBindingNode.cpp
@@ -843,7 +843,7 @@ FSlateColor FSequencerObjectBindingNode::GetDisplayNameColor() const
 		{
 			if (StaticCastSharedPtr<FSequencerObjectBindingNode>(CurrentNode)->GetBindingType() == EObjectBindingType::Spawnable)
 			{
-				return FLinearColor::Gray;
+				return FSlateColor::UseSubduedForeground();
 			}
 		}
 
diff --git a/Engine/Source/Editor/Sequencer/Private/DisplayNodes/SequencerTrackNode.cpp b/Engine/Source/Editor/Sequencer/Private/DisplayNodes/SequencerTrackNode.cpp
--- a/Engine/Source/Editor/Sequencer/Private/DisplayNodes/SequencerTrackNode.cpp
+++ b/Engine/Source/Editor/Sequencer/Private/DisplayNodes/SequencerTrackNode.cpp
@@ -744,7 +744,7 @@ FSlateColor FSequencerTrackNode::GetDisplayNameColor() const
 		// 3D transform tracks don't map to property bindings as below
 		if (Track->IsA<UMovieScene3DTransformTrack>() || Track->IsA<UMovieScenePrimitiveMaterialTrack>())
 		{
-			return bIsDimmed ? FLinearColor(0.6f, 0.6f, 0.6f, 0.6f) : FSlateColor::UseForeground();
+			return bIsDimmed ? FSlateColor::UseSubduedForeground() : FSlateColor::UseForeground();
 		}
 
 		FGuid ObjectBinding;
@@ -764,12 +764,12 @@ FSlateColor FSequencerTrackNode::GetDisplayNameColor() const
 					FTrackInstancePropertyBindings PropertyBinding(PropertyTrack->GetPropertyName(), PropertyTrack->GetPropertyPath().ToString());
 					if (PropertyBinding.GetProperty(*Object))
 					{
-						return bIsDimmed ? FLinearColor(0.6f, 0.6f, 0.6f, 0.6f) : FSlateColor::UseForeground();
+						return bIsDimmed ? FSlateColor::UseSubduedForeground() : FSlateColor::UseForeground();
 					}
 				}
 			}
 
-			return bIsDimmed ? FLinearColor(0.6f, 0.0f, 0.0f, 0.6f) : FLinearColor::Red;
+			return bIsDimmed ? FSlateColor(FLinearColor::Red.Desaturate(0.6f)) : FLinearColor::Red;
 		}
 	}
 
