Sequencer: When converting spawnable to possessable, set the default transform of the converted possessable to the current transform of the spawnable.

diff --git a/Engine/Source/Editor/Sequencer/Private/Sequencer.cpp b/Engine/Source/Editor/Sequencer/Private/Sequencer.cpp
--- a/Engine/Source/Editor/Sequencer/Private/Sequencer.cpp
+++ b/Engine/Source/Editor/Sequencer/Private/Sequencer.cpp
@@ -10842,6 +10842,19 @@ FMovieScenePossessable* FSequencer::ConvertToPossessableInternal(FGuid Spawnable
 		return nullptr;
 	}
 
+	FTransform DefaultTransform = SpawnableActorTemplate->GetActorTransform();
+	for (TWeakObjectPtr<> RuntimeObject : FindBoundObjects(SpawnableGuid, ActiveTemplateIDs.Top()) )
+	{
+		// Prefer the transform at the current time over the spawnable actor template's transform because that's most likely 0. 
+		// This makes it so that the object will return to the current position on restore state.
+		AActor* Actor = Cast<AActor>(RuntimeObject.Get());
+		if (Actor)
+		{
+			DefaultTransform = Actor->GetTransform();
+			break;
+		}
+	}
+
 	Sequence->Modify();
 	MovieScene->Modify();
 
@@ -10852,13 +10865,12 @@ FMovieScenePossessable* FSequencer::ConvertToPossessableInternal(FGuid Spawnable
 		MovieScene->RemoveTrack(*SpawnTrack);
 	}
 
-	FTransform SpawnTransform = SpawnableActorTemplate->GetActorTransform();
 	FActorSpawnParameters SpawnInfo;
 	SpawnInfo.bDeferConstruction = true;
 	SpawnInfo.Template = SpawnableActorTemplate;
 
 	UWorld* PlaybackContext = Cast<UWorld>(GetPlaybackContext());
-	AActor* PossessedActor = PlaybackContext->SpawnActor(Spawnable->GetObjectTemplate()->GetClass(), &SpawnTransform, SpawnInfo);
+	AActor* PossessedActor = PlaybackContext->SpawnActor(Spawnable->GetObjectTemplate()->GetClass(), &DefaultTransform, SpawnInfo);
 
 	if (!PossessedActor)
 	{
@@ -10868,7 +10880,7 @@ FMovieScenePossessable* FSequencer::ConvertToPossessableInternal(FGuid Spawnable
 	PossessedActor->SetActorLabel(Spawnable->GetName());
 
 	const bool bIsDefaultTransform = true;
-	PossessedActor->FinishSpawning(SpawnTransform, bIsDefaultTransform);
+	PossessedActor->FinishSpawning(DefaultTransform, bIsDefaultTransform);
 
 	const FGuid NewPossessableGuid = CreateBinding(*PossessedActor, PossessedActor->GetActorLabel());
 	const FGuid OldSpawnableGuid = Spawnable->GetGuid();
