Sequencer: When converting spawnable to possessable, set the default transform of the converted possessable to the current transform of the spawnable.

Follow up fix for deferred construction and dynamic root components

diff --git a/Engine/Source/Editor/Sequencer/Private/Sequencer.cpp b/Engine/Source/Editor/Sequencer/Private/Sequencer.cpp
--- a/Engine/Source/Editor/Sequencer/Private/Sequencer.cpp
+++ b/Engine/Source/Editor/Sequencer/Private/Sequencer.cpp
@@ -10882,6 +10882,9 @@ FMovieScenePossessable* FSequencer::ConvertToPossessableInternal(FGuid Spawnable
 	const bool bIsDefaultTransform = true;
 	PossessedActor->FinishSpawning(DefaultTransform, bIsDefaultTransform);
 
+	// The transform needs to be set again for deferred construction and dynamic root components. Until the fix for: UE-67537
+	PossessedActor->SetActorTransform(DefaultTransform);
+
 	const FGuid NewPossessableGuid = CreateBinding(*PossessedActor, PossessedActor->GetActorLabel());
 	const FGuid OldSpawnableGuid = Spawnable->GetGuid();
 
