Sequencer: Reworked binding updates to properly support FMovieSceneBindingID

Deprecated OnBindingsUpdated in favor of OnBindingIDsUpdated which allows the binding id to be compared and resolved as a fixed binding id so that conflicts with identical guids in different sequences can be avoided.

Follow up to support the bindings in the root sequence

diff --git a/Engine/Source/Editor/Sequencer/Private/Sequencer.cpp b/Engine/Source/Editor/Sequencer/Private/Sequencer.cpp
--- a/Engine/Source/Editor/Sequencer/Private/Sequencer.cpp
+++ b/Engine/Source/Editor/Sequencer/Private/Sequencer.cpp
@@ -10937,22 +10937,32 @@ void FSequencer::UpdateBindingIDs(FGuid OldGuid, FGuid NewGuid)
 	OldFixedToNewFixedMap.Add(UE::MovieScene::FFixedObjectBindingID(OldGuid, ActiveTemplateIDs.Top()), UE::MovieScene::FFixedObjectBindingID(NewGuid, ActiveTemplateIDs.Top()));
 
 	const FMovieSceneSequenceHierarchy* Hierarchy = CompiledDataManager->FindHierarchy(RootTemplateInstance.GetCompiledDataID());
-	if (!Hierarchy)
+
+	if (UMovieScene* MovieScene = GetRootMovieSceneSequence()->GetMovieScene())
 	{
-		return;
+		for (UMovieSceneSection* Section : MovieScene->GetAllSections())
+		{
+			if (Section)
+			{
+				Section->OnBindingIDsUpdated(OldFixedToNewFixedMap, GetRootTemplateID(), Hierarchy, *this);
+			}
+		}
 	}
 
-	for (const TTuple<FMovieSceneSequenceID, FMovieSceneSubSequenceData>& Pair : Hierarchy->AllSubSequenceData())
+	if (Hierarchy)
 	{
-		if (UMovieSceneSequence* Sequence = Pair.Value.GetSequence())
+		for (const TTuple<FMovieSceneSequenceID, FMovieSceneSubSequenceData>& Pair : Hierarchy->AllSubSequenceData())
 		{
-			if (UMovieScene* MovieScene = Sequence->GetMovieScene())
+			if (UMovieSceneSequence* Sequence = Pair.Value.GetSequence())
 			{
-				for (UMovieSceneSection* Section : MovieScene->GetAllSections())
+				if (UMovieScene* MovieScene = Sequence->GetMovieScene())
 				{
-					if (Section)
+					for (UMovieSceneSection* Section : MovieScene->GetAllSections())
 					{
-						Section->OnBindingIDsUpdated(OldFixedToNewFixedMap, Pair.Key, Hierarchy, *this);
+						if (Section)
+						{
+							Section->OnBindingIDsUpdated(OldFixedToNewFixedMap, Pair.Key, Hierarchy, *this);
+						}
 					}
 				}
 			}
