Sequencer: Expand section when adding keys otherwise newly created keys won't evaluate

diff --git a/Engine/Source/Editor/Sequencer/Private/Sequencer.cpp b/Engine/Source/Editor/Sequencer/Private/Sequencer.cpp
--- a/Engine/Source/Editor/Sequencer/Private/Sequencer.cpp
+++ b/Engine/Source/Editor/Sequencer/Private/Sequencer.cpp
@@ -10029,6 +10029,7 @@ void FSequencer::SetKey()
 					ModifiedSections.Add(Section);
 				}
 
+				Section->ExpandToFrame(AddKeyTime);
 
 				KeyArea->AddOrUpdateKey(AddKeyTime, ObjectBinding, *this);
 			}
