Restore sequencer to saved state when saving external actors and reset back to animated state afterward, similar to the save world behavior.

diff --git a/Engine/Source/Editor/Sequencer/Private/LevelEditorSequencerIntegration.cpp b/Engine/Source/Editor/Sequencer/Private/LevelEditorSequencerIntegration.cpp
--- a/Engine/Source/Editor/Sequencer/Private/LevelEditorSequencerIntegration.cpp
+++ b/Engine/Source/Editor/Sequencer/Private/LevelEditorSequencerIntegration.cpp
@@ -180,6 +180,14 @@ void FLevelEditorSequencerIntegration::Initialize(const FLevelEditorSequencerInt
 		FDelegateHandle Handle = FEditorDelegates::PostSaveWorldWithContext.AddRaw(this, &FLevelEditorSequencerIntegration::OnPostSaveWorld);
 		AcquiredResources.Add([=]{ FEditorDelegates::PostSaveWorldWithContext.Remove(Handle); });
 	}
+	{
+		FDelegateHandle Handle = FEditorDelegates::PreSaveExternalActors.AddRaw(this, &FLevelEditorSequencerIntegration::OnPreSaveExternalActors);
+		AcquiredResources.Add([=]{ FEditorDelegates::PreSaveExternalActors.Remove(Handle); });
+	}
+	{
+		FDelegateHandle Handle = FEditorDelegates::PostSaveExternalActors.AddRaw(this, &FLevelEditorSequencerIntegration::OnPostSaveExternalActors);
+		AcquiredResources.Add([=]{ FEditorDelegates::PostSaveExternalActors.Remove(Handle); });
+	}
 	{
 		FDelegateHandle Handle = FEditorDelegates::PreBeginPIE.AddRaw(this, &FLevelEditorSequencerIntegration::OnPreBeginPIE);
 		AcquiredResources.Add([=]{ FEditorDelegates::PreBeginPIE.Remove(Handle); });
@@ -329,40 +337,22 @@ void FLevelEditorSequencerIntegration::OnActorLabelChanged(AActor* ChangedActor)
 
 void FLevelEditorSequencerIntegration::OnPreSaveWorld(class UWorld* World, FObjectPreSaveContext ObjectSaveContext)
 {
-	// Restore the saved state so that the level save can save that instead of the animated state.
-	IterateAllSequencers(
-		[World](FSequencer& In, const FLevelEditorSequencerIntegrationOptions& Options)
-		{
-			if (Options.bRequiresLevelEvents)
-			{
-				for (const TSharedPtr<ISequencerTrackEditor>& TrackEditor : In.GetTrackEditors())
-				{
-					TrackEditor->OnPreSaveWorld(World);
-				}
-				In.RestorePreAnimatedState();
-			}
-		}
-	);
+	RestoreToSavedState(World);
 }
 
 void FLevelEditorSequencerIntegration::OnPostSaveWorld(class UWorld* World, FObjectPostSaveContext ObjectSaveContext)
 {
-	// Reset the time after saving so that an update will be triggered to put objects back to their animated state.
-	IterateAllSequencers(
-		[World](FSequencer& In, const FLevelEditorSequencerIntegrationOptions& Options)
-		{
-			if (Options.bRequiresLevelEvents)
-			{
-				In.InvalidateCachedData();
-				In.ForceEvaluate();
+	ResetToAnimatedState(World);
+}
 
-				for (const TSharedPtr<ISequencerTrackEditor>& TrackEditor : In.GetTrackEditors())
-				{
-					TrackEditor->OnPostSaveWorld(World);
-				}
-			}
-		}
-	);
+void FLevelEditorSequencerIntegration::OnPreSaveExternalActors(UWorld* World)
+{
+	RestoreToSavedState(World);
+}
+
+void FLevelEditorSequencerIntegration::OnPostSaveExternalActors(UWorld* World)
+{
+	ResetToAnimatedState(World);
 }
 
 void FLevelEditorSequencerIntegration::OnNewCurrentLevel()
@@ -1142,6 +1132,44 @@ void FLevelEditorSequencerIntegration::RestoreRealtimeViewports()
 	}
 }
 
+void FLevelEditorSequencerIntegration::RestoreToSavedState(UWorld* World)
+{
+	// Restore the saved state so that the level save can save that instead of the animated state.
+	IterateAllSequencers(
+		[World](FSequencer& In, const FLevelEditorSequencerIntegrationOptions& Options)
+		{
+			if (Options.bRequiresLevelEvents)
+			{
+				for (const TSharedPtr<ISequencerTrackEditor>& TrackEditor : In.GetTrackEditors())
+				{
+					TrackEditor->OnPreSaveWorld(World);
+				}
+				In.RestorePreAnimatedState();
+			}
+		}
+	);
+}
+
+void FLevelEditorSequencerIntegration::ResetToAnimatedState(UWorld* World)
+{
+	// Reset the time after saving so that an update will be triggered to put objects back to their animated state.
+	IterateAllSequencers(
+		[World](FSequencer& In, const FLevelEditorSequencerIntegrationOptions& Options)
+		{
+			if (Options.bRequiresLevelEvents)
+			{
+				In.InvalidateCachedData();
+				In.ForceEvaluate();
+
+				for (const TSharedPtr<ISequencerTrackEditor>& TrackEditor : In.GetTrackEditors())
+				{
+					TrackEditor->OnPostSaveWorld(World);
+				}
+			}
+		}
+	);
+}
+
 TSharedRef<FExtender> FLevelEditorSequencerIntegration::OnExtendLevelEditorViewMenu(const TSharedRef<FUICommandList> CommandList)
 {
 	// This would be where you added an extension to the Level Editor dropdown menus if needed.
diff --git a/Engine/Source/Editor/Sequencer/Public/LevelEditorSequencerIntegration.h b/Engine/Source/Editor/Sequencer/Public/LevelEditorSequencerIntegration.h
--- a/Engine/Source/Editor/Sequencer/Public/LevelEditorSequencerIntegration.h
+++ b/Engine/Source/Editor/Sequencer/Public/LevelEditorSequencerIntegration.h
@@ -95,6 +95,12 @@ private:
 	/** Called after the world has been saved. The sequencer updates to the animated state. */
 	void OnPostSaveWorld(UWorld* World, FObjectPostSaveContext ObjectSaveContext);
 
+	/** Called before any number of external actors are going to be saved. The sequencer puts everything back to its initial state. */
+	void OnPreSaveExternalActors(UWorld* World);
+
+	/** Called after any number of external actors has been saved. The sequencer puts everything back to its initial state. */
+	void OnPostSaveExternalActors(UWorld* World);
+
 	/** Called after a level has been added */
 	void OnLevelAdded(ULevel* InLevel, UWorld* InWorld);
 
@@ -172,6 +178,8 @@ private:
 	void DetachOutlinerColumn();
 	void ActivateRealtimeViewports();
 	void RestoreRealtimeViewports();
+	void RestoreToSavedState(UWorld* World);
+	void ResetToAnimatedState(UWorld* World);
 
 	struct FSequencerAndOptions
 	{
