Sequencer: Use relative transform when converting to possessable. If the current spawnable is attached to an object, getting the world transform will put the converted possessable in the correct world position, but then the attach and transform tracks will compound it.

Removing a parent will compensate the children at their world transform. We don't want that since we'll be replacing that parent right away. To negate that, we store the relative transform of these children and reset it after the parent is replaced with the new possessable.

diff --git a/Engine/Source/Editor/Sequencer/Private/Sequencer.cpp b/Engine/Source/Editor/Sequencer/Private/Sequencer.cpp
--- a/Engine/Source/Editor/Sequencer/Private/Sequencer.cpp
+++ b/Engine/Source/Editor/Sequencer/Private/Sequencer.cpp
@@ -10842,6 +10842,7 @@ FMovieScenePossessable* FSequencer::ConvertToPossessableInternal(FGuid Spawnable
 		return nullptr;
 	}
 
+	TMap<TWeakObjectPtr<AActor>, FTransform> AttachedChildTransforms;
 	FTransform DefaultTransform = SpawnableActorTemplate->GetActorTransform();
 	for (TWeakObjectPtr<> RuntimeObject : FindBoundObjects(SpawnableGuid, ActiveTemplateIDs.Top()) )
 	{
@@ -10850,7 +10851,30 @@ FMovieScenePossessable* FSequencer::ConvertToPossessableInternal(FGuid Spawnable
 		AActor* Actor = Cast<AActor>(RuntimeObject.Get());
 		if (Actor)
 		{
-			DefaultTransform = Actor->GetTransform();
+			DefaultTransform = Actor->GetRootComponent()->GetRelativeTransform();
+
+			// Removing a parent will compensate the children at their world transform. We don't want that since we'll be replacing that parent right away.
+			// To negate that, we store the relative transform of these children and reset it after the parent is replaced with the new possessable.
+			TArray<AActor*> AttachedActors;
+			Actor->GetAttachedActors(AttachedActors);
+			for (AActor* ChildActor : AttachedActors)
+			{
+				if (ChildActor)
+				{	
+					// Only do this for child actors that Sequencer is controlling
+					FGuid ExistingID;
+					for (FMovieSceneSequenceID SequenceID : ActiveTemplateIDs)
+					{
+						ExistingID = FindObjectId(*ChildActor, SequenceID);
+						if (ExistingID.IsValid())
+						{
+							AttachedChildTransforms.Add(ChildActor);
+							AttachedChildTransforms[ChildActor] = ChildActor->GetRootComponent()->GetRelativeTransform();
+							break;
+						}
+					}
+				}
+			}
 			break;
 		}
 	}
@@ -10925,6 +10949,14 @@ FMovieScenePossessable* FSequencer::ConvertToPossessableInternal(FGuid Spawnable
 
 		GEditor->SelectActor(PossessedActor, false, true);
 
+		for (TPair<TWeakObjectPtr<AActor>, FTransform> AttachedChildTransform : AttachedChildTransforms)
+		{
+			if (AActor* AttachedChild = AttachedChildTransform.Key.Get())
+			{
+				AttachedChild->GetRootComponent()->SetRelativeTransform(AttachedChildTransform.Value);
+			}
+		}
+
 		ForceEvaluate();
 	}
 
