Sequencer: Follow up to check for root component

diff --git a/Engine/Source/Editor/Sequencer/Private/Sequencer.cpp b/Engine/Source/Editor/Sequencer/Private/Sequencer.cpp
--- a/Engine/Source/Editor/Sequencer/Private/Sequencer.cpp
+++ b/Engine/Source/Editor/Sequencer/Private/Sequencer.cpp
@@ -10851,7 +10851,10 @@ FMovieScenePossessable* FSequencer::ConvertToPossessableInternal(FGuid Spawnable
 		AActor* Actor = Cast<AActor>(RuntimeObject.Get());
 		if (Actor)
 		{
-			DefaultTransform = Actor->GetRootComponent()->GetRelativeTransform();
+			if (Actor->GetRootComponent())
+			{
+				DefaultTransform = Actor->GetRootComponent()->GetRelativeTransform();
+			}
 
 			// Removing a parent will compensate the children at their world transform. We don't want that since we'll be replacing that parent right away.
 			// To negate that, we store the relative transform of these children and reset it after the parent is replaced with the new possessable.
@@ -10859,7 +10862,7 @@ FMovieScenePossessable* FSequencer::ConvertToPossessableInternal(FGuid Spawnable
 			Actor->GetAttachedActors(AttachedActors);
 			for (AActor* ChildActor : AttachedActors)
 			{
-				if (ChildActor)
+				if (ChildActor && ChildActor->GetRootComponent())
 				{	
 					// Only do this for child actors that Sequencer is controlling
 					FGuid ExistingID;
@@ -10953,7 +10956,10 @@ FMovieScenePossessable* FSequencer::ConvertToPossessableInternal(FGuid Spawnable
 		{
 			if (AActor* AttachedChild = AttachedChildTransform.Key.Get())
 			{
-				AttachedChild->GetRootComponent()->SetRelativeTransform(AttachedChildTransform.Value);
+				if (AttachedChild->GetRootComponent())
+				{
+					AttachedChild->GetRootComponent()->SetRelativeTransform(AttachedChildTransform.Value);
+				}
 			}
 		}
 
