Sequencer: Leave Old Motion Trails on by default.
#jra na
#trivial

diff --git a/Engine/Source/Editor/Sequencer/Private/SequencerEdMode.cpp b/Engine/Source/Editor/Sequencer/Private/SequencerEdMode.cpp
--- a/Engine/Source/Editor/Sequencer/Private/SequencerEdMode.cpp
+++ b/Engine/Source/Editor/Sequencer/Private/SequencerEdMode.cpp
@@ -41,7 +41,7 @@ static TAutoConsoleVariable<bool> CVarDrawMeshTrails(
 //if true still use old motion trails for sequencer objects.
 TAutoConsoleVariable<bool> CVarUseOldSequencerMotionTrails(
 	TEXT("Sequencer.UseOldSequencerTrails"),
-	false,
+	true,
 	TEXT("If true show old motion trails, if false use new editable motion trails."));
 
 
