Sequencer: before checking Slate to get modifier keys make sure we are in the right thead since we may be in a worker thread from anim/control rig evaluation.

diff --git a/Engine/Source/Editor/Sequencer/Private/MovieSceneTrackEditor.cpp b/Engine/Source/Editor/Sequencer/Private/MovieSceneTrackEditor.cpp
--- a/Engine/Source/Editor/Sequencer/Private/MovieSceneTrackEditor.cpp
+++ b/Engine/Source/Editor/Sequencer/Private/MovieSceneTrackEditor.cpp
@@ -92,7 +92,8 @@ void FMovieSceneTrackEditor::AnimatablePropertyChanged( FOnKeyProperty OnKeyProp
 		FKeyPropertyResult KeyPropertyResult = OnKeyProperty.Execute( KeyTime );
 
 		// When shift is down, enable adding keys/sections one after the other
-		if (FSlateApplication::Get().GetModifierKeys().IsShiftDown())
+		const bool bInProperSlateThread = (IsInGameThread() || IsInSlateThread() || IsInAsyncLoadingThread());
+		if (bInProperSlateThread && FSlateApplication::Get().GetModifierKeys().IsShiftDown())
 		{
 			for (TWeakObjectPtr<UMovieSceneSection> NewSection : KeyPropertyResult.SectionsCreated)
 			{
