Sequencer: Embedded curve view for float channels

Added menu options for float channels to "Show Curve" and set "Key Area Height". FMovieSceneFloatChannel now has a toggleable bShowCurve and the key area height is determined  by a sequencer setting, KeyAreaHeightWithCurves.

The FFortniteMainBranchObjectVersion has been bumped for the custom serialization of FMovieSceneFloatChannel with the new bShowCurve property.

diff --git a/Engine/Source/Editor/MovieSceneTools/Private/Channels/BuiltInChannelEditors.cpp b/Engine/Source/Editor/MovieSceneTools/Private/Channels/BuiltInChannelEditors.cpp
--- a/Engine/Source/Editor/MovieSceneTools/Private/Channels/BuiltInChannelEditors.cpp
+++ b/Engine/Source/Editor/MovieSceneTools/Private/Channels/BuiltInChannelEditors.cpp
@@ -8,6 +8,7 @@
 #include "ISequencerChannelInterface.h"
 #include "Widgets/SNullWidget.h"
 #include "ISequencer.h"
+#include "SequencerSettings.h"
 #include "MovieSceneCommonHelpers.h"
 #include "GameFramework/Actor.h"
 #include "EditorStyleSet.h"
@@ -1026,8 +1027,77 @@ struct FFloatChannelSectionMenuExtension : FExtender, TSharedFromThis<FFloatChan
 			LOCTEXT("SetPostInfinityExtrapTooltip", "Set post-infinity extrapolation"),
 			FNewMenuDelegate::CreateLambda([SharedThis](FMenuBuilder& SubMenuBuilder){ SharedThis->AddExtrapolationMenu(SubMenuBuilder, false); })
 			);
+
+		MenuBuilder.AddSubMenu(
+			LOCTEXT("DisplayOpyions", "Display"),
+			LOCTEXT("DisplayOptionsTooltip", "Display options"),
+			FNewMenuDelegate::CreateLambda([SharedThis](FMenuBuilder& SubMenuBuilder){ SharedThis->AddDisplayOptionsMenu(SubMenuBuilder); })
+			);
 	}
 
+	void AddDisplayOptionsMenu(FMenuBuilder& MenuBuilder)
+	{
+		TSharedRef<FFloatChannelSectionMenuExtension> SharedThis = AsShared();
+
+		// Menu entry for key area height
+		auto OnKeyAreaHeightChanged = [=](int32 NewValue) {
+
+			if (ISequencer* Sequencer = WeakSequencer.Pin().Get())
+			{
+				Sequencer->GetSequencerSettings()->SetKeyAreaHeightWithCurves((float)NewValue);
+			}
+		};
+
+		auto GetKeyAreaHeight = [=]() {
+
+			if (ISequencer* Sequencer = WeakSequencer.Pin().Get())
+			{
+				return (int)Sequencer->GetSequencerSettings()->GetKeyAreaHeightWithCurves();
+			}
+
+			return 15;
+		};
+
+
+		MenuBuilder.AddMenuEntry(
+			LOCTEXT("ToggleShowCurve", "Show Curve"),
+			LOCTEXT("ToggleShowCurveTooltip", "Toggle showing the curve in the track area"),
+			FSlateIcon(),
+			FUIAction(
+				FExecuteAction::CreateLambda([SharedThis]{ SharedThis->ToggleShowCurve(); }),
+				FCanExecuteAction(),
+				FIsActionChecked::CreateLambda([SharedThis]{ return SharedThis->IsShowCurve(); })
+			),
+			NAME_None,
+			EUserInterfaceActionType::ToggleButton
+		);	
+
+		MenuBuilder.AddWidget(
+			SNew(SHorizontalBox)
+			+ SHorizontalBox::Slot()
+				[
+					SNew(SSpacer)
+				]
+			+ SHorizontalBox::Slot()
+				.AutoWidth()
+				[
+					SNew(SBox)
+					.WidthOverride(30.f)
+					[
+						SNew(SSpinBox<int32>)
+						.Style(&FEditorStyle::GetWidgetStyle<FSpinBoxStyle>("Sequencer.HyperlinkSpinBox"))
+						.OnValueCommitted_Lambda([=](int32 Value, ETextCommit::Type CommitType) { OnKeyAreaHeightChanged(Value); })
+						.OnValueChanged_Lambda(OnKeyAreaHeightChanged)
+						.MinValue(15)
+						.MaxValue(300)
+						.Value_Lambda([=]() -> int32 { return GetKeyAreaHeight(); })
+					]
+				],
+				LOCTEXT("KeyAreaHeightText", "Key Area Height")
+		);
+	}
+
+
 	void AddExtrapolationMenu(FMenuBuilder& MenuBuilder, bool bPreInfinity)
 	{
 		TSharedRef<FFloatChannelSectionMenuExtension> SharedThis = AsShared();
@@ -1159,6 +1229,55 @@ struct FFloatChannelSectionMenuExtension : FExtender, TSharedFromThis<FFloatChan
 		return true;
 	}
 
+	void ToggleShowCurve()
+	{
+		bool bShowCurve = !IsShowCurve();
+		FScopedTransaction Transaction(LOCTEXT("ToggleShowCurve_Transaction", "Toggle Show Curve"));
+
+		bool bAnythingChanged = false;
+
+		// Modify all sections
+		for (TWeakObjectPtr<UMovieSceneSection> WeakSection : Sections)
+		{
+			if (UMovieSceneSection* Section = WeakSection.Get())
+			{
+				Section->Modify();
+			}
+		}
+
+		// Apply to all channels
+		for (const TMovieSceneChannelHandle<FMovieSceneFloatChannel>& Handle : Channels)
+		{
+			FMovieSceneFloatChannel* Channel = Handle.Get();
+
+			if (Channel)
+			{
+				Channel->SetShowCurve(bShowCurve);
+				bAnythingChanged = true;
+			}
+		}
+
+		if (!bAnythingChanged)
+		{
+			Transaction.Cancel();
+		}
+	}
+
+	bool IsShowCurve() const
+	{
+		for (const TMovieSceneChannelHandle<FMovieSceneFloatChannel>& Handle : Channels)
+		{
+			FMovieSceneFloatChannel* Channel = Handle.Get();
+
+			if (Channel && !Channel->GetShowCurve())
+			{
+				return false;
+			}
+		}
+
+		return true;
+	}
+
 private:
 
 	/** Hidden AsShared() methods to prevent CreateSP delegate use since this extender disappears with its menu. */
diff --git a/Engine/Source/Editor/Sequencer/Private/DisplayNodes/SequencerSectionKeyAreaNode.cpp b/Engine/Source/Editor/Sequencer/Private/DisplayNodes/SequencerSectionKeyAreaNode.cpp
--- a/Engine/Source/Editor/Sequencer/Private/DisplayNodes/SequencerSectionKeyAreaNode.cpp
+++ b/Engine/Source/Editor/Sequencer/Private/DisplayNodes/SequencerSectionKeyAreaNode.cpp
@@ -5,10 +5,12 @@
 #include "Widgets/SBoxPanel.h"
 #include "Widgets/Layout/SBox.h"
 #include "SSequencer.h"
+#include "SequencerSettings.h"
 #include "IKeyArea.h"
 #include "SKeyNavigationButtons.h"
 #include "SKeyAreaEditorSwitcher.h"
 #include "CurveModel.h"
+#include "Channels/MovieSceneFloatChannel.h"
 
 
 /* FSectionKeyAreaNode interface
@@ -91,8 +93,21 @@ FText FSequencerSectionKeyAreaNode::GetDisplayName() const
 
 float FSequencerSectionKeyAreaNode::GetNodeHeight() const
 {
-	//@todo sequencer: should be defined by the key area probably
-	return SequencerLayoutConstants::KeyAreaHeight;
+	float KeyAreaHeight = SequencerLayoutConstants::KeyAreaHeight;
+	for (const TSharedRef<IKeyArea>& KeyArea : KeyAreas)
+	{
+		const FMovieSceneChannelHandle& Channel = KeyArea->GetChannel();
+		if (Channel.GetChannelTypeName() == FMovieSceneFloatChannel::StaticStruct()->GetFName())
+		{
+			FMovieSceneFloatChannel* FloatChannel = Channel.Cast<FMovieSceneFloatChannel>().Get();
+			if (FloatChannel && FloatChannel->GetShowCurve())
+			{
+				KeyAreaHeight = FMath::Max(KeyAreaHeight, GetSequencer().GetSequencerSettings()->GetKeyAreaHeightWithCurves());
+			}
+		}
+	}
+
+	return KeyAreaHeight;
 }
 
 
diff --git a/Engine/Source/Editor/Sequencer/Private/SequencerKeyRenderer.cpp b/Engine/Source/Editor/Sequencer/Private/SequencerKeyRenderer.cpp
--- a/Engine/Source/Editor/Sequencer/Private/SequencerKeyRenderer.cpp
+++ b/Engine/Source/Editor/Sequencer/Private/SequencerKeyRenderer.cpp
@@ -445,9 +445,114 @@ void FKeyRenderer::FKeyDrawBatch::UpdateViewDependentData(FSequencer* Sequencer,
 
 		PrecomputedKeys.Add(NewKey);
 	}
+
+	PrecomputedCurve.Reset();
+
+	for (FCachedKeyDrawInformation& Info : KeyDrawInfo)
+	{
+		const TSharedPtr<IKeyArea>& ThisKeyArea = Info.CachedKeyPositions.GetKeyArea();
+		const FMovieSceneChannelHandle& Channel = ThisKeyArea.Get()->GetChannel();
+
+		if (Channel.GetChannelTypeName() != FMovieSceneFloatChannel::StaticStruct()->GetFName())
+		{
+			continue;
+		}
+			
+		FMovieSceneFloatChannel* FloatChannel = ThisKeyArea.Get()->GetChannel().Cast<FMovieSceneFloatChannel>().Get();
+		if (!FloatChannel->GetShowCurve())
+		{
+			continue;
+		}
+
+		const FFrameNumber DeltaFrame = TimeToPixelConverter.PixelDeltaToFrame(1.f).FrameNumber;		
+		const FFrameNumber LowerBoundFrame = (InCachedState.PaddedViewRange.GetLowerBoundValue() * TickResolution).CeilToFrame();
+		const FFrameNumber UpperBoundFrame = (InCachedState.PaddedViewRange.GetUpperBoundValue() * TickResolution).FloorToFrame();
+		TOptional<float> PreviousValue;
+		float MaxValue=-FLT_MAX, MinValue=FLT_MAX;
+		for (FFrameNumber FrameNumber = LowerBoundFrame; FrameNumber <= UpperBoundFrame; )
+		{
+			double EvalTime = TickResolution.AsSeconds(FrameNumber);
+			float Value = 0.f;
+			if (FloatChannel->Evaluate(FrameNumber, Value))
+			{
+				if ((PreviousValue.IsSet() && !FMath::IsNearlyEqual(Value, PreviousValue.GetValue())) || FrameNumber == LowerBoundFrame || FrameNumber == UpperBoundFrame)
+				{
+					MaxValue = FMath::Max(MaxValue, Value);
+					MinValue = FMath::Min(MinValue, Value);
+
+					FCurveKey NewKey;
+					NewKey.NormalizedValue = Value;
+					NewKey.FinalKeyPositionSeconds = EvalTime;
+
+					PrecomputedCurve.Add(NewKey);
+				}
+
+				if (!PreviousValue.IsSet())
+				{
+					PreviousValue = Value;
+				}
+			}
+
+			if (FrameNumber >= UpperBoundFrame)
+			{
+				break;
+			}
+			FrameNumber += DeltaFrame;
+			FrameNumber = FMath::Min(FrameNumber, UpperBoundFrame);
+		}
+
+		// Normalize
+		if (PrecomputedCurve.Num() > 0)
+		{
+			float DiffValue = MaxValue - MinValue;
+			for (FCurveKey& CurveKey : PrecomputedCurve)
+			{
+				CurveKey.NormalizedValue = (CurveKey.NormalizedValue - MinValue)/DiffValue;
+			}
+		}
+	}
 }
 
-void FKeyRenderer::FKeyDrawBatch::Draw(FSequencer* Sequencer, FSequencerSectionPainter& Painter, const FGeometry& KeyGeometry, const FPaintStyle& Style, const FKeyRendererPaintArgs& Args) const
+void FKeyRenderer::FKeyDrawBatch::DrawCurve(FSequencer* Sequencer, FSequencerSectionPainter& Painter, const FGeometry& KeyGeometry, const FPaintStyle& Style, const FKeyRendererPaintArgs& Args) const
+{
+	const FTimeToPixel& TimeToPixelConverter = Painter.GetTimeConverter();
+
+	TOptional<FSlateClippingState> PreviousClipState = Painter.DrawElements.GetClippingState();
+	Painter.DrawElements.PopClip();
+
+	const int32 KeyLayer = Painter.LayerId;
+
+	const ESlateDrawEffect BaseDrawEffects = Painter.bParentEnabled ? ESlateDrawEffect::None : ESlateDrawEffect::DisabledEffect;
+
+	TArray<FVector2D> CurvePoints;
+	for (const FCurveKey& CurveKey : PrecomputedCurve)
+	{		
+		CurvePoints.Add(
+		FVector2D(TimeToPixelConverter.SecondsToPixel(CurveKey.FinalKeyPositionSeconds),
+				(1.0f - CurveKey.NormalizedValue) * KeyGeometry.GetLocalSize().Y ) );
+	}
+
+	const float CurveThickness = 1.f;
+	const bool  bAntiAliasCurves = true;
+	const FLinearColor CurveColor(0.8f, 0.8f, 0.f, 0.7);
+
+	FSlateDrawElement::MakeLines(
+		Painter.DrawElements,
+		KeyLayer,
+		KeyGeometry.ToPaintGeometry(),
+		CurvePoints,
+		BaseDrawEffects,
+		CurveColor,
+		bAntiAliasCurves,
+		CurveThickness
+	);
+
+	Painter.LayerId = KeyLayer + 2;
+	Painter.DrawElements.GetClippingManager().PushClippingState(PreviousClipState.GetValue());
+}
+
+void
+FKeyRenderer::FKeyDrawBatch::Draw(FSequencer* Sequencer, FSequencerSectionPainter& Painter, const FGeometry& KeyGeometry, const FPaintStyle& Style, const FKeyRendererPaintArgs& Args) const
 {
 	const FTimeToPixel& TimeToPixelConverter = Painter.GetTimeConverter();
 
@@ -735,12 +840,17 @@ void FKeyRenderer::Paint(const FSectionLayout& InSectionLayout, const FWidgetSty
 		if (const FKeyDrawBatch* KeyDrawBatch = CachedKeyLayouts.Find(LayoutElement))
 		{
 			FGeometry KeyGeometry = LayoutElement.ComputeGeometry(InPainter.SectionGeometry);
+			
+			if (LayoutElement.GetType() == FSectionLayoutElement::Single)
+			{
+				KeyDrawBatch->DrawCurve(Sequencer, InPainter, KeyGeometry, Style, Args);
+			}
+			
 			KeyDrawBatch->Draw(Sequencer, InPainter, KeyGeometry, Style, Args);
 		}
 	}
 }
 
-
 } // namespace Sequencer
 } // namespace UE
 
diff --git a/Engine/Source/Editor/Sequencer/Private/SequencerKeyRenderer.h b/Engine/Source/Editor/Sequencer/Private/SequencerKeyRenderer.h
--- a/Engine/Source/Editor/Sequencer/Private/SequencerKeyRenderer.h
+++ b/Engine/Source/Editor/Sequencer/Private/SequencerKeyRenderer.h
@@ -179,6 +179,15 @@ private:
 		EKeyRenderingFlags Flags = EKeyRenderingFlags::None;
 	};
 
+	struct FCurveKey
+	{
+		/** Normalized key value */
+		float NormalizedValue;
+
+		/** The time in seconds that this key should be drawn - represents the average time for overlapping keys */
+		double FinalKeyPositionSeconds;
+	};
+
 
 	struct FPaintStyle
 	{
@@ -213,6 +222,11 @@ private:
 		void UpdateViewDependentData(FSequencer* Sequencer, const FSequencerSectionPainter& InPainter, const FCachedState& CachedState, ECacheFlags CacheFlags);
 
 
+		/**
+		 * Draw this batch's curve
+		 */
+		void DrawCurve(FSequencer* Sequencer, FSequencerSectionPainter& Painter, const FGeometry& KeyGeometry, const FPaintStyle& Style, const FKeyRendererPaintArgs& Args) const;
+
 		/**
 		 * Draw this batch
 		 */
@@ -225,6 +239,9 @@ private:
 
 		/** Computed final draw info */
 		TArray<FKey> PrecomputedKeys;
+
+		/** Computed final draw curve info */
+		TArray<FCurveKey> PrecomputedCurve;
 	};
 
 private:
diff --git a/Engine/Source/Editor/Sequencer/Private/SequencerSettings.cpp b/Engine/Source/Editor/Sequencer/Private/SequencerSettings.cpp
--- a/Engine/Source/Editor/Sequencer/Private/SequencerSettings.cpp
+++ b/Engine/Source/Editor/Sequencer/Private/SequencerSettings.cpp
@@ -3,6 +3,7 @@
 #include "SequencerSettings.h"
 #include "KeyParams.h"
 #include "ISequencer.h"
+#include "SSequencer.h"
 
 USequencerSettings::USequencerSettings( const FObjectInitializer& ObjectInitializer )
 	: Super( ObjectInitializer )
@@ -46,6 +47,7 @@ USequencerSettings::USequencerSettings( const FObjectInitializer& ObjectInitiali
 	bInfiniteKeyAreas = false;
 	bShowChannelColors = false;
 	ReduceKeysTolerance = KINDA_SMALL_NUMBER;
+	KeyAreaHeightWithCurves = SequencerLayoutConstants::KeyAreaHeight;
 	bDeleteKeysWhenTrimming = true;
 	bDisableSectionsAfterBaking = true;
 	bCleanPlaybackMode = true;
@@ -565,6 +567,20 @@ void USequencerSettings::SetShowChannelColors(bool InbShowChannelColors)
 	}
 }
 
+float USequencerSettings::GetKeyAreaHeightWithCurves() const
+{
+	return KeyAreaHeightWithCurves;
+}
+
+void USequencerSettings::SetKeyAreaHeightWithCurves(float InKeyAreaHeightWithCurves)
+{
+	if (KeyAreaHeightWithCurves != InKeyAreaHeightWithCurves)
+	{
+		KeyAreaHeightWithCurves = InKeyAreaHeightWithCurves;
+		SaveConfig();
+	}
+}
+
 float USequencerSettings::GetReduceKeysTolerance() const
 {
 	return ReduceKeysTolerance;
diff --git a/Engine/Source/Editor/Sequencer/Public/SequencerSettings.h b/Engine/Source/Editor/Sequencer/Public/SequencerSettings.h
--- a/Engine/Source/Editor/Sequencer/Public/SequencerSettings.h
+++ b/Engine/Source/Editor/Sequencer/Public/SequencerSettings.h
@@ -279,6 +279,11 @@ public:
 	/** Set whether to show channel colors */
 	void SetShowChannelColors(bool bInShowChannelColors);
 
+	/** @return The key area height when showing curves */
+	float GetKeyAreaHeightWithCurves() const;
+	/** Set the key area height when showing curves */
+	void SetKeyAreaHeightWithCurves(float InKeyAreaHeightWithCurves);
+
 	/** @return The tolerance to use when reducing keys */
 	float GetReduceKeysTolerance() const;
 	/** Set the tolerance to use when reducing keys */
@@ -506,6 +511,10 @@ protected:
 	UPROPERTY(config, EditAnywhere, Category = Timeline)
 	bool bShowChannelColors;
 
+	/** The key area height when showing curves */
+	UPROPERTY(config, EditAnywhere, Category = Timeline)
+	float KeyAreaHeightWithCurves;
+
 	/** The tolerance to use when reducing keys */
 	UPROPERTY(config, EditAnywhere, Category = Timeline)
 	float ReduceKeysTolerance;
diff --git a/Engine/Source/Runtime/Core/Public/UObject/FortniteMainBranchObjectVersion.h b/Engine/Source/Runtime/Core/Public/UObject/FortniteMainBranchObjectVersion.h
--- a/Engine/Source/Runtime/Core/Public/UObject/FortniteMainBranchObjectVersion.h
+++ b/Engine/Source/Runtime/Core/Public/UObject/FortniteMainBranchObjectVersion.h
@@ -167,6 +167,9 @@ struct CORE_API FFortniteMainBranchObjectVersion
 		// Enable material (edge, bending, and area stiffness) weight maps
 		ChaosClothAddMaterialWeightMaps,
 
+		// Added bShowCurve for movie scene float channel serialization
+		SerializeFloatChannelShowCurve,
+
 		// -----<new versions can be added above this line>-------------------------------------------------
 		VersionPlusOne,
 		LatestVersion = VersionPlusOne - 1
diff --git a/Engine/Source/Runtime/MovieScene/Private/Channels/MovieSceneFloatChannel.cpp b/Engine/Source/Runtime/MovieScene/Private/Channels/MovieSceneFloatChannel.cpp
--- a/Engine/Source/Runtime/MovieScene/Private/Channels/MovieSceneFloatChannel.cpp
+++ b/Engine/Source/Runtime/MovieScene/Private/Channels/MovieSceneFloatChannel.cpp
@@ -5,6 +5,7 @@
 #include "Channels/MovieSceneChannelProxy.h"
 #include "MovieSceneFrameMigration.h"
 #include "UObject/SequencerObjectVersion.h"
+#include "UObject/FortniteMainBranchObjectVersion.h"
 #include "HAL/ConsoleManager.h"
 
 int32 GSequencerLinearCubicInterpolation = 1;
@@ -1077,10 +1078,26 @@ void FMovieSceneFloatChannel::AddKeys(const TArray<FFrameNumber>& InTimes, const
 	}
 }
 
+#if WITH_EDITORONLY_DATA
+
+bool FMovieSceneFloatChannel::GetShowCurve() const
+{
+	return bShowCurve;
+}
+
+void FMovieSceneFloatChannel::SetShowCurve(bool bInShowCurve)
+{
+	bShowCurve = bInShowCurve;
+}
+
+#endif
+
 bool FMovieSceneFloatChannel::Serialize(FArchive& Ar)
 {
 	Ar.UsingCustomVersion(FSequencerObjectVersion::GUID);
-	if (Ar.CustomVer(FSequencerObjectVersion::GUID) < FSequencerObjectVersion::SerializeFloatChannelCompletely)
+	Ar.UsingCustomVersion(FFortniteMainBranchObjectVersion::GUID);
+	if (Ar.CustomVer(FSequencerObjectVersion::GUID) < FSequencerObjectVersion::SerializeFloatChannelCompletely &&
+		Ar.CustomVer(FFortniteMainBranchObjectVersion::GUID) < FFortniteMainBranchObjectVersion::SerializeFloatChannelShowCurve)
 	{
 		return false;
 	}
@@ -1161,6 +1178,13 @@ bool FMovieSceneFloatChannel::Serialize(FArchive& Ar)
 	{
 		Ar << KeyHandles;
 	}
+
+	if (Ar.CustomVer(FFortniteMainBranchObjectVersion::GUID) >= FFortniteMainBranchObjectVersion::SerializeFloatChannelShowCurve)
+	{
+#if WITH_EDITOR
+	Ar << bShowCurve;
+#endif
+	}
 	return true;
 }
 
diff --git a/Engine/Source/Runtime/MovieScene/Public/Channels/MovieSceneFloatChannel.h b/Engine/Source/Runtime/MovieScene/Public/Channels/MovieSceneFloatChannel.h
--- a/Engine/Source/Runtime/MovieScene/Public/Channels/MovieSceneFloatChannel.h
+++ b/Engine/Source/Runtime/MovieScene/Public/Channels/MovieSceneFloatChannel.h
@@ -141,8 +141,14 @@ struct MOVIESCENE_API FMovieSceneFloatChannel : public FMovieSceneChannel
 {
 	GENERATED_BODY()
 
-	FMovieSceneFloatChannel()
-		: PreInfinityExtrap(RCCE_Constant), PostInfinityExtrap(RCCE_Constant), DefaultValue(), bHasDefaultValue(false)
+	FMovieSceneFloatChannel() 
+		: PreInfinityExtrap(RCCE_Constant)
+		, PostInfinityExtrap(RCCE_Constant)
+		, DefaultValue()
+		, bHasDefaultValue(false)
+#if WITH_EDITORONLY_DATA
+		, bShowCurve(false)
+#endif
 	{}
 
 	/**
@@ -313,6 +319,14 @@ public:
 	* @param InValues Values to add
 	*/
 	void AddKeys(const TArray<FFrameNumber>& InTimes, const TArray<FMovieSceneFloatValue>& InValues);
+
+#if WITH_EDITORONLY_DATA
+	/* Get whether to show this curve in the UI */
+	bool GetShowCurve() const;
+	/* Set whether to show this curve in the UI */
+	void SetShowCurve(bool bInShowCurve);
+#endif
+
 public:
 
 	/** Pre-infinity extrapolation state */
@@ -365,6 +379,13 @@ private:
 	UPROPERTY()
 	FFrameRate TickResolution;
 
+#if WITH_EDITORONLY_DATA
+
+	UPROPERTY()
+	bool bShowCurve;
+
+#endif
+
 public:
 	//Set it's frame resolution
 	void SetTickResolution(FFrameRate InTickSolution)
