Sequencer: Improve drag and drop to indicate drop position and length if set

ISequencerTrackEditor::OnAllowDrop and OnDrop now takes FSequencerDragDropParams, which includes the frame number being dropped onto

diff --git a/Engine/Source/Editor/MovieSceneTools/Private/TrackEditors/AudioTrackEditor.cpp b/Engine/Source/Editor/MovieSceneTools/Private/TrackEditors/AudioTrackEditor.cpp
--- a/Engine/Source/Editor/MovieSceneTools/Private/TrackEditors/AudioTrackEditor.cpp
+++ b/Engine/Source/Editor/MovieSceneTools/Private/TrackEditors/AudioTrackEditor.cpp
@@ -1009,9 +1009,9 @@ void FAudioTrackEditor::Resize(float NewSize, UMovieSceneTrack* InTrack)
 	}
 }
 
-bool FAudioTrackEditor::OnAllowDrop(const FDragDropEvent& DragDropEvent, UMovieSceneTrack* Track, int32 RowIndex, const FGuid& TargetObjectGuid)
+bool FAudioTrackEditor::OnAllowDrop(const FDragDropEvent& DragDropEvent, FSequencerDragDropParams& DragDropParams)
 {
-	if (!Track->IsA(UMovieSceneAudioTrack::StaticClass()))
+	if (!DragDropParams.Track->IsA(UMovieSceneAudioTrack::StaticClass()))
 	{
 		return false;
 	}
@@ -1027,8 +1027,11 @@ bool FAudioTrackEditor::OnAllowDrop(const FDragDropEvent& DragDropEvent, UMovieS
 
 	for (const FAssetData& AssetData : DragDropOp->GetAssets())
 	{
-		if (Cast<USoundBase>(AssetData.GetAsset()))
+		if (USoundBase* Sound = Cast<USoundBase>(AssetData.GetAsset()))
 		{
+			FFrameRate TickResolution = GetSequencer()->GetFocusedTickResolution();
+			FFrameNumber LengthInFrames = TickResolution.AsFrameNumber(Sound->GetDuration());
+			DragDropParams.FrameRange = TRange<FFrameNumber>(DragDropParams.FrameNumber, DragDropParams.FrameNumber + LengthInFrames);
 			return true;
 		}
 	}
@@ -1037,9 +1040,9 @@ bool FAudioTrackEditor::OnAllowDrop(const FDragDropEvent& DragDropEvent, UMovieS
 }
 
 
-FReply FAudioTrackEditor::OnDrop(const FDragDropEvent& DragDropEvent, UMovieSceneTrack* Track, int32 RowIndex, const FGuid& TargetObjectGuid)
+FReply FAudioTrackEditor::OnDrop(const FDragDropEvent& DragDropEvent, const FSequencerDragDropParams& DragDropParams)
 {
-	if (!Track->IsA(UMovieSceneAudioTrack::StaticClass()))
+	if (!DragDropParams.Track->IsA(UMovieSceneAudioTrack::StaticClass()))
 	{
 		return FReply::Unhandled();
 	}
@@ -1051,13 +1054,13 @@ FReply FAudioTrackEditor::OnDrop(const FDragDropEvent& DragDropEvent, UMovieScen
 		return FReply::Unhandled();
 	}
 	
-	UMovieSceneAudioTrack* AudioTrack = Cast<UMovieSceneAudioTrack>(Track);
+	UMovieSceneAudioTrack* AudioTrack = Cast<UMovieSceneAudioTrack>(DragDropParams.Track);
 
 	const FScopedTransaction Transaction(LOCTEXT("DropAssets", "Drop Assets"));
 
 	TSharedPtr<FAssetDragDropOp> DragDropOp = StaticCastSharedPtr<FAssetDragDropOp>( Operation );
 
-	FMovieSceneTrackEditor::BeginKeying();
+	FMovieSceneTrackEditor::BeginKeying(DragDropParams.FrameNumber);
 
 	bool bAnyDropped = false;
 	for (const FAssetData& AssetData : DragDropOp->GetAssets())
@@ -1066,10 +1069,10 @@ FReply FAudioTrackEditor::OnDrop(const FDragDropEvent& DragDropEvent, UMovieScen
 
 		if (Sound)
 		{
-			if (TargetObjectGuid.IsValid())
+			if (DragDropParams.TargetObjectGuid.IsValid())
 			{
 				TArray<TWeakObjectPtr<>> OutObjects;
-				for (TWeakObjectPtr<> Object : GetSequencer()->FindObjectsInCurrentSequence(TargetObjectGuid))
+				for (TWeakObjectPtr<> Object : GetSequencer()->FindObjectsInCurrentSequence(DragDropParams.TargetObjectGuid))
 				{
 					OutObjects.Add(Object);
 				}
@@ -1078,7 +1081,7 @@ FReply FAudioTrackEditor::OnDrop(const FDragDropEvent& DragDropEvent, UMovieScen
 			}
 			else
 			{
-				AnimatablePropertyChanged(FOnKeyProperty::CreateRaw(this, &FAudioTrackEditor::AddNewMasterSound, Sound, AudioTrack, RowIndex));
+				AnimatablePropertyChanged(FOnKeyProperty::CreateRaw(this, &FAudioTrackEditor::AddNewMasterSound, Sound, AudioTrack, DragDropParams.RowIndex));
 			}
 
 			bAnyDropped = true;
diff --git a/Engine/Source/Editor/MovieSceneTools/Private/TrackEditors/CameraCutTrackEditor.cpp b/Engine/Source/Editor/MovieSceneTools/Private/TrackEditors/CameraCutTrackEditor.cpp
--- a/Engine/Source/Editor/MovieSceneTools/Private/TrackEditors/CameraCutTrackEditor.cpp
+++ b/Engine/Source/Editor/MovieSceneTools/Private/TrackEditors/CameraCutTrackEditor.cpp
@@ -273,9 +273,9 @@ const FSlateBrush* FCameraCutTrackEditor::GetIconBrush() const
 }
 
 
-bool FCameraCutTrackEditor::OnAllowDrop(const FDragDropEvent& DragDropEvent, UMovieSceneTrack* Track, int32 RowIndex, const FGuid& TargetObjectGuid)
+bool FCameraCutTrackEditor::OnAllowDrop(const FDragDropEvent& DragDropEvent, FSequencerDragDropParams& DragDropParams)
 {
-	if (!Track->IsA(UMovieSceneCameraCutTrack::StaticClass()))
+	if (!DragDropParams.Track->IsA(UMovieSceneCameraCutTrack::StaticClass()))
 	{
 		return false;
 	}
@@ -287,6 +287,8 @@ bool FCameraCutTrackEditor::OnAllowDrop(const FDragDropEvent& DragDropEvent, UMo
 		return false;
 	}
 	
+	UMovieSceneCameraCutTrack* CameraCutTrack = Cast<UMovieSceneCameraCutTrack>(DragDropParams.Track);
+
 	TSharedPtr<FActorDragDropGraphEdOp> DragDropOp = StaticCastSharedPtr<FActorDragDropGraphEdOp>( Operation );
 
 	for (auto& ActorPtr : DragDropOp->Actors)
@@ -298,6 +300,8 @@ bool FCameraCutTrackEditor::OnAllowDrop(const FDragDropEvent& DragDropEvent, UMo
 			UCameraComponent* CameraComponent = MovieSceneHelpers::CameraComponentFromActor(Actor);
 			if (CameraComponent)
 			{
+				FFrameNumber EndFrameNumber = CameraCutTrack->FindEndTimeForCameraCut(DragDropParams.FrameNumber);
+				DragDropParams.FrameRange = TRange<FFrameNumber>(DragDropParams.FrameNumber, EndFrameNumber);
 				return true;
 			}
 		}
@@ -307,9 +311,9 @@ bool FCameraCutTrackEditor::OnAllowDrop(const FDragDropEvent& DragDropEvent, UMo
 }
 
 
-FReply FCameraCutTrackEditor::OnDrop(const FDragDropEvent& DragDropEvent, UMovieSceneTrack* Track, int32 RowIndex, const FGuid& TargetObjectGuid)
+FReply FCameraCutTrackEditor::OnDrop(const FDragDropEvent& DragDropEvent, const FSequencerDragDropParams& DragDropParams)
 {
-	if (!Track->IsA(UMovieSceneCameraCutTrack::StaticClass()))
+	if (!DragDropParams.Track->IsA(UMovieSceneCameraCutTrack::StaticClass()))
 	{
 		return FReply::Unhandled();
 	}
@@ -323,7 +327,7 @@ FReply FCameraCutTrackEditor::OnDrop(const FDragDropEvent& DragDropEvent, UMovie
 	
 	TSharedPtr<FActorDragDropGraphEdOp> DragDropOp = StaticCastSharedPtr<FActorDragDropGraphEdOp>( Operation );
 
-	FMovieSceneTrackEditor::BeginKeying();
+	FMovieSceneTrackEditor::BeginKeying(DragDropParams.FrameNumber);
 
 	bool bAnyDropped = false;
 	for (auto& ActorPtr : DragDropOp->Actors)
diff --git a/Engine/Source/Editor/MovieSceneTools/Private/TrackEditors/CameraCutTrackEditor.h b/Engine/Source/Editor/MovieSceneTools/Private/TrackEditors/CameraCutTrackEditor.h
--- a/Engine/Source/Editor/MovieSceneTools/Private/TrackEditors/CameraCutTrackEditor.h
+++ b/Engine/Source/Editor/MovieSceneTools/Private/TrackEditors/CameraCutTrackEditor.h
@@ -62,8 +62,8 @@ public:
 	virtual bool SupportsSequence(UMovieSceneSequence* InSequence) const override;
 	virtual void Tick(float DeltaTime) override;
 	virtual const FSlateBrush* GetIconBrush() const override;
-	virtual bool OnAllowDrop(const FDragDropEvent& DragDropEvent, UMovieSceneTrack* Track, int32 RowIndex, const FGuid& TargetObjectGuid) override;
-	virtual FReply OnDrop(const FDragDropEvent& DragDropEvent, UMovieSceneTrack* Track, int32 RowIndex, const FGuid& TargetObjectGuid) override;
+	virtual bool OnAllowDrop(const FDragDropEvent& DragDropEvent, FSequencerDragDropParams& DragDropParams) override;
+	virtual FReply OnDrop(const FDragDropEvent& DragDropEvent, const FSequencerDragDropParams& DragDropParams) override;
 
 protected:
 
diff --git a/Engine/Source/Editor/MovieSceneTools/Private/TrackEditors/CinematicShotTrackEditor.cpp b/Engine/Source/Editor/MovieSceneTools/Private/TrackEditors/CinematicShotTrackEditor.cpp
--- a/Engine/Source/Editor/MovieSceneTools/Private/TrackEditors/CinematicShotTrackEditor.cpp
+++ b/Engine/Source/Editor/MovieSceneTools/Private/TrackEditors/CinematicShotTrackEditor.cpp
@@ -258,9 +258,9 @@ const FSlateBrush* FCinematicShotTrackEditor::GetIconBrush() const
 	return FEditorStyle::GetBrush("Sequencer.Tracks.CinematicShot");
 }
 
-bool FCinematicShotTrackEditor::OnAllowDrop(const FDragDropEvent& DragDropEvent, UMovieSceneTrack* Track, int32 RowIndex, const FGuid& TargetObjectGuid)
+bool FCinematicShotTrackEditor::OnAllowDrop(const FDragDropEvent& DragDropEvent, FSequencerDragDropParams& DragDropParams)
 {
-	if (!Track->IsA(UMovieSceneCinematicShotTrack::StaticClass()))
+	if (!DragDropParams.Track->IsA(UMovieSceneCinematicShotTrack::StaticClass()))
 	{
 		return false;
 	}
@@ -276,8 +276,16 @@ bool FCinematicShotTrackEditor::OnAllowDrop(const FDragDropEvent& DragDropEvent,
 
 	for (const FAssetData& AssetData : DragDropOp->GetAssets())
 	{
-		if (Cast<UMovieSceneSequence>(AssetData.GetAsset()))
+		if (UMovieSceneSequence* Sequence = Cast<UMovieSceneSequence>(AssetData.GetAsset()))
 		{
+			FFrameRate TickResolution = GetSequencer()->GetFocusedTickResolution();
+
+			const FQualifiedFrameTime InnerDuration = FQualifiedFrameTime(
+				UE::MovieScene::DiscreteSize(Sequence->GetMovieScene()->GetPlaybackRange()),
+				Sequence->GetMovieScene()->GetTickResolution());
+
+			FFrameNumber LengthInFrames = InnerDuration.ConvertTo(TickResolution).FrameNumber;
+			DragDropParams.FrameRange = TRange<FFrameNumber>(DragDropParams.FrameNumber, DragDropParams.FrameNumber + LengthInFrames);
 			return true;
 		}
 	}
@@ -286,9 +294,9 @@ bool FCinematicShotTrackEditor::OnAllowDrop(const FDragDropEvent& DragDropEvent,
 }
 
 
-FReply FCinematicShotTrackEditor::OnDrop(const FDragDropEvent& DragDropEvent, UMovieSceneTrack* Track, int32 RowIndex, const FGuid& TargetObjectGuid)
+FReply FCinematicShotTrackEditor::OnDrop(const FDragDropEvent& DragDropEvent, const FSequencerDragDropParams& DragDropParams)
 {
-	if (!Track->IsA(UMovieSceneCinematicShotTrack::StaticClass()))
+	if (!DragDropParams.Track->IsA(UMovieSceneCinematicShotTrack::StaticClass()))
 	{
 		return FReply::Unhandled();
 	}
@@ -304,7 +312,7 @@ FReply FCinematicShotTrackEditor::OnDrop(const FDragDropEvent& DragDropEvent, UM
 
 	TSharedPtr<FAssetDragDropOp> DragDropOp = StaticCastSharedPtr<FAssetDragDropOp>( Operation );
 	
-	FMovieSceneTrackEditor::BeginKeying();
+	FMovieSceneTrackEditor::BeginKeying(DragDropParams.FrameNumber);
 
 	bool bAnyDropped = false;
 	for (const FAssetData& AssetData : DragDropOp->GetAssets())
@@ -313,7 +321,7 @@ FReply FCinematicShotTrackEditor::OnDrop(const FDragDropEvent& DragDropEvent, UM
 
 		if (Sequence)
 		{
-			AnimatablePropertyChanged( FOnKeyProperty::CreateRaw( this, &FCinematicShotTrackEditor::AddKeyInternal, Sequence, RowIndex) );
+			AnimatablePropertyChanged( FOnKeyProperty::CreateRaw( this, &FCinematicShotTrackEditor::AddKeyInternal, Sequence, DragDropParams.RowIndex) );
 			
 			bAnyDropped = true;
 		}
diff --git a/Engine/Source/Editor/MovieSceneTools/Private/TrackEditors/CinematicShotTrackEditor.h b/Engine/Source/Editor/MovieSceneTools/Private/TrackEditors/CinematicShotTrackEditor.h
--- a/Engine/Source/Editor/MovieSceneTools/Private/TrackEditors/CinematicShotTrackEditor.h
+++ b/Engine/Source/Editor/MovieSceneTools/Private/TrackEditors/CinematicShotTrackEditor.h
@@ -62,8 +62,8 @@ public:
 	virtual void Tick(float DeltaTime) override;
 	virtual void BuildTrackContextMenu( FMenuBuilder& MenuBuilder, UMovieSceneTrack* Track ) override;
 	virtual const FSlateBrush* GetIconBrush() const override;
-	virtual bool OnAllowDrop(const FDragDropEvent& DragDropEvent, UMovieSceneTrack* Track, int32 RowIndex, const FGuid& TargetObjectGuid) override;
-	virtual FReply OnDrop(const FDragDropEvent& DragDropEvent, UMovieSceneTrack* Track, int32 RowIndex, const FGuid& TargetObjectGuid) override;
+	virtual bool OnAllowDrop(const FDragDropEvent& DragDropEvent, FSequencerDragDropParams& DragDropParams) override;
+	virtual FReply OnDrop(const FDragDropEvent& DragDropEvent, const FSequencerDragDropParams& DragDropParams) override;
 
 	/*
 	 * Insert shot. 
diff --git a/Engine/Source/Editor/MovieSceneTools/Private/TrackEditors/SkeletalAnimationTrackEditor.cpp b/Engine/Source/Editor/MovieSceneTools/Private/TrackEditors/SkeletalAnimationTrackEditor.cpp
--- a/Engine/Source/Editor/MovieSceneTools/Private/TrackEditors/SkeletalAnimationTrackEditor.cpp
+++ b/Engine/Source/Editor/MovieSceneTools/Private/TrackEditors/SkeletalAnimationTrackEditor.cpp
@@ -1737,9 +1737,9 @@ TSharedPtr<SWidget> FSkeletalAnimationTrackEditor::BuildOutlinerEditWidget(const
 	}
 }
 
-bool FSkeletalAnimationTrackEditor::OnAllowDrop(const FDragDropEvent& DragDropEvent, UMovieSceneTrack* Track, int32 RowIndex, const FGuid& TargetObjectGuid)
+bool FSkeletalAnimationTrackEditor::OnAllowDrop(const FDragDropEvent& DragDropEvent, FSequencerDragDropParams& DragDropParams)
 {
-	if (!Track->IsA(UMovieSceneSkeletalAnimationTrack::StaticClass()))
+	if (!DragDropParams.Track->IsA(UMovieSceneSkeletalAnimationTrack::StaticClass()))
 	{
 		return false;
 	}
@@ -1751,12 +1751,12 @@ bool FSkeletalAnimationTrackEditor::OnAllowDrop(const FDragDropEvent& DragDropEv
 		return false;
 	}
 	
-	if (!TargetObjectGuid.IsValid())
+	if (!DragDropParams.TargetObjectGuid.IsValid())
 	{
 		return false;
 	}
 
-	USkeleton* Skeleton = AcquireSkeletonFromObjectGuid(TargetObjectGuid, GetSequencer());
+	USkeleton* Skeleton = AcquireSkeletonFromObjectGuid(DragDropParams.TargetObjectGuid, GetSequencer());
 
 	TSharedPtr<FAssetDragDropOp> DragDropOp = StaticCastSharedPtr<FAssetDragDropOp>( Operation );
 
@@ -1767,6 +1767,9 @@ bool FSkeletalAnimationTrackEditor::OnAllowDrop(const FDragDropEvent& DragDropEv
 		const bool bValidAnimSequence = AnimSequence && AnimSequence->CanBeUsedInComposition();
 		if (bValidAnimSequence && Skeleton && Skeleton->IsCompatible(AnimSequence->GetSkeleton()))
 		{
+			FFrameRate TickResolution = GetSequencer()->GetFocusedTickResolution();
+			FFrameNumber LengthInFrames = TickResolution.AsFrameNumber(AnimSequence->GetPlayLength());
+			DragDropParams.FrameRange = TRange<FFrameNumber>(DragDropParams.FrameNumber, DragDropParams.FrameNumber + LengthInFrames);
 			return true;
 		}
 	}
@@ -1775,9 +1778,9 @@ bool FSkeletalAnimationTrackEditor::OnAllowDrop(const FDragDropEvent& DragDropEv
 }
 
 
-FReply FSkeletalAnimationTrackEditor::OnDrop(const FDragDropEvent& DragDropEvent, UMovieSceneTrack* Track, int32 RowIndex, const FGuid& TargetObjectGuid)
+FReply FSkeletalAnimationTrackEditor::OnDrop(const FDragDropEvent& DragDropEvent, const FSequencerDragDropParams& DragDropParams)
 {
-	if (!Track->IsA(UMovieSceneSkeletalAnimationTrack::StaticClass()))
+	if (!DragDropParams.Track->IsA(UMovieSceneSkeletalAnimationTrack::StaticClass()))
 	{
 		return FReply::Unhandled();
 	}
@@ -1789,18 +1792,18 @@ FReply FSkeletalAnimationTrackEditor::OnDrop(const FDragDropEvent& DragDropEvent
 		return FReply::Unhandled();
 	}
 	
-	if (!TargetObjectGuid.IsValid())
+	if (!DragDropParams.TargetObjectGuid.IsValid())
 	{
 		return FReply::Unhandled();
 	}
 
-	USkeleton* Skeleton = AcquireSkeletonFromObjectGuid(TargetObjectGuid, GetSequencer());
+	USkeleton* Skeleton = AcquireSkeletonFromObjectGuid(DragDropParams.TargetObjectGuid, GetSequencer());
 
 	const FScopedTransaction Transaction(LOCTEXT("DropAssets", "Drop Assets"));
 
 	TSharedPtr<FAssetDragDropOp> DragDropOp = StaticCastSharedPtr<FAssetDragDropOp>( Operation );
 
-	FMovieSceneTrackEditor::BeginKeying();
+	FMovieSceneTrackEditor::BeginKeying(DragDropParams.FrameNumber);
 
 	bool bAnyDropped = false;
 	for (const FAssetData& AssetData : DragDropOp->GetAssets())
@@ -1809,9 +1812,9 @@ FReply FSkeletalAnimationTrackEditor::OnDrop(const FDragDropEvent& DragDropEvent
 		const bool bValidAnimSequence = AnimSequence && AnimSequence->CanBeUsedInComposition();
 		if (bValidAnimSequence && Skeleton && Skeleton->IsCompatible(AnimSequence->GetSkeleton()))
 		{
-			UObject* Object = GetSequencer()->FindSpawnedObjectOrTemplate(TargetObjectGuid);
+			UObject* Object = GetSequencer()->FindSpawnedObjectOrTemplate(DragDropParams.TargetObjectGuid);
 				
-			AnimatablePropertyChanged( FOnKeyProperty::CreateRaw(this, &FSkeletalAnimationTrackEditor::AddKeyInternal, Object, AnimSequence, Track, RowIndex));
+			AnimatablePropertyChanged( FOnKeyProperty::CreateRaw(this, &FSkeletalAnimationTrackEditor::AddKeyInternal, Object, AnimSequence, DragDropParams.Track, DragDropParams.RowIndex));
 
 			bAnyDropped = true;
 		}
diff --git a/Engine/Source/Editor/MovieSceneTools/Private/TrackEditors/SkeletalAnimationTrackEditor.h b/Engine/Source/Editor/MovieSceneTools/Private/TrackEditors/SkeletalAnimationTrackEditor.h
--- a/Engine/Source/Editor/MovieSceneTools/Private/TrackEditors/SkeletalAnimationTrackEditor.h
+++ b/Engine/Source/Editor/MovieSceneTools/Private/TrackEditors/SkeletalAnimationTrackEditor.h
@@ -67,8 +67,8 @@ public:
 	virtual bool SupportsType( TSubclassOf<UMovieSceneTrack> Type ) const override;
 	virtual void BuildTrackContextMenu( FMenuBuilder& MenuBuilder, UMovieSceneTrack* Track ) override;
 	virtual TSharedPtr<SWidget> BuildOutlinerEditWidget(const FGuid& ObjectBinding, UMovieSceneTrack* Track, const FBuildEditWidgetParams& Params) override;
-	virtual bool OnAllowDrop(const FDragDropEvent& DragDropEvent, UMovieSceneTrack* Track, int32 RowIndex, const FGuid& TargetObjectGuid) override;
-	virtual FReply OnDrop(const FDragDropEvent& DragDropEvent, UMovieSceneTrack* Track, int32 RowIndex, const FGuid& TargetObjectGuid) override;
+	virtual bool OnAllowDrop(const FDragDropEvent& DragDropEvent, FSequencerDragDropParams& DragDropParams) override;
+	virtual FReply OnDrop(const FDragDropEvent& DragDropEvent, const FSequencerDragDropParams& DragDropParams) override;
 	virtual void OnInitialize() override;
 	virtual void OnRelease() override;
 
diff --git a/Engine/Source/Editor/MovieSceneTools/Private/TrackEditors/SubTrackEditor.cpp b/Engine/Source/Editor/MovieSceneTools/Private/TrackEditors/SubTrackEditor.cpp
--- a/Engine/Source/Editor/MovieSceneTools/Private/TrackEditors/SubTrackEditor.cpp
+++ b/Engine/Source/Editor/MovieSceneTools/Private/TrackEditors/SubTrackEditor.cpp
@@ -412,9 +412,9 @@ const FSlateBrush* FSubTrackEditor::GetIconBrush() const
 }
 
 
-bool FSubTrackEditor::OnAllowDrop(const FDragDropEvent& DragDropEvent, UMovieSceneTrack* Track, int32 RowIndex, const FGuid& TargetObjectGuid)
+bool FSubTrackEditor::OnAllowDrop(const FDragDropEvent& DragDropEvent, FSequencerDragDropParams& DragDropParams)
 {
-	if (!Track->IsA(UMovieSceneSubTrack::StaticClass()) || Track->IsA(UMovieSceneCinematicShotTrack::StaticClass()))
+	if (!DragDropParams.Track->IsA(UMovieSceneSubTrack::StaticClass()) || DragDropParams.Track->IsA(UMovieSceneCinematicShotTrack::StaticClass()))
 	{
 		return false;
 	}
@@ -430,8 +430,16 @@ bool FSubTrackEditor::OnAllowDrop(const FDragDropEvent& DragDropEvent, UMovieSce
 
 	for (const FAssetData& AssetData : DragDropOp->GetAssets())
 	{
-		if (Cast<UMovieSceneSequence>(AssetData.GetAsset()))
+		if (UMovieSceneSequence* Sequence = Cast<UMovieSceneSequence>(AssetData.GetAsset()))
 		{
+			FFrameRate TickResolution = GetSequencer()->GetFocusedTickResolution();
+
+			const FQualifiedFrameTime InnerDuration = FQualifiedFrameTime(
+				UE::MovieScene::DiscreteSize(Sequence->GetMovieScene()->GetPlaybackRange()),
+				Sequence->GetMovieScene()->GetTickResolution());
+
+			FFrameNumber LengthInFrames = InnerDuration.ConvertTo(TickResolution).FrameNumber;
+			DragDropParams.FrameRange = TRange<FFrameNumber>(DragDropParams.FrameNumber, DragDropParams.FrameNumber + LengthInFrames);
 			return true;
 		}
 	}
@@ -440,9 +448,9 @@ bool FSubTrackEditor::OnAllowDrop(const FDragDropEvent& DragDropEvent, UMovieSce
 }
 
 
-FReply FSubTrackEditor::OnDrop(const FDragDropEvent& DragDropEvent, UMovieSceneTrack* Track, int32 RowIndex, const FGuid& TargetObjectGuid)
+FReply FSubTrackEditor::OnDrop(const FDragDropEvent& DragDropEvent, const FSequencerDragDropParams& DragDropParams)
 {
-	if (!Track->IsA(UMovieSceneSubTrack::StaticClass()) || Track->IsA(UMovieSceneCinematicShotTrack::StaticClass()))
+	if (!DragDropParams.Track->IsA(UMovieSceneSubTrack::StaticClass()) || DragDropParams.Track->IsA(UMovieSceneCinematicShotTrack::StaticClass()))
 	{
 		return FReply::Unhandled();
 	}
@@ -458,7 +466,7 @@ FReply FSubTrackEditor::OnDrop(const FDragDropEvent& DragDropEvent, UMovieSceneT
 
 	TSharedPtr<FAssetDragDropOp> DragDropOp = StaticCastSharedPtr<FAssetDragDropOp>( Operation );
 	
-	FMovieSceneTrackEditor::BeginKeying();
+	FMovieSceneTrackEditor::BeginKeying(DragDropParams.FrameNumber);
 
 	bool bAnyDropped = false;
 	for (const FAssetData& AssetData : DragDropOp->GetAssets())
@@ -467,7 +475,7 @@ FReply FSubTrackEditor::OnDrop(const FDragDropEvent& DragDropEvent, UMovieSceneT
 
 		if (Sequence)
 		{
-			AnimatablePropertyChanged(FOnKeyProperty::CreateRaw(this, &FSubTrackEditor::HandleSequenceAdded, Sequence, RowIndex));
+			AnimatablePropertyChanged(FOnKeyProperty::CreateRaw(this, &FSubTrackEditor::HandleSequenceAdded, Sequence, DragDropParams.RowIndex));
 
 			bAnyDropped = true;
 		}
diff --git a/Engine/Source/Editor/MovieSceneTools/Public/TrackEditors/AudioTrackEditor.h b/Engine/Source/Editor/MovieSceneTools/Public/TrackEditors/AudioTrackEditor.h
--- a/Engine/Source/Editor/MovieSceneTools/Public/TrackEditors/AudioTrackEditor.h
+++ b/Engine/Source/Editor/MovieSceneTools/Public/TrackEditors/AudioTrackEditor.h
@@ -61,8 +61,8 @@ public:
 	virtual const FSlateBrush* GetIconBrush() const override;
 	virtual bool IsResizable(UMovieSceneTrack* InTrack) const override;
 	virtual void Resize(float NewSize, UMovieSceneTrack* InTrack) override;
-	virtual bool OnAllowDrop(const FDragDropEvent& DragDropEvent, UMovieSceneTrack* Track, int32 RowIndex, const FGuid& TargetObjectGuid) override;
-	virtual FReply OnDrop(const FDragDropEvent& DragDropEvent, UMovieSceneTrack* Track, int32 RowIndex, const FGuid& TargetObjectGuid) override;
+	virtual bool OnAllowDrop(const FDragDropEvent& DragDropEvent, FSequencerDragDropParams& DragDropParams) override;
+	virtual FReply OnDrop(const FDragDropEvent& DragDropEvent, const FSequencerDragDropParams& DragDropParams) override;
 	
 protected:
 
diff --git a/Engine/Source/Editor/MovieSceneTools/Public/TrackEditors/SubTrackEditor.h b/Engine/Source/Editor/MovieSceneTools/Public/TrackEditors/SubTrackEditor.h
--- a/Engine/Source/Editor/MovieSceneTools/Public/TrackEditors/SubTrackEditor.h
+++ b/Engine/Source/Editor/MovieSceneTools/Public/TrackEditors/SubTrackEditor.h
@@ -53,8 +53,8 @@ public:
 	virtual bool SupportsSequence(UMovieSceneSequence* InSequence) const override;
 	virtual bool SupportsType(TSubclassOf<UMovieSceneTrack> Type) const override;
 	virtual const FSlateBrush* GetIconBrush() const override;
-	virtual bool OnAllowDrop(const FDragDropEvent& DragDropEvent, UMovieSceneTrack* Track, int32 RowIndex, const FGuid& TargetObjectGuid) override;
-	virtual FReply OnDrop(const FDragDropEvent& DragDropEvent, UMovieSceneTrack* Track, int32 RowIndex, const FGuid& TargetObjectGuid) override;
+	virtual bool OnAllowDrop(const FDragDropEvent& DragDropEvent, FSequencerDragDropParams& DragDropParams) override;
+	virtual FReply OnDrop(const FDragDropEvent& DragDropEvent, const FSequencerDragDropParams& DragDropParams) override;
 	
 	/** Switch the selected sub section's take */
 	void SwitchTake(UObject* TakeObject);
diff --git a/Engine/Source/Editor/Sequencer/Private/MovieSceneTrackEditor.cpp b/Engine/Source/Editor/Sequencer/Private/MovieSceneTrackEditor.cpp
--- a/Engine/Source/Editor/Sequencer/Private/MovieSceneTrackEditor.cpp
+++ b/Engine/Source/Editor/Sequencer/Private/MovieSceneTrackEditor.cpp
@@ -36,8 +36,7 @@ bool FMovieSceneTrackEditor::bKeying;
 
 FFrameNumber FMovieSceneTrackEditor::GetTimeForKey()
 { 
-	// When shift is down, enable adding keys/sections one after the other
-	if (bKeying && NextKeyTime.IsSet() && FSlateApplication::Get().GetModifierKeys().IsShiftDown())
+	if (bKeying && NextKeyTime.IsSet())
 	{
 		return NextKeyTime.GetValue();
 	}
@@ -47,10 +46,10 @@ FFrameNumber FMovieSceneTrackEditor::GetTimeForKey()
 	return SequencerPin.IsValid() ? SequencerPin->GetLocalTime().Time.FrameNumber : FFrameNumber(0);
 }
 
-void FMovieSceneTrackEditor::BeginKeying()
+void FMovieSceneTrackEditor::BeginKeying(FFrameNumber InFrameNumber)
 {
 	bKeying = true;
-	NextKeyTime.Reset();
+	NextKeyTime = InFrameNumber;
 }
 
 void FMovieSceneTrackEditor::EndKeying()
@@ -92,12 +91,16 @@ void FMovieSceneTrackEditor::AnimatablePropertyChanged( FOnKeyProperty OnKeyProp
 
 		FKeyPropertyResult KeyPropertyResult = OnKeyProperty.Execute( KeyTime );
 
-		for (TWeakObjectPtr<UMovieSceneSection> NewSection : KeyPropertyResult.SectionsCreated)
+		// When shift is down, enable adding keys/sections one after the other
+		if (FSlateApplication::Get().GetModifierKeys().IsShiftDown())
 		{
-			if (NewSection.IsValid())
+			for (TWeakObjectPtr<UMovieSceneSection> NewSection : KeyPropertyResult.SectionsCreated)
 			{
-				NextKeyTime = NewSection.Get()->GetExclusiveEndFrame();
-				break;
+				if (NewSection.IsValid())
+				{
+					NextKeyTime = NewSection.Get()->GetExclusiveEndFrame();
+					break;
+				}
 			}
 		}
 
@@ -228,12 +231,12 @@ bool FMovieSceneTrackEditor::HandleAssetAdded(UObject* Asset, const FGuid& Targe
 	return false; 
 }
 
-bool FMovieSceneTrackEditor::OnAllowDrop(const FDragDropEvent& DragDropEvent, UMovieSceneTrack* Track, int32 RowIndex, const FGuid& TargetObjectGuid)
+bool FMovieSceneTrackEditor::OnAllowDrop(const FDragDropEvent& DragDropEvent, FSequencerDragDropParams& DragDropParams)
 {
 	return false;
 }
 
-FReply FMovieSceneTrackEditor::OnDrop(const FDragDropEvent& DragDropEvent, UMovieSceneTrack* Track, int32 RowIndex, const FGuid& TargetObjectGuid)
+FReply FMovieSceneTrackEditor::OnDrop(const FDragDropEvent& DragDropEvent, const FSequencerDragDropParams& DragDropParams)
 {
 	return FReply::Unhandled();
 }
diff --git a/Engine/Source/Editor/Sequencer/Private/SSequencer.cpp b/Engine/Source/Editor/Sequencer/Private/SSequencer.cpp
--- a/Engine/Source/Editor/Sequencer/Private/SSequencer.cpp
+++ b/Engine/Source/Editor/Sequencer/Private/SSequencer.cpp
@@ -3291,7 +3291,7 @@ void SSequencer::OnAssetsDropped( const FAssetDragDropOp& DragDropOp )
 	// If nobody took care of it, do the default behaviour.
 	if (DropResult == ESequencerDropResult::Unhandled)
 	{
-		FMovieSceneTrackEditor::BeginKeying();
+		FMovieSceneTrackEditor::BeginKeying(SequencerPtr.Pin()->GetLocalTime().Time.FrameNumber);
 
 		for (TArray<UObject*>::TConstIterator CurObjectIter = DroppedObjects.CreateConstIterator(); CurObjectIter; ++CurObjectIter)
 		{
diff --git a/Engine/Source/Editor/Sequencer/Private/SSequencerTrackArea.cpp b/Engine/Source/Editor/Sequencer/Private/SSequencerTrackArea.cpp
--- a/Engine/Source/Editor/Sequencer/Private/SSequencerTrackArea.cpp
+++ b/Engine/Source/Editor/Sequencer/Private/SSequencerTrackArea.cpp
@@ -16,6 +16,7 @@
 #include "ISequencerTrackEditor.h"
 #include "DisplayNodes/SequencerTrackNode.h"
 #include "DisplayNodes/SequencerObjectBindingNode.h"
+#include "CommonMovieSceneTools.h"
 
 FTrackAreaSlot::FTrackAreaSlot(const TSharedPtr<SSequencerTrackLane>& InSlotContent)
 	: TAlignmentWidgetSlotMixin<FTrackAreaSlot>(HAlign_Fill, VAlign_Top)
@@ -208,11 +209,21 @@ int32 SSequencerTrackArea::OnPaint(const FPaintArgs& Args, const FGeometry& Allo
 
 			int32 DashLayer = LayerId + 1;
 
+			float DropMinX = 0.f;
+			float DropMaxX = NodeGeometry.GetLocalSize().X;
+
+			if (DropFrameRange.IsSet())
+			{
+				FTimeToPixel TimeToPixel(NodeGeometry, Sequencer.Pin()->GetViewRange(), Sequencer.Pin()->GetFocusedTickResolution());
+				DropMinX = TimeToPixel.FrameToPixel(DropFrameRange.GetValue().GetLowerBoundValue());
+				DropMaxX = TimeToPixel.FrameToPixel(DropFrameRange.GetValue().GetUpperBoundValue());
+			}
+
 			// Top
 			FSlateDrawElement::MakeBox(
 				OutDrawElements,
 				DashLayer,
-				AllottedGeometry.ToPaintGeometry(FVector2D(0, TrackLane.Get()->GetPhysicalPosition()), FVector2D(NodeGeometry.GetLocalSize().X, HorizontalBrush->ImageSize.Y)),
+				AllottedGeometry.ToPaintGeometry(FVector2D(DropMinX, TrackLane.Get()->GetPhysicalPosition()), FVector2D(DropMaxX-DropMinX, HorizontalBrush->ImageSize.Y)),
 				HorizontalBrush,
 				ESlateDrawEffect::None,
 				DashColor);
@@ -221,7 +232,7 @@ int32 SSequencerTrackArea::OnPaint(const FPaintArgs& Args, const FGeometry& Allo
 			FSlateDrawElement::MakeBox(
 				OutDrawElements,
 				DashLayer,
-				AllottedGeometry.ToPaintGeometry(FVector2D(0, TrackLane.Get()->GetPhysicalPosition() + (TrackLane.Get()->GetCachedGeometry().GetLocalSize().Y - HorizontalBrush->ImageSize.Y)), FVector2D(AllottedGeometry.Size.X, HorizontalBrush->ImageSize.Y)),
+				AllottedGeometry.ToPaintGeometry(FVector2D(DropMinX, TrackLane.Get()->GetPhysicalPosition() + (TrackLane.Get()->GetCachedGeometry().GetLocalSize().Y - HorizontalBrush->ImageSize.Y)), FVector2D(DropMaxX-DropMinX, HorizontalBrush->ImageSize.Y)),
 				HorizontalBrush,
 				ESlateDrawEffect::None,
 				DashColor);
@@ -230,7 +241,7 @@ int32 SSequencerTrackArea::OnPaint(const FPaintArgs& Args, const FGeometry& Allo
 			FSlateDrawElement::MakeBox(
 				OutDrawElements,
 				DashLayer,
-				AllottedGeometry.ToPaintGeometry(FVector2D(0, TrackLane.Get()->GetPhysicalPosition()), FVector2D(VerticalBrush->ImageSize.X, TrackLane.Get()->GetCachedGeometry().GetLocalSize().Y)),
+				AllottedGeometry.ToPaintGeometry(FVector2D(DropMinX, TrackLane.Get()->GetPhysicalPosition()), FVector2D(VerticalBrush->ImageSize.X, TrackLane.Get()->GetCachedGeometry().GetLocalSize().Y)),
 				VerticalBrush,
 				ESlateDrawEffect::None,
 				DashColor);
@@ -239,7 +250,7 @@ int32 SSequencerTrackArea::OnPaint(const FPaintArgs& Args, const FGeometry& Allo
 			FSlateDrawElement::MakeBox(
 				OutDrawElements,
 				DashLayer,
-				AllottedGeometry.ToPaintGeometry(FVector2D(AllottedGeometry.GetLocalSize().X - VerticalBrush->ImageSize.X, TrackLane.Get()->GetPhysicalPosition()), FVector2D(VerticalBrush->ImageSize.X, TrackLane.Get()->GetCachedGeometry().GetLocalSize().Y)),
+				AllottedGeometry.ToPaintGeometry(FVector2D(DropMaxX - VerticalBrush->ImageSize.X, TrackLane.Get()->GetPhysicalPosition()), FVector2D(VerticalBrush->ImageSize.X, TrackLane.Get()->GetCachedGeometry().GetLocalSize().Y)),
 				VerticalBrush,
 				ESlateDrawEffect::None,
 				DashColor);
@@ -382,6 +393,7 @@ void SSequencerTrackArea::OnMouseEnter(const FGeometry& MyGeometry, const FPoint
 {
 	DroppedNode.Reset();
 	bAllowDrop = false;
+	DropFrameRange.Reset();
 
 	if ( Sequencer.IsValid() )
 	{
@@ -500,6 +512,7 @@ FReply SSequencerTrackArea::OnDragOver(const FGeometry& MyGeometry, const FDragD
 
 	DroppedNode = PinnedTreeView->HitTestNode(MyGeometry.AbsoluteToLocal(DragDropEvent.GetScreenSpacePosition()).Y);
 	bAllowDrop = false;
+	DropFrameRange.Reset();
 
 	if (DroppedNode.IsValid() && DroppedNode.Pin()->GetType() == ESequencerNode::Track && Sequencer.IsValid())
 	{
@@ -517,11 +530,21 @@ FReply SSequencerTrackArea::OnDragOver(const FGeometry& MyGeometry, const FDragD
 		// give track editors a chance to accept the drag event
 		auto TrackEditors = Sequencer.Pin()->GetTrackEditors();
 
+		FTimeToPixel TimeToPixel(MyGeometry, Sequencer.Pin()->GetViewRange(), Sequencer.Pin()->GetFocusedTickResolution());
+		FVector2D LocalPos = MyGeometry.AbsoluteToLocal(DragDropEvent.GetScreenSpacePosition());
+		FFrameNumber DropFrameNumber = TimeToPixel.PixelToFrame(LocalPos.X).FrameNumber;
+		if (Sequencer.Pin()->GetSequencerSettings()->GetIsSnapEnabled() && Sequencer.Pin()->GetSequencerSettings()->GetSnapPlayTimeToInterval())
+		{
+			DropFrameNumber = FFrameRate::Snap(DropFrameNumber, Sequencer.Pin()->GetFocusedTickResolution(), Sequencer.Pin()->GetFocusedDisplayRate()).FrameNumber;
+		}
+		FSequencerDragDropParams DragDropParams(Track, RowIndex, ObjectBinding, DropFrameNumber, TRange<FFrameNumber>());
+
 		for (const auto& TrackEditor : TrackEditors)
 		{
-			if (TrackEditor->OnAllowDrop(DragDropEvent, Track, RowIndex, ObjectBinding))
+			if (TrackEditor->OnAllowDrop(DragDropEvent, DragDropParams))
 			{
 				bAllowDrop = true;
+				DropFrameRange = DragDropParams.FrameRange;
 				return FReply::Handled();
 			}
 		}
@@ -553,13 +576,22 @@ FReply SSequencerTrackArea::OnDrop(const FGeometry& MyGeometry, const FDragDropE
 		// give track editors a chance to process the drag event
 		auto TrackEditors = Sequencer.Pin()->GetTrackEditors();
 
+		FTimeToPixel TimeToPixel(MyGeometry, Sequencer.Pin()->GetViewRange(), Sequencer.Pin()->GetFocusedTickResolution());
+		FVector2D LocalPos = MyGeometry.AbsoluteToLocal(DragDropEvent.GetScreenSpacePosition());
+		FFrameNumber DropFrameNumber = TimeToPixel.PixelToFrame(LocalPos.X).FrameNumber;
+		if (Sequencer.Pin()->GetSequencerSettings()->GetIsSnapEnabled() && Sequencer.Pin()->GetSequencerSettings()->GetSnapPlayTimeToInterval())
+		{
+			DropFrameNumber = FFrameRate::Snap(DropFrameNumber, Sequencer.Pin()->GetFocusedTickResolution(), Sequencer.Pin()->GetFocusedDisplayRate()).FrameNumber;
+		}
+		FSequencerDragDropParams DragDropParams(Track, RowIndex, ObjectBinding, DropFrameNumber, TRange<FFrameNumber>());
+
 		for (const auto& TrackEditor : TrackEditors)
 		{
-			if (TrackEditor->OnAllowDrop(DragDropEvent, Track, RowIndex, ObjectBinding))
+			if (TrackEditor->OnAllowDrop(DragDropEvent, DragDropParams))
 			{
 				DroppedNode.Reset();
 
-				return TrackEditor->OnDrop(DragDropEvent, Track, RowIndex, ObjectBinding);
+				return TrackEditor->OnDrop(DragDropEvent, DragDropParams);
 			}
 		}
 	}
diff --git a/Engine/Source/Editor/Sequencer/Private/SSequencerTrackArea.h b/Engine/Source/Editor/Sequencer/Private/SSequencerTrackArea.h
--- a/Engine/Source/Editor/Sequencer/Private/SSequencerTrackArea.h
+++ b/Engine/Source/Editor/Sequencer/Private/SSequencerTrackArea.h
@@ -155,6 +155,9 @@ private:
 	/** Whether the dropped node is allowed to be dropped onto */
 	bool bAllowDrop;
 
+	/** The frame range of the section about to be dropped */
+	TOptional<TRange<FFrameNumber>> DropFrameRange;
+
 	/** Whether this TrackArea is for pinned nodes or non-pinned nodes */
 	bool bShowPinnedNodes;
 
diff --git a/Engine/Source/Editor/Sequencer/Public/ISequencerTrackEditor.h b/Engine/Source/Editor/Sequencer/Public/ISequencerTrackEditor.h
--- a/Engine/Source/Editor/Sequencer/Public/ISequencerTrackEditor.h
+++ b/Engine/Source/Editor/Sequencer/Public/ISequencerTrackEditor.h
@@ -36,6 +36,37 @@ struct FBuildEditWidgetParams
 	int32 TrackInsertRowIndex;
 };
 
+struct FSequencerDragDropParams
+{
+	FSequencerDragDropParams()
+		: Track(nullptr)
+		, RowIndex(INDEX_NONE)
+		{}
+
+	FSequencerDragDropParams(UMovieSceneTrack* InTrack, int32 InRowIndex, FGuid InTargetObjectGuid, FFrameNumber InFrameNumber, const TRange<FFrameNumber>& InFrameRange)
+		: Track(InTrack)
+		, RowIndex(InRowIndex)
+		, TargetObjectGuid(InTargetObjectGuid)
+		, FrameNumber(InFrameNumber)
+		, FrameRange(InFrameRange)
+	{}
+
+	/** The track that is receiving this drop event */
+	UMovieSceneTrack* Track;
+	 
+	/** The row index to drop onto */
+	int32 RowIndex;
+
+	/** The object guid this asset is dropped onto, if applicable */
+	FGuid TargetObjectGuid;
+
+	/** The frame number that this drop event is being dropped at */
+	FFrameNumber FrameNumber;
+
+	/** The frame range that this drop event is being dropped at */
+	TRange<FFrameNumber> FrameRange;
+};
+
 /**
  * Interface for sequencer track editors.
  */
@@ -138,23 +169,25 @@ public:
 	 * Called when attempting to drop an asset directly onto a track.
 	 *
 	 * @param DragDropEvent The drag drop event.
-	 * @param Track The track that is receiving this drop event.
-	 * @param RowIndex The row index to drop onto.
-	 * @param TargetObjectGuid The object guid this asset is dropped onto, if applicable.
+	 * @param DragDropParams The drag drop parameters.
 	 * @return Whether the drop event can be handled.
 	 */
-	virtual bool OnAllowDrop(const FDragDropEvent& DragDropEvent, UMovieSceneTrack* Track, int32 RowIndex, const FGuid& TargetObjectGuid) = 0;
+	virtual bool OnAllowDrop(const FDragDropEvent& DragDropEvent, FSequencerDragDropParams& DragDropParams) = 0;
+
+	UE_DEPRECATED(4.27, "Use OnAllowDrop with DragDropParams.")
+	virtual bool OnAllowDrop(const FDragDropEvent& DragDropEvent, UMovieSceneTrack* Track, int32 RowIndex, const FGuid& TargetObjectGuid) { FSequencerDragDropParams DragDropParams; return OnAllowDrop(DragDropEvent, DragDropParams); }
 
 	/**
 	 * Called when an asset is dropped directly onto a track.
 	 *
 	 * @param DragDropEvent The drag drop event.
-	 * @param Track The track that is receiving this drop event.
-	 * @param RowIndex The row index to drop onto.
-	 * @param TargetObjectGuid The object guid this asset is dropped onto, if applicable.
+	 * @param DragDropParams The drag drop parameters.
 	 * @return Whether the drop event was handled.
 	 */	
-	virtual FReply OnDrop(const FDragDropEvent& DragDropEvent, UMovieSceneTrack* Track, int32 RowIndex, const FGuid& TargetObjectGuid) = 0;
+	virtual FReply OnDrop(const FDragDropEvent& DragDropEvent, const FSequencerDragDropParams& DragDropParams) = 0;
+
+	UE_DEPRECATED(4.27, "Use OnDrop with DragDropParams.")
+	virtual FReply OnDrop(const FDragDropEvent& DragDropEvent, UMovieSceneTrack* Track, int32 RowIndex, const FGuid& TargetObjectGuid) { return OnDrop(DragDropEvent, FSequencerDragDropParams()); }
 
 	/**
 	 * Called to generate a section layout for a particular section.
diff --git a/Engine/Source/Editor/Sequencer/Public/MovieSceneTrackEditor.h b/Engine/Source/Editor/Sequencer/Public/MovieSceneTrackEditor.h
--- a/Engine/Source/Editor/Sequencer/Public/MovieSceneTrackEditor.h
+++ b/Engine/Source/Editor/Sequencer/Public/MovieSceneTrackEditor.h
@@ -100,7 +100,7 @@ public:
 	 * This will allow for different behaviors, ie. keying all at the current time or keying 
 	 * one after the other.
 	 */
-	static void BeginKeying();
+	static void BeginKeying(FFrameNumber InFrameNumber);
 	static void EndKeying();
 
 	void UpdatePlaybackRange();
@@ -178,8 +178,8 @@ public:
 	virtual TSharedPtr<SWidget> BuildOutlinerEditWidget(const FGuid& ObjectBinding, UMovieSceneTrack* Track, const FBuildEditWidgetParams& Params) override;
 	virtual void BuildTrackContextMenu( FMenuBuilder& MenuBuilder, UMovieSceneTrack* Track ) override;
 	virtual bool HandleAssetAdded(UObject* Asset, const FGuid& TargetObjectGuid) override;
-	virtual bool OnAllowDrop(const FDragDropEvent& DragDropEvent, UMovieSceneTrack* Track, int32 RowIndex, const FGuid& TargetObjectGuid) override;
-	virtual FReply OnDrop(const FDragDropEvent& DragDropEvent, UMovieSceneTrack* Track, int32 RowIndex, const FGuid& TargetObjectGuid) override;
+	virtual bool OnAllowDrop(const FDragDropEvent& DragDropEvent, FSequencerDragDropParams& DragDropParams) override;
+	virtual FReply OnDrop(const FDragDropEvent& DragDropEvent, const FSequencerDragDropParams& DragDropParams) override;
 
 	virtual void OnInitialize() override;
 	virtual void OnRelease() override;
diff --git a/Engine/Source/Runtime/MovieSceneTracks/Public/Tracks/MovieSceneCameraCutTrack.h b/Engine/Source/Runtime/MovieSceneTracks/Public/Tracks/MovieSceneCameraCutTrack.h
--- a/Engine/Source/Runtime/MovieSceneTracks/Public/Tracks/MovieSceneCameraCutTrack.h
+++ b/Engine/Source/Runtime/MovieSceneTracks/Public/Tracks/MovieSceneCameraCutTrack.h
@@ -55,9 +55,9 @@ public:
 	virtual EMovieSceneSectionMovedResult OnSectionMoved(UMovieSceneSection& Section, const FMovieSceneSectionMovedParams& Params) override;
 #endif
 
-protected:
+	MOVIESCENETRACKS_API FFrameNumber FindEndTimeForCameraCut(FFrameNumber StartTime);
 
-	FFrameNumber FindEndTimeForCameraCut(FFrameNumber StartTime);
+protected:
 
 	virtual void PreCompileImpl(FMovieSceneTrackPreCompileResult& OutPreCompileResult) override;
 
