Sequencer: Improvements to how marked frames are shown globally.

The new "marked frames helper" looks into any kind of  sub-sequenc, including template sequences, which weren't taken into account previously. It does this by looking at the compiled sequence hierarchy. Then, it gathers marked frames with some special code to handle looping sub-sequences.

diff --git a/Engine/Source/Editor/Sequencer/Private/Sequencer.cpp b/Engine/Source/Editor/Sequencer/Private/Sequencer.cpp
--- a/Engine/Source/Editor/Sequencer/Private/Sequencer.cpp
+++ b/Engine/Source/Editor/Sequencer/Private/Sequencer.cpp
@@ -49,6 +49,7 @@
 #include "DisplayNodes/SequencerTrackNode.h"
 #include "MovieSceneClipboard.h"
 #include "SequencerCommonHelpers.h"
+#include "SequencerMarkedFrameHelper.h"
 #include "SSequencer.h"
 #include "SSequencerSection.h"
 #include "SequencerKeyCollection.h"
@@ -11572,67 +11573,6 @@ TArray<FMovieSceneMarkedFrame> FSequencer::GetMarkedFrames() const
 	return TArray<FMovieSceneMarkedFrame>();
 }
 
-void FindGlobalMarkedFrames(TArray<FMovieSceneMarkedFrame>& GlobalMarkedFrames, const UMovieScene* MovieScene, const FFrameNumber FrameOffset, const UMovieScene* FocusedMovieScene)
-{
-	if (MovieScene == nullptr)
-	{
-		return;
-	}
-
-	if (MovieScene != FocusedMovieScene && MovieScene->GetGloballyShowMarkedFrames())
-	{
-		TArray<FMovieSceneMarkedFrame> MarkedFrames = MovieScene->GetMarkedFrames();
-		for (FMovieSceneMarkedFrame MarkedFrame : MarkedFrames)
-		{
-			MarkedFrame.FrameNumber += FrameOffset;
-			GlobalMarkedFrames.Add(MarkedFrame);
-		}
-	}
-
-	const TArray<UMovieSceneTrack*>& Tracks = MovieScene->GetMasterTracks();
-	for (const UMovieSceneTrack* Track : Tracks)
-	{
-		if (Track->SupportsType(UMovieSceneSubSection::StaticClass()))
-		{
-			for (const UMovieSceneSection* Section : Track->GetAllSections())
-			{
-				const UMovieSceneSubSection* SubSection = Cast<UMovieSceneSubSection>(Section);
-				if (SubSection != nullptr)
-				{
-					const UMovieSceneSequence* SubSequence = SubSection->GetSequence();
-					if (SubSequence != nullptr)
-					{
-						const UMovieScene* SubSequenceScene = SubSequence->GetMovieScene();
-						if (SubSequenceScene != nullptr)
-						{
-							FindGlobalMarkedFrames(GlobalMarkedFrames, SubSequenceScene, SubSection->GetInclusiveStartFrame() + FrameOffset, FocusedMovieScene);
-						}
-					}
-				}
-			}
-		}
-		else if(Track->SupportsType(UMovieSceneCinematicShotSection::StaticClass()))
-		{
-			for (const UMovieSceneSection* Section : Track->GetAllSections())
-			{
-				const UMovieSceneCinematicShotSection* ShotSection = Cast<UMovieSceneCinematicShotSection>(Section);
-				if (ShotSection != nullptr)
-				{
-					const UMovieSceneSequence* SubSequence = ShotSection->GetSequence();
-					if (SubSequence != nullptr)
-					{
-						const UMovieScene* SubSequenceScene = SubSequence->GetMovieScene();
-						if (SubSequenceScene != nullptr)
-						{
-							FindGlobalMarkedFrames(GlobalMarkedFrames, SubSequenceScene, ShotSection->GetInclusiveStartFrame() + FrameOffset, FocusedMovieScene);
-						}
-					}
-				}
-			}
-		}
-	}
-}
-
 TArray<FMovieSceneMarkedFrame> FSequencer::GetGlobalMarkedFrames() const
 {
 	return GlobalMarkedFramesCache;
@@ -11642,51 +11582,16 @@ void FSequencer::UpdateGlobalMarkedFramesCache()
 {
 	GlobalMarkedFramesCache.Empty();
 
-	UMovieSceneSequence* FocusedMovieSequence = GetFocusedMovieSceneSequence();
-	UMovieScene* FocusedMovieScene = nullptr;
-	if (FocusedMovieSequence != nullptr)
-	{
-		FocusedMovieScene = FocusedMovieSequence->GetMovieScene();
-	}
-
-	UMovieSceneSequence* RootMovieSequence = GetRootMovieSceneSequence();
-	if (RootMovieSequence != nullptr)
-	{
-		UMovieScene* RootMovieScene = RootMovieSequence->GetMovieScene();
-		if (RootMovieScene != nullptr)
-		{
-
-			TRange<FFrameTime> Range(0,0);
-			Range = RootToLocalTransform.TransformRangeUnwarped(Range);
-			FFrameNumber FocusedMovieSceneFrameOffset = Range.GetLowerBoundValue().FrameNumber;
-
-			FindGlobalMarkedFrames(GlobalMarkedFramesCache, RootMovieScene, FocusedMovieSceneFrameOffset, FocusedMovieScene);
-		}
-	}
+	FSequencerMarkedFrameHelper::FindGlobalMarkedFrames(*this, GlobalMarkedFramesCache);
 	
 	bGlobalMarkedFramesCached = true;
 }
 
 void FSequencer::ClearGlobalMarkedFrames()
 {
-	bGlobalMarkedFramesCached = false;
-
-	const FMovieSceneSequenceHierarchy* Hierarchy = CompiledDataManager->FindHierarchy(RootTemplateInstance.GetCompiledDataID());
-	if (!Hierarchy)
-	{
-		return;
-	}
+	FSequencerMarkedFrameHelper::ClearGlobalMarkedFrames(*this);
 
-	for (const TTuple<FMovieSceneSequenceID, FMovieSceneSubSequenceData>& Pair : Hierarchy->AllSubSequenceData())
-	{
-		if (UMovieSceneSequence* Sequence = Pair.Value.GetSequence())
-		{
-			if (UMovieScene* MovieScene = Sequence->GetMovieScene())
-			{
-				MovieScene->SetGloballyShowMarkedFrames(false);
-			}
-		}
-	}
+	bGlobalMarkedFramesCached = false;
 }
 
 void FSequencer::ToggleMarkAtPlayPosition()
diff --git a/Engine/Source/Editor/Sequencer/Private/SequencerMarkedFrameHelper.cpp b/Engine/Source/Editor/Sequencer/Private/SequencerMarkedFrameHelper.cpp
new file mode 100644
--- /dev/null
+++ b/Engine/Source/Editor/Sequencer/Private/SequencerMarkedFrameHelper.cpp
@@ -0,0 +1,216 @@
+// Copyright Epic Games, Inc. All Rights Reserved.
+
+#include "SequencerMarkedFrameHelper.h"
+#include "ISequencer.h"
+#include "MovieSceneSequence.h"
+#include "Evaluation/MovieSceneEvaluationTemplateInstance.h"
+#include "Evaluation/MovieSceneTimeArray.h"
+#include "MovieSceneSequenceID.h"
+#include "MovieSceneTimeHelpers.h"
+#include "Sections/MovieSceneSubSection.h"
+#include "Tracks/MovieSceneSubTrack.h"
+
+namespace UE
+{
+namespace MovieScene
+{
+
+static void FindGlobalMarkedFrames(
+		const ISequencer& Sequencer, const FMovieSceneSequenceHierarchy* SequenceHierarchy,
+		FMovieSceneSequenceIDRef FocusedSequenceID, FMovieSceneSequenceIDRef SequenceID, 
+		TRange<FFrameNumber> GatherRange,
+		TMovieSceneTimeArray<FMovieSceneMarkedFrame>& OutTimestampedGlobalMarkedFrames)
+{
+	// Find the current sequence in the hierarchy.
+	const FMovieSceneSubSequenceData* const SequenceSubData = SequenceHierarchy->FindSubData(SequenceID);
+	const UMovieSceneSequence* const Sequence = SequenceSubData ? SequenceSubData->GetSequence() : Sequencer.GetRootMovieSceneSequence();
+	const UMovieScene* const MovieScene = Sequence ? Sequence->GetMovieScene() : nullptr;
+	if (UNLIKELY(!ensure(MovieScene)))
+	{
+		return;
+	}
+
+	// Get the marked frames of the current sequence if it's not the focused sequence.
+	if (SequenceID != FocusedSequenceID && MovieScene->GetGloballyShowMarkedFrames())
+	{
+		const TArray<FMovieSceneMarkedFrame>& MarkedFrames = MovieScene->GetMarkedFrames();
+		for (const FMovieSceneMarkedFrame& MarkedFrame : MarkedFrames)
+		{
+			if (GatherRange.Contains(MarkedFrame.FrameNumber))
+			{
+				OutTimestampedGlobalMarkedFrames.Add(MarkedFrame.FrameNumber, MarkedFrame);
+			}
+		}
+	}
+
+	// Dive into the current sequence's sub-sequences.
+	const FMovieSceneSequenceHierarchyNode* const SequenceNode = SequenceHierarchy->FindNode(SequenceID);
+	if (ensure(SequenceNode))
+	{
+		for (const FMovieSceneSequenceID ChildID : SequenceNode->Children)
+		{
+			const FMovieSceneSubSequenceData* const ChildSubData = SequenceHierarchy->FindSubData(ChildID);
+			if (UNLIKELY(!ensure(ChildSubData)))
+			{
+				continue;
+			}
+
+			const UMovieSceneSequence* const ChildSequence = ChildSubData->GetSequence();
+			const UMovieScene* const ChildMovieScene = ChildSequence ? ChildSequence->GetMovieScene() : nullptr;
+			if (UNLIKELY(!ensure(ChildMovieScene)))
+			{
+				continue;
+			}
+
+			if (!ChildSubData->bCanLoop)
+			{
+				// This child doesn't loop. We just need to take into account its offset and time scale, and how much of it we "see" through the
+				// parent section.
+				ensure(!ChildSubData->OuterToInnerTransform.IsWarping());
+				const FMovieSceneTimeTransform& OuterToInnerTransform(ChildSubData->OuterToInnerTransform.LinearTransform);
+				OutTimestampedGlobalMarkedFrames.PushTransform(OuterToInnerTransform);
+
+				// Compute the "window" of the child that we can see through the parent section. Note that the parent section could extend past
+				// the end of the child's playback... we don't trim that, so that any marked frame added after the playback end in the child sequence
+				// is still visible in the parent sequence as long as the parent section is long enough. This lets artists see what they're "missing".
+				const TRange<FFrameNumber> ParentPlayRange = ChildSubData->ParentPlayRange.Value;
+				const TRange<FFrameNumber> ChildGatherRange = ChildSubData->OuterToInnerTransform.TransformRangeConstrained(ParentPlayRange);
+
+				// Gather marked frames in this "window".
+				FindGlobalMarkedFrames(Sequencer, SequenceHierarchy, FocusedSequenceID, ChildID, ChildGatherRange, OutTimestampedGlobalMarkedFrames);
+
+				OutTimestampedGlobalMarkedFrames.PopTransform();
+			}
+			else if (ensure(ChildSubData->OuterToInnerTransform.NestedTransforms.Num() > 0))
+			{
+				// This child is looping. Things are... more complicated.
+				//
+				// First, we push the time transform for this child sequence. It should be found in the first NestedTransforms
+				// since that's how sub-datas are computed.
+				ensure(ChildSubData->OuterToInnerTransform.LinearTransform.IsIdentity());
+				const FMovieSceneNestedSequenceTransform& OuterToInnerTransform(ChildSubData->OuterToInnerTransform.NestedTransforms[0]);
+				OutTimestampedGlobalMarkedFrames.PushTransform(OuterToInnerTransform.LinearTransform, OuterToInnerTransform.Warping);
+
+				// Next, we'll need to gather the marked frames of the child sequence repeatedly, once for each loop.
+				// To know how many loops we have, we need to look at the play range of the parent sub-section, so let's grab
+				// that first.
+				const FMovieSceneSectionParameters SubSectionParameters = ChildSubData->ToSubSectionParameters();
+				const TRange<FFrameNumber> ChildPlayRange = UMovieSceneSubSection::GetValidatedInnerPlaybackRange(SubSectionParameters, *ChildMovieScene);
+				const FFrameNumber ChildLength = UE::MovieScene::DiscreteSize(ChildPlayRange);
+
+				// Now we need to know how long this child play range is in the parent's time space. This is how we can figure
+				// out how many loops we can fit.
+				const FFrameRate ParentFrameRate = MovieScene->GetTickResolution();
+				const FFrameRate ChildFrameRate   = ChildMovieScene->GetTickResolution();
+
+				const float ChildTimeScale = OuterToInnerTransform.LinearTransform.TimeScale;
+				const float InvChildTimeScale = FMath::IsNearlyZero(ChildTimeScale) ? 1.0f : 1.0f / ChildTimeScale;
+
+				const FFrameNumber ChildLengthInParentSpace = (ConvertFrameTime(ChildLength, ChildFrameRate, ParentFrameRate) * InvChildTimeScale).FrameNumber;
+				const FFrameNumber ChildFirstLoopLength = ChildLength - ChildSubData->ParentFirstLoopStartFrameOffset;
+				const FFrameNumber ChildFirstLoopLengthInParentSpace = (ConvertFrameTime(ChildFirstLoopLength, ChildFrameRate, ParentFrameRate) * InvChildTimeScale).FrameNumber;
+
+				// We can finally start iterating: we iterate for how many times as we can fit the child sequence's length,
+				// modified by the time scale, into the parent play range.
+				const TRange<FFrameNumber> ParentPlayRange = ChildSubData->ParentPlayRange.Value;
+				const FFrameNumber ParentExclusiveEnd = UE::MovieScene::DiscreteExclusiveUpper(ParentPlayRange);
+				FFrameNumber CurLoopStart = UE::MovieScene::DiscreteInclusiveLower(ParentPlayRange);
+				FFrameNumber CurLoopEnd = CurLoopStart + FMath::Min(
+						FMath::Max(ChildFirstLoopLengthInParentSpace, FFrameNumber(0)),
+						ParentExclusiveEnd);
+
+				while (CurLoopStart < ParentExclusiveEnd)
+				{
+					TRange<FFrameNumber> CurLoopChildGatherRange(CurLoopStart, CurLoopEnd);
+					CurLoopChildGatherRange = ChildSubData->OuterToInnerTransform.TransformRangeConstrained(CurLoopChildGatherRange);
+
+					FindGlobalMarkedFrames(Sequencer, SequenceHierarchy, FocusedSequenceID, ChildID, CurLoopChildGatherRange, OutTimestampedGlobalMarkedFrames);
+
+					CurLoopStart = CurLoopEnd;
+					CurLoopEnd += FMath::Max(ChildLengthInParentSpace, FFrameNumber(0));
+					
+					OutTimestampedGlobalMarkedFrames.IncrementWarpCounter();
+				}
+
+				OutTimestampedGlobalMarkedFrames.PopTransform();
+			}
+		}
+	}
+}
+
+} // namespace MovieScene
+} // namespace UE
+
+void FSequencerMarkedFrameHelper::FindGlobalMarkedFrames(ISequencer& Sequencer, TArray<FMovieSceneMarkedFrame>& OutGlobalMarkedFrames)
+{
+	// Get the focused sequence info. We want to gather all the marked frames that are in the subset of the sequence hierarchy
+	// that hangs below this focused sequence.
+	UMovieSceneSequence* FocusedMovieSequence = Sequencer.GetFocusedMovieSceneSequence();
+	const FMovieSceneSequenceID FocusedMovieSequenceID = Sequencer.GetFocusedTemplateID();
+
+	UMovieSceneSequence* RootMovieSequence = Sequencer.GetRootMovieSceneSequence();
+
+	if (!FocusedMovieSequence || !RootMovieSequence)
+	{
+		return;
+	}
+
+	// Get the sequence hierarchy so that we can iterate it.
+	const FMovieSceneRootEvaluationTemplateInstance& EvalTemplate = Sequencer.GetEvaluationTemplate();
+	const FMovieSceneSequenceHierarchy* SequenceHierarchy = EvalTemplate.GetHierarchy();
+	if (!SequenceHierarchy)
+	{
+		return;
+	}
+	
+	// Find the time offset of the focused sequence, so that we can offset all the marked frames we find.
+	// If it is the root sequence, there's no offset.
+	FFrameTime FrameZeroInFocusedMovieSequence(0);
+	const FMovieSceneSubSequenceData* FocusedMovieSequenceSubData = SequenceHierarchy->FindSubData(FocusedMovieSequenceID);
+	if (FocusedMovieSequenceSubData)
+	{
+		FrameZeroInFocusedMovieSequence = FocusedMovieSequenceSubData->RootToSequenceTransform.TransformTime(FrameZeroInFocusedMovieSequence);
+	}
+	const FFrameNumber FocusedMovieSequenceFrameOffset = FrameZeroInFocusedMovieSequence.FrameNumber;
+
+	// Grab the marked frames from the root sequence, and recursively across the whole hierarchy.
+	TMovieSceneTimeArray<FMovieSceneMarkedFrame> TimestampedGlobalMarkedFrames;
+	TimestampedGlobalMarkedFrames.PushTransform(FocusedMovieSequenceFrameOffset);
+	UE::MovieScene::FindGlobalMarkedFrames(Sequencer, SequenceHierarchy, FocusedMovieSequenceID, MovieSceneSequenceID::Root, TRange<FFrameNumber>::All(), TimestampedGlobalMarkedFrames);
+
+	// Export the modified timestamped entries.
+	for (const TMovieSceneTimeArrayEntry<FMovieSceneMarkedFrame>& Entry : TimestampedGlobalMarkedFrames.GetEntries())
+	{
+		FMovieSceneMarkedFrame MarkedFrame = Entry.Datum;
+		MarkedFrame.FrameNumber = Entry.RootTime.FrameNumber;
+		OutGlobalMarkedFrames.Add(MarkedFrame);
+	}
+}
+
+void FSequencerMarkedFrameHelper::ClearGlobalMarkedFrames(ISequencer& Sequencer)
+{
+	const FMovieSceneRootEvaluationTemplateInstance& EvalTemplate = Sequencer.GetEvaluationTemplate();
+
+	ClearGlobalMarkedFrames(EvalTemplate.GetRootSequence());
+
+	const FMovieSceneSequenceHierarchy* SequenceHierarchy = EvalTemplate.GetHierarchy();
+	if (SequenceHierarchy)
+	{
+		for (const TTuple<FMovieSceneSequenceID, FMovieSceneSubSequenceData>& Pair : SequenceHierarchy->AllSubSequenceData())
+		{
+			ClearGlobalMarkedFrames(Pair.Value.GetSequence());
+		}
+	}
+}
+
+void FSequencerMarkedFrameHelper::ClearGlobalMarkedFrames(UMovieSceneSequence* Sequence)
+{
+	if (Sequence)
+	{
+		if (UMovieScene* MovieScene = Sequence->GetMovieScene())
+		{
+			MovieScene->SetGloballyShowMarkedFrames(false);
+		}
+	}
+}
+
diff --git a/Engine/Source/Editor/Sequencer/Private/SequencerMarkedFrameHelper.h b/Engine/Source/Editor/Sequencer/Private/SequencerMarkedFrameHelper.h
new file mode 100644
--- /dev/null
+++ b/Engine/Source/Editor/Sequencer/Private/SequencerMarkedFrameHelper.h
@@ -0,0 +1,25 @@
+// Copyright Epic Games, Inc. All Rights Reserved.
+
+#pragma once
+
+#include "CoreMinimal.h"
+#include "MovieScene.h"
+
+class ISequencer;
+
+/**
+ * Helper class to manage marked frames across a sequence hierarchy.
+ */
+class FSequencerMarkedFrameHelper
+{
+public:
+	/** Find all marked frames in a sequence hierarchy, relative to the Sequencer's currently focused sequence */
+	static void FindGlobalMarkedFrames(ISequencer& Sequencer, TArray<FMovieSceneMarkedFrame>& OutGlobalMarkedFrames);
+
+	/** Clear the setting to show marked frames globally across an entire sequence hierarchy */
+	static void ClearGlobalMarkedFrames(ISequencer& Sequencer);
+
+private:
+	static void ClearGlobalMarkedFrames(UMovieSceneSequence* Sequence);
+};
+
