Sequencer: Add supports render movie toggle so that only level sequences have the render movie button

diff --git a/Engine/Plugins/MovieScene/LevelSequenceEditor/Source/LevelSequenceEditor/Private/LevelSequenceEditorToolkit.cpp b/Engine/Plugins/MovieScene/LevelSequenceEditor/Source/LevelSequenceEditor/Private/LevelSequenceEditorToolkit.cpp
--- a/Engine/Plugins/MovieScene/LevelSequenceEditor/Source/LevelSequenceEditor/Private/LevelSequenceEditorToolkit.cpp
+++ b/Engine/Plugins/MovieScene/LevelSequenceEditor/Source/LevelSequenceEditor/Private/LevelSequenceEditorToolkit.cpp
@@ -189,6 +189,7 @@ void FLevelSequenceEditorToolkit::Initialize(const EToolkitMode::Type Mode, cons
 		SequencerInitParams.HostCapabilities.bSupportsCurveEditor = true;
 		SequencerInitParams.HostCapabilities.bSupportsSaveMovieSceneAsset = true;
 		SequencerInitParams.HostCapabilities.bSupportsRecording = true;
+		SequencerInitParams.HostCapabilities.bSupportsRenderMovie = true;
 
 		TSharedRef<FExtender> ToolbarExtender = MakeShared<FExtender>();
 		ToolbarExtender->AddToolBarExtension("Base Commands", EExtensionHook::Before, nullptr, FToolBarExtensionDelegate::CreateSP(this, &FLevelSequenceEditorToolkit::ExtendSequencerToolbar));
diff --git a/Engine/Source/Editor/Sequencer/Private/SSequencer.cpp b/Engine/Source/Editor/Sequencer/Private/SSequencer.cpp
--- a/Engine/Source/Editor/Sequencer/Private/SSequencer.cpp
+++ b/Engine/Source/Editor/Sequencer/Private/SSequencer.cpp
@@ -1399,15 +1399,19 @@ TSharedRef<SWidget> SSequencer::MakeToolBar()
 
 			ToolBarBuilder.AddToolBarButton( FSequencerCommands::Get().FindInContentBrowser );
 			ToolBarBuilder.AddToolBarButton( FSequencerCommands::Get().CreateCamera );
-			ToolBarBuilder.AddToolBarButton( FSequencerCommands::Get().RenderMovie );
 
-			ToolBarBuilder.AddComboButton(
-				FUIAction(),
-				FOnGetContent::CreateSP( this, &SSequencer::MakeRenderMovieMenu ),
-				LOCTEXT( "RenderMovieOptions", "Render Movie Options" ),
-				LOCTEXT( "RenderMovieOptionsToolTip", "Render Movie Options" ),
-				TAttribute<FSlateIcon>(),
-				true );
+			if (SequencerPtr.Pin()->GetHostCapabilities().bSupportsRenderMovie)
+			{
+				ToolBarBuilder.AddToolBarButton( FSequencerCommands::Get().RenderMovie );
+				
+				ToolBarBuilder.AddComboButton(
+					FUIAction(),
+					FOnGetContent::CreateSP( this, &SSequencer::MakeRenderMovieMenu ),
+					LOCTEXT( "RenderMovieOptions", "Render Movie Options" ),
+					LOCTEXT( "RenderMovieOptionsToolTip", "Render Movie Options" ),
+					TAttribute<FSlateIcon>(),
+					true );
+			}
 
 			UMovieSceneSequence* RootSequence = SequencerPtr.Pin()->GetRootMovieSceneSequence();
 			if (RootSequence->GetTypedOuter<UBlueprint>() == nullptr)
diff --git a/Engine/Source/Editor/Sequencer/Private/Sequencer.cpp b/Engine/Source/Editor/Sequencer/Private/Sequencer.cpp
--- a/Engine/Source/Editor/Sequencer/Private/Sequencer.cpp
+++ b/Engine/Source/Editor/Sequencer/Private/Sequencer.cpp
@@ -13114,10 +13114,7 @@ void FSequencer::BindCommands()
 
 	SequencerCommandBindings->MapAction(
 		Commands.RenderMovie,
-		FExecuteAction::CreateLambda([this]{ RenderMovieInternal(GetPlaybackRange()); }),
-		FCanExecuteAction(),
-		FIsActionChecked(),
-		FIsActionButtonVisible::CreateLambda([this]{ return ExactCast<ULevelSequence>(GetFocusedMovieSceneSequence()) != nullptr; })
+		FExecuteAction::CreateLambda([this]{ RenderMovieInternal(GetPlaybackRange()); })
 	);
 
 	SequencerCommandBindings->MapAction(
diff --git a/Engine/Source/Editor/Sequencer/Public/ISequencerModule.h b/Engine/Source/Editor/Sequencer/Public/ISequencerModule.h
--- a/Engine/Source/Editor/Sequencer/Public/ISequencerModule.h
+++ b/Engine/Source/Editor/Sequencer/Public/ISequencerModule.h
@@ -115,10 +115,14 @@ struct FSequencerHostCapabilities
 	/** Do we support recording */
 	bool bSupportsRecording;
 
+	/** Do we support rendering to a movie */
+	bool bSupportsRenderMovie;
+
 	FSequencerHostCapabilities()
 		: bSupportsSaveMovieSceneAsset(false)
 		, bSupportsCurveEditor(false)
 		, bSupportsRecording(false)
+		, bSupportsRenderMovie(false)
 	{}
 };
 
