Sequencer: Check for null owning section

diff --git a/Engine/Source/Editor/Sequencer/Private/SSequencer.cpp b/Engine/Source/Editor/Sequencer/Private/SSequencer.cpp
--- a/Engine/Source/Editor/Sequencer/Private/SSequencer.cpp
+++ b/Engine/Source/Editor/Sequencer/Private/SSequencer.cpp
@@ -3575,7 +3575,10 @@ void SSequencer::StepToKey(bool bStepToNextKey, bool bCameraOnly)
 				SequencerHelpers::GetAllKeyAreas( Node, KeyAreas );
 				for ( TSharedPtr<IKeyArea> KeyArea : KeyAreas )
 				{
-					KeyArea->GetKeyTimes(AllTimes, KeyArea->GetOwningSection()->GetRange());
+					if (KeyArea->GetOwningSection())
+					{
+						KeyArea->GetKeyTimes(AllTimes, KeyArea->GetOwningSection()->GetRange());
+					}
 				}
 
 				TSet<TWeakObjectPtr<UMovieSceneSection> > Sections;
diff --git a/Engine/Source/Editor/Sequencer/Private/VirtualTrackArea.cpp b/Engine/Source/Editor/Sequencer/Private/VirtualTrackArea.cpp
--- a/Engine/Source/Editor/Sequencer/Private/VirtualTrackArea.cpp
+++ b/Engine/Source/Editor/Sequencer/Private/VirtualTrackArea.cpp
@@ -138,7 +138,7 @@ FSequencerSelectedKey FVirtualTrackArea::HitTestKey(FVector2D InPhysicalPosition
 		for (auto& KeyArea : KeyAreaNode->GetAllKeyAreas())
 		{
 			UMovieSceneSection* Section = KeyArea->GetOwningSection();
-			if (Section->GetRange().Overlaps(KeyRange))
+			if (Section && Section->GetRange().Overlaps(KeyRange))
 			{
 				KeyAreas.Add(KeyArea);
 			}
