Sequencer: Add support for renaming track rows.

Enable renaming for 3d Transform track

diff --git a/Engine/Plugins/MovieScene/SequencerScripting/Source/SequencerScripting/Private/ExtensionLibraries/MovieSceneTrackExtensions.cpp b/Engine/Plugins/MovieScene/SequencerScripting/Source/SequencerScripting/Private/ExtensionLibraries/MovieSceneTrackExtensions.cpp
--- a/Engine/Plugins/MovieScene/SequencerScripting/Source/SequencerScripting/Private/ExtensionLibraries/MovieSceneTrackExtensions.cpp
+++ b/Engine/Plugins/MovieScene/SequencerScripting/Source/SequencerScripting/Private/ExtensionLibraries/MovieSceneTrackExtensions.cpp
@@ -22,6 +22,24 @@ FText UMovieSceneTrackExtensions::GetDisplayName(UMovieSceneTrack* Track)
 	return FText::GetEmpty();
 }
 
+void UMovieSceneTrackExtensions::SetTrackRowDisplayName(UMovieSceneTrack* Track, const FText& InName, int32 RowIndex)
+{
+	if (UMovieSceneNameableTrack* NameableTrack = Cast<UMovieSceneNameableTrack>(Track))
+	{
+#if WITH_EDITORONLY_DATA
+		NameableTrack->SetTrackRowDisplayName(InName, RowIndex);
+#endif
+	}
+}
+
+FText UMovieSceneTrackExtensions::GetTrackRowDisplayName(UMovieSceneTrack* Track, int32 RowIndex)
+{
+#if WITH_EDITORONLY_DATA
+	return Track->GetTrackRowDisplayName(RowIndex);
+#endif
+	return FText::GetEmpty();
+}
+
 UMovieSceneSection* UMovieSceneTrackExtensions::AddSection(UMovieSceneTrack* Track)
 {
 	UMovieSceneSection* NewSection = Track->CreateNewSection();
diff --git a/Engine/Plugins/MovieScene/SequencerScripting/Source/SequencerScripting/Public/ExtensionLibraries/MovieSceneTrackExtensions.h b/Engine/Plugins/MovieScene/SequencerScripting/Source/SequencerScripting/Public/ExtensionLibraries/MovieSceneTrackExtensions.h
--- a/Engine/Plugins/MovieScene/SequencerScripting/Source/SequencerScripting/Public/ExtensionLibraries/MovieSceneTrackExtensions.h
+++ b/Engine/Plugins/MovieScene/SequencerScripting/Source/SequencerScripting/Public/ExtensionLibraries/MovieSceneTrackExtensions.h
@@ -41,6 +41,27 @@ public:
 	UFUNCTION(BlueprintCallable, Category = "Sequencer|Track", meta=(ScriptMethod, DevelopmentOnly))
 	static FText GetDisplayName(UMovieSceneTrack* Track);
 
+	/**
+	 * Set this track row's display name
+	 *
+	 * @param Track        The track to use
+	 * @param InName The name for this track
+	 * @param RowIndex The row index for the track
+	 */
+	UFUNCTION(BlueprintCallable, Category = "Sequencer|Track", meta=(ScriptMethod, DevelopmentOnly))
+	static void SetTrackRowDisplayName(UMovieSceneTrack* Track, const FText& InName, int32 RowIndex);
+
+	/**
+	 * Get this track row's display name
+	 *
+	 * @param Track        The track to use
+	 * @param RowIndex The row index for the track
+	 * @return This track's display name
+	 * 
+	 */
+	UFUNCTION(BlueprintCallable, Category = "Sequencer|Track", meta=(ScriptMethod, DevelopmentOnly))
+	static FText GetTrackRowDisplayName(UMovieSceneTrack* Track, int32 RowIndex);
+
 	/**
 	 * Add a new section to this track
 	 *
diff --git a/Engine/Source/Editor/Sequencer/Private/DisplayNodes/SequencerTrackNode.cpp b/Engine/Source/Editor/Sequencer/Private/DisplayNodes/SequencerTrackNode.cpp
--- a/Engine/Source/Editor/Sequencer/Private/DisplayNodes/SequencerTrackNode.cpp
+++ b/Engine/Source/Editor/Sequencer/Private/DisplayNodes/SequencerTrackNode.cpp
@@ -724,6 +724,16 @@ void FSequencerTrackNode::GetChildKeyAreaNodesRecursively(TArray<TSharedRef<FSeq
 
 FText FSequencerTrackNode::GetDisplayName() const
 {
+	if (UMovieSceneNameableTrack* NameableTrack = Cast<UMovieSceneNameableTrack>(AssociatedTrack.Get()))
+	{
+		if (SubTrackMode == ESubTrackMode::SubTrack)
+		{	
+			int32 TrackRowIndex = GetRowIndex();
+
+			return NameableTrack->GetTrackRowDisplayName(TrackRowIndex);
+		}
+	}
+
 	return AssociatedTrack.IsValid() ? AssociatedTrack->GetDisplayName() : FText::GetEmpty();
 }
 
@@ -813,13 +823,21 @@ ESequencerNode::Type FSequencerTrackNode::GetType() const
 
 void FSequencerTrackNode::SetDisplayName(const FText& NewDisplayName)
 {
-	auto NameableTrack = Cast<UMovieSceneNameableTrack>(AssociatedTrack.Get());
-
-	if (NameableTrack != nullptr && !NameableTrack->GetDisplayName().EqualTo(NewDisplayName))
+	if (UMovieSceneNameableTrack* NameableTrack = Cast<UMovieSceneNameableTrack>(AssociatedTrack.Get()))
 	{
 		const FScopedTransaction Transaction(NSLOCTEXT("SequencerTrackNode", "RenameTrack", "Rename Track"));
 
-		NameableTrack->SetDisplayName(NewDisplayName);
+		if (SubTrackMode == ESubTrackMode::SubTrack)
+		{
+			int32 TrackRowIndex = GetRowIndex();
+		
+			NameableTrack->SetTrackRowDisplayName(NewDisplayName, TrackRowIndex);
+		}
+		else
+		{
+			NameableTrack->SetDisplayName(NewDisplayName);
+		}
+
 		GetSequencer().NotifyMovieSceneDataChanged(EMovieSceneDataChangeType::TrackValueChanged);
 
 		SetNodeName(FName(*NewDisplayName.ToString()));
diff --git a/Engine/Source/Runtime/MovieScene/Private/MovieSceneNameableTrack.cpp b/Engine/Source/Runtime/MovieScene/Private/MovieSceneNameableTrack.cpp
--- a/Engine/Source/Runtime/MovieScene/Private/MovieSceneNameableTrack.cpp
+++ b/Engine/Source/Runtime/MovieScene/Private/MovieSceneNameableTrack.cpp
@@ -24,6 +24,24 @@ void UMovieSceneNameableTrack::SetDisplayName(const FText& NewDisplayName)
 	DisplayName = NewDisplayName;
 }
 
+void UMovieSceneNameableTrack::SetTrackRowDisplayName(const FText& NewDisplayName, int32 TrackRowIndex)
+{
+	if (TrackRowIndex >= TrackRowDisplayNames.Num())
+	{
+		TrackRowDisplayNames.AddDefaulted(TrackRowIndex+1);
+	}
+
+	if (NewDisplayName.EqualTo(TrackRowDisplayNames[TrackRowIndex]))
+	{
+		return;
+	}
+
+	SetFlags(RF_Transactional);
+	Modify();
+
+	TrackRowDisplayNames[TrackRowIndex] = NewDisplayName;
+}
+
 bool UMovieSceneNameableTrack::ValidateDisplayName(const FText& NewDisplayName, FText& OutErrorMessage) const
 {
 	if (NewDisplayName.IsEmpty())
@@ -57,6 +75,21 @@ FText UMovieSceneNameableTrack::GetDisplayName() const
 	return DisplayName;
 }
 
+FText UMovieSceneNameableTrack::GetTrackRowDisplayName(int32 TrackRowIndex) const
+{
+	if (TrackRowIndex < TrackRowDisplayNames.Num() && !TrackRowDisplayNames[TrackRowIndex].IsEmpty())
+	{
+		return TrackRowDisplayNames[TrackRowIndex];
+	}
+
+	if (DisplayName.IsEmpty())
+	{
+		return GetDefaultDisplayName();
+	}
+
+	return DisplayName;
+}
+
 FText UMovieSceneNameableTrack::GetDefaultDisplayName() const
 
 { 
diff --git a/Engine/Source/Runtime/MovieScene/Public/MovieSceneNameableTrack.h b/Engine/Source/Runtime/MovieScene/Public/MovieSceneNameableTrack.h
--- a/Engine/Source/Runtime/MovieScene/Public/MovieSceneNameableTrack.h
+++ b/Engine/Source/Runtime/MovieScene/Public/MovieSceneNameableTrack.h
@@ -27,6 +27,14 @@ public:
 	 */
 	MOVIESCENE_API virtual void SetDisplayName(const FText& NewDisplayName);
 
+	/**
+	 * Set the track row's display name.
+	 *
+	 * @param NewName The name to set.
+	 * @param RowIndex The track row index to set.
+	 */
+	MOVIESCENE_API virtual void SetTrackRowDisplayName(const FText& NewDisplayName, int32 TrackRowIndex);
+
 	/**
 	 * Can rename this track.
 	 *
@@ -47,6 +55,7 @@ public:
 	// UMovieSceneTrack interface
 
 	MOVIESCENE_API virtual FText GetDisplayName() const override;
+	MOVIESCENE_API virtual FText GetTrackRowDisplayName(int32 RowIndex) const override;
 	MOVIESCENE_API virtual FText GetDefaultDisplayName() const;
 
 private:
@@ -54,5 +63,10 @@ private:
 	/** The track's human readable display name. */
 	UPROPERTY(EditAnywhere, Category="Track")
 	FText DisplayName;
+
+	/** The track display name per row. */
+	UPROPERTY(EditAnywhere, Category="Track")
+	TArray<FText> TrackRowDisplayNames;
+
 #endif
 };
diff --git a/Engine/Source/Runtime/MovieScene/Public/MovieSceneTrack.h b/Engine/Source/Runtime/MovieScene/Public/MovieSceneTrack.h
--- a/Engine/Source/Runtime/MovieScene/Public/MovieSceneTrack.h
+++ b/Engine/Source/Runtime/MovieScene/Public/MovieSceneTrack.h
@@ -425,6 +425,13 @@ public:
 	 */
 	virtual FText GetDisplayName() const PURE_VIRTUAL(UMovieSceneTrack::GetDisplayName, return FText::FromString(TEXT("Unnamed Track")););
 
+	/**
+	 * Get the track row's display name.
+	 *
+	 * @return Display name text.
+	 */
+	virtual FText GetTrackRowDisplayName(int32 RowIndex) const PURE_VIRTUAL(UMovieSceneTrack::GetTrackRowDisplayName, return FText::FromString(TEXT("Unnamed Track")););
+
 	/**
 	 * Get the track's display tooltip text to be shown on the track's name.
 	 *
diff --git a/Engine/Source/Runtime/MovieSceneTracks/Public/Tracks/MovieScene3DTransformTrack.h b/Engine/Source/Runtime/MovieSceneTracks/Public/Tracks/MovieScene3DTransformTrack.h
--- a/Engine/Source/Runtime/MovieSceneTracks/Public/Tracks/MovieScene3DTransformTrack.h
+++ b/Engine/Source/Runtime/MovieSceneTracks/Public/Tracks/MovieScene3DTransformTrack.h
@@ -65,6 +65,10 @@ public:
 	virtual bool SupportsType(TSubclassOf<UMovieSceneSection> SectionClass) const override;
 	virtual UMovieSceneSection* CreateNewSection() override;
 
+#if WITH_EDITORONLY_DATA
+	virtual bool CanRename() const override { return true; }
+#endif
+
 #if WITH_EDITOR
 
 	MOVIESCENETRACKS_API TArray<FTrajectoryKey> GetTrajectoryData(FFrameNumber Time, int32 MaxNumDataPoints, TRange<FFrameNumber>) const;
