Added LevelSequenceEditorSubsystem to expose FixActorReferences() instead of living inside Sequencer. Extracted some utility functions to SequencerUtilities.

diff --git a/Engine/Plugins/MovieScene/LevelSequenceEditor/Source/LevelSequenceEditor/LevelSequenceEditor.Build.cs b/Engine/Plugins/MovieScene/LevelSequenceEditor/Source/LevelSequenceEditor/LevelSequenceEditor.Build.cs
--- a/Engine/Plugins/MovieScene/LevelSequenceEditor/Source/LevelSequenceEditor/LevelSequenceEditor.Build.cs
+++ b/Engine/Plugins/MovieScene/LevelSequenceEditor/Source/LevelSequenceEditor/LevelSequenceEditor.Build.cs
@@ -24,7 +24,8 @@ public class LevelSequenceEditor : ModuleRules
 				"CurveEditor",
 				"EditorFramework",
                 "EditorStyle",
-                "Engine",
+				"EditorSubsystem",
+				"Engine",
                 "InputCore",
                 "Kismet",
                 "KismetCompiler",
diff --git a/Engine/Plugins/MovieScene/LevelSequenceEditor/Source/LevelSequenceEditor/Private/LevelSequenceEditorCommands.cpp b/Engine/Plugins/MovieScene/LevelSequenceEditor/Source/LevelSequenceEditor/Private/LevelSequenceEditorCommands.cpp
--- a/Engine/Plugins/MovieScene/LevelSequenceEditor/Source/LevelSequenceEditor/Private/LevelSequenceEditorCommands.cpp
+++ b/Engine/Plugins/MovieScene/LevelSequenceEditor/Source/LevelSequenceEditor/Private/LevelSequenceEditorCommands.cpp
@@ -14,6 +14,9 @@ void FLevelSequenceEditorCommands::RegisterCommands()
 	UI_COMMAND(CreateNewLevelSequenceInLevel, "Add Level Sequence", "Create a new level sequence asset, and place an instance of it in this level", EUserInterfaceActionType::Button, FInputChord());
 	UI_COMMAND(CreateNewMasterSequenceInLevel, "Add Master Sequence", "Create a new master sequence asset, and place an instance of it in this level", EUserInterfaceActionType::Button, FInputChord());
 	UI_COMMAND(ToggleCinematicViewportCommand, "Cinematic Viewport", "A viewport layout tailored to cinematic preview", EUserInterfaceActionType::RadioButton, FInputChord());
+
+	UI_COMMAND(FixActorReferences, "Fix Actor References", "Try to automatically fix up broken actor bindings.", EUserInterfaceActionType::Button, FInputChord());
+
 }
 
 #undef LOCTEXT_NAMESPACE
diff --git a/Engine/Plugins/MovieScene/LevelSequenceEditor/Source/LevelSequenceEditor/Private/LevelSequenceEditorCommands.h b/Engine/Plugins/MovieScene/LevelSequenceEditor/Source/LevelSequenceEditor/Private/LevelSequenceEditorCommands.h
--- a/Engine/Plugins/MovieScene/LevelSequenceEditor/Source/LevelSequenceEditor/Private/LevelSequenceEditorCommands.h
+++ b/Engine/Plugins/MovieScene/LevelSequenceEditor/Source/LevelSequenceEditor/Private/LevelSequenceEditorCommands.h
@@ -19,4 +19,8 @@ public:
 	TSharedPtr<FUICommandInfo> CreateNewLevelSequenceInLevel;
 	TSharedPtr<FUICommandInfo> CreateNewMasterSequenceInLevel;
 	TSharedPtr<FUICommandInfo> ToggleCinematicViewportCommand;
+	
+	/** Attempts to fix broken actor references. */
+	TSharedPtr< FUICommandInfo > FixActorReferences;
+
 };
diff --git a/Engine/Plugins/MovieScene/LevelSequenceEditor/Source/LevelSequenceEditor/Private/LevelSequenceEditorSubsystem.cpp b/Engine/Plugins/MovieScene/LevelSequenceEditor/Source/LevelSequenceEditor/Private/LevelSequenceEditorSubsystem.cpp
new file mode 100644
--- /dev/null
+++ b/Engine/Plugins/MovieScene/LevelSequenceEditor/Source/LevelSequenceEditor/Private/LevelSequenceEditorSubsystem.cpp
@@ -0,0 +1,154 @@
+// Copyright Epic Games, Inc. All Rights Reserved.
+
+#include "LevelSequenceEditorSubsystem.h"
+#include "EngineUtils.h"
+#include "ISequencer.h"
+#include "ISequencerModule.h"
+#include "SequencerUtilities.h"
+#include "MovieScene.h"
+#include "MovieSceneSequence.h"
+#include "LevelSequenceEditorCommands.h"
+#include "ActorTreeItem.h"
+#include "LevelEditor.h"
+#include "Compilation/MovieSceneCompiledDataManager.h"
+
+
+DEFINE_LOG_CATEGORY(LogLevelSequenceEditor);
+
+void ULevelSequenceEditorSubsystem::Initialize(FSubsystemCollectionBase& Collection)
+{
+	UE_LOG(LogLevelSequenceEditor, Log, TEXT("LevelSequenceEditor subsystem initialized."));
+
+	ISequencerModule& SequencerModule = FModuleManager::Get().LoadModuleChecked<ISequencerModule>("Sequencer");
+	OnSequencerCreatedHandle = SequencerModule.RegisterOnSequencerCreated(FOnSequencerCreated::FDelegate::CreateUObject(this, &ULevelSequenceEditorSubsystem::OnSequencerCreated));
+
+	/* Commands for this subsystem */
+	CommandList = MakeShareable(new FUICommandList);
+	CommandList->MapAction(FLevelSequenceEditorCommands::Get().FixActorReferences,
+		FExecuteAction::CreateUObject(this, &ULevelSequenceEditorSubsystem::FixActorReferences)
+	);
+
+	/* Createand register "Fix Actor References" extension */
+	FixActorReferencesMenuExtender = MakeShareable(new FExtender);
+	FixActorReferencesMenuExtender->AddMenuExtension("Bindings", EExtensionHook::First, CommandList, FMenuExtensionDelegate::CreateStatic([](FMenuBuilder& MenuBuilder) {
+		MenuBuilder.AddMenuEntry(FLevelSequenceEditorCommands::Get().FixActorReferences);
+		}));
+
+	SequencerModule.GetActionsMenuExtensibilityManager()->AddExtender(FixActorReferencesMenuExtender);
+}
+
+void ULevelSequenceEditorSubsystem::Deinitialize()
+{
+	UE_LOG(LogLevelSequenceEditor, Log, TEXT("LevelSequenceEditor subsystem deinitialized."));
+
+	ISequencerModule* SequencerModulePtr = FModuleManager::Get().GetModulePtr<ISequencerModule>("Sequencer");
+	if (SequencerModulePtr)
+	{
+		SequencerModulePtr->UnregisterOnSequencerCreated(OnSequencerCreatedHandle);
+	}
+
+}
+
+void ULevelSequenceEditorSubsystem::OnSequencerCreated(TSharedRef<ISequencer> InSequencer)
+{
+	UE_LOG(LogLevelSequenceEditor, VeryVerbose, TEXT("ULevelSequenceEditorSubsystem::OnSequencerCreated"));
+
+	Sequencers.Add(TWeakPtr<ISequencer>(InSequencer));
+}
+
+TSharedPtr<ISequencer> ULevelSequenceEditorSubsystem::GetActiveSequencer()
+{
+	for (TWeakPtr<ISequencer> Ptr : Sequencers)
+	{
+		if (Ptr.IsValid())
+		{
+			return Ptr.Pin();
+		}
+	}
+
+	return nullptr;
+}
+
+void ULevelSequenceEditorSubsystem::FixActorReferences()
+{
+	TSharedPtr<ISequencer> Sequencer = GetActiveSequencer();
+	if (Sequencer == nullptr)
+	{
+		return;
+	}
+
+	UWorld* PlaybackContext = Cast<UWorld>(Sequencer->GetPlaybackContext());
+
+	if (!PlaybackContext)
+	{
+		return;
+	}
+
+	FScopedTransaction FixActorReferencesTransaction(NSLOCTEXT("LevelSequenceEditor", "FixActorReferences", "Fix Actor References"));
+
+	UMovieScene* FocusedMovieScene = Sequencer->GetFocusedMovieSceneSequence()->GetMovieScene();
+	if (!FocusedMovieScene)
+	{
+		return;
+	}
+
+	TMap<FString, AActor*> ActorNameToActorMap;
+
+	for (TActorIterator<AActor> ActorItr(PlaybackContext); ActorItr; ++ActorItr)
+	{
+		// Same as with the Object Iterator, access the subclass instance with the * or -> operators.
+		AActor* Actor = *ActorItr;
+		ActorNameToActorMap.Add(Actor->GetActorLabel(), Actor);
+	}
+
+	// Cache the possessables to fix up first since the bindings will change as the fix ups happen.
+	TArray<FMovieScenePossessable> ActorsPossessablesToFix;
+	for (int32 i = 0; i < FocusedMovieScene->GetPossessableCount(); i++)
+	{
+		FMovieScenePossessable& Possessable = FocusedMovieScene->GetPossessable(i);
+		// Possessables with parents are components so ignore them.
+		if (Possessable.GetParent().IsValid() == false)
+		{
+			if (Sequencer->FindBoundObjects(Possessable.GetGuid(), Sequencer->GetFocusedTemplateID()).Num() == 0)
+			{
+				ActorsPossessablesToFix.Add(Possessable);
+			}
+		}
+	}
+
+	// For the possessables to fix, look up the actors by name and reassign them if found.
+	TMap<FGuid, FGuid> OldGuidToNewGuidMap;
+	for (const FMovieScenePossessable& ActorPossessableToFix : ActorsPossessablesToFix)
+	{
+		AActor** ActorPtr = ActorNameToActorMap.Find(ActorPossessableToFix.GetName());
+		if (ActorPtr != nullptr)
+		{
+			FGuid OldGuid = ActorPossessableToFix.GetGuid();
+
+			// The actor might have an existing guid while the possessable with the same name might not. 
+			// In that case, make sure we also replace the existing guid with the new guid 
+			FGuid ExistingGuid = Sequencer->FindObjectId(**ActorPtr, Sequencer->GetFocusedTemplateID());
+
+			FGuid NewGuid = FSequencerUtilities::DoAssignActor(Sequencer.Get(), ActorPtr, 1, ActorPossessableToFix.GetGuid());
+
+			OldGuidToNewGuidMap.Add(OldGuid, NewGuid);
+
+			if (ExistingGuid.IsValid())
+			{
+				OldGuidToNewGuidMap.Add(ExistingGuid, NewGuid);
+			}
+		}
+	}
+
+	UMovieSceneCompiledDataManager* CompiledDataManager = FindObject<UMovieSceneCompiledDataManager>(GetTransientPackage(), TEXT("SequencerCompiledDataManager"));
+	if (!CompiledDataManager)
+	{
+		CompiledDataManager = NewObject<UMovieSceneCompiledDataManager>(GetTransientPackage(), "SequencerCompiledDataManager");
+	}
+
+	for (TPair<FGuid, FGuid> GuidPair : OldGuidToNewGuidMap)
+	{
+		FSequencerUtilities::UpdateBindingIDs(Sequencer.Get(), CompiledDataManager, GuidPair.Key, GuidPair.Value);
+	}
+}
+
diff --git a/Engine/Plugins/MovieScene/LevelSequenceEditor/Source/LevelSequenceEditor/Public/LevelSequenceEditorSubsystem.h b/Engine/Plugins/MovieScene/LevelSequenceEditor/Source/LevelSequenceEditor/Public/LevelSequenceEditorSubsystem.h
new file mode 100644
--- /dev/null
+++ b/Engine/Plugins/MovieScene/LevelSequenceEditor/Source/LevelSequenceEditor/Public/LevelSequenceEditorSubsystem.h
@@ -0,0 +1,47 @@
+// Copyright Epic Games, Inc. All Rights Reserved.
+
+
+#pragma once
+
+#include "CoreMinimal.h"
+#include "EditorSubsystem.h"
+#include "ISequencer.h"
+
+#include "LevelSequenceEditorSubsystem.generated.h"
+
+DECLARE_LOG_CATEGORY_EXTERN(LogLevelSequenceEditor, Log, All);
+
+class UMovieSceneCompiledDataManager;
+
+/**
+* ULevelSequenceEditorSubsystem
+* Subsystem for level sequencer related utilities to scripts
+*/
+UCLASS(Blueprintable)
+class ULevelSequenceEditorSubsystem : public UEditorSubsystem
+{
+	GENERATED_BODY()
+
+public:
+
+	virtual void Initialize(FSubsystemCollectionBase& Collection) override;
+	virtual void Deinitialize() override;
+
+	void OnSequencerCreated(TSharedRef<ISequencer> InSequencer);
+
+	/** Attempts to automatically fix up broken actor references in the current scene. */
+	UFUNCTION(BlueprintCallable, Category = "Level Sequence Editor")
+	void FixActorReferences();
+
+private:
+	TSharedPtr<ISequencer> GetActiveSequencer();
+
+	FDelegateHandle OnSequencerCreatedHandle;
+
+	/* List of sequencers that have been created */
+	TArray<TWeakPtr<ISequencer>> Sequencers;
+
+	TSharedPtr<FUICommandList> CommandList;
+
+	TSharedPtr<FExtender> FixActorReferencesMenuExtender;
+};
diff --git a/Engine/Source/Editor/Sequencer/Private/SSequencer.cpp b/Engine/Source/Editor/Sequencer/Private/SSequencer.cpp
--- a/Engine/Source/Editor/Sequencer/Private/SSequencer.cpp
+++ b/Engine/Source/Editor/Sequencer/Private/SSequencer.cpp
@@ -1971,7 +1971,11 @@ bool SSequencer::IsTrackLevelFilterActive(const FString LevelName) const
 
 TSharedRef<SWidget> SSequencer::MakeActionsMenu()
 {
-	FMenuBuilder MenuBuilder(true, SequencerPtr.Pin()->GetCommandBindings());
+	ISequencerModule& SequencerModule = FModuleManager::GetModuleChecked<ISequencerModule>("Sequencer");
+
+	TSharedPtr<FExtender> ActionsMenuExtender = SequencerModule.GetActionsMenuExtensibilityManager()->GetAllExtenders();
+
+	FMenuBuilder MenuBuilder(true, SequencerPtr.Pin()->GetCommandBindings(), ActionsMenuExtender);
 	TSharedPtr<FSequencer> Sequencer = SequencerPtr.Pin();
 
 	MenuBuilder.BeginSection("SequenceOptions", LOCTEXT("SequenceOptionsHeader", "Sequence"));
@@ -2043,7 +2047,6 @@ void SSequencer::FillAdvancedMenu(FMenuBuilder& MenuBuilder)
 	{
 		MenuBuilder.BeginSection("Bindings", LOCTEXT("BindingsMenuHeader", "Bindings"));
 
-		MenuBuilder.AddMenuEntry(FSequencerCommands::Get().FixActorReferences);
 		MenuBuilder.AddMenuEntry(FSequencerCommands::Get().RebindPossessableReferences);
 
 		MenuBuilder.EndSection();
diff --git a/Engine/Source/Editor/Sequencer/Private/Sequencer.cpp b/Engine/Source/Editor/Sequencer/Private/Sequencer.cpp
--- a/Engine/Source/Editor/Sequencer/Private/Sequencer.cpp
+++ b/Engine/Source/Editor/Sequencer/Private/Sequencer.cpp
@@ -1998,7 +1998,7 @@ void FSequencer::StretchTime(FFrameTime InDeltaTime)
 
 	if (FocusedMovieScene->IsReadOnly())
 	{
-		ShowReadOnlyError();
+		FSequencerUtilities::ShowReadOnlyError();
 		return;
 	}
 
@@ -2039,7 +2039,7 @@ void FSequencer::ShrinkTime(FFrameTime InDeltaTime)
 
 	if (FocusedMovieScene->IsReadOnly())
 	{
-		ShowReadOnlyError();
+		FSequencerUtilities::ShowReadOnlyError();
 		return;
 	}
 
@@ -2240,7 +2240,7 @@ void FSequencer::BakeTransform()
 
 	if (FocusedMovieScene->IsReadOnly())
 	{
-		ShowReadOnlyError();
+		FSequencerUtilities::ShowReadOnlyError();
 		return;
 	}
 
@@ -3123,7 +3123,7 @@ void FSequencer::TogglePlaybackRangeLocked()
 
 		if (MovieScene->IsReadOnly())
 		{
-			ShowReadOnlyError();
+			FSequencerUtilities::ShowReadOnlyError();
 			return;
 		}
 
@@ -5624,13 +5624,6 @@ bool FSequencer::IsReadOnly() const
 	return bReadOnly || (GetFocusedMovieSceneSequence() && GetFocusedMovieSceneSequence()->GetMovieScene()->IsReadOnly());
 }
 
-void FSequencer::ShowReadOnlyError() const
-{
-	FNotificationInfo Info(NSLOCTEXT("Sequencer", "SequenceReadOnly", "Sequence is read only."));
-	Info.ExpireDuration = 5.0f;
-	FSlateNotificationManager::Get().AddNotification(Info)->SetCompletionState(SNotificationItem::CS_Fail);
-}
-
 void FSequencer::VerticalScroll(float ScrollAmountUnits)
 {
 	SequencerWidget->GetTreeView()->ScrollByDelta(ScrollAmountUnits);
@@ -5675,7 +5668,7 @@ FGuid FSequencer::MakeNewSpawnable( UObject& Object, UActorFactory* ActorFactory
 
 	if (MovieScene->IsReadOnly())
 	{
-		ShowReadOnlyError();
+		FSequencerUtilities::ShowReadOnlyError();
 		return FGuid();
 	}
 
@@ -5724,7 +5717,7 @@ void FSequencer::AddSubSequence(UMovieSceneSequence* Sequence)
 
 	if (OwnerMovieScene->IsReadOnly())
 	{
-		ShowReadOnlyError();
+		FSequencerUtilities::ShowReadOnlyError();
 		return;
 	}
 
@@ -5797,7 +5790,7 @@ bool FSequencer::OnRequestNodeDeleted( TSharedRef<const FSequencerDisplayNode> N
 
 	if (OwnerMovieScene->IsReadOnly())
 	{
-		ShowReadOnlyError();
+		FSequencerUtilities::ShowReadOnlyError();
 		return bAnythingRemoved;
 	}
 
@@ -6028,7 +6021,7 @@ void FSequencer::OnNewActorsDropped(const TArray<UObject*>& DroppedObjects, cons
 
 		if (OwnerMovieScene->IsReadOnly())
 		{
-			ShowReadOnlyError();
+			FSequencerUtilities::ShowReadOnlyError();
 			return;
 		}
 
@@ -6414,7 +6407,7 @@ TArray<FGuid> FSequencer::AddActors(const TArray<TWeakObjectPtr<AActor> >& InAct
 
 	if (GetFocusedMovieSceneSequence()->GetMovieScene()->IsReadOnly())
 	{
-		ShowReadOnlyError();
+		FSequencerUtilities::ShowReadOnlyError();
 		return PossessableGuids;
 	}
 
@@ -7789,7 +7782,7 @@ void FSequencer::AssignActor(FMenuBuilder& MenuBuilder, FGuid InObjectBinding)
 				FOnActorPicked::CreateLambda([=](AActor* Actor){
 					// Create a new binding for this actor
 					FSlateApplication::Get().DismissAllMenus();
-					DoAssignActor(&Actor, 1, InObjectBinding);
+					FSequencerUtilities::DoAssignActor(this, &Actor, 1, InObjectBinding);
 				})
 			)
 		];
@@ -7799,177 +7792,6 @@ void FSequencer::AssignActor(FMenuBuilder& MenuBuilder, FGuid InObjectBinding)
 }
 
 
-FGuid FSequencer::DoAssignActor(AActor*const* InActors, int32 NumActors, FGuid InObjectBinding)
-{
-	if (NumActors <= 0)
-	{
-		return FGuid();
-	}
-
-	//@todo: this code doesn't work with multiple actors, or when the existing binding is bound to multiple actors
-
-	AActor* Actor = InActors[0];
-
-	if (Actor == nullptr)
-	{
-		return FGuid();
-	}
-
-	UMovieSceneSequence* OwnerSequence = GetFocusedMovieSceneSequence();
-	UMovieScene* OwnerMovieScene = OwnerSequence->GetMovieScene();
-
-	if (OwnerMovieScene->IsReadOnly())
-	{
-		ShowReadOnlyError();
-		return FGuid();
-	}
-
-	FScopedTransaction AssignActor( NSLOCTEXT("Sequencer", "AssignActor", "Assign Actor") );
-
-	Actor->Modify();
-	OwnerSequence->Modify();
-	OwnerMovieScene->Modify();
-
-	TArrayView<TWeakObjectPtr<>> RuntimeObjects = FindObjectsInCurrentSequence(InObjectBinding);
-
-	UObject* RuntimeObject = RuntimeObjects.Num() ? RuntimeObjects[0].Get() : nullptr;
-
-	// Replace the object itself
-	FMovieScenePossessable NewPossessableActor;
-	FGuid NewGuid;
-	{
-		// Get the object guid to assign, remove the binding if it already exists
-		FGuid ParentGuid = FindObjectId(*Actor, ActiveTemplateIDs.Top());
-		FString NewActorLabel = Actor->GetActorLabel();
-		if (ParentGuid.IsValid())
-		{
-			OwnerMovieScene->RemovePossessable(ParentGuid);
-			OwnerSequence->UnbindPossessableObjects(ParentGuid);
-		}
-
-		// Add this object
-		NewPossessableActor = FMovieScenePossessable( NewActorLabel, Actor->GetClass());
-		NewGuid = NewPossessableActor.GetGuid();
-		OwnerSequence->BindPossessableObject(NewPossessableActor.GetGuid(), *Actor, GetPlaybackContext());
-
-		// Defer replacing this object until the components have been updated
-	}
-
-	auto UpdateComponent = [&]( FGuid OldComponentGuid, UActorComponent* NewComponent )
-	{
-		// Get the object guid to assign, remove the binding if it already exists
-		FGuid NewComponentGuid = FindObjectId( *NewComponent, ActiveTemplateIDs.Top() );
-		if ( NewComponentGuid.IsValid() )
-		{
-			OwnerMovieScene->RemovePossessable( NewComponentGuid );
-			OwnerSequence->UnbindPossessableObjects( NewComponentGuid );
-		}
-
-		// Add this object
-		FMovieScenePossessable NewPossessable( NewComponent->GetName(), NewComponent->GetClass() );
-		OwnerSequence->BindPossessableObject( NewPossessable.GetGuid(), *NewComponent, Actor );
-
-		// Replace
-		OwnerMovieScene->ReplacePossessable( OldComponentGuid, NewPossessable );
-		OwnerSequence->UnbindPossessableObjects( OldComponentGuid );
-		State.Invalidate(OldComponentGuid, ActiveTemplateIDs.Top());
-		State.Invalidate(NewPossessable.GetGuid(), ActiveTemplateIDs.Top());
-
-		FMovieScenePossessable* ThisPossessable = OwnerMovieScene->FindPossessable( NewPossessable.GetGuid() );
-		if ( ensure( ThisPossessable ) )
-		{
-			ThisPossessable->SetParent( NewGuid );
-		}
-	};
-
-	// Handle components
-	AActor* ActorToReplace = Cast<AActor>(RuntimeObject);
-	if (ActorToReplace != nullptr && ActorToReplace->IsActorBeingDestroyed() == false)
-	{
-		for (UActorComponent* ComponentToReplace : ActorToReplace->GetComponents())
-		{
-			if (ComponentToReplace != nullptr)
-			{
-				FGuid ComponentGuid = FindObjectId(*ComponentToReplace, ActiveTemplateIDs.Top());
-				if (ComponentGuid.IsValid())
-				{
-					bool bComponentWasUpdated = false;
-					for (UActorComponent* NewComponent : Actor->GetComponents())
-					{
-						if (NewComponent->GetFullName(Actor) == ComponentToReplace->GetFullName(ActorToReplace))
-						{
-							UpdateComponent( ComponentGuid, NewComponent );
-							bComponentWasUpdated = true;
-						}
-					}
-
-					// Clear the parent guid since this possessable component doesn't match to any component on the new actor
-					if (!bComponentWasUpdated)
-					{
-						FMovieScenePossessable* ThisPossessable = OwnerMovieScene->FindPossessable(ComponentGuid);
-						ThisPossessable->SetParent(FGuid());
-					}
-				}
-			}
-		}
-	}
-	else // If the actor didn't exist, try to find components who's parent guids were the previous actors guid.
-	{
-		TMap<FString, UActorComponent*> ComponentNameToComponent;
-		for ( UActorComponent* Component : Actor->GetComponents() )
-		{
-			ComponentNameToComponent.Add( Component->GetName(), Component );
-		}
-		for ( int32 i = 0; i < OwnerMovieScene->GetPossessableCount(); i++ )
-		{
-			FMovieScenePossessable& OldPossessable = OwnerMovieScene->GetPossessable(i);
-			if ( OldPossessable.GetParent() == InObjectBinding )
-			{
-				UActorComponent** ComponentPtr = ComponentNameToComponent.Find( OldPossessable.GetName() );
-				if ( ComponentPtr != nullptr )
-				{
-					UpdateComponent( OldPossessable.GetGuid(), *ComponentPtr );
-				}
-			}
-		}
-	}
-
-	// Replace the actor itself after components have been updated
-	OwnerMovieScene->ReplacePossessable(InObjectBinding, NewPossessableActor);
-	OwnerSequence->UnbindPossessableObjects(InObjectBinding);
-
-	State.Invalidate(InObjectBinding, ActiveTemplateIDs.Top());
-	State.Invalidate(NewPossessableActor.GetGuid(), ActiveTemplateIDs.Top());
-
-	// Try to fix up folders
-	TArray<UMovieSceneFolder*> FoldersToCheck;
-	FoldersToCheck.Append(GetFocusedMovieSceneSequence()->GetMovieScene()->GetRootFolders());
-	bool bFolderFound = false;
-	while ( FoldersToCheck.Num() > 0 && bFolderFound == false )
-	{
-		UMovieSceneFolder* Folder = FoldersToCheck[0];
-		FoldersToCheck.RemoveAt(0);
-		if ( Folder->GetChildObjectBindings().Contains( InObjectBinding ) )
-		{
-			Folder->RemoveChildObjectBinding( InObjectBinding );
-			Folder->AddChildObjectBinding( NewGuid );
-			bFolderFound = true;
-		}
-
-		for ( UMovieSceneFolder* ChildFolder : Folder->GetChildFolders() )
-		{
-			FoldersToCheck.Add( ChildFolder );
-		}
-	}
-
-	RestorePreAnimatedState();
-
-	NotifyMovieSceneDataChanged( EMovieSceneDataChangeType::MovieSceneStructureItemsChanged );
-
-	return NewGuid;
-}
-
-
 void FSequencer::AddActorsToBinding(FGuid InObjectBinding, const TArray<AActor*>& InActors)
 {
 	if (!InActors.Num())
@@ -8435,7 +8257,7 @@ void FSequencer::MoveSelectedNodesToFolder(UMovieSceneFolder* DestinationFolder)
 
 	if (FocusedMovieScene->IsReadOnly())
 	{
-		ShowReadOnlyError();
+		FSequencerUtilities::ShowReadOnlyError();
 		return;
 	}
 
@@ -8602,7 +8424,7 @@ void FSequencer::MoveSelectedNodesToNewFolder()
 
 	if (FocusedMovieScene->IsReadOnly())
 	{
-		ShowReadOnlyError();
+		FSequencerUtilities::ShowReadOnlyError();
 		return;
 	}
 
@@ -8765,7 +8587,7 @@ void FSequencer::RemoveSelectedNodesFromFolders()
 
 	if (FocusedMovieScene->IsReadOnly())
 	{
-		ShowReadOnlyError();
+		FSequencerUtilities::ShowReadOnlyError();
 		return;
 	}
 
@@ -9120,7 +8942,7 @@ bool FSequencer::DoPaste(bool bClearSelection)
 {
 	if (IsReadOnly())
 	{
-		ShowReadOnlyError();
+		FSequencerUtilities::ShowReadOnlyError();
 		// If we cancel the paste due to being read-only, count that as having handled the paste operation
 		return true;
 	}
@@ -9456,7 +9278,7 @@ bool FSequencer::PasteObjectBindings(const FString& TextToImport, UMovieSceneFol
 
 					if (!ExistingGuid.IsValid())
 					{
-						FGuid NewGuid = DoAssignActor(&Actor, 1, Possessable->GetGuid());
+						FGuid NewGuid = FSequencerUtilities::DoAssignActor(this, &Actor, 1, Possessable->GetGuid());
 
 						// If assigning produces a new guid, update the possessable guids and the bindings pasted data
 						if (NewGuid.IsValid())
@@ -10288,7 +10110,7 @@ void FSequencer::GroupSelectedSections()
 	UMovieScene* MovieScene = GetFocusedMovieSceneSequence()->GetMovieScene();
 	if (MovieScene->IsReadOnly())
 	{
-		ShowReadOnlyError();
+		FSequencerUtilities::ShowReadOnlyError();
 		return;
 	}
 
@@ -10331,7 +10153,7 @@ void FSequencer::UngroupSelectedSections()
 	UMovieScene* MovieScene = GetFocusedMovieSceneSequence()->GetMovieScene();
 	if (MovieScene->IsReadOnly())
 	{
-		ShowReadOnlyError();
+		FSequencerUtilities::ShowReadOnlyError();
 		return;
 	}
 
@@ -10366,7 +10188,7 @@ void FSequencer::SaveSelectedNodesSpawnableState()
 
 	if (MovieScene->IsReadOnly())
 	{
-		ShowReadOnlyError();
+		FSequencerUtilities::ShowReadOnlyError();
 		return;
 	}
 
@@ -10413,7 +10235,7 @@ void FSequencer::SetSelectedNodesSpawnableLevel(FName InLevelName)
 
 	if (MovieScene->IsReadOnly())
 	{
-		ShowReadOnlyError();
+		FSequencerUtilities::ShowReadOnlyError();
 		return;
 	}
 
@@ -10440,7 +10262,7 @@ void FSequencer::ConvertToSpawnable(TSharedRef<FSequencerObjectBindingNode> Node
 {
 	if (GetFocusedMovieSceneSequence()->GetMovieScene()->IsReadOnly())
 	{
-		ShowReadOnlyError();
+		FSequencerUtilities::ShowReadOnlyError();
 		return;
 	}
 
@@ -10479,7 +10301,7 @@ void FSequencer::ConvertSelectedNodesToSpawnables()
 
 	if (MovieScene->IsReadOnly())
 	{
-		ShowReadOnlyError();
+		FSequencerUtilities::ShowReadOnlyError();
 		return;
 	}
 
@@ -10569,7 +10391,7 @@ TArray<FGuid> FSequencer::ExpandMultiplePossessableBindings(FGuid PossessableGui
 
 	if (MovieScene->IsReadOnly())
 	{
-		ShowReadOnlyError();
+		FSequencerUtilities::ShowReadOnlyError();
 		return TArray<FGuid>();
 	}
 
@@ -10675,7 +10497,7 @@ TArray<FMovieSceneSpawnable*> FSequencer::ConvertToSpawnableInternal(FGuid Posse
 
 	if (MovieScene->IsReadOnly() || !Sequence->AllowsSpawnableObjects())
 	{
-		ShowReadOnlyError();
+		FSequencerUtilities::ShowReadOnlyError();
 		return TArray<FMovieSceneSpawnable*>();
 	}
 
@@ -10759,7 +10581,7 @@ TArray<FMovieSceneSpawnable*> FSequencer::ConvertToSpawnableInternal(FGuid Posse
 			SpawnRegister->HandleConvertPossessableToSpawnable(FoundObject, *this, TransformData);
 			SpawnRegister->SetupDefaultsForSpawnable(nullptr, Spawnable->GetGuid(), TransformData, AsShared(), Settings);
 
-			UpdateBindingIDs(PossessableGuid, Spawnable->GetGuid());
+			FSequencerUtilities::UpdateBindingIDs(this, CompiledDataManager, PossessableGuid, Spawnable->GetGuid());
 
 			ForceEvaluate();
 
@@ -10774,7 +10596,7 @@ void FSequencer::ConvertToPossessable(TSharedRef<FSequencerObjectBindingNode> No
 {
 	if (GetFocusedMovieSceneSequence()->GetMovieScene()->IsReadOnly())
 	{
-		ShowReadOnlyError();
+		FSequencerUtilities::ShowReadOnlyError();
 		return;
 	}
 
@@ -10798,7 +10620,7 @@ void FSequencer::ConvertSelectedNodesToPossessables()
 
 	if (MovieScene->IsReadOnly())
 	{
-		ShowReadOnlyError();
+		FSequencerUtilities::ShowReadOnlyError();
 		return;
 	}
 
@@ -10882,7 +10704,7 @@ FMovieScenePossessable* FSequencer::ConvertToPossessableInternal(FGuid Spawnable
 
 	if (MovieScene->IsReadOnly())
 	{
-		ShowReadOnlyError();
+		FSequencerUtilities::ShowReadOnlyError();
 		return nullptr;
 	}
 
@@ -11006,7 +10828,7 @@ FMovieScenePossessable* FSequencer::ConvertToPossessableInternal(FGuid Spawnable
 		static const FName SequencerActorTag(TEXT("SequencerActor"));
 		PossessedActor->Tags.Remove(SequencerActorTag);
 
-		UpdateBindingIDs(OldSpawnableGuid, NewPossessableGuid);
+		FSequencerUtilities::UpdateBindingIDs(this, CompiledDataManager, OldSpawnableGuid, NewPossessableGuid);
 
 		GEditor->SelectActor(PossessedActor, false, true);
 
@@ -11027,45 +10849,6 @@ FMovieScenePossessable* FSequencer::ConvertToPossessableInternal(FGuid Spawnable
 	return Possessable;
 }
 
-void FSequencer::UpdateBindingIDs(FGuid OldGuid, FGuid NewGuid)
-{
-	TMap<UE::MovieScene::FFixedObjectBindingID, UE::MovieScene::FFixedObjectBindingID> OldFixedToNewFixedMap;
-	OldFixedToNewFixedMap.Add(UE::MovieScene::FFixedObjectBindingID(OldGuid, ActiveTemplateIDs.Top()), UE::MovieScene::FFixedObjectBindingID(NewGuid, ActiveTemplateIDs.Top()));
-
-	const FMovieSceneSequenceHierarchy* Hierarchy = CompiledDataManager->FindHierarchy(RootTemplateInstance.GetCompiledDataID());
-
-	if (UMovieScene* MovieScene = GetRootMovieSceneSequence()->GetMovieScene())
-	{
-		for (UMovieSceneSection* Section : MovieScene->GetAllSections())
-		{
-			if (Section)
-			{
-				Section->OnBindingIDsUpdated(OldFixedToNewFixedMap, GetRootTemplateID(), Hierarchy, *this);
-			}
-		}
-	}
-
-	if (Hierarchy)
-	{
-		for (const TTuple<FMovieSceneSequenceID, FMovieSceneSubSequenceData>& Pair : Hierarchy->AllSubSequenceData())
-		{
-			if (UMovieSceneSequence* Sequence = Pair.Value.GetSequence())
-			{
-				if (UMovieScene* MovieScene = Sequence->GetMovieScene())
-				{
-					for (UMovieSceneSection* Section : MovieScene->GetAllSections())
-					{
-						if (Section)
-						{
-							Section->OnBindingIDsUpdated(OldFixedToNewFixedMap, Pair.Key, Hierarchy, *this);
-						}
-					}
-				}
-			}
-		}
-	}
-}
-
 void FSequencer::OnLoadRecordedData()
 {
 	UMovieSceneSequence* FocusedMovieSceneSequence = GetFocusedMovieSceneSequence();
@@ -11134,7 +10917,7 @@ bool FSequencer::ReplaceFolderBindingGUID(UMovieSceneFolder* Folder, FGuid Origi
 
 	if (MovieScene->IsReadOnly())
 	{
-		ShowReadOnlyError();
+		FSequencerUtilities::ShowReadOnlyError();
 		return true;
 	}
 
@@ -11169,7 +10952,7 @@ void FSequencer::OnAddFolder()
 
 	if (FocusedMovieScene->IsReadOnly())
 	{
-		ShowReadOnlyError();
+		FSequencerUtilities::ShowReadOnlyError();
 		return;
 	}
 
@@ -12362,7 +12145,7 @@ void FSequencer::CreateCamera()
 
 	if (FocusedMovieScene->IsReadOnly())
 	{
-		ShowReadOnlyError();
+		FSequencerUtilities::ShowReadOnlyError();
 		return;
 	}
 
@@ -12454,78 +12237,6 @@ void FSequencer::NewCameraAdded(ACameraActor* NewCamera, FGuid CameraGuid)
 	MovieSceneToolHelpers::CreateCameraCutSectionForCamera(GetFocusedMovieSceneSequence()->GetMovieScene(), CameraGuid, GetLocalTime().Time.FloorToFrame());
 }
 
-
-void FSequencer::FixActorReferences()
-{
-	UWorld* PlaybackContext = Cast<UWorld>(GetPlaybackContext());
-
-	if (!PlaybackContext)
-	{
-		return;
-	}
-
-	FScopedTransaction FixActorReferencesTransaction( NSLOCTEXT( "Sequencer", "FixActorReferences", "Fix Actor References" ) );
-
-	UMovieScene* FocusedMovieScene = GetFocusedMovieSceneSequence()->GetMovieScene();
-	if (!FocusedMovieScene)
-	{
-		return;
-	}
-
-	TMap<FString, AActor*> ActorNameToActorMap;
-
-	for ( TActorIterator<AActor> ActorItr( PlaybackContext ); ActorItr; ++ActorItr )
-	{
-		// Same as with the Object Iterator, access the subclass instance with the * or -> operators.
-		AActor *Actor = *ActorItr;
-		ActorNameToActorMap.Add( Actor->GetActorLabel(), Actor);
-	}
-
-	// Cache the possessables to fix up first since the bindings will change as the fix ups happen.
-	TArray<FMovieScenePossessable> ActorsPossessablesToFix;
-	for ( int32 i = 0; i < FocusedMovieScene->GetPossessableCount(); i++ )
-	{
-		FMovieScenePossessable& Possessable = FocusedMovieScene->GetPossessable( i );
-		// Possessables with parents are components so ignore them.
-		if ( Possessable.GetParent().IsValid() == false )
-		{
-			if ( FindBoundObjects(Possessable.GetGuid(), ActiveTemplateIDs.Top()).Num() == 0 )
-			{
-				ActorsPossessablesToFix.Add( Possessable );
-			}
-		}
-	}
-
-	// For the possessables to fix, look up the actors by name and reassign them if found.
-	TMap<FGuid, FGuid> OldGuidToNewGuidMap;
-	for ( const FMovieScenePossessable& ActorPossessableToFix : ActorsPossessablesToFix )
-	{
-		AActor** ActorPtr = ActorNameToActorMap.Find( ActorPossessableToFix.GetName() );
-		if ( ActorPtr != nullptr )
-		{
-			FGuid OldGuid = ActorPossessableToFix.GetGuid();
-
-			// The actor might have an existing guid while the possessable with the same name might not. 
-			// In that case, make sure we also replace the existing guid with the new guid 
-			FGuid ExistingGuid = FindObjectId( **ActorPtr, ActiveTemplateIDs.Top() );
-
-			FGuid NewGuid = DoAssignActor( ActorPtr, 1, ActorPossessableToFix.GetGuid() );
-
-			OldGuidToNewGuidMap.Add(OldGuid, NewGuid);
-
-			if (ExistingGuid.IsValid())
-			{
-				OldGuidToNewGuidMap.Add(ExistingGuid, NewGuid);
-			}
-		}
-	}
-
-	for (TPair<FGuid, FGuid> GuidPair : OldGuidToNewGuidMap)
-	{
-		UpdateBindingIDs(GuidPair.Key, GuidPair.Value);
-	}
-}
-
 void FSequencer::RebindPossessableReferences()
 {
 	UMovieSceneSequence* FocusedSequence = GetFocusedMovieSceneSequence();
@@ -12537,7 +12248,7 @@ void FSequencer::RebindPossessableReferences()
 
 	if (FocusedMovieScene->IsReadOnly())
 	{
-		ShowReadOnlyError();
+		FSequencerUtilities::ShowReadOnlyError();
 		return;
 	}
 
@@ -13507,11 +13218,6 @@ void FSequencer::BindCommands()
 		FExecuteAction::CreateSP( this, &FSequencer::SyncSectionsUsingSourceTimecode ),
 		FCanExecuteAction::CreateLambda( [this]{ return (GetSelection().GetSelectedSections().Num() > 1); } ) );
 
-	SequencerCommandBindings->MapAction(
-		Commands.FixActorReferences,
-		FExecuteAction::CreateSP( this, &FSequencer::FixActorReferences ),
-		FCanExecuteAction::CreateLambda( []{ return true; } ) );
-
 	SequencerCommandBindings->MapAction(
 		Commands.RebindPossessableReferences,
 		FExecuteAction::CreateSP( this, &FSequencer::RebindPossessableReferences ),
diff --git a/Engine/Source/Editor/Sequencer/Private/Sequencer.h b/Engine/Source/Editor/Sequencer/Private/Sequencer.h
--- a/Engine/Source/Editor/Sequencer/Private/Sequencer.h
+++ b/Engine/Source/Editor/Sequencer/Private/Sequencer.h
@@ -526,7 +526,6 @@ public:
 
 	/** Called when a user executes the assign actor to track menu item */
 	void AssignActor(FMenuBuilder& MenuBuilder, FGuid ObjectBinding);
-	FGuid DoAssignActor(AActor*const* InActors, int32 NumActors, FGuid ObjectBinding);
 
 	/** Called when a user executes the assign selected to track menu item */
 	void AddActorsToBinding(FGuid ObjectBinding, const TArray<AActor*>& InActors);
@@ -662,9 +661,6 @@ public:
 	/** Called when a new camera is added. Locks the viewport to the NewCamera. */
 	void NewCameraAdded(ACameraActor* NewCamera, FGuid CameraGuid);
 
-	/** Attempts to automatically fix up broken actor references in the current scene. */
-	void FixActorReferences();
-
 	/** Rebinds all possessable references in the current sequence to update them to the latest referencing mechanism. */
 	void RebindPossessableReferences();
 
@@ -678,9 +674,6 @@ public:
 	/** Exports the animation to a camera anim asset. */
 	void ExportToCameraAnim();
 
-	/** */
-	void ShowReadOnlyError() const;
-
 public:
 	
 	/** Access the currently active track area edit tool */
@@ -1092,8 +1085,6 @@ private:
 
 	/** Internal conversion function that doesn't perform expensive reset/update tasks */
 	FMovieScenePossessable* ConvertToPossessableInternal(FGuid SpawnableGuid);
-	
-	void UpdateBindingIDs(FGuid OldGuid, FGuid NewGuid);
 
 	/** Recurses through a folder to replace converted GUID with new GUID */
 	bool ReplaceFolderBindingGUID(UMovieSceneFolder *Folder, FGuid Original, FGuid Converted);
diff --git a/Engine/Source/Editor/Sequencer/Private/SequencerCommands.cpp b/Engine/Source/Editor/Sequencer/Private/SequencerCommands.cpp
--- a/Engine/Source/Editor/Sequencer/Private/SequencerCommands.cpp
+++ b/Engine/Source/Editor/Sequencer/Private/SequencerCommands.cpp
@@ -152,7 +152,6 @@ void FSequencerCommands::RegisterCommands()
 	UI_COMMAND( SaveCurrentSpawnableState, "Save Default State", "Save the current state of this spawnable object as its default properties.", EUserInterfaceActionType::Button, FInputChord() );
 	UI_COMMAND( RestoreAnimatedState, "Restore Pre-Animated State", "Restore any objects that have been animated by sequencer back to their original state.", EUserInterfaceActionType::Button, FInputChord(EModifierKey::Control, EKeys::R) );
 
-	UI_COMMAND( FixActorReferences, "Fix Actor References", "Try to automatically fix up broken actor bindings.", EUserInterfaceActionType::Button, FInputChord() );
 	UI_COMMAND( RebindPossessableReferences, "Rebind Possesable References", "Rebinds all possessables in the current sequence to ensure they're using the most robust referencing mechanism.", EUserInterfaceActionType::Button, FInputChord() );
 
 	UI_COMMAND( ImportFBX, "Import...", "Import the animation from an FBX file.", EUserInterfaceActionType::Button, FInputChord() );
diff --git a/Engine/Source/Editor/Sequencer/Private/SequencerModule.cpp b/Engine/Source/Editor/Sequencer/Private/SequencerModule.cpp
--- a/Engine/Source/Editor/Sequencer/Private/SequencerModule.cpp
+++ b/Engine/Source/Editor/Sequencer/Private/SequencerModule.cpp
@@ -305,6 +305,7 @@ public:
 		ObjectBindingContextMenuExtensibilityManager = MakeShareable( new FExtensibilityManager );
 		AddTrackMenuExtensibilityManager = MakeShareable( new FExtensibilityManager );
 		ToolBarExtensibilityManager = MakeShareable(new FExtensibilityManager);
+		ActionsMenuExtensibilityManager = MakeShareable(new FExtensibilityManager);
 
 		SequencerCustomizationManager = MakeShareable(new FSequencerCustomizationManager);
 	}
@@ -377,6 +378,7 @@ public:
 	virtual TSharedPtr<FExtensibilityManager> GetObjectBindingContextMenuExtensibilityManager() const override { return ObjectBindingContextMenuExtensibilityManager; }
 	virtual TSharedPtr<FExtensibilityManager> GetAddTrackMenuExtensibilityManager() const override { return AddTrackMenuExtensibilityManager; }
 	virtual TSharedPtr<FExtensibilityManager> GetToolBarExtensibilityManager() const override { return ToolBarExtensibilityManager; }
+	virtual TSharedPtr<FExtensibilityManager> GetActionsMenuExtensibilityManager() const override { return ActionsMenuExtensibilityManager; }
 
 	virtual TSharedPtr<FSequencerCustomizationManager> GetSequencerCustomizationManager() const override { return SequencerCustomizationManager; }
 
@@ -447,6 +449,7 @@ private:
 	TSharedPtr<FExtensibilityManager> ObjectBindingContextMenuExtensibilityManager;
 	TSharedPtr<FExtensibilityManager> AddTrackMenuExtensibilityManager;
 	TSharedPtr<FExtensibilityManager> ToolBarExtensibilityManager;
+	TSharedPtr<FExtensibilityManager> ActionsMenuExtensibilityManager;
 
 	TSharedPtr<FSequencerCustomizationManager> SequencerCustomizationManager;
 
diff --git a/Engine/Source/Editor/Sequencer/Private/SequencerUtilities.cpp b/Engine/Source/Editor/Sequencer/Private/SequencerUtilities.cpp
--- a/Engine/Source/Editor/Sequencer/Private/SequencerUtilities.cpp
+++ b/Engine/Source/Editor/Sequencer/Private/SequencerUtilities.cpp
@@ -16,6 +16,8 @@
 #include "MovieSceneTrack.h"
 #include "MovieSceneSection.h"
 #include "MovieSceneTimeHelpers.h"
+#include "MovieSceneFolder.h"
+#include "Compilation/MovieSceneCompiledDataManager.h"
 #include "Framework/MultiBox/MultiBoxBuilder.h"
 #include "ISequencerTrackEditor.h"
 #include "ISequencer.h"
@@ -27,6 +29,8 @@
 #include "AssetRegistryModule.h"
 #include "FileHelpers.h"
 #include "LevelSequence.h"
+#include "Framework/Notifications/NotificationManager.h"
+#include "Widgets/Notifications/SNotificationList.h"
 
 
 
@@ -342,5 +346,225 @@ TArray<FString> FSequencerUtilities::GetAssociatedMapPackages(const ULevelSequen
 	return AssociatedMaps;
 }
 
+FGuid FSequencerUtilities::DoAssignActor(ISequencer * InSequencerPtr, AActor* const* InActors, int32 NumActors, FGuid InObjectBinding)
+{
+	if (NumActors <= 0)
+	{
+		return FGuid();
+	}
+
+	//@todo: this code doesn't work with multiple actors, or when the existing binding is bound to multiple actors
+
+	AActor* Actor = InActors[0];
+
+	if (Actor == nullptr)
+	{
+		return FGuid();
+	}
+
+	UMovieSceneSequence* OwnerSequence = InSequencerPtr->GetFocusedMovieSceneSequence();
+	UMovieScene* OwnerMovieScene = OwnerSequence->GetMovieScene();
+
+	if (OwnerMovieScene->IsReadOnly())
+	{
+		FSequencerUtilities::ShowReadOnlyError();
+		return FGuid();
+	}
+
+	FScopedTransaction AssignActor(LOCTEXT("AssignActor", "Assign Actor"));
+
+	Actor->Modify();
+	OwnerSequence->Modify();
+	OwnerMovieScene->Modify();
+
+	TArrayView<TWeakObjectPtr<>> RuntimeObjects = InSequencerPtr->FindObjectsInCurrentSequence(InObjectBinding);
+
+	UObject* RuntimeObject = RuntimeObjects.Num() ? RuntimeObjects[0].Get() : nullptr;
+
+	// Replace the object itself
+	FMovieScenePossessable NewPossessableActor;
+	FGuid NewGuid;
+	{
+		// Get the object guid to assign, remove the binding if it already exists
+		FGuid ParentGuid = InSequencerPtr->FindObjectId(*Actor, InSequencerPtr->GetFocusedTemplateID());
+		FString NewActorLabel = Actor->GetActorLabel();
+		if (ParentGuid.IsValid())
+		{
+			OwnerMovieScene->RemovePossessable(ParentGuid);
+			OwnerSequence->UnbindPossessableObjects(ParentGuid);
+		}
+
+		// Add this object
+		NewPossessableActor = FMovieScenePossessable(NewActorLabel, Actor->GetClass());
+		NewGuid = NewPossessableActor.GetGuid();
+		OwnerSequence->BindPossessableObject(NewPossessableActor.GetGuid(), *Actor, InSequencerPtr->GetPlaybackContext());
+
+		// Defer replacing this object until the components have been updated
+	}
+
+	auto UpdateComponent = [&](FGuid OldComponentGuid, UActorComponent* NewComponent)
+	{
+		FMovieSceneSequenceIDRef FocusedGuid = InSequencerPtr->GetFocusedTemplateID();
+
+		// Get the object guid to assign, remove the binding if it already exists
+		FGuid NewComponentGuid = InSequencerPtr->FindObjectId(*NewComponent, FocusedGuid);
+		if (NewComponentGuid.IsValid())
+		{
+			OwnerMovieScene->RemovePossessable(NewComponentGuid);
+			OwnerSequence->UnbindPossessableObjects(NewComponentGuid);
+		}
+
+		// Add this object
+		FMovieScenePossessable NewPossessable(NewComponent->GetName(), NewComponent->GetClass());
+		OwnerSequence->BindPossessableObject(NewPossessable.GetGuid(), *NewComponent, Actor);
+
+		// Replace
+		OwnerMovieScene->ReplacePossessable(OldComponentGuid, NewPossessable);
+		OwnerSequence->UnbindPossessableObjects(OldComponentGuid);
+		InSequencerPtr->State.Invalidate(OldComponentGuid, FocusedGuid);
+		InSequencerPtr->State.Invalidate(NewPossessable.GetGuid(), FocusedGuid);
+
+		FMovieScenePossessable* ThisPossessable = OwnerMovieScene->FindPossessable(NewPossessable.GetGuid());
+		if (ensure(ThisPossessable))
+		{
+			ThisPossessable->SetParent(NewGuid);
+		}
+	};
+
+	// Handle components
+	AActor* ActorToReplace = Cast<AActor>(RuntimeObject);
+	if (ActorToReplace != nullptr && ActorToReplace->IsActorBeingDestroyed() == false)
+	{
+		for (UActorComponent* ComponentToReplace : ActorToReplace->GetComponents())
+		{
+			if (ComponentToReplace != nullptr)
+			{
+				FGuid ComponentGuid = InSequencerPtr->FindObjectId(*ComponentToReplace, InSequencerPtr->GetFocusedTemplateID());
+				if (ComponentGuid.IsValid())
+				{
+					bool bComponentWasUpdated = false;
+					for (UActorComponent* NewComponent : Actor->GetComponents())
+					{
+						if (NewComponent->GetFullName(Actor) == ComponentToReplace->GetFullName(ActorToReplace))
+						{
+							UpdateComponent(ComponentGuid, NewComponent);
+							bComponentWasUpdated = true;
+						}
+					}
+
+					// Clear the parent guid since this possessable component doesn't match to any component on the new actor
+					if (!bComponentWasUpdated)
+					{
+						FMovieScenePossessable* ThisPossessable = OwnerMovieScene->FindPossessable(ComponentGuid);
+						ThisPossessable->SetParent(FGuid());
+					}
+				}
+			}
+		}
+	}
+	else // If the actor didn't exist, try to find components who's parent guids were the previous actors guid.
+	{
+		TMap<FString, UActorComponent*> ComponentNameToComponent;
+		for (UActorComponent* Component : Actor->GetComponents())
+		{
+			ComponentNameToComponent.Add(Component->GetName(), Component);
+		}
+		for (int32 i = 0; i < OwnerMovieScene->GetPossessableCount(); i++)
+		{
+			FMovieScenePossessable& OldPossessable = OwnerMovieScene->GetPossessable(i);
+			if (OldPossessable.GetParent() == InObjectBinding)
+			{
+				UActorComponent** ComponentPtr = ComponentNameToComponent.Find(OldPossessable.GetName());
+				if (ComponentPtr != nullptr)
+				{
+					UpdateComponent(OldPossessable.GetGuid(), *ComponentPtr);
+				}
+			}
+		}
+	}
+
+	// Replace the actor itself after components have been updated
+	OwnerMovieScene->ReplacePossessable(InObjectBinding, NewPossessableActor);
+	OwnerSequence->UnbindPossessableObjects(InObjectBinding);
+
+	InSequencerPtr->State.Invalidate(InObjectBinding, InSequencerPtr->GetFocusedTemplateID());
+	InSequencerPtr->State.Invalidate(NewPossessableActor.GetGuid(), InSequencerPtr->GetFocusedTemplateID());
+
+	// Try to fix up folders
+	TArray<UMovieSceneFolder*> FoldersToCheck;
+	FoldersToCheck.Append(InSequencerPtr->GetFocusedMovieSceneSequence()->GetMovieScene()->GetRootFolders());
+	bool bFolderFound = false;
+	while (FoldersToCheck.Num() > 0 && bFolderFound == false)
+	{
+		UMovieSceneFolder* Folder = FoldersToCheck[0];
+		FoldersToCheck.RemoveAt(0);
+		if (Folder->GetChildObjectBindings().Contains(InObjectBinding))
+		{
+			Folder->RemoveChildObjectBinding(InObjectBinding);
+			Folder->AddChildObjectBinding(NewGuid);
+			bFolderFound = true;
+		}
+
+		for (UMovieSceneFolder* ChildFolder : Folder->GetChildFolders())
+		{
+			FoldersToCheck.Add(ChildFolder);
+		}
+	}
+
+	InSequencerPtr->RestorePreAnimatedState();
+
+	InSequencerPtr->NotifyMovieSceneDataChanged(EMovieSceneDataChangeType::MovieSceneStructureItemsChanged);
+
+	return NewGuid;
+}
+
+void FSequencerUtilities::UpdateBindingIDs(ISequencer* InSequencerPtr, UMovieSceneCompiledDataManager* InCompiledDataManagerPtr, FGuid OldGuid, FGuid NewGuid)
+{
+	FMovieSceneSequenceIDRef FocusedGuid = InSequencerPtr->GetFocusedTemplateID();
+
+	TMap<UE::MovieScene::FFixedObjectBindingID, UE::MovieScene::FFixedObjectBindingID> OldFixedToNewFixedMap;
+	OldFixedToNewFixedMap.Add(UE::MovieScene::FFixedObjectBindingID(OldGuid, FocusedGuid), UE::MovieScene::FFixedObjectBindingID(NewGuid, FocusedGuid));
+
+	const FMovieSceneSequenceHierarchy* Hierarchy = InCompiledDataManagerPtr->FindHierarchy(InSequencerPtr->GetEvaluationTemplate().GetCompiledDataID());
+
+	if (UMovieScene* MovieScene = InSequencerPtr->GetRootMovieSceneSequence()->GetMovieScene())
+	{
+		for (UMovieSceneSection* Section : MovieScene->GetAllSections())
+		{
+			if (Section)
+			{
+				Section->OnBindingIDsUpdated(OldFixedToNewFixedMap, InSequencerPtr->GetRootTemplateID(), Hierarchy, *InSequencerPtr);
+			}
+		}
+	}
+
+	if (Hierarchy)
+	{
+		for (const TTuple<FMovieSceneSequenceID, FMovieSceneSubSequenceData>& Pair : Hierarchy->AllSubSequenceData())
+		{
+			if (UMovieSceneSequence* Sequence = Pair.Value.GetSequence())
+			{
+				if (UMovieScene* MovieScene = Sequence->GetMovieScene())
+				{
+					for (UMovieSceneSection* Section : MovieScene->GetAllSections())
+					{
+						if (Section)
+						{
+							Section->OnBindingIDsUpdated(OldFixedToNewFixedMap, Pair.Key, Hierarchy, *InSequencerPtr);
+						}
+					}
+				}
+			}
+		}
+	}
+}
+
+void FSequencerUtilities::ShowReadOnlyError()
+{
+	FNotificationInfo Info(LOCTEXT("SequenceReadOnly", "Sequence is read only."));
+	Info.ExpireDuration = 5.0f;
+	FSlateNotificationManager::Get().AddNotification(Info)->SetCompletionState(SNotificationItem::CS_Fail);
+}
+
 
 #undef LOCTEXT_NAMESPACE
diff --git a/Engine/Source/Editor/Sequencer/Public/ISequencerModule.h b/Engine/Source/Editor/Sequencer/Public/ISequencerModule.h
--- a/Engine/Source/Editor/Sequencer/Public/ISequencerModule.h
+++ b/Engine/Source/Editor/Sequencer/Public/ISequencerModule.h
@@ -280,6 +280,13 @@ public:
 	 */
 	virtual TSharedPtr<FExtensibilityManager> GetToolBarExtensibilityManager() const = 0;
 
+	/**
+	 * Get the extensibility manager for toolbars.
+	 *
+	 * @return Toolbar extensibility manager.
+	 */
+	virtual TSharedPtr<FExtensibilityManager> GetActionsMenuExtensibilityManager() const = 0;
+
 	/**
 	 * Get the sequencer customization manager, which handles editor customizations applied based on
 	 * the currently focused sequence type and other dynamic criteria.
diff --git a/Engine/Source/Editor/Sequencer/Public/SequencerCommands.h b/Engine/Source/Editor/Sequencer/Public/SequencerCommands.h
--- a/Engine/Source/Editor/Sequencer/Public/SequencerCommands.h
+++ b/Engine/Source/Editor/Sequencer/Public/SequencerCommands.h
@@ -382,9 +382,6 @@ public:
 	/** Restores all animated state for the current sequence. */
 	TSharedPtr< FUICommandInfo > RestoreAnimatedState;
 
-	/** Attempts to fix broken actor references. */
-	TSharedPtr< FUICommandInfo > FixActorReferences;
-
 	/** Rebinds all possessable references with their current bindings. */
 	TSharedPtr< FUICommandInfo > RebindPossessableReferences;
 
diff --git a/Engine/Source/Editor/Sequencer/Public/SequencerUtilities.h b/Engine/Source/Editor/Sequencer/Public/SequencerUtilities.h
--- a/Engine/Source/Editor/Sequencer/Public/SequencerUtilities.h
+++ b/Engine/Source/Editor/Sequencer/Public/SequencerUtilities.h
@@ -12,6 +12,7 @@ class UMovieSceneSection;
 class ISequencer;
 class FMenuBuilder;
 class ULevelSequence;
+class UMovieSceneCompiledDataManager;
 
 struct SEQUENCER_API FSequencerUtilities
 {
@@ -30,4 +31,10 @@ struct SEQUENCER_API FSequencerUtilities
 	 * The name is made unique by appending a number to the end.
 	 */
 	static FName GetUniqueName(FName CandidateName, const TArray<FName>& ExistingNames);
+
+	static FGuid DoAssignActor(ISequencer* InSequencerPtr, AActor* const* InActors, int32 NumActors, FGuid InObjectBinding);
+
+	static void UpdateBindingIDs(ISequencer* InSequencerPtr, UMovieSceneCompiledDataManager* InCompiledDataManagerPtr, FGuid OldGuid, FGuid NewGuid);
+
+	static void ShowReadOnlyError();
 };
