Sequencer: Control Rig: When selecting controls we refresh the tree on next tick. Also when get key area's we don't get them if from track or object since that messes with control rig selection.

diff --git a/Engine/Plugins/Animation/ControlRig/Source/ControlRigEditor/Private/Sequencer/ControlRigParameterTrackEditor.cpp b/Engine/Plugins/Animation/ControlRig/Source/ControlRigEditor/Private/Sequencer/ControlRigParameterTrackEditor.cpp
--- a/Engine/Plugins/Animation/ControlRig/Source/ControlRigEditor/Private/Sequencer/ControlRigParameterTrackEditor.cpp
+++ b/Engine/Plugins/Animation/ControlRig/Source/ControlRigEditor/Private/Sequencer/ControlRigParameterTrackEditor.cpp
@@ -2287,9 +2287,8 @@ void FControlRigParameterTrackEditor::HandleControlSelected(UControlRig* Subject
 			GetSequencer()->GetSelectedKeyAreas(KeyAreas, UseSelectedKeys);
 			SelectRigsAndControls(Subject, KeyAreas);
 
-			//Force refresh now, not later
-			GetSequencer()->RefreshTree();
-
+			//Force refresh later, not now
+			GetSequencer()->NotifyMovieSceneDataChanged(EMovieSceneDataChangeType::RefreshTree);				
 		}
 	}
 }
diff --git a/Engine/Source/Editor/Sequencer/Private/Sequencer.cpp b/Engine/Source/Editor/Sequencer/Private/Sequencer.cpp
--- a/Engine/Source/Editor/Sequencer/Private/Sequencer.cpp
+++ b/Engine/Source/Editor/Sequencer/Private/Sequencer.cpp
@@ -2634,7 +2634,11 @@ void FSequencer::NotifyMovieSceneDataChanged( EMovieSceneDataChangeType DataChan
 		}
 	}
 
-	if ( DataChangeType == EMovieSceneDataChangeType::MovieSceneStructureItemRemoved ||
+	if (DataChangeType == EMovieSceneDataChangeType::RefreshTree)
+	{
+		bNeedTreeRefresh = true;
+	}
+	else if ( DataChangeType == EMovieSceneDataChangeType::MovieSceneStructureItemRemoved ||
 		DataChangeType == EMovieSceneDataChangeType::MovieSceneStructureItemsChanged ||
 		DataChangeType == EMovieSceneDataChangeType::Unknown )
 	{
@@ -7333,15 +7337,27 @@ void FSequencer::GetSelectedKeyAreas(TArray<const IKeyArea*>& OutSelectedKeyArea
 {
 	TSet<TSharedRef<FSequencerDisplayNode>> NodesToKey = Selection.GetSelectedOutlinerNodes();
 	{
+
 		TSet<TSharedRef<FSequencerDisplayNode>> ChildNodes;
+		int32 Index = 0;
 		for (TSharedRef<FSequencerDisplayNode> Node : NodesToKey.Array())
 		{
-			ChildNodes.Reset();
-			SequencerHelpers::GetDescendantNodes(Node, ChildNodes);
-
-			for (TSharedRef<FSequencerDisplayNode> ChildNode : ChildNodes)
+			// No need to gather key areas from binding/tracks because they have no key areas
+			if (Node->GetType()== ESequencerNode::Object || Node->GetType() == ESequencerNode::Track)
+			{
+				NodesToKey.Array().RemoveAtSwap(Index);
+				++Index;
+			}
+			else
 			{
-				NodesToKey.Remove(ChildNode);
+				ChildNodes.Reset();
+				SequencerHelpers::GetDescendantNodes(Node, ChildNodes);
+
+				for (TSharedRef<FSequencerDisplayNode> ChildNode : ChildNodes)
+				{
+					NodesToKey.Remove(ChildNode);
+				}
+				++Index;
 			}
 		}
 	}
diff --git a/Engine/Source/Editor/Sequencer/Public/ISequencer.h b/Engine/Source/Editor/Sequencer/Public/ISequencer.h
--- a/Engine/Source/Editor/Sequencer/Public/ISequencer.h
+++ b/Engine/Source/Editor/Sequencer/Public/ISequencer.h
@@ -149,7 +149,10 @@ enum class EMovieSceneDataChangeType
 	/** Rebuild and evaluate everything immediately. */
 	RefreshAllImmediately,
 	/** It's not known what data has changed. */
-	Unknown
+	Unknown,
+	/** Refresh Tree on Next Tick */
+	RefreshTree
+
 };
 
 /**
@@ -513,7 +516,6 @@ public:
 
 	/** Refresh the sequencer tree view */
 	virtual void RefreshTree() = 0;
-
 protected:
 	virtual void NotifyMovieSceneDataChangedInternal() = 0;
 
