Sequencer: Output error when the possessed object class is different from the class of the bound object.

Added a "Fix Possessed Object Class" menu item that operates on the focused movie scene.

diff --git a/Engine/Source/Editor/Sequencer/Private/SSequencer.cpp b/Engine/Source/Editor/Sequencer/Private/SSequencer.cpp
--- a/Engine/Source/Editor/Sequencer/Private/SSequencer.cpp
+++ b/Engine/Source/Editor/Sequencer/Private/SSequencer.cpp
@@ -2047,6 +2047,7 @@ void SSequencer::FillAdvancedMenu(FMenuBuilder& MenuBuilder)
 	{
 		MenuBuilder.BeginSection("Bindings", LOCTEXT("BindingsMenuHeader", "Bindings"));
 
+		MenuBuilder.AddMenuEntry(FSequencerCommands::Get().FixPossessableObjectClass);
 		MenuBuilder.AddMenuEntry(FSequencerCommands::Get().RebindPossessableReferences);
 
 		MenuBuilder.EndSection();
diff --git a/Engine/Source/Editor/Sequencer/Private/Sequencer.cpp b/Engine/Source/Editor/Sequencer/Private/Sequencer.cpp
--- a/Engine/Source/Editor/Sequencer/Private/Sequencer.cpp
+++ b/Engine/Source/Editor/Sequencer/Private/Sequencer.cpp
@@ -12253,6 +12253,68 @@ void FSequencer::NewCameraAdded(ACameraActor* NewCamera, FGuid CameraGuid)
 	MovieSceneToolHelpers::CreateCameraCutSectionForCamera(GetFocusedMovieSceneSequence()->GetMovieScene(), CameraGuid, GetLocalTime().Time.FloorToFrame());
 }
 
+
+void FSequencer::FixPossessableObjectClassInternal(UMovieSceneSequence* Sequence, FMovieSceneSequenceIDRef SequenceID)
+{
+	UMovieScene* MovieScene = Sequence->GetMovieScene();
+
+	MovieScene->Modify();
+
+	for (int32 Index = 0; Index < MovieScene->GetPossessableCount(); ++Index)
+	{
+		FMovieScenePossessable& Possessable = MovieScene->GetPossessable(Index);
+
+		TOptional<UClass*> CommonBaseClass;
+
+		for (TWeakObjectPtr<> WeakObject : FindBoundObjects(Possessable.GetGuid(), SequenceID))
+		{
+			if (WeakObject.IsValid())
+			{
+				if (!CommonBaseClass.IsSet())
+				{
+					CommonBaseClass = WeakObject->GetClass();
+				}
+				else
+				{
+					CommonBaseClass = UClass::FindCommonBase(WeakObject->GetClass(), CommonBaseClass.GetValue());
+				}
+			}
+		}
+
+		if (CommonBaseClass.IsSet() && CommonBaseClass.GetValue() != Possessable.GetPossessedObjectClass())
+		{
+			UE_LOG(LogSequencer, Display, TEXT("Updated possessed object class in: %s for: %s, from: %s, to: %s"), *GetNameSafe(Sequence), *Possessable.GetName(), *GetNameSafe(Possessable.GetPossessedObjectClass()), *GetNameSafe(CommonBaseClass.GetValue()));
+			Possessable.SetPossessedObjectClass(CommonBaseClass.GetValue());
+		}
+	}
+}
+
+void FSequencer::FixPossessableObjectClass()
+{
+	FScopedTransaction FixPossessableObjectClassTransaction( NSLOCTEXT( "Sequencer", "FixPossessableObjectClass", "Fix Possessable Object Class" ) );
+
+	FMovieSceneRootEvaluationTemplateInstance& RootTemplate = GetEvaluationTemplate();
+		
+	UMovieSceneSequence* Sequence = RootTemplate.GetSequence(MovieSceneSequenceID::Root);
+
+	FixPossessableObjectClassInternal(Sequence, MovieSceneSequenceID::Root);
+
+	const FMovieSceneSequenceHierarchy* Hierarchy = CompiledDataManager->FindHierarchy(RootTemplateInstance.GetCompiledDataID());
+	if (Hierarchy)
+	{
+		FMovieSceneEvaluationTreeRangeIterator Iter = Hierarchy->GetTree().IterateFromTime(PlayPosition.GetCurrentPosition().FrameNumber);
+
+		for (const FMovieSceneSubSequenceTreeEntry& Entry : Hierarchy->GetTree().GetAllData(Iter.Node()))
+		{
+			UMovieSceneSequence* SubSequence = Hierarchy->FindSubSequence(Entry.SequenceID);
+			if (SubSequence)
+			{
+				FixPossessableObjectClassInternal(SubSequence, Entry.SequenceID);
+			}
+		}
+	}
+}
+
 void FSequencer::RebindPossessableReferences()
 {
 	UMovieSceneSequence* FocusedSequence = GetFocusedMovieSceneSequence();
@@ -13234,6 +13296,11 @@ void FSequencer::BindCommands()
 		FExecuteAction::CreateSP( this, &FSequencer::SyncSectionsUsingSourceTimecode ),
 		FCanExecuteAction::CreateLambda( [this]{ return (GetSelection().GetSelectedSections().Num() > 1); } ) );
 
+	SequencerCommandBindings->MapAction(
+		Commands.FixPossessableObjectClass,
+		FExecuteAction::CreateSP( this, &FSequencer::FixPossessableObjectClass ),
+		FCanExecuteAction::CreateLambda( []{ return true; } ) );
+
 	SequencerCommandBindings->MapAction(
 		Commands.RebindPossessableReferences,
 		FExecuteAction::CreateSP( this, &FSequencer::RebindPossessableReferences ),
diff --git a/Engine/Source/Editor/Sequencer/Private/Sequencer.h b/Engine/Source/Editor/Sequencer/Private/Sequencer.h
--- a/Engine/Source/Editor/Sequencer/Private/Sequencer.h
+++ b/Engine/Source/Editor/Sequencer/Private/Sequencer.h
@@ -661,6 +661,10 @@ public:
 	/** Called when a new camera is added. Locks the viewport to the NewCamera. */
 	void NewCameraAdded(ACameraActor* NewCamera, FGuid CameraGuid);
 
+	/** Attempts to automatically fix up possessables whose object class don't match the object class of their currently bound objects */
+	void FixPossessableObjectClass();
+	void FixPossessableObjectClassInternal(UMovieSceneSequence* Sequence, FMovieSceneSequenceIDRef SequenceID);
+
 	/** Rebinds all possessable references in the current sequence to update them to the latest referencing mechanism. */
 	void RebindPossessableReferences();
 
diff --git a/Engine/Source/Editor/Sequencer/Private/SequencerCommands.cpp b/Engine/Source/Editor/Sequencer/Private/SequencerCommands.cpp
--- a/Engine/Source/Editor/Sequencer/Private/SequencerCommands.cpp
+++ b/Engine/Source/Editor/Sequencer/Private/SequencerCommands.cpp
@@ -152,6 +152,7 @@ void FSequencerCommands::RegisterCommands()
 	UI_COMMAND( SaveCurrentSpawnableState, "Save Default State", "Save the current state of this spawnable object as its default properties.", EUserInterfaceActionType::Button, FInputChord() );
 	UI_COMMAND( RestoreAnimatedState, "Restore Pre-Animated State", "Restore any objects that have been animated by sequencer back to their original state.", EUserInterfaceActionType::Button, FInputChord(EModifierKey::Control, EKeys::R) );
 
+	UI_COMMAND( FixPossessableObjectClass, "Fix Possessable Object Class", "Try to automatically fix up possessables whose object class don't match the object class of their currently bound objects.", EUserInterfaceActionType::Button, FInputChord() );
 	UI_COMMAND( RebindPossessableReferences, "Rebind Possesable References", "Rebinds all possessables in the current sequence to ensure they're using the most robust referencing mechanism.", EUserInterfaceActionType::Button, FInputChord() );
 
 	UI_COMMAND( ImportFBX, "Import...", "Import the animation from an FBX file.", EUserInterfaceActionType::Button, FInputChord() );
diff --git a/Engine/Source/Editor/Sequencer/Public/SequencerCommands.h b/Engine/Source/Editor/Sequencer/Public/SequencerCommands.h
--- a/Engine/Source/Editor/Sequencer/Public/SequencerCommands.h
+++ b/Engine/Source/Editor/Sequencer/Public/SequencerCommands.h
@@ -382,6 +382,9 @@ public:
 	/** Restores all animated state for the current sequence. */
 	TSharedPtr< FUICommandInfo > RestoreAnimatedState;
 
+	/** Attempts to fix possessables whose object class don't match the object class of their currently bound objects. */
+	TSharedPtr< FUICommandInfo > FixPossessableObjectClass;
+
 	/** Rebinds all possessable references with their current bindings. */
 	TSharedPtr< FUICommandInfo > RebindPossessableReferences;
 
diff --git a/Engine/Source/Runtime/MovieScene/Private/Evaluation/MovieSceneEvaluationState.cpp b/Engine/Source/Runtime/MovieScene/Private/Evaluation/MovieSceneEvaluationState.cpp
--- a/Engine/Source/Runtime/MovieScene/Private/Evaluation/MovieSceneEvaluationState.cpp
+++ b/Engine/Source/Runtime/MovieScene/Private/Evaluation/MovieSceneEvaluationState.cpp
@@ -351,6 +351,11 @@ void FMovieSceneObjectCache::UpdateBindings(const FGuid& InGuid, IMovieScenePlay
 					for (UObject* Object : FoundObjects)
 					{
 						Bindings->Objects.Add(Object);
+
+						if (!Object->GetClass()->IsChildOf(Possessable->GetPossessedObjectClass()))
+						{
+							UE_LOG(LogMovieScene, Warning, TEXT("Mismatch in %s with posssessed object class for: %s, Expected: %s, but was: %s"), *GetNameSafe(Sequence), *Possessable->GetName(), *GetNameSafe(Possessable->GetPossessedObjectClass()), *GetNameSafe(Object->GetClass()));
+						}
 					}
 				}
 			}
@@ -362,6 +367,11 @@ void FMovieSceneObjectCache::UpdateBindings(const FGuid& InGuid, IMovieScenePlay
 				for (UObject* Object : FoundObjects)
 				{
 					Bindings->Objects.Add(Object);
+
+					if (!Object->GetClass()->IsChildOf(Possessable->GetPossessedObjectClass()))
+					{
+						UE_LOG(LogMovieScene, Warning, TEXT("Mismatch in %s with posssessed object class for: %s, Expected: %s, but was: %s"), *GetNameSafe(Sequence), *Possessable->GetName(), *GetNameSafe(Possessable->GetPossessedObjectClass()), *GetNameSafe(Object->GetClass()));
+					}
 				}
 			}
 		}
