Sequencer: Follow up to infinite section selection (18193503), so that they draw selected

diff --git a/Engine/Source/Editor/Sequencer/Private/SSequencerSection.cpp b/Engine/Source/Editor/Sequencer/Private/SSequencerSection.cpp
--- a/Engine/Source/Editor/Sequencer/Private/SSequencerSection.cpp
+++ b/Engine/Source/Editor/Sequencer/Private/SSequencerSection.cpp
@@ -241,12 +241,6 @@ struct FSequencerSectionPainterImpl : FSequencerSectionPainter
 
 	void CalculateSelectionColor()
 	{
-		// Don't draw selected if infinite
-		if (Section.GetRange() == TRange<FFrameNumber>::All())
-		{
-			return;
-		}
-
 		FSequencerSelection& Selection = Sequencer.GetSelection();
 		FSequencerSelectionPreview& SelectionPreview = Sequencer.GetSelectionPreview();
 
