Sequencer: Add option to snap play time to section bounds. Follow up to CL 17818309

diff --git a/Engine/Source/Editor/Sequencer/Private/SSequencer.cpp b/Engine/Source/Editor/Sequencer/Private/SSequencer.cpp
--- a/Engine/Source/Editor/Sequencer/Private/SSequencer.cpp
+++ b/Engine/Source/Editor/Sequencer/Private/SSequencer.cpp
@@ -2686,6 +2686,7 @@ TSharedRef<SWidget> SSequencer::MakeSnapMenu()
 	{
 		MenuBuilder.AddMenuEntry( FSequencerCommands::Get().ToggleSnapPlayTimeToInterval );
 		MenuBuilder.AddMenuEntry( FSequencerCommands::Get().ToggleSnapPlayTimeToKeys );
+		MenuBuilder.AddMenuEntry( FSequencerCommands::Get().ToggleSnapPlayTimeToSections );
 		MenuBuilder.AddMenuEntry( FSequencerCommands::Get().ToggleSnapPlayTimeToMarkers );
 		MenuBuilder.AddMenuEntry( FSequencerCommands::Get().ToggleSnapPlayTimeToPressedKey );
 		MenuBuilder.AddMenuEntry( FSequencerCommands::Get().ToggleSnapPlayTimeToDraggedKey );
diff --git a/Engine/Source/Editor/Sequencer/Private/Sequencer.cpp b/Engine/Source/Editor/Sequencer/Private/Sequencer.cpp
--- a/Engine/Source/Editor/Sequencer/Private/Sequencer.cpp
+++ b/Engine/Source/Editor/Sequencer/Private/Sequencer.cpp
@@ -3341,6 +3341,11 @@ void FSequencer::SetLocalTime( FFrameTime NewTime, ESnapTimeMode SnapTimeMode)
 			EnumAddFlags(NearestKeyOption, ENearestKeyOption::NKO_SearchKeys);
 		}
 
+		if (Settings->GetSnapPlayTimeToSections() || FSlateApplication::Get().GetModifierKeys().IsShiftDown())
+		{
+			EnumAddFlags(NearestKeyOption, ENearestKeyOption::NKO_SearchSections);
+		}
+
 		if (Settings->GetSnapPlayTimeToMarkers() || FSlateApplication::Get().GetModifierKeys().IsShiftDown())
 		{
 			EnumAddFlags(NearestKeyOption, ENearestKeyOption::NKO_SearchMarkers);
@@ -5345,7 +5350,7 @@ void FSequencer::OnClampRangeChanged( TRange<double> NewClampRange )
 
 FFrameNumber FSequencer::OnGetNearestKey(FFrameTime InTime, ENearestKeyOption NearestKeyOption)
 {
-	FFrameNumber NearestKeyTime = InTime.FloorToFrame();
+	const FFrameNumber CurrentTime = InTime.FloorToFrame();
 
 	if (EnumHasAnyFlags(NearestKeyOption, ENearestKeyOption::NKO_SearchAllTracks))
 	{
@@ -5356,18 +5361,22 @@ FFrameNumber FSequencer::OnGetNearestKey(FFrameTime InTime, ENearestKeyOption Ne
 		GetKeysFromSelection(SelectedKeyCollection, SMALL_NUMBER);
 	}
 
+	TOptional<FFrameNumber> NearestTime;
+
 	if (EnumHasAnyFlags(NearestKeyOption, ENearestKeyOption::NKO_SearchKeys) && SelectedKeyCollection.IsValid())
 	{
-		TRange<FFrameNumber> FindRangeBackwards(TRangeBound<FFrameNumber>::Open(), NearestKeyTime);
+		TOptional<FFrameNumber> NearestKeyTime;
+
+		TRange<FFrameNumber> FindRangeBackwards(TRangeBound<FFrameNumber>::Open(), CurrentTime);
 		TOptional<FFrameNumber> NewTimeBackwards = SelectedKeyCollection->FindFirstKeyInRange(FindRangeBackwards, EFindKeyDirection::Backwards);
 
-		TRange<FFrameNumber> FindRangeForwards(NearestKeyTime, TRangeBound<FFrameNumber>::Open());
+		TRange<FFrameNumber> FindRangeForwards(CurrentTime, TRangeBound<FFrameNumber>::Open());
 		TOptional<FFrameNumber> NewTimeForwards = SelectedKeyCollection->FindFirstKeyInRange(FindRangeForwards, EFindKeyDirection::Forwards);
 		if (NewTimeForwards.IsSet())
 		{
 			if (NewTimeBackwards.IsSet())
 			{
-				if (FMath::Abs(NewTimeForwards.GetValue() - NearestKeyTime) < FMath::Abs(NewTimeBackwards.GetValue() - NearestKeyTime))
+				if (FMath::Abs(NewTimeForwards.GetValue() - CurrentTime) < FMath::Abs(NewTimeBackwards.GetValue() - CurrentTime))
 				{
 					NearestKeyTime = NewTimeForwards.GetValue();
 				}
@@ -5385,38 +5394,90 @@ FFrameNumber FSequencer::OnGetNearestKey(FFrameTime InTime, ENearestKeyOption Ne
 		{
 			NearestKeyTime = NewTimeBackwards.GetValue();
 		}
+
+		if (NearestKeyTime.IsSet())
+		{
+			NearestTime = NearestKeyTime.GetValue();
+		}
 	}
 
-	if (EnumHasAnyFlags(NearestKeyOption, ENearestKeyOption::NKO_SearchMarkers))
+	if (EnumHasAnyFlags(NearestKeyOption, ENearestKeyOption::NKO_SearchSections) && SelectedKeyCollection.IsValid())
 	{
-		FFrameNumber MinDiff = FMath::Abs(NearestKeyTime - InTime.FloorToFrame());
+		TOptional<FFrameNumber> NearestSectionTime;
+
+		TRange<FFrameNumber> FindRangeBackwards(TRangeBound<FFrameNumber>::Open(), CurrentTime);
+		TOptional<FFrameNumber> NewTimeBackwards = SelectedKeyCollection->FindFirstSectionKeyInRange(FindRangeBackwards, EFindKeyDirection::Backwards);
+
+		TRange<FFrameNumber> FindRangeForwards(CurrentTime, TRangeBound<FFrameNumber>::Open());
+		TOptional<FFrameNumber> NewTimeForwards = SelectedKeyCollection->FindFirstSectionKeyInRange(FindRangeForwards, EFindKeyDirection::Forwards);
+		if (NewTimeForwards.IsSet())
+		{
+			if (NewTimeBackwards.IsSet())
+			{
+				if (FMath::Abs(NewTimeForwards.GetValue() - CurrentTime) < FMath::Abs(NewTimeBackwards.GetValue() - CurrentTime))
+				{
+					NearestSectionTime = NewTimeForwards.GetValue();
+				}
+				else
+				{
+					NearestSectionTime = NewTimeBackwards.GetValue();
+				}
+			}
+			else
+			{
+				NearestSectionTime = NewTimeForwards.GetValue();
+			}
+		}
+		else if (NewTimeBackwards.IsSet())
+		{
+			NearestSectionTime = NewTimeBackwards.GetValue();
+		}
 
+		if (NearestSectionTime.IsSet())
+		{
+			if (!NearestTime.IsSet())
+			{
+				NearestTime = NearestSectionTime.GetValue();
+			}
+			else if (FMath::Abs(NearestSectionTime.GetValue() - CurrentTime) < FMath::Abs(NearestTime.GetValue() - CurrentTime))
+			{
+				NearestTime = NearestSectionTime.GetValue();
+			}
+		}
+	}
+	
+	if (EnumHasAnyFlags(NearestKeyOption, ENearestKeyOption::NKO_SearchMarkers))
+	{
 		TArray<FMovieSceneMarkedFrame> MarkedFrames = GetMarkedFrames();
 		for (const FMovieSceneMarkedFrame& MarkedFrame : MarkedFrames)
 		{
-			FFrameNumber Diff = FMath::Abs(MarkedFrame.FrameNumber - InTime.FloorToFrame());
-
-			if (Diff < MinDiff)
+			if (!NearestTime.IsSet())
 			{
-				MinDiff = Diff;
-				NearestKeyTime = MarkedFrame.FrameNumber;
+				NearestTime = MarkedFrame.FrameNumber;
+			}
+			else if (FMath::Abs(MarkedFrame.FrameNumber - CurrentTime) < FMath::Abs(NearestTime.GetValue() - CurrentTime))
+			{
+				NearestTime = MarkedFrame.FrameNumber;
 			}
 		}
 
 		TArray<FMovieSceneMarkedFrame> GlobalMarkedFrames = GetGlobalMarkedFrames();
 		for (const FMovieSceneMarkedFrame& GlobalMarkedFrame : GlobalMarkedFrames)
 		{
-			FFrameNumber Diff = FMath::Abs(GlobalMarkedFrame.FrameNumber - InTime.FloorToFrame());
+			FFrameNumber Diff = FMath::Abs(GlobalMarkedFrame.FrameNumber - CurrentTime);
 
-			if (Diff < MinDiff)
+			if (!NearestTime.IsSet())
 			{
-				MinDiff = Diff;
-				NearestKeyTime = GlobalMarkedFrame.FrameNumber;
+				NearestTime = GlobalMarkedFrame.FrameNumber;
 			}
-		}		
+			else if (FMath::Abs(GlobalMarkedFrame.FrameNumber - CurrentTime) < FMath::Abs(NearestTime.GetValue() - CurrentTime))
+			{
+				NearestTime = GlobalMarkedFrame.FrameNumber;
+			}
+		}
 	}
 
-	return NearestKeyTime;
+	return NearestTime.IsSet() ? NearestTime.GetValue() : CurrentTime;
 }
 
 void FSequencer::OnScrubPositionChanged( FFrameTime NewScrubPosition, bool bScrubbing )
@@ -13067,6 +13128,12 @@ void FSequencer::BindCommands()
 		FCanExecuteAction::CreateLambda( []{ return true; } ),
 		FIsActionChecked::CreateLambda( [this]{ return Settings->GetSnapPlayTimeToKeys(); } ) );
 
+	SequencerCommandBindings->MapAction(
+		Commands.ToggleSnapPlayTimeToSections,
+		FExecuteAction::CreateLambda( [this]{ Settings->SetSnapPlayTimeToSections( !Settings->GetSnapPlayTimeToSections() ); } ),
+		FCanExecuteAction::CreateLambda( []{ return true; } ),
+		FIsActionChecked::CreateLambda( [this]{ return Settings->GetSnapPlayTimeToSections(); } ) );
+
 	SequencerCommandBindings->MapAction(
 		Commands.ToggleSnapPlayTimeToMarkers,
 		FExecuteAction::CreateLambda( [this]{ Settings->SetSnapPlayTimeToMarkers( !Settings->GetSnapPlayTimeToMarkers() ); } ),
diff --git a/Engine/Source/Editor/Sequencer/Private/SequencerCommands.cpp b/Engine/Source/Editor/Sequencer/Private/SequencerCommands.cpp
--- a/Engine/Source/Editor/Sequencer/Private/SequencerCommands.cpp
+++ b/Engine/Source/Editor/Sequencer/Private/SequencerCommands.cpp
@@ -124,6 +124,7 @@ void FSequencerCommands::RegisterCommands()
 	UI_COMMAND( ToggleSnapKeysAndSectionsToPlayRange, "Snap Keys and Sections to the Playback Range", "When checked, keys and sections will be snapped to the playback range bounds", EUserInterfaceActionType::ToggleButton, FInputChord());
 
 	UI_COMMAND( ToggleSnapPlayTimeToKeys, "Snap to Keys While Scrubbing", "Snap the current time to keys of the selected track while scrubbing", EUserInterfaceActionType::ToggleButton, FInputChord() );
+	UI_COMMAND( ToggleSnapPlayTimeToSections, "Snap to Sections While Scrubbing", "Snap the current time to section bounds while scrubbing", EUserInterfaceActionType::ToggleButton, FInputChord() );
 	UI_COMMAND( ToggleSnapPlayTimeToMarkers, "Snap to Markers While Scrubbing", "Snap the current time to markers while scrubbing", EUserInterfaceActionType::ToggleButton, FInputChord() );
 	UI_COMMAND( ToggleSnapPlayTimeToInterval, "Snap to the Interval While Scrubbing", "Snap the current time to the time snapping interval while scrubbing", EUserInterfaceActionType::ToggleButton, FInputChord() );
 	UI_COMMAND( ToggleSnapPlayTimeToPressedKey, "Snap to the Pressed Key", "Snap the current time to the pressed key", EUserInterfaceActionType::ToggleButton, FInputChord() );
diff --git a/Engine/Source/Editor/Sequencer/Private/SequencerKeyCollection.cpp b/Engine/Source/Editor/Sequencer/Private/SequencerKeyCollection.cpp
--- a/Engine/Source/Editor/Sequencer/Private/SequencerKeyCollection.cpp
+++ b/Engine/Source/Editor/Sequencer/Private/SequencerKeyCollection.cpp
@@ -157,17 +157,29 @@ bool FSequencerKeyCollection::Update(const FSequencerKeyCollectionSignature& InS
 	}
 
 	TArray<FFrameNumber> AllTimes;
+	TArray<FFrameNumber> AllSectionTimes;
 
 	// Get all the key times for the key areas
 	for (auto& Pair : InSignature.GetKeyAreas())
 	{
-		if (Pair.Key->GetOwningSection())
+		if (UMovieSceneSection* Section = Pair.Key->GetOwningSection())
 		{
-			Pair.Key->GetKeyTimes(AllTimes, Pair.Key->GetOwningSection()->GetRange());
+			Pair.Key->GetKeyTimes(AllTimes, Section->GetRange());
+
+			if (Section->HasStartFrame())
+			{
+				AllSectionTimes.Add(Section->GetInclusiveStartFrame());
+			}
+			
+			if (Section->HasEndFrame())
+			{
+				AllSectionTimes.Add(Section->GetExclusiveEndFrame());
+			}
 		}
 	}
 
 	AllTimes.Sort();
+	AllSectionTimes.Sort();
 
 	GroupedTimes.Reset(AllTimes.Num());
 	int32 Index = 0;
@@ -182,118 +194,107 @@ bool FSequencerKeyCollection::Update(const FSequencerKeyCollectionSignature& InS
 	}
 	GroupedTimes.Shrink();
 
+	GroupedSectionTimes.Reset(AllSectionTimes.Num());
+	Index = 0;
+	while (Index < AllSectionTimes.Num())
+	{
+		FFrameNumber PredicateTime = AllSectionTimes[Index];
+		GroupedSectionTimes.Add(PredicateTime);
+		while (Index < AllSectionTimes.Num() && FMath::Abs(AllSectionTimes[Index] - PredicateTime) <= InSignature.GetDuplicateThreshold())
+		{
+			++Index;
+		}
+	}
+	GroupedSectionTimes.Shrink();
+
 	Signature = InSignature;
 
 	return true;
 }
 
-TOptional<FFrameNumber> FSequencerKeyCollection::FindFirstKeyInRange(const TRange<FFrameNumber>& Range, EFindKeyDirection Direction) const
-{
-	TArrayView<const FFrameNumber> KeysInRange = GetKeysInRange(Range);
-	if (KeysInRange.Num())
-	{
-		return Direction == EFindKeyDirection::Forwards ? KeysInRange[0] : KeysInRange[KeysInRange.Num()-1];
-	}
-	return TOptional<FFrameNumber>();
-}
-
-TArrayView<const FFrameNumber> FSequencerKeyCollection::GetKeysInRange(const TRange<FFrameNumber>& Range) const
+TArrayView<const FFrameNumber> GetKeysInRangeInternal(const TArray<FFrameNumber>&Times, const TRange<FFrameNumber>& Range)
 {
 	// Binary search the first time that's >= the lower bound
-	int32 FirstVisibleIndex = Range.GetLowerBound().IsClosed() ? Algo::LowerBound(GroupedTimes, Range.GetLowerBoundValue()) : 0;
+	int32 FirstVisibleIndex = Range.GetLowerBound().IsClosed() ? Algo::LowerBound(Times, Range.GetLowerBoundValue()) : 0;
 	// Binary search the last time that's > the upper bound
-	int32 LastVisibleIndex  = Range.GetUpperBound().IsClosed() ? Algo::UpperBound(GroupedTimes, Range.GetUpperBoundValue()) : GroupedTimes.Num();
+	int32 LastVisibleIndex = Range.GetUpperBound().IsClosed() ? Algo::UpperBound(Times, Range.GetUpperBoundValue()) : Times.Num();
 
 	int32 Num = LastVisibleIndex - FirstVisibleIndex;
-	if (GroupedTimes.IsValidIndex(FirstVisibleIndex) && LastVisibleIndex <= GroupedTimes.Num() && Num > 0)
+	if (Times.IsValidIndex(FirstVisibleIndex) && LastVisibleIndex <= Times.Num() && Num > 0)
 	{
-		return MakeArrayView(&GroupedTimes[FirstVisibleIndex], Num);
+		return MakeArrayView(&Times[FirstVisibleIndex], Num);
 	}
 
 	return TArrayView<const FFrameNumber>();
 }
-/*
-FFrameNumber FSequencerKeyCollection::GetNextKey(FFrameNumber FrameNumber, EFindKeyDirection Direction) const
-{
-	FFrameNumber NextFrame = FrameNumber;
-	int32 Num = GroupedTimes.Num();
-	if (Num > 0)
-	{
-		if (FrameNumber < GroupedTimes[0])
-		{
-			NextFrame = GroupedTimes[0];
-		}
-		else if (FrameNumber > GroupedTimes[Num - 1])
-		{
-			NextFrame = GroupedTimes[Num - 1];
-		}
-		else
-		{
-			int32 Index = Algo::LowerBound(GroupedTimes, FrameNumber);
-			if (GroupedTimes[Index] != FrameNumber)
-			{
-				if (Direction == EFindKeyDirection::Forwards)
-				{
-					NextFrame = GroupedTimes[Index];
-				}
-				else
-				{
-					NextFrame = GroupedTimes[Index - 1];
-				}
-			}
-			else {
-				if (Direction == EFindKeyDirection::Forwards)
-				{
-					if (++Index >= Num)
-					{
-						Index = Num - 1;
-					}
-					NextFrame = GroupedTimes[Index];
-				}
-				else
-				{
-					if (--Index < 0)
-					{
-						Index = 0;
-					}
-					NextFrame = GroupedTimes[Index];
-				}
-			}
-		}
-	}
-	return NextFrame;
 
-}
-
-*/
-
-TOptional<FFrameNumber> FSequencerKeyCollection::GetNextKey(FFrameNumber FrameNumber, EFindKeyDirection Direction) const
+TOptional<FFrameNumber> GetNextKeyInternal(const TArray<FFrameNumber>& Times, FFrameNumber FrameNumber, EFindKeyDirection Direction)
 {
 	int32 Index = INDEX_NONE;
 	if (Direction == EFindKeyDirection::Forwards)
 	{
-		Index = Algo::UpperBound(GroupedTimes, FrameNumber);
+		Index = Algo::UpperBound(Times, FrameNumber);
 	}
 	else
 	{
-		Index = Algo::LowerBound(GroupedTimes, FrameNumber) - 1;
+		Index = Algo::LowerBound(Times, FrameNumber) - 1;
 	}
 
-	if (GroupedTimes.IsValidIndex(Index))
+	if (Times.IsValidIndex(Index))
 	{
-		return GroupedTimes[Index];
+		return Times[Index];
 	}
-	else if (GroupedTimes.Num() > 0)
+	else if (Times.Num() > 0)
 	{
 		if (Direction == EFindKeyDirection::Forwards)
 		{
-			return GroupedTimes[0];
+			return Times[0];
 		}
 		else
 		{
-			return GroupedTimes.Last();
+			return Times.Last();
 		}
 	}
 
 	return TOptional<FFrameNumber>();
 }
+
+TOptional<FFrameNumber> FindFirstKeyInRangeInternal(const TArray<FFrameNumber>& Times, const TRange<FFrameNumber>& Range, EFindKeyDirection Direction)
+{
+	TArrayView<const FFrameNumber> KeysInRange = GetKeysInRangeInternal(Times, Range);
+	if (KeysInRange.Num())
+	{
+		return Direction == EFindKeyDirection::Forwards ? KeysInRange[0] : KeysInRange[KeysInRange.Num() - 1];
+	}
+	return TOptional<FFrameNumber>();
+}
+
+TOptional<FFrameNumber> FSequencerKeyCollection::FindFirstKeyInRange(const TRange<FFrameNumber>& Range, EFindKeyDirection Direction) const
+{
+	return FindFirstKeyInRangeInternal(GroupedTimes, Range, Direction);
+}
+
+TOptional<FFrameNumber> FSequencerKeyCollection::FindFirstSectionKeyInRange(const TRange<FFrameNumber>& Range, EFindKeyDirection Direction) const
+{
+	return FindFirstKeyInRangeInternal(GroupedSectionTimes, Range, Direction);
+}
+
+TArrayView<const FFrameNumber> FSequencerKeyCollection::GetKeysInRange(const TRange<FFrameNumber>& Range) const
+{
+	return GetKeysInRangeInternal(GroupedTimes, Range);
+}
+
+TArrayView<const FFrameNumber> FSequencerKeyCollection::GetSectionKeysInRange(const TRange<FFrameNumber>& Range) const
+{
+	return GetKeysInRangeInternal(GroupedSectionTimes, Range);
+}
+
+TOptional<FFrameNumber> FSequencerKeyCollection::GetNextKey(FFrameNumber FrameNumber, EFindKeyDirection Direction) const
+{
+	return GetNextKeyInternal(GroupedTimes, FrameNumber, Direction);
+}
+
+TOptional<FFrameNumber> FSequencerKeyCollection::GetNextSectionKey(FFrameNumber FrameNumber, EFindKeyDirection Direction) const
+{
+	return GetNextKeyInternal(GroupedSectionTimes, FrameNumber, Direction);
+}
diff --git a/Engine/Source/Editor/Sequencer/Private/SequencerSettings.cpp b/Engine/Source/Editor/Sequencer/Private/SequencerSettings.cpp
--- a/Engine/Source/Editor/Sequencer/Private/SequencerSettings.cpp
+++ b/Engine/Source/Editor/Sequencer/Private/SequencerSettings.cpp
@@ -22,6 +22,7 @@ USequencerSettings::USequencerSettings( const FObjectInitializer& ObjectInitiali
 	bSnapSectionTimesToInterval = true;
 	bSnapSectionTimesToSections = true;
 	bSnapPlayTimeToKeys = false;
+	bSnapPlayTimeToSections = false;
 	bSnapPlayTimeToMarkers = false;
 	bSnapPlayTimeToInterval = true;
 	bSnapPlayTimeToPressedKey = true;
@@ -261,6 +262,20 @@ void USequencerSettings::SetSnapPlayTimeToKeys(bool InbSnapPlayTimeToKeys)
 	}
 }
 
+bool USequencerSettings::GetSnapPlayTimeToSections() const
+{
+	return bSnapPlayTimeToSections;
+}
+
+void USequencerSettings::SetSnapPlayTimeToSections(bool InbSnapPlayTimeToSections)
+{
+	if (bSnapPlayTimeToSections != InbSnapPlayTimeToSections)
+	{
+		bSnapPlayTimeToSections = InbSnapPlayTimeToSections;
+		SaveConfig();
+	}
+}
+
 bool USequencerSettings::GetSnapPlayTimeToMarkers() const
 {
 	return bSnapPlayTimeToMarkers;
diff --git a/Engine/Source/Editor/Sequencer/Private/SequencerTimeSliderController.cpp b/Engine/Source/Editor/Sequencer/Private/SequencerTimeSliderController.cpp
--- a/Engine/Source/Editor/Sequencer/Private/SequencerTimeSliderController.cpp
+++ b/Engine/Source/Editor/Sequencer/Private/SequencerTimeSliderController.cpp
@@ -1737,6 +1737,11 @@ FFrameTime FSequencerTimeSliderController::SnapTimeToNearestKey(const FPointerEv
 			EnumAddFlags(NearestKeyOption, ENearestKeyOption::NKO_SearchKeys);
 		}
 
+		if (WeakSequencer.Pin()->GetSequencerSettings()->GetSnapPlayTimeToSections() || MouseEvent.IsShiftDown())
+		{
+			EnumAddFlags(NearestKeyOption, ENearestKeyOption::NKO_SearchSections);
+		}
+
 		if (WeakSequencer.Pin()->GetSequencerSettings()->GetSnapPlayTimeToMarkers() || MouseEvent.IsShiftDown())
 		{
 			EnumAddFlags(NearestKeyOption, ENearestKeyOption::NKO_SearchMarkers);
diff --git a/Engine/Source/Editor/Sequencer/Public/SequencerCommands.h b/Engine/Source/Editor/Sequencer/Public/SequencerCommands.h
--- a/Engine/Source/Editor/Sequencer/Public/SequencerCommands.h
+++ b/Engine/Source/Editor/Sequencer/Public/SequencerCommands.h
@@ -316,6 +316,9 @@ public:
 	/** Toggles whether or not snap to key times while scrubbing. */
 	TSharedPtr< FUICommandInfo > ToggleSnapPlayTimeToKeys;
 
+	/** Toggles whether or not snap to section bounds while scrubbing. */
+	TSharedPtr< FUICommandInfo > ToggleSnapPlayTimeToSections;
+
 	/** Toggles whether or not snap to markers while scrubbing. */
 	TSharedPtr< FUICommandInfo > ToggleSnapPlayTimeToMarkers;
 
diff --git a/Engine/Source/Editor/Sequencer/Public/SequencerKeyCollection.h b/Engine/Source/Editor/Sequencer/Public/SequencerKeyCollection.h
--- a/Engine/Source/Editor/Sequencer/Public/SequencerKeyCollection.h
+++ b/Engine/Source/Editor/Sequencer/Public/SequencerKeyCollection.h
@@ -80,6 +80,7 @@ public:
 	 * @return (Optional) the time of the key that matched the range
 	 */
 	SEQUENCER_API TOptional<FFrameNumber> FindFirstKeyInRange(const TRange<FFrameNumber>& Range, EFindKeyDirection Direction) const;
+	SEQUENCER_API TOptional<FFrameNumber> FindFirstSectionKeyInRange(const TRange<FFrameNumber>& Range, EFindKeyDirection Direction) const;
 
 	/**
 	 * Get a view of all key times that reside within the specified range
@@ -88,6 +89,7 @@ public:
 	 * @return A (possibly empty) array view of all the times that lie within the range
 	 */
 	SEQUENCER_API TArrayView<const FFrameNumber> GetKeysInRange(const TRange<FFrameNumber>& Range) const;
+	SEQUENCER_API TArrayView<const FFrameNumber> GetSectionKeysInRange(const TRange<FFrameNumber>& Range) const;
 
 	/**
 	* Search forwards or backwards for the next key from the specified frame number
@@ -96,6 +98,7 @@ public:
 	* @return (Optional)  Frame number of the key that's next or previous from that time 
 	*/
 	SEQUENCER_API TOptional<FFrameNumber> GetNextKey(FFrameNumber FrameNumber, EFindKeyDirection Direction) const;
+	SEQUENCER_API TOptional<FFrameNumber> GetNextSectionKey(FFrameNumber FrameNumber, EFindKeyDirection Direction) const;
 
 	/**
 	 * Access the signature this collection was generated with
@@ -122,6 +125,9 @@ private:
 	/** Times grouped by the supplied threshold */
 	TArray<FFrameNumber> GroupedTimes;
 
+	/** Section times grouped by the supplied threshold */
+	TArray<FFrameNumber> GroupedSectionTimes;
+
 	/** The signature with which the above array was generated */
 	FSequencerKeyCollectionSignature Signature;
 };
diff --git a/Engine/Source/Editor/Sequencer/Public/SequencerSettings.h b/Engine/Source/Editor/Sequencer/Public/SequencerSettings.h
--- a/Engine/Source/Editor/Sequencer/Public/SequencerSettings.h
+++ b/Engine/Source/Editor/Sequencer/Public/SequencerSettings.h
@@ -168,6 +168,11 @@ public:
 	/** Sets whether or not to snap the play time to keys while scrubbing. */
 	void SetSnapPlayTimeToKeys(bool InbSnapPlayTimeToKeys);
 
+	/** Gets whether or not to snap the play time to section bounds while scrubbing. */
+	bool GetSnapPlayTimeToSections() const;
+	/** Sets whether or not to snap the play time to section bounds while scrubbing. */
+	void SetSnapPlayTimeToSections(bool InbSnapPlayTimeToSections);
+
 	/** Gets whether or not to snap the play time to markers while scrubbing. */
 	bool GetSnapPlayTimeToMarkers() const;
 	/** Sets whether or not to snap the play time to markers while scrubbing. */
@@ -425,10 +430,14 @@ protected:
 	UPROPERTY(config, EditAnywhere, Category = Timeline)
 	bool bSnapKeysAndSectionsToPlayRange;
 
-	/** Enable or disable snapping the current time to keys of the selected track while scrubbing. */
+	/** Enable or disable snapping the current time to keys while scrubbing. */
 	UPROPERTY( config, EditAnywhere, Category=Snapping )
 	bool bSnapPlayTimeToKeys;
 
+	/** Enable or disable snapping the current time to section bounds while scrubbing. */
+	UPROPERTY( config, EditAnywhere, Category=Snapping )
+	bool bSnapPlayTimeToSections;
+
 	/** Enable or disable snapping the current time to markers while scrubbing. */
 	UPROPERTY( config, EditAnywhere, Category=Snapping )
 	bool bSnapPlayTimeToMarkers;
diff --git a/Engine/Source/Editor/SequencerWidgets/Public/ITimeSlider.h b/Engine/Source/Editor/SequencerWidgets/Public/ITimeSlider.h
--- a/Engine/Source/Editor/SequencerWidgets/Public/ITimeSlider.h
+++ b/Engine/Source/Editor/SequencerWidgets/Public/ITimeSlider.h
@@ -32,11 +32,14 @@ enum class ENearestKeyOption : uint8
 	/* Search keys */
 	NKO_SearchKeys = 0x01,
 
+	/* Search sections */
+	NKO_SearchSections = 0x02,
+
 	/* Search markers */
-	NKO_SearchMarkers = 0x02,
+	NKO_SearchMarkers = 0x04,
 
 	/** Search all tracks */
-	NKO_SearchAllTracks = 0x04
+	NKO_SearchAllTracks = 0x08
 };
 ENUM_CLASS_FLAGS(ENearestKeyOption);
 
