MovieScene: remove DeltaFrame member from FMovieSceneTimecodeSource

The DeltaFrame value was being set in UMovieSceneSection::MoveSection(), but
was not being used anywhere.

diff --git a/Engine/Source/Editor/Sequencer/Private/Sequencer.cpp b/Engine/Source/Editor/Sequencer/Private/Sequencer.cpp
--- a/Engine/Source/Editor/Sequencer/Private/Sequencer.cpp
+++ b/Engine/Source/Editor/Sequencer/Private/Sequencer.cpp
@@ -2581,7 +2581,7 @@ void FSequencer::SyncSectionsUsingSourceTimecode()
 
 	const UMovieSceneSection* FirstSection = Sections[0];
 	FFrameNumber FirstSectionSourceTimecode = FirstSection->TimecodeSource.Timecode.ToFrameNumber(GetFocusedTickResolution());
-	FFrameNumber FirstSectionCurrentStartFrame = FirstSection->GetInclusiveStartFrame();// - FirstSection->TimecodeSource.DeltaFrame;
+	FFrameNumber FirstSectionCurrentStartFrame = FirstSection->GetInclusiveStartFrame();
 	Sections.RemoveAt(0);
 
 	for (auto Section : Sections)
@@ -2589,7 +2589,7 @@ void FSequencer::SyncSectionsUsingSourceTimecode()
 		if (Section->HasStartFrame())
 		{
 			FFrameNumber SectionSourceTimecode = Section->TimecodeSource.Timecode.ToFrameNumber(GetFocusedTickResolution());
-			FFrameNumber SectionCurrentStartFrame = Section->GetInclusiveStartFrame();// - Section->TimecodeSource.DeltaFrame;
+			FFrameNumber SectionCurrentStartFrame = Section->GetInclusiveStartFrame();
 
 			FFrameNumber TimecodeDelta = SectionSourceTimecode - FirstSectionSourceTimecode;
 			FFrameNumber CurrentDelta = SectionCurrentStartFrame - FirstSectionCurrentStartFrame;
diff --git a/Engine/Source/Runtime/MovieScene/Private/MovieSceneSection.cpp b/Engine/Source/Runtime/MovieScene/Private/MovieSceneSection.cpp
--- a/Engine/Source/Runtime/MovieScene/Private/MovieSceneSection.cpp
+++ b/Engine/Source/Runtime/MovieScene/Private/MovieSceneSection.cpp
@@ -219,13 +219,8 @@ void UMovieSceneSection::MoveSection(FFrameNumber DeltaFrame)
 			}
 		}
 	}
-
-#if WITH_EDITORONLY_DATA
-	TimecodeSource.DeltaFrame += DeltaFrame;
-#endif
 }
 
-
 TRange<FFrameNumber> UMovieSceneSection::ComputeEffectiveRange() const
 {
 	if (!SectionRange.Value.GetLowerBound().IsOpen() && !SectionRange.Value.GetUpperBound().IsOpen())
diff --git a/Engine/Source/Runtime/MovieScene/Public/MovieScene.h b/Engine/Source/Runtime/MovieScene/Public/MovieScene.h
--- a/Engine/Source/Runtime/MovieScene/Public/MovieScene.h
+++ b/Engine/Source/Runtime/MovieScene/Public/MovieScene.h
@@ -50,21 +50,19 @@ struct FMovieSceneTimecodeSource
 
 	FMovieSceneTimecodeSource(FTimecode InTimecode)
 		: Timecode(InTimecode)
-		, DeltaFrame(FFrameNumber())
 	{}
 
 	FMovieSceneTimecodeSource()
 		: Timecode(FTimecode())
-		, DeltaFrame(FFrameNumber())
 	{}
 
 	FORCEINLINE bool operator==(const FMovieSceneTimecodeSource& Other) const
 	{
-		return Timecode == Other.Timecode && DeltaFrame == Other.DeltaFrame;
+		return Timecode == Other.Timecode;
 	}
 	FORCEINLINE bool operator!=(const FMovieSceneTimecodeSource& Other) const
 	{
-		return Timecode != Other.Timecode || DeltaFrame != Other.DeltaFrame;
+		return Timecode != Other.Timecode;
 	}
 
 public:
@@ -72,10 +70,6 @@ public:
 	/** The global timecode at which this target is based (ie. the timecode at the beginning of the movie scene section when it was recorded) */
 	UPROPERTY(EditAnywhere, Category="Timecode")
 	FTimecode Timecode;
-
-	/** The delta from the original placement of this target */
-	UPROPERTY(VisibleAnywhere, Category="Timecode")
-	FFrameNumber DeltaFrame;
 };
 
 USTRUCT(BlueprintType)
