Sequencer: Convert the playback position to input space when resetting it.

diff --git a/Engine/Source/Editor/Sequencer/Private/Sequencer.cpp b/Engine/Source/Editor/Sequencer/Private/Sequencer.cpp
--- a/Engine/Source/Editor/Sequencer/Private/Sequencer.cpp
+++ b/Engine/Source/Editor/Sequencer/Private/Sequencer.cpp
@@ -434,7 +434,7 @@ void FSequencer::InitSequencer(const FSequencerInitParams& InitParams, const TSh
 	ResetTimeController();
 
 	UpdateTimeBases();
-	PlayPosition.Reset(GetPlaybackRange().GetLowerBoundValue());
+	PlayPosition.Reset(ConvertFrameTime(GetPlaybackRange().GetLowerBoundValue(), GetRootTickResolution(), PlayPosition.GetInputRate()));
 
 	// Make internal widgets
 	SequencerWidget = SNew( SSequencer, SharedThis( this ) )
@@ -945,7 +945,7 @@ void FSequencer::ResetToNewRootSequence(UMovieSceneSequence& NewSequence)
 	ResetPerMovieSceneData();
 	SequencerWidget->ResetBreadcrumbs();
 
-	PlayPosition.Reset(GetPlaybackRange().GetLowerBoundValue());
+	PlayPosition.Reset(ConvertFrameTime(GetPlaybackRange().GetLowerBoundValue(), GetRootTickResolution(), PlayPosition.GetInputRate()));
 	TimeController->Reset(FQualifiedFrameTime(PlayPosition.GetCurrentPosition(), GetRootTickResolution()));
 
 	UpdateSequencerCustomizations();
