Sequencer: Add ability for middle mouse to change time without evaluating sequencer.

diff --git a/Engine/Source/Editor/CurveEditor/Private/CurveEditor.cpp b/Engine/Source/Editor/CurveEditor/Private/CurveEditor.cpp
--- a/Engine/Source/Editor/CurveEditor/Private/CurveEditor.cpp
+++ b/Engine/Source/Editor/CurveEditor/Private/CurveEditor.cpp
@@ -784,11 +784,11 @@ void FCurveEditor::StepToNextKey()
 
 	if (NextTime.IsSet())
 	{
-		WeakTimeSliderController.Pin()->SetScrubPosition(NextTime.GetValue() * TickResolution);
+		WeakTimeSliderController.Pin()->SetScrubPosition(NextTime.GetValue() * TickResolution,/*bEvaluate*/ true);
 	}
 	else if (MinTime.IsSet())
 	{
-		WeakTimeSliderController.Pin()->SetScrubPosition(MinTime.GetValue() * TickResolution);
+		WeakTimeSliderController.Pin()->SetScrubPosition(MinTime.GetValue() * TickResolution, /*bEvaluate*/ true);
 	}
 }
 
@@ -842,11 +842,11 @@ void FCurveEditor::StepToPreviousKey()
 
 	if (PreviousTime.IsSet())
 	{
-		WeakTimeSliderController.Pin()->SetScrubPosition(PreviousTime.GetValue() * TickResolution);
+		WeakTimeSliderController.Pin()->SetScrubPosition(PreviousTime.GetValue() * TickResolution,/*bEvaluate*/ true);
 	}
 	else if (MaxTime.IsSet())
 	{
-		WeakTimeSliderController.Pin()->SetScrubPosition(MaxTime.GetValue() * TickResolution);
+		WeakTimeSliderController.Pin()->SetScrubPosition(MaxTime.GetValue() * TickResolution, /*bEvaluate*/ true);
 	}
 }
 
@@ -863,7 +863,7 @@ void FCurveEditor::StepForward()
 
 	FFrameTime OneFrame = FFrameRate::TransformTime(FFrameTime(1), DisplayRate, TickResolution);
 
-	WeakTimeSliderController.Pin()->SetScrubPosition(WeakTimeSliderController.Pin()->GetScrubPosition() + OneFrame);
+	WeakTimeSliderController.Pin()->SetScrubPosition(WeakTimeSliderController.Pin()->GetScrubPosition() + OneFrame, /*bEvaluate*/ true);
 }
 
 void FCurveEditor::StepBackward()
@@ -878,7 +878,7 @@ void FCurveEditor::StepBackward()
 
 	FFrameTime OneFrame = FFrameRate::TransformTime(FFrameTime(1), DisplayRate, TickResolution);
 
-	WeakTimeSliderController.Pin()->SetScrubPosition(WeakTimeSliderController.Pin()->GetScrubPosition() - OneFrame);
+	WeakTimeSliderController.Pin()->SetScrubPosition(WeakTimeSliderController.Pin()->GetScrubPosition() - OneFrame, /*bEvaluate*/ true);
 }
 
 void FCurveEditor::JumpToStart()
@@ -888,7 +888,7 @@ void FCurveEditor::JumpToStart()
 		return;
 	}
 
-	WeakTimeSliderController.Pin()->SetScrubPosition(WeakTimeSliderController.Pin()->GetPlayRange().GetLowerBoundValue());
+	WeakTimeSliderController.Pin()->SetScrubPosition(WeakTimeSliderController.Pin()->GetPlayRange().GetLowerBoundValue(), /*bEvaluate*/ true);
 }
 
 void FCurveEditor::JumpToEnd()
@@ -908,7 +908,7 @@ void FCurveEditor::JumpToEnd()
 	FFrameTime OneFrame = bInsetDisplayFrame ? FFrameRate::TransformTime(FFrameTime(1), DisplayRate, TickResolution) : FFrameTime(1);
 	FFrameTime NewTime = WeakTimeSliderController.Pin()->GetPlayRange().GetUpperBoundValue() - OneFrame;
 
-	WeakTimeSliderController.Pin()->SetScrubPosition(NewTime);
+	WeakTimeSliderController.Pin()->SetScrubPosition(NewTime, /*bEvaluate*/ true);
 }
 
 void FCurveEditor::SetSelectionRangeStart()
diff --git a/Engine/Source/Editor/Persona/Private/AnimTimeline/AnimTimeSliderController.cpp b/Engine/Source/Editor/Persona/Private/AnimTimeline/AnimTimeSliderController.cpp
--- a/Engine/Source/Editor/Persona/Private/AnimTimeline/AnimTimeSliderController.cpp
+++ b/Engine/Source/Editor/Persona/Private/AnimTimeline/AnimTimeSliderController.cpp
@@ -809,7 +809,7 @@ void FAnimTimeSliderController::CommitScrubPosition( FFrameTime NewValue, bool b
 		TimeSliderArgs.ScrubPosition.Set( NewValue );
 	}
 
-	TimeSliderArgs.OnScrubPositionChanged.ExecuteIfBound( NewValue, bIsScrubbing );
+	TimeSliderArgs.OnScrubPositionChanged.ExecuteIfBound( NewValue, bIsScrubbing, /*bEvaluate*/ true); //todo change if anim timeline needs to handle sequencer style middle mouse manipulation which changes time but doesn't evaluate
 }
 
 FReply FAnimTimeSliderController::OnMouseWheel( SWidget& WidgetOwner, const FGeometry& MyGeometry, const FPointerEvent& MouseEvent )
diff --git a/Engine/Source/Editor/Persona/Private/AnimTimeline/AnimTimeSliderController.h b/Engine/Source/Editor/Persona/Private/AnimTimeline/AnimTimeSliderController.h
--- a/Engine/Source/Editor/Persona/Private/AnimTimeline/AnimTimeSliderController.h
+++ b/Engine/Source/Editor/Persona/Private/AnimTimeline/AnimTimeSliderController.h
@@ -65,7 +65,7 @@ public:
 	virtual FAnimatedRange GetClampRange() const override { return TimeSliderArgs.ClampRange.Get(); }
 	virtual TRange<FFrameNumber> GetPlayRange() const override { return TimeSliderArgs.PlaybackRange.Get(TRange<FFrameNumber>()); }
 	virtual FFrameTime GetScrubPosition() const { return TimeSliderArgs.ScrubPosition.Get(); }
-	virtual void SetScrubPosition(FFrameTime InTime) { TimeSliderArgs.OnScrubPositionChanged.ExecuteIfBound(InTime, false); }
+	virtual void SetScrubPosition(FFrameTime InTime, bool bEvaluate) { TimeSliderArgs.OnScrubPositionChanged.ExecuteIfBound(InTime, false, bEvaluate); }
 	/** End ITimeSliderController Interface */
 
 	/**
diff --git a/Engine/Source/Editor/Persona/Private/AnimTimeline/SAnimTimeline.cpp b/Engine/Source/Editor/Persona/Private/AnimTimeline/SAnimTimeline.cpp
--- a/Engine/Source/Editor/Persona/Private/AnimTimeline/SAnimTimeline.cpp
+++ b/Engine/Source/Editor/Persona/Private/AnimTimeline/SAnimTimeline.cpp
@@ -795,7 +795,7 @@ class UAnimSingleNodeInstance* SAnimTimeline::GetPreviewInstance() const
 	return PreviewMeshComponent && PreviewMeshComponent->IsPreviewOn()? PreviewMeshComponent->PreviewInstance : nullptr;
 }
 
-void SAnimTimeline::HandleScrubPositionChanged(FFrameTime NewScrubPosition, bool bIsScrubbing)
+void SAnimTimeline::HandleScrubPositionChanged(FFrameTime NewScrubPosition, bool bIsScrubbing, bool bEvaluate)
 {
 	if (UAnimSingleNodeInstance* PreviewInstance = GetPreviewInstance())
 	{
diff --git a/Engine/Source/Editor/Persona/Private/AnimTimeline/SAnimTimeline.h b/Engine/Source/Editor/Persona/Private/AnimTimeline/SAnimTimeline.h
--- a/Engine/Source/Editor/Persona/Private/AnimTimeline/SAnimTimeline.h
+++ b/Engine/Source/Editor/Persona/Private/AnimTimeline/SAnimTimeline.h
@@ -78,7 +78,7 @@ private:
 
 
 
-	void HandleScrubPositionChanged(FFrameTime NewScrubPosition, bool bIsScrubbing);
+	void HandleScrubPositionChanged(FFrameTime NewScrubPosition, bool bIsScrubbing, bool bEvaluate);
 
 	void OnCropAnimSequence(bool bFromStart, float CurrentTime);
 
diff --git a/Engine/Source/Editor/Sequencer/Private/Sequencer.cpp b/Engine/Source/Editor/Sequencer/Private/Sequencer.cpp
--- a/Engine/Source/Editor/Sequencer/Private/Sequencer.cpp
+++ b/Engine/Source/Editor/Sequencer/Private/Sequencer.cpp
@@ -2307,6 +2307,8 @@ void FSequencer::NotifyMovieSceneDataChanged( EMovieSceneDataChangeType DataChan
 	if (DataChangeType == EMovieSceneDataChangeType::RefreshTree)
 	{
 		bNeedTreeRefresh = true;
+		OnMovieSceneDataChangedDelegate.Broadcast(DataChangeType);
+		return;
 	}
 	else if ( DataChangeType == EMovieSceneDataChangeType::MovieSceneStructureItemRemoved ||
 		DataChangeType == EMovieSceneDataChangeType::MovieSceneStructureItemsChanged ||
@@ -2984,7 +2986,7 @@ FQualifiedFrameTime FSequencer::GetGlobalTime() const
 	return FQualifiedFrameTime(RootTime, PlayPosition.GetOutputRate());
 }
 
-void FSequencer::SetLocalTime( FFrameTime NewTime, ESnapTimeMode SnapTimeMode)
+void FSequencer::SetLocalTime( FFrameTime NewTime, ESnapTimeMode SnapTimeMode, bool bEvaluate)
 {
 	FFrameRate LocalResolution = GetFocusedTickResolution();
 
@@ -3024,11 +3026,11 @@ void FSequencer::SetLocalTime( FFrameTime NewTime, ESnapTimeMode SnapTimeMode)
 		NewTime = OnGetNearestKey(NewTime, NearestKeyOption);
 	}
 
-	SetLocalTimeDirectly(NewTime);
+	SetLocalTimeDirectly(NewTime, bEvaluate);
 }
 
 
-void FSequencer::SetLocalTimeDirectly(FFrameTime NewTime)
+void FSequencer::SetLocalTimeDirectly(FFrameTime NewTime, bool bEvaluate)
 {
 	TWeakPtr<SWidget> PreviousFocusedWidget = FSlateApplication::Get().GetKeyboardFocusedWidget();
 
@@ -3039,7 +3041,7 @@ void FSequencer::SetLocalTimeDirectly(FFrameTime NewTime)
 	}
 
 	// Transform the time to the root time-space
-	SetGlobalTime(NewTime * RootToLocalTransform.InverseFromWarp(RootToLocalLoopCounter));
+	SetGlobalTime(NewTime * RootToLocalTransform.InverseFromWarp(RootToLocalLoopCounter), bEvaluate);
 
 	if (PreviousFocusedWidget.IsValid())
 	{
@@ -3048,7 +3050,7 @@ void FSequencer::SetLocalTimeDirectly(FFrameTime NewTime)
 }
 
 
-void FSequencer::SetGlobalTime(FFrameTime NewTime)
+void FSequencer::SetGlobalTime(FFrameTime NewTime, bool bEvaluate)
 {
 	NewTime = ConvertFrameTime(NewTime, GetRootTickResolution(), PlayPosition.GetInputRate());
 	if (PlayPosition.GetEvaluationType() == EMovieSceneEvaluationType::FrameLocked)
@@ -3061,7 +3063,11 @@ void FSequencer::SetGlobalTime(FFrameTime NewTime)
 	TOptional<FFrameTime> CurrentPosition = PlayPosition.GetCurrentPosition();
 	if (PlayPosition.GetCurrentPosition() != NewTime)
 	{
-		EvaluateInternal(PlayPosition.JumpTo(NewTime));
+		FMovieSceneEvaluationRange EvalRange = PlayPosition.JumpTo(NewTime);
+		if (bEvaluate)
+		{
+			EvaluateInternal(EvalRange);
+		}
 	}
 
 	if (AutoScrubTarget.IsSet())
@@ -3129,7 +3135,6 @@ void FSequencer::EvaluateInternal(FMovieSceneEvaluationRange InRange, bool bHasJ
 	bNeedsEvaluate = false;
 
 	UpdateCachedPlaybackContextAndClient();
-	
 	if (EventContextsAttribute.IsBound())
 	{
 		CachedEventContexts.Reset();
@@ -3165,6 +3170,7 @@ void FSequencer::EvaluateInternal(FMovieSceneEvaluationRange InRange, bool bHasJ
 	{
 		OnGlobalTimeChangedDelegate.Broadcast();
 	}
+
 }
 
 void FSequencer::UpdateCachedPlaybackContextAndClient()
@@ -5150,7 +5156,7 @@ FFrameNumber FSequencer::OnGetNearestKey(FFrameTime InTime, ENearestKeyOption Ne
 	return NearestTime.IsSet() ? NearestTime.GetValue() : CurrentTime;
 }
 
-void FSequencer::OnScrubPositionChanged( FFrameTime NewScrubPosition, bool bScrubbing )
+void FSequencer::OnScrubPositionChanged( FFrameTime NewScrubPosition, bool bScrubbing , bool bEvaluate)
 {
 	if (PlaybackState == EMovieScenePlayerStatus::Scrubbing)
 	{
@@ -5176,7 +5182,7 @@ void FSequencer::OnScrubPositionChanged( FFrameTime NewScrubPosition, bool bScru
 	}
 	else
 	{
-		SetLocalTimeDirectly(NewScrubPosition);
+		SetLocalTimeDirectly(NewScrubPosition, bEvaluate);
 	}
 }
 
diff --git a/Engine/Source/Editor/Sequencer/Private/Sequencer.h b/Engine/Source/Editor/Sequencer/Private/Sequencer.h
--- a/Engine/Source/Editor/Sequencer/Private/Sequencer.h
+++ b/Engine/Source/Editor/Sequencer/Private/Sequencer.h
@@ -755,9 +755,9 @@ public:
 	virtual FQualifiedFrameTime GetLocalTime() const override;
 	virtual FQualifiedFrameTime GetGlobalTime() const override;
 	virtual uint32 GetLocalLoopIndex() const override;
-	virtual void SetLocalTime(FFrameTime Time, ESnapTimeMode SnapTimeMode = ESnapTimeMode::STM_None) override;
-	virtual void SetLocalTimeDirectly(FFrameTime NewTime) override;
-	virtual void SetGlobalTime(FFrameTime Time) override;
+	virtual void SetLocalTime(FFrameTime Time, ESnapTimeMode SnapTimeMode = ESnapTimeMode::STM_None, bool bEvaluate = true) override;
+	virtual void SetLocalTimeDirectly(FFrameTime NewTime, bool bEvaluate = true) override;
+	virtual void SetGlobalTime(FFrameTime Time, bool bEvaluate = true) override;
 	virtual void PlayTo(FMovieSceneSequencePlaybackParams PlaybackParams) override;
 	virtual void RestorePlaybackSpeed() override;
 	virtual void SnapToClosestPlaybackSpeed() override;
@@ -907,8 +907,10 @@ protected:
 	 * This will stop any playback from happening
 	 *
 	 * @param NewScrubPosition	The new scrub position
+	 * @param bScrubbing If scrubbing
+	 * @param bEvaluate  Do evaluate sequencer after changing time
 	 */
-	void OnScrubPositionChanged( FFrameTime NewScrubPosition, bool bScrubbing );
+	void OnScrubPositionChanged( FFrameTime NewScrubPosition, bool bScrubbing, bool bEvaluate);
 
 	/** Called when the user has begun scrubbing */
 	void OnBeginScrubbing();
diff --git a/Engine/Source/Editor/Sequencer/Private/SequencerTimeSliderController.cpp b/Engine/Source/Editor/Sequencer/Private/SequencerTimeSliderController.cpp
--- a/Engine/Source/Editor/Sequencer/Private/SequencerTimeSliderController.cpp
+++ b/Engine/Source/Editor/Sequencer/Private/SequencerTimeSliderController.cpp
@@ -64,6 +64,21 @@ FSequencerTimeSliderController::FSequencerTimeSliderController( const FTimeSlide
 	VanillaScrubHandleUpBrush      = FEditorStyle::GetBrush( TEXT( "Sequencer.Timeline.VanillaScrubHandleUp" ) ); 
 	VanillaScrubHandleDownBrush    = FEditorStyle::GetBrush( TEXT( "Sequencer.Timeline.VanillaScrubHandleDown" ) );
 	ContextMenuSuppression = 0;
+	TSharedPtr<FSequencer> Sequencer = WeakSequencer.Pin();
+	if (Sequencer.IsValid())
+	{
+		Sequencer->OnGlobalTimeChanged().AddRaw(this, &FSequencerTimeSliderController::SetIsEvaluating);
+	}
+
+}
+
+FSequencerTimeSliderController::~FSequencerTimeSliderController()
+{
+	TSharedPtr<FSequencer> Sequencer = WeakSequencer.Pin();
+	if (Sequencer.IsValid())
+	{
+		Sequencer->OnGlobalTimeChanged().RemoveAll(this);
+	}
 }
 
 FFrameTime FSequencerTimeSliderController::ComputeScrubTimeFromMouse(const FGeometry& Geometry, const FPointerEvent& MouseEvent, FScrubRangeToScreen RangeToScreen) const
@@ -461,17 +476,25 @@ int32 FSequencerTimeSliderController::OnPaintTimeSlider( bool bMirrorLabels, con
 		const int32 ArrowLayer = LayerId + 2;
 		FPaintGeometry MyGeometry =	AllottedGeometry.ToPaintGeometry( FVector2D( HandleStart, 0 ), FVector2D( HandleEnd - HandleStart, AllottedGeometry.Size.Y ) );
 		FLinearColor ScrubColor = InWidgetStyle.GetColorAndOpacityTint();
+		if(bIsEvaluating)
 		{
 			// @todo Sequencer this color should be specified in the style
 			ScrubColor.A = ScrubColor.A * 0.75f;
 			ScrubColor.B *= 0.1f;
 			ScrubColor.G *= 0.2f;
 		}
-
+		else
+		{
+			ScrubColor.A = ScrubColor.A * 0.75f;
+			ScrubColor.R = 0.7f;
+			ScrubColor.B = 0.1f;
+			ScrubColor.G = 0.7f;
+		}
 		const FSlateBrush* Brush = ScrubMetrics.Style == ESequencerScrubberStyle::Vanilla
 			? ( bMirrorLabels ? VanillaScrubHandleUpBrush    : VanillaScrubHandleDownBrush )
 			: ( bMirrorLabels ? FrameBlockScrubHandleUpBrush : FrameBlockScrubHandleDownBrush );
 
+		
 		FSlateDrawElement::MakeBox(
 			OutDrawElements,
 			ArrowLayer,
@@ -480,7 +503,7 @@ int32 FSequencerTimeSliderController::OnPaintTimeSlider( bool bMirrorLabels, con
 			DrawEffects,
 			ScrubColor
 		);
-
+		
 		LayerId = DrawMarkedFrames(AllottedGeometry, RangeToScreen, OutDrawElements, LayerId, DrawEffects, true);
 
 		{
@@ -779,7 +802,7 @@ FReply FSequencerTimeSliderController::OnMouseButtonUp( SWidget& WidgetOwner, co
 {
 	bool bHandleLeftMouseButton  = MouseEvent.GetEffectingButton() == EKeys::LeftMouseButton  && WidgetOwner.HasMouseCapture();
 	bool bHandleRightMouseButton = MouseEvent.GetEffectingButton() == EKeys::RightMouseButton && WidgetOwner.HasMouseCapture() && TimeSliderArgs.AllowZoom ;
-
+	bool bHandleMiddleMouseButton = MouseEvent.GetEffectingButton() == EKeys::MiddleMouseButton && WidgetOwner.HasMouseCapture();
 	FScrubRangeToScreen RangeToScreen = FScrubRangeToScreen(GetViewRange(), MyGeometry.Size);
 	FFrameTime          MouseTime     = ComputeFrameTimeFromMouse(MyGeometry, MouseEvent.GetScreenSpacePosition(), RangeToScreen);
 
@@ -814,7 +837,7 @@ FReply FSequencerTimeSliderController::OnMouseButtonUp( SWidget& WidgetOwner, co
 		
 		return FReply::Handled().ReleaseMouseCapture();
 	}
-	else if ( bHandleLeftMouseButton )
+	else if ( bHandleLeftMouseButton || bHandleMiddleMouseButton )
 	{
 		if (MouseDragType == DRAG_PLAYBACK_START)
 		{
@@ -885,7 +908,7 @@ FReply FSequencerTimeSliderController::OnMouseButtonUp( SWidget& WidgetOwner, co
 				ScrubTime = SnapTimeToNearestKey(MouseEvent, RangeToScreen, CursorPos.X, ScrubTime);
 			}
 
-			CommitScrubPosition( ScrubTime, /*bIsScrubbing=*/false );
+			CommitScrubPosition( ScrubTime, /*bIsScrubbing=*/false , /*bEvaluate*/ !bHandleMiddleMouseButton); //if middle mouse button down we don't evaluate on the time change
 		}
 
 		MouseDragType = DRAG_NONE;
@@ -910,6 +933,7 @@ FReply FSequencerTimeSliderController::OnMouseMove( SWidget& WidgetOwner, const
 
 	bool bHandleLeftMouseButton  = MouseEvent.IsMouseButtonDown( EKeys::LeftMouseButton  );
 	bool bHandleRightMouseButton = MouseEvent.IsMouseButtonDown( EKeys::RightMouseButton ) && TimeSliderArgs.AllowZoom;
+	bool bHandleMiddleMouseButton = MouseEvent.IsMouseButtonDown(EKeys::MiddleMouseButton);
 
 	if (bHandleRightMouseButton)
 	{
@@ -958,7 +982,7 @@ FReply FSequencerTimeSliderController::OnMouseMove( SWidget& WidgetOwner, const
 			SetViewRange(NewViewOutputMin, NewViewOutputMax, EViewRangeInterpolation::Immediate);
 		}
 	}
-	else if (bHandleLeftMouseButton)
+	else if (bHandleLeftMouseButton || bHandleMiddleMouseButton)
 	{
 		TRange<double> LocalViewRange = GetViewRange();
 		FScrubRangeToScreen RangeToScreen(LocalViewRange, MyGeometry.Size);
@@ -979,36 +1003,36 @@ FReply FSequencerTimeSliderController::OnMouseMove( SWidget& WidgetOwner, const
 				TRange<double>   PlaybackRange    = TimeSliderArgs.PlaybackRange.Get()  / TickResolution;
 
 				// Disable selection range test if it's empty so that the playback range scrubbing gets priority
-				if (!SelectionRange.IsEmpty() && !bHitScrubber && HitTestRangeEnd(RangeToScreen, SelectionRange, MouseDownPixel))
+				if (!SelectionRange.IsEmpty() && !bHitScrubber && HitTestRangeEnd(RangeToScreen, SelectionRange, MouseDownPixel) && bHandleMiddleMouseButton == false)
 				{
 					// selection range end scrubber
 					MouseDragType = DRAG_SELECTION_END;
 					TimeSliderArgs.OnSelectionRangeBeginDrag.ExecuteIfBound();
 				}
-				else if (!SelectionRange.IsEmpty() && !bHitScrubber && HitTestRangeStart(RangeToScreen, SelectionRange, MouseDownPixel))
+				else if (!SelectionRange.IsEmpty() && !bHitScrubber && HitTestRangeStart(RangeToScreen, SelectionRange, MouseDownPixel) && bHandleMiddleMouseButton == false)
 				{
 					// selection range start scrubber
 					MouseDragType = DRAG_SELECTION_START;
 					TimeSliderArgs.OnSelectionRangeBeginDrag.ExecuteIfBound();
 				}
-				else if (!bLockedPlayRange && !bHitScrubber && HitTestRangeEnd(RangeToScreen, PlaybackRange, MouseDownPixel))
+				else if (!bLockedPlayRange && !bHitScrubber && HitTestRangeEnd(RangeToScreen, PlaybackRange, MouseDownPixel) && bHandleMiddleMouseButton == false)
 				{
 					// playback range end scrubber
 					MouseDragType = DRAG_PLAYBACK_END;
 					TimeSliderArgs.OnPlaybackRangeBeginDrag.ExecuteIfBound();
 				}
-				else if (!bLockedPlayRange && !bHitScrubber && HitTestRangeStart(RangeToScreen, PlaybackRange, MouseDownPixel))
+				else if (!bLockedPlayRange && !bHitScrubber && HitTestRangeStart(RangeToScreen, PlaybackRange, MouseDownPixel) && bHandleMiddleMouseButton == false)
 				{
 					// playback range start scrubber
 					MouseDragType = DRAG_PLAYBACK_START;
 					TimeSliderArgs.OnPlaybackRangeBeginDrag.ExecuteIfBound();
 				}
-				else if (!bHitScrubber && HitTestMark(RangeToScreen, MouseDownPixel, DragMarkIndex))
+				else if (!bHitScrubber && HitTestMark(RangeToScreen, MouseDownPixel, DragMarkIndex) && bHandleMiddleMouseButton == false)
 				{
 					MouseDragType = DRAG_MARK;
 					TimeSliderArgs.OnMarkBeginDrag.ExecuteIfBound();
 				}
-				else if (FSlateApplication::Get().GetModifierKeys().AreModifersDown(EModifierKey::Control))
+				else if (FSlateApplication::Get().GetModifierKeys().AreModifersDown(EModifierKey::Control) && bHandleMiddleMouseButton == false)
 				{
 					MouseDragType = DRAG_SETTING_RANGE;
 				}
@@ -1084,7 +1108,7 @@ FReply FSequencerTimeSliderController::OnMouseMove( SWidget& WidgetOwner, const
 			else if (MouseDragType == DRAG_SCRUBBING_TIME)
 			{
 				// Delegate responsibility for clamping to the current viewrange to the client
-				CommitScrubPosition( ScrubTime, /*bIsScrubbing=*/true );
+				CommitScrubPosition( ScrubTime, /*bIsScrubbing=*/true, /*bEvaluate*/ !bHandleMiddleMouseButton); //if middle mouse button down we don't evaluate on the time change
 			}
 			else if (MouseDragType == DRAG_SETTING_RANGE)
 			{
@@ -1103,8 +1127,9 @@ FReply FSequencerTimeSliderController::OnMouseMove( SWidget& WidgetOwner, const
 }
 
 
-void FSequencerTimeSliderController::CommitScrubPosition( FFrameTime NewValue, bool bIsScrubbing )
+void FSequencerTimeSliderController::CommitScrubPosition( FFrameTime NewValue, bool bIsScrubbing, bool bEvaluate)
 {
+	bIsEvaluating = bEvaluate;
 	// The user can scrub past the viewing range of the time slider controller, so we clamp it to the view range.
 	TSharedPtr<FSequencer> Sequencer = WeakSequencer.Pin();
 	if(Sequencer.IsValid())
@@ -1132,7 +1157,7 @@ void FSequencerTimeSliderController::CommitScrubPosition( FFrameTime NewValue, b
 		TimeSliderArgs.ScrubPosition.Set( NewValue );
 	}
 
-	TimeSliderArgs.OnScrubPositionChanged.ExecuteIfBound( NewValue, bIsScrubbing );
+	TimeSliderArgs.OnScrubPositionChanged.ExecuteIfBound( NewValue, bIsScrubbing, bEvaluate);
 }
 
 FReply FSequencerTimeSliderController::OnMouseWheel( SWidget& WidgetOwner, const FGeometry& MyGeometry, const FPointerEvent& MouseEvent )
diff --git a/Engine/Source/Editor/Sequencer/Public/ISequencer.h b/Engine/Source/Editor/Sequencer/Public/ISequencer.h
--- a/Engine/Source/Editor/Sequencer/Public/ISequencer.h
+++ b/Engine/Source/Editor/Sequencer/Public/ISequencer.h
@@ -357,14 +357,15 @@ public:
 	 *
 	 * @param Time The local time to set.
 	 * @param SnapTimeMode The type of time snapping allowed.
+	 * @param bEvaluate If True also evaluate
 	 */
-	virtual void SetLocalTime(FFrameTime Time, ESnapTimeMode SnapTimeMode = ESnapTimeMode::STM_None) = 0;
+	virtual void SetLocalTime(FFrameTime Time, ESnapTimeMode SnapTimeMode = ESnapTimeMode::STM_None, bool bEvaluate = true) = 0;
 
 	/** Set the current local time directly, with no other snapping, scrolling or manipulation */
-	virtual void SetLocalTimeDirectly(FFrameTime NewTime) = 0;
+	virtual void SetLocalTimeDirectly(FFrameTime NewTime, bool bEvaluate = true) = 0;
 
 	/** Set the global time directly, without performing any auto-scroll, snapping or other adjustments to the supplied time  */
-	virtual void SetGlobalTime(FFrameTime Time) = 0;
+	virtual void SetGlobalTime(FFrameTime Time, bool bEvaluate = true) = 0;
 
 	/** Play from the current time to the requested time */
 	virtual void PlayTo(FMovieSceneSequencePlaybackParams PlaybackParams) = 0;
diff --git a/Engine/Source/Editor/Sequencer/Public/SequencerTimeSliderController.h b/Engine/Source/Editor/Sequencer/Public/SequencerTimeSliderController.h
--- a/Engine/Source/Editor/Sequencer/Public/SequencerTimeSliderController.h
+++ b/Engine/Source/Editor/Sequencer/Public/SequencerTimeSliderController.h
@@ -23,6 +23,7 @@ class FSequencerTimeSliderController : public ITimeSliderController, public TSha
 {
 public:
 	FSequencerTimeSliderController( const FTimeSliderArgs& InArgs, TWeakPtr<FSequencer> InWeakSequencer );
+	~FSequencerTimeSliderController();
 
 	/**
 	* Determines the optimal spacing between tick marks in the slider for a given pixel density
@@ -66,7 +67,7 @@ public:
 	virtual FFrameTime GetScrubPosition() const override { return TimeSliderArgs.ScrubPosition.Get(FFrameTime()); }
 
 	/** Get the current time for the Scrub handle which indicates what range is being evaluated. */
-	virtual void SetScrubPosition(FFrameTime InTime) override { CommitScrubPosition(InTime, false); } 
+	virtual void SetScrubPosition(FFrameTime InTime, bool bEvaluate) override { CommitScrubPosition(InTime, false, bEvaluate); }
 
 	/**
 	 * Clamp the given range to the clamp range 
@@ -171,6 +172,12 @@ public:
 		}
 	};
 
+	/** Set that's evaluating */
+	void SetIsEvaluating()
+	{
+		bIsEvaluating = true;
+	}
+
 private:
 	// forward declared as class members to prevent name collision with similar types defined in other units
 	struct FDrawTickArgs;
@@ -180,8 +187,9 @@ private:
 	 *
 	 * @param NewValue				Value resulting from the user's interaction
 	 * @param bIsScrubbing			True if done via scrubbing, false if just releasing scrubbing
+	 * @param bEvaluate				If true evaluate, if not just change time
 	 */
-	void CommitScrubPosition( FFrameTime NewValue, bool bIsScrubbing );
+	void CommitScrubPosition( FFrameTime NewValue, bool bIsScrubbing, bool bEvaluate);
 
 	/**
 	 * Draw time tick marks
@@ -339,6 +347,9 @@ private:
 	/** When > 0, we should not show context menus */
 	int32 ContextMenuSuppression;
 	
+	/** If evaluating, if not we draw the time box to be yellow not default*/
+	bool bIsEvaluating = true;
+
 	friend FContextMenuSuppressor;
 	
 };
diff --git a/Engine/Source/Editor/SequencerWidgets/Public/ITimeSlider.h b/Engine/Source/Editor/SequencerWidgets/Public/ITimeSlider.h
--- a/Engine/Source/Editor/SequencerWidgets/Public/ITimeSlider.h
+++ b/Engine/Source/Editor/SequencerWidgets/Public/ITimeSlider.h
@@ -43,7 +43,7 @@ enum class ENearestKeyOption : uint8
 };
 ENUM_CLASS_FLAGS(ENearestKeyOption);
 
-DECLARE_DELEGATE_TwoParams( FOnScrubPositionChanged, FFrameTime, bool )
+DECLARE_DELEGATE_ThreeParams( FOnScrubPositionChanged, FFrameTime, bool, bool )
 DECLARE_DELEGATE_TwoParams( FOnViewRangeChanged, TRange<double>, EViewRangeInterpolation )
 DECLARE_DELEGATE_OneParam( FOnTimeRangeChanged, TRange<double> )
 DECLARE_DELEGATE_OneParam( FOnFrameRangeChanged, TRange<FFrameNumber> )
@@ -283,7 +283,7 @@ public:
 	virtual FFrameTime GetScrubPosition() const { return FFrameTime(); }
 
 	/** Set the current time for the Scrub handle which indicates what range is being evaluated. */
-	virtual void SetScrubPosition(FFrameTime InTime) {}
+	virtual void SetScrubPosition(FFrameTime InTime, bool bEvaluate) {}
 
 	/**
 	 * Set a new range based on a min, max and an interpolation mode
