Sequencer: Expose Track Filters to Python/Blueprints.

SetTrackFilterEnabled
IsTrackFilterEnabled
GetTrackFilterNames

diff --git a/Engine/Plugins/Animation/ControlRig/Source/ControlRigEditor/Private/ControlRigBlueprintActions.cpp b/Engine/Plugins/Animation/ControlRig/Source/ControlRigEditor/Private/ControlRigBlueprintActions.cpp
--- a/Engine/Plugins/Animation/ControlRig/Source/ControlRigEditor/Private/ControlRigBlueprintActions.cpp
+++ b/Engine/Plugins/Animation/ControlRig/Source/ControlRigEditor/Private/ControlRigBlueprintActions.cpp
@@ -346,7 +346,7 @@ void FControlRigBlueprintActions::OnSpawnedSkeletalMeshActorChanged(UObject* InO
 				WeakSequencer.Pin()->ThrobSectionSelection();
 				WeakSequencer.Pin()->ObjectImplicitlyAdded(ControlRig);
 				FText Name = LOCTEXT("SequenceTrackFilter_ControlRigControls", "Control Rig Controls");
-				WeakSequencer.Pin()->SetFilterOn(Name, true);
+				WeakSequencer.Pin()->SetTrackFilterEnabled(Name, true);
 				WeakSequencer.Pin()->NotifyMovieSceneDataChanged(EMovieSceneDataChangeType::MovieSceneStructureItemAdded);
 				FControlRigEditMode* ControlRigEditMode = static_cast<FControlRigEditMode*>(GLevelEditorModeTools().GetActiveMode(FControlRigEditMode::ModeName));
 				if (!ControlRigEditMode)
diff --git a/Engine/Plugins/Animation/ControlRig/Source/ControlRigEditor/Private/ControlRigEditorModule.cpp b/Engine/Plugins/Animation/ControlRig/Source/ControlRigEditor/Private/ControlRigEditorModule.cpp
--- a/Engine/Plugins/Animation/ControlRig/Source/ControlRigEditor/Private/ControlRigEditorModule.cpp
+++ b/Engine/Plugins/Animation/ControlRig/Source/ControlRigEditor/Private/ControlRigEditorModule.cpp
@@ -747,7 +747,7 @@ void FControlRigEditorModule::BakeToControlRig(UClass* ControlRigClass, UAnimSeq
 					WeakSequencer.Pin()->ThrobSectionSelection();
 					WeakSequencer.Pin()->ObjectImplicitlyAdded(ControlRig);
 					FText Name = LOCTEXT("SequenceTrackFilter_ControlRigControls", "Control Rig Controls");
-					WeakSequencer.Pin()->SetFilterOn(Name, true);
+					WeakSequencer.Pin()->SetTrackFilterEnabled(Name, true);
 					WeakSequencer.Pin()->NotifyMovieSceneDataChanged(EMovieSceneDataChangeType::MovieSceneStructureItemAdded);
 					FControlRigEditMode* ControlRigEditMode = static_cast<FControlRigEditMode*>(GLevelEditorModeTools().GetActiveMode(FControlRigEditMode::ModeName));
 					if (!ControlRigEditMode)
diff --git a/Engine/Plugins/MovieScene/LevelSequenceEditor/Source/LevelSequenceEditor/Private/LevelSequenceEditorBlueprintLibrary.cpp b/Engine/Plugins/MovieScene/LevelSequenceEditor/Source/LevelSequenceEditor/Private/LevelSequenceEditorBlueprintLibrary.cpp
--- a/Engine/Plugins/MovieScene/LevelSequenceEditor/Source/LevelSequenceEditor/Private/LevelSequenceEditorBlueprintLibrary.cpp
+++ b/Engine/Plugins/MovieScene/LevelSequenceEditor/Source/LevelSequenceEditor/Private/LevelSequenceEditorBlueprintLibrary.cpp
@@ -412,6 +412,39 @@ void ULevelSequenceEditorBlueprintLibrary::SetLockLevelSequence(bool bLock)
 	}
 }
 
+bool ULevelSequenceEditorBlueprintLibrary::IsTrackFilterEnabled(const FText& TrackFilterName)
+{
+	if (CurrentSequencer.IsValid())
+	{
+		TSharedPtr<ISequencer> Sequencer = CurrentSequencer.Pin();
+
+		return Sequencer->IsTrackFilterEnabled(TrackFilterName);
+	}
+	return false;
+}
+
+void ULevelSequenceEditorBlueprintLibrary::SetTrackFilterEnabled(const FText& TrackFilterName, bool bEnabled)
+{
+	if (CurrentSequencer.IsValid())
+	{
+		TSharedPtr<ISequencer> Sequencer = CurrentSequencer.Pin();
+
+		Sequencer->SetTrackFilterEnabled(TrackFilterName, bEnabled);
+	}
+}
+
+TArray<FText> ULevelSequenceEditorBlueprintLibrary::GetTrackFilterNames()
+{
+	if (CurrentSequencer.IsValid())
+	{
+		TSharedPtr<ISequencer> Sequencer = CurrentSequencer.Pin();
+
+		return Sequencer->GetTrackFilterNames();
+	}
+
+	return TArray<FText>();
+}
+
 bool ULevelSequenceEditorBlueprintLibrary::HasCustomColorForChannel(UClass* Class, const FString& Identifier)
 {
 	const UCurveEditorSettings* Settings = GetDefault<UCurveEditorSettings>();
diff --git a/Engine/Plugins/MovieScene/LevelSequenceEditor/Source/LevelSequenceEditor/Public/LevelSequenceEditorBlueprintLibrary.h b/Engine/Plugins/MovieScene/LevelSequenceEditor/Source/LevelSequenceEditor/Public/LevelSequenceEditorBlueprintLibrary.h
--- a/Engine/Plugins/MovieScene/LevelSequenceEditor/Source/LevelSequenceEditor/Public/LevelSequenceEditorBlueprintLibrary.h
+++ b/Engine/Plugins/MovieScene/LevelSequenceEditor/Source/LevelSequenceEditor/Public/LevelSequenceEditorBlueprintLibrary.h
@@ -208,6 +208,21 @@ public:
 	static void SetLockCameraCutToViewport(bool bLock);
 
 public:
+
+	/** Gets whether the specified track filter is on/off */
+	UFUNCTION(BlueprintPure, Category = "Level Sequence Editor")
+	static bool IsTrackFilterEnabled(const FText& TrackFilterName);
+
+	/** Sets the specified track filter to be on or off */
+	UFUNCTION(BlueprintCallable, Category = "Level Sequence Editor")
+	static void SetTrackFilterEnabled(const FText& TrackFilterName, bool bEnabled);
+
+	/** Gets all the available track filter names */
+	UFUNCTION(BlueprintPure, Category = "Level Sequence Editor")
+	static TArray<FText> GetTrackFilterNames();
+
+public:
+
 	/** Get if a custom color for specified channel idendified by it's class and identifier exists */
 	UFUNCTION(BlueprintCallable, Category = "Level Sequence Editor")
 	static bool HasCustomColorForChannel(UClass* Class, const FString& Identifier);
diff --git a/Engine/Source/Editor/Sequencer/Private/SSequencer.cpp b/Engine/Source/Editor/Sequencer/Private/SSequencer.cpp
--- a/Engine/Source/Editor/Sequencer/Private/SSequencer.cpp
+++ b/Engine/Source/Editor/Sequencer/Private/SSequencer.cpp
@@ -1177,6 +1177,16 @@ void SSequencer::BindCommands(TSharedRef<FUICommandList> SequencerCommandBinding
 		FSequencerCommands::Get().OpenNodeGroupsManager,
 		FExecuteAction::CreateSP(this, &SSequencer::OpenNodeGroupsManager)
 	);
+
+	TSharedPtr<FSequencer> Sequencer = SequencerPtr.Pin();
+
+	for (TSharedRef<FSequencerTrackFilter> TrackFilter : AllTrackFilters)
+	{
+		if (TrackFilter->SupportsSequence(Sequencer->GetFocusedMovieSceneSequence()))
+		{
+			TrackFilter->BindCommands(SequencerCommandBindings, SequencerPtr.Pin());
+		}
+	}
 }
 
 void SSequencer::OpenTickResolutionOptions()
@@ -1853,29 +1863,70 @@ void SSequencer::FillLevelFilterMenu(FMenuBuilder& InMenuBarBuilder)
 	}
 }
 
-
-void SSequencer::SetFilterOn(const FText& InName, bool bOn)
+void SSequencer::SetTrackFilterEnabled(const FText& InTrackFilterName, bool bEnabled)
 {
 	TSharedPtr<FSequencer> Sequencer = SequencerPtr.Pin();
 
-
+	bool bFound = false;
 	for (TSharedRef<FSequencerTrackFilter> TrackFilter : AllTrackFilters)
 	{
 		if (TrackFilter->SupportsSequence(Sequencer->GetFocusedMovieSceneSequence()))
 		{
-			if (InName.EqualToCaseIgnored(TrackFilter->GetDisplayName()))
+			if (InTrackFilterName.EqualToCaseIgnored(TrackFilter->GetDisplayName()))
 			{
-				if (bOn && !IsTrackFilterActive(TrackFilter))
+				bFound = true;
+				if (bEnabled && !IsTrackFilterActive(TrackFilter))
 				{
 					Sequencer->GetNodeTree()->AddFilter(TrackFilter);
 				}
-				else if (!bOn && IsTrackFilterActive(TrackFilter))
+				else if (!bEnabled && IsTrackFilterActive(TrackFilter))
 				{
 					Sequencer->GetNodeTree()->RemoveFilter(TrackFilter);
 				}
 			}
 		}
 	}
+
+	if (!bFound)
+	{
+		UE_LOG(LogSequencer, Warning, TEXT("'%s' track filter not found"), *InTrackFilterName.ToString());
+	}
+}
+
+bool SSequencer::IsTrackFilterEnabled(const FText& InTrackFilterName) const
+{
+	TSharedPtr<FSequencer> Sequencer = SequencerPtr.Pin();
+
+	for (TSharedRef<FSequencerTrackFilter> TrackFilter : AllTrackFilters)
+	{
+		if (TrackFilter->SupportsSequence(Sequencer->GetFocusedMovieSceneSequence()))
+		{
+			if (InTrackFilterName.EqualToCaseIgnored(TrackFilter->GetDisplayName()))
+			{
+				return Sequencer->GetNodeTree()->IsTrackFilterActive(TrackFilter);
+			}
+		}
+	}
+
+	UE_LOG(LogSequencer, Warning, TEXT("'%s' track filter not found"), *InTrackFilterName.ToString());
+
+	return false;
+}
+
+TArray<FText> SSequencer::GetTrackFilterNames() const
+{
+	TSharedPtr<FSequencer> Sequencer = SequencerPtr.Pin();
+
+	TArray<FText> Filters;
+	for (TSharedRef<FSequencerTrackFilter> TrackFilter : AllTrackFilters)
+	{
+		if (TrackFilter->SupportsSequence(Sequencer->GetFocusedMovieSceneSequence()))
+		{
+			Filters.Add(TrackFilter->GetDisplayName());
+		}
+	}
+
+	return Filters;
 }
 
 void SSequencer::OnResetFilters()
diff --git a/Engine/Source/Editor/Sequencer/Private/SSequencer.h b/Engine/Source/Editor/Sequencer/Private/SSequencer.h
--- a/Engine/Source/Editor/Sequencer/Private/SSequencer.h
+++ b/Engine/Source/Editor/Sequencer/Private/SSequencer.h
@@ -347,8 +347,14 @@ public:
 	/** Sets the play time for the sequence. Will extend the working range if out of bounds. */
 	void SetPlayTime(double Frame);
 
-	/** Sets the specified filter to be on or off*/
-	void SetFilterOn(const FText& InName, bool bOn);
+	/** Sets the specified track filter to be on or off */
+	void SetTrackFilterEnabled(const FText& InTrackFilterName, bool bEnabled);
+
+	/** Gets whether the specified track filter is on/off */
+	bool IsTrackFilterEnabled(const FText& InTrackFilterName) const;
+
+	/** Gets all the available track filter names */
+	TArray<FText> GetTrackFilterNames() const;
 
 	/** Sets the text to search by */
 	void SetSearchText(const FText& InSearchText);
diff --git a/Engine/Source/Editor/Sequencer/Private/Sequencer.cpp b/Engine/Source/Editor/Sequencer/Private/Sequencer.cpp
--- a/Engine/Source/Editor/Sequencer/Private/Sequencer.cpp
+++ b/Engine/Source/Editor/Sequencer/Private/Sequencer.cpp
@@ -9737,11 +9737,20 @@ void FSequencer::ObjectImplicitlyRemoved(UObject* InObject) const
 	}
 }
 
-void FSequencer::SetFilterOn(const FText& InName, bool bOn)
+void FSequencer::SetTrackFilterEnabled(const FText& InTrackFilterName, bool bEnabled)
 {
-	SequencerWidget->SetFilterOn(InName, bOn);
+	SequencerWidget->SetTrackFilterEnabled(InTrackFilterName, bEnabled);
 }
 
+bool FSequencer::IsTrackFilterEnabled(const FText& InTrackFilterName) const
+{
+	return SequencerWidget->IsTrackFilterEnabled(InTrackFilterName);
+}
+
+TArray<FText> FSequencer::GetTrackFilterNames() const
+{
+	return SequencerWidget->GetTrackFilterNames();
+}
 
 void FSequencer::ImportSectionsFromText(const FString& TextToImport, /*out*/ TArray<UMovieSceneSection*>& ImportedSections)
 {
diff --git a/Engine/Source/Editor/Sequencer/Private/Sequencer.h b/Engine/Source/Editor/Sequencer/Private/Sequencer.h
--- a/Engine/Source/Editor/Sequencer/Private/Sequencer.h
+++ b/Engine/Source/Editor/Sequencer/Private/Sequencer.h
@@ -847,7 +847,9 @@ public:
 	virtual void SetSequencerSettings(USequencerSettings* InSettings) override { Settings = InSettings; }
 	virtual TSharedPtr<class ITimeSlider> GetTopTimeSliderWidget() const override;
 	virtual void ResetTimeController() override;
-	virtual void SetFilterOn(const FText& InName, bool bOn) override;
+	virtual void SetTrackFilterEnabled(const FText& InTrackFilterName, bool bEnabled) override;
+	virtual bool IsTrackFilterEnabled(const FText& InTrackFilterName) const override;
+	virtual TArray<FText> GetTrackFilterNames() const override;
 
 public:
 
diff --git a/Engine/Source/Editor/Sequencer/Public/ISequencer.h b/Engine/Source/Editor/Sequencer/Public/ISequencer.h
--- a/Engine/Source/Editor/Sequencer/Public/ISequencer.h
+++ b/Engine/Source/Editor/Sequencer/Public/ISequencer.h
@@ -717,12 +717,16 @@ public:
 	virtual void ObjectImplicitlyRemoved(UObject* InObject) const = 0;
 
 public:
-	/**
-	*    Turn on/off the filter with the specified name
-	* @InName The name of the of the filter
-	* @bOn   Whether or not the filter is on or off.
-	*/
-	virtual void SetFilterOn(const FText& InName, bool bOn) = 0;
+
+	/** Sets the specified track filter to be on or off */
+	virtual void SetTrackFilterEnabled(const FText& InTrackFilterName, bool bEnabled) = 0;
+
+	/** Gets whether the specified track filter is on/off */
+	virtual bool IsTrackFilterEnabled(const FText& InTrackFilterName) const = 0;
+
+	/** Gets all the available track filter names */
+	virtual TArray<FText> GetTrackFilterNames() const = 0;
+
 public:
 
 	/**
