Sequencer: Add shortcut U for enable/disable "Animated Tracks" filter

Allow Track filters to bind commands
Enable "Animated Tracks" filter for UMG

diff --git a/Engine/Source/Editor/Sequencer/Private/SequencerTrackFilters.cpp b/Engine/Source/Editor/Sequencer/Private/SequencerTrackFilters.cpp
--- a/Engine/Source/Editor/Sequencer/Private/SequencerTrackFilters.cpp
+++ b/Engine/Source/Editor/Sequencer/Private/SequencerTrackFilters.cpp
@@ -1,7 +1,11 @@
 // Copyright Epic Games, Inc. All Rights Reserved.
 
 #include "SequencerTrackFilters.h"
+#include "EditorStyleSet.h"
 #include "Engine/World.h"
+#include "ISequencer.h"
+
+#define LOCTEXT_NAMESPACE "Sequencer"
 
 FSequencerTrackFilter_LevelFilter::~FSequencerTrackFilter_LevelFilter()
 {
@@ -115,4 +119,67 @@ void FSequencerTrackFilter_LevelFilter::HandleLevelsChanged()
 	{
 		BroadcastChangedEvent();
 	}
-}
\ No newline at end of file
+}
+
+class FSequencerTrackFilter_AnimatedCommands
+	: public TCommands<FSequencerTrackFilter_AnimatedCommands>
+{
+public:
+
+	FSequencerTrackFilter_AnimatedCommands()
+		: TCommands<FSequencerTrackFilter_AnimatedCommands>
+	(
+		"FSequencerTrackFilter_Animated",
+		NSLOCTEXT("Contexts", "FSequencerTrackFilter_Animated", "FSequencerTrackFilter_Animated"),
+		NAME_None,
+		FEditorStyle::GetStyleSetName() // Icon Style Set
+	)
+	{ }
+		
+	/** Toggle the animated tracks filter */
+	TSharedPtr< FUICommandInfo > ToggleAnimatedTracks;
+
+	/** Initialize commands */
+	virtual void RegisterCommands() override
+	{
+		UI_COMMAND(ToggleAnimatedTracks, "Animated Tracks", "Toggle the filter for Animated Tracks.", EUserInterfaceActionType::ToggleButton, FInputChord(EKeys::U));
+	}
+};
+
+FSequencerTrackFilter_Animated::FSequencerTrackFilter_Animated()
+{
+	FSequencerTrackFilter_AnimatedCommands::Register();	
+}
+
+FSequencerTrackFilter_Animated::~FSequencerTrackFilter_Animated()
+{
+	FSequencerTrackFilter_AnimatedCommands::Unregister();	
+}
+
+FText FSequencerTrackFilter_Animated::GetToolTipText() const 
+{ 
+	const FSequencerTrackFilter_AnimatedCommands& Commands = FSequencerTrackFilter_AnimatedCommands::Get();
+
+	const TSharedRef<const FInputChord> FirstActiveChord = Commands.ToggleAnimatedTracks->GetFirstValidChord();
+	
+	FText Tooltip = LOCTEXT("SequencerTrackFilter_AnimatedTip", "Show Only Animated Tracks."); 
+
+	if (FirstActiveChord->IsValidChord())
+	{
+		return FText::Join(FText::FromString(TEXT(" ")), Tooltip, FirstActiveChord->GetInputText());
+	}
+	return Tooltip;
+}
+
+void FSequencerTrackFilter_Animated::BindCommands(TSharedRef<FUICommandList> CommandBindings, TWeakPtr<ISequencer> Sequencer)
+{
+	const FSequencerTrackFilter_AnimatedCommands& Commands = FSequencerTrackFilter_AnimatedCommands::Get();
+
+	CommandBindings->MapAction(
+		Commands.ToggleAnimatedTracks,
+		FExecuteAction::CreateLambda( [this, Sequencer]{ Sequencer.Pin()->SetTrackFilterEnabled(GetDisplayName(), !Sequencer.Pin()->IsTrackFilterEnabled(GetDisplayName())); } ),
+		FCanExecuteAction::CreateLambda( [this, Sequencer]{ return true; } ),
+		FIsActionChecked::CreateLambda( [this, Sequencer]{ return Sequencer.Pin()->IsTrackFilterEnabled(GetDisplayName()); } ) );
+}
+
+#undef LOCTEXT_NAMESPACE
diff --git a/Engine/Source/Editor/Sequencer/Private/SequencerTrackFilters.h b/Engine/Source/Editor/Sequencer/Private/SequencerTrackFilters.h
--- a/Engine/Source/Editor/Sequencer/Private/SequencerTrackFilters.h
+++ b/Engine/Source/Editor/Sequencer/Private/SequencerTrackFilters.h
@@ -320,9 +320,13 @@ private:
 
 class FSequencerTrackFilter_Animated : public FSequencerTrackFilter
 {
+public:
+	FSequencerTrackFilter_Animated();
+	~FSequencerTrackFilter_Animated();
+
 	virtual FString GetName() const override { return TEXT("AnimatedFilter"); }
 	virtual FText GetDisplayName() const override { return LOCTEXT("SequenceTrackFilter_Animated", "Animated Tracks"); }
-	virtual FText GetToolTipText() const override { return LOCTEXT("SequencerTrackFilter_AnimatedTip", "Show Only Animated Tracks."); }
+	virtual FText GetToolTipText() const override;
 	virtual FSlateIcon GetIcon() const { return FSlateIcon(FEditorStyle::GetStyleSetName(), "Sequencer.IconKeyUser"); }
 
 	virtual bool ShouldUpdateOnTrackValueChanged() const override
@@ -340,5 +344,15 @@ class FSequencerTrackFilter_Animated : public FSequencerTrackFilter
 		return (InMovieSceneChannel && InMovieSceneChannel->GetNumKeys() > 0);
 	}
 
+	virtual bool SupportsSequence(UMovieSceneSequence* InSequence) const override
+	{
+		static UClass* LevelSequenceClass = FindObject<UClass>(ANY_PACKAGE, TEXT("LevelSequence"), true);
+		static UClass* WidgetAnimationClass = FindObject<UClass>(ANY_PACKAGE, TEXT("WidgetAnimation"), true);
+		return InSequence != nullptr &&
+			((LevelSequenceClass != nullptr && InSequence->GetClass()->IsChildOf(LevelSequenceClass)) ||
+			(WidgetAnimationClass != nullptr && InSequence->GetClass()->IsChildOf(WidgetAnimationClass)));
+	}
+	
+	virtual void BindCommands(TSharedRef<FUICommandList> CommandBindings, TWeakPtr<ISequencer> Sequencer) override;
 };
 #undef LOCTEXT_NAMESPACE
\ No newline at end of file
diff --git a/Engine/Source/Editor/Sequencer/Public/SequencerTrackFilterBase.h b/Engine/Source/Editor/Sequencer/Public/SequencerTrackFilterBase.h
--- a/Engine/Source/Editor/Sequencer/Public/SequencerTrackFilterBase.h
+++ b/Engine/Source/Editor/Sequencer/Public/SequencerTrackFilterBase.h
@@ -9,6 +9,8 @@
 #include "MovieSceneSequence.h"
 
 class FMenuBuilder;
+class FUICommandList;
+class ISequencer;
 
 typedef const UObject* FTrackFilterType;
 
@@ -51,6 +53,8 @@ public:
 
 	virtual bool PassesFilterChannel(FMovieSceneChannel* InMovieSceneChannel) const { return false; }
 
+	virtual void BindCommands(TSharedRef<FUICommandList> CommandBindings, TWeakPtr<ISequencer> Sequencer) {}
+
 protected:
 	void BroadcastChangedEvent() const { ChangedEvent.Broadcast(); }
 
