Sequencer: Fix include

diff --git a/Engine/Source/Editor/Sequencer/Private/SequencerTrackFilters.cpp b/Engine/Source/Editor/Sequencer/Private/SequencerTrackFilters.cpp
--- a/Engine/Source/Editor/Sequencer/Private/SequencerTrackFilters.cpp
+++ b/Engine/Source/Editor/Sequencer/Private/SequencerTrackFilters.cpp
@@ -3,6 +3,7 @@
 #include "SequencerTrackFilters.h"
 #include "EditorStyleSet.h"
 #include "Engine/World.h"
+#include "Framework/Commands/Commands.h"
 #include "ISequencer.h"
 
 #define LOCTEXT_NAMESPACE "Sequencer"
