Sequencer: Disable play rate menu if locked so that it's not confusing when you try to set a value and it doesn't stick.

diff --git a/Engine/Source/Editor/Sequencer/Private/SSequencerPlayRateCombo.cpp b/Engine/Source/Editor/Sequencer/Private/SSequencerPlayRateCombo.cpp
--- a/Engine/Source/Editor/Sequencer/Private/SSequencerPlayRateCombo.cpp
+++ b/Engine/Source/Editor/Sequencer/Private/SSequencerPlayRateCombo.cpp
@@ -217,6 +217,12 @@ FText SSequencerPlayRateCombo::GetFrameRateMismatchErrorDescription() const
 	return FText();
 }
 
+bool SSequencerPlayRateCombo::GetIsSequenceReadOnly() const
+{
+	TSharedPtr<FSequencer> Sequencer = WeakSequencer.Pin();
+	return !Sequencer.IsValid() || Sequencer->IsReadOnly();
+}
+
 TSharedRef<SWidget> SSequencerPlayRateCombo::OnCreateMenu()
 {
 	TSharedPtr<FSequencer> Sequencer       = WeakSequencer.Pin();
@@ -261,6 +267,7 @@ TSharedRef<SWidget> SSequencerPlayRateCombo::OnCreateMenu()
 				SNew(SFrameRateEntryBox)
 				.Value(this, &SSequencerPlayRateCombo::GetDisplayRate)
 				.OnValueChanged(this, &SSequencerPlayRateCombo::SetDisplayRate)
+				.IsEnabled_Lambda([=] { return !GetIsSequenceReadOnly(); })
 			],
 			LOCTEXT("CustomFramerateDisplayLabel", "Custom")
 		);
@@ -304,7 +311,7 @@ TSharedRef<SWidget> SSequencerPlayRateCombo::OnCreateMenu()
 		FSlateIcon(),
 		FUIAction(
 			FExecuteAction::CreateSP(this, &SSequencerPlayRateCombo::OnToggleFrameLocked),
-			FCanExecuteAction(),
+			FCanExecuteAction::CreateLambda([this] { return !GetIsSequenceReadOnly(); } ),
 			FGetActionCheckState::CreateSP(this, &SSequencerPlayRateCombo::OnGetFrameLockedCheckState)
 		),
 		NAME_None,
@@ -378,7 +385,7 @@ void SSequencerPlayRateCombo::PopulateClockSourceMenu(FMenuBuilder& MenuBuilder)
 						FNewMenuDelegate::CreateSP(this, &SSequencerPlayRateCombo::PopulateCustomClockSourceMenu),
 						FUIAction(
 							FExecuteAction::CreateSP(this, &SSequencerPlayRateCombo::SetClockSource, Value),
-							FCanExecuteAction(),
+							FCanExecuteAction::CreateLambda([this]{ return !GetIsSequenceReadOnly(); }),
 							FIsActionChecked::CreateLambda([=]{ return RootSequence->GetMovieScene()->GetClockSource() == Value; })
 						),
 						NAME_None,
@@ -393,7 +400,7 @@ void SSequencerPlayRateCombo::PopulateClockSourceMenu(FMenuBuilder& MenuBuilder)
 						FSlateIcon(),
 						FUIAction(
 							FExecuteAction::CreateSP(this, &SSequencerPlayRateCombo::SetClockSource, Value),
-							FCanExecuteAction(),
+							FCanExecuteAction::CreateLambda([this]{ return !GetIsSequenceReadOnly(); }),
 							FIsActionChecked::CreateLambda([=]{ return RootSequence->GetMovieScene()->GetClockSource() == Value; })
 						),
 						NAME_None,
@@ -478,7 +485,7 @@ void SSequencerPlayRateCombo::AddMenuEntry(FMenuBuilder& MenuBuilder, const FCom
 		FSlateIcon(),
 		FUIAction(
 			FExecuteAction::CreateSP(this, &SSequencerPlayRateCombo::SetDisplayRate, Info.FrameRate),
-			FCanExecuteAction(),
+			FCanExecuteAction::CreateLambda([this]{ return !GetIsSequenceReadOnly(); }),
 			FIsActionChecked::CreateSP(this, &SSequencerPlayRateCombo::IsSameDisplayRate, Info.FrameRate)
 		),
 		NAME_None,
diff --git a/Engine/Source/Editor/Sequencer/Private/SSequencerPlayRateCombo.h b/Engine/Source/Editor/Sequencer/Private/SSequencerPlayRateCombo.h
--- a/Engine/Source/Editor/Sequencer/Private/SSequencerPlayRateCombo.h
+++ b/Engine/Source/Editor/Sequencer/Private/SSequencerPlayRateCombo.h
@@ -69,6 +69,8 @@ private:
 	EVisibility GetClockSourceVisibility() const;
 	const FSlateBrush* GetClockSourceImage() const;
 
+	bool GetIsSequenceReadOnly() const;
+
 	/** Sequencer pointer */
 	TWeakPtr<FSequencer> WeakSequencer;
 	TWeakPtr<SSequencer> WeakSequencerWidget;
