Sequencer: No Longer evalute when begging a drag, would cause the middle mouse drag to re-evaluate when dragging.
#trivial

diff --git a/Engine/Source/Editor/Sequencer/Private/Sequencer.cpp b/Engine/Source/Editor/Sequencer/Private/Sequencer.cpp
--- a/Engine/Source/Editor/Sequencer/Private/Sequencer.cpp
+++ b/Engine/Source/Editor/Sequencer/Private/Sequencer.cpp
@@ -5189,9 +5189,6 @@ void FSequencer::OnScrubPositionChanged( FFrameTime NewScrubPosition, bool bScru
 
 void FSequencer::OnBeginScrubbing()
 {
-	// Pause first since there's no explicit evaluation in the stopped state when OnEndScrubbing() is called
-	Pause();
-
 	SetPlaybackStatus(EMovieScenePlayerStatus::Scrubbing);
 	SequencerWidget->RegisterActiveTimerForPlayback();
 
