Sequencer: Display the bound object class in tooltip

diff --git a/Engine/Source/Editor/Sequencer/Private/DisplayNodes/SequencerObjectBindingNode.cpp b/Engine/Source/Editor/Sequencer/Private/DisplayNodes/SequencerObjectBindingNode.cpp
--- a/Engine/Source/Editor/Sequencer/Private/DisplayNodes/SequencerObjectBindingNode.cpp
+++ b/Engine/Source/Editor/Sequencer/Private/DisplayNodes/SequencerObjectBindingNode.cpp
@@ -807,6 +807,7 @@ FText FSequencerObjectBindingNode::GetDisplayNameToolTipText() const
 	else
 	{
 		TArray<FString> ValidBoundObjectLabels;
+		FName BoundObjectClass;
 		bool bAddEllipsis = false;
 		int32 NumMissing = 0;
 		for (const TWeakObjectPtr<>& Ptr : BoundObjects)
@@ -819,6 +820,11 @@ FText FSequencerObjectBindingNode::GetDisplayNameToolTipText() const
 				continue;
 			}
 
+			if (Obj->GetClass())
+			{
+				BoundObjectClass = Obj->GetClass()->GetFName();
+			}
+
 			if (AActor* Actor = Cast<AActor>(Obj))
 			{
 				ValidBoundObjectLabels.Add(Actor->GetActorLabel());
@@ -846,7 +852,7 @@ FText FSequencerObjectBindingNode::GetDisplayNameToolTipText() const
 					return FText::Format(LOCTEXT("SpawnableBoundObjectToolTip", "Spawnable Class: {0} (BindingID: {1})"), FText::FromName(ClassForObjectBinding->GetFName()), FText::FromString(LexToString(ObjectBinding)));
 				}
 			}
-			return FText::Format(LOCTEXT("PossessableBoundObjectToolTip", "(BindingID: {0}"), FText::FromString(LexToString(ObjectBinding)));
+			return FText::Format(LOCTEXT("PossessableBoundObjectToolTip", "Class: {0} (BindingID: {1}"), FText::FromName(BoundObjectClass), FText::FromString(LexToString(ObjectBinding)));
 		}
 		else if (ValidBoundObjectLabels.Num() == 0 && NumMissing == 1)
 		{
