Sequencer: Enhance the display rate mismatch tooltip with the names of the level sequences that are not matching.

diff --git a/Engine/Source/Editor/Sequencer/Private/SSequencerPlayRateCombo.cpp b/Engine/Source/Editor/Sequencer/Private/SSequencerPlayRateCombo.cpp
--- a/Engine/Source/Editor/Sequencer/Private/SSequencerPlayRateCombo.cpp
+++ b/Engine/Source/Editor/Sequencer/Private/SSequencerPlayRateCombo.cpp
@@ -194,6 +194,8 @@ FText SSequencerPlayRateCombo::GetFrameRateMismatchErrorDescription() const
 			return FText();
 		}
 
+		TArray<FText> SubSequenceDisplayRates;
+
 		for (const TTuple<FMovieSceneSequenceID, FMovieSceneSubSequenceData>& Pair : Hierarchy->AllSubSequenceData())
 		{
 			UMovieSceneSequence* SubSequence = Pair.Value.GetSequence();
@@ -203,15 +205,25 @@ FText SSequencerPlayRateCombo::GetFrameRateMismatchErrorDescription() const
 			{
 				FFrameRate SubDisplayRate = MovieScene->GetDisplayRate();
 
-				const FCommonFrameRateInfo* DisplayRateInfo = FCommonFrameRates::Find(DisplayRate);
 				const FCommonFrameRateInfo* SubDisplayRateInfo = FCommonFrameRates::Find(SubDisplayRate);
 
-				FText DisplayRateText = DisplayRateInfo ? DisplayRateInfo->DisplayName : FText::Format(LOCTEXT("DisplayRateFormat", "{0} fps"), DisplayRate.AsDecimal());
 				FText SubDisplayRateText = SubDisplayRateInfo ? SubDisplayRateInfo->DisplayName : FText::Format(LOCTEXT("SubDisplayRateFormat", "{0} fps"), SubDisplayRate.AsDecimal());
 
-				return FText::Format(LOCTEXT("FrameRateMismatchDescription", "At least one mismatch in display rate: {0} is at {1} and {2} is at {3}"), Sequencer->GetRootMovieSceneSequence()->GetDisplayName(), DisplayRateText, SubSequence->GetDisplayName(), SubDisplayRateText);
+				FText SubSequenceDescription = FText::Format(LOCTEXT("SubSequenceFrameRateMismatchDescription", "\t{0} is at {1}"), SubSequence->GetDisplayName(), SubDisplayRateText);
+				SubSequenceDisplayRates.Add(SubSequenceDescription);
 			}
 		}
+
+		if (SubSequenceDisplayRates.Num() != 0)
+		{
+			const FCommonFrameRateInfo* DisplayRateInfo = FCommonFrameRates::Find(DisplayRate);
+			FText DisplayRateText = DisplayRateInfo ? DisplayRateInfo->DisplayName : FText::Format(LOCTEXT("DisplayRateFormat", "{0} fps"), DisplayRate.AsDecimal());		
+		
+			FText Description = FText::Format(LOCTEXT("FrameRateMismatchDescription", "Mismatch in display rate: {0} is at {1}"), Sequencer->GetRootMovieSceneSequence()->GetDisplayName(), DisplayRateText);
+			SubSequenceDisplayRates.Insert(Description, 0);
+
+			return FText::Join(FText::FromString(TEXT("\n")), SubSequenceDisplayRates);
+		}
 	}
 
 	return FText();
