PR #8888: Fixing Sequencer scrub bar to more fully respect the "ReadOnly" status of Sequencer (Contributed by [REDACTED_PERSON])

Lock button shows blue/on when locked
Disable adding/moving marks when locked
Disable playback range and mark menu options when locked
Selection is still allowed

diff --git a/Engine/Source/Editor/Sequencer/Private/SSequencer.cpp b/Engine/Source/Editor/Sequencer/Private/SSequencer.cpp
--- a/Engine/Source/Editor/Sequencer/Private/SSequencer.cpp
+++ b/Engine/Source/Editor/Sequencer/Private/SSequencer.cpp
@@ -709,7 +709,8 @@ void SSequencer::Construct(const FArguments& InArgs, TSharedRef<FSequencer> InSe
 						.Style(FAppStyle::Get(),"ToggleButtonCheckBoxAlt")
 						.Type(ESlateCheckBoxType::CheckBox) // Use CheckBox instead of ToggleType since we're not putting ohter widget inside
 						.Padding(FMargin(0.f))
-						.IsFocusable(false)		
+						.IsFocusable(false)
+						.IsEnabled(!SequencerPtr.Pin()->IsReadOnly())		
 						.IsChecked_Lambda([this] { return GetIsSequenceReadOnly() ? ECheckBoxState::Checked : ECheckBoxState::Unchecked; } )
 						.OnCheckStateChanged(this, &SSequencer::OnSetSequenceReadOnly)
 						.ToolTipText_Lambda([this] { return GetIsSequenceReadOnly() ? LOCTEXT("UnlockSequence", "Unlock the animation so that it is editable") : LOCTEXT("LockSequence", "Lock the animation so that it is not editable"); } )
diff --git a/Engine/Source/Editor/Sequencer/Private/Sequencer.cpp b/Engine/Source/Editor/Sequencer/Private/Sequencer.cpp
--- a/Engine/Source/Editor/Sequencer/Private/Sequencer.cpp
+++ b/Engine/Source/Editor/Sequencer/Private/Sequencer.cpp
@@ -2729,6 +2729,11 @@ void FSequencer::SetPlaybackRangeToAllShots()
 
 bool FSequencer::IsPlaybackRangeLocked() const
 {
+	if (IsReadOnly())
+	{
+		return true;
+	}
+	
 	UMovieSceneSequence* FocusedMovieSceneSequence = GetFocusedMovieSceneSequence();
 	if (FocusedMovieSceneSequence != nullptr)
 	{
@@ -11347,129 +11352,187 @@ void FSequencer::ClearGlobalMarkedFrames()
 void FSequencer::ToggleMarkAtPlayPosition()
 {
 	UMovieSceneSequence* FocusedMovieSequence = GetFocusedMovieSceneSequence();
-	if (FocusedMovieSequence != nullptr)
+	if (!FocusedMovieSequence)
 	{
-		UMovieScene* FocusedMovieScene = FocusedMovieSequence->GetMovieScene();
-		if (FocusedMovieScene != nullptr)
-		{
-			FFrameNumber TickFrameNumber = GetLocalTime().Time.FloorToFrame();
-			int32 MarkedFrameIndex = FocusedMovieScene->FindMarkedFrameByFrameNumber(TickFrameNumber);
-			if (MarkedFrameIndex != INDEX_NONE)
-			{
-				FScopedTransaction DeleteMarkedFrameTransaction(LOCTEXT("DeleteMarkedFrames_Transaction", "Delete Marked Frame"));
+		return;
+	}
+	
+	UMovieScene* FocusedMovieScene = FocusedMovieSequence->GetMovieScene();
+	if (!FocusedMovieScene)
+	{
+		return;
+	}
 
-				FocusedMovieScene->Modify();
-				FocusedMovieScene->DeleteMarkedFrame(MarkedFrameIndex);
-			}
-			else
-			{
-				FScopedTransaction AddMarkedFrameTransaction(LOCTEXT("AddMarkedFrame_Transaction", "Add Marked Frame"));
+	if (FocusedMovieScene->IsReadOnly())
+	{
+		FSequencerUtilities::ShowReadOnlyError();
+		return;
+	}
 
-				FocusedMovieScene->Modify();
-				FocusedMovieScene->AddMarkedFrame(FMovieSceneMarkedFrame(TickFrameNumber));
-			}
-		}
+	FFrameNumber TickFrameNumber = GetLocalTime().Time.FloorToFrame();
+	int32 MarkedFrameIndex = FocusedMovieScene->FindMarkedFrameByFrameNumber(TickFrameNumber);
+	if (MarkedFrameIndex != INDEX_NONE)
+	{
+		FScopedTransaction DeleteMarkedFrameTransaction(LOCTEXT("DeleteMarkedFrames_Transaction", "Delete Marked Frame"));
+
+		FocusedMovieScene->Modify();
+		FocusedMovieScene->DeleteMarkedFrame(MarkedFrameIndex);
+	}
+	else
+	{
+		FScopedTransaction AddMarkedFrameTransaction(LOCTEXT("AddMarkedFrame_Transaction", "Add Marked Frame"));
+
+		FocusedMovieScene->Modify();
+		FocusedMovieScene->AddMarkedFrame(FMovieSceneMarkedFrame(TickFrameNumber));
 	}
 }
 
 void FSequencer::SetMarkedFrame(int32 InMarkIndex, FFrameNumber InFrameNumber)
 {
 	UMovieSceneSequence* FocusedMovieSequence = GetFocusedMovieSceneSequence();
-	if (FocusedMovieSequence != nullptr)
+	if (!FocusedMovieSequence)
 	{
-		UMovieScene* FocusedMovieScene = FocusedMovieSequence->GetMovieScene();
-		if (FocusedMovieScene != nullptr)
-		{
-			FocusedMovieScene->Modify();
-			FocusedMovieScene->SetMarkedFrame(InMarkIndex, InFrameNumber);
-		}
+		return;
+	}
+	
+	UMovieScene* FocusedMovieScene = FocusedMovieSequence->GetMovieScene();
+	if (!FocusedMovieScene)
+	{
+		return;
 	}
+
+	if (FocusedMovieScene->IsReadOnly())
+	{
+		FSequencerUtilities::ShowReadOnlyError();
+		return;
+	}
+			
+	FocusedMovieScene->Modify();
+	FocusedMovieScene->SetMarkedFrame(InMarkIndex, InFrameNumber);
 }
 
 void FSequencer::AddMarkedFrame(FFrameNumber FrameNumber)
 {
 	UMovieSceneSequence* FocusedMovieSequence = GetFocusedMovieSceneSequence();
-	if (FocusedMovieSequence != nullptr)
+	if (!FocusedMovieSequence)
 	{
-		UMovieScene* FocusedMovieScene = FocusedMovieSequence->GetMovieScene();
-		if (FocusedMovieScene != nullptr)
-		{
-			FScopedTransaction AddMarkedFrameTransaction(LOCTEXT("AddMarkedFrame_Transaction", "Add Marked Frame"));
+		return;
+	}
+	
+	UMovieScene* FocusedMovieScene = FocusedMovieSequence->GetMovieScene();
+	if (!FocusedMovieScene)
+	{
+		return;
+	}
 
-			FocusedMovieScene->Modify();
-			FocusedMovieScene->AddMarkedFrame(FMovieSceneMarkedFrame(FrameNumber));
-		}
+	if (FocusedMovieScene->IsReadOnly())
+	{
+		FSequencerUtilities::ShowReadOnlyError();
+		return;
 	}
+
+	FScopedTransaction AddMarkedFrameTransaction(LOCTEXT("AddMarkedFrame_Transaction", "Add Marked Frame"));
+
+	FocusedMovieScene->Modify();
+	FocusedMovieScene->AddMarkedFrame(FMovieSceneMarkedFrame(FrameNumber));
 }
 
 void FSequencer::DeleteMarkedFrame(int32 InMarkIndex)
 {
 	UMovieSceneSequence* FocusedMovieSequence = GetFocusedMovieSceneSequence();
-	if (FocusedMovieSequence != nullptr)
+	if (!FocusedMovieSequence)
 	{
-		UMovieScene* FocusedMovieScene = FocusedMovieSequence->GetMovieScene();
-		if (FocusedMovieScene != nullptr)
-		{
-			if (InMarkIndex != INDEX_NONE)
-			{
-				FScopedTransaction DeleteMarkedFrameTransaction(LOCTEXT("DeleteMarkedFrame_Transaction", "Delete Marked Frame"));
+		return;
+	}
+	
+	UMovieScene* FocusedMovieScene = FocusedMovieSequence->GetMovieScene();
+	if (!FocusedMovieScene)
+	{
+		return;
+	}
 
-				FocusedMovieScene->Modify();
-				FocusedMovieScene->DeleteMarkedFrame(InMarkIndex);
-			}
-		}
+	if (FocusedMovieScene->IsReadOnly())
+	{
+		FSequencerUtilities::ShowReadOnlyError();
+		return;
+	}
+
+	if (InMarkIndex != INDEX_NONE)
+	{
+		FScopedTransaction DeleteMarkedFrameTransaction(LOCTEXT("DeleteMarkedFrame_Transaction", "Delete Marked Frame"));
+
+		FocusedMovieScene->Modify();
+		FocusedMovieScene->DeleteMarkedFrame(InMarkIndex);
 	}
 }
 
 void FSequencer::DeleteAllMarkedFrames()
 {
 	UMovieSceneSequence* FocusedMovieSequence = GetFocusedMovieSceneSequence();
-	if (FocusedMovieSequence != nullptr)
+	if (!FocusedMovieSequence)
 	{
-		UMovieScene* FocusedMovieScene = FocusedMovieSequence->GetMovieScene();
-		if (FocusedMovieScene != nullptr)
-		{
-			FScopedTransaction DeleteAllMarkedFramesTransaction(LOCTEXT("DeleteAllMarkedFrames_Transaction", "Delete All Marked Frames"));
+		return;
+	}
+	
+	UMovieScene* FocusedMovieScene = FocusedMovieSequence->GetMovieScene();
+	if (!FocusedMovieScene)
+	{
+		return;
+	}
 
-			FocusedMovieScene->Modify();
-			FocusedMovieScene->DeleteMarkedFrames();
-		}
+	if (FocusedMovieScene->IsReadOnly())
+	{
+		FSequencerUtilities::ShowReadOnlyError();
+		return;
 	}
+			
+	FScopedTransaction DeleteAllMarkedFramesTransaction(LOCTEXT("DeleteAllMarkedFrames_Transaction", "Delete All Marked Frames"));
+
+	FocusedMovieScene->Modify();
+	FocusedMovieScene->DeleteMarkedFrames();
 }
 
 void FSequencer::StepToNextMark()
 {
 	UMovieSceneSequence* FocusedMovieSequence = GetFocusedMovieSceneSequence();
-	if (FocusedMovieSequence != nullptr)
+	if (!FocusedMovieSequence)
 	{
-		UMovieScene* FocusedMovieScene = FocusedMovieSequence->GetMovieScene();
-		if (FocusedMovieScene != nullptr)
-		{
-			const bool bForwards = true;
-			int32 MarkedIndex = FocusedMovieScene->FindNextMarkedFrame(GetLocalTime().Time.FloorToFrame(), bForwards);
-			if (MarkedIndex != INDEX_NONE)
-			{
-				AutoScrubToTime(FocusedMovieScene->GetMarkedFrames()[MarkedIndex].FrameNumber.Value);
-			}
-		}
+		return;
+	}
+	
+	UMovieScene* FocusedMovieScene = FocusedMovieSequence->GetMovieScene();
+	if (!FocusedMovieScene)
+	{
+		return;
+	}
+
+	const bool bForwards = true;
+	int32 MarkedIndex = FocusedMovieScene->FindNextMarkedFrame(GetLocalTime().Time.FloorToFrame(), bForwards);
+	if (MarkedIndex != INDEX_NONE)
+	{
+		AutoScrubToTime(FocusedMovieScene->GetMarkedFrames()[MarkedIndex].FrameNumber.Value);
 	}
 }
 
 void FSequencer::StepToPreviousMark()
 {
 	UMovieSceneSequence* FocusedMovieSequence = GetFocusedMovieSceneSequence();
-	if (FocusedMovieSequence != nullptr)
+	if (!FocusedMovieSequence)
 	{
-		UMovieScene* FocusedMovieScene = FocusedMovieSequence->GetMovieScene();
-		if (FocusedMovieScene != nullptr)
-		{
-			const bool bForwards = false;
-			int32 MarkedIndex = FocusedMovieScene->FindNextMarkedFrame(GetLocalTime().Time.FloorToFrame(), bForwards);
-			if (MarkedIndex != INDEX_NONE)
-			{
-				AutoScrubToTime(FocusedMovieScene->GetMarkedFrames()[MarkedIndex].FrameNumber.Value);
-			}
-		}
+		return;
+	}
+	
+	UMovieScene* FocusedMovieScene = FocusedMovieSequence->GetMovieScene();
+	if (!FocusedMovieScene)
+	{
+		return;
+	}
+	
+	const bool bForwards = false;
+	int32 MarkedIndex = FocusedMovieScene->FindNextMarkedFrame(GetLocalTime().Time.FloorToFrame(), bForwards);
+	if (MarkedIndex != INDEX_NONE)
+	{
+		AutoScrubToTime(FocusedMovieScene->GetMarkedFrames()[MarkedIndex].FrameNumber.Value);
 	}
 }
 
diff --git a/Engine/Source/Editor/Sequencer/Private/SequencerTimeSliderController.cpp b/Engine/Source/Editor/Sequencer/Private/SequencerTimeSliderController.cpp
--- a/Engine/Source/Editor/Sequencer/Private/SequencerTimeSliderController.cpp
+++ b/Engine/Source/Editor/Sequencer/Private/SequencerTimeSliderController.cpp
@@ -994,6 +994,7 @@ FReply FSequencerTimeSliderController::OnMouseMove( SWidget& WidgetOwner, const
 			{
 				FFrameTime MouseDownFree = ComputeFrameTimeFromMouse(MyGeometry, MouseDownPosition[0], RangeToScreen, false);
 
+				const bool       bReadOnly          = Sequencer->IsReadOnly();
 				const FFrameRate TickResolution     = GetTickResolution();
 				const bool       bLockedPlayRange   = TimeSliderArgs.IsPlaybackRangeLocked.Get();
 				const float      MouseDownPixel     = RangeToScreen.InputToLocalX(MouseDownFree / TickResolution);
@@ -1027,7 +1028,7 @@ FReply FSequencerTimeSliderController::OnMouseMove( SWidget& WidgetOwner, const
 					MouseDragType = DRAG_PLAYBACK_START;
 					TimeSliderArgs.OnPlaybackRangeBeginDrag.ExecuteIfBound();
 				}
-				else if (!bHitScrubber && HitTestMark(RangeToScreen, MouseDownPixel, DragMarkIndex) && bHandleMiddleMouseButton == false)
+				else if (!bReadOnly && !bHitScrubber && HitTestMark(RangeToScreen, MouseDownPixel, DragMarkIndex) && bHandleMiddleMouseButton == false)
 				{
 					MouseDragType = DRAG_MARK;
 					TimeSliderArgs.OnMarkBeginDrag.ExecuteIfBound();
@@ -1207,6 +1208,7 @@ FCursorReply FSequencerTimeSliderController::OnCursorQuery( TSharedRef<const SWi
 
 	FScrubRangeToScreen RangeToScreen(GetViewRange(), MyGeometry.Size);
 
+	const bool       bReadOnly        = Sequencer->IsReadOnly();
 	const FFrameRate TickResolution   = GetTickResolution();
 	const bool       bLockedPlayRange = TimeSliderArgs.IsPlaybackRangeLocked.Get();
 	const float      HitTestPixel     = MyGeometry.AbsoluteToLocal(CursorEvent.GetScreenSpacePosition()).X;
@@ -1234,7 +1236,7 @@ FCursorReply FSequencerTimeSliderController::OnCursorQuery( TSharedRef<const SWi
 	}
 
 	int32 DummyMarkIndex = INDEX_NONE;
-	if (MouseDragType == DRAG_MARK || (!bHitScrubber && HitTestMark(RangeToScreen, HitTestPixel, DummyMarkIndex)))
+	if (MouseDragType == DRAG_MARK || (!bReadOnly && !bHitScrubber && HitTestMark(RangeToScreen, HitTestPixel, DummyMarkIndex)))
 	{
 		return FCursorReply::Cursor(EMouseCursor::CardinalCross);
 	}
@@ -1338,6 +1340,9 @@ int32 FSequencerTimeSliderController::OnPaintViewArea( const FGeometry& Allotted
 
 TSharedRef<SWidget> FSequencerTimeSliderController::OpenSetPlaybackRangeMenu(const FGeometry& MyGeometry, const FPointerEvent& MouseEvent)
 {
+	TSharedPtr<ISequencer> Sequencer = WeakSequencer.Pin();
+	const bool bReadOnly = Sequencer && Sequencer->IsReadOnly();
+	
 	FScrubRangeToScreen RangeToScreen = FScrubRangeToScreen(TimeSliderArgs.ViewRange.Get(), MyGeometry.Size);
 	const float MousePixel = MyGeometry.AbsoluteToLocal(MouseEvent.GetScreenSpacePosition()).X;
 	FFrameNumber FrameNumber = ComputeFrameTimeFromMouse(MyGeometry, MouseEvent.GetScreenSpacePosition(), RangeToScreen).FrameNumber;
@@ -1348,7 +1353,6 @@ TSharedRef<SWidget> FSequencerTimeSliderController::OpenSetPlaybackRangeMenu(con
 	FText CurrentTimeText;
 	CurrentTimeText = FText::FromString(TimeSliderArgs.NumericTypeInterface->ToString(FrameNumber.Value));
 	
-
 	TRange<FFrameNumber> PlaybackRange = TimeSliderArgs.PlaybackRange.Get();
 
 	MenuBuilder.BeginSection("SequencerPlaybackRangeMenu", FText::Format(LOCTEXT("PlaybackRangeTextFormat", "Playback Range ({0}):"), CurrentTimeText));
@@ -1379,7 +1383,7 @@ TSharedRef<SWidget> FSequencerTimeSliderController::OpenSetPlaybackRangeMenu(con
 			FSlateIcon(),
 			FUIAction(
 				FExecuteAction::CreateLambda([=] { TimeSliderArgs.OnTogglePlaybackRangeLocked.ExecuteIfBound(); }),
-				FCanExecuteAction(),
+				FCanExecuteAction::CreateLambda([=]{ return !bReadOnly; }),
 				FIsActionChecked::CreateLambda([=] { return TimeSliderArgs.IsPlaybackRangeLocked.Get(); })
 			),
 			NAME_None,
@@ -1535,7 +1539,9 @@ TSharedRef<SWidget> FSequencerTimeSliderController::OpenSetPlaybackRangeMenu(con
 				LOCTEXT("AddMark", "Add Mark"),
 				FText(),
 				FSlateIcon(),
-				FUIAction(FExecuteAction::CreateLambda( [=]{ AddMarkAtFrame(FrameNumber); }))
+				FUIAction(
+					FExecuteAction::CreateLambda( [=]{ AddMarkAtFrame(FrameNumber); }),
+					FCanExecuteAction::CreateLambda([=]{ return !bReadOnly; }))
 			);
 		}
 		else 
@@ -1544,7 +1550,9 @@ TSharedRef<SWidget> FSequencerTimeSliderController::OpenSetPlaybackRangeMenu(con
 				LOCTEXT("DeleteMark", "Delete Mark"),
 				FText(),
 				FSlateIcon(),
-				FUIAction(FExecuteAction::CreateLambda([=]{ DeleteMarkAtIndex(MarkedIndex); }))
+				FUIAction(
+					FExecuteAction::CreateLambda([=]{ DeleteMarkAtIndex(MarkedIndex); }),
+					FCanExecuteAction::CreateLambda([=]{ return !bReadOnly; }))
 			);
 		}
 
@@ -1554,9 +1562,8 @@ TSharedRef<SWidget> FSequencerTimeSliderController::OpenSetPlaybackRangeMenu(con
 			FSlateIcon(),
 			FUIAction(
 				FExecuteAction::CreateLambda([=]{ DeleteAllMarks(); }),
-				FCanExecuteAction::CreateLambda([=]{ return bHasMarks; })
-			)
-		);
+				FCanExecuteAction::CreateLambda([=]{ return !bReadOnly && bHasMarks; }))
+			);
 	}
 	MenuBuilder.EndSection(); // SequencerMarkMenu
 
