Sequencer: Fix missing use of DilationPixels which is used to expand the hit region when testing for the scrubber. This fixes an issue where sometimes the hit test would fall through to dragging the playback ranges instead of dragging the time scrubber. Also, don't display the left/right cursor if the drag type is none

diff --git a/Engine/Source/Editor/Sequencer/Private/SequencerTimeSliderController.cpp b/Engine/Source/Editor/Sequencer/Private/SequencerTimeSliderController.cpp
--- a/Engine/Source/Editor/Sequencer/Private/SequencerTimeSliderController.cpp
+++ b/Engine/Source/Editor/Sequencer/Private/SequencerTimeSliderController.cpp
@@ -174,7 +174,7 @@ FSequencerTimeSliderController::FScrubberMetrics FSequencerTimeSliderController:
 	}
 
 	// Store off the pixel width of the frame
-	Metrics.FrameExtentsPx = TRange<float>(FrameStartPixel, FrameEndPixel);
+	Metrics.FrameExtentsPx = TRange<float>(FrameStartPixel - DilationPixels, FrameEndPixel + DilationPixels);
 
 	// Set the style of the scrub handle
 	TSharedPtr<FSequencer> Sequencer = WeakSequencer.Pin();
@@ -189,7 +189,7 @@ FSequencerTimeSliderController::FScrubberMetrics FSequencerTimeSliderController:
 		Metrics.Style = ESequencerScrubberStyle::Vanilla;
 
 		float ScrubPixel = RangeToScreen.InputToLocalX(ScrubTime.AsSeconds());
-		Metrics.HandleRangePx = TRange<float>(ScrubPixel - MinScrubSize*.5f, ScrubPixel + MinScrubSize*.5f);
+		Metrics.HandleRangePx = TRange<float>(ScrubPixel - MinScrubSize*.5f - DilationPixels, ScrubPixel + MinScrubSize*.5f + DilationPixels);
 	}
 	else
 	{
