Sequencer: Fix paste object bindings.

Pasted component bindings need to be remapped to their spawnables. Previously, a temp active spawn section would be injected so that the component bindings could find their parent spawnables. This isn't necessary though since MakeNewSpawnable would have already made the spawnable.

Also, attach sections and other sections that have bindings need to be remapped from the old copied binding to the new pasted binding.

diff --git a/Engine/Source/Editor/Sequencer/Private/Sequencer.cpp b/Engine/Source/Editor/Sequencer/Private/Sequencer.cpp
--- a/Engine/Source/Editor/Sequencer/Private/Sequencer.cpp
+++ b/Engine/Source/Editor/Sequencer/Private/Sequencer.cpp
@@ -8926,27 +8926,6 @@ bool FSequencer::PasteObjectBindings(const FString& TextToImport, UMovieSceneFol
 
 	OnMovieSceneBindingsPastedDelegate.Broadcast(BindingsPasted);
 
-	// Temporarily spawn all spawnables so that component bindings can be fixed
-	TArray<TWeakObjectPtr<UMovieSceneSection> > SpawnSectionsToRemove;
-	for (auto SpawnableGuid : SpawnableGuids)
-	{
-		UMovieSceneSpawnTrack* SpawnTrack = MovieScene->FindTrack<UMovieSceneSpawnTrack>(SpawnableGuid);
-
-		if (SpawnTrack)
-		{
-			for (UMovieSceneSection* SpawnSection : SpawnTrack->GetAllSections())
-			{
-				SpawnSection->SetIsActive(false);
-			}
-
-			// Spawnable could have animated spawned state, so temporarily override it to spawn infinitely
-			UMovieSceneSpawnSection* SpawnSection = Cast<UMovieSceneSpawnSection>(SpawnTrack->CreateNewSection());
-			SpawnSection->GetChannel().Reset();
-			SpawnSection->GetChannel().SetDefault(true);
-			SpawnSectionsToRemove.Add(SpawnSection);
-		}
-	}
-	
 	// Refresh all immediately so that spawned actors will be generated immediately
 	ForceEvaluate();
 
@@ -8990,25 +8969,13 @@ bool FSequencer::PasteObjectBindings(const FString& TextToImport, UMovieSceneFol
 			}
 		}
 	}
-	
-	for (TWeakObjectPtr<UMovieSceneSection> SpawnSectionToRemove : SpawnSectionsToRemove)
-	{
-		if (SpawnSectionToRemove.IsValid())
-		{
-			UMovieSceneTrack* SpawnTrack = SpawnSectionToRemove->GetTypedOuter<UMovieSceneTrack>();
-			if (SpawnTrack)
-			{
-				SpawnTrack->Modify();
-				SpawnTrack->RemoveSection(*SpawnSectionToRemove);
 
-				for (UMovieSceneSection* SpawnSection : SpawnTrack->GetAllSections())
-				{
-					SpawnSection->SetIsActive(true);
-				}
-			}
-		}
+	// Remap bindings in sections (ie. attach tracks)
+	for (TPair<FGuid, FGuid> GuidPair : OldToNewGuidMap)
+	{
+		FSequencerUtilities::UpdateBindingIDs(this, CompiledDataManager, GuidPair.Key, GuidPair.Value);
 	}
-	
+
 	NotifyMovieSceneDataChanged(EMovieSceneDataChangeType::MovieSceneStructureItemsChanged);
 
 	return true;
