Sequencer: Speculative fix for tooltip crash when adjusting ease in/out

diff --git a/Engine/Source/Editor/Sequencer/Private/SSequencerSectionAreaView.cpp b/Engine/Source/Editor/Sequencer/Private/SSequencerSectionAreaView.cpp
--- a/Engine/Source/Editor/Sequencer/Private/SSequencerSectionAreaView.cpp
+++ b/Engine/Source/Editor/Sequencer/Private/SSequencerSectionAreaView.cpp
@@ -157,14 +157,14 @@ FText SSequencerSectionAreaView::GetSectionToolTip(TSharedRef<ISequencerSection>
 				SectionToolTipContent);
 		}
 	
-		if (SectionObject->Easing.GetEaseInDuration() > 0)
+		if (SectionObject->Easing.EaseIn.GetObject() && SectionObject->Easing.GetEaseInDuration() > 0)
 		{
 			int32 EaseInFrames = ConvertFrameTime(SectionObject->Easing.GetEaseInDuration(), TickResolution, DisplayRate).RoundToFrame().Value;
 			FText EaseInText = FText::Format(NSLOCTEXT("SequencerSection", "EaseInFormat", "Ease In: {0} ({1} frames)"), SectionObject->Easing.EaseIn->GetDisplayName(), EaseInFrames);
 			SectionToolTip = FText::Join(FText::FromString("\n"), SectionToolTip, EaseInText);
 		}
 
-		if (SectionObject->Easing.GetEaseOutDuration() > 0)
+		if (SectionObject->Easing.EaseOut.GetObject() && SectionObject->Easing.GetEaseOutDuration() > 0)
 		{
 			int32 EaseOutFrames = ConvertFrameTime(SectionObject->Easing.GetEaseOutDuration(), TickResolution, DisplayRate).RoundToFrame().Value;
 			FText EaseOutText = FText::Format(NSLOCTEXT("SequencerSection", "EaseOutFormat", "Ease Out: {0} ({1} frames)"), SectionObject->Easing.EaseOut->GetDisplayName(), EaseOutFrames);
