Invalidate path tracer when sequencer time changes

This also allows the path tracer to be used in the niagara editor (with constant invalidation, but avoiding blurring)

diff --git a/Engine/Source/Editor/Sequencer/Private/Sequencer.cpp b/Engine/Source/Editor/Sequencer/Private/Sequencer.cpp
--- a/Engine/Source/Editor/Sequencer/Private/Sequencer.cpp
+++ b/Engine/Source/Editor/Sequencer/Private/Sequencer.cpp
@@ -152,6 +152,7 @@
 #include "EntitySystem/MovieScenePreAnimatedStateSystem.h"
 #include "Systems/MovieSceneMotionVectorSimulationSystem.h"
 
+#include "EngineModule.h"
 
 #define LOCTEXT_NAMESPACE "Sequencer"
 
@@ -3129,6 +3130,7 @@ void FSequencer::EvaluateInternal(FMovieSceneEvaluationRange InRange, bool bHasJ
 	if (!IsInSilentMode())
 	{
 		OnGlobalTimeChangedDelegate.Broadcast();
+		GetRendererModule().InvalidatePathTracedOutput();
 	}
 
 }
