Sequencer: Fix crash in GetSelectedKeyAreas on selection. No need to RemoveAtSwap

diff --git a/Engine/Source/Editor/Sequencer/Private/Sequencer.cpp b/Engine/Source/Editor/Sequencer/Private/Sequencer.cpp
--- a/Engine/Source/Editor/Sequencer/Private/Sequencer.cpp
+++ b/Engine/Source/Editor/Sequencer/Private/Sequencer.cpp
@@ -7035,14 +7035,12 @@ void FSequencer::GetSelectedKeyAreas(TArray<const IKeyArea*>& OutSelectedKeyArea
 	{
 
 		TSet<TSharedRef<FSequencerDisplayNode>> ChildNodes;
-		int32 Index = 0;
 		for (TSharedRef<FSequencerDisplayNode> Node : NodesToKey.Array())
 		{
 			// No need to gather key areas from binding/tracks because they have no key areas
 			if (Node->GetType()== ESequencerNode::Object || Node->GetType() == ESequencerNode::Track)
 			{
-				NodesToKey.Array().RemoveAtSwap(Index);
-				++Index;
+				continue;
 			}
 			else
 			{
@@ -7053,7 +7051,6 @@ void FSequencer::GetSelectedKeyAreas(TArray<const IKeyArea*>& OutSelectedKeyArea
 				{
 					NodesToKey.Remove(ChildNode);
 				}
-				++Index;
 			}
 		}
 	}
