Sequencer: Clamp hierarchy browser to display height / 2 for very deep sequence hierarchies

diff --git a/Engine/Source/Editor/Sequencer/Private/SSequencerHierarchyBrowser.cpp b/Engine/Source/Editor/Sequencer/Private/SSequencerHierarchyBrowser.cpp
--- a/Engine/Source/Editor/Sequencer/Private/SSequencerHierarchyBrowser.cpp
+++ b/Engine/Source/Editor/Sequencer/Private/SSequencerHierarchyBrowser.cpp
@@ -146,6 +146,9 @@ void SSequencerHierarchyBrowser::Construct(const FArguments& InArgs, TWeakPtr<FS
 		OutChildren.Append(InParent->Children);
 	};
 
+	FDisplayMetrics DisplayMetrics;
+	FSlateApplication::Get().GetCachedDisplayMetrics(DisplayMetrics);
+
 	TreeView = SNew(STreeView<TSharedPtr<FSequencerHierarchyNode>>)
 		.OnGenerateRow_Lambda(HandleGenerateRow)
 		.OnGetChildren_Lambda(HandleGetChildren)
@@ -160,6 +163,7 @@ void SSequencerHierarchyBrowser::Construct(const FArguments& InArgs, TWeakPtr<FS
 			SNew(SVerticalBox)
 
 			+ SVerticalBox::Slot()
+			.MaxHeight(DisplayMetrics.PrimaryDisplayHeight * 0.5)
 			[
 				SNew(SScrollBorder, TreeView.ToSharedRef())
 				[
