Add `OptionalFwd.h` for `TOptional`

diff --git a/Engine/Source/Developer/Android/AndroidTargetPlatform/Private/AndroidTargetDevice.inl b/Engine/Source/Developer/Android/AndroidTargetPlatform/Private/AndroidTargetDevice.inl
--- a/Engine/Source/Developer/Android/AndroidTargetPlatform/Private/AndroidTargetDevice.inl
+++ b/Engine/Source/Developer/Android/AndroidTargetPlatform/Private/AndroidTargetDevice.inl
@@ -12,8 +12,6 @@ class FAndroidTargetDevice;
 struct FTargetDeviceProcessInfo;
 enum class ETargetDeviceFeatures;
 
-template<typename OptionalType> struct TOptional;
-
 inline FString FAndroidTargetDevice::GetOperatingSystemName()
 {
 	if (!AndroidVersionString.IsEmpty())
diff --git a/Engine/Source/Editor/Sequencer/Public/SequencerKeyCollection.h b/Engine/Source/Editor/Sequencer/Public/SequencerKeyCollection.h
--- a/Engine/Source/Editor/Sequencer/Public/SequencerKeyCollection.h
+++ b/Engine/Source/Editor/Sequencer/Public/SequencerKeyCollection.h
@@ -2,17 +2,17 @@
 
 #pragma once
 
-#include "CoreTypes.h"
 #include "Containers/ContainersFwd.h"
+#include "CoreTypes.h"
+#include "Misc/FrameNumber.h"
 #include "Misc/Guid.h"
+#include "Misc/OptionalFwd.h"
 #include "Templates/SharedPointer.h"
-#include "Misc/FrameNumber.h"
 
 class UMovieSceneSection;
 class FSequencerDisplayNode;
 class IKeyArea;
 
-template<typename> struct TOptional;
 template<typename> class TRange;
 
 /** Enumeration used to define how to search for keys */
diff --git a/Engine/Source/Runtime/Core/Public/Math/UnitConversion.inl b/Engine/Source/Runtime/Core/Public/Math/UnitConversion.inl
--- a/Engine/Source/Runtime/Core/Public/Math/UnitConversion.inl
+++ b/Engine/Source/Runtime/Core/Public/Math/UnitConversion.inl
@@ -5,15 +5,13 @@
 
 #include "CoreTypes.h"
 #include "CoreFwd.h"
+#include "Misc/OptionalFwd.h"
 
 struct FMath;
 struct FUnitConversion;
 enum class EUnit : uint8;
 enum class EUnitType;
 template<typename NumericType> struct FNumericUnit;
-/// @cond DOXYGEN_WARNINGS
-template<typename OptionalType> struct TOptional;
-/// @endcond
 template<typename ValueType, typename ErrorType> class TValueOrError;
 
 namespace UnitConversion
diff --git a/Engine/Source/Runtime/Core/Public/Misc/OptionalFwd.h b/Engine/Source/Runtime/Core/Public/Misc/OptionalFwd.h
new file mode 100644
--- /dev/null
+++ b/Engine/Source/Runtime/Core/Public/Misc/OptionalFwd.h
@@ -0,0 +1,10 @@
+// Copyright Epic Games, Inc. All Rights Reserved.
+
+#pragma once
+
+/**
+ * When we have an optional value IsSet() returns true, and GetValue() is meaningful.
+ * Otherwise GetValue() is not meaningful.
+ */
+template<typename OptionalType>
+struct TOptional;
diff --git a/Engine/Source/Runtime/Core/Public/Templates/NonNullPointer.h b/Engine/Source/Runtime/Core/Public/Templates/NonNullPointer.h
--- a/Engine/Source/Runtime/Core/Public/Templates/NonNullPointer.h
+++ b/Engine/Source/Runtime/Core/Public/Templates/NonNullPointer.h
@@ -3,12 +3,9 @@
 #pragma once
 
 #include "CoreTypes.h"
+#include "Misc/OptionalFwd.h"
 #include "Templates/PointerIsConvertibleFromTo.h"
 
-// Forward declaration
-template<typename OptionalType>
-struct TOptional;
-
 // So we can construct TNonNullPtrs
 enum class EDefaultConstructNonNullPtr { UnsafeDoNotUse };
 
