Sequencer: Backout IsEnabled from 18954849 because if the sequence is locked, the button won't be enabled to unlock the sequence

diff --git a/Engine/Source/Editor/Sequencer/Private/SSequencer.cpp b/Engine/Source/Editor/Sequencer/Private/SSequencer.cpp
--- a/Engine/Source/Editor/Sequencer/Private/SSequencer.cpp
+++ b/Engine/Source/Editor/Sequencer/Private/SSequencer.cpp
@@ -710,7 +710,6 @@ void SSequencer::Construct(const FArguments& InArgs, TSharedRef<FSequencer> InSe
 						.Type(ESlateCheckBoxType::CheckBox) // Use CheckBox instead of ToggleType since we're not putting ohter widget inside
 						.Padding(FMargin(0.f))
 						.IsFocusable(false)
-						.IsEnabled(!SequencerPtr.Pin()->IsReadOnly())		
 						.IsChecked_Lambda([this] { return GetIsSequenceReadOnly() ? ECheckBoxState::Checked : ECheckBoxState::Unchecked; } )
 						.OnCheckStateChanged(this, &SSequencer::OnSetSequenceReadOnly)
 						.ToolTipText_Lambda([this] { return GetIsSequenceReadOnly() ? LOCTEXT("UnlockSequence", "Unlock the animation so that it is editable") : LOCTEXT("LockSequence", "Lock the animation so that it is not editable"); } )
