Sequencer: Fix localization warning for Outliner Columns

diff --git a/Engine/Source/Editor/Sequencer/Private/LevelEditorSequencerIntegration.cpp b/Engine/Source/Editor/Sequencer/Private/LevelEditorSequencerIntegration.cpp
--- a/Engine/Source/Editor/Sequencer/Private/LevelEditorSequencerIntegration.cpp
+++ b/Engine/Source/Editor/Sequencer/Private/LevelEditorSequencerIntegration.cpp
@@ -1019,13 +1019,15 @@ void FLevelEditorSequencerIntegration::AttachOutlinerColumn()
 	FSceneOutlinerModule& SceneOutlinerModule = FModuleManager::LoadModuleChecked< FSceneOutlinerModule >("SceneOutliner");
 
 	FSceneOutlinerColumnInfo SpawnColumnInfo(ESceneOutlinerColumnVisibility::Visible, 11, 
-		FCreateSceneOutlinerColumn::CreateRaw( this, &FLevelEditorSequencerIntegration::CreateSequencerSpawnableColumn));
+		FCreateSceneOutlinerColumn::CreateRaw( this, &FLevelEditorSequencerIntegration::CreateSequencerSpawnableColumn),
+		true, TOptional<float>(), LOCTEXT("SpawnableColumnName", "SpawnableInfo"));
 
 	SceneOutlinerModule.RegisterDefaultColumnType< Sequencer::FSequencerSpawnableColumn >(SpawnColumnInfo);
 	AcquiredResources.Add([=]{ this->DetachOutlinerColumn(); });
 
 	FSceneOutlinerColumnInfo ColumnInfo(ESceneOutlinerColumnVisibility::Visible, 15, 
-		FCreateSceneOutlinerColumn::CreateRaw( this, &FLevelEditorSequencerIntegration::CreateSequencerInfoColumn));
+		FCreateSceneOutlinerColumn::CreateRaw( this, &FLevelEditorSequencerIntegration::CreateSequencerInfoColumn), 
+		true, TOptional<float>(), LOCTEXT("SequencerColumnName", "SequencerInfo"));
 
 	SceneOutlinerModule.RegisterDefaultColumnType< Sequencer::FSequencerInfoColumn >(ColumnInfo);
 
