Sequencer: Curve Editor:  Curve Editor Sync Isse with Filter
Since Selection ticks before any filter is applied and the curve editor tree is udpated in FSequencer::Tick we need to sync the curve editor on the next tick so the filter is applied first, otherwise the curve tree hierarchy won't be up to date and the curve may not be found.

diff --git a/Engine/Source/Editor/Sequencer/Private/SSequencer.cpp b/Engine/Source/Editor/Sequencer/Private/SSequencer.cpp
--- a/Engine/Source/Editor/Sequencer/Private/SSequencer.cpp
+++ b/Engine/Source/Editor/Sequencer/Private/SSequencer.cpp
@@ -1318,23 +1318,31 @@ void SSequencer::HandleOutlinerNodeSelectionChanged()
 
 		if (GetSequencerSettings()->ShouldSyncCurveEditorSelection())
 		{
-			TSharedRef<FSequencerNodeTree> NodeTree = Sequencer->GetNodeTree();
-			CurveEditor->SuspendBroadcast();
-			// Clear the tree selection
-			CurveEditorTree->ClearSelection();
-			for (TSharedRef<FSequencerDisplayNode> Node : SelectedDisplayNodes)
+			//since Selection ticks before any filter is applied and the curve editor tree is udpated in FSequencer::Tick we need to sync 
+			//the curve editor on the next tick so the filter is applied first, otherwise the curve tree hierarchy won't be up to date
+			//and the curve may not be found.
+			GEditor->GetTimerManager()->SetTimerForNextTick([Sequencer, this]()
 			{
-				FCurveEditorTreeItemID CurveEditorTreeItem = NodeTree->FindCurveEditorTreeItem(Node);
-				if (CurveEditorTreeItem != FCurveEditorTreeItemID::Invalid())
+				const TSet<TSharedRef<FSequencerDisplayNode>>& SelectedDisplayNodes = Sequencer->GetSelection().GetSelectedOutlinerNodes();
+				TSharedPtr<FCurveEditor> CurveEditor = Sequencer->GetCurveEditor();
+				TSharedRef<FSequencerNodeTree> NodeTree = Sequencer->GetNodeTree();
+				CurveEditor->SuspendBroadcast();
+				// Clear the tree selection
+				CurveEditorTree->ClearSelection();
+				for (TSharedRef<FSequencerDisplayNode> Node : SelectedDisplayNodes)
 				{
-					if (!CurveEditorTree->IsItemSelected(CurveEditorTreeItem))
+					FCurveEditorTreeItemID CurveEditorTreeItem = NodeTree->FindCurveEditorTreeItem(Node);
+					if (CurveEditorTreeItem != FCurveEditorTreeItemID::Invalid())
 					{
-						CurveEditorTree->SetItemSelection(CurveEditorTreeItem, true);
-						CurveEditorTree->RequestScrollIntoView(CurveEditorTreeItem);
+						if (!CurveEditorTree->IsItemSelected(CurveEditorTreeItem))
+						{
+							CurveEditorTree->SetItemSelection(CurveEditorTreeItem, true);
+							CurveEditorTree->RequestScrollIntoView(CurveEditorTreeItem);
+						}
 					}
 				}
-			}
-			CurveEditor->ResumeBroadcast();
+				CurveEditor->ResumeBroadcast();
+			});
 		}
 	}
 
