Sequencer: Clamp to playback range before snapping to frames, otherwise the scrubbed to frame could be a subframe when snapping is on.

diff --git a/Engine/Source/Editor/Sequencer/Private/SequencerTimeSliderController.cpp b/Engine/Source/Editor/Sequencer/Private/SequencerTimeSliderController.cpp
--- a/Engine/Source/Editor/Sequencer/Private/SequencerTimeSliderController.cpp
+++ b/Engine/Source/Editor/Sequencer/Private/SequencerTimeSliderController.cpp
@@ -94,6 +94,12 @@ FFrameTime FSequencerTimeSliderController::ComputeScrubTimeFromMouse(const FGeom
 		return ScrubTime;
 	}
 	
+	// Clamp first, snap to frame last
+	if (Sequencer->GetSequencerSettings()->ShouldKeepCursorInPlayRangeWhileScrubbing())
+	{
+		ScrubTime = UE::MovieScene::ClampToDiscreteRange(ScrubTime, TimeSliderArgs.PlaybackRange.Get());
+	}
+
 	if ( Sequencer->GetSequencerSettings()->GetIsSnapEnabled() || MouseEvent.IsShiftDown() )
 	{
 		if (Sequencer->GetSequencerSettings()->GetSnapPlayTimeToInterval())
@@ -115,11 +121,6 @@ FFrameTime FSequencerTimeSliderController::ComputeScrubTimeFromMouse(const FGeom
 		ScrubTime = SnapTimeToNearestKey(MouseEvent, RangeToScreen, CursorPos.X, ScrubTime);
 	}
 
-	if (Sequencer->GetSequencerSettings()->ShouldKeepCursorInPlayRangeWhileScrubbing())
-	{
-		ScrubTime = UE::MovieScene::ClampToDiscreteRange(ScrubTime, TimeSliderArgs.PlaybackRange.Get());
-	}
-
 	return ScrubTime;
 }
 
