Sequencer: New Transport Tool hotkeys, for toggle play, step frame and scrub frame.Also Game View in the Level Editor will hide all Controls in that viewport.

diff --git a/Engine/Plugins/Animation/ControlRig/Source/ControlRigEditor/Private/EditMode/ControlRigEditMode.cpp b/Engine/Plugins/Animation/ControlRig/Source/ControlRigEditor/Private/EditMode/ControlRigEditMode.cpp
--- a/Engine/Plugins/Animation/ControlRig/Source/ControlRigEditor/Private/EditMode/ControlRigEditMode.cpp
+++ b/Engine/Plugins/Animation/ControlRig/Source/ControlRigEditor/Private/EditMode/ControlRigEditMode.cpp
@@ -681,8 +681,8 @@ TSet<FName> FControlRigEditMode::GetActiveControlsFromSequencer(UControlRig* Con
 void FControlRigEditMode::Render(const FSceneView* View, FViewport* Viewport, FPrimitiveDrawInterface* PDI)
 {	
 	const UControlRigEditModeSettings* Settings = GetDefault<UControlRigEditModeSettings>();
-	bool bRender = !Settings->bHideControlShapes;
-
+	const bool bIsInGameView = IsInLevelEditor() ? (GCurrentLevelEditingViewportClient && GCurrentLevelEditingViewportClient->IsInGameView()) : false;
+	bool bRender = !Settings->bHideControlShapes && !bIsInGameView;
 	for (TWeakObjectPtr<UControlRig>& ControlRigPtr : RuntimeControlRigs)
 	{
 		UControlRig* ControlRig = ControlRigPtr.Get();
@@ -3526,7 +3526,9 @@ void FControlRigEditMode::TickControlShape(AControlRigShapeActor* ShapeActor, co
 			if (FRigControlElement* ControlElement = ControlRig->FindControl(ShapeActor->ControlName))
 			{
 				ShapeActor->SetShapeColor(ControlElement->Settings.ShapeColor);
-				ShapeActor->SetIsTemporarilyHiddenInEditor(!ControlElement->Settings.bShapeVisible || Settings->bHideControlShapes || !ControlRig->GetControlsVisible());
+				const bool bIsInGameView = IsInLevelEditor() ? (GCurrentLevelEditingViewportClient && GCurrentLevelEditingViewportClient->IsInGameView()) : false;
+
+				ShapeActor->SetIsTemporarilyHiddenInEditor(bIsInGameView || !ControlElement->Settings.bShapeVisible || Settings->bHideControlShapes || !ControlRig->GetControlsVisible());
 				if (!IsInLevelEditor()) //don't change this in level editor otherwise we can never select it
 				{
 					ShapeActor->SetSelectable(ControlElement->Settings.bShapeVisible && !Settings->bHideControlShapes && ControlElement->Settings.bAnimatable && ControlRig->GetControlsVisible());
diff --git a/Engine/Source/Editor/Sequencer/Private/Sequencer.cpp b/Engine/Source/Editor/Sequencer/Private/Sequencer.cpp
--- a/Engine/Source/Editor/Sequencer/Private/Sequencer.cpp
+++ b/Engine/Source/Editor/Sequencer/Private/Sequencer.cpp
@@ -12342,14 +12342,28 @@ void FSequencer::BindCommands()
 {
 	const FSequencerCommands& Commands = FSequencerCommands::Get();
 
+	SequencerCommandBindings->MapAction(
+		Commands.TogglePlayViewport,
+		FExecuteAction::CreateSP(this, &FSequencer::TogglePlay));
+
 	SequencerCommandBindings->MapAction(
 		Commands.StepToNextKey,
-		FExecuteAction::CreateSP( this, &FSequencer::StepToNextKey ) );
+		FExecuteAction::CreateSP( this, &FSequencer::StepToNextKey) );
 
 	SequencerCommandBindings->MapAction(
 		Commands.StepToPreviousKey,
 		FExecuteAction::CreateSP( this, &FSequencer::StepToPreviousKey ) );
 
+	SequencerCommandBindings->MapAction(
+		Commands.StepForwardViewport,
+		FExecuteAction::CreateSP(this, &FSequencer::StepForward),
+		EUIActionRepeatMode::RepeatEnabled);
+
+	SequencerCommandBindings->MapAction(
+		Commands.StepBackwardViewport,
+		FExecuteAction::CreateSP(this, &FSequencer::StepBackward),
+		EUIActionRepeatMode::RepeatEnabled);
+
 	SequencerCommandBindings->MapAction(
 		Commands.StepToNextCameraKey,
 		FExecuteAction::CreateSP( this, &FSequencer::StepToNextCameraKey ) );
@@ -13006,21 +13020,12 @@ void FSequencer::BindCommands()
 		FExecuteAction::CreateSP(this, &FSequencer::StepForward),
 		EUIActionRepeatMode::RepeatEnabled);
 
-	SequencerCommandBindings->MapAction(
-		Commands.StepForward2,
-		FExecuteAction::CreateSP(this, &FSequencer::StepForward),
-		EUIActionRepeatMode::RepeatEnabled);
 
 	SequencerCommandBindings->MapAction(
 		Commands.StepBackward,
 		FExecuteAction::CreateSP(this, &FSequencer::StepBackward),
 		EUIActionRepeatMode::RepeatEnabled);
 
-	SequencerCommandBindings->MapAction(
-		Commands.StepBackward2,
-		FExecuteAction::CreateSP(this, &FSequencer::StepBackward),
-		EUIActionRepeatMode::RepeatEnabled);
-
 	SequencerCommandBindings->MapAction(
 		Commands.JumpForward,
 		FExecuteAction::CreateSP(this, &FSequencer::JumpForward),
@@ -13173,6 +13178,7 @@ void FSequencer::BindCommands()
 	// while also freeing us up from issues that result from Sequencer already using two lists (for which our commands might be spread
 	// across both lists which makes a direct copy like it already uses difficult).
 	CurveEditorSharedBindings->MapAction(Commands.TogglePlay,			*SequencerCommandBindings->GetActionForCommand(Commands.TogglePlay));
+	CurveEditorSharedBindings->MapAction(Commands.TogglePlayViewport,	*SequencerCommandBindings->GetActionForCommand(Commands.TogglePlayViewport));
 	CurveEditorSharedBindings->MapAction(Commands.PlayForward,			*SequencerCommandBindings->GetActionForCommand(Commands.PlayForward));
 	CurveEditorSharedBindings->MapAction(Commands.JumpToStart,			*SequencerCommandBindings->GetActionForCommand(Commands.JumpToStart));
 	CurveEditorSharedBindings->MapAction(Commands.JumpToEnd,			*SequencerCommandBindings->GetActionForCommand(Commands.JumpToEnd));
@@ -13181,8 +13187,8 @@ void FSequencer::BindCommands()
 	CurveEditorSharedBindings->MapAction(Commands.Pause,				*SequencerCommandBindings->GetActionForCommand(Commands.Pause));
 	CurveEditorSharedBindings->MapAction(Commands.StepForward,			*SequencerCommandBindings->GetActionForCommand(Commands.StepForward));
 	CurveEditorSharedBindings->MapAction(Commands.StepBackward,			*SequencerCommandBindings->GetActionForCommand(Commands.StepBackward));
-	CurveEditorSharedBindings->MapAction(Commands.StepForward2,         *SequencerCommandBindings->GetActionForCommand(Commands.StepForward2));
-	CurveEditorSharedBindings->MapAction(Commands.StepBackward2,        *SequencerCommandBindings->GetActionForCommand(Commands.StepBackward2));
+	CurveEditorSharedBindings->MapAction(Commands.StepForwardViewport,  *SequencerCommandBindings->GetActionForCommand(Commands.StepForwardViewport));
+	CurveEditorSharedBindings->MapAction(Commands.StepBackwardViewport, *SequencerCommandBindings->GetActionForCommand(Commands.StepBackwardViewport));
 	CurveEditorSharedBindings->MapAction(Commands.JumpForward,          *SequencerCommandBindings->GetActionForCommand(Commands.JumpForward));
 	CurveEditorSharedBindings->MapAction(Commands.JumpBackward,         *SequencerCommandBindings->GetActionForCommand(Commands.JumpBackward));
 	CurveEditorSharedBindings->MapAction(Commands.StepToNextKey,		*SequencerCommandBindings->GetActionForCommand(Commands.StepToNextKey));
diff --git a/Engine/Source/Editor/Sequencer/Private/SequencerCommands.cpp b/Engine/Source/Editor/Sequencer/Private/SequencerCommands.cpp
--- a/Engine/Source/Editor/Sequencer/Private/SequencerCommands.cpp
+++ b/Engine/Source/Editor/Sequencer/Private/SequencerCommands.cpp
@@ -7,6 +7,8 @@
 void FSequencerCommands::RegisterCommands()
 {
 	UI_COMMAND( TogglePlay, "Toggle Play", "Toggle the timeline playing", EUserInterfaceActionType::Button, FInputChord(EKeys::SpaceBar) );
+	UI_COMMAND( TogglePlayViewport, "Toggle Play(Viewport)", "Toggle the timeline playing in all viewports and sequencer", EUserInterfaceActionType::Button, FInputChord(EModifierKey::Shift, EKeys::SpaceBar));
+	UI_COMMAND( ScrubTimeViewport, "Scrub Time(Viewport)", "Scrub mouse left and right to change time", EUserInterfaceActionType::Button, FInputChord(EKeys::K) );
 	UI_COMMAND( PlayForward, "Play Forward", "Play the timeline forward", EUserInterfaceActionType::Button, FInputChord(EKeys::Down) );
 	UI_COMMAND( JumpToStart, "Jump to Start", "Jump to the start of the playback range", EUserInterfaceActionType::Button, FInputChord(EKeys::Up) );
 	UI_COMMAND( JumpToEnd, "Jump to End", "Jump to the end of the playback range", EUserInterfaceActionType::Button, FInputChord(EModifierKey::Control, EKeys::Up) );
@@ -16,12 +18,12 @@ void FSequencerCommands::RegisterCommands()
 	UI_COMMAND( RestorePlaybackSpeed, "Restore Speed", "Restores the playback speed to 1.", EUserInterfaceActionType::Button, FInputChord() );
 	UI_COMMAND( StepForward, "Step Forward", "Step the timeline forward", EUserInterfaceActionType::Button, FInputChord(EKeys::Right) );
 	UI_COMMAND( StepBackward, "Step Backward", "Step the timeline backward", EUserInterfaceActionType::Button, FInputChord(EKeys::Left) );
-	UI_COMMAND( StepForward2, "Step Forward", "Step the timeline forward", EUserInterfaceActionType::Button, FInputChord(EModifierKey::Alt, EKeys::Period) );
-	UI_COMMAND( StepBackward2, "Step Backward", "Step the timeline backward", EUserInterfaceActionType::Button, FInputChord(EModifierKey::Alt, EKeys::Comma) );
+	UI_COMMAND( StepForwardViewport, "Step Forward (Viewport)", "Step the timeline forward in all viewports and sequencer", EUserInterfaceActionType::Button, FInputChord(EModifierKey::Alt, EKeys::Period) );
+	UI_COMMAND( StepBackwardViewport, "Step Backward (Viewport)", "Step the timeline backward in all viewports and sequencer", EUserInterfaceActionType::Button, FInputChord(EModifierKey::Alt, EKeys::Comma) );
 	UI_COMMAND( JumpForward, "Jump Forward", "Jump the timeline forward a user defined number of frames/times", EUserInterfaceActionType::Button, FInputChord(EModifierKey::Shift, EKeys::Right) );
 	UI_COMMAND( JumpBackward, "Jump Backward", "Jump the timeline backward a user defined number of frames/times", EUserInterfaceActionType::Button, FInputChord(EModifierKey::Shift, EKeys::Left) );
-	UI_COMMAND( StepToNextKey, "Step to Next Key", "Step to the next key", EUserInterfaceActionType::Button, FInputChord(EKeys::Period) );
-	UI_COMMAND( StepToPreviousKey, "Step to Previous Key", "Step to the previous key", EUserInterfaceActionType::Button, FInputChord(EKeys::Comma) );
+	UI_COMMAND( StepToNextKey, "Step to Next Key (Viewport)", "Step to the next key in all viewports and sequencer", EUserInterfaceActionType::Button, FInputChord(EKeys::Period) );
+	UI_COMMAND( StepToPreviousKey, "Step to Previous Key(Viewport)", "Step to the previous key in all viewports and sequencer", EUserInterfaceActionType::Button, FInputChord(EKeys::Comma) );
 	UI_COMMAND( StepToNextCameraKey, "Step to Next Camera Key", "Step to the next camera key", EUserInterfaceActionType::Button, FInputChord() );
 	UI_COMMAND( StepToPreviousCameraKey, "Step to Previous Camera Key", "Step to the previous camera key", EUserInterfaceActionType::Button, FInputChord() );
 	UI_COMMAND( StepToNextShot, "Step to Next Shot", "Step to the next shot", EUserInterfaceActionType::Button, FInputChord(EModifierKey::Shift, EKeys::Period) );
diff --git a/Engine/Source/Editor/Sequencer/Private/SequencerEdMode.cpp b/Engine/Source/Editor/Sequencer/Private/SequencerEdMode.cpp
--- a/Engine/Source/Editor/Sequencer/Private/SequencerEdMode.cpp
+++ b/Engine/Source/Editor/Sequencer/Private/SequencerEdMode.cpp
@@ -30,6 +30,7 @@
 #include "SequencerSectionPainter.h"
 #include "MovieSceneToolHelpers.h"
 #include "Tools/MotionTrailOptions.h"
+#include "SequencerCommands.h"
 
 const FEditorModeID FSequencerEdMode::EM_SequencerMode(TEXT("EM_SequencerMode"));
 
@@ -158,6 +159,8 @@ FSequencerEdMode::~FSequencerEdMode()
 
 void FSequencerEdMode::Enter()
 {
+	bIsTracking = false;;
+	StartXValue.Reset();
 	FEdMode::Enter();
 }
 
@@ -193,15 +196,164 @@ bool FSequencerEdMode::InputKey( FEditorViewportClient* ViewportClient, FViewpor
 	{
 		FModifierKeysState KeyState = FSlateApplication::Get().GetModifierKeys();
 
-		if (ActiveSequencer->GetCommandBindings(ESequencerCommandBindings::Shared).Get()->ProcessCommandBindings(Key, KeyState, (Event == IE_Repeat) ))
+ 		if (ActiveSequencer->GetCommandBindings(ESequencerCommandBindings::Shared).Get()->ProcessCommandBindings(Key, KeyState, (Event == IE_Repeat) ))
 		{
 			return true;
 		}
 	}
-
 	return FEdMode::InputKey(ViewportClient, Viewport, Key, Event);
 }
 
+bool FSequencerEdMode::IsPressingMoveTimeSlider(FViewport* InViewport) const
+{
+	const FSequencerCommands& Commands = FSequencerCommands::Get();
+	bool bIsMovingTimeSlider = false;
+	// Need to iterate through primary and secondary to make sure they are all pressed.
+	for (uint32 i = 0; i < static_cast<uint8>(EMultipleKeyBindingIndex::NumChords); ++i)
+	{
+		EMultipleKeyBindingIndex ChordIndex = static_cast<EMultipleKeyBindingIndex>(i);
+		const FInputChord& Chord = *Commands.ScrubTimeViewport->GetActiveChord(ChordIndex);
+		bIsMovingTimeSlider |= Chord.IsValidChord() && InViewport->KeyState(Chord.Key);
+	}
+	return bIsMovingTimeSlider;
+}
+
+bool FSequencerEdMode::StartTracking(FEditorViewportClient* InViewportClient, FViewport* InViewport)
+{
+	return IsPressingMoveTimeSlider(InViewport);
+}
+
+bool FSequencerEdMode::InputDelta(FEditorViewportClient* InViewportClient, FViewport* InViewport, FVector& InDrag, FRotator& InRot, FVector& InScale)
+{
+	return IsPressingMoveTimeSlider(InViewport);
+}
+
+bool FSequencerEdMode::ProcessCapturedMouseMoves(FEditorViewportClient* InViewportClient, FViewport* InViewport, const TArrayView<FIntPoint>& CapturedMouseMoves)
+{
+	if (CapturedMouseMoves.Num() > 0)
+	{
+		const bool bTimeChange = IsPressingMoveTimeSlider(InViewport);
+		if (bTimeChange)
+		{
+			TSharedPtr<ISequencer> ActiveSequencer;
+
+			for (TWeakPtr<ISequencer> WeakSequencer : Sequencers)
+			{
+				ActiveSequencer = WeakSequencer.Pin();
+				if (ActiveSequencer.IsValid())
+				{
+					break;
+				}
+			}
+
+			if (ActiveSequencer.IsValid())
+			{
+				for (int32 Index = 0; Index < CapturedMouseMoves.Num(); ++Index)
+				{
+					int32 X = CapturedMouseMoves[Index].X;
+					if (StartXValue.IsSet() == false)
+					{
+						StartXValue = X;
+						FQualifiedFrameTime CurrentTime = ActiveSequencer->GetLocalTime();
+						StartFrameNumber = CurrentTime.Time.GetFrame();
+					}
+					else
+					{
+						int32 Diff = X - StartXValue.GetValue();
+						if (Diff != 0)
+						{
+							FIntPoint Origin, Size;
+							InViewportClient->GetViewportDimensions(Origin, Size);
+							const float ViewPortSize = (float)Size.X;
+							const float FloatViewDiff = (float)(Diff) / ViewPortSize;
+
+							FFrameRate TickResolution = ActiveSequencer->GetFocusedTickResolution();
+							TPair<FFrameNumber, FFrameNumber> ViewRange(TickResolution.AsFrameNumber(ActiveSequencer->GetViewRange().GetLowerBoundValue()), TickResolution.AsFrameNumber(ActiveSequencer->GetViewRange().GetUpperBoundValue()));
+							FFrameNumber FrameDiff = ViewRange.Value - ViewRange.Key;
+							FrameDiff = FrameDiff * FloatViewDiff;
+							ActiveSequencer->SetLocalTimeDirectly(StartFrameNumber + FrameDiff);
+						}
+					}
+					bIsTracking = true;
+				}
+			}
+			else
+			{
+				bIsTracking = false;
+				StartXValue.Reset();
+			}
+		}
+		else
+		{
+			bIsTracking = false;
+			StartXValue.Reset();
+		}
+		return bIsTracking;
+	}
+	bIsTracking = false;
+	StartXValue.Reset();
+	return false;
+}
+
+bool FSequencerEdMode::MouseMove(FEditorViewportClient* ViewportClient, FViewport* InViewport, int32 X, int32 Y)
+{
+	const bool bTimeChange = IsPressingMoveTimeSlider(InViewport);
+	const bool bMouseButtonDown = InViewport->KeyState(EKeys::LeftMouseButton);
+
+	TSharedPtr<ISequencer> ActiveSequencer;
+	if (bTimeChange && bMouseButtonDown)
+	{
+		for (TWeakPtr<ISequencer> WeakSequencer : Sequencers)
+		{
+			ActiveSequencer = WeakSequencer.Pin();
+			if (ActiveSequencer.IsValid())
+			{
+				break;
+			}
+		}
+
+		if (ActiveSequencer.IsValid())
+		{
+			if (StartXValue.IsSet() == false)
+			{
+				StartXValue = X;
+				FQualifiedFrameTime CurrentTime = ActiveSequencer->GetLocalTime();
+				StartFrameNumber = CurrentTime.Time.GetFrame();
+			}
+			else
+			{
+				int32 Diff = StartXValue.GetValue() - X;
+				if (Diff != 0)
+				{		
+					FIntPoint Origin, Size;
+					ViewportClient->GetViewportDimensions(Origin, Size);
+					const float ViewPortSize = (float)Size.X;
+					const float FloatViewDiff = (float)(Diff) / ViewPortSize;
+					
+					FFrameRate TickResolution = ActiveSequencer->GetFocusedTickResolution();
+					TPair<FFrameNumber, FFrameNumber> ViewRange(TickResolution.AsFrameNumber(ActiveSequencer->GetViewRange().GetLowerBoundValue()), TickResolution.AsFrameNumber(ActiveSequencer->GetViewRange().GetUpperBoundValue()));
+					FFrameNumber FrameDiff = ViewRange.Value - ViewRange.Key;
+					FrameDiff = FrameDiff * FloatViewDiff;
+					ActiveSequencer->SetLocalTimeDirectly(StartFrameNumber + FrameDiff);
+				}
+			}
+			bIsTracking = true;
+		}
+		else
+		{
+			bIsTracking = false;
+			StartXValue.Reset();
+		}
+	}
+	else
+	{
+		bIsTracking = false;
+		StartXValue.Reset();
+	}
+	return bIsTracking;
+}
+
+
 void FSequencerEdMode::Render(const FSceneView* View, FViewport* Viewport, FPrimitiveDrawInterface* PDI)
 {
 	FEdMode::Render(View, Viewport, PDI);
diff --git a/Engine/Source/Editor/Sequencer/Private/SequencerEdMode.h b/Engine/Source/Editor/Sequencer/Private/SequencerEdMode.h
--- a/Engine/Source/Editor/Sequencer/Private/SequencerEdMode.h
+++ b/Engine/Source/Editor/Sequencer/Private/SequencerEdMode.h
@@ -58,6 +58,13 @@ public:
 	virtual bool UsesTransformWidget(UE::Widget::EWidgetMode CheckMode) const override { return false; }
 	virtual void AddReferencedObjects(FReferenceCollector& Collector) override;
 
+
+	virtual bool StartTracking(FEditorViewportClient* InViewportClient, FViewport* InViewport) override;
+	virtual bool MouseMove(FEditorViewportClient* ViewportClient, FViewport* Viewport, int32 x, int32 y) override;
+	virtual bool ProcessCapturedMouseMoves(FEditorViewportClient* InViewportClient, FViewport* InViewport, const TArrayView<FIntPoint>& CapturedMouseMoves) override;
+	virtual bool InputDelta(FEditorViewportClient* InViewportClient, FViewport* InViewport, FVector& InDrag, FRotator& InRot, FVector& InScale) override;
+
+	bool IsPressingMoveTimeSlider(FViewport* InViewport) const;
 	void AddSequencer(TWeakPtr<FSequencer> InSequencer) { Sequencers.AddUnique(InSequencer); }
 	void RemoveSequencer(TWeakPtr<FSequencer> InSequencer) { Sequencers.Remove(InSequencer); }
 
@@ -90,6 +97,14 @@ private:
 
 	/** The audio texture used for drawing the audio spatialization points */
 	UTexture2D* AudioTexture;
+
+	//params to handle mouse move for changing time
+	/** If we are tracking */
+	bool bIsTracking = false;
+	/** Starting X Value*/
+	TOptional<int32> StartXValue;
+	/** Starting Time Value*/
+	FFrameNumber StartFrameNumber;
 };
 
 /**
diff --git a/Engine/Source/Editor/Sequencer/Public/SequencerCommands.h b/Engine/Source/Editor/Sequencer/Public/SequencerCommands.h
--- a/Engine/Source/Editor/Sequencer/Public/SequencerCommands.h
+++ b/Engine/Source/Editor/Sequencer/Public/SequencerCommands.h
@@ -22,6 +22,12 @@ public:
 	/** Toggle play */
 	TSharedPtr< FUICommandInfo > TogglePlay;
 
+	/** Toggle play viewport*/
+	TSharedPtr< FUICommandInfo > TogglePlayViewport;
+
+	/** Scrub time in the viewport*/
+	TSharedPtr< FUICommandInfo > ScrubTimeViewport;
+
 	/** Play forward */
 	TSharedPtr< FUICommandInfo > PlayForward;
 
@@ -50,10 +56,10 @@ public:
 	TSharedPtr< FUICommandInfo > StepBackward;
 
 	/** Step forward */
-	TSharedPtr< FUICommandInfo > StepForward2;
+	TSharedPtr< FUICommandInfo > StepForwardViewport;
 
 	/** Step backward */
-	TSharedPtr< FUICommandInfo > StepBackward2;
+	TSharedPtr< FUICommandInfo > StepBackwardViewport;
 
 	/** Jump forward */
 	TSharedPtr< FUICommandInfo > JumpForward;
