DDC Fix: Same Chord
UE-148352

diff --git a/Engine/Source/Editor/Sequencer/Private/SequencerCommands.cpp b/Engine/Source/Editor/Sequencer/Private/SequencerCommands.cpp
--- a/Engine/Source/Editor/Sequencer/Private/SequencerCommands.cpp
+++ b/Engine/Source/Editor/Sequencer/Private/SequencerCommands.cpp
@@ -14,7 +14,7 @@ void FSequencerCommands::RegisterCommands()
 	UI_COMMAND( JumpToEnd, "Jump to End", "Jump to the end of the playback range", EUserInterfaceActionType::Button, FInputChord(EModifierKey::Control, EKeys::Up) );
 	UI_COMMAND( ShuttleBackward, "Shuttle Backward", "Shuttle backward", EUserInterfaceActionType::Button, FInputChord(EKeys::J) );
 	UI_COMMAND( ShuttleForward, "Shuttle Forward", "Shuttle forward", EUserInterfaceActionType::Button, FInputChord(EKeys::L) );
-	UI_COMMAND( Pause, "Pause", "Pause playback", EUserInterfaceActionType::Button, FInputChord(EKeys::K) );
+	UI_COMMAND( Pause, "Pause", "Pause playback", EUserInterfaceActionType::Button, FInputChord(EModifierKey::Shift,EKeys::K) );
 	UI_COMMAND( RestorePlaybackSpeed, "Restore Speed", "Restores the playback speed to 1.", EUserInterfaceActionType::Button, FInputChord() );
 	UI_COMMAND( StepForward, "Step Forward", "Step the timeline forward", EUserInterfaceActionType::Button, FInputChord(EKeys::Right) );
 	UI_COMMAND( StepBackward, "Step Backward", "Step the timeline backward", EUserInterfaceActionType::Button, FInputChord(EKeys::Left) );
