Sequencer: Viewport time change handles snapping.
#prefligh 624de6d3a64871f1a12487d3t
#trivial

diff --git a/Engine/Source/Editor/Sequencer/Private/SequencerEdMode.cpp b/Engine/Source/Editor/Sequencer/Private/SequencerEdMode.cpp
--- a/Engine/Source/Editor/Sequencer/Private/SequencerEdMode.cpp
+++ b/Engine/Source/Editor/Sequencer/Private/SequencerEdMode.cpp
@@ -228,6 +228,19 @@ bool FSequencerEdMode::InputDelta(FEditorViewportClient* InViewportClient, FView
 	return IsPressingMoveTimeSlider(InViewport);
 }
 
+static void SnapSequencerTime(TSharedPtr<ISequencer>& Sequencer, FFrameTime& ScrubTime)
+{
+
+	if (Sequencer->GetSequencerSettings()->GetIsSnapEnabled())
+	{
+		if (Sequencer->GetSequencerSettings()->GetSnapPlayTimeToInterval())
+		{
+			// Snap (round) to display rate
+			ScrubTime = FFrameRate::Snap(ScrubTime, Sequencer->GetFocusedTickResolution(), Sequencer->GetFocusedDisplayRate());
+		}
+	}
+}
+
 bool FSequencerEdMode::ProcessCapturedMouseMoves(FEditorViewportClient* InViewportClient, FViewport* InViewport, const TArrayView<FIntPoint>& CapturedMouseMoves)
 {
 	if (CapturedMouseMoves.Num() > 0)
@@ -271,7 +284,9 @@ bool FSequencerEdMode::ProcessCapturedMouseMoves(FEditorViewportClient* InViewpo
 							TPair<FFrameNumber, FFrameNumber> ViewRange(TickResolution.AsFrameNumber(ActiveSequencer->GetViewRange().GetLowerBoundValue()), TickResolution.AsFrameNumber(ActiveSequencer->GetViewRange().GetUpperBoundValue()));
 							FFrameNumber FrameDiff = ViewRange.Value - ViewRange.Key;
 							FrameDiff = FrameDiff * FloatViewDiff;
-							ActiveSequencer->SetLocalTimeDirectly(StartFrameNumber + FrameDiff);
+							FFrameTime ScrubTime = StartFrameNumber + FrameDiff;
+							SnapSequencerTime(ActiveSequencer, ScrubTime);
+							ActiveSequencer->SetLocalTime(ScrubTime.GetFrame());
 						}
 					}
 					bIsTracking = true;
@@ -334,7 +349,9 @@ bool FSequencerEdMode::MouseMove(FEditorViewportClient* ViewportClient, FViewpor
 					TPair<FFrameNumber, FFrameNumber> ViewRange(TickResolution.AsFrameNumber(ActiveSequencer->GetViewRange().GetLowerBoundValue()), TickResolution.AsFrameNumber(ActiveSequencer->GetViewRange().GetUpperBoundValue()));
 					FFrameNumber FrameDiff = ViewRange.Value - ViewRange.Key;
 					FrameDiff = FrameDiff * FloatViewDiff;
-					ActiveSequencer->SetLocalTimeDirectly(StartFrameNumber + FrameDiff);
+					FFrameTime ScrubTime = StartFrameNumber + FrameDiff;
+					SnapSequencerTime(ActiveSequencer, ScrubTime);
+					ActiveSequencer->SetLocalTime(ScrubTime.GetFrame());
 				}
 			}
 			bIsTracking = true;
