Sequencer: Add support for the new property setter functions, ie UPROPERTY(Setter = MyPropertySetter)

diff --git a/Engine/Source/Editor/Sequencer/Private/SequencerObjectChangeListener.cpp b/Engine/Source/Editor/Sequencer/Private/SequencerObjectChangeListener.cpp
--- a/Engine/Source/Editor/Sequencer/Private/SequencerObjectChangeListener.cpp
+++ b/Engine/Source/Editor/Sequencer/Private/SequencerObjectChangeListener.cpp
@@ -203,6 +203,8 @@ const FOnAnimatablePropertyChanged* FSequencerObjectChangeListener::FindProperty
 			PropertyVarName.RemoveFromStart("b", ESearchCase::CaseSensitive);
 		}
 
+		const bool bFoundSetter = Property.HasSetter();
+
 		static const FString Set(TEXT("Set"));
 
 		const FString FunctionString = Set + PropertyVarName;
@@ -247,7 +249,7 @@ const FOnAnimatablePropertyChanged* FSequencerObjectChangeListener::FindProperty
 		const bool bIsHiddenFunction = IsHiddenFunction(PropertyStructure, FAnimatedPropertyKey::FromProperty(&Property), Property.GetName());
 
 		// Valid if there's a setter function and the property is editable. Also valid if there's an interp keyword.
-		if (((bFoundValidFunction && bFoundEdit && !bFoundEditDefaultsOnly) || bFoundValidInterp) && !bIsHiddenFunction)
+		if (((bFoundValidFunction && bFoundEdit && !bFoundEditDefaultsOnly) || bFoundValidInterp || bFoundSetter) && !bIsHiddenFunction)
 		{
 			return DelegatePtr;
 		}
diff --git a/Engine/Source/Runtime/MovieScene/Private/EntitySystem/MovieScenePropertyRegistry.cpp b/Engine/Source/Runtime/MovieScene/Private/EntitySystem/MovieScenePropertyRegistry.cpp
--- a/Engine/Source/Runtime/MovieScene/Private/EntitySystem/MovieScenePropertyRegistry.cpp
+++ b/Engine/Source/Runtime/MovieScene/Private/EntitySystem/MovieScenePropertyRegistry.cpp
@@ -18,6 +18,13 @@ TOptional<uint16> ComputeFastPropertyPtrOffset(UClass* ObjectClass, const FMovie
 	using namespace UE::MovieScene;
 
 	FProperty* Property = ObjectClass->FindPropertyByName(PropertyBinding.PropertyName);
+	
+	const bool bFoundSetter = (!Property || Property->HasSetter());
+	if (bFoundSetter)
+	{
+		return TOptional<uint16>();
+	}
+
 	// @todo: Constructing FNames from strings is _very_ costly and we really shouldn't be doing this at runtime.
 	UFunction* Setter   = ObjectClass->FindFunctionByName(*(FString("Set") + PropertyBinding.PropertyName.ToString()));
 	if (Property && !Setter)
diff --git a/Engine/Source/Runtime/MovieScene/Public/MovieSceneCommonHelpers.h b/Engine/Source/Runtime/MovieScene/Public/MovieSceneCommonHelpers.h
--- a/Engine/Source/Runtime/MovieScene/Public/MovieSceneCommonHelpers.h
+++ b/Engine/Source/Runtime/MovieScene/Public/MovieSceneCommonHelpers.h
@@ -185,7 +185,13 @@ public:
 	void CallFunction( UObject& InRuntimeObject, typename TCallTraits<ValueType>::ParamType PropertyValue )
 	{
 		FPropertyAndFunction PropAndFunction = FindOrAdd(InRuntimeObject);
-		if (UFunction* SetterFunction = PropAndFunction.SetterFunction.Get())
+
+		FProperty* Property = GetProperty(InRuntimeObject);
+		if (Property && Property->HasSetter())
+		{
+			Property->CallSetter(&InRuntimeObject, &PropertyValue);
+		}
+		else if (UFunction* SetterFunction = PropAndFunction.SetterFunction.Get())
 		{
 			InvokeSetterFunction(&InRuntimeObject, SetterFunction, PropertyValue);
 		}
