Sequencer: Control Rig: Don't auto expand if you select a child control

diff --git a/Engine/Source/Editor/Sequencer/Private/Sequencer.cpp b/Engine/Source/Editor/Sequencer/Private/Sequencer.cpp
--- a/Engine/Source/Editor/Sequencer/Private/Sequencer.cpp
+++ b/Engine/Source/Editor/Sequencer/Private/Sequencer.cpp
@@ -7177,22 +7177,9 @@ void FSequencer::SelectByChannels(UMovieSceneSection* Section, TArrayView<const
 	
 	if (bSelect)
 	{
-		for (const TSharedRef<FSequencerDisplayNode>& DisplayNode : Nodes)
-		{
-			if (DisplayNode->GetParent().IsValid() && DisplayNode->GetParent()->GetType() == ESequencerNode::Category && !DisplayNode->GetParent()->IsExpanded())
-			{
-				DisplayNode->GetParent()->SetExpansionState(true);
-			}
-			//MAY NEED TO EXPAND TRACK ABOVE THE CATEGORY
-			if (DisplayNode->GetParent()->GetParent().IsValid() && DisplayNode->GetParent()->GetParent()->GetType() == ESequencerNode::Track && !DisplayNode->GetParent()->GetParent()->IsExpanded())
-			{
-				DisplayNode->GetParent()->GetParent()->SetExpansionState(true);
-			}
-			NodesToSelect.Add(DisplayNode);
-		}
-
 		if (NodesToSelect.Num() > 0)
 		{
+			//todo hide behind preference 
 			SequencerWidget->GetTreeView()->RequestScrollIntoView(NodesToSelect[0]);
 		
 			Selection.AddToSelection(NodesToSelect);
