Sequencer: Control Rig: Don't auto expand if you select a child contro pt2l

diff --git a/Engine/Source/Editor/Sequencer/Private/Sequencer.cpp b/Engine/Source/Editor/Sequencer/Private/Sequencer.cpp
--- a/Engine/Source/Editor/Sequencer/Private/Sequencer.cpp
+++ b/Engine/Source/Editor/Sequencer/Private/Sequencer.cpp
@@ -7116,15 +7116,6 @@ void FSequencer::SelectByNthCategoryNode(UMovieSceneSection* Section, int Index,
 	}
 	if (bSelect)
 	{
-		for (const TSharedRef<FSequencerDisplayNode>& DisplayNode : NodesToSelect)
-		{
-			if (DisplayNode->GetParent().IsValid() && DisplayNode->GetParent()->GetType() == ESequencerNode::Track && !DisplayNode->GetParent()->IsExpanded())
-			{
-				DisplayNode->GetParent()->SetExpansionState(true);
-				break;
-			}
-		}
-
 		if (NodesToSelect.Num() > 0)
 		{
 			SequencerWidget->GetTreeView()->RequestScrollIntoView(NodesToSelect[0]);
