Sequencer: Return shortcut k to Pause and set B as the shortcut for scrubbing time in the viewport

diff --git a/Engine/Source/Editor/Sequencer/Private/SequencerCommands.cpp b/Engine/Source/Editor/Sequencer/Private/SequencerCommands.cpp
--- a/Engine/Source/Editor/Sequencer/Private/SequencerCommands.cpp
+++ b/Engine/Source/Editor/Sequencer/Private/SequencerCommands.cpp
@@ -8,13 +8,13 @@ void FSequencerCommands::RegisterCommands()
 {
 	UI_COMMAND( TogglePlay, "Toggle Play", "Toggle the timeline playing", EUserInterfaceActionType::Button, FInputChord(EKeys::SpaceBar) );
 	UI_COMMAND( TogglePlayViewport, "Toggle Play(Viewport)", "Toggle the timeline playing in all viewports and sequencer", EUserInterfaceActionType::Button, FInputChord(EModifierKey::Shift, EKeys::SpaceBar));
-	UI_COMMAND( ScrubTimeViewport, "Scrub Time(Viewport)", "Scrub mouse left and right to change time", EUserInterfaceActionType::Button, FInputChord(EKeys::K) );
+	UI_COMMAND( ScrubTimeViewport, "Scrub Time(Viewport)", "Scrub mouse left and right to change time", EUserInterfaceActionType::Button, FInputChord(EKeys::B) );
 	UI_COMMAND( PlayForward, "Play Forward", "Play the timeline forward", EUserInterfaceActionType::Button, FInputChord(EKeys::Down) );
 	UI_COMMAND( JumpToStart, "Jump to Start", "Jump to the start of the playback range", EUserInterfaceActionType::Button, FInputChord(EKeys::Up) );
 	UI_COMMAND( JumpToEnd, "Jump to End", "Jump to the end of the playback range", EUserInterfaceActionType::Button, FInputChord(EModifierKey::Control, EKeys::Up) );
 	UI_COMMAND( ShuttleBackward, "Shuttle Backward", "Shuttle backward", EUserInterfaceActionType::Button, FInputChord(EKeys::J) );
 	UI_COMMAND( ShuttleForward, "Shuttle Forward", "Shuttle forward", EUserInterfaceActionType::Button, FInputChord(EKeys::L) );
-	UI_COMMAND( Pause, "Pause", "Pause playback", EUserInterfaceActionType::Button, FInputChord(EModifierKey::Shift,EKeys::K) );
+	UI_COMMAND( Pause, "Pause", "Pause playback", EUserInterfaceActionType::Button, FInputChord(EKeys::K) );
 	UI_COMMAND( RestorePlaybackSpeed, "Restore Speed", "Restores the playback speed to 1.", EUserInterfaceActionType::Button, FInputChord() );
 	UI_COMMAND( StepForward, "Step Forward", "Step the timeline forward", EUserInterfaceActionType::Button, FInputChord(EKeys::Right) );
 	UI_COMMAND( StepBackward, "Step Backward", "Step the timeline backward", EUserInterfaceActionType::Button, FInputChord(EKeys::Left) );
