Sequencer: Clamp viewport scrubbing with Keep Cursor in Playback Range while Scrubbing and make the snapping consistent with SequencerTimeSliderController

diff --git a/Engine/Source/Editor/Sequencer/Private/SequencerEdMode.cpp b/Engine/Source/Editor/Sequencer/Private/SequencerEdMode.cpp
--- a/Engine/Source/Editor/Sequencer/Private/SequencerEdMode.cpp
+++ b/Engine/Source/Editor/Sequencer/Private/SequencerEdMode.cpp
@@ -28,6 +28,7 @@
 #include "MovieSceneTracksComponentTypes.h"
 #include "EntitySystem/Interrogation/MovieSceneInterrogationLinker.h"
 #include "SequencerSectionPainter.h"
+#include "MovieSceneTimeHelpers.h"
 #include "MovieSceneToolHelpers.h"
 #include "Tools/MotionTrailOptions.h"
 #include "SequencerCommands.h"
@@ -228,15 +229,30 @@ bool FSequencerEdMode::InputDelta(FEditorViewportClient* InViewportClient, FView
 	return IsPressingMoveTimeSlider(InViewport);
 }
 
-static void SnapSequencerTime(TSharedPtr<ISequencer>& Sequencer, FFrameTime& ScrubTime)
+static void SnapSequencerTime(TSharedPtr<FSequencer>& Sequencer, FFrameTime& ScrubTime)
 {
+	// Clamp first, snap to frame last
+	if (Sequencer->GetSequencerSettings()->ShouldKeepCursorInPlayRangeWhileScrubbing())
+	{
+		TRange<FFrameNumber> PlaybackRange = Sequencer->GetFocusedMovieSceneSequence()->GetMovieScene()->GetPlaybackRange();
+		ScrubTime = UE::MovieScene::ClampToDiscreteRange(ScrubTime, PlaybackRange);
+	}
 
 	if (Sequencer->GetSequencerSettings()->GetIsSnapEnabled())
 	{
-		if (Sequencer->GetSequencerSettings()->GetSnapPlayTimeToInterval())
+		FFrameRate TickResolution =  Sequencer->GetFocusedTickResolution();
+		FFrameRate DisplayRate = Sequencer->GetFocusedDisplayRate();
+
+		// Set the style of the scrub handle
+		if (Sequencer->GetScrubStyle() == ESequencerScrubberStyle::FrameBlock)
+		{
+			// Floor to the display frame
+			ScrubTime = ConvertFrameTime(ConvertFrameTime(ScrubTime, TickResolution, DisplayRate).FloorToFrame(), DisplayRate, TickResolution);
+		}
+		else
 		{
 			// Snap (round) to display rate
-			ScrubTime = FFrameRate::Snap(ScrubTime, Sequencer->GetFocusedTickResolution(), Sequencer->GetFocusedDisplayRate());
+			ScrubTime = FFrameRate::Snap(ScrubTime, TickResolution, DisplayRate);
 		}
 	}
 }
@@ -248,9 +264,9 @@ bool FSequencerEdMode::ProcessCapturedMouseMoves(FEditorViewportClient* InViewpo
 		const bool bTimeChange = IsPressingMoveTimeSlider(InViewport);
 		if (bTimeChange)
 		{
-			TSharedPtr<ISequencer> ActiveSequencer;
+			TSharedPtr<FSequencer> ActiveSequencer;
 
-			for (TWeakPtr<ISequencer> WeakSequencer : Sequencers)
+			for (TWeakPtr<FSequencer> WeakSequencer : Sequencers)
 			{
 				ActiveSequencer = WeakSequencer.Pin();
 				if (ActiveSequencer.IsValid())
@@ -315,10 +331,10 @@ bool FSequencerEdMode::MouseMove(FEditorViewportClient* ViewportClient, FViewpor
 	const bool bTimeChange = IsPressingMoveTimeSlider(InViewport);
 	const bool bMouseButtonDown = InViewport->KeyState(EKeys::LeftMouseButton);
 
-	TSharedPtr<ISequencer> ActiveSequencer;
+	TSharedPtr<FSequencer> ActiveSequencer;
 	if (bTimeChange && bMouseButtonDown)
 	{
-		for (TWeakPtr<ISequencer> WeakSequencer : Sequencers)
+		for (TWeakPtr<FSequencer> WeakSequencer : Sequencers)
 		{
 			ActiveSequencer = WeakSequencer.Pin();
 			if (ActiveSequencer.IsValid())
