Editor/Sequencer:

-Expose commands to open curve editor and determine if curve editor is opened

diff --git a/Engine/Source/Editor/Sequencer/Private/Sequencer.h b/Engine/Source/Editor/Sequencer/Private/Sequencer.h
--- a/Engine/Source/Editor/Sequencer/Private/Sequencer.h
+++ b/Engine/Source/Editor/Sequencer/Private/Sequencer.h
@@ -501,9 +501,9 @@ public:
 	bool IsLevelEditorSequencer() const { return bIsEditingWithinLevelEditor; }
 
 	/** @return Whether to show the curve editor or not */
-	void SetShowCurveEditor(bool bInShowCurveEditor);
+	virtual void SetShowCurveEditor(bool bInShowCurveEditor) override;
 	/** @return If the curve editor is currently visible. */
-	bool GetCurveEditorIsVisible() const;
+	virtual bool GetCurveEditorIsVisible() const override;
 
 	/** Called to save the current movie scene */
 	void SaveCurrentMovieScene();
diff --git a/Engine/Source/Editor/Sequencer/Public/ISequencer.h b/Engine/Source/Editor/Sequencer/Public/ISequencer.h
--- a/Engine/Source/Editor/Sequencer/Public/ISequencer.h
+++ b/Engine/Source/Editor/Sequencer/Public/ISequencer.h
@@ -444,6 +444,11 @@ public:
 	*/
 	virtual void RecreateCurveEditor() {};
 
+	/** Whether to show the curve editor or not */
+	virtual void SetShowCurveEditor(bool bInShowCurveEditor) {}
+	/** @return If the curve editor is currently visible. */
+	virtual bool GetCurveEditorIsVisible() const { return false; }
+
 	/*
 	 * Puts sequencer in a silent state (whereby it will not redraw viewports, or attempt to update external state besides the sequence itself)
 	 */
