Sequencer: Fix for Animated Track Filter command registration.

When opening another sequence, FSequencer initializes the first sequence and then closes the previous sequence. This causes the track filter commands to be initialized for the first sequence and then destroyed when the subsequence sequence is opened. For now, register the commands before calling Get().

diff --git a/Engine/Source/Editor/Sequencer/Private/SequencerTrackFilters.cpp b/Engine/Source/Editor/Sequencer/Private/SequencerTrackFilters.cpp
--- a/Engine/Source/Editor/Sequencer/Private/SequencerTrackFilters.cpp
+++ b/Engine/Source/Editor/Sequencer/Private/SequencerTrackFilters.cpp
@@ -159,6 +159,14 @@ FSequencerTrackFilter_Animated::~FSequencerTrackFilter_Animated()
 
 FText FSequencerTrackFilter_Animated::GetToolTipText() const 
 { 
+	// When opening another sequence, FSequencer initializes the first sequence and then closes the previous sequence. 
+	// This causes the track filter commands to be initialized for the first sequence and then destroyed when the subsequence 
+	// sequence is opened. For now, register the commands before calling Get()
+	if (!FSequencerTrackFilter_AnimatedCommands::IsRegistered())
+	{
+		FSequencerTrackFilter_AnimatedCommands::Register();	
+	}
+
 	const FSequencerTrackFilter_AnimatedCommands& Commands = FSequencerTrackFilter_AnimatedCommands::Get();
 
 	const TSharedRef<const FInputChord> FirstActiveChord = Commands.ToggleAnimatedTracks->GetFirstValidChord();
@@ -174,6 +182,12 @@ FText FSequencerTrackFilter_Animated::GetToolTipText() const
 
 void FSequencerTrackFilter_Animated::BindCommands(TSharedRef<FUICommandList> CommandBindings, TWeakPtr<ISequencer> Sequencer)
 {
+	// See comment above
+	if (!FSequencerTrackFilter_AnimatedCommands::IsRegistered())
+	{
+		FSequencerTrackFilter_AnimatedCommands::Register();	
+	}
+
 	const FSequencerTrackFilter_AnimatedCommands& Commands = FSequencerTrackFilter_AnimatedCommands::Get();
 
 	CommandBindings->MapAction(
