Sequencer: Set the display name for properties and commands that use "Play time" to "Playhead"

diff --git a/Engine/Source/Editor/Sequencer/Private/SSequencer.cpp b/Engine/Source/Editor/Sequencer/Private/SSequencer.cpp
--- a/Engine/Source/Editor/Sequencer/Private/SSequencer.cpp
+++ b/Engine/Source/Editor/Sequencer/Private/SSequencer.cpp
@@ -2734,7 +2734,7 @@ TSharedRef<SWidget> SSequencer::MakeSnapMenu()
 	}
 	MenuBuilder.EndSection();
 
-	MenuBuilder.BeginSection( "PlayTimeSnapping", LOCTEXT( "SnappingMenuPlayTimeHeader", "Play Time Snapping" ) );
+	MenuBuilder.BeginSection( "PlayHeadSnapping", LOCTEXT( "SnappingMenuPlayHeadHeader", "Play Head Snapping" ) );
 	{
 		MenuBuilder.AddMenuEntry( FSequencerCommands::Get().ToggleSnapPlayTimeToInterval );
 		MenuBuilder.AddMenuEntry( FSequencerCommands::Get().ToggleSnapPlayTimeToKeys );
diff --git a/Engine/Source/Editor/Sequencer/Private/SequencerCommands.cpp b/Engine/Source/Editor/Sequencer/Private/SequencerCommands.cpp
--- a/Engine/Source/Editor/Sequencer/Private/SequencerCommands.cpp
+++ b/Engine/Source/Editor/Sequencer/Private/SequencerCommands.cpp
@@ -46,7 +46,7 @@ void FSequencerCommands::RegisterCommands()
 	UI_COMMAND( ToggleRerunConstructionScripts, "Rerun Construction Scripts", "Rerun construction scripts on bound actors every frame.", EUserInterfaceActionType::ToggleButton, FInputChord() );
 	UI_COMMAND( ToggleAsyncEvaluation, "Async Evaluation", "When enabled, enables a single asynchronous evaluation once per-frame. When disabled, forces a full blocking evaluation every time this sequence is evaluated (should be avoided for real-time content).", EUserInterfaceActionType::ToggleButton, FInputChord() );
 
-	UI_COMMAND( ToggleKeepCursorInPlaybackRangeWhileScrubbing, "Keep Cursor in Playback Range While Scrubbing", "When checked, the cursor will be constrained to the current playback range while scrubbing", EUserInterfaceActionType::ToggleButton, FInputChord() );
+	UI_COMMAND( ToggleKeepCursorInPlaybackRangeWhileScrubbing, "Keep Playhead in Playback Range While Scrubbing", "When checked, the playhead will be constrained to the current playback range while scrubbing", EUserInterfaceActionType::ToggleButton, FInputChord() );
 	UI_COMMAND( ToggleKeepPlaybackRangeInSectionBounds, "Keep Playback Range in Section Bounds", "When checked, the playback range will be synchronized to the section bounds", EUserInterfaceActionType::ToggleButton, FInputChord() );
 
 	UI_COMMAND( ToggleExpandCollapseNodes, "Expand/Collapse Nodes", "Toggle expand or collapse selected nodes", EUserInterfaceActionType::Button, FInputChord(EKeys::V) );
@@ -123,12 +123,12 @@ void FSequencerCommands::RegisterCommands()
 	UI_COMMAND( ToggleSnapSectionTimesToSections, "Snap to Keys and Sections", "Snap sections to other keys and sections in this sequence", EUserInterfaceActionType::ToggleButton, FInputChord() );
 	UI_COMMAND( ToggleSnapKeysAndSectionsToPlayRange, "Snap Keys and Sections to the Playback Range", "When checked, keys and sections will be snapped to the playback range bounds", EUserInterfaceActionType::ToggleButton, FInputChord());
 
-	UI_COMMAND( ToggleSnapPlayTimeToKeys, "Snap to Keys While Scrubbing", "Snap the current time to keys of the selected track while scrubbing", EUserInterfaceActionType::ToggleButton, FInputChord() );
-	UI_COMMAND( ToggleSnapPlayTimeToSections, "Snap to Sections While Scrubbing", "Snap the current time to section bounds while scrubbing", EUserInterfaceActionType::ToggleButton, FInputChord() );
-	UI_COMMAND( ToggleSnapPlayTimeToMarkers, "Snap to Markers While Scrubbing", "Snap the current time to markers while scrubbing", EUserInterfaceActionType::ToggleButton, FInputChord() );
-	UI_COMMAND( ToggleSnapPlayTimeToInterval, "Snap to the Interval While Scrubbing", "Snap the current time to the time snapping interval while scrubbing", EUserInterfaceActionType::ToggleButton, FInputChord() );
-	UI_COMMAND( ToggleSnapPlayTimeToPressedKey, "Snap to the Pressed Key", "Snap the current time to the pressed key", EUserInterfaceActionType::ToggleButton, FInputChord() );
-	UI_COMMAND( ToggleSnapPlayTimeToDraggedKey, "Snap to the Dragged Key", "Snap the current time to the dragged key", EUserInterfaceActionType::ToggleButton, FInputChord() );
+	UI_COMMAND( ToggleSnapPlayTimeToKeys, "Snap to Keys While Scrubbing", "Snap the playhead to keys of the selected track while scrubbing", EUserInterfaceActionType::ToggleButton, FInputChord() );
+	UI_COMMAND( ToggleSnapPlayTimeToSections, "Snap to Sections While Scrubbing", "Snap the playhead to section bounds while scrubbing", EUserInterfaceActionType::ToggleButton, FInputChord() );
+	UI_COMMAND( ToggleSnapPlayTimeToMarkers, "Snap to Markers While Scrubbing", "Snap the playhead to markers while scrubbing", EUserInterfaceActionType::ToggleButton, FInputChord() );
+	UI_COMMAND( ToggleSnapPlayTimeToInterval, "Snap to the Interval While Scrubbing", "Snap the playhead to the time snapping interval while scrubbing", EUserInterfaceActionType::ToggleButton, FInputChord() );
+	UI_COMMAND( ToggleSnapPlayTimeToPressedKey, "Snap to the Pressed Key", "Snap the playhead to the pressed key", EUserInterfaceActionType::ToggleButton, FInputChord() );
+	UI_COMMAND( ToggleSnapPlayTimeToDraggedKey, "Snap to the Dragged Key", "Snap the playhead to the dragged key", EUserInterfaceActionType::ToggleButton, FInputChord() );
 
 	UI_COMMAND( ToggleSnapCurveValueToInterval, "Snap Curve Key Values", "Snap curve keys to the value snapping interval", EUserInterfaceActionType::ToggleButton, FInputChord() );
 
diff --git a/Engine/Source/Editor/Sequencer/Public/SequencerSettings.h b/Engine/Source/Editor/Sequencer/Public/SequencerSettings.h
--- a/Engine/Source/Editor/Sequencer/Public/SequencerSettings.h
+++ b/Engine/Source/Editor/Sequencer/Public/SequencerSettings.h
@@ -29,7 +29,7 @@ UENUM()
 enum ESequencerZoomPosition
 {
 	/** Current Time. */
-	SZP_CurrentTime UMETA(DisplayName="Current Time"),
+	SZP_CurrentTime UMETA(DisplayName="Playhead"),
 
 	/** Mouse Position. */
 	SZP_MousePosition UMETA(DisplayName="Mouse Position"),
@@ -430,28 +430,28 @@ protected:
 	UPROPERTY(config, EditAnywhere, Category = Timeline)
 	bool bSnapKeysAndSectionsToPlayRange;
 
-	/** Enable or disable snapping the current time to keys while scrubbing. */
-	UPROPERTY( config, EditAnywhere, Category=Snapping )
+	/** Enable or disable snapping the playhead to keys while scrubbing. */
+	UPROPERTY( config, EditAnywhere, Category=Snapping, meta = (DisplayName = "Snap Playhead to Keys"))
 	bool bSnapPlayTimeToKeys;
 
-	/** Enable or disable snapping the current time to section bounds while scrubbing. */
-	UPROPERTY( config, EditAnywhere, Category=Snapping )
+	/** Enable or disable snapping the playhead to section bounds while scrubbing. */
+	UPROPERTY( config, EditAnywhere, Category=Snapping, meta = (DisplayName = "Snap Playhead to Sectins"))
 	bool bSnapPlayTimeToSections;
 
-	/** Enable or disable snapping the current time to markers while scrubbing. */
-	UPROPERTY( config, EditAnywhere, Category=Snapping )
+	/** Enable or disable snapping the playhead to markers while scrubbing. */
+	UPROPERTY( config, EditAnywhere, Category=Snapping, meta = (DisplayName = "Snap Playhead to Markers"))
 	bool bSnapPlayTimeToMarkers;
 
-	/** Enable or disable snapping the current time to the time snapping interval while scrubbing. */
-	UPROPERTY( config, EditAnywhere, Category=Snapping )
+	/** Enable or disable snapping the playhead to the time snapping interval while scrubbing. */
+	UPROPERTY( config, EditAnywhere, Category=Snapping, meta = (DisplayName = "Snap Playhead to Interval"))
 	bool bSnapPlayTimeToInterval;
 
-	/** Enable or disable snapping the current time to the pressed key. */
-	UPROPERTY( config, EditAnywhere, Category=Snapping )
+	/** Enable or disable snapping the playhead to the pressed key. */
+	UPROPERTY( config, EditAnywhere, Category=Snapping, meta = (DisplayName = "Snap Playhead to Pressed Key"))
 	bool bSnapPlayTimeToPressedKey;
 
-	/** Enable or disable snapping the current time to the dragged key. */
-	UPROPERTY( config, EditAnywhere, Category=Snapping )
+	/** Enable or disable snapping the playhead to the dragged key. */
+	UPROPERTY( config, EditAnywhere, Category=Snapping, meta = (DisplayName = "Snap Playhead to Dragged Key"))
 	bool bSnapPlayTimeToDraggedKey;
 
 	/** The curve value interval to snap to. */
@@ -496,8 +496,8 @@ protected:
 	UPROPERTY( config )
 	TEnumAsByte<ESequencerLoopMode> LoopMode;
 
-	/** Enable or disable keeping the cursor in the current playback range while scrubbing. */
-	UPROPERTY(config, EditAnywhere, Category = Timeline)
+	/** Enable or disable keeping the playhead in the current playback range while scrubbing. */
+	UPROPERTY(config, EditAnywhere, Category = Timeline, meta = (DisplayName = "Keep Playhead in Play Range While Scrubbing"))
 	bool bKeepCursorInPlayRangeWhileScrubbing;
 
 	/** Enable or disable keeping the playback range constrained to the section bounds. */
