Sequencer: Match outliner indent amount to 12 so that there's more separation between parent and child nodes

diff --git a/Engine/Source/Editor/Sequencer/Private/SSequencer.h b/Engine/Source/Editor/Sequencer/Private/SSequencer.h
--- a/Engine/Source/Editor/Sequencer/Private/SSequencer.h
+++ b/Engine/Source/Editor/Sequencer/Private/SSequencer.h
@@ -46,7 +46,7 @@ struct FSequencerSelectionCurveFilter;
 namespace SequencerLayoutConstants
 {
 	/** The amount to indent child nodes of the layout tree */
-	const float IndentAmount = 10.0f;
+	const float IndentAmount = 12.0f;
 
 	/** Height of each folder node */
 	const float FolderNodeHeight = 20.0f;
