Sequencer: Don't fire GetOnOutlinerNodeSelectionChanged when right clicking on a section because the node selection shouldn't be changing (just the selection)

diff --git a/Engine/Source/Editor/Sequencer/Private/SequencerCommonHelpers.cpp b/Engine/Source/Editor/Sequencer/Private/SequencerCommonHelpers.cpp
--- a/Engine/Source/Editor/Sequencer/Private/SequencerCommonHelpers.cpp
+++ b/Engine/Source/Editor/Sequencer/Private/SequencerCommonHelpers.cpp
@@ -331,7 +331,6 @@ void SequencerHelpers::PerformDefaultSelection(FSequencer& Sequencer, const FPoi
 		}
 		
 		Selection.ResumeBroadcast();
-		Selection.GetOnOutlinerNodeSelectionChanged().Broadcast();
 		return;
 	}
 
