Sequencer: Fix typo

diff --git a/Engine/Source/Editor/Sequencer/Public/SequencerSettings.h b/Engine/Source/Editor/Sequencer/Public/SequencerSettings.h
--- a/Engine/Source/Editor/Sequencer/Public/SequencerSettings.h
+++ b/Engine/Source/Editor/Sequencer/Public/SequencerSettings.h
@@ -435,7 +435,7 @@ protected:
 	bool bSnapPlayTimeToKeys;
 
 	/** Enable or disable snapping the playhead to section bounds while scrubbing. */
-	UPROPERTY( config, EditAnywhere, Category=Snapping, meta = (DisplayName = "Snap Playhead to Sectins"))
+	UPROPERTY( config, EditAnywhere, Category=Snapping, meta = (DisplayName = "Snap Playhead to Sections"))
 	bool bSnapPlayTimeToSections;
 
 	/** Enable or disable snapping the playhead to markers while scrubbing. */
