Sequencer:  Viewport scrubbing needs to go into scrub mode so audio plays.

diff --git a/Engine/Source/Editor/Sequencer/Private/SequencerEdMode.cpp b/Engine/Source/Editor/Sequencer/Private/SequencerEdMode.cpp
--- a/Engine/Source/Editor/Sequencer/Private/SequencerEdMode.cpp
+++ b/Engine/Source/Editor/Sequencer/Private/SequencerEdMode.cpp
@@ -201,6 +201,10 @@ bool FSequencerEdMode::InputKey( FEditorViewportClient* ViewportClient, FViewpor
 		{
 			return true;
 		}
+		if (IsPressingMoveTimeSlider(Viewport)) //this is needed to make sure we get all of the processed mouse events, for some reason the above may not return true
+		{
+			return true;
+		}
 	}
 	return FEdMode::InputKey(ViewportClient, Viewport, Key, Event);
 }
@@ -221,9 +225,28 @@ bool FSequencerEdMode::IsPressingMoveTimeSlider(FViewport* InViewport) const
 
 bool FSequencerEdMode::StartTracking(FEditorViewportClient* InViewportClient, FViewport* InViewport)
 {
-	return IsPressingMoveTimeSlider(InViewport);
+	if (IsPressingMoveTimeSlider(InViewport))
+	{
+		TSharedPtr<FSequencer> ActiveSequencer;
+		for (TWeakPtr<FSequencer> WeakSequencer : Sequencers)
+		{
+			ActiveSequencer = WeakSequencer.Pin();
+			if (ActiveSequencer.IsValid())
+			{
+				break;
+			}
+		}
+		if (ActiveSequencer.IsValid())
+		{
+			ActiveSequencer->SetPlaybackStatus(EMovieScenePlayerStatus::Scrubbing);
+			ActiveSequencer->OnBeginScrubbingEvent().Broadcast();
+		}
+		return true;
+	}
+	return false;
 }
 
+
 bool FSequencerEdMode::InputDelta(FEditorViewportClient* InViewportClient, FViewport* InViewport, FVector& InDrag, FRotator& InRot, FVector& InScale)
 {
 	return IsPressingMoveTimeSlider(InViewport);
@@ -259,23 +282,21 @@ static void SnapSequencerTime(TSharedPtr<FSequencer>& Sequencer, FFrameTime& Scr
 
 bool FSequencerEdMode::ProcessCapturedMouseMoves(FEditorViewportClient* InViewportClient, FViewport* InViewport, const TArrayView<FIntPoint>& CapturedMouseMoves)
 {
+	TSharedPtr<FSequencer> ActiveSequencer;
+	for (TWeakPtr<FSequencer> WeakSequencer : Sequencers)
+	{
+		ActiveSequencer = WeakSequencer.Pin();
+		if (ActiveSequencer.IsValid())
+		{
+			break;
+		}
+	}
+	const bool bTimeChange = IsPressingMoveTimeSlider(InViewport);
 	if (CapturedMouseMoves.Num() > 0)
 	{
-		const bool bTimeChange = IsPressingMoveTimeSlider(InViewport);
-		if (bTimeChange)
+		if (ActiveSequencer.IsValid())
 		{
-			TSharedPtr<FSequencer> ActiveSequencer;
-
-			for (TWeakPtr<FSequencer> WeakSequencer : Sequencers)
-			{
-				ActiveSequencer = WeakSequencer.Pin();
-				if (ActiveSequencer.IsValid())
-				{
-					break;
-				}
-			}
-
-			if (ActiveSequencer.IsValid())
+			if (bTimeChange)
 			{
 				for (int32 Index = 0; Index < CapturedMouseMoves.Num(); ++Index)
 				{
@@ -305,12 +326,19 @@ bool FSequencerEdMode::ProcessCapturedMouseMoves(FEditorViewportClient* InViewpo
 							ActiveSequencer->SetLocalTime(ScrubTime.GetFrame());
 						}
 					}
-					bIsTracking = true;
+					if (bIsTracking == false)
+					{
+						ActiveSequencer->SetPlaybackStatus(EMovieScenePlayerStatus::Scrubbing);
+						ActiveSequencer->OnBeginScrubbingEvent().Broadcast();
+						bIsTracking = true;
+					}
 				}
 			}
-			else
+			else if(bIsTracking)
 			{
 				bIsTracking = false;
+				ActiveSequencer->SetPlaybackStatus(EMovieScenePlayerStatus::Stopped);
+				ActiveSequencer->OnEndScrubbingEvent().Broadcast();
 				StartXValue.Reset();
 			}
 		}
@@ -321,56 +349,14 @@ bool FSequencerEdMode::ProcessCapturedMouseMoves(FEditorViewportClient* InViewpo
 		}
 		return bIsTracking;
 	}
-	bIsTracking = false;
-	StartXValue.Reset();
-	return false;
-}
-
-bool FSequencerEdMode::MouseMove(FEditorViewportClient* ViewportClient, FViewport* InViewport, int32 X, int32 Y)
-{
-	const bool bTimeChange = IsPressingMoveTimeSlider(InViewport);
-	const bool bMouseButtonDown = InViewport->KeyState(EKeys::LeftMouseButton);
-
-	TSharedPtr<FSequencer> ActiveSequencer;
-	if (bTimeChange && bMouseButtonDown)
+	else if (bTimeChange == false && bIsTracking)
 	{
-		for (TWeakPtr<FSequencer> WeakSequencer : Sequencers)
-		{
-			ActiveSequencer = WeakSequencer.Pin();
-			if (ActiveSequencer.IsValid())
-			{
-				break;
-			}
-		}
-
 		if (ActiveSequencer.IsValid())
 		{
-			if (StartXValue.IsSet() == false)
-			{
-				StartXValue = X;
-				FQualifiedFrameTime CurrentTime = ActiveSequencer->GetLocalTime();
-				StartFrameNumber = CurrentTime.Time.GetFrame();
-			}
-			else
-			{
-				int32 Diff = StartXValue.GetValue() - X;
-				if (Diff != 0)
-				{		
-					FIntPoint Origin, Size;
-					ViewportClient->GetViewportDimensions(Origin, Size);
-					const float ViewPortSize = (float)Size.X;
-					const float FloatViewDiff = (float)(Diff) / ViewPortSize;
-					
-					FFrameRate TickResolution = ActiveSequencer->GetFocusedTickResolution();
-					TPair<FFrameNumber, FFrameNumber> ViewRange(TickResolution.AsFrameNumber(ActiveSequencer->GetViewRange().GetLowerBoundValue()), TickResolution.AsFrameNumber(ActiveSequencer->GetViewRange().GetUpperBoundValue()));
-					FFrameNumber FrameDiff = ViewRange.Value - ViewRange.Key;
-					FrameDiff = FrameDiff * FloatViewDiff;
-					FFrameTime ScrubTime = StartFrameNumber + FrameDiff;
-					SnapSequencerTime(ActiveSequencer, ScrubTime);
-					ActiveSequencer->SetLocalTime(ScrubTime.GetFrame());
-				}
-			}
-			bIsTracking = true;
+			bIsTracking = false;
+			ActiveSequencer->SetPlaybackStatus(EMovieScenePlayerStatus::Stopped);
+			ActiveSequencer->OnEndScrubbingEvent().Broadcast();
+			StartXValue.Reset();
 		}
 		else
 		{
@@ -378,12 +364,13 @@ bool FSequencerEdMode::MouseMove(FEditorViewportClient* ViewportClient, FViewpor
 			StartXValue.Reset();
 		}
 	}
-	else
-	{
-		bIsTracking = false;
-		StartXValue.Reset();
-	}
-	return bIsTracking;
+	return false;
+}
+
+//Currently this is handled by the processed mouse events above, but don't fully trust ed modes this so leaving around for now
+bool FSequencerEdMode::MouseMove(FEditorViewportClient* ViewportClient, FViewport* InViewport, int32 X, int32 Y)
+{
+	return false;
 }
 
 
