Sequencer: Convert main toolbar to Tool Menus

diff --git a/Engine/Plugins/MovieScene/LevelSequenceEditor/Source/LevelSequenceEditor/LevelSequenceEditor.Build.cs b/Engine/Plugins/MovieScene/LevelSequenceEditor/Source/LevelSequenceEditor/LevelSequenceEditor.Build.cs
--- a/Engine/Plugins/MovieScene/LevelSequenceEditor/Source/LevelSequenceEditor/LevelSequenceEditor.Build.cs
+++ b/Engine/Plugins/MovieScene/LevelSequenceEditor/Source/LevelSequenceEditor/LevelSequenceEditor.Build.cs
@@ -46,6 +46,7 @@ public class LevelSequenceEditor : ModuleRules
                 "UnrealEd",
                 "VREditor",
 				"TimeManagement",
+				"ToolMenus",
 				"ToolWidgets",
 			}
 		);
diff --git a/Engine/Plugins/MovieScene/LevelSequenceEditor/Source/LevelSequenceEditor/Private/LevelSequenceEditorToolkit.cpp b/Engine/Plugins/MovieScene/LevelSequenceEditor/Source/LevelSequenceEditor/Private/LevelSequenceEditorToolkit.cpp
--- a/Engine/Plugins/MovieScene/LevelSequenceEditor/Source/LevelSequenceEditor/Private/LevelSequenceEditorToolkit.cpp
+++ b/Engine/Plugins/MovieScene/LevelSequenceEditor/Source/LevelSequenceEditor/Private/LevelSequenceEditorToolkit.cpp
@@ -2,6 +2,7 @@
 
 #include "LevelSequenceEditorToolkit.h"
 #include "Misc/LevelSequencePlaybackContext.h"
+#include "Misc/LevelSequenceEditorMenuContext.h"
 #include "Framework/MultiBox/MultiBoxBuilder.h"
 #include "UObject/UnrealType.h"
 #include "GameFramework/Actor.h"
@@ -51,6 +52,7 @@
 #include "MovieSceneCaptureDialogModule.h"
 #include "MovieScene.h"
 #include "UnrealEdMisc.h"
+#include "ToolMenus.h"
 
 // @todo sequencer: hack: setting defaults for transform tracks
 
@@ -191,17 +193,16 @@ void FLevelSequenceEditorToolkit::Initialize(const EToolkitMode::Type Mode, cons
 		SequencerInitParams.ViewParams.UniqueName = "LevelSequenceEditor";
 		SequencerInitParams.ViewParams.ScrubberStyle = ESequencerScrubberStyle::FrameBlock;
 		SequencerInitParams.ViewParams.OnReceivedFocus.BindRaw(this, &FLevelSequenceEditorToolkit::OnSequencerReceivedFocus);
+		SequencerInitParams.ViewParams.OnInitToolMenuContext.BindRaw(this, &FLevelSequenceEditorToolkit::OnInitToolMenuContext);
 
 		SequencerInitParams.HostCapabilities.bSupportsCurveEditor = true;
 		SequencerInitParams.HostCapabilities.bSupportsSaveMovieSceneAsset = true;
 		SequencerInitParams.HostCapabilities.bSupportsRecording = true;
 		SequencerInitParams.HostCapabilities.bSupportsRenderMovie = true;
-
-		TSharedRef<FExtender> ToolbarExtender = MakeShared<FExtender>();
-		ToolbarExtender->AddToolBarExtension("Base Commands", EExtensionHook::Before, nullptr, FToolBarExtensionDelegate::CreateSP(this, &FLevelSequenceEditorToolkit::ExtendSequencerToolbar));
-		SequencerInitParams.ViewParams.ToolbarExtender = ToolbarExtender;
 	}
 
+	ExtendSequencerToolbar("Sequencer.MainToolBar");
+
 	Sequencer = FModuleManager::LoadModuleChecked<ISequencerModule>("Sequencer").CreateSequencer(SequencerInitParams);
 	SpawnRegister->SetSequencer(Sequencer);
 	Sequencer->OnActorAddedToSequencer().AddSP(this, &FLevelSequenceEditorToolkit::HandleActorAddedToSequencer);
@@ -286,13 +287,37 @@ FString FLevelSequenceEditorToolkit::GetWorldCentricTabPrefix() const
 	return LOCTEXT("WorldCentricTabPrefix", "Sequencer ").ToString();
 }
 
-
 /* FLevelSequenceEditorToolkit implementation
  *****************************************************************************/
 
-void FLevelSequenceEditorToolkit::ExtendSequencerToolbar(FToolBarBuilder& ToolbarBuilder)
+void FLevelSequenceEditorToolkit::ExtendSequencerToolbar(FName InToolMenuName)
 {
-	ToolbarBuilder.AddWidget(PlaybackContext->BuildWorldPickerCombo());
+	FToolMenuOwnerScoped OwnerScoped(this);
+
+	UToolMenu* ToolMenu = UToolMenus::Get()->ExtendMenu(InToolMenuName);
+
+	const FToolMenuInsert SectionInsertLocation("BaseCommands", EToolMenuInsertType::Before);
+
+	{
+		ToolMenu->AddDynamicSection("LevelSequenceEditorDynamic", FNewToolMenuDelegate::CreateLambda([](UToolMenu* InMenu)
+		{	
+			ULevelSequenceEditorMenuContext* LevelSequenceEditorMenuContext = InMenu->FindContext<ULevelSequenceEditorMenuContext>();
+			if (LevelSequenceEditorMenuContext && LevelSequenceEditorMenuContext->Toolkit.IsValid())
+			{
+				const FName SequencerToolbarStyleName = "SequencerToolbar";
+			
+				FToolMenuEntry PlaybackContextEntry = FToolMenuEntry::InitWidget(
+					"PlaybackContext",
+					LevelSequenceEditorMenuContext->Toolkit.Pin()->PlaybackContext->BuildWorldPickerCombo(),
+					LOCTEXT("PlaybackContext", "PlaybackContext")
+				);
+				PlaybackContextEntry.StyleNameOverride = SequencerToolbarStyleName;
+
+				FToolMenuSection& Section = InMenu->AddSection("LevelSequenceEditor");
+				Section.AddEntry(PlaybackContextEntry);
+			}
+		}), SectionInsertLocation);
+	}
 }
 
 void FLevelSequenceEditorToolkit::AddDefaultTracksForActor(AActor& Actor, const FGuid Binding)
@@ -543,6 +568,13 @@ void FLevelSequenceEditorToolkit::OnSequencerReceivedFocus()
 	}
 }
 
+void FLevelSequenceEditorToolkit::OnInitToolMenuContext(FToolMenuContext& MenuContext)
+{
+	ULevelSequenceEditorMenuContext* LevelSequenceEditorMenuContext = NewObject<ULevelSequenceEditorMenuContext>();
+	LevelSequenceEditorMenuContext->Toolkit = SharedThis(this);
+	MenuContext.AddObject(LevelSequenceEditorMenuContext);
+}
+
 void FLevelSequenceEditorToolkit::HandleAddComponentActionExecute(UActorComponent* Component)
 {
 	const FScopedTransaction Transaction(LOCTEXT("AddComponent", "Add Component"));
diff --git a/Engine/Plugins/MovieScene/LevelSequenceEditor/Source/LevelSequenceEditor/Private/LevelSequenceEditorToolkit.h b/Engine/Plugins/MovieScene/LevelSequenceEditor/Source/LevelSequenceEditor/Private/LevelSequenceEditorToolkit.h
--- a/Engine/Plugins/MovieScene/LevelSequenceEditor/Source/LevelSequenceEditor/Private/LevelSequenceEditorToolkit.h
+++ b/Engine/Plugins/MovieScene/LevelSequenceEditor/Source/LevelSequenceEditor/Private/LevelSequenceEditorToolkit.h
@@ -116,9 +116,12 @@ protected:
 	/** Called whenever sequencer has received focus */
 	void OnSequencerReceivedFocus();
 
+	/** Called whenever sequencer in initializing tool menu context */
+	void OnInitToolMenuContext(FToolMenuContext& MenuContext);
+
 private:
 
-	void ExtendSequencerToolbar(FToolBarBuilder& ToolbarBuilder);
+	void ExtendSequencerToolbar(FName InToolMenuName);
 
 	/** Callback for executing the Add Component action. */
 	void HandleAddComponentActionExecute(UActorComponent* Component);
diff --git a/Engine/Plugins/MovieScene/LevelSequenceEditor/Source/LevelSequenceEditor/Private/Misc/LevelSequenceEditorMenuContext.h b/Engine/Plugins/MovieScene/LevelSequenceEditor/Source/LevelSequenceEditor/Private/Misc/LevelSequenceEditorMenuContext.h
new file mode 100644
--- /dev/null
+++ b/Engine/Plugins/MovieScene/LevelSequenceEditor/Source/LevelSequenceEditor/Private/Misc/LevelSequenceEditorMenuContext.h
@@ -0,0 +1,16 @@
+// Copyright Epic Games, Inc. All Rights Reserved.
+
+#pragma once
+
+#include "CoreMinimal.h"
+#include "ToolMenuContext.h"
+#include "LevelSequenceEditorToolkit.h"
+#include "LevelSequenceEditorMenuContext.generated.h"
+
+UCLASS()
+class ULevelSequenceEditorMenuContext : public UObject
+{
+	GENERATED_BODY()
+public:
+	TWeakPtr<FLevelSequenceEditorToolkit> Toolkit;
+};
diff --git a/Engine/Source/Editor/Sequencer/Private/SSequencer.cpp b/Engine/Source/Editor/Sequencer/Private/SSequencer.cpp
--- a/Engine/Source/Editor/Sequencer/Private/SSequencer.cpp
+++ b/Engine/Source/Editor/Sequencer/Private/SSequencer.cpp
@@ -46,6 +46,7 @@
 #include "ISequencerWidgetsModule.h"
 #include "ScopedTransaction.h"
 #include "SequencerTimeSliderController.h"
+#include "SequencerToolMenuContext.h"
 #include "SSequencerSectionOverlay.h"
 #include "SSequencerTrackArea.h"
 #include "SSequencerTrackOutliner.h"
@@ -104,6 +105,7 @@
 #include "SequencerCustomizationManager.h"
 #include "EditorActorFolders.h"
 #include "Tracks/MovieSceneEventTrack.h"
+#include "ToolMenus.h"
 
 #define LOCTEXT_NAMESPACE "Sequencer"
 
@@ -340,6 +342,7 @@ void SSequencer::Construct(const FArguments& InArgs, TSharedRef<FSequencer> InSe
 	OnMarkEndDrag = InArgs._OnMarkEndDrag;
 
 	OnReceivedFocus = InArgs._OnReceivedFocus;
+	OnInitToolMenuContext = InArgs._OnInitToolMenuContext;
 
 	RootCustomization.OnReceivedDragOver = InArgs._OnReceivedDragOver;
 	RootCustomization.OnReceivedDrop = InArgs._OnReceivedDrop;
@@ -1397,6 +1400,14 @@ TSharedRef<SWidget> SSequencer::MakeFilterButton()
 
 TSharedRef<SWidget> SSequencer::MakeToolBar()
 {
+	static const FName MenuName("Sequencer.MainToolBar");
+	
+	if (!UToolMenus::Get()->IsMenuRegistered(MenuName))
+	{
+		UToolMenu* Toolbar = UToolMenus::Get()->RegisterMenu(MenuName, NAME_None, EMultiBoxType::ToolBar);
+		Toolbar->AddDynamicSection("PopulateToolBar", FNewToolMenuDelegate::CreateStatic(&SSequencer::PopulateToolBar));
+	}
+
 	TArray<TSharedPtr<FExtender>> AllExtenders;
 	ISequencerModule& SequencerModule = FModuleManager::GetModuleChecked<ISequencerModule>("Sequencer");
 	AllExtenders.Add(SequencerModule.GetToolBarExtensibilityManager()->GetAllExtenders());
@@ -1404,96 +1415,156 @@ TSharedRef<SWidget> SSequencer::MakeToolBar()
 
 	TSharedPtr<FExtender> Extender = FExtender::Combine(AllExtenders);
 
-	FSlimHorizontalToolBarBuilder ToolBarBuilder( SequencerPtr.Pin()->GetCommandBindings(), FMultiBoxCustomization::None, Extender, true);
+	TSharedPtr<FSequencer> Sequencer = SequencerPtr.Pin();
+
+	USequencerToolMenuContext* ContextObject = NewObject<USequencerToolMenuContext>();
+	ContextObject->SequencerWidget = SharedThis(this);
+
+	FToolMenuContext Context(SequencerPtr.Pin()->GetCommandBindings(), Extender, ContextObject);
+
+	// Allow any toolkits to initialize their menu context
+	OnInitToolMenuContext.ExecuteIfBound(Context);
+
+	return UToolMenus::Get()->GenerateWidget(MenuName, Context);
+}
+
+void SSequencer::PopulateToolBar(UToolMenu* InMenu)
+{
+	const FName SequencerToolbarStyleName = "SequencerToolbar";
 	
-	ToolBarBuilder.BeginSection("Base Commands");
-	ToolBarBuilder.BeginStyleOverride("SequencerToolbar");
+	USequencerToolMenuContext* ContextObject = InMenu->FindContext<USequencerToolMenuContext>();
+	if (!ContextObject)
 	{
-		// General 
-		if (SequencerPtr.Pin()->IsLevelEditorSequencer())
+		return;
+	}
+
+	TSharedPtr<SSequencer> SequencerWidget = ContextObject->SequencerWidget.Pin();
+	if (!SequencerWidget.Get())
+	{
+		return;
+	}
+	
+	TWeakPtr<FSequencer> Sequencer = SequencerWidget.Get()->GetSequencer();
+	if (!Sequencer.IsValid())
+	{
+		return;
+	}
+
+	{
+		FToolMenuSection& Section = InMenu->AddSection("BaseCommands");
+
+		if (Sequencer.Pin()->IsLevelEditorSequencer())
 		{
 			TAttribute<FSlateIcon> SaveIcon;
-			SaveIcon.Bind(TAttribute<FSlateIcon>::FGetter::CreateLambda([&] {
-
-				TArray<UMovieScene*> MovieScenesToSave;
-				MovieSceneHelpers::GetDescendantMovieScenes(SequencerPtr.Pin()->GetRootMovieSceneSequence(), MovieScenesToSave);
-				for (UMovieScene* MovieSceneToSave : MovieScenesToSave)
+			SaveIcon.Bind(TAttribute<FSlateIcon>::FGetter::CreateLambda([Sequencer] 
+			{
+				if (Sequencer.IsValid())
 				{
-					UPackage* MovieScenePackageToSave = MovieSceneToSave->GetOuter()->GetOutermost();
-					if (MovieScenePackageToSave->IsDirty())
+					TArray<UMovieScene*> MovieScenesToSave;
+					MovieSceneHelpers::GetDescendantMovieScenes(Sequencer.Pin()->GetRootMovieSceneSequence(), MovieScenesToSave);
+					for (UMovieScene* MovieSceneToSave : MovieScenesToSave)
 					{
-						return FSlateIcon(FAppStyle::Get().GetStyleSetName(), "Icons.SaveChanged");
+						UPackage* MovieScenePackageToSave = MovieSceneToSave->GetOuter()->GetOutermost();
+						if (MovieScenePackageToSave->IsDirty())
+						{
+							return FSlateIcon(FAppStyle::Get().GetStyleSetName(), "Icons.SaveChanged");
+						}
 					}
 				}
 				return FSlateIcon(FAppStyle::Get().GetStyleSetName(), "Icons.Save");
 			}));
 
-			if (SequencerPtr.Pin()->GetHostCapabilities().bSupportsSaveMovieSceneAsset)
+			if (Sequencer.Pin()->GetHostCapabilities().bSupportsSaveMovieSceneAsset)
 			{
-				ToolBarBuilder.AddToolBarButton(
-					FUIAction(FExecuteAction::CreateSP(this, &SSequencer::OnSaveMovieSceneClicked)),
-					NAME_None,
+				FToolMenuEntry SaveEntry = FToolMenuEntry::InitToolBarButton(
+					"Save",
+					FUIAction(FExecuteAction::CreateSP(SequencerWidget.Get(), &SSequencer::OnSaveMovieSceneClicked)),
 					LOCTEXT("SaveDirtyPackages", "Save"),
 					LOCTEXT("SaveDirtyPackagesTooltip", "Saves the current sequence and any subsequences"),
 					SaveIcon
 				);
+				SaveEntry.StyleNameOverride = SequencerToolbarStyleName;
+				Section.AddEntry(SaveEntry);				
 			}
 
-			ToolBarBuilder.AddToolBarButton( FSequencerCommands::Get().FindInContentBrowser, NAME_None, TAttribute<FText>(), TAttribute<FText>(), FSlateIcon(FAppStyle::Get().GetStyleSetName(), "SystemWideCommands.FindInContentBrowser") );
-			ToolBarBuilder.AddToolBarButton( FSequencerCommands::Get().CreateCamera );
+			FToolMenuEntry FindInContentBrowserEntry = FToolMenuEntry::InitToolBarButton(FSequencerCommands::Get().FindInContentBrowser, TAttribute<FText>(), TAttribute<FText>(), FSlateIcon(FAppStyle::Get().GetStyleSetName(), "SystemWideCommands.FindInContentBrowser"));
+			FindInContentBrowserEntry.StyleNameOverride = SequencerToolbarStyleName;
+			Section.AddEntry(FindInContentBrowserEntry);
 
-			if (SequencerPtr.Pin()->GetHostCapabilities().bSupportsRenderMovie)
+			FToolMenuEntry CreateCameraEntry = FToolMenuEntry::InitToolBarButton(FSequencerCommands::Get().CreateCamera);
+			CreateCameraEntry.StyleNameOverride = SequencerToolbarStyleName;
+			Section.AddEntry(CreateCameraEntry);
+
+			if (Sequencer.Pin()->GetHostCapabilities().bSupportsRenderMovie)
 			{
-				ToolBarBuilder.AddToolBarButton( FSequencerCommands::Get().RenderMovie, NAME_None, TAttribute<FText>(), TAttribute<FText>(), FSlateIcon(FAppStyle::GetAppStyleSetName(), "LevelEditor.OpenCinematic"));
-				
-				ToolBarBuilder.AddComboButton(
+				FToolMenuEntry RenderMovieEntry = FToolMenuEntry::InitToolBarButton(FSequencerCommands::Get().RenderMovie, TAttribute<FText>(), TAttribute<FText>(), FSlateIcon(FAppStyle::GetAppStyleSetName(), "LevelEditor.OpenCinematic"));
+				RenderMovieEntry.StyleNameOverride = SequencerToolbarStyleName;
+				Section.AddEntry(RenderMovieEntry);
+			
+				FToolMenuEntry RenderMovieOptionsEntry = FToolMenuEntry::InitComboButton(
+					"RenderMovieOptions",
 					FUIAction(),
-					FOnGetContent::CreateSP( this, &SSequencer::MakeRenderMovieMenu ),
-					LOCTEXT( "RenderMovieOptions", "Render Movie Options" ),
-					LOCTEXT( "RenderMovieOptionsToolTip", "Render Movie Options" ),
+					FOnGetContent::CreateSP(SequencerWidget.Get(), &SSequencer::MakeRenderMovieMenu),
+					LOCTEXT("RenderMovieOptions", "Render Movie Options"),
+					LOCTEXT("RenderMovieOptionsToolTip", "Render Movie Options"),
 					TAttribute<FSlateIcon>(),
-					true );
+					true
+				);
+				RenderMovieOptionsEntry.StyleNameOverride = SequencerToolbarStyleName;
+				Section.AddEntry(RenderMovieOptionsEntry);
 			}
 
-			UMovieSceneSequence* RootSequence = SequencerPtr.Pin()->GetRootMovieSceneSequence();
+			UMovieSceneSequence* RootSequence = Sequencer.Pin()->GetRootMovieSceneSequence();
 			if (RootSequence->GetTypedOuter<UBlueprint>() == nullptr && UMovieScene::IsTrackClassAllowed(UMovieSceneEventTrack::StaticClass()))
 			{
 				// Only show this button where it makes sense (ie, if the sequence is not contained within a blueprint already)
-				ToolBarBuilder.AddToolBarButton(FSequencerCommands::Get().OpenDirectorBlueprint, NAME_None, TAttribute<FText>(), TAttribute<FText>(), FSlateIcon(FAppStyle::Get().GetStyleSetName(), "LevelEditor.OpenLevelBlueprint"));
+				FToolMenuEntry OpenDirectorBlueprintEntry = FToolMenuEntry::InitToolBarButton(FSequencerCommands::Get().OpenDirectorBlueprint, TAttribute<FText>(), TAttribute<FText>(), FSlateIcon(FAppStyle::Get().GetStyleSetName(), "LevelEditor.OpenLevelBlueprint"));
+				OpenDirectorBlueprintEntry.StyleNameOverride = SequencerToolbarStyleName;
+				Section.AddEntry(OpenDirectorBlueprintEntry);
 			}
 
-			ToolBarBuilder.AddSeparator("Level Sequence Separator");
+			Section.AddSeparator(NAME_None);
 		}
 
-		ToolBarBuilder.AddComboButton(
+		FToolMenuEntry ActionsEntry = FToolMenuEntry::InitComboButton(
+			"Actions",
 			FUIAction(),
-			FOnGetContent::CreateSP(this, &SSequencer::MakeActionsMenu),
+			FOnGetContent::CreateSP(SequencerWidget.Get(), &SSequencer::MakeActionsMenu),
 			LOCTEXT("Actions", "Actions"),
 			LOCTEXT("ActionsToolTip", "Actions"),
 			FSlateIcon(FAppStyle::GetAppStyleSetName(), "Sequencer.Actions")
 		);
+		ActionsEntry.StyleNameOverride = SequencerToolbarStyleName;
+		Section.AddEntry(ActionsEntry);
 
-		ToolBarBuilder.AddComboButton(
+		FToolMenuEntry ViewOptionsEntry = FToolMenuEntry::InitComboButton(
+			"ViewOptions",
 			FUIAction(),
-			FOnGetContent::CreateSP(this, &SSequencer::MakeViewMenu),
+			FOnGetContent::CreateSP(SequencerWidget.Get(), &SSequencer::MakeViewMenu),
 			LOCTEXT("ViewOptions", "View Options"),
 			LOCTEXT("ViewOptionsToolTip", "View Options"),
 			FSlateIcon(FAppStyle::Get().GetStyleSetName(), "Icons.Visibility")
 		);
+		ViewOptionsEntry.StyleNameOverride = SequencerToolbarStyleName;
+		Section.AddEntry(ViewOptionsEntry);
 
-		ToolBarBuilder.AddComboButton(
+		FToolMenuEntry PlaybackOptionsEntry = FToolMenuEntry::InitComboButton(
+			"PlaybackOptions",
 			FUIAction(),
-			FOnGetContent::CreateSP(this, &SSequencer::MakePlaybackMenu),
+			FOnGetContent::CreateSP(SequencerWidget.Get(), &SSequencer::MakePlaybackMenu),
 			LOCTEXT("PlaybackOptions", "Playback Options"),
 			LOCTEXT("PlaybackOptionsToolTip", "Playback Options"),
 			FSlateIcon(FAppStyle::GetAppStyleSetName(), "Sequencer.PlaybackOptions")
 		);
+		PlaybackOptionsEntry.StyleNameOverride = SequencerToolbarStyleName;
+		Section.AddEntry(PlaybackOptionsEntry);
 
-		ToolBarBuilder.AddSeparator();
+		Section.AddSeparator(NAME_None);
 
 		TAttribute<FSlateIcon> KeyGroupModeIcon;
-		KeyGroupModeIcon.Bind(TAttribute<FSlateIcon>::FGetter::CreateLambda([&] {
-			switch (SequencerPtr.Pin()->GetKeyGroupMode())
+		KeyGroupModeIcon.Bind(TAttribute<FSlateIcon>::FGetter::CreateLambda([Sequencer] 
+		{
+			switch (Sequencer.Pin()->GetKeyGroupMode())
 			{
 			case EKeyGroupMode::KeyAll:
 				return FSequencerCommands::Get().SetKeyAll->GetIcon();
@@ -1505,8 +1576,9 @@ TSharedRef<SWidget> SSequencer::MakeToolBar()
 		}));
 
 		TAttribute<FText> KeyGroupModeToolTip;
-		KeyGroupModeToolTip.Bind(TAttribute<FText>::FGetter::CreateLambda([&] {
-			switch (SequencerPtr.Pin()->GetKeyGroupMode())
+		KeyGroupModeToolTip.Bind(TAttribute<FText>::FGetter::CreateLambda([Sequencer] 
+		{
+			switch (Sequencer.Pin()->GetKeyGroupMode())
 			{
 			case EKeyGroupMode::KeyAll:
 				return FSequencerCommands::Get().SetKeyAll->GetDescription();
@@ -1517,18 +1589,23 @@ TSharedRef<SWidget> SSequencer::MakeToolBar()
 			}
 		}));
 
-		ToolBarBuilder.AddComboButton(
+		FToolMenuEntry KeyGroupEntry = FToolMenuEntry::InitComboButton(
+			"KeyGroup",
 			FUIAction(),
-			FOnGetContent::CreateSP(this, &SSequencer::MakeKeyGroupMenu),
+			FOnGetContent::CreateSP(SequencerWidget.Get(), &SSequencer::MakeKeyGroupMenu),
 			LOCTEXT("KeyGroup", "Key All"),
 			KeyGroupModeToolTip,
-			KeyGroupModeIcon);
-
-		if (IVREditorModule::Get().IsVREditorModeActive() || (SequencerPtr.Pin()->IsLevelEditorSequencer() && ExactCast<ULevelSequence>(SequencerPtr.Pin()->GetFocusedMovieSceneSequence()) == nullptr))
+			KeyGroupModeIcon
+		);
+		KeyGroupEntry.StyleNameOverride = SequencerToolbarStyleName;
+		Section.AddEntry(KeyGroupEntry);
+	
+		if (IVREditorModule::Get().IsVREditorModeActive() || (Sequencer.Pin()->IsLevelEditorSequencer() && ExactCast<ULevelSequence>(Sequencer.Pin()->GetFocusedMovieSceneSequence()) == nullptr))
 		{
 			TAttribute<FSlateIcon> AutoChangeModeIcon;
-			AutoChangeModeIcon.Bind(TAttribute<FSlateIcon>::FGetter::CreateLambda( [&] {
-				switch ( SequencerPtr.Pin()->GetAutoChangeMode() )
+			AutoChangeModeIcon.Bind(TAttribute<FSlateIcon>::FGetter::CreateLambda([Sequencer] 
+			{
+				switch ( Sequencer.Pin()->GetAutoChangeMode() )
 				{
 				case EAutoChangeMode::AutoKey:
 					return FSequencerCommands::Get().SetAutoKey->GetIcon();
@@ -1542,8 +1619,9 @@ TSharedRef<SWidget> SSequencer::MakeToolBar()
 			} ) );
 
 			TAttribute<FText> AutoChangeModeToolTip;
-			AutoChangeModeToolTip.Bind( TAttribute<FText>::FGetter::CreateLambda( [&] {
-				switch ( SequencerPtr.Pin()->GetAutoChangeMode() )
+			AutoChangeModeToolTip.Bind( TAttribute<FText>::FGetter::CreateLambda([Sequencer] 
+			{
+				switch ( Sequencer.Pin()->GetAutoChangeMode() )
 				{
 				case EAutoChangeMode::AutoKey:
 					return FSequencerCommands::Get().SetAutoKey->GetDescription();
@@ -1556,23 +1634,30 @@ TSharedRef<SWidget> SSequencer::MakeToolBar()
 				}
 			} ) );
 			
-			ToolBarBuilder.AddComboButton(
+			FToolMenuEntry AutoChangeEntry = FToolMenuEntry::InitComboButton(
+				"AutoChange",
 				FUIAction(),
-				FOnGetContent::CreateSP(this, &SSequencer::MakeAutoChangeMenu),
+				FOnGetContent::CreateSP(SequencerWidget.Get(), &SSequencer::MakeAutoChangeMenu),
 				LOCTEXT("AutoChangeMode", "Auto-Change Mode"),
 				AutoChangeModeToolTip,
-				AutoChangeModeIcon);
+				AutoChangeModeIcon
+			);
+			AutoChangeEntry.StyleNameOverride = SequencerToolbarStyleName;
+			Section.AddEntry(AutoChangeEntry);
 		}
-		else
+		else	
 		{
-			ToolBarBuilder.AddToolBarButton( FSequencerCommands::Get().ToggleAutoKeyEnabled );
+			FToolMenuEntry ToggleAutoKeyEntry = FToolMenuEntry::InitToolBarButton(FSequencerCommands::Get().ToggleAutoKeyEnabled);
+			ToggleAutoKeyEntry.StyleNameOverride = SequencerToolbarStyleName;
+			Section.AddEntry(ToggleAutoKeyEntry);
 		}
 
-		if( SequencerPtr.Pin()->IsLevelEditorSequencer() )
+		if (Sequencer.Pin()->IsLevelEditorSequencer())
 		{
 			TAttribute<FSlateIcon> AllowEditsModeIcon;
-			AllowEditsModeIcon.Bind(TAttribute<FSlateIcon>::FGetter::CreateLambda( [&] {
-				switch ( SequencerPtr.Pin()->GetAllowEditsMode() )
+			AllowEditsModeIcon.Bind(TAttribute<FSlateIcon>::FGetter::CreateLambda([Sequencer] 
+			{
+				switch ( Sequencer.Pin()->GetAllowEditsMode() )
 				{
 				case EAllowEditsMode::AllEdits:
 					return FSequencerCommands::Get().AllowAllEdits->GetIcon();
@@ -1584,8 +1669,9 @@ TSharedRef<SWidget> SSequencer::MakeToolBar()
 			} ) );
 
 			TAttribute<FText> AllowEditsModeToolTip;
-			AllowEditsModeToolTip.Bind( TAttribute<FText>::FGetter::CreateLambda( [&] {
-				switch ( SequencerPtr.Pin()->GetAllowEditsMode() )
+			AllowEditsModeToolTip.Bind( TAttribute<FText>::FGetter::CreateLambda([Sequencer] 
+			{
+				switch ( Sequencer.Pin()->GetAllowEditsMode() )
 				{
 				case EAllowEditsMode::AllEdits:
 					return FSequencerCommands::Get().AllowAllEdits->GetDescription();
@@ -1596,50 +1682,64 @@ TSharedRef<SWidget> SSequencer::MakeToolBar()
 				}
 			} ) );
 
-			ToolBarBuilder.AddComboButton(
+			FToolMenuEntry AllowEditsModeEntry = FToolMenuEntry::InitComboButton(
+				"AllowEditsMode",
 				FUIAction(),
-				FOnGetContent::CreateSP(this, &SSequencer::MakeAllowEditsMenu),
-				LOCTEXT("AllowMode", "Allow Edits"),
+				FOnGetContent::CreateSP(SequencerWidget.Get(), &SSequencer::MakeAllowEditsMenu),
+				LOCTEXT("AllowEditsMode", "Allow Edits"),
 				AllowEditsModeToolTip,
-				AllowEditsModeIcon);
+				AllowEditsModeIcon
+			);
+			AllowEditsModeEntry.StyleNameOverride = SequencerToolbarStyleName;
+			Section.AddEntry(AllowEditsModeEntry);
 		}
-	}
-	ToolBarBuilder.EndSection();
 
+		Section.AddSeparator(NAME_None);
+	}
 
-	ToolBarBuilder.BeginSection("Snapping");
 	{
-		ToolBarBuilder.AddToolBarButton( FSequencerCommands::Get().ToggleIsSnapEnabled, NAME_None, TAttribute<FText>( FText::GetEmpty() ) );
+		FToolMenuSection& SnappingSection = InMenu->AddSection("Snapping");
+
+		FToolMenuEntry ToggleIsSnapEnabledEntry = FToolMenuEntry::InitToolBarButton(FSequencerCommands::Get().ToggleIsSnapEnabled, TAttribute<FText>(FText::GetEmpty()));
+		ToggleIsSnapEnabledEntry.StyleNameOverride = SequencerToolbarStyleName;
+		SnappingSection.AddEntry(ToggleIsSnapEnabledEntry);
 
-		ToolBarBuilder.AddComboButton(
+		FToolMenuEntry SnapOptionsEntry = FToolMenuEntry::InitComboButton(
+			"SnapOptions",
 			FUIAction(),
-			FOnGetContent::CreateSP( this, &SSequencer::MakeSnapMenu ),
-			LOCTEXT( "SnapOptions", "Options" ),
-			LOCTEXT( "SnapOptionsToolTip", "Snapping Options" ),
+			FOnGetContent::CreateSP(SequencerWidget.Get(), &SSequencer::MakeSnapMenu),
+			LOCTEXT("SnapOptions", "Options"),
+			LOCTEXT("SnapOptionsToolTip", "Snapping Options"),
 			TAttribute<FSlateIcon>(),
-			true );
+			true
+		);
+		SnapOptionsEntry.StyleNameOverride = SequencerToolbarStyleName;
+		SnappingSection.AddEntry(SnapOptionsEntry);
 
-		ToolBarBuilder.AddSeparator();
+		SnappingSection.AddSeparator(NAME_None);
 
-		ToolBarBuilder.AddWidget(SNew(SSequencerPlayRateCombo, SequencerPtr.Pin(), SharedThis(this)));
+		FToolMenuEntry PlayRateEntry = FToolMenuEntry::InitWidget(
+			"PlayRate",
+			SNew(SSequencerPlayRateCombo, Sequencer.Pin(), SequencerWidget.ToSharedRef()),
+			LOCTEXT("PlayRate", "PlayRate")
+		);
+		PlayRateEntry.StyleNameOverride = SequencerToolbarStyleName;
+		SnappingSection.AddEntry(PlayRateEntry);
 	}
-	ToolBarBuilder.EndSection();
 
-	ToolBarBuilder.BeginSection("Curve Editor");
 	{
+		FToolMenuSection& CurveEditorSection = InMenu->AddSection("CurveEditor");
+
 		// Only add the button if supported
-		if (SequencerPtr.Pin()->GetHostCapabilities().bSupportsCurveEditor)
+		if (Sequencer.Pin()->GetHostCapabilities().bSupportsCurveEditor)
 		{
-			ToolBarBuilder.AddToolBarButton( FSequencerCommands::Get().ToggleShowCurveEditor, NAME_None, TAttribute<FText>(), TAttribute<FText>(), FSlateIcon(FAppStyle::GetAppStyleSetName(), "GenericCurveEditor"));
+			FToolMenuEntry ShowCurveEditorEntry = FToolMenuEntry::InitToolBarButton(FSequencerCommands::Get().ToggleShowCurveEditor, TAttribute<FText>(), TAttribute<FText>(), FSlateIcon(FAppStyle::GetAppStyleSetName(), "GenericCurveEditor"));
+			ShowCurveEditorEntry.StyleNameOverride = SequencerToolbarStyleName;
+			CurveEditorSection.AddEntry(ShowCurveEditorEntry);
 		}
 	}
-	ToolBarBuilder.EndSection();
-	ToolBarBuilder.EndStyleOverride();
-
-	return ToolBarBuilder.MakeWidget();
 }
 
-
 void SSequencer::GetContextMenuContent(FMenuBuilder& MenuBuilder)
 {
 	// let toolkits populate the menu
diff --git a/Engine/Source/Editor/Sequencer/Private/SSequencer.h b/Engine/Source/Editor/Sequencer/Private/SSequencer.h
--- a/Engine/Source/Editor/Sequencer/Private/SSequencer.h
+++ b/Engine/Source/Editor/Sequencer/Private/SSequencer.h
@@ -249,6 +249,9 @@ public:
 		/** Called when any widget contained within sequencer has received focus */
 		SLATE_EVENT( FSimpleDelegate, OnReceivedFocus )
 
+		/** Called when initializing tool menu context */
+		SLATE_EVENT(FOnInitToolMenuContext, OnInitToolMenuContext)
+
 		/** Called when something is dragged over the sequencer. */
 		SLATE_EVENT( FOptionalOnDragDrop, OnReceivedDragOver )
 
@@ -400,6 +403,10 @@ private:
 	/** Get context menu contents. */
 	void GetContextMenuContent(FMenuBuilder& MenuBuilder);
 
+	TWeakPtr<FSequencer> GetSequencer() { return SequencerPtr; }
+
+	static void PopulateToolBar(UToolMenu* InMenu);
+
 	/** Makes the toolbar. */
 	TSharedRef<SWidget> MakeToolBar();
 
@@ -744,6 +751,9 @@ private:
 	/** Called when any widget contained within sequencer has received focus */
 	FSimpleDelegate OnReceivedFocus;
 
+	/** Called when initializing tool menu context */
+	FOnInitToolMenuContext OnInitToolMenuContext;
+
 	/** Called when something is dragged over the sequencer. */
 	TArray<FOptionalOnDragDrop> OnReceivedDragOver;
 
diff --git a/Engine/Source/Editor/Sequencer/Private/Sequencer.cpp b/Engine/Source/Editor/Sequencer/Private/Sequencer.cpp
--- a/Engine/Source/Editor/Sequencer/Private/Sequencer.cpp
+++ b/Engine/Source/Editor/Sequencer/Private/Sequencer.cpp
@@ -477,6 +477,7 @@ void FSequencer::InitSequencer(const FSequencerInitParams& InitParams, const TSh
 		.OnBuildCustomContextMenuForGuid(InitParams.ViewParams.OnBuildCustomContextMenuForGuid)
 		.OnGetPlaybackSpeeds(InitParams.ViewParams.OnGetPlaybackSpeeds)
 		.OnReceivedFocus(InitParams.ViewParams.OnReceivedFocus)
+		.OnInitToolMenuContext(InitParams.ViewParams.OnInitToolMenuContext)
 		.AddMenuExtender(InitParams.ViewParams.AddMenuExtender)
 		.ToolbarExtender(InitParams.ViewParams.ToolbarExtender);
 
diff --git a/Engine/Source/Editor/Sequencer/Private/SequencerToolMenuContext.h b/Engine/Source/Editor/Sequencer/Private/SequencerToolMenuContext.h
new file mode 100644
--- /dev/null
+++ b/Engine/Source/Editor/Sequencer/Private/SequencerToolMenuContext.h
@@ -0,0 +1,16 @@
+// Copyright Epic Games, Inc. All Rights Reserved.
+
+#pragma once
+
+#include "CoreMinimal.h"
+#include "ToolMenuContext.h"
+#include "SSequencer.h"
+#include "SequencerToolMenuContext.generated.h"
+
+UCLASS()
+class USequencerToolMenuContext : public UObject
+{
+	GENERATED_BODY()
+public:
+	TWeakPtr<SSequencer> SequencerWidget;
+};
\ No newline at end of file
diff --git a/Engine/Source/Editor/Sequencer/Public/ISequencerModule.h b/Engine/Source/Editor/Sequencer/Public/ISequencerModule.h
--- a/Engine/Source/Editor/Sequencer/Public/ISequencerModule.h
+++ b/Engine/Source/Editor/Sequencer/Public/ISequencerModule.h
@@ -52,6 +52,9 @@ DECLARE_DELEGATE_RetVal_OneParam(TSharedRef<ISequencerEditorObjectBinding>, FOnC
 /** A delegate that is executed when adding menu content. */
 DECLARE_DELEGATE_TwoParams(FOnGetAddMenuContent, FMenuBuilder& /*MenuBuilder*/, TSharedRef<ISequencer>);
 
+/** A delegate that is executed when initializing tool menu context. */
+DECLARE_DELEGATE_OneParam(FOnInitToolMenuContext, FToolMenuContext&);
+
 /** A delegate that is executed when menu object is clicked. Unlike FExtender delegates we pass in the FGuid which exists even for deleted objects. */
 DECLARE_DELEGATE_TwoParams(FOnBuildCustomContextMenuForGuid, FMenuBuilder&, FGuid);
 
@@ -73,6 +76,9 @@ struct FSequencerViewParams
 	/** Called when this sequencer has received user focus */
 	FSimpleDelegate OnReceivedFocus;
 
+	/** Called when this sequencer is initializing the tool menu context */
+	FOnInitToolMenuContext OnInitToolMenuContext;
+
 	/** The playback speed options available */
 	ISequencer::FOnGetPlaybackSpeeds OnGetPlaybackSpeeds;
 	
