Sequencer: Fix nonunity

diff --git a/Engine/Source/Editor/Sequencer/Private/SSequencer.h b/Engine/Source/Editor/Sequencer/Private/SSequencer.h
--- a/Engine/Source/Editor/Sequencer/Private/SSequencer.h
+++ b/Engine/Source/Editor/Sequencer/Private/SSequencer.h
@@ -12,6 +12,8 @@
 #include "Widgets/SCompoundWidget.h"
 #include "MovieSceneSequenceID.h"
 #include "ITimeSlider.h"
+#include "ISequencerModule.h"
+#include "ToolMenu.h"
 #include "Framework/Commands/UICommandList.h"
 #include "Widgets/Input/NumericTypeInterface.h"
 #include "Widgets/Input/SSpinBox.h"
diff --git a/Engine/Source/Editor/Sequencer/Public/ISequencerModule.h b/Engine/Source/Editor/Sequencer/Public/ISequencerModule.h
--- a/Engine/Source/Editor/Sequencer/Public/ISequencerModule.h
+++ b/Engine/Source/Editor/Sequencer/Public/ISequencerModule.h
@@ -4,6 +4,7 @@
 
 #include "CoreMinimal.h"
 #include "Misc/Attribute.h"
+#include "ToolMenuContext.h"
 #include "ISequencer.h"
 #include "Modules/ModuleInterface.h"
 #include "AnimatedPropertyKey.h"
