Sequencer: Added IgnoreFitlersWhenCollapsed setting so if a section is collapsed and there is a filter, it will only show or key filtered key areas when that flag is false.
Defaults to true to keep existing behavior.

diff --git a/Engine/Source/Editor/Sequencer/Private/DisplayNodes/SequencerDisplayNode.cpp b/Engine/Source/Editor/Sequencer/Private/DisplayNodes/SequencerDisplayNode.cpp
--- a/Engine/Source/Editor/Sequencer/Private/DisplayNodes/SequencerDisplayNode.cpp
+++ b/Engine/Source/Editor/Sequencer/Private/DisplayNodes/SequencerDisplayNode.cpp
@@ -627,11 +627,11 @@ FGuid FSequencerDisplayNode::GetObjectGuid() const
 	return ObjectBindingNode ? ObjectBindingNode->GetObjectBinding() : FGuid();
 }
 
-bool FSequencerDisplayNode::Traverse_ChildFirst(const TFunctionRef<bool(FSequencerDisplayNode&)>& InPredicate, bool bIncludeThisNode)
+bool FSequencerDisplayNode::Traverse_ChildFirst(const TFunctionRef<bool(FSequencerDisplayNode&)>& InPredicate, bool bIncludeThisNode, bool bJustVisible)
 {
 	for (auto& Child : GetChildNodes())
 	{
-		if (!Child->Traverse_ChildFirst(InPredicate, true))
+		if ((bJustVisible == true  && Child->IsVisible() == false) || !Child->Traverse_ChildFirst(InPredicate, true))
 		{
 			return false;
 		}
@@ -641,7 +641,7 @@ bool FSequencerDisplayNode::Traverse_ChildFirst(const TFunctionRef<bool(FSequenc
 }
 
 
-bool FSequencerDisplayNode::Traverse_ParentFirst(const TFunctionRef<bool(FSequencerDisplayNode&)>& InPredicate, bool bIncludeThisNode)
+bool FSequencerDisplayNode::Traverse_ParentFirst(const TFunctionRef<bool(FSequencerDisplayNode&)>& InPredicate, bool bIncludeThisNode, bool bJustVisible)
 {
 	if (bIncludeThisNode && !InPredicate(*this))
 	{
@@ -650,7 +650,7 @@ bool FSequencerDisplayNode::Traverse_ParentFirst(const TFunctionRef<bool(FSequen
 
 	for (auto& Child : GetChildNodes())
 	{
-		if (!Child->Traverse_ParentFirst(InPredicate, true))
+		if ((bJustVisible == true &&  Child->IsVisible() == false) || !Child->Traverse_ParentFirst(InPredicate, true))
 		{
 			return false;
 		}
@@ -735,7 +735,8 @@ bool FSequencerDisplayNode::IsDimmed() const
 		{
 			TArray<TSharedRef<FSequencerSectionKeyAreaNode>> KeyAreaNodes;
 			const FSequencerTrackNode& TrackNode = static_cast<FSequencerTrackNode&>(InNode);
-			TrackNode.GetChildKeyAreaNodesRecursively(KeyAreaNodes);
+			const bool bJustVisible = true;
+			TrackNode.GetChildKeyAreaNodesRecursively(KeyAreaNodes, true);
 			if (KeyAreaNodes.Num() > 0)
 			{
 				for (TSharedRef<FSequencerSectionKeyAreaNode> KeyAreaNode : KeyAreaNodes)
@@ -1283,16 +1284,16 @@ void FSequencerDisplayNode::BuildOrganizeContextMenu(FMenuBuilder& MenuBuilder)
 }
 
 
-void FSequencerDisplayNode::GetChildKeyAreaNodesRecursively(TArray< TSharedRef<FSequencerSectionKeyAreaNode> >& OutNodes) const
+void FSequencerDisplayNode::GetChildKeyAreaNodesRecursively(TArray< TSharedRef<FSequencerSectionKeyAreaNode> >& OutNodes, bool bJustVisible) const
 {
 	for (const TSharedRef<FSequencerDisplayNode>& Node : ChildNodes)
 	{
-		if (Node->GetType() == ESequencerNode::KeyArea)
+		if (Node->GetType() == ESequencerNode::KeyArea && (bJustVisible == false || Node->IsVisible()))
 		{
 			OutNodes.Add(StaticCastSharedRef<FSequencerSectionKeyAreaNode>(Node));
 		}
 
-		Node->GetChildKeyAreaNodesRecursively(OutNodes);
+		Node->GetChildKeyAreaNodesRecursively(OutNodes, bJustVisible);
 	}
 }
 
diff --git a/Engine/Source/Editor/Sequencer/Private/DisplayNodes/SequencerDisplayNode.h b/Engine/Source/Editor/Sequencer/Private/DisplayNodes/SequencerDisplayNode.h
--- a/Engine/Source/Editor/Sequencer/Private/DisplayNodes/SequencerDisplayNode.h
+++ b/Engine/Source/Editor/Sequencer/Private/DisplayNodes/SequencerDisplayNode.h
@@ -380,18 +380,20 @@ public:
 	 *
 	 * @param 	InPredicate			Predicate to call for each node, returning whether to continue iteration or not
 	 * @param 	bIncludeThisNode	Whether to include this node in the iteration, or just children
+	 * @param   bJustVisible		Whether or not we check the visible state of the node when adding it
 	 * @return  true where the client prematurely exited the iteration, false otherwise
 	 */
-	bool Traverse_ChildFirst(const TFunctionRef<bool(FSequencerDisplayNode&)>& InPredicate, bool bIncludeThisNode = true);
+	bool Traverse_ChildFirst(const TFunctionRef<bool(FSequencerDisplayNode&)>& InPredicate, bool bIncludeThisNode = true, bool bJustVisible = false);
 
 	/**
 	 * Iterate this entire node tree, parent first.
 	 *
 	 * @param 	InPredicate			Predicate to call for each node, returning whether to continue iteration or not
 	 * @param 	bIncludeThisNode	Whether to include this node in the iteration, or just children
+	 * @param   bJustVisible		Whether or not we check the visible state of the node when adding it
 	 * @return  true where the client prematurely exited the iteration, false otherwise
 	 */
-	bool Traverse_ParentFirst(const TFunctionRef<bool(FSequencerDisplayNode&)>& InPredicate, bool bIncludeThisNode = true);
+	bool Traverse_ParentFirst(const TFunctionRef<bool(FSequencerDisplayNode&)>& InPredicate, bool bIncludeThisNode = true, bool bJustVisible = false);
 
 	/**
 	 * Iterate any visible portions of this node's sub-tree, child first.
@@ -457,7 +459,7 @@ public:
 	}
 
 	/** Gets all the key area nodes recursively, including this node if applicable */
-	virtual void GetChildKeyAreaNodesRecursively(TArray<TSharedRef<class FSequencerSectionKeyAreaNode>>& OutNodes) const;
+	virtual void GetChildKeyAreaNodesRecursively(TArray<TSharedRef<class FSequencerSectionKeyAreaNode>>& OutNodes, bool bJustVisible = false) const;
 
 	/**
 	 * @return The base node this node belongs to, for collections of tracks that are part of an object
diff --git a/Engine/Source/Editor/Sequencer/Private/DisplayNodes/SequencerTrackNode.cpp b/Engine/Source/Editor/Sequencer/Private/DisplayNodes/SequencerTrackNode.cpp
--- a/Engine/Source/Editor/Sequencer/Private/DisplayNodes/SequencerTrackNode.cpp
+++ b/Engine/Source/Editor/Sequencer/Private/DisplayNodes/SequencerTrackNode.cpp
@@ -711,9 +711,9 @@ void FSequencerTrackNode::Resize(float NewSize)
 	}
 }
 
-void FSequencerTrackNode::GetChildKeyAreaNodesRecursively(TArray<TSharedRef<FSequencerSectionKeyAreaNode>>& OutNodes) const
+void FSequencerTrackNode::GetChildKeyAreaNodesRecursively(TArray<TSharedRef<FSequencerSectionKeyAreaNode>>& OutNodes, bool bJustVisible) const
 {
-	FSequencerDisplayNode::GetChildKeyAreaNodesRecursively(OutNodes);
+	FSequencerDisplayNode::GetChildKeyAreaNodesRecursively(OutNodes, bJustVisible);
 
 	if (TopLevelKeyNode.IsValid())
 	{
diff --git a/Engine/Source/Editor/Sequencer/Private/DisplayNodes/SequencerTrackNode.h b/Engine/Source/Editor/Sequencer/Private/DisplayNodes/SequencerTrackNode.h
--- a/Engine/Source/Editor/Sequencer/Private/DisplayNodes/SequencerTrackNode.h
+++ b/Engine/Source/Editor/Sequencer/Private/DisplayNodes/SequencerTrackNode.h
@@ -129,7 +129,7 @@ public:
 	virtual bool CanRenameNode() const override;
 	virtual bool ValidateDisplayName(const FText& NewDisplayName, FText& OutErrorMessage) const override;
 	virtual TSharedRef<SWidget> GetCustomOutlinerContent() override;
-	virtual void GetChildKeyAreaNodesRecursively(TArray<TSharedRef<FSequencerSectionKeyAreaNode>>& OutNodes) const override;
+	virtual void GetChildKeyAreaNodesRecursively(TArray<TSharedRef<FSequencerSectionKeyAreaNode>>& OutNode, bool bJustVisible = false) const override;
 	virtual FText GetDisplayName() const override;
 	virtual FSlateColor GetDisplayNameColor() const override;
 	virtual FSlateFontInfo GetDisplayNameFont() const override;
diff --git a/Engine/Source/Editor/Sequencer/Private/SKeyNavigationButtons.h b/Engine/Source/Editor/Sequencer/Private/SKeyNavigationButtons.h
--- a/Engine/Source/Editor/Sequencer/Private/SKeyNavigationButtons.h
+++ b/Engine/Source/Editor/Sequencer/Private/SKeyNavigationButtons.h
@@ -24,6 +24,7 @@
 #include "Modules/ModuleManager.h"
 #include "IKeyArea.h"
 #include "SequencerAddKeyOperation.h"
+#include "SequencerSettings.h"
 
 #define LOCTEXT_NAMESPACE "SKeyNavigationButtons"
 
@@ -247,7 +248,9 @@ public:
 		FFrameTime  CurrentTime = Sequencer.GetLocalTime().Time;
 
 		FScopedTransaction Transaction(LOCTEXT("AddKeys", "Add Keys at Current Time"));
-		FAddKeyOperation::FromNode(DisplayNode.ToSharedRef()).Commit(CurrentTime.FrameNumber, Sequencer);
+		const bool bJustVisible = true;
+
+		FAddKeyOperation::FromNode(DisplayNode.ToSharedRef(), bJustVisible ).Commit(CurrentTime.FrameNumber, Sequencer);
 
 		return FReply::Handled();
 	}
diff --git a/Engine/Source/Editor/Sequencer/Private/SectionLayout.cpp b/Engine/Source/Editor/Sequencer/Private/SectionLayout.cpp
--- a/Engine/Source/Editor/Sequencer/Private/SectionLayout.cpp
+++ b/Engine/Source/Editor/Sequencer/Private/SectionLayout.cpp
@@ -2,6 +2,8 @@
 
 #include "SectionLayout.h"
 #include "Layout/Geometry.h"
+#include "Sequencer.h"
+#include "SequencerSettings.h"
 
 bool FSectionLayoutElementKeyFuncs::Matches(const FSectionLayoutElement& A, const FSectionLayoutElement& B)
 {
@@ -37,7 +39,8 @@ uint32 FSectionLayoutElementKeyFuncs::GetKeyHash(const FSectionLayoutElement& Ke
 FSectionLayoutElement FSectionLayoutElement::FromGroup(const TSharedRef<FSequencerDisplayNode>& InNode, UMovieSceneSection* InSection, float InOffset)
 {
 	TArray< TSharedRef<FSequencerSectionKeyAreaNode> > ChildKeyAreaNodes;
-	InNode->GetChildKeyAreaNodesRecursively(ChildKeyAreaNodes);
+	const bool bJustVisible = true;
+	InNode->GetChildKeyAreaNodesRecursively(ChildKeyAreaNodes, bJustVisible);
 
 	FSectionLayoutElement Tmp;
 	Tmp.Type = Group;
diff --git a/Engine/Source/Editor/Sequencer/Private/SequencerAddKeyOperation.cpp b/Engine/Source/Editor/Sequencer/Private/SequencerAddKeyOperation.cpp
--- a/Engine/Source/Editor/Sequencer/Private/SequencerAddKeyOperation.cpp
+++ b/Engine/Source/Editor/Sequencer/Private/SequencerAddKeyOperation.cpp
@@ -17,7 +17,7 @@ namespace Sequencer
 {
 
 
-FAddKeyOperation FAddKeyOperation::FromNodes(const TSet<TSharedRef<FSequencerDisplayNode>>& InNodes)
+FAddKeyOperation FAddKeyOperation::FromNodes(const TSet<TSharedRef<FSequencerDisplayNode>>& InNodes, bool bJustVisible)
 {
 	FAddKeyOperation Operation;
 
@@ -42,14 +42,14 @@ FAddKeyOperation FAddKeyOperation::FromNodes(const TSet<TSharedRef<FSequencerDis
 		continue;
 	}
 
-	Operation.AddPreFilteredNodes(FilteredNodes);
+	Operation.AddPreFilteredNodes(FilteredNodes, bJustVisible);
 	return Operation;
 }
 
-FAddKeyOperation FAddKeyOperation::FromNode(TSharedRef<FSequencerDisplayNode> InNode)
+FAddKeyOperation FAddKeyOperation::FromNode(TSharedRef<FSequencerDisplayNode> InNode, bool bJustVisible)
 {
 	FAddKeyOperation Operation;
-	Operation.AddPreFilteredNodes(MakeArrayView(&InNode, 1));
+	Operation.AddPreFilteredNodes(MakeArrayView(&InNode, 1), bJustVisible);
 	return Operation;
 }
 
@@ -66,9 +66,9 @@ FAddKeyOperation FAddKeyOperation::FromKeyAreas(ISequencerTrackEditor* TrackEdit
 	return Operation;
 }
 
-void FAddKeyOperation::AddPreFilteredNodes(TArrayView<const TSharedRef<FSequencerDisplayNode>> FilteredNodes)
+void FAddKeyOperation::AddPreFilteredNodes(TArrayView<const TSharedRef<FSequencerDisplayNode>> FilteredNodes, bool bJustVisible)
 {
-	auto KeyChildTrackArea = [this](FSequencerDisplayNode& InNode)
+	auto KeyChildTrackArea = [this, bJustVisible](FSequencerDisplayNode& InNode)
 	{
 		if (InNode.GetType() == ESequencerNode::Track)
 		{
@@ -76,7 +76,7 @@ void FAddKeyOperation::AddPreFilteredNodes(TArrayView<const TSharedRef<FSequence
 			if (TrackNode->GetSubTrackMode() != FSequencerTrackNode::ESubTrackMode::ParentTrack)
 			{
 				// Consider everything underneath this track for keying
-				this->ConsiderKeyableAreas(TrackNode, TrackNode);
+				this->ConsiderKeyableAreas(TrackNode, TrackNode, bJustVisible);
 			}
 		}
 		return true;
@@ -86,20 +86,20 @@ void FAddKeyOperation::AddPreFilteredNodes(TArrayView<const TSharedRef<FSequence
 	{
 		if (TSharedPtr<FSequencerTrackNode> ParentTrack = Node->FindParentTrackNode())
 		{
-			ConsiderKeyableAreas(ParentTrack.Get(), &Node.Get());
+			ConsiderKeyableAreas(ParentTrack.Get(), &Node.Get(), bJustVisible);
 		}
 		else
 		{
-			Node->Traverse_ParentFirst(KeyChildTrackArea);
+			Node->Traverse_ParentFirst(KeyChildTrackArea, true/*bIncludeThisNode*/, bJustVisible);
 		}
 	}
 }
 
-bool FAddKeyOperation::ConsiderKeyableAreas(FSequencerTrackNode* InTrackNode, FSequencerDisplayNode* KeyAnythingBeneath)
+bool FAddKeyOperation::ConsiderKeyableAreas(FSequencerTrackNode* InTrackNode, FSequencerDisplayNode* KeyAnythingBeneath, bool bJustVisible)
 {
 	bool bKeyedAnything = false;
 
-	auto Traversal = [this, InTrackNode, &bKeyedAnything](FSequencerDisplayNode& InNode)
+	auto Traversal = [this, InTrackNode, &bKeyedAnything, bJustVisible](FSequencerDisplayNode& InNode)
 	{
 		if (InNode.GetType() == ESequencerNode::Track)
 		{
@@ -107,28 +107,30 @@ bool FAddKeyOperation::ConsiderKeyableAreas(FSequencerTrackNode* InTrackNode, FS
 
 			if (TSharedPtr<FSequencerSectionKeyAreaNode> KeyAreaNode = ThisTrackNode->GetTopLevelKeyNode())
 			{
-				bKeyedAnything |= this->ProcessKeyAreaNode(ThisTrackNode, KeyAreaNode.Get());
+				bKeyedAnything |= this->ProcessKeyAreaNode(ThisTrackNode, KeyAreaNode.Get(), bJustVisible);
 			}
 		}
 
 		if (InNode.GetType() == ESequencerNode::KeyArea)
 		{
-			bKeyedAnything |= this->ProcessKeyAreaNode(InTrackNode, static_cast<FSequencerSectionKeyAreaNode*>(&InNode));
+			bKeyedAnything |= this->ProcessKeyAreaNode(InTrackNode, static_cast<FSequencerSectionKeyAreaNode*>(&InNode), bJustVisible);
 		}
 		return true;
 	};
-	KeyAnythingBeneath->Traverse_ParentFirst(Traversal, true);
+	KeyAnythingBeneath->Traverse_ParentFirst(Traversal, true/*bIncludeThisNode*/, bJustVisible);
 
 	return bKeyedAnything;
 }
 
-bool FAddKeyOperation::ProcessKeyAreaNode(FSequencerTrackNode* InTrackNode, const FSequencerSectionKeyAreaNode* KeyAreaNode)
+bool FAddKeyOperation::ProcessKeyAreaNode(FSequencerTrackNode* InTrackNode, const FSequencerSectionKeyAreaNode* KeyAreaNode, bool bJustVisible)
 {
 	bool bKeyedAnything = false;
-
-	for (const TSharedRef<IKeyArea>& KeyArea : KeyAreaNode->GetAllKeyAreas())
+	if (bJustVisible == false || InTrackNode->IsVisible())
 	{
-		bKeyedAnything |= ProcessKeyArea(InTrackNode, KeyArea);
+		for (const TSharedRef<IKeyArea>& KeyArea : KeyAreaNode->GetAllKeyAreas())
+		{
+			bKeyedAnything |= ProcessKeyArea(InTrackNode, KeyArea);
+		}
 	}
 
 	return bKeyedAnything;
diff --git a/Engine/Source/Editor/Sequencer/Private/SequencerKeyCollection.cpp b/Engine/Source/Editor/Sequencer/Private/SequencerKeyCollection.cpp
--- a/Engine/Source/Editor/Sequencer/Private/SequencerKeyCollection.cpp
+++ b/Engine/Source/Editor/Sequencer/Private/SequencerKeyCollection.cpp
@@ -6,6 +6,8 @@
 #include "DisplayNodes/SequencerDisplayNode.h"
 #include "DisplayNodes/SequencerTrackNode.h"
 #include "DisplayNodes/SequencerSectionKeyAreaNode.h"
+#include "Sequencer.h"
+#include "SequencerSettings.h"
 
 FSequencerKeyCollectionSignature FSequencerKeyCollectionSignature::FromNodes(const TArray<FSequencerDisplayNode*>& InNodes, FFrameNumber InDuplicateThresholdTime)
 {
@@ -46,6 +48,8 @@ FSequencerKeyCollectionSignature FSequencerKeyCollectionSignature::FromNodesRecu
 
 	TArray<TSharedRef<FSequencerSectionKeyAreaNode>> AllKeyAreaNodes;
 	AllKeyAreaNodes.Reserve(36);
+	const bool bJustVisible = true;
+
 	for (FSequencerDisplayNode* Node : InNodes)
 	{
 		if (Node->GetType() == ESequencerNode::KeyArea)
@@ -53,7 +57,7 @@ FSequencerKeyCollectionSignature FSequencerKeyCollectionSignature::FromNodesRecu
 			AllKeyAreaNodes.Add(StaticCastSharedRef<FSequencerSectionKeyAreaNode>(Node->AsShared()));
 		}
 
-		Node->GetChildKeyAreaNodesRecursively(AllKeyAreaNodes);
+		Node->GetChildKeyAreaNodesRecursively(AllKeyAreaNodes, bJustVisible);
 	}
 
 	for (const TSharedRef<FSequencerSectionKeyAreaNode>& Node : AllKeyAreaNodes)
@@ -75,7 +79,9 @@ FSequencerKeyCollectionSignature FSequencerKeyCollectionSignature::FromNodeRecur
 
 	TArray<TSharedRef<FSequencerSectionKeyAreaNode>> AllKeyAreaNodes;
 	AllKeyAreaNodes.Reserve(36);
-	InNode.GetChildKeyAreaNodesRecursively(AllKeyAreaNodes);
+
+	const bool bJustVisible = true;
+	InNode.GetChildKeyAreaNodesRecursively(AllKeyAreaNodes, bJustVisible);
 
 	for (const auto& Node : AllKeyAreaNodes)
 	{
diff --git a/Engine/Source/Editor/Sequencer/Public/SequencerAddKeyOperation.h b/Engine/Source/Editor/Sequencer/Public/SequencerAddKeyOperation.h
--- a/Engine/Source/Editor/Sequencer/Public/SequencerAddKeyOperation.h
+++ b/Engine/Source/Editor/Sequencer/Public/SequencerAddKeyOperation.h
@@ -32,14 +32,18 @@ struct SEQUENCER_API FAddKeyOperation
 	 * Construct an operation from any set of display nodes. Each node in the set will receive keys for all decendant key areas.
 	 *
 	 * @param InNodes A set of all the nodes to key
+	 * @param bJustVisible Only do it on visible key area nodes, so we can respect filtering when keying
 	 */
-	static FAddKeyOperation FromNodes(const TSet<TSharedRef<FSequencerDisplayNode>>& InNodes);
+	static FAddKeyOperation FromNodes(const TSet<TSharedRef<FSequencerDisplayNode>>& InNodes, bool bJustVisible = false);
 
 
 	/**
 	 * Construct an operation from a single display node. Every key area underneath this node will receive keys.
+	 * 
+	 * @param InNode Node to key
+	 * @param bJustVisible Only do it on visible key area nodes, so we can respect filtering when keying
 	 */
-	static FAddKeyOperation FromNode(TSharedRef<FSequencerDisplayNode> InNode);
+	static FAddKeyOperation FromNode(TSharedRef<FSequencerDisplayNode> InNode, bool bJustVisible = false);
 
 
 	/**
@@ -62,8 +66,9 @@ private:
 	 * Add a set of nodes to this operation that have already had child nodes removed (ie only parent nodes should exist in the set)
 	 *
 	 * @param InNodes     A set of nodes to add to this operation that contains no child nodes
+	 * @param bJustVisible Only do it on visible key area nodes, so we can respect filtering when keying
 	 */
-	void AddPreFilteredNodes(TArrayView<const TSharedRef<FSequencerDisplayNode>> InNodes);
+	void AddPreFilteredNodes(TArrayView<const TSharedRef<FSequencerDisplayNode>> InNodes, bool bJustVisible);
 
 
 	/**
@@ -71,8 +76,9 @@ private:
 	 *
 	 * @param InTrackNode           The current track node
 	 * @param InKeyAnythingBeneath  A node to search within for key areas
+	 * @param bJustVisible Only do it on visible key area nodes, so we can respect filtering when keying
 	 */
-	bool ConsiderKeyableAreas(FSequencerTrackNode* InTrackNode, FSequencerDisplayNode* InKeyAnythingBeneath);
+	bool ConsiderKeyableAreas(FSequencerTrackNode* InTrackNode, FSequencerDisplayNode* InKeyAnythingBeneath, bool bJustVisible);
 
 
 	/**
@@ -80,8 +86,9 @@ private:
 	 *
 	 * @param InTrackNode         The current track node
 	 * @param InKeyAreaNode       The key area node to add key areas from
+	 * @param bJustVisible Only do it on visible key area nodes, so we can respect filtering when keying
 	 */
-	bool ProcessKeyAreaNode(FSequencerTrackNode* InTrackNode, const FSequencerSectionKeyAreaNode* InKeyAreaNode);
+	bool ProcessKeyAreaNode(FSequencerTrackNode* InTrackNode, const FSequencerSectionKeyAreaNode* InKeyAreaNode, bool bJustVisible);
 
 
 	/**
