Moving Sequencer command registration to post engine init

diff --git a/Engine/Source/Editor/Sequencer/Private/SequencerModule.cpp b/Engine/Source/Editor/Sequencer/Private/SequencerModule.cpp
--- a/Engine/Source/Editor/Sequencer/Private/SequencerModule.cpp
+++ b/Engine/Source/Editor/Sequencer/Private/SequencerModule.cpp
@@ -279,8 +279,6 @@ public:
 	{
 		if (GIsEditor)
 		{
-			FSequencerCommands::Register();
-
 			FEditorModeRegistry::Get().RegisterMode<FSequencerEdMode>(
 				FSequencerEdMode::EM_SequencerMode,
 				NSLOCTEXT("Sequencer", "SequencerEditMode", "Sequencer Mode"),
@@ -289,10 +287,12 @@ public:
 
 			if (UToolMenus::TryGet())
 			{
+				FSequencerCommands::Register();
 				RegisterMenus();
 			}
 			else
 			{
+				FCoreDelegates::OnPostEngineInit.AddStatic(&FSequencerCommands::Register);
 				FCoreDelegates::OnPostEngineInit.AddRaw(this, &FSequencerModule::RegisterMenus);
 			}
 
