Sequencer: Fixed Tick Resolution migration code not successfully migrating locked sections, and not migrating Marked Frames. This also adds the ability to recursively apply the time migration to all child sequences, which was the user-expected behavior and enables mass time migration for a sequence.

diff --git a/Engine/Source/Editor/Sequencer/Private/SSequencerTimePanel.cpp b/Engine/Source/Editor/Sequencer/Private/SSequencerTimePanel.cpp
--- a/Engine/Source/Editor/Sequencer/Private/SSequencerTimePanel.cpp
+++ b/Engine/Source/Editor/Sequencer/Private/SSequencerTimePanel.cpp
@@ -17,11 +17,12 @@
 #include "Sequencer.h"
 #include "Styling/AppStyle.h"
 #include "MovieSceneToolHelpers.h"
-
+#include "MovieSceneSequenceVisitor.h"
+#include "Framework/Notifications/NotificationManager.h"
+#include "Widgets/Notifications/SNotificationList.h"
 
 #define LOCTEXT_NAMESPACE "SSequencerTimePanel"
 
-
 void SSequencerTimePanel::Construct(const FArguments& InArgs, TWeakPtr<FSequencer> InSequencer)
 {
 	WeakSequencer = InSequencer;
@@ -44,9 +45,9 @@ void SSequencerTimePanel::Construct(const FArguments& InArgs, TWeakPtr<FSequence
 
 	static FMargin Col1Padding(0.f, 0.f, HorizontalGridPadding, VerticalGridPadding);
 	static FMargin Col2Padding(HorizontalGridPadding, 0.f, 0.f, VerticalGridPadding);
+	bIsRecursive = true;
 
 	static FLinearColor WarningColor(FColor(0xffbbbb44));
-	int32 CurrentRow = 0;
 	ChildSlot
 	[
 		SNew(SBorder)
@@ -96,6 +97,23 @@ void SSequencerTimePanel::Construct(const FArguments& InArgs, TWeakPtr<FSequence
 					]
 
 					+ SGridPanel::Slot(0, 1)
+					.HAlign(HAlign_Right)
+					.Padding(FMargin(0, 0, HorizontalGridPadding, 0))
+					[
+						SNew(STextBlock)
+						.Text(LOCTEXT("ApplyRecursively", "Apply Recursively"))
+						.ToolTipText(LOCTEXT("ApplyRecursively_Tooltip", "If true then the tick rate change will propagate to all child sub-sequences of the current sequence. The UI warnings only apply to the top level sequence."))
+					]
+					+ SGridPanel::Slot(1, 1)
+					.HAlign(HAlign_Left)
+					.Padding(FMargin(HorizontalGridPadding, 0, 0, 0))
+					[
+						SNew(SCheckBox)
+						.IsChecked(this, &SSequencerTimePanel::GetIsRecursive)
+						.OnCheckStateChanged(this, &SSequencerTimePanel::OnSetIsRecursive)
+					]
+
+					+ SGridPanel::Slot(0, 2)
 					.ColumnSpan(2)
 					.Padding(FMargin(0.f, VerticalGridPadding, 0.f, VerticalGridPadding))
 					.HAlign(HAlign_Left)
@@ -105,28 +123,28 @@ void SSequencerTimePanel::Construct(const FArguments& InArgs, TWeakPtr<FSequence
 						.Text(LOCTEXT("NewTickInterval_Tip", "Sequence will have the following properties if applied:"))
 					]
 
-					+ SGridPanel::Slot(0, 2)
+					+ SGridPanel::Slot(0, 3)
 					.Padding(Col1Padding)
 					.HAlign(HAlign_Right)
 					[
 						SNew(STextBlock)
 						.Text(LOCTEXT("ResultingRange", "Time Range"))
 					]
-					+ SGridPanel::Slot(1, 2)
+					+ SGridPanel::Slot(1, 3)
 					.Padding(Col2Padding)
 					[
 						SNew(STextBlock)
 						.Text(this, &SSequencerTimePanel::GetSupportedTimeRange)
 					]
 
-					+ SGridPanel::Slot(0, 3)
+					+ SGridPanel::Slot(0, 4)
 					.Padding(Col1Padding)
 					.HAlign(HAlign_Right)
 					[
 						SNew(STextBlock)
 						.Text(LOCTEXT("SupportedFrameRates", "Supported Rates"))
 					]
-					+ SGridPanel::Slot(1, 3)
+					+ SGridPanel::Slot(1, 4)
 					.Padding(Col2Padding)
 					[
 						SAssignNew(CommonFrameRates, SVerticalBox)
@@ -212,10 +230,51 @@ FReply SSequencerTimePanel::Apply()
 	{
 		FFrameRate Src = MovieScene->GetTickResolution();
 		FFrameRate Dst = GetCurrentTickResolution();
+		bool bRecursive = GetIsRecursive() == ECheckBoxState::Checked;
 
 		FScopedTransaction ScopedTransaction(FText::Format(LOCTEXT("MigrateFrameTimes", "Convert sequence tick interval from {0} to {1}"), Src.ToPrettyText(), Dst.ToPrettyText()));
 
-		UE::MovieScene::TimeHelpers::MigrateFrameTimes(Src, Dst, MovieScene);
+		// We quickly iterate through the sequence hierarchy to check for any readonly sequences. If we find a read-only sequence we're going to just warn the user
+		// that we automatically unlocked, edited, and re-locked the sequence (as there should be no data change apparent to the user). We're doing this here to avoid
+		// firing off the warning for every sequence, though we will print in the output log which ones were unlocked.
+		{
+			struct FSequenceReadOnlyVisitor : UE::MovieScene::ISequenceVisitor
+			{
+				virtual void VisitSubSequence(UMovieSceneSequence* InSequence, const FGuid&, const UE::MovieScene::FSubSequenceSpace& LocalSpace) 
+				{
+					if (InSequence->GetMovieScene()->IsReadOnly())
+					{
+						ReadOnlyMovieScenes.Add(InSequence->GetMovieScene());
+					}
+				}
+				TArray<UMovieScene*> ReadOnlyMovieScenes;
+			};
+
+			UE::MovieScene::FSequenceVisitParams Params;
+			Params.bVisitMasterTracks = true;
+			Params.bVisitSubSequences = true;
+			FSequenceReadOnlyVisitor ReadOnlyVisitor;
+			VisitSequence(Sequence, Params, ReadOnlyVisitor);
+
+			for (UMovieScene* ReadOnlyMovieScene : ReadOnlyVisitor.ReadOnlyMovieScenes)
+			{
+				UE_LOG(LogMovieScene, Log, TEXT("The following sequence was read-only but had its tick resolution adjusted automatically. No action is required. %s"), *ReadOnlyMovieScene->GetPathName());
+			}
+
+			if (ReadOnlyVisitor.ReadOnlyMovieScenes.Num() > 0)
+			{
+				FNotificationInfo Info(LOCTEXT("ReadOnlyScenesEdited", "Edited read-only sequence, see output log for details. No action is required."));
+				Info.bFireAndForget = true;
+				Info.FadeOutDuration = 0.5f;
+				Info.ExpireDuration = 5.0f;
+
+				TSharedPtr<SNotificationItem> NotificationPtr = FSlateNotificationManager::Get().AddNotification(Info);
+				NotificationPtr->SetCompletionState(SNotificationItem::CS_Success);
+			}
+
+		}
+
+		UE::MovieScene::TimeHelpers::MigrateFrameTimes(Src, Dst, MovieScene, bRecursive);
 	}
 
 	return Close();
@@ -299,7 +358,15 @@ void SSequencerTimePanel::OnSetTickResolution(FFrameRate InTickResolution)
 }
 
 
+ECheckBoxState SSequencerTimePanel::GetIsRecursive() const
+{
+	return bIsRecursive ? ECheckBoxState::Checked : ECheckBoxState::Unchecked;
+}
 
+void SSequencerTimePanel::OnSetIsRecursive(ECheckBoxState InCheckBoxState)
+{
+	bIsRecursive = InCheckBoxState == ECheckBoxState::Checked;
+}
 
 
 #undef LOCTEXT_NAMESPACE
\ No newline at end of file
diff --git a/Engine/Source/Editor/Sequencer/Private/SSequencerTimePanel.h b/Engine/Source/Editor/Sequencer/Private/SSequencerTimePanel.h
--- a/Engine/Source/Editor/Sequencer/Private/SSequencerTimePanel.h
+++ b/Engine/Source/Editor/Sequencer/Private/SSequencerTimePanel.h
@@ -41,6 +41,10 @@ private:
 	void OnSetTickResolution(FFrameRate InTickResolution);
 
 	void UpdateCommonFrameRates();
+	
+	ECheckBoxState GetIsRecursive() const;
+	void OnSetIsRecursive(ECheckBoxState InCheckBoxState);
+
 
 	UMovieSceneSequence* GetFocusedSequence() const;
 
@@ -51,4 +55,5 @@ private:
 	TOptional<FFrameRate> CurrentTickResolution;
 	TSharedPtr<SVerticalBox> CommonFrameRates;
 	TWeakPtr<FSequencer> WeakSequencer;
+	bool bIsRecursive;
 };
\ No newline at end of file
diff --git a/Engine/Source/Runtime/MovieScene/Private/MovieSceneTimeHelpers.cpp b/Engine/Source/Runtime/MovieScene/Private/MovieSceneTimeHelpers.cpp
--- a/Engine/Source/Runtime/MovieScene/Private/MovieSceneTimeHelpers.cpp
+++ b/Engine/Source/Runtime/MovieScene/Private/MovieSceneTimeHelpers.cpp
@@ -7,9 +7,12 @@
 #include "Misc/FrameTime.h"
 #include "Misc/ScopedSlowTask.h"
 #include "Sections/MovieSceneSubSection.h"
+#include "Sections/MovieSceneSubSection.h"
+#include "Tracks/MovieSceneSubTrack.h"
 #include "MovieScene.h"
 #include "MovieSceneSection.h"
 #include "MovieSceneTrack.h"
+#include "MovieSceneSequence.h"
 #include "Channels/MovieSceneChannelProxy.h"
 #include "Channels/MovieSceneChannel.h"
 
@@ -55,6 +58,8 @@ TRange<FFrameNumber> MigrateFrameRange(const TRange<FFrameNumber>& SourceRange,
 void MigrateFrameTimes(FFrameRate SourceRate, FFrameRate DestinationRate, UMovieSceneSection* Section)
 {
 	Section->Modify();
+	const bool bSectionWasLocked = Section->IsLocked();
+	Section->SetIsLocked(false);
 
 	TRangeBound<FFrameNumber> NewLowerBound, NewUpperBound;
 
@@ -98,6 +103,8 @@ void MigrateFrameTimes(FFrameRate SourceRate, FFrameRate DestinationRate, UMovie
 			Channel->ChangeFrameResolution(SourceRate, DestinationRate);
 		}
 	}
+
+	Section->SetIsLocked(bSectionWasLocked);
 }
 
 void MigrateFrameTimes(FFrameRate SourceRate, FFrameRate DestinationRate, UMovieSceneTrack* Track)
@@ -111,7 +118,7 @@ void MigrateFrameTimes(FFrameRate SourceRate, FFrameRate DestinationRate, UMovie
 	}
 }
 
-void TimeHelpers::MigrateFrameTimes(FFrameRate SourceRate, FFrameRate DestinationRate, UMovieScene* MovieScene)
+void TimeHelpers::MigrateFrameTimes(FFrameRate SourceRate, FFrameRate DestinationRate, UMovieScene* MovieScene, bool bApplyRecursively)
 {
 	int32 TotalNumTracks = MovieScene->GetMasterTracks().Num() + (MovieScene->GetCameraCutTrack() ? 1 : 0);
 	for (const FMovieSceneBinding& Binding : MovieScene->GetBindings())
@@ -123,6 +130,10 @@ void TimeHelpers::MigrateFrameTimes(FFrameRate SourceRate, FFrameRate Destinatio
 	SlowTask.MakeDialogDelayed(0.25f, true);
 
 	MovieScene->Modify();
+#if WITH_EDITOR
+	const bool bMovieSceneReadOnly = MovieScene->IsReadOnly();
+	MovieScene->SetReadOnly(false);
+#endif
 
 	MovieScene->SetPlaybackRange(MigrateFrameRange(MovieScene->GetPlaybackRange(), SourceRate, DestinationRate));
 #if WITH_EDITORONLY_DATA
@@ -133,6 +144,27 @@ void TimeHelpers::MigrateFrameTimes(FFrameRate SourceRate, FFrameRate Destinatio
 	{
 		SlowTask.EnterProgressFrame();
 		UE::MovieScene::MigrateFrameTimes(SourceRate, DestinationRate, Track);
+
+		// We iterate through recursively here (and not in MigrateFrameTimes) so that the movie scene is taken
+		// into account for locking/modifying/etc.
+		if (bApplyRecursively && Track->IsA<UMovieSceneSubTrack>())
+		{
+			for (UMovieSceneSection* Section : Track->GetAllSections())
+			{
+				UMovieSceneSubSection* SubSection = Cast<UMovieSceneSubSection>(Section);
+				if (SubSection)
+				{
+					if (UMovieSceneSequence* SubSequence = SubSection->GetSequence())
+					{
+						UMovieScene* ChildMovieScene = SubSequence->GetMovieScene();
+						if (ChildMovieScene)
+						{
+							TimeHelpers::MigrateFrameTimes(SourceRate, DestinationRate, ChildMovieScene, bApplyRecursively);
+						}
+					}
+				}
+			}
+		}
 	}
 
 	if (UMovieSceneTrack* Track = MovieScene->GetCameraCutTrack())
@@ -150,7 +182,28 @@ void TimeHelpers::MigrateFrameTimes(FFrameRate SourceRate, FFrameRate Destinatio
 		}
 	}
 
+	{
+		TArray<FMovieSceneMarkedFrame> MarkedFrames = MovieScene->GetMarkedFrames();
+
+		// Clear the marked frames as the returned array is immutable
+		MovieScene->DeleteMarkedFrames();
+
+		for (FMovieSceneMarkedFrame& MarkedFrame : MarkedFrames)
+		{
+			MarkedFrame.FrameNumber = ConvertFrameTime(MarkedFrame.FrameNumber, SourceRate, DestinationRate).RoundToFrame();
+
+			// Add it back in
+			MovieScene->AddMarkedFrame(MarkedFrame);
+		}
+
+		// Ensure they're in order as they may not have been before.
+		MovieScene->SortMarkedFrames();
+
+	}
 	MovieScene->SetTickResolutionDirectly(DestinationRate);
+#if WITH_EDITOR
+	MovieScene->SetReadOnly(bMovieSceneReadOnly);
+#endif
 }
 
 } // namespace MovieScene
diff --git a/Engine/Source/Runtime/MovieScene/Public/MovieSceneTimeHelpers.h b/Engine/Source/Runtime/MovieScene/Public/MovieSceneTimeHelpers.h
--- a/Engine/Source/Runtime/MovieScene/Public/MovieSceneTimeHelpers.h
+++ b/Engine/Source/Runtime/MovieScene/Public/MovieSceneTimeHelpers.h
@@ -24,7 +24,7 @@ public:
  * Migrate the frame times of the movie scene from the source frame rate to the destination frame rate
  */
 
-static void MigrateFrameTimes(FFrameRate SourceRate, FFrameRate DestinationRate, UMovieScene* MovieScene);
+static void MigrateFrameTimes(FFrameRate SourceRate, FFrameRate DestinationRate, UMovieScene* MovieScene, bool bApplyRecursively = false);
 
 };
 
