Sequencer: Add editor preference to do drag selection on left mouse click. Added new class to the ed mode to handle the needed drag behavior and then get it created when needed.

diff --git a/Engine/Source/Editor/Sequencer/Private/SequencerEdMode.cpp b/Engine/Source/Editor/Sequencer/Private/SequencerEdMode.cpp
--- a/Engine/Source/Editor/Sequencer/Private/SequencerEdMode.cpp
+++ b/Engine/Source/Editor/Sequencer/Private/SequencerEdMode.cpp
@@ -32,7 +32,8 @@
 #include "MovieSceneToolHelpers.h"
 #include "Tools/MotionTrailOptions.h"
 #include "SequencerCommands.h"
-
+#include "SnappingUtils.h"
+#include "SequencerSettings.h"
 const FEditorModeID FSequencerEdMode::EM_SequencerMode(TEXT("EM_SequencerMode"));
 
 static TAutoConsoleVariable<bool> CVarDrawMeshTrails(
@@ -49,6 +50,8 @@ TAutoConsoleVariable<bool> CVarUseOldSequencerMotionTrails(
 
 namespace UE
 {
+
+
 namespace SequencerEdMode
 {
 
@@ -223,6 +226,32 @@ bool FSequencerEdMode::IsPressingMoveTimeSlider(FViewport* InViewport) const
 	return bIsMovingTimeSlider;
 }
 
+//just get the first one.
+USequencerSettings* FSequencerEdMode:: GetSequencerSettings() const
+{
+	TSharedPtr<FSequencer> ActiveSequencer;
+	for (TWeakPtr<FSequencer> WeakSequencer : Sequencers)
+	{
+		ActiveSequencer = WeakSequencer.Pin();
+		if (ActiveSequencer.IsValid())
+		{
+			return ActiveSequencer->GetSequencerSettings();
+		}
+	}
+	return nullptr;
+}
+
+bool FSequencerEdMode::IsDoingDrag(FViewport* InViewport) const
+{
+	const USequencerSettings* SequencerSettings = GetSequencerSettings();
+	const bool LeftMouseButtonDown = InViewport->KeyState(EKeys::LeftMouseButton);
+	const bool bIsCtrlKeyDown = InViewport->KeyState(EKeys::LeftControl) || InViewport->KeyState(EKeys::RightControl);
+	const bool bIsShiftKeyDown = InViewport->KeyState(EKeys::LeftShift) || InViewport->KeyState(EKeys::RightShift);
+	const bool bIsAltKeyDown = InViewport->KeyState(EKeys::LeftAlt) || InViewport->KeyState(EKeys::RightAlt);
+
+	return LeftMouseButtonDown && !bIsCtrlKeyDown && !bIsShiftKeyDown && !bIsAltKeyDown && (SequencerSettings ? SequencerSettings->GetLeftMouseDragDoesMarquee() : false);
+}
+
 bool FSequencerEdMode::StartTracking(FEditorViewportClient* InViewportClient, FViewport* InViewport)
 {
 	if (IsPressingMoveTimeSlider(InViewport))
@@ -243,13 +272,29 @@ bool FSequencerEdMode::StartTracking(FEditorViewportClient* InViewportClient, FV
 		}
 		return true;
 	}
+	else if(IsDoingDrag(InViewport))
+	{
+		DragToolHandler.MakeDragTool(InViewportClient);
+		return DragToolHandler.StartTracking(InViewportClient, InViewport);
+	}
 	return false;
 }
-
+bool FSequencerEdMode::EndTracking(FEditorViewportClient* InViewportClient, FViewport* InViewport)
+{
+	return DragToolHandler.EndTracking(InViewportClient, InViewport);
+}
 
 bool FSequencerEdMode::InputDelta(FEditorViewportClient* InViewportClient, FViewport* InViewport, FVector& InDrag, FRotator& InRot, FVector& InScale)
 {
-	return IsPressingMoveTimeSlider(InViewport);
+	if (IsPressingMoveTimeSlider(InViewport))
+	{
+		return true;
+	}
+	else if(IsDoingDrag(InViewport))
+	{
+		return DragToolHandler.InputDelta(InViewportClient, InViewport, InDrag, InRot, InScale);
+	}
+	return false;
 }
 
 static void SnapSequencerTime(TSharedPtr<FSequencer>& Sequencer, FFrameTime& ScrubTime)
@@ -378,6 +423,8 @@ void FSequencerEdMode::Render(const FSceneView* View, FViewport* Viewport, FPrim
 {
 	FEdMode::Render(View, Viewport, PDI);
 
+	DragToolHandler.Render3DDragTool(View, PDI);
+
 #if WITH_EDITORONLY_DATA
 	if (PDI)
 	{
@@ -402,6 +449,8 @@ void FSequencerEdMode::DrawHUD(FEditorViewportClient* ViewportClient,FViewport*
 {
 	FEdMode::DrawHUD(ViewportClient,Viewport,View,Canvas);
 
+	DragToolHandler.RenderDragTool(View, Canvas);
+
 	if( ViewportClient->AllowsCinematicControl() )
 	{
 		// Get the size of the viewport
@@ -913,3 +962,109 @@ bool FSequencerEdModeTool::InputKey(FEditorViewportClient* ViewportClient, FView
 
 	return FModeTool::InputKey(ViewportClient, Viewport, Key, Event);
 }
+
+/*
+*
+*    Drag Tool Handler
+*
+*/
+
+FMarqueeDragTool::FMarqueeDragTool()
+	: bIsDeletingDragTool(false)
+{
+}
+
+bool FMarqueeDragTool::StartTracking(FEditorViewportClient* InViewportClient, FViewport* InViewport)
+{
+	return (DragTool.IsValid() && InViewportClient->GetCurrentWidgetAxis() == EAxisList::None);
+}
+
+bool FMarqueeDragTool::EndTracking(FEditorViewportClient* InViewportClient, FViewport* InViewport)
+{
+	if (!bIsDeletingDragTool)
+	{
+		// Ending the drag tool may pop up a modal dialog which can cause unwanted reentrancy - protect against this.
+		TGuardValue<bool> RecursionGuard(bIsDeletingDragTool, true);
+
+		// Delete the drag tool if one exists.
+		if (DragTool.IsValid())
+		{
+			if (DragTool->IsDragging())
+			{
+				DragTool->EndDrag();
+			}
+			DragTool.Reset();
+			return false;
+		}
+	}
+	
+	return true;
+}
+
+bool FMarqueeDragTool::InputDelta(FEditorViewportClient* InViewportClient, FViewport* InViewport, FVector& InDrag, FRotator& InRot, FVector& InScale)
+{
+	if (DragTool.IsValid() == false || InViewportClient->GetCurrentWidgetAxis() != EAxisList::None)
+	{
+		return false;
+	}
+	if (DragTool->IsDragging() == false)
+	{
+		int32 InX = InViewport->GetMouseX();
+		int32 InY = InViewport->GetMouseY();
+		FVector2D Start(InX, InY);
+
+		DragTool->StartDrag(InViewportClient, GEditor->ClickLocation,Start);
+	}
+	const bool bUsingDragTool = UsingDragTool();
+	if (bUsingDragTool == false)
+	{
+		return false;
+	}
+
+	DragTool->AddDelta(InDrag);
+	return true;
+}
+
+/**
+ * @return		true if a drag tool is being used by the tracker, false otherwise.
+ */
+bool FMarqueeDragTool::UsingDragTool() const
+{
+	return DragTool.IsValid() && DragTool->IsDragging();
+}
+
+/**
+ * Renders the drag tool.  Does nothing if no drag tool exists.
+ */
+void FMarqueeDragTool::Render3DDragTool(const FSceneView* View, FPrimitiveDrawInterface* PDI)
+{
+	if (DragTool.IsValid())
+	{
+		DragTool->Render3D(View, PDI);
+	}
+}
+
+/**
+ * Renders the drag tool.  Does nothing if no drag tool exists.
+ */
+void FMarqueeDragTool::RenderDragTool(const FSceneView* View, FCanvas* Canvas)
+{
+	if (DragTool.IsValid())
+	{
+		DragTool->Render(View, Canvas);
+	}
+}
+
+void FMarqueeDragTool::MakeDragTool(FEditorViewportClient* InViewportClient)
+{
+	DragTool.Reset();
+	if (InViewportClient->IsOrtho())
+	{
+		DragTool = InViewportClient->MakeDragTool(EDragTool::BoxSelect);
+	}
+	else
+	{
+		DragTool = InViewportClient->MakeDragTool(EDragTool::FrustumSelect);
+	}
+
+}
\ No newline at end of file
diff --git a/Engine/Source/Editor/Sequencer/Private/SequencerEdMode.h b/Engine/Source/Editor/Sequencer/Private/SequencerEdMode.h
--- a/Engine/Source/Editor/Sequencer/Private/SequencerEdMode.h
+++ b/Engine/Source/Editor/Sequencer/Private/SequencerEdMode.h
@@ -21,6 +21,42 @@ class UMovieScene3DTransformSection;
 class UMovieScene3DTransformTrack;
 struct FMovieSceneEvaluationTrack;
 struct FMovieSceneInterrogationData;
+class USequencerSettings;
+
+#include "EditorDragTools.h"
+
+// This struct wraps up functionality for creating a marquee(frustum or box) selection drag tool
+// It does so based upon the type of viewport being drawn.
+// It's also agnostic to any key presses or mouse button type. 
+// If it's created it's assumed it will track.
+struct FMarqueeDragTool
+{
+	FMarqueeDragTool();
+	~FMarqueeDragTool() {};
+
+	bool StartTracking(FEditorViewportClient* InViewportClient, FViewport* InViewport);
+	bool EndTracking(FEditorViewportClient* InViewportClient, FViewport* InViewport);
+	void MakeDragTool(FEditorViewportClient* InViewportClient);
+	bool InputDelta(FEditorViewportClient* InViewportClient, FViewport* InViewport, FVector& InDrag, FRotator& InRot, FVector& InScale);
+
+	bool UsingDragTool() const;
+	void Render3DDragTool(const FSceneView* View, FPrimitiveDrawInterface* PDI);
+	void RenderDragTool(const FSceneView* View, FCanvas* Canvas);
+
+
+private:
+	
+	/**
+	 * If there is a dragging tool being used, this will point to it.
+	 * Gets newed/deleted in StartTracking/EndTracking.
+	 */
+	TSharedPtr<FDragTool> DragTool;
+
+	/** Tracks whether the drag tool is in the process of being deleted (to protect against reentrancy) */
+	bool bIsDeletingDragTool = false;
+
+};
+
 
 /** Stores the transform track and associated mesh trail for each drawn track */
 struct FMeshTrailData
@@ -60,14 +96,16 @@ public:
 
 
 	virtual bool StartTracking(FEditorViewportClient* InViewportClient, FViewport* InViewport) override;
+	virtual bool EndTracking(FEditorViewportClient* InViewportClient, FViewport* InViewport) override;
 	virtual bool MouseMove(FEditorViewportClient* ViewportClient, FViewport* Viewport, int32 x, int32 y) override;
 	virtual bool ProcessCapturedMouseMoves(FEditorViewportClient* InViewportClient, FViewport* InViewport, const TArrayView<FIntPoint>& CapturedMouseMoves) override;
 	virtual bool InputDelta(FEditorViewportClient* InViewportClient, FViewport* InViewport, FVector& InDrag, FRotator& InRot, FVector& InScale) override;
 
 	bool IsPressingMoveTimeSlider(FViewport* InViewport) const;
+	bool IsDoingDrag(FViewport* InViewport) const;
 	void AddSequencer(TWeakPtr<FSequencer> InSequencer) { Sequencers.AddUnique(InSequencer); }
 	void RemoveSequencer(TWeakPtr<FSequencer> InSequencer) { Sequencers.Remove(InSequencer); }
-
+	USequencerSettings* GetSequencerSettings() const;
 	void OnSequencerReceivedFocus(TWeakPtr<FSequencer> InSequencer) { Sequencers.Sort([=](TWeakPtr<FSequencer> A, TWeakPtr<FSequencer> B){ return A == InSequencer; }); }
 
 	void OnKeySelected(FViewport* Viewport, HMovieSceneKeyProxy* KeyProxy);
@@ -105,6 +143,9 @@ private:
 	TOptional<int32> StartXValue;
 	/** Starting Time Value*/
 	FFrameNumber StartFrameNumber;
+
+	FMarqueeDragTool DragToolHandler;
+
 };
 
 /**
diff --git a/Engine/Source/Editor/Sequencer/Private/SequencerSettings.cpp b/Engine/Source/Editor/Sequencer/Private/SequencerSettings.cpp
--- a/Engine/Source/Editor/Sequencer/Private/SequencerSettings.cpp
+++ b/Engine/Source/Editor/Sequencer/Private/SequencerSettings.cpp
@@ -41,6 +41,7 @@ USequencerSettings::USequencerSettings( const FObjectInitializer& ObjectInitiali
 	bKeepCursorInPlayRangeWhileScrubbing = false;
 	bKeepPlayRangeInSectionBounds = true;
 	bCompileDirectorOnEvaluate = true;
+	bLeftMouseDragDoesMarquee = false;
 	ZeroPadFrames = 0;
 	JumpFrameIncrement = FFrameNumber(5);
 	bShowCombinedKeyframes = true;
@@ -207,6 +208,19 @@ void USequencerSettings::SetSnapKeyTimesToKeys(bool InbSnapKeyTimesToKeys)
 	}
 }
 
+bool USequencerSettings::GetLeftMouseDragDoesMarquee() const
+{
+	return bLeftMouseDragDoesMarquee;
+}
+void USequencerSettings::SetLeftMouseDragDoesMarque(bool bDoMarque)
+{
+	if (bLeftMouseDragDoesMarquee != bDoMarque)
+	{
+		bLeftMouseDragDoesMarquee = bDoMarque;
+		SaveConfig();
+	}
+}
+
 bool USequencerSettings::GetSnapSectionTimesToInterval() const
 {
 	return bSnapSectionTimesToInterval;
diff --git a/Engine/Source/Editor/Sequencer/Public/SequencerSettings.h b/Engine/Source/Editor/Sequencer/Public/SequencerSettings.h
--- a/Engine/Source/Editor/Sequencer/Public/SequencerSettings.h
+++ b/Engine/Source/Editor/Sequencer/Public/SequencerSettings.h
@@ -214,10 +214,10 @@ public:
 	void SetShowSelectedNodesOnly(bool Visible);
 	FOnShowSelectedNodesOnlyChanged& GetOnShowSelectedNodesOnlyChanged() { return OnShowSelectedNodesOnlyChangedEvent; }
 
-	/** Gets whether to jump to the start of the sequence when we start a recording or not. */
-	bool ShouldRewindOnRecord() const;
-	/** Sets whether to jump to the start of the sequence when we start a recording. */
-	void SetRewindOnRecord(bool bInRewindOnRecord);
+	/** Gets whether left mouse drag does marquee select instead of camera orbit */
+	bool GetLeftMouseDragDoesMarquee() const;
+	/** Sets whether left mouse drag does marquee select instead of camera orbit */
+	void SetLeftMouseDragDoesMarque(bool bDoMarque);
 
 	/** Get zoom in/out position (mouse position or current time). */
 	ESequencerZoomPosition GetZoomPosition() const;
@@ -478,6 +478,10 @@ protected:
 	UPROPERTY(config, EditAnywhere, Category=General)
 	bool bRewindOnRecord;
 
+	/** Defines whether left mouse drag does marquee select instead of camera orbit */
+	UPROPERTY(config, EditAnywhere, Category = General)
+	bool bLeftMouseDragDoesMarquee;
+
 	/** Whether to zoom in on the current position or the current time in the timeline. */
 	UPROPERTY( config, EditAnywhere, Category=Timeline )
 	TEnumAsByte<ESequencerZoomPosition> ZoomPosition;
