Sequencer: Fix input handling

diff --git a/Engine/Source/Editor/Sequencer/Private/SequencerEdMode.cpp b/Engine/Source/Editor/Sequencer/Private/SequencerEdMode.cpp
--- a/Engine/Source/Editor/Sequencer/Private/SequencerEdMode.cpp
+++ b/Engine/Source/Editor/Sequencer/Private/SequencerEdMode.cpp
@@ -272,16 +272,24 @@ bool FSequencerEdMode::StartTracking(FEditorViewportClient* InViewportClient, FV
 		}
 		return true;
 	}
-	else if(IsDoingDrag(InViewport))
+	else if (IsDoingDrag(InViewport))
 	{
 		DragToolHandler.MakeDragTool(InViewportClient);
 		return DragToolHandler.StartTracking(InViewportClient, InViewport);
 	}
-	return false;
+	return FEdMode::StartTracking(InViewportClient, InViewport);
 }
 bool FSequencerEdMode::EndTracking(FEditorViewportClient* InViewportClient, FViewport* InViewport)
 {
-	return DragToolHandler.EndTracking(InViewportClient, InViewport);
+	if (IsPressingMoveTimeSlider(InViewport))
+	{
+		return true;
+	}
+	else if (IsDoingDrag(InViewport))
+	{
+		return DragToolHandler.EndTracking(InViewportClient, InViewport);
+	}
+	return FEdMode::EndTracking(InViewportClient, InViewport);
 }
 
 bool FSequencerEdMode::InputDelta(FEditorViewportClient* InViewportClient, FViewport* InViewport, FVector& InDrag, FRotator& InRot, FVector& InScale)
@@ -290,11 +298,11 @@ bool FSequencerEdMode::InputDelta(FEditorViewportClient* InViewportClient, FView
 	{
 		return true;
 	}
-	else if(IsDoingDrag(InViewport))
+	else if (IsDoingDrag(InViewport))
 	{
 		return DragToolHandler.InputDelta(InViewportClient, InViewport, InDrag, InRot, InScale);
 	}
-	return false;
+	return FEdMode::InputDelta(InViewportClient, InViewport, InDrag, InRot, InScale);
 }
 
 static void SnapSequencerTime(TSharedPtr<FSequencer>& Sequencer, FFrameTime& ScrubTime)
