Sequencer:Control Rig: Support indirect manip with CTRL drags with control rig objects.

diff --git a/Engine/Plugins/Animation/ControlRig/Source/ControlRigEditor/Private/EditMode/ControlRigEditMode.cpp b/Engine/Plugins/Animation/ControlRig/Source/ControlRigEditor/Private/EditMode/ControlRigEditMode.cpp
--- a/Engine/Plugins/Animation/ControlRig/Source/ControlRigEditor/Private/EditMode/ControlRigEditMode.cpp
+++ b/Engine/Plugins/Animation/ControlRig/Source/ControlRigEditor/Private/EditMode/ControlRigEditMode.cpp
@@ -1638,9 +1638,8 @@ bool FControlRigEditMode::InputDelta(FEditorViewportClient* InViewportClient, FV
 	FVector Scale = InScale;
 
 	const bool bCtrlDown = InViewport->KeyState(EKeys::LeftControl) || InViewport->KeyState(EKeys::RightControl);
-	const bool bShiftDown = InViewport->KeyState(EKeys::LeftShift) || InViewport->KeyState(EKeys::RightShift);
-	//bAltDown We don't care about if it is down we still want to move and not clone.
-	const bool bMouseButtonDown = InViewport->KeyState(EKeys::LeftMouseButton);
+	//button down if left and ctrl and right is down, needed for indirect posting
+	const bool bMouseButtonDown = InViewport->KeyState(EKeys::LeftMouseButton) || (bCtrlDown && InViewport->KeyState(EKeys::RightMouseButton));
 
 	const UE::Widget::EWidgetMode WidgetMode = InViewportClient->GetWidgetMode();
 	const EAxisList::Type CurrentAxis = InViewportClient->GetCurrentWidgetAxis();
@@ -1651,7 +1650,7 @@ bool FControlRigEditMode::InputDelta(FEditorViewportClient* InViewportClient, FV
 	const bool bDoScale = !Scale.IsZero() && WidgetMode == UE::Widget::WM_Scale;
 
 
-	if (InteractionScopes.Num() > 0 && bMouseButtonDown && !bCtrlDown && !bShiftDown && CurrentAxis != EAxisList::None
+	if (InteractionScopes.Num() > 0 && bMouseButtonDown && CurrentAxis != EAxisList::None
 		&& (bDoRotation || bDoTranslation || bDoScale))
 	{
 		for (TPair<UControlRig*, TArray<AControlRigShapeActor*>>& Pairs : ControlRigShapeActors)
diff --git a/Engine/Source/Editor/Sequencer/Private/SequencerEdMode.cpp b/Engine/Source/Editor/Sequencer/Private/SequencerEdMode.cpp
--- a/Engine/Source/Editor/Sequencer/Private/SequencerEdMode.cpp
+++ b/Engine/Source/Editor/Sequencer/Private/SequencerEdMode.cpp
@@ -246,10 +246,10 @@ bool FSequencerEdMode::IsDoingDrag(FViewport* InViewport) const
 	const USequencerSettings* SequencerSettings = GetSequencerSettings();
 	const bool LeftMouseButtonDown = InViewport->KeyState(EKeys::LeftMouseButton);
 	const bool bIsCtrlKeyDown = InViewport->KeyState(EKeys::LeftControl) || InViewport->KeyState(EKeys::RightControl);
-	const bool bIsShiftKeyDown = InViewport->KeyState(EKeys::LeftShift) || InViewport->KeyState(EKeys::RightShift);
 	const bool bIsAltKeyDown = InViewport->KeyState(EKeys::LeftAlt) || InViewport->KeyState(EKeys::RightAlt);
+	//if shfit is down we still want to drag
 
-	return LeftMouseButtonDown && !bIsCtrlKeyDown && !bIsShiftKeyDown && !bIsAltKeyDown && (SequencerSettings ? SequencerSettings->GetLeftMouseDragDoesMarquee() : false);
+	return LeftMouseButtonDown && !bIsCtrlKeyDown  && !bIsAltKeyDown && (SequencerSettings ? SequencerSettings->GetLeftMouseDragDoesMarquee() : false);
 }
 
 bool FSequencerEdMode::StartTracking(FEditorViewportClient* InViewportClient, FViewport* InViewport)
@@ -285,9 +285,9 @@ bool FSequencerEdMode::EndTracking(FEditorViewportClient* InViewportClient, FVie
 	{
 		return true;
 	}
-	else if (IsDoingDrag(InViewport))
+	else if (DragToolHandler.EndTracking(InViewportClient, InViewport))
 	{
-		return DragToolHandler.EndTracking(InViewportClient, InViewport);
+		return true;
 	}
 	return FEdMode::EndTracking(InViewportClient, InViewport);
 }
@@ -1002,11 +1002,11 @@ bool FMarqueeDragTool::EndTracking(FEditorViewportClient* InViewportClient, FVie
 				DragTool->EndDrag();
 			}
 			DragTool.Reset();
-			return false;
+			return true;
 		}
 	}
 	
-	return true;
+	return false;
 }
 
 bool FMarqueeDragTool::InputDelta(FEditorViewportClient* InViewportClient, FViewport* InViewport, FVector& InDrag, FRotator& InRot, FVector& InScale)
