Sequencer: Only clamp to view range when scrubbing. Otherwise setting the time from shortcuts which call into the time slider controller get clamped to the view range.

diff --git a/Engine/Source/Editor/Sequencer/Private/SequencerTimeSliderController.cpp b/Engine/Source/Editor/Sequencer/Private/SequencerTimeSliderController.cpp
--- a/Engine/Source/Editor/Sequencer/Private/SequencerTimeSliderController.cpp
+++ b/Engine/Source/Editor/Sequencer/Private/SequencerTimeSliderController.cpp
@@ -1134,7 +1134,7 @@ void FSequencerTimeSliderController::CommitScrubPosition( FFrameTime NewValue, b
 	bIsEvaluating = bEvaluate;
 	// The user can scrub past the viewing range of the time slider controller, so we clamp it to the view range.
 	TSharedPtr<FSequencer> Sequencer = WeakSequencer.Pin();
-	if(Sequencer.IsValid())
+	if(Sequencer.IsValid() && bIsScrubbing)
 	{
 		FAnimatedRange ViewRange = GetViewRange();
 		
