Sequencer: MVVM2 branch and Layer Bars

Copying //Tasks/UE5/Dev-SequencerMVVM2 to Main (//UE5/Main) [REDACTED_HANDLE]

diff --git a/Engine/Plugins/Animation/ControlRig/Source/ControlRigEditor/ControlRigEditor.Build.cs b/Engine/Plugins/Animation/ControlRig/Source/ControlRigEditor/ControlRigEditor.Build.cs
--- a/Engine/Plugins/Animation/ControlRig/Source/ControlRigEditor/ControlRigEditor.Build.cs
+++ b/Engine/Plugins/Animation/ControlRig/Source/ControlRigEditor/ControlRigEditor.Build.cs
@@ -54,6 +54,7 @@
                     "MovieScene",
                     "MovieSceneTracks",
                     "MovieSceneTools",
+                    "SequencerCore",
                     "Sequencer",
 					"LevelSequenceEditor",
                     "ClassViewer",
diff --git a/Engine/Plugins/Animation/ControlRig/Source/ControlRigEditor/Private/ControlRigSpaceChannelCurveModel.cpp b/Engine/Plugins/Animation/ControlRig/Source/ControlRigEditor/Private/ControlRigSpaceChannelCurveModel.cpp
--- a/Engine/Plugins/Animation/ControlRig/Source/ControlRigEditor/Private/ControlRigSpaceChannelCurveModel.cpp
+++ b/Engine/Plugins/Animation/ControlRig/Source/ControlRigEditor/Private/ControlRigSpaceChannelCurveModel.cpp
@@ -15,11 +15,14 @@
 #include "SequencerSectionPainter.h"
 #include "ISequencer.h"
 #include "ControlRigSpaceChannelEditors.h"
+#include "MVVM/Views/STrackAreaView.h"
 
 ECurveEditorViewID FControlRigSpaceChannelCurveModel::ViewID = ECurveEditorViewID::Invalid;
 
 FControlRigSpaceChannelCurveModel::FControlRigSpaceChannelCurveModel(TMovieSceneChannelHandle<FMovieSceneControlRigSpaceChannel> InChannel, UMovieSceneSection* OwningSection, TWeakPtr<ISequencer> InWeakSequencer)
 {
+	using namespace UE::Sequencer;
+	
 	FMovieSceneChannelProxy* NewChannelProxy = &OwningSection->GetChannelProxy();
 	ChannelHandle = NewChannelProxy->MakeHandle<FMovieSceneControlRigSpaceChannel>(InChannel.GetChannelIndex());
 	if (FMovieSceneChannelProxy* ChannelProxy = InChannel.GetChannelProxy())
@@ -31,7 +34,7 @@
 	WeakSequencer = InWeakSequencer;
 	SupportedViews = ViewID;
 
-	Color = FSequencerSectionPainter::BlendColor(OwningSection->GetTypedOuter<UMovieSceneTrack>()->GetColorTint());
+	Color = STrackAreaView::BlendDefaultTrackColor(OwningSection->GetTypedOuter<UMovieSceneTrack>()->GetColorTint());
 }
 
 FControlRigSpaceChannelCurveModel::~FControlRigSpaceChannelCurveModel()
diff --git a/Engine/Plugins/Animation/ControlRig/Source/ControlRigEditor/Private/Sequencer/ControlRigSpaceChannelEditors.cpp b/Engine/Plugins/Animation/ControlRig/Source/ControlRigEditor/Private/Sequencer/ControlRigSpaceChannelEditors.cpp
--- a/Engine/Plugins/Animation/ControlRig/Source/ControlRigEditor/Private/Sequencer/ControlRigSpaceChannelEditors.cpp
+++ b/Engine/Plugins/Animation/ControlRig/Source/ControlRigEditor/Private/Sequencer/ControlRigSpaceChannelEditors.cpp
@@ -58,6 +58,7 @@
 #include "Rendering/DrawElements.h"
 #include "Fonts/FontMeasure.h"
 #include "CurveEditorSettings.h"
+#include "TimeToPixel.h"
 
 #define LOCTEXT_NAMESPACE "ControlRigEditMode"
 
@@ -1525,6 +1526,8 @@
 
 void DrawExtra(FMovieSceneControlRigSpaceChannel* Channel, const UMovieSceneSection* Owner,const FGeometry& AllottedGeometry, FSequencerSectionPainter& Painter)
 {
+	using namespace UE::Sequencer;
+
 	if (const UMovieSceneControlRigParameterSection* Section = Cast<UMovieSceneControlRigParameterSection>(Owner))
 	{
 		TArray<FKeyBarCurveModel::FBarRange> Ranges = FControlRigSpaceChannelHelpers::FindRanges(Channel, Owner);
diff --git a/Engine/Plugins/Experimental/GeometryCollectionPlugin/Source/GeometryCollectionSequencer/Private/GeometryCollectionTrackEditor.cpp b/Engine/Plugins/Experimental/GeometryCollectionPlugin/Source/GeometryCollectionSequencer/Private/GeometryCollectionTrackEditor.cpp
--- a/Engine/Plugins/Experimental/GeometryCollectionPlugin/Source/GeometryCollectionSequencer/Private/GeometryCollectionTrackEditor.cpp
+++ b/Engine/Plugins/Experimental/GeometryCollectionPlugin/Source/GeometryCollectionSequencer/Private/GeometryCollectionTrackEditor.cpp
@@ -29,6 +29,7 @@
 #include "GeometryCollection/GeometryCollectionComponent.h"
 #include "Styling/SlateIconFinder.h"
 #include "LevelSequence.h"
+#include "TimeToPixel.h"
 
 namespace GeometryCollectionEditorConstants
 {
@@ -97,6 +98,8 @@
 
 int32 FGeometryCollectionTrackSection::OnPaintSection( FSequencerSectionPainter& Painter ) const
 {
+	using namespace UE::Sequencer;
+
 	const ESlateDrawEffect DrawEffects = Painter.bParentEnabled ? ESlateDrawEffect::None : ESlateDrawEffect::DisabledEffect;
 	
 	const FTimeToPixel& TimeToPixelConverter = Painter.GetTimeConverter();
diff --git a/Engine/Plugins/FX/Niagara/Source/NiagaraEditor/Private/ViewModels/NiagaraSystemViewModel.cpp b/Engine/Plugins/FX/Niagara/Source/NiagaraEditor/Private/ViewModels/NiagaraSystemViewModel.cpp
--- a/Engine/Plugins/FX/Niagara/Source/NiagaraEditor/Private/ViewModels/NiagaraSystemViewModel.cpp
+++ b/Engine/Plugins/FX/Niagara/Source/NiagaraEditor/Private/ViewModels/NiagaraSystemViewModel.cpp
@@ -1487,16 +1487,16 @@ void FNiagaraSystemViewModel::RefreshSequencerTracks()
 			EmitterHandleIdToTrackMap.Add(EmitterHandleViewModel->GetId(), EmitterTrack);
 		}
 
-		TArray<UMovieSceneFolder*>& MovieSceneRootFolders = NiagaraSequence->GetMovieScene()->GetRootFolders();
-		MovieSceneRootFolders.Empty();
+		UMovieScene* MovieScene = NiagaraSequence->GetMovieScene();
+		MovieScene->EmptyRootFolders();
 
 		const UNiagaraSystemEditorData& SystemEditorData = GetEditorData();
 		UNiagaraSystemEditorFolder& RootFolder = SystemEditorData.GetRootFolder();
 		for (const UNiagaraSystemEditorFolder* RootChildFolder : RootFolder.GetChildFolders())
 		{
-			UMovieSceneFolder* MovieSceneRootFolder = NewObject<UMovieSceneFolder>(NiagaraSequence->GetMovieScene(), RootChildFolder->GetFolderName(), RF_Transactional);
+			UMovieSceneFolder* MovieSceneRootFolder = NewObject<UMovieSceneFolder>(MovieScene, RootChildFolder->GetFolderName(), RF_Transactional);
 			MovieSceneRootFolder->SetFolderName(RootChildFolder->GetFolderName());
-			MovieSceneRootFolders.Add(MovieSceneRootFolder);
+			MovieScene->AddRootFolder(MovieSceneRootFolder);
 			PopulateChildMovieSceneFoldersFromNiagaraFolders(RootChildFolder, MovieSceneRootFolder, EmitterHandleIdToTrackMap);
 		}
 		Sequencer->NotifyMovieSceneDataChanged(EMovieSceneDataChangeType::MovieSceneStructureItemsChanged);
@@ -1896,7 +1896,7 @@ void FNiagaraSystemViewModel::UpdateSimulationFromParameterChange()
 	}
 }
 
-void PopulateNiagaraFoldersFromMovieSceneFolders(const TArray<UMovieSceneFolder*>& MovieSceneFolders, const TArray<UMovieSceneTrack*>& MovieSceneTracks, UNiagaraSystemEditorFolder* ParentFolder)
+void PopulateNiagaraFoldersFromMovieSceneFolders(TArrayView<UMovieSceneFolder* const> MovieSceneFolders, const TArray<UMovieSceneTrack*>& MovieSceneTracks, UNiagaraSystemEditorFolder* ParentFolder)
 {
 	TArray<FName> ValidFolderNames;
 	for (UMovieSceneFolder* MovieSceneFolder : MovieSceneFolders)
@@ -2054,7 +2054,7 @@ void FNiagaraSystemViewModel::SequencerDataChanged(EMovieSceneDataChangeType Dat
 		}
 
 		TArray<UMovieSceneTrack*> RootTracks;
-		TArray<UMovieSceneFolder*> RootFolders = NiagaraSequence->GetMovieScene()->GetRootFolders();
+		TArrayView<UMovieSceneFolder* const> RootFolders = NiagaraSequence->GetMovieScene()->GetRootFolders();
 		if (RootFolders.Num() != 0 || GetEditorData().GetRootFolder().GetChildFolders().Num() != 0)
 		{
 			PopulateNiagaraFoldersFromMovieSceneFolders(RootFolders, RootTracks, &GetEditorData().GetRootFolder());
diff --git a/Engine/Plugins/Importers/USDImporter/Source/USDStage/Private/USDLevelSequenceHelper.cpp b/Engine/Plugins/Importers/USDImporter/Source/USDStage/Private/USDLevelSequenceHelper.cpp
--- a/Engine/Plugins/Importers/USDImporter/Source/USDStage/Private/USDLevelSequenceHelper.cpp
+++ b/Engine/Plugins/Importers/USDImporter/Source/USDStage/Private/USDLevelSequenceHelper.cpp
@@ -1700,7 +1700,7 @@
 			{
 				if ( FMovieScenePossessable* ComponentPossessable = MovieScene->FindPossessable( ComponentBinding ) )
 				{
-					ComponentPossessable->SetParent( ActorBinding );
+					ComponentPossessable->SetParent( ActorBinding, MovieScene );
 				}
 			}
 
diff --git a/Engine/Plugins/Media/MediaCompositing/Source/MediaCompositingEditor/Private/Sequencer/MediaThumbnailSection.cpp b/Engine/Plugins/Media/MediaCompositing/Source/MediaCompositingEditor/Private/Sequencer/MediaThumbnailSection.cpp
--- a/Engine/Plugins/Media/MediaCompositing/Source/MediaCompositingEditor/Private/Sequencer/MediaThumbnailSection.cpp
+++ b/Engine/Plugins/Media/MediaCompositing/Source/MediaCompositingEditor/Private/Sequencer/MediaThumbnailSection.cpp
@@ -20,6 +20,7 @@
 #include "CommonMovieSceneTools.h"
 #include "Widgets/Layout/SBox.h"
 #include "Widgets/Text/STextBlock.h"
+#include "TimeToPixel.h"
 
 #include "MovieSceneMediaData.h"
 
@@ -276,6 +277,8 @@
 
 void FMediaThumbnailSection::DrawLoopIndicators(FSequencerSectionPainter& InPainter, FTimespan MediaDuration, FVector2D SectionSize) const
 {
+	using namespace UE::Sequencer;
+
 	static const FSlateBrush* GenericBrush = FCoreStyle::Get().GetBrush("GenericWhiteBox");
 
 	UMovieSceneMediaSection* MediaSection = Cast<UMovieSceneMediaSection>(Section);
@@ -305,6 +308,8 @@
 
 void FMediaThumbnailSection::DrawSampleStates(FSequencerSectionPainter& InPainter, FTimespan MediaDuration, FVector2D SectionSize, const TRangeSet<FTimespan>& RangeSet, const FLinearColor& Color) const
 {
+	using namespace UE::Sequencer;
+
 	static const FSlateBrush* GenericBrush = FCoreStyle::Get().GetBrush("GenericWhiteBox");
 
 	UMovieSceneMediaSection* MediaSection = Cast<UMovieSceneMediaSection>(Section);
diff --git a/Engine/Plugins/MovieScene/ActorSequence/Source/ActorSequenceEditor/Private/ActorSequenceEditorTabSummoner.cpp b/Engine/Plugins/MovieScene/ActorSequence/Source/ActorSequenceEditor/Private/ActorSequenceEditorTabSummoner.cpp
--- a/Engine/Plugins/MovieScene/ActorSequence/Source/ActorSequenceEditor/Private/ActorSequenceEditorTabSummoner.cpp
+++ b/Engine/Plugins/MovieScene/ActorSequence/Source/ActorSequenceEditor/Private/ActorSequenceEditorTabSummoner.cpp
@@ -19,6 +19,7 @@
 #include "Framework/MultiBox/MultiBoxBuilder.h"
 #include "Framework/Application/SlateApplication.h"
 #include "ActorSequenceEditorStyle.h"
+#include "EventHandlers/ISignedObjectEventHandler.h"
 #include "UObject/ObjectSaveContext.h"
 #include "SSubobjectEditor.h"
 
@@ -189,7 +190,7 @@
 	TArray<TSharedPtr<FSCSEditorTreeNode>> RootNodes;
 };
 
-class SActorSequenceEditorWidgetImpl : public SCompoundWidget, public FEditorUndoClient
+class SActorSequenceEditorWidgetImpl : public SCompoundWidget, public FEditorUndoClient, UE::MovieScene::TIntrusiveEventHandler<UE::MovieScene::ISignedObjectEventHandler>
 {
 public:
 
@@ -325,19 +326,13 @@
 
 	void SetActorSequence(UActorSequence* NewSequence)
 	{
-		if (UActorSequence* OldSequence = WeakSequence.Get())
-		{
-			if (OnSequenceChangedHandle.IsValid())
-			{
-				OldSequence->OnSignatureChanged().Remove(OnSequenceChangedHandle);
-			}
-		}
+		Unlink();
 
 		WeakSequence = NewSequence;
 
 		if (NewSequence)
 		{
-			OnSequenceChangedHandle = NewSequence->OnSignatureChanged().AddSP(this, &SActorSequenceEditorWidgetImpl::OnSequenceChanged);
+			NewSequence->EventHandlers.Link(this);
 		}
 
 		// If we already have a sequencer open, just assign the sequence
@@ -610,7 +605,7 @@
 		return Actor ? FindObject<UActorComponent>(Actor, *ActorSequence->GetOuter()->GetName()) : nullptr;
 	}
 
-	void OnSequenceChanged()
+	void OnModifiedIndirectly(UMovieSceneSignedObject*) override
 	{
 		UActorSequence* ActorSequence = WeakSequence.Get();
 		UBlueprint* Blueprint = ActorSequence ? ActorSequence->GetParentBlueprint() : nullptr;
@@ -620,6 +615,10 @@
 			FBlueprintEditorUtils::MarkBlueprintAsModified(Blueprint);
 		}
 	}
+	void OnModifiedDirectly(UMovieSceneSignedObject* Object) override
+	{
+		OnModifiedIndirectly(Object);
+	}
 
 private:
 	TWeakObjectPtr<UActorSequence> WeakSequence;
@@ -632,8 +631,6 @@
 	FDelegateHandle OnBlueprintPreCompileHandle;
 	FDelegateHandle OnObjectSavedHandle;
 
-	FDelegateHandle OnSequenceChangedHandle;
-
 	FDelegateHandle OnSelectionChangedHandle;
 	FDelegateHandle OnGlobalTimeChangedHandle;
 };
diff --git a/Engine/Plugins/MovieScene/LevelSequenceEditor/Source/LevelSequenceEditor/Private/CinematicViewport/SCinematicTransportRange.cpp b/Engine/Plugins/MovieScene/LevelSequenceEditor/Source/LevelSequenceEditor/Private/CinematicViewport/SCinematicTransportRange.cpp
--- a/Engine/Plugins/MovieScene/LevelSequenceEditor/Source/LevelSequenceEditor/Private/CinematicViewport/SCinematicTransportRange.cpp
+++ b/Engine/Plugins/MovieScene/LevelSequenceEditor/Source/LevelSequenceEditor/Private/CinematicViewport/SCinematicTransportRange.cpp
@@ -9,6 +9,7 @@
 #include "MovieScene.h"
 #include "ISequencer.h"
 #include "CommonMovieSceneTools.h"
+#include "TimeToPixel.h"
 
 #define LOCTEXT_NAMESPACE "SCinematicTransportRange"
 
@@ -93,6 +94,8 @@
 
 void SCinematicTransportRange::Tick(const FGeometry& AllottedGeometry, const double InCurrentTime, const float InDeltaTime)
 {
+	using namespace UE::Sequencer;
+
 	ISequencer* Sequencer = GetSequencer();
 	if (Sequencer)
 	{
@@ -163,6 +166,8 @@
 
 int32 SCinematicTransportRange::OnPaint(const FPaintArgs& Args, const FGeometry& AllottedGeometry, const FSlateRect& MyCullingRect, FSlateWindowElementList& OutDrawElements, int32 LayerId, const FWidgetStyle& InWidgetStyle, bool bParentEnabled) const
 {
+	using namespace UE::Sequencer;
+
 	ISequencer* Sequencer = GetSequencer();
 
 	const ESlateDrawEffect DrawEffects = bParentEnabled ? ESlateDrawEffect::None : ESlateDrawEffect::DisabledEffect;
diff --git a/Engine/Plugins/MovieScene/MatineeToLevelSequence/Source/MatineeToLevelSequence/Private/CameraAnimToTemplateSequenceConverter.cpp b/Engine/Plugins/MovieScene/MatineeToLevelSequence/Source/MatineeToLevelSequence/Private/CameraAnimToTemplateSequenceConverter.cpp
--- a/Engine/Plugins/MovieScene/MatineeToLevelSequence/Source/MatineeToLevelSequence/Private/CameraAnimToTemplateSequenceConverter.cpp
+++ b/Engine/Plugins/MovieScene/MatineeToLevelSequence/Source/MatineeToLevelSequence/Private/CameraAnimToTemplateSequenceConverter.cpp
@@ -552,7 +552,7 @@
 			UCameraComponent* CameraComponent = CameraTemplate->GetCameraComponent();
 			CameraComponentBindingID = NewMovieScene->AddPossessable(CameraComponent->GetName(), CameraComponent->GetClass());
 			FMovieScenePossessable* CameraComponentPossessable = NewMovieScene->FindPossessable(CameraComponentBindingID);
-			CameraComponentPossessable->SetParent(SpawnableGuid);
+			CameraComponentPossessable->SetParent(SpawnableGuid, NewMovieScene);
 			NewSequence->BindPossessableObject(CameraComponentBindingID, *CameraComponent, nullptr);
 		}
 	};
diff --git a/Engine/Plugins/MovieScene/MatineeToLevelSequence/Source/MatineeToLevelSequence/Private/MatineeConverter.cpp b/Engine/Plugins/MovieScene/MatineeToLevelSequence/Source/MatineeToLevelSequence/Private/MatineeConverter.cpp
--- a/Engine/Plugins/MovieScene/MatineeToLevelSequence/Source/MatineeToLevelSequence/Private/MatineeConverter.cpp
+++ b/Engine/Plugins/MovieScene/MatineeToLevelSequence/Source/MatineeToLevelSequence/Private/MatineeConverter.cpp
@@ -123,7 +123,7 @@
 	{
 		FolderToUse = NewObject<UMovieSceneFolder>(MovieScene, NAME_None, RF_Transactional);
 		FolderToUse->SetFolderName(FolderName);
-		MovieScene->GetRootFolders().Add(FolderToUse);
+		MovieScene->AddRootFolder(FolderToUse);
 	}
 
 	return FolderToUse;
@@ -153,7 +153,7 @@
 		{
 			ComponentGUID = NewMovieScene->AddPossessable(SkelMeshComponent->GetName(), SkelMeshComponent->GetClass());
 			FMovieScenePossessable* ChildPossesable = NewMovieScene->FindPossessable(ComponentGUID);
-			ChildPossesable->SetParent(PossessableGuid);
+			ChildPossesable->SetParent(PossessableGuid, NewMovieScene);
 			NewSequence->BindPossessableObject(ComponentGUID, *SkelMeshComponent, GroupActor->GetWorld());
 		}
 	}
@@ -170,7 +170,7 @@
 		{
 			ComponentGUID = NewMovieScene->AddPossessable(StaticMeshComponent->GetName(), StaticMeshComponent->GetClass());
 			FMovieScenePossessable* ChildPossesable = NewMovieScene->FindPossessable(ComponentGUID);
-			ChildPossesable->SetParent(PossessableGuid);
+			ChildPossesable->SetParent(PossessableGuid, NewMovieScene);
 			NewSequence->BindPossessableObject(ComponentGUID, *StaticMeshComponent, GroupActor->GetWorld());
 		}
 	}
diff --git a/Engine/Plugins/MovieScene/MatineeToLevelSequence/Source/MatineeToLevelSequence/Private/MatineeConverter.h b/Engine/Plugins/MovieScene/MatineeToLevelSequence/Source/MatineeToLevelSequence/Private/MatineeConverter.h
--- a/Engine/Plugins/MovieScene/MatineeToLevelSequence/Source/MatineeToLevelSequence/Private/MatineeConverter.h
+++ b/Engine/Plugins/MovieScene/MatineeToLevelSequence/Source/MatineeToLevelSequence/Private/MatineeConverter.h
@@ -109,7 +109,7 @@
 
 		if (ChildPossessable)
 		{
-			ChildPossessable->SetParent(PossessableGuid);
+			ChildPossessable->SetParent(PossessableGuid, NewMovieScene);
 		}
 
 		// cbb: String manipulations to get the property path in the right form for sequencer
diff --git a/Engine/Plugins/MovieScene/SequencerScripting/Source/SequencerScripting/Private/ExtensionLibraries/MovieSceneBindingExtensions.cpp b/Engine/Plugins/MovieScene/SequencerScripting/Source/SequencerScripting/Private/ExtensionLibraries/MovieSceneBindingExtensions.cpp
--- a/Engine/Plugins/MovieScene/SequencerScripting/Source/SequencerScripting/Private/ExtensionLibraries/MovieSceneBindingExtensions.cpp
+++ b/Engine/Plugins/MovieScene/SequencerScripting/Source/SequencerScripting/Private/ExtensionLibraries/MovieSceneBindingExtensions.cpp
@@ -268,7 +268,7 @@
 		if (Possessable)
 		{
 			MovieScene->Modify();
-			Possessable->SetParent(InParentBinding.BindingID);
+			Possessable->SetParent(InParentBinding.BindingID, MovieScene);
 		}
 	}
 }
diff --git a/Engine/Plugins/MovieScene/SequencerScripting/Source/SequencerScripting/Private/ExtensionLibraries/MovieSceneFolderExtensions.cpp b/Engine/Plugins/MovieScene/SequencerScripting/Source/SequencerScripting/Private/ExtensionLibraries/MovieSceneFolderExtensions.cpp
--- a/Engine/Plugins/MovieScene/SequencerScripting/Source/SequencerScripting/Private/ExtensionLibraries/MovieSceneFolderExtensions.cpp
+++ b/Engine/Plugins/MovieScene/SequencerScripting/Source/SequencerScripting/Private/ExtensionLibraries/MovieSceneFolderExtensions.cpp
@@ -65,7 +65,8 @@
 		return TArray<UMovieSceneFolder*>();
 	}
 
-	return Folder->GetChildFolders();
+	TArray<UMovieSceneFolder*> Result(Folder->GetChildFolders());
+	return Result;
 }
 
 bool UMovieSceneFolderExtensions::AddChildFolder(UMovieSceneFolder* TargetFolder, UMovieSceneFolder* FolderToAdd)
diff --git a/Engine/Plugins/MovieScene/SequencerScripting/Source/SequencerScripting/Private/ExtensionLibraries/MovieSceneSequenceExtensions.cpp b/Engine/Plugins/MovieScene/SequencerScripting/Source/SequencerScripting/Private/ExtensionLibraries/MovieSceneSequenceExtensions.cpp
--- a/Engine/Plugins/MovieScene/SequencerScripting/Source/SequencerScripting/Private/ExtensionLibraries/MovieSceneSequenceExtensions.cpp
+++ b/Engine/Plugins/MovieScene/SequencerScripting/Source/SequencerScripting/Private/ExtensionLibraries/MovieSceneSequenceExtensions.cpp
@@ -1005,7 +1005,7 @@
 		MovieScene->Modify();
 		NewFolder = NewObject<UMovieSceneFolder>(MovieScene);
 		NewFolder->SetFolderName(FName(*NewFolderName));
-		MovieScene->GetRootFolders().Add(NewFolder);
+		MovieScene->AddRootFolder(NewFolder);
 	}
 #endif
 
@@ -1032,7 +1032,7 @@
 	if (FolderMovieScene)
 	{
 		FolderMovieScene->Modify();
-		const int32 NumFoldersRemoved = FolderMovieScene->GetRootFolders().Remove(Folder);
+		const int32 NumFoldersRemoved = FolderMovieScene->RemoveRootFolder(Folder);
 		if (NumFoldersRemoved == 0)
 		{
 			FFrame::KismetExecutionMessage(*FString::Printf(TEXT("The specified folder '%s' is not a root folder"),
diff --git a/Engine/Plugins/MovieScene/TemplateSequence/Source/TemplateSequence/Private/TemplateSequence.cpp b/Engine/Plugins/MovieScene/TemplateSequence/Source/TemplateSequence/Private/TemplateSequence.cpp
--- a/Engine/Plugins/MovieScene/TemplateSequence/Source/TemplateSequence/Private/TemplateSequence.cpp
+++ b/Engine/Plugins/MovieScene/TemplateSequence/Source/TemplateSequence/Private/TemplateSequence.cpp
@@ -199,7 +199,7 @@
 		FMovieScenePossessable* ChildPossessable = MovieScene->FindPossessable(NewGuid);
 		if (ensure(ChildPossessable))
 		{
-			ChildPossessable->SetParent(ParentGuid);
+			ChildPossessable->SetParent(ParentGuid, MovieScene);
 		}
 
 		FMovieSceneSpawnable* ParentSpawnable = MovieScene->FindSpawnable(ParentGuid);
diff --git a/Engine/Plugins/VirtualProduction/SequencerPlaylists/Source/SequencerPlaylists/Private/SequencerPlaylistItem_Sequence.cpp b/Engine/Plugins/VirtualProduction/SequencerPlaylists/Source/SequencerPlaylists/Private/SequencerPlaylistItem_Sequence.cpp
--- a/Engine/Plugins/VirtualProduction/SequencerPlaylists/Source/SequencerPlaylists/Private/SequencerPlaylistItem_Sequence.cpp
+++ b/Engine/Plugins/VirtualProduction/SequencerPlaylists/Source/SequencerPlaylists/Private/SequencerPlaylistItem_Sequence.cpp
@@ -334,7 +334,7 @@
 	{
 		FolderToUse = NewObject<UMovieSceneFolder>(RootScene, NAME_None, RF_Transactional);
 		FolderToUse->SetFolderName(PlaylistFolderName);
-		RootScene->GetRootFolders().Add(FolderToUse);
+		//RootScene->GetRootFolders().Add(FolderToUse);
 	}
 
 	FolderToUse->AddChildMasterTrack(NewWorkingTrack);
diff --git a/Engine/Plugins/VirtualProduction/Takes/Source/TakeRecorder/Private/SerializedRecorder.cpp b/Engine/Plugins/VirtualProduction/Takes/Source/TakeRecorder/Private/SerializedRecorder.cpp
--- a/Engine/Plugins/VirtualProduction/Takes/Source/TakeRecorder/Private/SerializedRecorder.cpp
+++ b/Engine/Plugins/VirtualProduction/Takes/Source/TakeRecorder/Private/SerializedRecorder.cpp
@@ -500,7 +500,7 @@
 		{
 			FolderToUse = NewObject<UMovieSceneFolder>(MovieScene, NAME_None, RF_Transactional);
 			FolderToUse->SetFolderName(ActorHeader.FolderName);
-			MovieScene->GetRootFolders().Add(FolderToUse);
+			MovieScene->AddRootFolder(FolderToUse);
 		}
 
 		FolderToUse->AddChildObjectBinding(ActorHeader.Guid);
@@ -525,7 +525,7 @@
 
 		// Set up parent/child guids for possessables within spawnables
 		FMovieScenePossessable* ChildPossessablePtr = InMovieScene->FindPossessable(ActorProperty.Guid);
-		ChildPossessablePtr->SetParent(ActorHeader.Guid);
+		ChildPossessablePtr->SetParent(ActorHeader.Guid, InMovieScene);
 		
 		FMovieSceneSpawnable* ParentSpawnable = InMovieScene->FindSpawnable(ActorHeader.Guid);
 		if (ParentSpawnable)
diff --git a/Engine/Plugins/VirtualProduction/Takes/Source/TakeRecorderSources/Private/TakeRecorderActorSource.cpp b/Engine/Plugins/VirtualProduction/Takes/Source/TakeRecorderSources/Private/TakeRecorderActorSource.cpp
--- a/Engine/Plugins/VirtualProduction/Takes/Source/TakeRecorderSources/Private/TakeRecorderActorSource.cpp
+++ b/Engine/Plugins/VirtualProduction/Takes/Source/TakeRecorderSources/Private/TakeRecorderActorSource.cpp
@@ -406,7 +406,7 @@
 		FMovieScenePossessable* ChildPossessable = MovieScene->FindPossessable(Guid);
 		if (ensure(ChildPossessable))
 		{
-			ChildPossessable->SetParent(CachedObjectBindingGuid);
+			ChildPossessable->SetParent(CachedObjectBindingGuid, MovieScene);
 		}
 
 		FMovieSceneSpawnable* ParentSpawnable = MovieScene->FindSpawnable(CachedObjectBindingGuid);
diff --git a/Engine/Plugins/VirtualProduction/Takes/Source/TakeRecorderSources/Private/TakeRecorderWorldSource.cpp b/Engine/Plugins/VirtualProduction/Takes/Source/TakeRecorderSources/Private/TakeRecorderWorldSource.cpp
--- a/Engine/Plugins/VirtualProduction/Takes/Source/TakeRecorderSources/Private/TakeRecorderWorldSource.cpp
+++ b/Engine/Plugins/VirtualProduction/Takes/Source/TakeRecorderSources/Private/TakeRecorderWorldSource.cpp
@@ -196,7 +196,7 @@
 	{
 		FolderToUse = NewObject<UMovieSceneFolder>(MovieScene, NAME_None, RF_Transactional);
 		FolderToUse->SetFolderName(AutotrackedFolderName);
-		MovieScene->GetRootFolders().Add(FolderToUse);
+		MovieScene->AddRootFolder(FolderToUse);
 	}
 
 	for (FGuid AddedGuid : AddedGuids)
diff --git a/Engine/Plugins/VirtualProduction/Takes/Source/TakesCore/Private/TakeRecorderSources.cpp b/Engine/Plugins/VirtualProduction/Takes/Source/TakesCore/Private/TakeRecorderSources.cpp
--- a/Engine/Plugins/VirtualProduction/Takes/Source/TakesCore/Private/TakeRecorderSources.cpp
+++ b/Engine/Plugins/VirtualProduction/Takes/Source/TakesCore/Private/TakeRecorderSources.cpp
@@ -621,7 +621,7 @@
 	{
 		FolderToUse = NewObject<UMovieSceneFolder>(InMovieScene, NAME_None, RF_Transactional);
 		FolderToUse->SetFolderName(FolderName);
-		InMovieScene->GetRootFolders().Add(FolderToUse);
+		InMovieScene->AddRootFolder(FolderToUse);
 	}
 
 	// We want to expand these folders in the Sequencer UI (since these are visible as they record).
diff --git a/Engine/Source/Editor/EditorStyle/Private/StarshipStyle.cpp b/Engine/Source/Editor/EditorStyle/Private/StarshipStyle.cpp
--- a/Engine/Source/Editor/EditorStyle/Private/StarshipStyle.cpp
+++ b/Engine/Source/Editor/EditorStyle/Private/StarshipStyle.cpp
@@ -2122,6 +2122,11 @@
 		Set("Sequencer.Section.BackgroundText", DEFAULT_FONT("Bold", 24));
 		Set("Sequencer.Section.EmptySpace", new BOX_BRUSH(TEXT("Sequencer/EmptySpace"), FMargin(0.f, 7.f / 14.f)));
 
+		Set("Sequencer.LayerBar.Background",  new BOX_BRUSH(TEXT("Sequencer/LayerBarBackground"), FMargin(4.f / 16.f)));
+		Set("Sequencer.LayerBar.Border",      new BORDER_BRUSH(TEXT("Sequencer/LayerBarBorder"), FMargin(4.f / 16.f)));
+		Set("Sequencer.LayerBar.HandleLeft",  new BOX_BRUSH(TEXT("Sequencer/LayerBarHandleLeft"), FMargin(4.f / 16.f)));
+		Set("Sequencer.LayerBar.HandleRight", new BOX_BRUSH(TEXT("Sequencer/LayerBarHandleRight"), FMargin(4.f / 16.f)));
+
 		Set("Sequencer.ExposedNamePill_BG", new BOX_BRUSH(TEXT("Sequencer/ExposedNamePill_BG"), FMargin(14.f / 30.f), FLinearColor(1.f, 1.f, 1.f, .8f)));
 		Set("Sequencer.ExposedNamePill", FButtonStyle()
 			.SetNormal(  BOX_BRUSH(TEXT("Sequencer/ExposedNamePill_BG"), FMargin(14.f / 30.f), FLinearColor(1.f, 1.f, 1.f, .8f)) )
diff --git a/Engine/Source/Editor/MovieSceneTools/Private/MVVM/ViewModels/CameraCutTrackModel.h b/Engine/Source/Editor/MovieSceneTools/Private/MVVM/ViewModels/CameraCutTrackModel.h
new file mode 100644
--- /dev/null
+++ b/Engine/Source/Editor/MovieSceneTools/Private/MVVM/ViewModels/CameraCutTrackModel.h
@@ -0,0 +1,30 @@
+// Copyright Epic Games, Inc. All Rights Reserved.
+
+#pragma once
+
+#include "MVVM/ViewModels/TrackModel.h"
+
+class UMovieSceneCameraCutTrack;
+
+namespace UE
+{
+namespace Sequencer
+{
+
+class MOVIESCENETOOLS_API FCameraCutTrackModel
+	: public FTrackModel
+{
+public:
+
+	UE_SEQUENCER_DECLARE_CASTABLE(FCameraCutTrackModel, FTrackModel);
+
+	static TSharedPtr<FTrackModel> CreateTrackModel(UMovieSceneTrack* Track);
+
+	explicit FCameraCutTrackModel(UMovieSceneCameraCutTrack* Track);
+
+	FSortingKey GetSortingKey() const override;
+};
+
+} // namespace Sequencer
+} // namespace UE
+
diff --git a/Engine/Source/Editor/MovieSceneTools/Private/MVVM/ViewModels/CinematicShotTrackModel.h b/Engine/Source/Editor/MovieSceneTools/Private/MVVM/ViewModels/CinematicShotTrackModel.h
new file mode 100644
--- /dev/null
+++ b/Engine/Source/Editor/MovieSceneTools/Private/MVVM/ViewModels/CinematicShotTrackModel.h
@@ -0,0 +1,30 @@
+// Copyright Epic Games, Inc. All Rights Reserved.
+
+#pragma once
+
+#include "MVVM/ViewModels/TrackModel.h"
+
+class UMovieSceneCinematicShotTrack;
+
+namespace UE
+{
+namespace Sequencer
+{
+
+class MOVIESCENETOOLS_API FCinematicShotTrackModel
+	: public FTrackModel
+{
+public:
+
+	UE_SEQUENCER_DECLARE_CASTABLE(FCinematicShotTrackModel, FTrackModel);
+
+	static TSharedPtr<FTrackModel> CreateTrackModel(UMovieSceneTrack* Track);
+
+	explicit FCinematicShotTrackModel(UMovieSceneCinematicShotTrack* Track);
+
+	FSortingKey GetSortingKey() const override;
+};
+
+} // namespace Sequencer
+} // namespace UE
+
diff --git a/Engine/Source/Editor/MovieSceneTools/Private/MovieSceneObjectBindingIDCustomization.cpp b/Engine/Source/Editor/MovieSceneTools/Private/MovieSceneObjectBindingIDCustomization.cpp
--- a/Engine/Source/Editor/MovieSceneTools/Private/MovieSceneObjectBindingIDCustomization.cpp
+++ b/Engine/Source/Editor/MovieSceneTools/Private/MovieSceneObjectBindingIDCustomization.cpp
@@ -40,13 +40,22 @@
 
 void FMovieSceneObjectBindingIDCustomization::CustomizeHeader(TSharedRef<IPropertyHandle> PropertyHandle, FDetailWidgetRow& HeaderRow, IPropertyTypeCustomizationUtils& CustomizationUtils)
 {
+	using namespace UE::Sequencer;
+
 	StructProperty = PropertyHandle;
 
 	Initialize();
 
 	auto IsAcceptable = [](TSharedPtr<FDragDropOperation> Operation)
 	{
-		return Operation->IsOfType<FSequencerObjectBindingDragDropOp>() && static_cast<FSequencerObjectBindingDragDropOp*>(Operation.Get())->GetDraggedBindings().Num() == 1;
+		using namespace UE::Sequencer;
+		if (!Operation->IsOfType<FSequencerObjectBindingDragDropOp>())
+		{
+			return false;
+		}
+
+		FSequencerObjectBindingDragDropOp* DragDropOp = static_cast<FSequencerObjectBindingDragDropOp*>(Operation.Get());
+		return DragDropOp->GetDraggedRebindableBindings().Num() == 1;
 	};
 
 	HeaderRow
@@ -90,10 +99,12 @@
 
 FReply FMovieSceneObjectBindingIDCustomization::OnDrop(const FGeometry& InGeometry, const FDragDropEvent& InDragDropEvent)
 {
+	using namespace UE::Sequencer;
+
 	TSharedPtr<FSequencerObjectBindingDragDropOp> SequencerOp = InDragDropEvent.GetOperationAs<FSequencerObjectBindingDragDropOp>();
 	if (SequencerOp)
 	{
-		TArray<UE::MovieScene::FFixedObjectBindingID> Bindings = SequencerOp->GetDraggedBindings();
+		TArray<UE::MovieScene::FFixedObjectBindingID> Bindings = SequencerOp->GetDraggedRebindableBindings();
 		if (Bindings.Num() == 1)
 		{
 			SetBindingId(Bindings[0]);
diff --git a/Engine/Source/Editor/MovieSceneTools/Private/MovieSceneToolsModule.cpp b/Engine/Source/Editor/MovieSceneTools/Private/MovieSceneToolsModule.cpp
--- a/Engine/Source/Editor/MovieSceneTools/Private/MovieSceneToolsModule.cpp
+++ b/Engine/Source/Editor/MovieSceneTools/Private/MovieSceneToolsModule.cpp
@@ -54,6 +54,9 @@
 #include "TrackEditors/CameraShakeSourceShakeTrackEditor.h"
 #include "TrackEditors/CVarTrackEditor.h"
 
+#include "MVVM/ViewModels/CameraCutTrackModel.h"
+#include "MVVM/ViewModels/CinematicShotTrackModel.h"
+
 #include "MovieSceneBuiltInEasingFunctionCustomization.h"
 #include "MovieSceneObjectBindingIDCustomization.h"
 #include "MovieSceneEventCustomization.h"
@@ -98,6 +101,8 @@
 
 void FMovieSceneToolsModule::StartupModule()
 {
+	using namespace UE::Sequencer;
+
 	if (GIsEditor)
 	{
 		if (ISettingsModule* SettingsModule = FModuleManager::GetModulePtr<ISettingsModule>("Settings"))
@@ -152,6 +157,10 @@
 		CameraShakeSourceShakeCreateEditorHandle = SequencerModule.RegisterTrackEditor(FOnCreateTrackEditor::CreateStatic(&FCameraShakeSourceShakeTrackEditor::CreateTrackEditor));
 		CVarTrackCreateEditorHandle = SequencerModule.RegisterTrackEditor(FOnCreateTrackEditor::CreateStatic(&FCVarTrackEditor::CreateTrackEditor));
 
+		// register track models
+		CameraCutTrackModelHandle = SequencerModule.RegisterTrackModel(FOnCreateTrackModel::CreateStatic(&FCameraCutTrackModel::CreateTrackModel));
+		CinematicShotTrackModelHandle = SequencerModule.RegisterTrackModel(FOnCreateTrackModel::CreateStatic(&FCinematicShotTrackModel::CreateTrackModel));
+
 		RegisterClipboardConversions();
 
 		// register details customization
@@ -288,6 +297,10 @@
 	SequencerModule.UnRegisterTrackEditor( PrimitiveMaterialCreateEditorHandle );
 	SequencerModule.UnRegisterTrackEditor( CVarTrackCreateEditorHandle );
 
+	// unregister track models
+	SequencerModule.UnregisterTrackModel( CameraCutTrackModelHandle );
+	SequencerModule.UnregisterTrackModel( CinematicShotTrackModelHandle );
+
 	if (FModuleManager::Get().IsModuleLoaded("PropertyEditor"))
 	{	
 		FPropertyEditorModule& PropertyModule = FModuleManager::LoadModuleChecked<FPropertyEditorModule>("PropertyEditor");
diff --git a/Engine/Source/Editor/MovieSceneTools/Private/TrackEditors/FadeTrackEditor.cpp b/Engine/Source/Editor/MovieSceneTools/Private/TrackEditors/FadeTrackEditor.cpp
--- a/Engine/Source/Editor/MovieSceneTools/Private/TrackEditors/FadeTrackEditor.cpp
+++ b/Engine/Source/Editor/MovieSceneTools/Private/TrackEditors/FadeTrackEditor.cpp
@@ -160,7 +160,9 @@
 	UMovieSceneSection* NewSection = FadeTrack->CreateNewSection();
 	check(NewSection);
 
+	FadeTrack->Modify();
 	FadeTrack->AddSection(*NewSection);
+
 	if (GetSequencer().IsValid())
 	{
 		GetSequencer()->OnAddTrack(FadeTrack, FGuid());
diff --git a/Engine/Source/Editor/MovieSceneTools/Private/TrackEditors/SlomoTrackEditor.cpp b/Engine/Source/Editor/MovieSceneTools/Private/TrackEditors/SlomoTrackEditor.cpp
--- a/Engine/Source/Editor/MovieSceneTools/Private/TrackEditors/SlomoTrackEditor.cpp
+++ b/Engine/Source/Editor/MovieSceneTools/Private/TrackEditors/SlomoTrackEditor.cpp
@@ -92,6 +92,7 @@
 	UMovieSceneSection* NewSection = SlomoTrack->CreateNewSection();
 	check(NewSection);
 
+	SlomoTrack->Modify();
 	SlomoTrack->AddSection(*NewSection);
 
 	if (GetSequencer().IsValid())
diff --git a/Engine/Source/Editor/MovieSceneTools/Public/CommonMovieSceneTools.h b/Engine/Source/Editor/MovieSceneTools/Public/CommonMovieSceneTools.h
--- a/Engine/Source/Editor/MovieSceneTools/Public/CommonMovieSceneTools.h
+++ b/Engine/Source/Editor/MovieSceneTools/Public/CommonMovieSceneTools.h
@@ -1,133 +1,23 @@
-// Copyright Epic Games, Inc. All Rights Reserved.
-
-#pragma once
-
-#include "CoreTypes.h"
-#include "FrameNumberNumericInterface.h"
-#include "Layout/Geometry.h"
-#include "Misc/FrameRate.h"
-
-class FSequencerSectionPainter;
-
-/**
- * Utility for converting time units to slate pixel units and vice versa
- */
-struct FTimeToPixel
-{
-public:
-	FTimeToPixel( const FGeometry& AllottedGeometry, const TRange<double>& InLocalViewRange, const FFrameRate& InTickResolution )
-		: ViewRangeStartSeconds( InLocalViewRange.GetLowerBoundValue() )
-		, TickResolution( InTickResolution )
-	{
-		double VisibleWidth = InLocalViewRange.Size<double>();
-
-		const float MaxPixelsPerSecond = 1000.f;
-		PixelsPerSecond = VisibleWidth > 0 ? AllottedGeometry.GetLocalSize().X / VisibleWidth : MaxPixelsPerSecond;
-	}
-
-	/**
-	 * Converts a time to a pixel point relative to the geometry of a widget (passed into the constructor)
-	 *
-	 * @param Time	The time to convert
-	 * @return The pixel equivalent of the time
-	 */
-	float SecondsToPixel( double Time ) const
-	{
-		return (Time - ViewRangeStartSeconds) * PixelsPerSecond;
-	}
-
-	/**
-	 * Converts a time delta to a pixel delta
-	 *
-	 * @param TimeDelta  The time delta to convert
-	 * @return           The pixel equivalent of the delta time
-	 */
-	float SecondsDeltaToPixel( double TimeDelta ) const
-	{
-		return TimeDelta * PixelsPerSecond;
-	}
-
-	/**
-	 * Converts a pixel value to time 
-	 *
-	 * @param PixelX The x value of a pixel coordinate relative to the geometry that was passed into the constructor.
-	 * @return The time where the pixel is located
-	 */
-	double PixelToSeconds( float PixelX ) const
-	{
-		return (PixelX/PixelsPerSecond) + ViewRangeStartSeconds;
-	}
-
-	/**
-	 * Converts a frame time to a pixel point relative to the geometry of a widget (passed into the constructor)
-	 *
-	 * @param Time The time to convert
-	 * @return The pixel equivalent of the frame time
-	 */
-	float FrameToPixel( const FFrameTime& Time ) const
-	{
-		return (Time / TickResolution - ViewRangeStartSeconds) * PixelsPerSecond;
-	}
-
-	/**
-	 * Converts a frame delta value to pixel delta
-	 *
-	 * @param TimeDelta   The time delta to convert
-	 * @return            The pixel equivalent of the delta time
-	 */
-	float FrameDeltaToPixel( const FFrameTime& TimeDelta) const
-	{
-		return (TimeDelta / TickResolution) * PixelsPerSecond;
-	}
-
-	/**
-	 * Converts a pixel value to frame time 
-	 *
-	 * @param PixelX The x value of a pixel coordinate relative to the geometry that was passed into the constructor.
-	 * @return The frame time where the pixel is located
-	 */
-	FFrameTime PixelToFrame( float PixelX ) const
-	{
-		return ( PixelX / PixelsPerSecond + ViewRangeStartSeconds ) * TickResolution;
-	}
-
-	/**
-	 * Converts a pixel delta value to delta frame time 
-	 *
-	 * @param PixelDelta The delta value in pixel space
-	 * @return The equivalent delta frame time
-	 */
-	FFrameTime PixelDeltaToFrame( float PixelDelta ) const
-	{
-		return ( PixelDelta / PixelsPerSecond ) * TickResolution;
-	}
-
-	/**
-	 * Retrieve the tick resolution of the current sequence
-	 */
-	FFrameRate GetTickResolution() const
-	{
-		return TickResolution;
-	}
-
-private:
-
-	/** time range of the sequencer in seconds */
-	double ViewRangeStartSeconds;
-	/** The tick resolution of the current timeline */
-	FFrameRate TickResolution;
-	/** The number of pixels in the view range */
-	float PixelsPerSecond;
-};
-
-
-/**
- * Draw a frame time next to the scrub handle
- *
- * @param InPainter Structure that affords common painting operations
- * @param CurrentTime Current time of the scrub handle
- * @param FrameTime Frame time to draw
- * @param FrameNumberInterface (optional) Interface to control the display format and/or frame rate conversion of the drawn frame time.
- *            If not provided, the frame time will be drawn as a frame number without any subframe.
- */
-void DrawFrameTimeHint(FSequencerSectionPainter& InPainter, const FFrameTime& CurrentTime, const FFrameTime& FrameTime, const FFrameNumberInterface* FrameNumberInterface = nullptr);
+// Copyright Epic Games, Inc. All Rights Reserved.
+
+#pragma once
+
+#include "CoreTypes.h"
+#include "FrameNumberNumericInterface.h"
+#include "Layout/Geometry.h"
+#include "Misc/FrameRate.h"
+#include "SequencerSectionPainter.h"
+#include "TimeToPixel.h"
+
+class FSequencerSectionPainter;
+
+/**
+ * Draw a frame time next to the scrub handle
+ *
+ * @param InPainter Structure that affords common painting operations
+ * @param CurrentTime Current time of the scrub handle
+ * @param FrameTime Frame time to draw
+ * @param FrameNumberInterface (optional) Interface to control the display format and/or frame rate conversion of the drawn frame time.
+ *            If not provided, the frame time will be drawn as a frame number without any subframe.
+ */
+void DrawFrameTimeHint(FSequencerSectionPainter& InPainter, const FFrameTime& CurrentTime, const FFrameTime& FrameTime, const FFrameNumberInterface* FrameNumberInterface = nullptr);
diff --git a/Engine/Source/Editor/MovieSceneTools/Public/MovieSceneToolsModule.h b/Engine/Source/Editor/MovieSceneTools/Public/MovieSceneToolsModule.h
--- a/Engine/Source/Editor/MovieSceneTools/Public/MovieSceneToolsModule.h
+++ b/Engine/Source/Editor/MovieSceneTools/Public/MovieSceneToolsModule.h
@@ -124,6 +124,9 @@
 	FDelegateHandle CameraShakeSourceShakeCreateEditorHandle;
 	FDelegateHandle CVarTrackCreateEditorHandle;
 
+	FDelegateHandle CameraCutTrackModelHandle;
+	FDelegateHandle CinematicShotTrackModelHandle;
+
 	FDelegateHandle GenerateEventEntryPointsHandle;
 	FDelegateHandle FixupPayloadParameterNameHandle;
 	FDelegateHandle UpgradeLegacyEventEndpointHandle;
diff --git a/Engine/Source/Editor/SequenceRecorder/Private/ActorRecording.cpp b/Engine/Source/Editor/SequenceRecorder/Private/ActorRecording.cpp
--- a/Engine/Source/Editor/SequenceRecorder/Private/ActorRecording.cpp
+++ b/Engine/Source/Editor/SequenceRecorder/Private/ActorRecording.cpp
@@ -373,7 +373,7 @@
 	{
 		FolderToUse = NewObject<UMovieSceneFolder>(MovieScene, NAME_None, RF_Transactional);
 		FolderToUse->SetFolderName(FolderName);
-		MovieScene->GetRootFolders().Add(FolderToUse);
+		MovieScene->AddRootFolder(FolderToUse);
 	}
 
 	FolderToUse->AddChildObjectBinding(Guid);
@@ -641,7 +641,7 @@
 	FMovieScenePossessable* ChildPossessable = OwnerMovieScene->FindPossessable(PossessableGuid);
 	if (ensure(ChildPossessable))
 	{
-		ChildPossessable->SetParent(Guid);
+		ChildPossessable->SetParent(Guid, OwnerMovieScene);
 	}
 
 	FMovieSceneSpawnable* ParentSpawnable = OwnerMovieScene->FindSpawnable(Guid);
diff --git a/Engine/Source/Editor/Sequencer/Private/DisplayNodes/SequencerKeyTimeCache.cpp b/Engine/Source/Editor/Sequencer/Private/DisplayNodes/SequencerKeyTimeCache.cpp
--- a/Engine/Source/Editor/Sequencer/Private/DisplayNodes/SequencerKeyTimeCache.cpp
+++ b/Engine/Source/Editor/Sequencer/Private/DisplayNodes/SequencerKeyTimeCache.cpp
@@ -1,13 +1,22 @@
 // Copyright Epic Games, Inc. All Rights Reserved.
 
 #include "SequencerKeyTimeCache.h"
+#include "MVVM/ViewModels/ChannelModel.h"
 #include "MovieSceneSection.h"
 #include "Algo/Sort.h"
 #include "Algo/BinarySearch.h"
 
 bool FSequencerCachedKeys::Update(FFrameRate SourceResolution)
 {
-	UMovieSceneSection* Section = KeyArea->GetOwningSection();
+	using namespace UE::Sequencer;
+
+	TSharedPtr<FChannelModel> Channel = WeakChannel.Pin();
+	if (!Channel)
+	{
+		return false;
+	}
+
+	UMovieSceneSection* Section = Channel->GetSection();
 	if (!Section || !CachedSignature.IsValid() || Section->GetSignature() != CachedSignature || SourceResolution != CachedTickResolution)
 	{
 		CachedSignature = Section ? Section->GetSignature() : FGuid();
@@ -16,7 +25,7 @@
 		CachedKeyFrames.Reset();
 
 		TArray<FKeyHandle> Handles;
-		KeyArea->GetKeyInfo(&Handles, &CachedKeyFrames);
+		Channel->GetKeyArea()->GetKeyInfo(&Handles, &CachedKeyFrames);
 
 		CachedKeyTimes.Reset(CachedKeyFrames.Num());
 		CachedKeyHandles.Reset(CachedKeyFrames.Num());
diff --git a/Engine/Source/Editor/Sequencer/Private/DisplayNodes/SequencerKeyTimeCache.h b/Engine/Source/Editor/Sequencer/Private/DisplayNodes/SequencerKeyTimeCache.h
--- a/Engine/Source/Editor/Sequencer/Private/DisplayNodes/SequencerKeyTimeCache.h
+++ b/Engine/Source/Editor/Sequencer/Private/DisplayNodes/SequencerKeyTimeCache.h
@@ -12,11 +12,19 @@
 #include "IKeyArea.h"
 #include "Misc/FrameRate.h"
 
+namespace UE
+{
+namespace Sequencer
+{
+class FChannelModel;
+}
+}
+
 /** Simple structure that caches the sorted key times for a given key area */
 struct FSequencerCachedKeys
 {
-	FSequencerCachedKeys(TSharedRef<IKeyArea> InKeyArea)
-		: KeyArea(InKeyArea)
+	FSequencerCachedKeys(TSharedPtr<UE::Sequencer::FChannelModel> InChannel)
+		: WeakChannel(InChannel)
 	{}
 
 	/**
@@ -29,9 +37,9 @@
 	void GetKeysInRange(const TRange<double>& Range, TArrayView<const double>* OutTimes, TArrayView<const FFrameNumber>* OutKeyFrames, TArrayView<const FKeyHandle>* OutHandles) const;
 
 	/** Get the key area this cache was generated for, or nullptr if the cache has never been updated */
-	const TSharedPtr<IKeyArea>& GetKeyArea() const
+	TSharedPtr<UE::Sequencer::FChannelModel> GetChannel() const
 	{
-		return KeyArea;
+		return WeakChannel.Pin();
 	}
 
 private:
@@ -47,6 +55,6 @@
 	FFrameRate CachedTickResolution;
 
 	/** The key area this cache is for */
-	TSharedPtr<IKeyArea> KeyArea;
+	TWeakPtr<UE::Sequencer::FChannelModel> WeakChannel;
 };
 
diff --git a/Engine/Source/Editor/Sequencer/Private/IKeyArea.cpp b/Engine/Source/Editor/Sequencer/Private/IKeyArea.cpp
--- a/Engine/Source/Editor/Sequencer/Private/IKeyArea.cpp
+++ b/Engine/Source/Editor/Sequencer/Private/IKeyArea.cpp
@@ -206,21 +206,19 @@
 	}
 }
 
-void IKeyArea::PasteKeys(const FMovieSceneClipboardKeyTrack& KeyTrack, const FMovieSceneClipboardEnvironment& SrcEnvironment, const FSequencerPasteEnvironment& DstEnvironment)
+TArray<FKeyHandle> IKeyArea::PasteKeys(const FMovieSceneClipboardKeyTrack& KeyTrack, const FMovieSceneClipboardEnvironment& SrcEnvironment, const FSequencerPasteEnvironment& DstEnvironment)
 {
+	TArray<FKeyHandle> PastedKeys;
+
 	ISequencerChannelInterface* EditorInterface = FindChannelEditorInterface();
 	FMovieSceneChannel* Channel = ChannelHandle.Get();
 	UMovieSceneSection* OwningSection = GetOwningSection();
 	if (EditorInterface && Channel && OwningSection)
 	{
-		TArray<FKeyHandle> PastedKeys;
 		EditorInterface->PasteKeys_Raw(Channel, OwningSection, KeyTrack, SrcEnvironment, DstEnvironment, PastedKeys);
+	}
 
-		for (FKeyHandle KeyHandle : PastedKeys)
-		{
-			DstEnvironment.ReportPastedKey(KeyHandle, *this);
-		}
-	}
+	return PastedKeys;
 }
 
 FText GetOwningObjectBindingName(UMovieSceneTrack* InTrack, ISequencer& InSequencer)
diff --git a/Engine/Source/Editor/Sequencer/Private/MVVM/CurveEditorIntegrationExtension.cpp b/Engine/Source/Editor/Sequencer/Private/MVVM/CurveEditorIntegrationExtension.cpp
new file mode 100644
--- /dev/null
+++ b/Engine/Source/Editor/Sequencer/Private/MVVM/CurveEditorIntegrationExtension.cpp
@@ -0,0 +1,199 @@
+// Copyright Epic Games, Inc. All Rights Reserved.
+
+#include "MVVM/CurveEditorIntegrationExtension.h"
+#include "MVVM/TrackModelStorageExtension.h"
+#include "MVVM/ViewModels/ViewModel.h"
+#include "MVVM/SharedViewModelData.h"
+#include "MVVM/ViewModelPtr.h"
+#include "MVVM/ViewModels/SequenceModel.h"
+#include "MVVM/ViewModels/ViewModelIterators.h"
+#include "MVVM/Extensions/ICurveEditorTreeItemExtension.h"
+#include "CurveEditor.h"
+
+namespace UE
+{
+namespace Sequencer
+{
+
+FCurveEditorIntegrationExtension::FCurveEditorIntegrationExtension()
+{
+}
+
+void FCurveEditorIntegrationExtension::OnCreated(TSharedRef<FViewModel> InWeakOwner)
+{
+	ensureMsgf(!WeakOwnerModel.Pin().IsValid(), TEXT("This extension was already created!"));
+	WeakOwnerModel = InWeakOwner->CastThisShared<FSequenceModel>();
+
+	FSimpleMulticastDelegate& HierarchyChanged = InWeakOwner->GetSharedData()->SubscribeToHierarchyChanged(InWeakOwner);
+	HierarchyChanged.AddSP(this, &FCurveEditorIntegrationExtension::OnHierarchyChanged);
+}
+
+void FCurveEditorIntegrationExtension::OnHierarchyChanged()
+{
+	UpdateCurveEditor();
+}
+
+void FCurveEditorIntegrationExtension::UpdateCurveEditor()
+{
+	TSharedPtr<FSequenceModel> OwnerModel = WeakOwnerModel.Pin();
+	if (!OwnerModel)
+	{
+		return;
+	}
+
+	ICurveEditorExtension* CurveEditorExtension = OwnerModel->CastThis<ICurveEditorExtension>();
+	if (!CurveEditorExtension)
+	{
+		return;
+	}
+
+	TSharedPtr<FCurveEditor> CurveEditor = CurveEditorExtension->GetCurveEditor();
+	if (!CurveEditor)
+	{
+		return;
+	}
+
+	FCurveEditorTree* CurveEditorTree = CurveEditor->GetTree();
+
+	// Guard against multiple broadcasts here and defer them until the end of this function
+	FScopedCurveEditorTreeEventGuard ScopedEventGuard = CurveEditorTree->ScopedEventGuard();
+
+	// Remove any stale tree items
+	for (auto It = ViewModelToTreeItemIDMap.CreateIterator(); It; ++It)
+	{
+		bool bIsRelevant = false;
+		const FCurveEditorTreeItemID ItemID(It.Value());
+		TViewModelPtr<ICurveEditorTreeItemExtension> ViewModel(It.Key().Pin());
+		if (ViewModel)
+		{
+			TViewModelPtr<IOutlinerExtension> OutlinerModel = ViewModel.ImplicitCast();
+			const bool bIsVisible = OutlinerModel && !OutlinerModel->IsFilteredOut();
+			
+			const FCurveEditorTreeItem* TreeItem = CurveEditorTree->FindItem(ItemID);
+
+			FCurveEditorTreeItemID ParentID;
+			if (auto Parent = ViewModel.AsModel()->FindAncestorOfType<ICurveEditorTreeItemExtension>())
+			{
+				ParentID = ViewModelToTreeItemIDMap.FindRef(Parent.AsModel());
+			}
+			
+			bIsRelevant = bIsVisible && TreeItem && TreeItem->GetParentID() == ParentID;
+		}
+		
+		if (!bIsRelevant)
+		{
+			if (ViewModel)
+			{
+				ViewModel->OnRemovedFromCurveEditor(CurveEditor);
+			}
+			CurveEditor->RemoveTreeItem(ItemID);
+			It.RemoveCurrent();
+		}
+	}
+
+	// Do a second pass to remove any items that were removed recursively above
+	for (auto It = ViewModelToTreeItemIDMap.CreateIterator(); It; ++It)
+	{
+		if (CurveEditorTree->FindItem(It->Value) == nullptr)
+		{
+			It.RemoveCurrent();
+		}
+	}
+
+	// Iterate all non-filtered out outliners and check for curve editor tree extensions
+	const bool bIncludeRootNode = false;
+	for (TParentFirstChildIterator<IOutlinerExtension> It(OwnerModel, bIncludeRootNode); It; ++It)
+	{
+		if (It->IsFilteredOut())
+		{
+			It.IgnoreCurrentChildren();
+			continue;
+		}
+
+		if (TViewModelPtr<ICurveEditorTreeItemExtension> ChildViewModel = (*It).ImplicitCast())
+		{
+			const FCurveEditorTreeItemID ItemID = ViewModelToTreeItemIDMap.FindRef(ChildViewModel.AsModel());
+			if (!ItemID.IsValid())
+			{
+				AddToCurveEditor(ChildViewModel.AsModel(), CurveEditor);
+			}
+		}
+	}
+}
+
+FCurveEditorTreeItemID FCurveEditorIntegrationExtension::AddToCurveEditor(TViewModelPtr<ICurveEditorTreeItemExtension> InViewModel, TSharedPtr<FCurveEditor> CurveEditor)
+{
+	// If the view model doesn't want to be in the curve editor, bail out
+	// Note that this means we will create curve editor items for each parent in the hierarchy up
+	// until the first parent that doesn't implement ICurveEditorTreeItemExtension
+	// That is: we don't create "dummy" entries when there's a "gap" in the hierarchy
+	if (!InViewModel)
+	{
+		return FCurveEditorTreeItemID::Invalid();
+	}
+
+	// Check if we already have a valid curve editor ID
+	if (FCurveEditorTreeItemID* Existing = ViewModelToTreeItemIDMap.Find(InViewModel.AsModel()))
+	{
+		if (CurveEditor->GetTree()->FindItem(*Existing) != nullptr)
+		{
+			return *Existing;
+		}
+	}
+
+	// Recursively create any needed parent curve editor items
+	TSharedPtr<FViewModel> Parent = InViewModel.AsModel()->GetParent();
+
+	FCurveEditorTreeItemID ParentID = Parent.IsValid() ? 
+		AddToCurveEditor(Parent, CurveEditor) : 
+		FCurveEditorTreeItemID::Invalid();
+
+	// Create the new curve editor item
+	FCurveEditorTreeItem* NewItem = CurveEditor->AddTreeItem(ParentID);
+	TSharedPtr<ICurveEditorTreeItem> CurveEditorTreeItem = InViewModel->GetCurveEditorTreeItem();
+	NewItem->SetWeakItem(CurveEditorTreeItem);
+
+	// Register the new ID in our map and notify the view model
+	ViewModelToTreeItemIDMap.Add(InViewModel.AsModel(), NewItem->GetID());
+	InViewModel->OnAddedToCurveEditor(NewItem->GetID(), CurveEditor);
+
+	return NewItem->GetID();
+}
+
+void FCurveEditorIntegrationExtension::RecreateCurveEditor()
+{
+	TSharedPtr<FSequenceModel> OwnerModel = WeakOwnerModel.Pin();
+	if (!OwnerModel)
+	{
+		return;
+	}
+
+	ICurveEditorExtension* CurveEditorExtension = OwnerModel->CastThis<ICurveEditorExtension>();
+	if (!CurveEditorExtension)
+	{
+		return;
+	}
+
+	TSharedPtr<FCurveEditor> CurveEditor = CurveEditorExtension->GetCurveEditor();
+	if (!CurveEditor)
+	{
+		return;
+	}
+
+	for (auto Pair : ViewModelToTreeItemIDMap)
+	{
+		FCurveEditorTreeItemID ItemID(Pair.Value);
+		TViewModelPtr<ICurveEditorTreeItemExtension> ViewModel(Pair.Key.Pin());
+		if (ViewModel)
+		{
+			ViewModel->OnRemovedFromCurveEditor(CurveEditor);
+		}
+		CurveEditor->RemoveTreeItem(ItemID);
+	}
+
+	ViewModelToTreeItemIDMap.Reset();
+}
+
+} // namespace Sequencer
+} // namespace UE
+
diff --git a/Engine/Source/Editor/Sequencer/Private/MVVM/ViewModels/ObjectBindingModel.cpp b/Engine/Source/Editor/Sequencer/Private/MVVM/ViewModels/ObjectBindingModel.cpp
new file mode 100644
--- /dev/null
+++ b/Engine/Source/Editor/Sequencer/Private/MVVM/ViewModels/ObjectBindingModel.cpp
@@ -0,0 +1,1431 @@
+// Copyright Epic Games, Inc. All Rights Reserved.
+
+#include "MVVM/ViewModels/ObjectBindingModel.h"
+#include "MVVM/ViewModels/FolderModel.h"
+#include "MVVM/ViewModels/SequenceModel.h"
+#include "MVVM/ViewModels/TrackModel.h"
+#include "MVVM/ViewModels/LayerBarModel.h"
+#include "MVVM/ViewModels/ViewModelIterators.h"
+#include "MVVM/TrackModelStorageExtension.h"
+#include "MVVM/ObjectBindingModelStorageExtension.h"
+#include "MVVM/ViewModels/SequencerEditorViewModel.h"
+#include "MVVM/ViewModels/OutlinerViewModelDragDropOp.h"
+#include "MVVM/Views/SOutlinerObjectBindingView.h"
+#include "Modules/ModuleManager.h"
+#include "ISequencerModule.h"
+#include "MovieScene.h"
+#include "MovieSceneFolder.h"
+#include "MovieSceneBinding.h"
+#include "PropertyPath.h"
+#include "ScopedTransaction.h"
+#include "Sequencer.h"
+#include "SequencerNodeTree.h"
+#include "SequencerCommands.h"
+#include "SequencerUtilities.h"
+#include "ISequencerTrackEditor.h"
+#include "SObjectBindingTag.h"
+#include "Containers/ArrayBuilder.h"
+#include "ObjectBindingTagCache.h"
+#include "Styling/AppStyle.h"
+#include "Styling/SlateIconFinder.h"
+#include "Tracks/MovieSceneSpawnTrack.h"
+#include "ClassViewerModule.h"
+#include "Algo/Sort.h"
+#include "Engine/LevelStreaming.h"
+#include "Framework/MultiBox/MultiBoxBuilder.h"
+
+#define LOCTEXT_NAMESPACE "ObjectBindingModel"
+
+namespace UE
+{
+namespace Sequencer
+{
+
+namespace
+{
+
+struct PropertyMenuData
+{
+	FString MenuName;
+	FPropertyPath PropertyPath;
+};
+
+void GetKeyablePropertyPaths(UClass* Class, void* ValuePtr, UStruct* PropertySource, FPropertyPath PropertyPath, FSequencer& Sequencer, TArray<FPropertyPath>& KeyablePropertyPaths)
+{
+	//@todo need to resolve this between UMG and the level editor sequencer
+	const bool bRecurseAllProperties = Sequencer.IsLevelEditorSequencer();
+
+	for (TFieldIterator<FProperty> PropertyIterator(PropertySource); PropertyIterator; ++PropertyIterator)
+	{
+		FProperty* Property = *PropertyIterator;
+
+		if (Property && !Property->HasAnyPropertyFlags(CPF_Deprecated))
+		{
+			PropertyPath.AddProperty(FPropertyInfo(Property));
+
+			bool bIsPropertyKeyable = Sequencer.CanKeyProperty(FCanKeyPropertyParams(Class, PropertyPath));
+			if (bIsPropertyKeyable)
+			{
+				KeyablePropertyPaths.Add(PropertyPath);
+			}
+
+			FArrayProperty* ArrayProperty = CastField<FArrayProperty>(Property);
+			if (!bIsPropertyKeyable && ArrayProperty)
+			{
+				FScriptArrayHelper ArrayHelper(ArrayProperty, ArrayProperty->ContainerPtrToValuePtr<void>(ValuePtr));
+				for (int32 Index = 0; Index < ArrayHelper.Num(); ++Index)
+				{
+					PropertyPath.AddProperty(FPropertyInfo(ArrayProperty->Inner, Index));
+
+					if (Sequencer.CanKeyProperty(FCanKeyPropertyParams(Class, PropertyPath)))
+					{
+						KeyablePropertyPaths.Add(PropertyPath);
+						bIsPropertyKeyable = true;
+					}
+					else if (FStructProperty* StructProperty = CastField<FStructProperty>(ArrayProperty->Inner))
+					{
+						GetKeyablePropertyPaths(Class, ArrayHelper.GetRawPtr(Index), StructProperty->Struct, PropertyPath, Sequencer, KeyablePropertyPaths);
+					}
+
+					PropertyPath = *PropertyPath.TrimPath(1);
+				}
+			}
+
+			if (!bIsPropertyKeyable || bRecurseAllProperties)
+			{
+				if (FStructProperty* StructProperty = CastField<FStructProperty>(Property))
+				{
+					GetKeyablePropertyPaths(Class, StructProperty->ContainerPtrToValuePtr<void>(ValuePtr), StructProperty->Struct, PropertyPath, Sequencer, KeyablePropertyPaths);
+				}
+			}
+
+			PropertyPath = *PropertyPath.TrimPath(1);
+		}
+	}
+}
+
+struct FMovieSceneSpawnableFlagCheckState
+{
+	FSequencer* Sequencer;
+	UMovieScene* MovieScene;
+	bool FMovieSceneSpawnable::*PtrToFlag;
+
+	ECheckBoxState operator()() const
+	{
+		using namespace UE::Sequencer;
+
+		ECheckBoxState CheckState = ECheckBoxState::Undetermined;
+		for (TWeakPtr<FViewModel> WeakItem : Sequencer->GetSelection().GetSelectedOutlinerItems())
+		{
+			if (IObjectBindingExtension* ObjectBindingID = ICastable::CastWeakPtr<IObjectBindingExtension>(WeakItem))
+			{
+				FMovieSceneSpawnable* SelectedSpawnable = MovieScene->FindSpawnable(ObjectBindingID->GetObjectGuid());
+				if (SelectedSpawnable)
+				{
+					if (CheckState != ECheckBoxState::Undetermined && SelectedSpawnable->*PtrToFlag != ( CheckState == ECheckBoxState::Checked ))
+					{
+						return ECheckBoxState::Undetermined;
+					}
+					CheckState = SelectedSpawnable->*PtrToFlag ? ECheckBoxState::Checked : ECheckBoxState::Unchecked;
+				}
+			}
+		}
+		return CheckState;
+	}
+};
+
+struct FMovieSceneSpawnableFlagToggler
+{
+	FSequencer* Sequencer;
+	UMovieScene* MovieScene;
+	bool FMovieSceneSpawnable::*PtrToFlag;
+	FText TransactionText;
+
+	void operator()() const
+	{
+		using namespace UE::Sequencer;
+
+		FScopedTransaction Transaction(TransactionText);
+
+		const ECheckBoxState CheckState = FMovieSceneSpawnableFlagCheckState{Sequencer, MovieScene, PtrToFlag}();
+
+		MovieScene->Modify();
+		for (TWeakPtr<FViewModel> WeakItem : Sequencer->GetSelection().GetSelectedOutlinerItems())
+		{
+			if (IObjectBindingExtension* ObjectBinding = ICastable::CastWeakPtr<IObjectBindingExtension>(WeakItem))
+			{
+				FMovieSceneSpawnable* SelectedSpawnable = MovieScene->FindSpawnable(ObjectBinding->GetObjectGuid());
+				if (SelectedSpawnable)
+				{
+					SelectedSpawnable->*PtrToFlag = (CheckState == ECheckBoxState::Unchecked);
+				}
+			}
+		}
+	}
+};
+
+} // anon-namespace
+
+FObjectBindingModel::FObjectBindingModel(FSequenceModel* InOwnerModel, const FMovieSceneBinding& InBinding)
+	: ObjectBindingID(InBinding.GetObjectGuid())
+	, TrackAreaList(EViewModelListType::TrackArea)
+	, OwnerModel(InOwnerModel)
+{
+	RegisterChildList(&TrackAreaList);
+
+	// Objects are expanded by default
+	SetExpansion(true);
+
+	SetIdentifier(*ObjectBindingID.ToString());
+}
+
+FObjectBindingModel::~FObjectBindingModel()
+{
+}
+
+void FObjectBindingModel::OnConstruct()
+{
+	if (!LayerBar)
+	{
+		LayerBar = MakeShared<FLayerBarModel>(AsShared());
+		LayerBar->SetLinkedOutlinerItem(AsShared());
+
+		GetChildrenForList(&TrackAreaList).AddChild(LayerBar);
+	}
+
+	UMovieScene* MovieScene = OwnerModel->GetMovieScene();
+	check(MovieScene);
+
+	FMovieSceneBinding* Binding = MovieScene->FindBinding(ObjectBindingID);
+	check(Binding);
+
+	FScopedViewModelListHead RecycledHead(AsShared(), EViewModelListType::Recycled);
+	GetChildrenForList(&OutlinerChildList).MoveChildrenTo(RecycledHead.GetChildren());
+
+	for (UMovieSceneTrack* Track : Binding->GetTracks())
+	{
+		AddTrack(Track);
+	}
+}
+
+void FObjectBindingModel::SetParentBindingID(const FGuid& InObjectBindingID)
+{
+	ParentObjectBindingID = InObjectBindingID;
+}
+
+FGuid FObjectBindingModel::GetDesiredParentBinding() const
+{
+	return ParentObjectBindingID;
+}
+
+EObjectBindingType FObjectBindingModel::GetType() const
+{
+	return EObjectBindingType::Unknown;
+}
+
+const UClass* FObjectBindingModel::FindObjectClass() const
+{
+	return UObject::StaticClass();
+}
+
+bool FObjectBindingModel::SupportsRebinding() const
+{
+	return true;
+}
+
+FTrackAreaParameters FObjectBindingModel::GetTrackAreaParameters() const
+{
+	FTrackAreaParameters Params;
+	Params.LaneType = ETrackAreaLaneType::Inline;
+	return Params;
+}
+
+FViewModelVariantIterator FObjectBindingModel::GetTrackAreaModelList() const
+{
+	return &TrackAreaList;
+}
+
+void FObjectBindingModel::AddTrack(UMovieSceneTrack* Track)
+{
+	FTrackModelStorageExtension* TrackStorage = OwnerModel->CastDynamic<FTrackModelStorageExtension>();
+
+	TSharedPtr<FTrackModel> TrackModel = TrackStorage->CreateModelForTrack(Track, AsShared());
+
+	GetChildrenForList(&OutlinerChildList).AddChild(TrackModel);
+}
+
+void FObjectBindingModel::RemoveTrack(UMovieSceneTrack* Track)
+{
+	FTrackModelStorageExtension* TrackStorage = OwnerModel->CastDynamic<FTrackModelStorageExtension>();
+
+	TSharedPtr<FTrackModel> TrackModel = GetChildrenOfType<FTrackModel>().FindBy(Track, &FTrackModel::GetTrack);
+	if (TrackModel)
+	{
+		TrackModel->RemoveFromParent();
+	}
+}
+
+FGuid FObjectBindingModel::GetObjectGuid() const
+{
+	return ObjectBindingID;
+}
+
+FOutlinerSizing FObjectBindingModel::GetOutlinerSizing() const
+{
+	return FOutlinerSizing(20.f, 4.f);
+}
+
+void FObjectBindingModel::GetIdentifierForGrouping(TStringBuilder<128>& OutString) const
+{
+	FOutlinerItemModel::GetIdentifier().ToString(OutString);
+}
+
+TSharedRef<SWidget> FObjectBindingModel::CreateOutlinerView(const FCreateOutlinerViewParams& InParams)
+{
+	TSharedPtr<FSequencerEditorViewModel> EditorViewModel = GetEditor();
+	TSharedPtr<FSequencer> Sequencer = EditorViewModel->GetSequencerImpl();
+
+	const FMovieSceneSequenceID SequenceID = OwnerModel->GetSequenceID();
+	const MovieScene::FFixedObjectBindingID FixedObjectBindingID(ObjectBindingID, SequenceID);
+
+	return SNew(SOutlinerObjectBindingView, SharedThis(this), EditorViewModel, InParams.TreeViewRow)
+		.AdditionalLabelContent()
+		[
+			SNew(SObjectBindingTags, FixedObjectBindingID, Sequencer->GetObjectBindingTagCache())
+		];
+}
+
+bool FObjectBindingModel::GetDefaultExpansionState() const
+{
+	// Object binding nodes are always expanded by default
+	return true;
+}
+
+bool FObjectBindingModel::CanRename() const
+{
+	return true;
+}
+
+void FObjectBindingModel::Rename(const FText& NewName)
+{
+	TSharedPtr<ISequencer> Sequencer = OwnerModel->GetSequencer();
+	UMovieSceneSequence* MovieSceneSequence = OwnerModel->GetSequence();
+
+	if (MovieSceneSequence && Sequencer)
+	{
+		UMovieScene* MovieScene = MovieSceneSequence->GetMovieScene();
+
+		FScopedTransaction Transaction(LOCTEXT("SetTrackName", "Set Track Name"));
+
+		// Modify the movie scene so that it gets marked dirty and renames are saved consistently.
+		MovieScene->Modify();
+
+		FMovieSceneSpawnable*   Spawnable   = MovieScene->FindSpawnable(ObjectBindingID);
+		FMovieScenePossessable* Possessable = MovieScene->FindPossessable(ObjectBindingID);
+
+		// If there is only one binding, set the name of the bound actor
+		TArrayView<TWeakObjectPtr<>> Objects = Sequencer->FindObjectsInCurrentSequence(ObjectBindingID);
+		if (Objects.Num() == 1)
+		{
+			if (AActor* Actor = Cast<AActor>(Objects[0].Get()))
+			{
+				Actor->SetActorLabel(NewName.ToString());
+			}
+		}
+
+		if (Spawnable)
+		{
+			// Otherwise set our display name
+			Spawnable->SetName(NewName.ToString());
+		}
+		else if (Possessable)
+		{
+			Possessable->SetName(NewName.ToString());
+		}
+		else
+		{
+			MovieScene->SetObjectDisplayName(ObjectBindingID, NewName);
+		}
+
+		SetIdentifier(FName(*NewName.ToString()));
+	}
+}
+
+FText FObjectBindingModel::GetLabel() const
+{
+	UMovieSceneSequence* MovieSceneSequence = OwnerModel->GetSequence();
+	if (MovieSceneSequence != nullptr)
+	{
+		return MovieSceneSequence->GetMovieScene()->GetObjectDisplayName(ObjectBindingID);
+	}
+
+	return FText();
+}
+
+FSlateColor FObjectBindingModel::GetLabelColor() const
+{
+	TSharedPtr<ISequencer> Sequencer = OwnerModel->GetSequencer();
+
+	if (!Sequencer)
+	{
+		return FLinearColor::Red;
+	}
+
+	TArrayView<TWeakObjectPtr<> > BoundObjects = Sequencer->FindBoundObjects(ObjectBindingID, OwnerModel->GetSequenceID());
+
+	if (BoundObjects.Num() > 0)
+	{
+		int32 NumValidObjects = 0;
+		for (const TWeakObjectPtr<>& BoundObject : BoundObjects)
+		{
+			if (BoundObject.IsValid())
+			{
+				++NumValidObjects;
+			}
+		}
+
+		if (NumValidObjects == BoundObjects.Num())
+		{
+			return FOutlinerItemModel::GetLabelColor();
+		}
+
+		if (NumValidObjects > 0)
+		{
+			return FLinearColor::Yellow;
+		}
+	}
+
+	// Spawnables don't have valid object bindings when their track hasn't spawned them yet,
+	// so we override the default behavior of red with a gray so that users don't think there is something wrong.
+	constexpr bool bIncludeThis = true;
+	for (TSharedPtr<FObjectBindingModel> Parent : GetAncestorsOfType<FObjectBindingModel>(bIncludeThis))
+	{
+		if (Parent->GetType() == EObjectBindingType::Spawnable)
+		{
+			return FSlateColor::UseSubduedForeground();
+		}
+	}
+
+	return FLinearColor::Red;
+}
+
+FText FObjectBindingModel::GetTooltipForSingleObjectBinding() const
+{
+	return FText::Format(LOCTEXT("PossessableBoundObjectToolTip", "(BindingID: {0}"), FText::FromString(LexToString(ObjectBindingID)));
+}
+
+FText FObjectBindingModel::GetLabelToolTipText() const
+{
+	TSharedPtr<ISequencer> Sequencer = OwnerModel->GetSequencer();
+	if (!Sequencer)
+	{
+		return FText();
+	}
+
+	TArrayView<TWeakObjectPtr<>> BoundObjects = Sequencer->FindBoundObjects(ObjectBindingID, OwnerModel->GetSequenceID());
+
+	if ( BoundObjects.Num() == 0 )
+	{
+		return FText::Format(LOCTEXT("InvalidBoundObjectToolTip", "The object bound to this track is missing (BindingID: {0})."), FText::FromString(LexToString(ObjectBindingID)));
+	}
+	else
+	{
+		TArray<FString> ValidBoundObjectLabels;
+		FName BoundObjectClass;
+		bool bAddEllipsis = false;
+		int32 NumMissing = 0;
+		for (const TWeakObjectPtr<>& Ptr : BoundObjects)
+		{
+			UObject* Obj = Ptr.Get();
+
+			if (Obj == nullptr)
+			{
+				++NumMissing;
+				continue;
+			}
+
+			if (Obj->GetClass())
+			{
+				BoundObjectClass = Obj->GetClass()->GetFName();
+			}
+
+			if (AActor* Actor = Cast<AActor>(Obj))
+			{
+				ValidBoundObjectLabels.Add(Actor->GetActorLabel());
+			}
+			else
+			{
+				ValidBoundObjectLabels.Add(Obj->GetName());
+			}
+
+			if (ValidBoundObjectLabels.Num() > 3)
+			{
+				bAddEllipsis = true;
+				break;
+			}
+		}
+
+		// If only 1 bound object, display a simpler tooltip.
+		if (ValidBoundObjectLabels.Num() == 1 && NumMissing == 0)
+		{
+			return GetTooltipForSingleObjectBinding();
+		}
+		else if (ValidBoundObjectLabels.Num() == 0 && NumMissing == 1)
+		{
+			return FText::Format(LOCTEXT("InvalidBoundObjectToolTip", "The object bound to this track is missing (BindingID: {0})."), FText::FromString(LexToString(ObjectBindingID)));
+		}
+
+		FString MultipleBoundObjectLabel = FString::Join(ValidBoundObjectLabels, TEXT(", "));
+		if (bAddEllipsis)
+		{
+			MultipleBoundObjectLabel += FString::Printf(TEXT("... %d more"), BoundObjects.Num()-3);
+		}
+
+		if (NumMissing != 0)
+		{
+			MultipleBoundObjectLabel += FString::Printf(TEXT(" (%d missing)"), NumMissing);
+		}
+
+		return FText::FromString(MultipleBoundObjectLabel + FString::Printf(TEXT(" Class: %s (BindingID: %s)"), *LexToString(BoundObjectClass), *LexToString(ObjectBindingID)));
+	}
+}
+
+const FSlateBrush* FObjectBindingModel::GetIconBrush() const
+{
+	const UClass* ClassForObjectBinding = FindObjectClass();
+	if (ClassForObjectBinding)
+	{
+		return FSlateIconFinder::FindIconBrushForClass(ClassForObjectBinding);
+	}
+
+	return FAppStyle::GetBrush("Sequencer.InvalidSpawnableIcon");
+}
+
+TSharedRef<SWidget> FObjectBindingModel::GetAddTrackMenuContent()
+{
+	TSharedPtr<FSequencer> Sequencer = OwnerModel->GetSequencerImpl();
+	check(Sequencer);
+
+	//@todo need to resolve this between UMG and the level editor sequencer
+	const bool bUseSubMenus = Sequencer->IsLevelEditorSequencer();
+
+	UObject* BoundObject = Sequencer->FindSpawnedObjectOrTemplate(ObjectBindingID);
+
+	const UClass* MainSelectionObjectClass = FindObjectClass();
+
+	TArray<FGuid> ObjectBindings;
+	ObjectBindings.Add(ObjectBindingID);
+
+	TArray<UClass*> ObjectClasses;
+	ObjectClasses.Add(const_cast<UClass*>(MainSelectionObjectClass));
+
+	// Only include other selected object bindings if this binding is selected. Otherwise, this will lead to 
+	// confusion with multiple tracks being added to possibly unrelated objects
+	if (Sequencer->GetSelection().IsSelected(AsShared()))
+	{
+		for (const TWeakPtr<FViewModel>& Node : Sequencer->GetSelection().GetSelectedOutlinerItems())
+		{
+			const FObjectBindingModel* ObjectBindingNode = Node.Pin()->CastThisChecked<FObjectBindingModel>();
+			if (!ObjectBindingNode)
+			{
+				continue;
+			}
+
+			const FGuid Guid = ObjectBindingNode->GetObjectGuid();
+			for (auto RuntimeObject : Sequencer->FindBoundObjects(Guid, OwnerModel->GetSequenceID()))
+			{
+				if (RuntimeObject != nullptr)
+				{
+					ObjectBindings.AddUnique(Guid);
+					ObjectClasses.Add(RuntimeObject->GetClass());
+					continue;
+				}
+			}
+		}
+	}
+
+	ISequencerModule& SequencerModule = FModuleManager::GetModuleChecked<ISequencerModule>( "Sequencer" );
+	TSharedRef<FUICommandList> CommandList(new FUICommandList);
+
+	TSharedRef<FExtender> Extender = SequencerModule.GetAddTrackMenuExtensibilityManager()->GetAllExtenders(CommandList, TArrayBuilder<UObject*>().Add(BoundObject)).ToSharedRef();
+
+	const UClass* ObjectClass = UClass::FindCommonBase(ObjectClasses);
+
+	for (const TSharedPtr<ISequencerTrackEditor>& CurTrackEditor : Sequencer->GetTrackEditors())
+	{
+		CurTrackEditor->ExtendObjectBindingTrackMenu(Extender, ObjectBindings, ObjectClass);
+	}
+
+	// The menu are generated through reflection and sometime the API exposes some recursivity (think about a Widget returning it parent which is also a Widget). Just by reflection
+	// it is not possible to determine when the root object is reached. It needs a kind of simulation which is not implemented. Also, even if the recursivity was correctly handled, the possible
+	// permutations tend to grow exponentially. Until a clever solution is found, the simple approach is to disable recursively searching those menus. User can still search the current one though.
+	// See UE-131257
+	const bool bInRecursivelySearchable = false;
+
+	FMenuBuilder AddTrackMenuBuilder(true, nullptr, Extender, false, &FCoreStyle::Get(), true, NAME_None, bInRecursivelySearchable);
+
+	const int32 NumStartingBlocks = AddTrackMenuBuilder.GetMultiBox()->GetBlocks().Num();
+
+	AddTrackMenuBuilder.BeginSection("Tracks", LOCTEXT("TracksMenuHeader" , "Tracks"));
+	Sequencer->BuildObjectBindingTrackMenu(AddTrackMenuBuilder, ObjectBindings, ObjectClass);
+	AddTrackMenuBuilder.EndSection();
+
+	TArray<FPropertyPath> KeyablePropertyPaths;
+
+	if (BoundObject != nullptr)
+	{
+		FPropertyPath PropertyPath;
+		GetKeyablePropertyPaths(BoundObject->GetClass(), BoundObject, BoundObject->GetClass(), PropertyPath, *Sequencer, KeyablePropertyPaths);
+	}
+
+	// [Aspect Ratio]
+	// [PostProcess Settings] [Bloom1Tint] [X]
+	// [PostProcess Settings] [Bloom1Tint] [Y]
+	// [PostProcess Settings] [ColorGrading]
+	// [Ortho View]
+
+	// Create property menu data based on keyable property paths
+	TArray<PropertyMenuData> KeyablePropertyMenuData;
+	for (const FPropertyPath& KeyablePropertyPath : KeyablePropertyPaths)
+	{
+		FProperty* Property = KeyablePropertyPath.GetRootProperty().Property.Get();
+		if (Property)
+		{
+			PropertyMenuData KeyableMenuData;
+			KeyableMenuData.PropertyPath = KeyablePropertyPath;
+			if (KeyablePropertyPath.GetRootProperty().ArrayIndex != INDEX_NONE)
+			{
+				KeyableMenuData.MenuName = FText::Format(LOCTEXT("PropertyMenuTextFormat", "{0} [{1}]"), Property->GetDisplayNameText(), FText::AsNumber(KeyablePropertyPath.GetRootProperty().ArrayIndex)).ToString();
+			}
+			else
+			{
+				KeyableMenuData.MenuName = Property->GetDisplayNameText().ToString();
+			}
+			KeyablePropertyMenuData.Add(KeyableMenuData);
+		}
+	}
+
+	// Sort on the menu name
+	KeyablePropertyMenuData.Sort([](const PropertyMenuData& A, const PropertyMenuData& B)
+	{
+		int32 CompareResult = A.MenuName.Compare(B.MenuName);
+		return CompareResult < 0;
+	});
+	
+
+	// Add menu items
+	AddTrackMenuBuilder.BeginSection( SequencerMenuExtensionPoints::AddTrackMenu_PropertiesSection, LOCTEXT("PropertiesMenuHeader" , "Properties"));
+	for (int32 MenuDataIndex = 0; MenuDataIndex < KeyablePropertyMenuData.Num(); )
+	{
+		TArray<FPropertyPath> KeyableSubMenuPropertyPaths;
+
+		KeyableSubMenuPropertyPaths.Add(KeyablePropertyMenuData[MenuDataIndex].PropertyPath);
+
+		// If this menu data only has one property name, add the menu item
+		if (KeyablePropertyMenuData[MenuDataIndex].PropertyPath.GetNumProperties() == 1 || !bUseSubMenus)
+		{
+			AddPropertyMenuItems(AddTrackMenuBuilder, KeyableSubMenuPropertyPaths, 0, -1);
+			++MenuDataIndex;
+		}
+		// Otherwise, look to the next menu data to gather up new data
+		else
+		{
+			for (; MenuDataIndex < KeyablePropertyMenuData.Num()-1; )
+			{
+				if (KeyablePropertyMenuData[MenuDataIndex].MenuName == KeyablePropertyMenuData[MenuDataIndex+1].MenuName)
+				{	
+					++MenuDataIndex;
+					KeyableSubMenuPropertyPaths.Add(KeyablePropertyMenuData[MenuDataIndex].PropertyPath);
+				}
+				else
+				{
+					break;
+				}
+			}
+
+			AddTrackMenuBuilder.AddSubMenu(
+				FText::FromString(KeyablePropertyMenuData[MenuDataIndex].MenuName),
+				FText::GetEmpty(), 
+				FNewMenuDelegate::CreateSP(this, &FObjectBindingModel::HandleAddTrackSubMenuNew, KeyableSubMenuPropertyPaths, 0));
+
+			++MenuDataIndex;
+		}
+	}
+	AddTrackMenuBuilder.EndSection();
+
+	if (AddTrackMenuBuilder.GetMultiBox()->GetBlocks().Num() == NumStartingBlocks)
+	{
+		TSharedRef<SWidget> EmptyTip = SNew(SBox)
+			.Padding(FMargin(15.f, 7.5f))
+			[
+				SNew(STextBlock)
+				.Text(LOCTEXT("NoKeyablePropertiesFound", "No keyable properties or tracks"))
+				.ColorAndOpacity(FSlateColor::UseSubduedForeground())
+			];
+
+		AddTrackMenuBuilder.AddWidget(EmptyTip, FText(), true, false);
+	}
+
+	return AddTrackMenuBuilder.MakeWidget();
+}
+
+void FObjectBindingModel::HandleAddTrackSubMenuNew(FMenuBuilder& AddTrackMenuBuilder, TArray<FPropertyPath> KeyablePropertyPaths, int32 PropertyNameIndexStart)
+{
+	// [PostProcessSettings] [Bloom1Tint] [X]
+	// [PostProcessSettings] [Bloom1Tint] [Y]
+	// [PostProcessSettings] [ColorGrading]
+
+	// Create property menu data based on keyable property paths
+	TArray<FProperty*> PropertiesTraversed;
+	TArray<int32> ArrayIndicesTraversed;
+	TArray<PropertyMenuData> KeyablePropertyMenuData;
+	for (const FPropertyPath& KeyablePropertyPath : KeyablePropertyPaths)
+	{
+		PropertyMenuData KeyableMenuData;
+		KeyableMenuData.PropertyPath = KeyablePropertyPath;
+
+		// If the path is greater than 1, keep track of the actual properties (not channels) and only add these properties once since we can't do single channel keying of a property yet.
+		if (KeyablePropertyPath.GetNumProperties() > 1) //@todo
+		{
+			const FPropertyInfo& PropertyInfo = KeyablePropertyPath.GetPropertyInfo(1);
+			FProperty* Property = PropertyInfo.Property.Get();
+
+			// Search for any array elements
+			int32 ArrayIndex = INDEX_NONE;
+			for (int32 PropertyInfoIndex = 0; PropertyInfoIndex < KeyablePropertyPath.GetNumProperties(); ++PropertyInfoIndex)
+			{
+				const FPropertyInfo& ArrayPropertyInfo = KeyablePropertyPath.GetPropertyInfo(PropertyInfoIndex);
+				if (ArrayPropertyInfo.ArrayIndex != INDEX_NONE)
+				{
+					ArrayIndex = ArrayPropertyInfo.ArrayIndex;
+					break;
+				}
+			}
+
+			bool bFound = false;
+			for (int32 TraversedIndex = 0; TraversedIndex < PropertiesTraversed.Num(); ++TraversedIndex)
+			{
+				if (PropertiesTraversed[TraversedIndex] == Property && ArrayIndicesTraversed[TraversedIndex] == ArrayIndex)
+				{
+					bFound = true;
+					break;
+				}
+			}
+
+			if (bFound)
+			{
+				continue;
+			}
+
+			if (ArrayIndex != INDEX_NONE)
+			{
+				KeyableMenuData.MenuName = FText::Format(LOCTEXT("ArrayElementFormat", "{0} [{1}]"), Property->GetDisplayNameText(), FText::AsNumber(ArrayIndex)).ToString();
+			}
+			else
+			{
+				//KeyableMenuData.MenuName = FObjectEditorUtils::GetCategoryFName(Property).ToString();
+			}
+
+			PropertiesTraversed.Add(Property);
+			ArrayIndicesTraversed.Add(ArrayIndex);
+		}
+		else
+		{
+			// No sub menus items, so skip
+			continue; 
+		}
+		KeyablePropertyMenuData.Add(KeyableMenuData);
+	}
+
+	// Sort on the menu name
+	KeyablePropertyMenuData.Sort([](const PropertyMenuData& A, const PropertyMenuData& B)
+	{
+		int32 CompareResult = A.MenuName.Compare(B.MenuName);
+		return CompareResult < 0;
+	});
+
+	// Add menu items
+	for (int32 MenuDataIndex = 0; MenuDataIndex < KeyablePropertyMenuData.Num(); )
+	{
+		TArray<FPropertyPath> KeyableSubMenuPropertyPaths;
+		KeyableSubMenuPropertyPaths.Add(KeyablePropertyMenuData[MenuDataIndex].PropertyPath);
+
+		for (; MenuDataIndex < KeyablePropertyMenuData.Num()-1; )
+		{
+			if (KeyablePropertyMenuData[MenuDataIndex].MenuName == KeyablePropertyMenuData[MenuDataIndex+1].MenuName)
+			{
+				++MenuDataIndex;
+				KeyableSubMenuPropertyPaths.Add(KeyablePropertyMenuData[MenuDataIndex].PropertyPath);
+			}
+			else
+			{
+				break;
+			}
+		}
+
+		AddTrackMenuBuilder.AddSubMenu(
+			FText::FromString(KeyablePropertyMenuData[MenuDataIndex].MenuName),
+			FText::GetEmpty(), 
+			FNewMenuDelegate::CreateSP(this, &FObjectBindingModel::AddPropertyMenuItems, KeyableSubMenuPropertyPaths, PropertyNameIndexStart + 1, PropertyNameIndexStart + 2));
+
+		++MenuDataIndex;
+	}
+}
+
+void FObjectBindingModel::HandlePropertyMenuItemExecute(FPropertyPath PropertyPath)
+{
+	TSharedPtr<FSequencer> Sequencer = OwnerModel->GetSequencerImpl();
+	UObject* BoundObject = Sequencer->FindSpawnedObjectOrTemplate(ObjectBindingID);
+
+	TArray<UObject*> KeyableBoundObjects;
+	if (BoundObject != nullptr)
+	{
+		if (Sequencer->CanKeyProperty(FCanKeyPropertyParams(BoundObject->GetClass(), PropertyPath)))
+		{
+			KeyableBoundObjects.Add(BoundObject);
+		}
+	}
+
+	for (const TWeakPtr<FViewModel>& Node : Sequencer->GetSelection().GetSelectedOutlinerItems())
+	{
+		FObjectBindingModel* ObjectBindingNode = ICastable::CastWeakPtr<FObjectBindingModel>(Node);
+		if (ObjectBindingNode)
+		{
+			FGuid Guid = ObjectBindingNode->GetObjectGuid();
+			for (auto RuntimeObject : Sequencer->FindBoundObjects(Guid, OwnerModel->GetSequenceID()))
+			{
+				if (Sequencer->CanKeyProperty(FCanKeyPropertyParams(RuntimeObject->GetClass(), PropertyPath)))
+				{
+					KeyableBoundObjects.AddUnique(RuntimeObject.Get());
+				}
+			}
+		}
+	}
+
+	// When auto setting track defaults are disabled, force add a key so that the changed
+	// value is saved and is propagated to the property.
+	FKeyPropertyParams KeyPropertyParams(KeyableBoundObjects, PropertyPath, Sequencer->GetAutoSetTrackDefaults() == false ? ESequencerKeyMode::ManualKeyForced : ESequencerKeyMode::ManualKey);
+
+	Sequencer->KeyProperty(KeyPropertyParams);
+}
+
+void FObjectBindingModel::AddPropertyMenuItems(FMenuBuilder& AddTrackMenuBuilder, TArray<FPropertyPath> KeyableProperties, int32 PropertyNameIndexStart, int32 PropertyNameIndexEnd)
+{
+	TArray<PropertyMenuData> KeyablePropertyMenuData;
+
+	for (auto KeyableProperty : KeyableProperties)
+	{
+		TArray<FString> PropertyNames;
+		if (PropertyNameIndexEnd == -1)
+		{
+			PropertyNameIndexEnd = KeyableProperty.GetNumProperties();
+		}
+
+		//@todo
+		if (PropertyNameIndexStart >= KeyableProperty.GetNumProperties())
+		{
+			continue;
+		}
+
+		for (int32 PropertyNameIndex = PropertyNameIndexStart; PropertyNameIndex < PropertyNameIndexEnd; ++PropertyNameIndex)
+		{
+			PropertyNames.Add(KeyableProperty.GetPropertyInfo(PropertyNameIndex).Property.Get()->GetDisplayNameText().ToString());
+		}
+
+		PropertyMenuData KeyableMenuData;
+		{
+			KeyableMenuData.PropertyPath = KeyableProperty;
+			KeyableMenuData.MenuName = FString::Join( PropertyNames, TEXT( "." ) );
+		}
+
+		KeyablePropertyMenuData.Add(KeyableMenuData);
+	}
+
+	// Sort on the menu name
+	KeyablePropertyMenuData.Sort([](const PropertyMenuData& A, const PropertyMenuData& B)
+	{
+		int32 CompareResult = A.MenuName.Compare(B.MenuName);
+		return CompareResult < 0;
+	});
+
+	// Add menu items
+	for (int32 MenuDataIndex = 0; MenuDataIndex < KeyablePropertyMenuData.Num(); ++MenuDataIndex)
+	{
+		FUIAction AddTrackMenuAction(FExecuteAction::CreateSP(this, &FObjectBindingModel::HandlePropertyMenuItemExecute, KeyablePropertyMenuData[MenuDataIndex].PropertyPath));
+		AddTrackMenuBuilder.AddMenuEntry(FText::FromString(KeyablePropertyMenuData[MenuDataIndex].MenuName), FText(), FSlateIcon(), AddTrackMenuAction);
+	}
+}
+
+
+void FObjectBindingModel::BuildContextMenu(FMenuBuilder& MenuBuilder)
+{
+	FSequencer* Sequencer = GetEditor()->GetSequencerImpl().Get();
+	ISequencerModule& SequencerModule = FModuleManager::GetModuleChecked<ISequencerModule>("Sequencer");
+
+	UObject* BoundObject = Sequencer->FindSpawnedObjectOrTemplate(ObjectBindingID);
+	const UClass* ObjectClass = FindObjectClass();
+
+	TSharedRef<FUICommandList> CommandList(new FUICommandList);
+	TSharedPtr<FExtender> Extender = SequencerModule.GetObjectBindingContextMenuExtensibilityManager()->GetAllExtenders(CommandList, TArrayBuilder<UObject*>().Add(BoundObject));
+	if (Extender.IsValid())
+	{
+		MenuBuilder.PushExtender(Extender.ToSharedRef());
+	}
+
+	if (Sequencer->IsLevelEditorSequencer())
+	{
+		UMovieScene* MovieScene = OwnerModel->GetMovieScene();
+		FMovieSceneSpawnable* Spawnable = MovieScene->FindSpawnable(ObjectBindingID);
+
+		if (Spawnable)
+		{
+			MenuBuilder.BeginSection("Spawnable", LOCTEXT("SpawnableMenuSectionName", "Spawnable"));
+	
+			MenuBuilder.AddSubMenu(
+				LOCTEXT("OwnerLabel", "Spawned Object Owner"),
+				LOCTEXT("OwnerTooltip", "Specifies how the spawned object is to be owned"),
+				FNewMenuDelegate::CreateSP(this, &FObjectBindingModel::AddSpawnOwnershipMenu)
+			);
+
+			MenuBuilder.AddSubMenu(
+				LOCTEXT("SubLevelLabel", "Spawnable Level"),
+				LOCTEXT("SubLevelTooltip", "Specifies which level the spawnable should be spawned into"),
+				FNewMenuDelegate::CreateSP(this, &FObjectBindingModel::AddSpawnLevelMenu)
+			);
+
+			MenuBuilder.AddSubMenu(
+				LOCTEXT("ChangeClassLabel", "Change Class"),
+				LOCTEXT("ChangeClassTooltip", "Change the class (object template) that this spawns from"),
+				FNewMenuDelegate::CreateSP(this, &FObjectBindingModel::AddChangeClassMenu));
+
+			MenuBuilder.AddMenuEntry(
+				LOCTEXT("ContinuouslyRespawn", "Continuously Respawn"),
+				LOCTEXT("ContinuouslyRespawnTooltip", "When enabled, this spawnable will always be respawned if it gets destroyed externally. When disabled, this object will only ever be spawned once for each spawn key even if it is destroyed externally"),
+				FSlateIcon(),
+				FUIAction(
+					FExecuteAction::CreateLambda(FMovieSceneSpawnableFlagToggler{Sequencer, MovieScene, &FMovieSceneSpawnable::bContinuouslyRespawn, LOCTEXT("ContinuouslyRespawnTransaction", "Set Continuously Respawn")}),
+					FCanExecuteAction(),
+					FGetActionCheckState::CreateLambda(FMovieSceneSpawnableFlagCheckState{Sequencer, MovieScene, &FMovieSceneSpawnable::bContinuouslyRespawn})
+				),
+				NAME_None,
+				EUserInterfaceActionType::ToggleButton
+			);
+
+			MenuBuilder.AddMenuEntry(
+				LOCTEXT("EvaluateTracksWhenNotSpawned", "Evaluate Tracks When Not Spawned"),
+				LOCTEXT("EvaluateTracksWhenNotSpawnedTooltip", "When enabled, any tracks on this object binding or its children will still be evaluated even when the object is not spawned."),
+				FSlateIcon(),
+				FUIAction(
+					FExecuteAction::CreateLambda(FMovieSceneSpawnableFlagToggler{Sequencer, MovieScene, &FMovieSceneSpawnable::bEvaluateTracksWhenNotSpawned, LOCTEXT("EvaluateTracksWhenNotSpawned_Transaction", "Evaluate Tracks When Not Spawned")}),
+					FCanExecuteAction(),
+					FGetActionCheckState::CreateLambda(FMovieSceneSpawnableFlagCheckState{Sequencer, MovieScene, &FMovieSceneSpawnable::bEvaluateTracksWhenNotSpawned})
+				),
+				NAME_None,
+				EUserInterfaceActionType::ToggleButton
+			);
+
+			MenuBuilder.AddMenuEntry(
+				LOCTEXT("NetAddressable", "Net Addressable"),
+				LOCTEXT("NetAddressableTooltip", "When enabled, this spawnable will be spawned using a unique name that allows it to be addressed by the server and client (useful for relative movement calculations on spawned props)"),
+				FSlateIcon(),
+				FUIAction(
+					FExecuteAction::CreateLambda(FMovieSceneSpawnableFlagToggler{Sequencer, MovieScene, &FMovieSceneSpawnable::bNetAddressableName, LOCTEXT("NetAddressableTransaction", "Set Net Addressable")}),
+					FCanExecuteAction(),
+					FGetActionCheckState::CreateLambda(FMovieSceneSpawnableFlagCheckState{Sequencer, MovieScene, &FMovieSceneSpawnable::bNetAddressableName})
+				),
+				NAME_None,
+				EUserInterfaceActionType::ToggleButton
+			);
+
+			MenuBuilder.AddMenuEntry( FSequencerCommands::Get().SaveCurrentSpawnableState );
+			MenuBuilder.AddMenuEntry( FSequencerCommands::Get().ConvertToPossessable );
+
+			MenuBuilder.EndSection();
+		}
+		else
+		{
+			MenuBuilder.BeginSection("Possessable");
+
+			MenuBuilder.AddMenuEntry( FSequencerCommands::Get().ConvertToSpawnable );
+
+			MenuBuilder.EndSection();
+		}
+
+		MenuBuilder.BeginSection("Import/Export", LOCTEXT("ImportExportMenuSectionName", "Import/Export"));
+		
+		MenuBuilder.AddMenuEntry(
+			LOCTEXT("ImportFBX", "Import..."),
+			LOCTEXT("ImportFBXTooltip", "Import FBX animation to this object"),
+			FSlateIcon(),
+			FUIAction(
+				FExecuteAction::CreateLambda([=]{ Sequencer->ImportFBXOntoSelectedNodes(); })
+			));
+
+		MenuBuilder.AddMenuEntry(
+			LOCTEXT("ExportFBX", "Export..."),
+			LOCTEXT("ExportFBXTooltip", "Export FBX animation from this object"),
+			FSlateIcon(),
+			FUIAction(
+				FExecuteAction::CreateLambda([=]{ Sequencer->ExportFBX(); })
+			));
+
+		MenuBuilder.AddMenuEntry(
+			LOCTEXT("ExportToCameraAnim", "Export to Camera Anim..."),
+			LOCTEXT("ExportToCameraAnimTooltip", "Exports the animation to a camera anim asset"),
+			FSlateIcon(),
+			FUIAction(
+				FExecuteAction::CreateLambda([=]{ Sequencer->ExportToCameraAnim(); })
+			));
+		MenuBuilder.EndSection();
+	}
+
+	Sequencer->BuildCustomContextMenuForGuid(MenuBuilder, ObjectBindingID);
+	TArray<FGuid> ObjectBindings;
+	ObjectBindings.Add(ObjectBindingID);
+	for (const TSharedPtr<ISequencerTrackEditor>& TrackEditor : Sequencer->GetTrackEditors())
+	{
+		TrackEditor->BuildObjectBindingContextMenu(MenuBuilder, ObjectBindings, ObjectClass);
+	}
+
+	FOutlinerItemModel::BuildContextMenu(MenuBuilder);
+}
+
+void FObjectBindingModel::BuildOrganizeContextMenu(FMenuBuilder& MenuBuilder)
+{
+	MenuBuilder.AddSubMenu(
+		LOCTEXT("TagsLabel", "Tags"),
+		LOCTEXT("TagsTooltip", "Show this object binding's tags"),
+		FNewMenuDelegate::CreateSP(this, &FObjectBindingModel::AddTagMenu)
+	);
+
+	FOutlinerItemModel::BuildOrganizeContextMenu(MenuBuilder);
+}
+
+void FObjectBindingModel::AddSpawnOwnershipMenu(FMenuBuilder& MenuBuilder)
+{
+	TSharedPtr<ISequencer> Sequencer = OwnerModel->GetSequencer();
+	UMovieScene* MovieScene = OwnerModel->GetMovieScene();
+	FMovieSceneSpawnable* Spawnable = MovieScene->FindSpawnable(ObjectBindingID);
+	if (!Spawnable)
+	{
+		return;
+	}
+	auto Callback = [=](ESpawnOwnership NewOwnership){
+
+		FScopedTransaction Transaction(LOCTEXT("SetSpawnOwnership", "Set Spawnable Ownership"));
+
+		Spawnable->SetSpawnOwnership(NewOwnership);
+
+		// Overwrite the completion state for all spawn sections to ensure the expected behaviour.
+		EMovieSceneCompletionMode NewCompletionMode = NewOwnership == ESpawnOwnership::InnerSequence ? EMovieSceneCompletionMode::RestoreState : EMovieSceneCompletionMode::KeepState;
+
+		// Make all spawn sections retain state
+		UMovieSceneSpawnTrack* SpawnTrack = MovieScene->FindTrack<UMovieSceneSpawnTrack>(ObjectBindingID);
+		if (SpawnTrack)
+		{
+			for (UMovieSceneSection* Section : SpawnTrack->GetAllSections())
+			{
+				Section->Modify();
+				Section->EvalOptions.CompletionMode = NewCompletionMode;
+			}
+		}
+	};
+
+	MenuBuilder.AddMenuEntry(
+		LOCTEXT("ThisSequence_Label", "This Sequence"),
+		LOCTEXT("ThisSequence_Tooltip", "Indicates that this sequence will own the spawned object. The object will be destroyed at the end of the sequence."),
+		FSlateIcon(),
+		FUIAction(
+			FExecuteAction::CreateLambda(Callback, ESpawnOwnership::InnerSequence),
+			FCanExecuteAction(),
+			FIsActionChecked::CreateLambda([=]{ return Spawnable->GetSpawnOwnership() == ESpawnOwnership::InnerSequence; })
+		),
+		NAME_None,
+		EUserInterfaceActionType::ToggleButton
+	);
+
+	MenuBuilder.AddMenuEntry(
+		LOCTEXT("MasterSequence_Label", "Master Sequence"),
+		LOCTEXT("MasterSequence_Tooltip", "Indicates that the outermost sequence will own the spawned object. The object will be destroyed when the outermost sequence stops playing."),
+		FSlateIcon(),
+		FUIAction(
+			FExecuteAction::CreateLambda(Callback, ESpawnOwnership::MasterSequence),
+			FCanExecuteAction(),
+			FIsActionChecked::CreateLambda([=]{ return Spawnable->GetSpawnOwnership() == ESpawnOwnership::MasterSequence; })
+		),
+		NAME_None,
+		EUserInterfaceActionType::ToggleButton
+	);
+
+	MenuBuilder.AddMenuEntry(
+		LOCTEXT("External_Label", "External"),
+		LOCTEXT("External_Tooltip", "Indicates this object's lifetime is managed externally once spawned. It will not be destroyed by sequencer."),
+		FSlateIcon(),
+		FUIAction(
+			FExecuteAction::CreateLambda(Callback, ESpawnOwnership::External),
+			FCanExecuteAction(),
+			FIsActionChecked::CreateLambda([=]{ return Spawnable->GetSpawnOwnership() == ESpawnOwnership::External; })
+		),
+		NAME_None,
+		EUserInterfaceActionType::ToggleButton
+	);
+}
+
+void FObjectBindingModel::AddSpawnLevelMenu(FMenuBuilder& MenuBuilder)
+{
+	TSharedPtr<FSequencer> Sequencer = OwnerModel->GetSequencerImpl();
+	UMovieScene* MovieScene = OwnerModel->GetMovieScene();
+	FMovieSceneSpawnable* Spawnable = MovieScene->FindSpawnable(ObjectBindingID);
+	if (!Spawnable)
+	{
+		return;
+	}
+
+	MenuBuilder.AddMenuEntry(
+		NSLOCTEXT("UnrealEd", "PersistentLevel", "Persistent Level"),
+		NSLOCTEXT("UnrealEd", "PersistentLevel", "Persistent Level"),
+		FSlateIcon(),
+		FUIAction(
+			FExecuteAction::CreateLambda([=] { Sequencer->SetSelectedNodesSpawnableLevel(NAME_None); }),
+			FCanExecuteAction(),
+			FIsActionChecked::CreateLambda([=] { return Spawnable->GetLevelName() == NAME_None; })
+		),
+		NAME_None,
+		EUserInterfaceActionType::ToggleButton
+	);
+
+	UWorld* World = Cast<UWorld>(Sequencer->GetPlaybackContext());
+	if (!World)
+	{
+		return;
+	}
+
+	for (ULevelStreaming* LevelStreaming : World->GetStreamingLevels())
+	{
+		if (LevelStreaming)
+		{
+			FName LevelName = FPackageName::GetShortFName( LevelStreaming->GetWorldAssetPackageFName() );
+
+			MenuBuilder.AddMenuEntry(
+				FText::FromName(LevelName),
+				FText::FromName(LevelName),
+				FSlateIcon(),
+				FUIAction(
+					FExecuteAction::CreateLambda([=] { Sequencer->SetSelectedNodesSpawnableLevel(LevelName); }),
+					FCanExecuteAction(),
+					FIsActionChecked::CreateLambda([=] { return Spawnable->GetLevelName() == LevelName; })
+				),
+				NAME_None,
+				EUserInterfaceActionType::ToggleButton
+			);
+		}
+	}
+}
+
+void FObjectBindingModel::AddChangeClassMenu(FMenuBuilder& MenuBuilder)
+{
+	UMovieScene* MovieScene = OwnerModel->GetMovieScene();
+	FMovieSceneSpawnable* Spawnable = MovieScene->FindSpawnable(ObjectBindingID);
+	if (!Spawnable)
+	{
+		return;
+	}
+
+	FClassViewerModule& ClassViewerModule = FModuleManager::LoadModuleChecked<FClassViewerModule>("ClassViewer");
+
+	FClassViewerInitializationOptions Options;
+	Options.Mode = EClassViewerMode::ClassPicker;
+	Options.bIsActorsOnly = true;
+	Options.bIsPlaceableOnly = true;
+
+	const UClass* ClassForObjectBinding = FindObjectClass();
+	if (ClassForObjectBinding)
+	{
+		Options.ViewerTitleString = FText::FromString(TEXT("Change from: ") + ClassForObjectBinding->GetFName().ToString());
+	}
+	else
+	{
+		Options.ViewerTitleString = FText::FromString(TEXT("Change from: (empty)"));
+	}
+
+	MenuBuilder.AddWidget(
+		SNew(SBox)
+		.MinDesiredWidth(300.0f)
+		.MaxDesiredHeight(400.0f)
+		[
+			ClassViewerModule.CreateClassViewer(Options, FOnClassPicked::CreateRaw(this, &FObjectBindingModel::HandleTemplateActorClassPicked))
+		],
+		FText(), true, false
+	);
+}
+
+void FObjectBindingModel::HandleTemplateActorClassPicked(UClass* ChosenClass)
+{
+	FSlateApplication::Get().DismissAllMenus();
+
+	TSharedPtr<ISequencer> Sequencer = OwnerModel->GetSequencer();
+	UMovieScene* MovieScene = OwnerModel->GetMovieScene();
+	FMovieSceneSpawnable* Spawnable = MovieScene->FindSpawnable(ObjectBindingID);
+	if (!Spawnable)
+	{
+		return;
+	}
+
+	FScopedTransaction Transaction(LOCTEXT("ChangeClass", "Change Class"));
+
+	MovieScene->Modify();
+
+	TValueOrError<FNewSpawnable, FText> Result = Sequencer->GetSpawnRegister().CreateNewSpawnableType(*ChosenClass, *MovieScene, nullptr);
+	if (Result.IsValid())
+	{
+		Spawnable->SetObjectTemplate(Result.GetValue().ObjectTemplate);
+
+		Sequencer->GetSpawnRegister().DestroySpawnedObject(Spawnable->GetGuid(), OwnerModel->GetSequenceID(), *Sequencer.Get());
+		Sequencer->ForceEvaluate();
+	}
+}
+
+void FObjectBindingModel::AddTagMenu(FMenuBuilder& MenuBuilder)
+{
+	MenuBuilder.AddMenuEntry(FSequencerCommands::Get().OpenTaggedBindingManager);
+
+	TSharedPtr<FSequencer> Sequencer = OwnerModel->GetSequencerImpl();
+
+	UMovieSceneSequence* Sequence   = Sequencer->GetRootMovieSceneSequence();
+	UMovieScene*         MovieScene = Sequence->GetMovieScene();
+
+	MenuBuilder.BeginSection(NAME_None, LOCTEXT("ObjectTagsHeader", "Object Tags"));
+	{
+		TSet<FName> AllTags;
+
+		// Gather all the tags on all currently selected object binding IDs
+		FMovieSceneSequenceID SequenceID = OwnerModel->GetSequenceID();
+		for (const TWeakPtr<FViewModel>& Node : Sequencer->GetSelection().GetSelectedOutlinerItems())
+		{
+			if (FObjectBindingModel* ObjectBindingNode = ICastable::CastWeakPtr<FObjectBindingModel>(Node))
+			{
+				const FGuid& ObjectID = ObjectBindingNode->GetObjectGuid();
+
+				UE::MovieScene::FFixedObjectBindingID BindingID(ObjectID, SequenceID);
+				for (auto It = Sequencer->GetObjectBindingTagCache()->IterateTags(BindingID); It; ++It)
+				{
+					AllTags.Add(It.Value());
+				}
+			}
+		}
+
+		bool bIsReadOnly = MovieScene->IsReadOnly();
+		for (const FName& TagName : AllTags)
+		{
+			MenuBuilder.AddMenuEntry(
+				FText::FromName(TagName),
+				FText(),
+				FSlateIcon(),
+				FUIAction(
+					FExecuteAction::CreateSP(this, &FObjectBindingModel::ToggleTag, TagName),
+					FCanExecuteAction::CreateLambda([bIsReadOnly] { return bIsReadOnly == false; }),
+					FGetActionCheckState::CreateSP(this, &FObjectBindingModel::GetTagCheckState, TagName)
+				),
+				NAME_None,
+				EUserInterfaceActionType::ToggleButton
+			);
+		}
+	}
+	MenuBuilder.EndSection();
+
+	MenuBuilder.BeginSection(NAME_None, LOCTEXT("AddNewHeader", "Add Tag"));
+	{
+		if (!MovieScene->IsReadOnly())
+		{
+			TSharedRef<SWidget> Widget =
+				SNew(SObjectBindingTag)
+				.OnCreateNew(this, &FObjectBindingModel::HandleAddTag);
+
+			MenuBuilder.AddWidget(Widget, FText());
+		}
+	}
+	MenuBuilder.EndSection();
+}
+
+ECheckBoxState FObjectBindingModel::GetTagCheckState(FName TagName)
+{
+	ECheckBoxState CheckBoxState = ECheckBoxState::Undetermined;
+
+	TSharedPtr<FSequencer> Sequencer = OwnerModel->GetSequencerImpl();
+	FMovieSceneSequenceID SequenceID = OwnerModel->GetSequenceID();
+
+	for (const TWeakPtr<FViewModel>& Node : Sequencer->GetSelection().GetSelectedOutlinerItems())
+	{
+		if (FObjectBindingModel* ObjectBindingNode = ICastable::CastWeakPtr<FObjectBindingModel>(Node))
+		{
+			const FGuid& ObjectID = ObjectBindingNode->GetObjectGuid();
+
+			UE::MovieScene::FFixedObjectBindingID BindingID(ObjectID, SequenceID);
+			ECheckBoxState ThisCheckState = Sequencer->GetObjectBindingTagCache()->HasTag(BindingID, TagName)
+				? ECheckBoxState::Checked
+				: ECheckBoxState::Unchecked;
+
+			if (CheckBoxState == ECheckBoxState::Undetermined)
+			{
+				CheckBoxState = ThisCheckState;
+			}
+			else if (CheckBoxState != ThisCheckState)
+			{
+				return ECheckBoxState::Undetermined;
+			}
+		}
+	}
+
+	return CheckBoxState;
+}
+
+void FObjectBindingModel::ToggleTag(FName TagName)
+{
+	TSharedPtr<FSequencer> Sequencer = OwnerModel->GetSequencerImpl();
+	FMovieSceneSequenceID SequenceID = OwnerModel->GetSequenceID();
+
+	for (TWeakPtr<FViewModel> Node : Sequencer->GetSelection().GetSelectedOutlinerItems())
+	{
+		if (FObjectBindingModel* ObjectBindingNode = ICastable::CastWeakPtr<FObjectBindingModel>(Node))
+		{
+			const FGuid& ObjectID = ObjectBindingNode->GetObjectGuid();
+
+			UE::MovieScene::FFixedObjectBindingID BindingID(ObjectID, SequenceID);
+			if (!Sequencer->GetObjectBindingTagCache()->HasTag(BindingID, TagName))
+			{
+				HandleAddTag(TagName);
+				return;
+			}
+		}
+	}
+
+	HandleDeleteTag(TagName);
+}
+
+void FObjectBindingModel::HandleDeleteTag(FName TagName)
+{
+	FScopedTransaction Transaction(FText::Format(LOCTEXT("RemoveBindingTag", "Remove tag '{0}' from binding(s)"), FText::FromName(TagName)));
+
+	TSharedPtr<ISequencer> Sequencer = OwnerModel->GetSequencer();
+	UMovieScene* MovieScene = OwnerModel->GetMovieScene();
+	MovieScene->Modify();
+
+	FMovieSceneSequenceID SequenceID = OwnerModel->GetSequenceID();
+	for (TWeakPtr<FViewModel> Node : Sequencer->GetSelection().GetSelectedOutlinerItems())
+	{
+		if (FObjectBindingModel* ObjectBindingNode = ICastable::CastWeakPtr<FObjectBindingModel>(Node))
+		{
+			const FGuid& ObjectID = ObjectBindingNode->GetObjectGuid();
+
+			MovieScene->UntagBinding(TagName, UE::MovieScene::FFixedObjectBindingID(ObjectID, SequenceID));
+		}
+	}
+}
+
+void FObjectBindingModel::HandleAddTag(FName TagName)
+{
+	FScopedTransaction Transaction(FText::Format(LOCTEXT("CreateBindingTag", "Add new tag {0} to binding(s)"), FText::FromName(TagName)));
+
+	TSharedPtr<ISequencer> Sequencer = OwnerModel->GetSequencer();
+	UMovieScene* MovieScene = Sequencer->GetRootMovieSceneSequence()->GetMovieScene();
+	MovieScene->Modify();
+
+	FMovieSceneSequenceID SequenceID = OwnerModel->GetSequenceID();
+	for (TWeakPtr<FViewModel> Node : Sequencer->GetSelection().GetSelectedOutlinerItems())
+	{
+		if (FObjectBindingModel* ObjectBindingNode = ICastable::CastWeakPtr<FObjectBindingModel>(Node))
+		{
+			const FGuid& ObjectID = ObjectBindingNode->GetObjectGuid();
+
+			MovieScene->TagBinding(TagName, UE::MovieScene::FFixedObjectBindingID(ObjectID, SequenceID));
+		}
+	}
+}
+
+void FObjectBindingModel::SortChildren()
+{
+	ISortableExtension::SortChildren(SharedThis(this), ESortingMode::PriorityFirst);
+}
+
+FSortingKey FObjectBindingModel::GetSortingKey() const
+{
+	FSortingKey SortingKey;
+
+	if (OwnerModel)
+	{
+		UMovieScene* MovieScene = OwnerModel->GetMovieScene();
+		const FMovieSceneBinding* MovieSceneBinding = MovieScene->GetBindings().FindByPredicate([&](FMovieSceneBinding& Binding)
+		{
+			return Binding.GetObjectGuid() == ObjectBindingID;
+		});
+
+		if (MovieSceneBinding)
+		{
+			SortingKey.CustomOrder = MovieSceneBinding->GetSortingOrder();
+		}
+
+		SortingKey.DisplayName = MovieScene->GetObjectDisplayName(ObjectBindingID);
+	}
+
+	// When inside object bindings, we come before tracks. Elsewhere, we come after tracks.
+	const bool bHasParentObjectBinding = (CastParent<IObjectBindingExtension>() != nullptr);
+	SortingKey.PrioritizeBy(bHasParentObjectBinding ? 2 : 1);
+
+	return SortingKey;
+}
+
+void FObjectBindingModel::SetCustomOrder(int32 InCustomOrder)
+{
+	if (OwnerModel)
+	{
+		UMovieScene* MovieScene = OwnerModel->GetMovieScene();
+		FMovieSceneBinding* MovieSceneBinding = MovieScene->FindBinding(ObjectBindingID);
+		if (MovieSceneBinding)
+		{
+			MovieSceneBinding->SetSortingOrder(InCustomOrder);
+		}
+	}
+}
+
+bool FObjectBindingModel::CanDrag() const
+{
+	// Can only drag top level object bindings
+	TSharedPtr<IObjectBindingExtension> ObjectBindingExtension = FindAncestorOfType<IObjectBindingExtension>();
+	return ObjectBindingExtension == nullptr;
+}
+
+bool FObjectBindingModel::CanDelete(FText* OutErrorMessage) const
+{
+	return true;
+}
+
+void FObjectBindingModel::Delete()
+{
+	UMovieScene* MovieScene = OwnerModel->GetMovieScene();
+
+	// Untag this binding
+	UE::MovieScene::FFixedObjectBindingID BindingID(ObjectBindingID, OwnerModel->GetSequenceID());
+	for (auto It = OwnerModel->GetSequencerImpl()->GetObjectBindingTagCache()->IterateTags(BindingID); It; ++It)
+	{
+		MovieScene->UntagBinding(It.Value(), BindingID);
+	}
+
+	// Delete any child object bindings - this will remove their tracks implicitly
+	// so no need to delete those manually
+	for (const TViewModelPtr<FObjectBindingModel>& ChildObject : GetChildrenOfType<FObjectBindingModel>(EViewModelListType::Outliner).ToArray())
+	{
+		ChildObject->Delete();
+	}
+
+	// Remove from a parent folder if necessary.
+	if (TViewModelPtr<FFolderModel> ParentFolder = CastParent<FFolderModel>())
+	{
+		ParentFolder->GetFolder()->RemoveChildObjectBinding(ObjectBindingID);
+	}
+}
+
+} // namespace Sequencer
+} // namespace UE
+
+#undef LOCTEXT_NAMESPACE
+
