Sequencer: MVVM2 branch and Layer Bars

diff --git a/Engine/Source/Editor/Sequencer/Private/MVVM/ViewModels/OutlinerItemModel.cpp b/Engine/Source/Editor/Sequencer/Private/MVVM/ViewModels/OutlinerItemModel.cpp
new file mode 100644
--- /dev/null
+++ b/Engine/Source/Editor/Sequencer/Private/MVVM/ViewModels/OutlinerItemModel.cpp
@@ -0,0 +1,721 @@
+// Copyright Epic Games, Inc. All Rights Reserved.
+
+#include "MVVM/ViewModels/OutlinerItemModel.h"
+#include "MVVM/ViewModels/ViewModelIterators.h"
+#include "MVVM/ViewModels/SequenceModel.h"
+#include "MVVM/ViewModels/FolderModel.h"
+#include "MVVM/Extensions/IDraggableOutlinerExtension.h"
+#include "MVVM/Extensions/ITrackExtension.h"
+#include "MVVM/Extensions/IObjectBindingExtension.h"
+#include "MVVM/ViewModels/OutlinerViewModel.h"
+#include "MVVM/ViewModels/SequencerEditorViewModel.h"
+
+#include "CurveEditor.h"
+#include "Framework/Commands/GenericCommands.h"
+#include "Framework/Commands/UIAction.h"
+#include "Framework/MultiBox/MultiBoxBuilder.h"
+#include "MovieSceneSequence.h"
+#include "ScopedTransaction.h"
+#include "Sequencer.h"
+#include "SequencerSelectionCurveFilter.h"
+#include "Styling/AppStyle.h"
+#include "Tree/CurveEditorTreeFilter.h"
+#include "Tree/SCurveEditorTreePin.h"
+#include "Tree/SCurveEditorTreeSelect.h"
+#include "Widgets/Images/SImage.h"
+#include "Widgets/Layout/SSpacer.h"
+
+#define LOCTEXT_NAMESPACE "OutlinerItemModel"
+
+namespace UE
+{
+namespace Sequencer
+{
+
+static bool NodeMatchesTextFilterTerm(TSharedPtr<const UE::Sequencer::IOutlinerExtension> Node, const FCurveEditorTreeTextFilterTerm& Term)
+{
+	using namespace UE::Sequencer;
+
+	bool bMatched = false;
+
+	for (const FCurveEditorTreeTextFilterToken& Token : Term.ChildToParentTokens)
+	{
+		if (!Node)
+		{
+			// No match - ran out of parents
+			return false;
+		}
+		else if (!Token.Match(*Node->GetIdentifier().ToString()))
+		{
+			return false;
+		}
+
+		bMatched = true;
+		//Node = Node->GetParent();
+		break;
+	}
+
+	return bMatched;
+}
+
+void AddEvalOptionsPropertyMenuItem(FMenuBuilder& MenuBuilder, FCanExecuteAction InCanExecute, const TArray<UMovieSceneTrack*>& AllTracks, const FBoolProperty* Property, TFunction<bool(UMovieSceneTrack*)> Validator = nullptr)
+{
+	bool bIsChecked = AllTracks.ContainsByPredicate(
+		[=](UMovieSceneTrack* InTrack)
+		{
+			return (!Validator || Validator(InTrack)) && Property->GetPropertyValue(Property->ContainerPtrToValuePtr<void>(&InTrack->EvalOptions));
+		});
+
+	MenuBuilder.AddMenuEntry(
+		Property->GetDisplayNameText(),
+		Property->GetToolTipText(),
+		FSlateIcon(),
+		FUIAction(
+			FExecuteAction::CreateLambda([AllTracks, Property, Validator, bIsChecked]{
+				FScopedTransaction Transaction(FText::Format(NSLOCTEXT("Sequencer", "TrackNodeSetRoundEvaluation", "Set '{0}'"), Property->GetDisplayNameText()));
+				for (UMovieSceneTrack* Track : AllTracks)
+				{
+					if (Validator && !Validator(Track))
+					{
+						continue;
+					}
+					void* PropertyContainer = Property->ContainerPtrToValuePtr<void>(&Track->EvalOptions);
+					Track->Modify();
+					Property->SetPropertyValue(PropertyContainer, !bIsChecked);
+				}
+			}),
+			InCanExecute,
+			FIsActionChecked::CreateLambda([=]{ return bIsChecked; })
+		),
+		NAME_None,
+		EUserInterfaceActionType::Check
+	);
+}
+
+void AddDisplayOptionsPropertyMenuItem(FMenuBuilder& MenuBuilder, FCanExecuteAction InCanExecute, const TArray<UMovieSceneTrack*>& AllTracks, const FBoolProperty* Property, TFunction<bool(UMovieSceneTrack*)> Validator = nullptr)
+{
+	bool bIsChecked = AllTracks.ContainsByPredicate(
+		[=](UMovieSceneTrack* InTrack)
+	{
+		return (!Validator || Validator(InTrack)) && Property->GetPropertyValue(Property->ContainerPtrToValuePtr<void>(&InTrack->DisplayOptions));
+	});
+
+	MenuBuilder.AddMenuEntry(
+		Property->GetDisplayNameText(),
+		Property->GetToolTipText(),
+		FSlateIcon(),
+		FUIAction(
+			FExecuteAction::CreateLambda([AllTracks, Property, Validator, bIsChecked] {
+				FScopedTransaction Transaction(FText::Format(NSLOCTEXT("Sequencer", "TrackNodeSetDisplayOption", "Set '{0}'"), Property->GetDisplayNameText()));
+				for (UMovieSceneTrack* Track : AllTracks)
+				{
+					if (Validator && !Validator(Track))
+					{
+						continue;
+					}
+					void* PropertyContainer = Property->ContainerPtrToValuePtr<void>(&Track->DisplayOptions);
+					Track->Modify();
+					Property->SetPropertyValue(PropertyContainer, !bIsChecked);
+				}
+			}),
+			InCanExecute,
+			FIsActionChecked::CreateLambda([=] { return bIsChecked; })
+		),
+		NAME_None,
+		EUserInterfaceActionType::Check
+	);
+}
+
+FOutlinerItemModel::FOutlinerItemModel()
+	: OutlinerChildList(EViewModelListType::Outliner)
+	, bInitializedExpansion(false)
+	, bInitializedPinnedState(false)
+	, bIsForcedFilteredOut(false)
+{
+	RegisterChildList(&OutlinerChildList);
+}
+
+TSharedPtr<FSequencerEditorViewModel> FOutlinerItemModel::GetEditor() const
+{
+	TSharedPtr<FSequenceModel> SequenceModel = FindAncestorOfType<FSequenceModel>();
+	return SequenceModel ? SequenceModel->GetEditor() : nullptr;
+}
+
+bool FOutlinerItemModel::IsForceFilteredOut() const
+{
+	return bIsForcedFilteredOut;
+}
+
+void FOutlinerItemModel::SetForceFilteredOut(bool bInIsForcedFilteredOut)
+{
+	bIsForcedFilteredOut = bInIsForcedFilteredOut;
+}
+
+FName FOutlinerItemModel::GetIdentifier() const
+{
+	return TreeItemIdentifier;
+}
+
+void FOutlinerItemModel::SetIdentifier(FName InNewIdentifier)
+{
+	TreeItemIdentifier = InNewIdentifier;
+}
+
+bool FOutlinerItemModel::IsExpanded() const
+{
+	if (bInitializedExpansion)
+	{
+		return bIsExpanded;
+	}
+
+	bInitializedExpansion = true;
+
+	TStringBuilder<256> StringBuilder;
+	IOutlinerExtension::GetPathName(*this, StringBuilder);
+
+	TSharedPtr<FSequenceModel> SequenceModel = FindAncestorOfType<FSequenceModel>();
+	UMovieSceneSequence*       Sequence      = SequenceModel ? SequenceModel->GetSequence() : nullptr;
+
+	if (Sequence)
+	{
+		FStringView StringView = StringBuilder.ToView();
+		UMovieScene* MovieScene = Sequence->GetMovieScene();
+		FMovieSceneEditorData& EditorData = MovieScene->GetEditorData();
+		if (const FMovieSceneExpansionState* Expansion = EditorData.ExpansionStates.FindByHash(GetTypeHash(StringView), StringView))
+		{
+			const_cast<FOutlinerItemModel*>(this)->bIsExpanded = Expansion->bExpanded;
+		}
+		else
+		{
+			const_cast<FOutlinerItemModel*>(this)->bIsExpanded = GetDefaultExpansionState();
+		}
+	}
+
+	return bIsExpanded;
+}
+
+bool FOutlinerItemModel::GetDefaultExpansionState() const
+{
+	return false;
+}
+
+void FOutlinerItemModel::SetExpansion(bool bInIsExpanded)
+{
+	SetExpansionWithoutSaving(bInIsExpanded);
+
+	if (GetParent())
+	{
+		// Expansion state has changed, save it to the movie scene now
+		TSharedPtr<FSequenceModel> SequenceModel = FindAncestorOfType<FSequenceModel>();
+		if (SequenceModel)
+		{
+			TSharedPtr<FSequencer> Sequencer = SequenceModel->GetSequencerImpl();
+			Sequencer->GetNodeTree()->SaveExpansionState(*this, bInIsExpanded);
+		}
+	}
+}
+
+void FOutlinerItemModel::SetExpansionWithoutSaving(bool bInIsExpanded)
+{
+	FOutlinerExtensionShim::SetExpansion(bInIsExpanded);
+
+	// Force this flag in case a sub-class wants a given expansion state before the
+	// getter is called.
+	bInitializedExpansion = true;
+}
+
+bool FOutlinerItemModel::IsFilteredOut() const
+{
+	return bIsFilteredOut || bIsForcedFilteredOut;
+}
+
+bool FOutlinerItemModel::IsPinned() const
+{
+	if (bInitializedPinnedState)
+	{
+		return FPinnableExtensionShim::IsPinned();
+	}
+
+	bInitializedPinnedState = true;
+
+	// Initialize expansion states for tree items
+	// Assign the saved expansion state when this node is initialized for the first time
+	const bool bIsRootModel = (GetHierarchicalDepth() == 1);
+	if (bIsRootModel)
+	{
+		TSharedPtr<FSequenceModel> SequenceModel = FindAncestorOfType<FSequenceModel>();
+		TSharedPtr<FSequencer> Sequencer = SequenceModel->GetSequencerImpl();
+		const bool bWasPinned = Sequencer->GetNodeTree()->GetSavedPinnedState(*this);
+		const_cast<FOutlinerItemModel*>(this)->FPinnableExtensionShim::SetPinned(bWasPinned);
+	}
+
+	return FPinnableExtensionShim::IsPinned();
+}
+
+bool FOutlinerItemModel::IsDimmed() const
+{
+	TSharedPtr<FSequencerEditorViewModel> EditorViewModel = GetEditor();
+	TSharedPtr<FSequencer> Sequencer = EditorViewModel ? EditorViewModel->GetSequencerImpl() : nullptr;
+	if (Sequencer)
+	{
+		TSharedPtr<FViewModel> ThisShared(const_cast<FOutlinerItemModel*>(this)->AsShared());
+		FSequencerNodeTree& ParentTree = Sequencer->GetNodeTree().Get();
+		if (ParentTree.IsNodeMute(ThisShared) || (ParentTree.HasSoloNodes() && !ParentTree.IsNodeSolo(ThisShared)))
+		{
+			return true;
+		}
+	}
+	return false;
+}
+
+bool FOutlinerItemModel::IsSolo() const
+{
+	TSharedPtr<FSequencerEditorViewModel> EditorViewModel = GetEditor();
+	TSharedPtr<FSequencer> Sequencer = EditorViewModel ? EditorViewModel->GetSequencerImpl() : nullptr;
+	if (Sequencer)
+	{
+		TSharedPtr<FViewModel> ThisShared(const_cast<FOutlinerItemModel*>(this)->AsShared());
+		return Sequencer->GetNodeTree()->IsNodeSolo(ThisShared);
+	}
+	return false;
+}
+
+bool FOutlinerItemModel::IsMuted() const
+{
+	TSharedPtr<FSequencerEditorViewModel> EditorViewModel = GetEditor();
+	TSharedPtr<FSequencer> Sequencer = EditorViewModel ? EditorViewModel->GetSequencerImpl() : nullptr;
+	if (Sequencer)
+	{
+		TSharedPtr<FViewModel> ThisShared(const_cast<FOutlinerItemModel*>(this)->AsShared());
+		return Sequencer->GetNodeTree()->IsNodeMute(ThisShared);
+	}
+	return false;
+}
+
+bool FOutlinerItemModel::IsRootModelPinned() const
+{
+	TSharedPtr<IPinnableExtension> PinnableParent = FindAncestorOfType<IPinnableExtension>(true);
+	return PinnableParent && PinnableParent->IsPinned();
+}
+
+void FOutlinerItemModel::ToggleRootModelPinned()
+{
+	FSequenceModel* RootModel = GetRoot()->CastThis<FSequenceModel>();
+	TSharedPtr<IPinnableExtension> PinnableParent = FindAncestorOfType<IPinnableExtension>(true);
+	if (RootModel && PinnableParent)
+	{
+		TSharedPtr<FOutlinerViewModel> Outliner = RootModel->GetEditor()->GetOutliner();
+		Outliner->UnpinAllNodes();
+
+		const bool bShouldPin = !PinnableParent->IsPinned();
+		PinnableParent->SetPinned(bShouldPin);
+
+		TSharedPtr<FSequencer> Sequencer = RootModel->GetSequencerImpl();
+		Sequencer->GetNodeTree()->SavePinnedState(*this, bShouldPin);
+		Sequencer->RefreshTree();
+	}
+}
+
+bool FOutlinerItemModel::IsSelectedModelsSolo() const
+{
+	TSharedPtr<FSequencerEditorViewModel> EditorViewModel = GetEditor();
+	TSharedPtr<FSequencer> Sequencer = EditorViewModel ? EditorViewModel->GetSequencerImpl() : nullptr;
+	if (Sequencer)
+	{
+		return Sequencer->GetNodeTree()->IsSelectedNodesSolo();
+	}
+	return false;
+}
+
+void FOutlinerItemModel::ToggleSelectedModelsSolo()
+{
+	TSharedPtr<FSequencerEditorViewModel> EditorViewModel = GetEditor();
+	TSharedPtr<FSequencer> Sequencer = EditorViewModel ? EditorViewModel->GetSequencerImpl() : nullptr;
+	if (Sequencer)
+	{
+		Sequencer->GetNodeTree()->ToggleSelectedNodesSolo();
+	}
+}
+
+bool FOutlinerItemModel::IsSelectedModelsMuted() const
+{
+	TSharedPtr<FSequencerEditorViewModel> EditorViewModel = GetEditor();
+	TSharedPtr<FSequencer> Sequencer = EditorViewModel ? EditorViewModel->GetSequencerImpl() : nullptr;
+	if (Sequencer)
+	{
+		return Sequencer->GetNodeTree()->IsSelectedNodesMute();
+	}
+	return false;
+}
+
+void FOutlinerItemModel::ToggleSelectedModelsMuted()
+{
+	TSharedPtr<FSequencerEditorViewModel> EditorViewModel = GetEditor();
+	TSharedPtr<FSequencer> Sequencer = EditorViewModel ? EditorViewModel->GetSequencerImpl() : nullptr;
+	if (Sequencer)
+	{
+		return Sequencer->GetNodeTree()->ToggleSelectedNodesMute();
+	}
+}
+
+TSharedPtr<SWidget> FOutlinerItemModel::CreateContextMenuWidget(const FCreateOutlinerContextMenuWidgetParams& InParams)
+{
+	TSharedPtr<FSequencerEditorViewModel> EditorViewModel = GetEditor();
+	TSharedPtr<ISequencer> Sequencer = EditorViewModel->GetSequencer();
+
+	if (Sequencer)
+	{
+		const bool bShouldCloseWindowAfterMenuSelection = true;
+		FMenuBuilder MenuBuilder(bShouldCloseWindowAfterMenuSelection, Sequencer->GetCommandBindings());
+
+		BuildContextMenu(MenuBuilder);
+
+		return MenuBuilder.MakeWidget();
+	}
+
+	return nullptr;
+}
+
+void FOutlinerItemModel::BuildContextMenu(FMenuBuilder& MenuBuilder)
+{
+	TSharedPtr<FSequencer> Sequencer = StaticCastSharedPtr<FSequencer>(GetEditor()->GetSequencer());
+
+	if (!Sequencer)
+	{
+		return;
+	}
+
+	const bool bIsReadOnly = Sequencer->IsReadOnly();
+	FCanExecuteAction CanExecute = FCanExecuteAction::CreateLambda([bIsReadOnly]{ return !bIsReadOnly; });
+
+	MenuBuilder.BeginSection("Edit", LOCTEXT("EditContextMenuSectionName", "Edit"));
+	{
+		MenuBuilder.AddMenuEntry(
+			LOCTEXT("ToggleNodeLock", "Locked"),
+			LOCTEXT("ToggleNodeLockTooltip", "Lock or unlock this node or selected tracks"),
+			FSlateIcon(),
+			FUIAction(
+				FExecuteAction::CreateSP(Sequencer.Get(), &FSequencer::ToggleNodeLocked),
+				CanExecute,
+				FIsActionChecked::CreateSP(Sequencer.Get(), &FSequencer::IsNodeLocked)
+			),
+			NAME_None,
+			EUserInterfaceActionType::ToggleButton
+		);
+
+		// Only support pinning root nodes
+		const bool bIsRootModel = (GetHierarchicalDepth() == 1);
+		if (bIsRootModel)
+		{
+			MenuBuilder.AddMenuEntry(
+				LOCTEXT("ToggleNodePin", "Pinned"),
+				LOCTEXT("ToggleNodePinTooltip", "Pin or unpin this node or selected tracks"),
+				FSlateIcon(),
+				FUIAction(
+					FExecuteAction::CreateSP(this, &FOutlinerItemModel::ToggleRootModelPinned),
+					FCanExecuteAction(),
+					FIsActionChecked::CreateSP(this, &FOutlinerItemModel::IsRootModelPinned)
+				),
+				NAME_None,
+				EUserInterfaceActionType::ToggleButton
+			);
+		}
+
+		// We already know we are soloable and mutable
+		MenuBuilder.AddMenuEntry(
+			LOCTEXT("ToggleNodeSolo", "Solo"),
+			LOCTEXT("ToggleNodeSoloTooltip", "Solo or unsolo this node or selected tracks"),
+			FSlateIcon(),
+			FUIAction(
+				FExecuteAction::CreateSP(this, &FOutlinerItemModel::ToggleSelectedModelsSolo),
+				CanExecute,
+				FIsActionChecked::CreateSP(this, &FOutlinerItemModel::IsSelectedModelsSolo)
+			),
+			NAME_None,
+			EUserInterfaceActionType::ToggleButton
+		);
+
+		MenuBuilder.AddMenuEntry(
+			LOCTEXT("ToggleNodeMute", "Mute"),
+			LOCTEXT("ToggleNodeMuteTooltip", "Mute or unmute this node or selected tracks"),
+			FSlateIcon(),
+			FUIAction(
+				FExecuteAction::CreateSP(this, &FOutlinerItemModel::ToggleSelectedModelsMuted),
+				CanExecute,
+				FIsActionChecked::CreateSP(this, &FOutlinerItemModel::IsSelectedModelsMuted)
+			),
+			NAME_None,
+			EUserInterfaceActionType::ToggleButton
+		);
+
+		// Add cut, copy and paste functions to the tracks
+		MenuBuilder.AddMenuEntry(FGenericCommands::Get().Cut);
+
+		MenuBuilder.AddMenuEntry(FGenericCommands::Get().Copy);
+		
+		MenuBuilder.AddMenuEntry(FGenericCommands::Get().Paste);
+		
+		MenuBuilder.AddMenuEntry(FGenericCommands::Get().Duplicate);
+		
+		TSharedRef<FViewModel> ThisNode = AsShared();
+
+		MenuBuilder.AddMenuEntry(
+			LOCTEXT("DeleteNode", "Delete"),
+			LOCTEXT("DeleteNodeTooltip", "Delete this or selected tracks"),
+			FSlateIcon(FAppStyle::GetAppStyleSetName(), "ContentBrowser.AssetActions.Delete"),
+			FUIAction(FExecuteAction::CreateSP(Sequencer.Get(), &FSequencer::DeleteNode, ThisNode, false), CanExecute)
+		);
+
+		if (IsA<IObjectBindingExtension>())
+		{
+			MenuBuilder.AddMenuEntry(
+				LOCTEXT("DeleteNodeAndKeepState", "Delete and Keep State"),
+				LOCTEXT("DeleteNodeAndKeepStateTooltip", "Delete this object's tracks and keep its current animated state"),
+				FSlateIcon(FAppStyle::GetAppStyleSetName(), "ContentBrowser.AssetActions.Delete"),
+				FUIAction(FExecuteAction::CreateSP(Sequencer.Get(), &FSequencer::DeleteNode, ThisNode, true), CanExecute)
+			);
+		}
+
+		MenuBuilder.AddMenuEntry(FGenericCommands::Get().Rename);
+	}
+	MenuBuilder.EndSection();
+
+	TArray<UMovieSceneTrack*> AllTracks;
+	for (TWeakPtr<FViewModel> Node : Sequencer->GetSelection().GetSelectedOutlinerItems())
+	{
+		if (ITrackExtension* TrackExtension = ICastable::CastWeakPtr<ITrackExtension>(Node))
+		{
+			UMovieSceneTrack* Track = TrackExtension->GetTrack();
+			if (Track)
+			{
+				AllTracks.Add(Track);
+			}
+		}
+	}
+
+	MenuBuilder.BeginSection("Organize", LOCTEXT("OrganizeContextMenuSectionName", "Organize"));
+	BuildOrganizeContextMenu(MenuBuilder);
+	MenuBuilder.EndSection();
+
+	if (AllTracks.Num())
+	{
+		MenuBuilder.BeginSection("GeneralTrackOptions", NSLOCTEXT("Sequencer", "TrackNodeGeneralOptions", "Track Options"));
+		{
+			UStruct* EvalOptionsStruct = FMovieSceneTrackEvalOptions::StaticStruct();
+
+			const FBoolProperty* NearestSectionProperty = CastField<FBoolProperty>(EvalOptionsStruct->FindPropertyByName(GET_MEMBER_NAME_CHECKED(FMovieSceneTrackEvalOptions, bEvalNearestSection)));
+			auto CanEvaluateNearest = [](UMovieSceneTrack* InTrack) { return InTrack->EvalOptions.bCanEvaluateNearestSection != 0; };
+			if (NearestSectionProperty && AllTracks.ContainsByPredicate(CanEvaluateNearest))
+			{
+				TFunction<bool(UMovieSceneTrack*)> Validator = CanEvaluateNearest;
+				AddEvalOptionsPropertyMenuItem(MenuBuilder, CanExecute, AllTracks, NearestSectionProperty, Validator);
+			}
+
+			const FBoolProperty* PrerollProperty = CastField<FBoolProperty>(EvalOptionsStruct->FindPropertyByName(GET_MEMBER_NAME_CHECKED(FMovieSceneTrackEvalOptions, bEvaluateInPreroll)));
+			if (PrerollProperty)
+			{
+				AddEvalOptionsPropertyMenuItem(MenuBuilder, CanExecute, AllTracks, PrerollProperty);
+			}
+
+			const FBoolProperty* PostrollProperty = CastField<FBoolProperty>(EvalOptionsStruct->FindPropertyByName(GET_MEMBER_NAME_CHECKED(FMovieSceneTrackEvalOptions, bEvaluateInPostroll)));
+			if (PostrollProperty)
+			{
+				AddEvalOptionsPropertyMenuItem(MenuBuilder, CanExecute, AllTracks, PostrollProperty);
+			}
+		}
+		MenuBuilder.EndSection();
+
+		MenuBuilder.BeginSection("TrackDisplayOptions", NSLOCTEXT("Sequencer", "TrackNodeDisplayOptions", "Display Options"));
+		{
+			UStruct* DisplayOptionsStruct = FMovieSceneTrackDisplayOptions::StaticStruct();
+
+			const FBoolProperty* ShowVerticalFramesProperty = CastField<FBoolProperty>(DisplayOptionsStruct->FindPropertyByName(GET_MEMBER_NAME_CHECKED(FMovieSceneTrackDisplayOptions, bShowVerticalFrames)));
+			if (ShowVerticalFramesProperty)
+			{
+				AddDisplayOptionsPropertyMenuItem(MenuBuilder, CanExecute, AllTracks, ShowVerticalFramesProperty);
+			}
+		}
+		MenuBuilder.EndSection();
+	}
+}
+
+void FOutlinerItemModel::BuildOrganizeContextMenu(FMenuBuilder& MenuBuilder)
+{
+	TSharedRef<FViewModel> ThisNode = SharedThis(this);
+	TSharedPtr<FSequencer> Sequencer = GetEditor()->GetSequencerImpl();
+
+	const bool bFilterableNode = (IsA<ITrackExtension>() || IsA<IObjectBindingExtension>() || IsA<FFolderModel>());
+	const bool bIsReadOnly = Sequencer->IsReadOnly();
+	
+	TArray<UMovieSceneTrack*> AllTracks;
+	TArray<TSharedPtr<FViewModel> > DraggableNodes;
+	for (TWeakPtr<FViewModel> WeakNode : Sequencer->GetSelection().GetSelectedOutlinerItems())
+	{
+		TSharedPtr<FViewModel> Node = WeakNode.Pin();
+
+		if (ITrackExtension* TrackExtension = Node->CastThis<ITrackExtension>())
+		{
+			UMovieSceneTrack* Track = TrackExtension->GetTrack();
+			if (Track)
+			{
+				AllTracks.Add(Track);
+			}
+		}
+
+		if (IDraggableOutlinerExtension* DraggableExtension = Node->CastThis<IDraggableOutlinerExtension>())
+		{
+			if (DraggableExtension->CanDrag())
+			{
+				DraggableNodes.Add(Node);
+			}
+		}
+	}
+
+	if (bFilterableNode && !bIsReadOnly)
+	{
+		MenuBuilder.AddSubMenu(
+			LOCTEXT("AddNodesToNodeGroup", "Add to Group"),
+			LOCTEXT("AddNodesToNodeGroupTooltip", "Add selected nodes to a group"),
+			FNewMenuDelegate::CreateSP(Sequencer.Get(), &FSequencer::BuildAddSelectedToNodeGroupMenu));
+
+	}
+
+	if (DraggableNodes.Num() && !bIsReadOnly)
+	{
+		MenuBuilder.AddSubMenu(
+			LOCTEXT("MoveToFolder", "Move to Folder"),
+			LOCTEXT("MoveToFolderTooltip", "Move the selected nodes to a folder"),
+			FNewMenuDelegate::CreateSP(Sequencer.Get(), &FSequencer::BuildAddSelectedToFolderMenu));
+
+		MenuBuilder.AddMenuEntry(
+			LOCTEXT("RemoveFromFolder", "Remove from Folder"),
+			LOCTEXT("RemoveFromFolderTooltip", "Remove selected nodes from their folders"),
+			FSlateIcon(),
+			FUIAction(
+				FExecuteAction::CreateSP(Sequencer.Get(), &FSequencer::RemoveSelectedNodesFromFolders),
+				FCanExecuteAction::CreateLambda( [Sequencer] { return Sequencer->GetSelectedNodesInFolders().Num() > 0; } )));
+	}
+}
+
+bool FOutlinerItemModel::HasCurves() const
+{
+	return false;
+}
+
+TSharedPtr<ICurveEditorTreeItem> FOutlinerItemModel::GetCurveEditorTreeItem() const
+{
+	return StaticCastSharedRef<ICurveEditorTreeItem>(SharedThis(const_cast<FOutlinerItemModel*>(this)));
+}
+
+TSharedPtr<SWidget> FOutlinerItemModel::GenerateCurveEditorTreeWidget(const FName& InColumnName, TWeakPtr<FCurveEditor> InCurveEditor, FCurveEditorTreeItemID InTreeItemID, const TSharedRef<ITableRow>& TableRow)
+{
+	using namespace UE::Sequencer;
+
+	auto GetCurveEditorHighlightText = [](TWeakPtr<FCurveEditor> InCurveEditor) -> FText 
+	{
+		TSharedPtr<FCurveEditor> PinnedCurveEditor = InCurveEditor.Pin();
+		if (!PinnedCurveEditor)
+		{
+			return FText::GetEmpty();
+		}
+
+		const FCurveEditorTreeFilter* Filter = PinnedCurveEditor->GetTree()->FindFilterByType(ECurveEditorTreeFilterType::Text);
+		if (Filter)
+		{
+			return static_cast<const FCurveEditorTreeTextFilter*>(Filter)->InputText;
+		}
+
+		return FText::GetEmpty();
+	};
+
+	if (InColumnName == ColumnNames.Label)
+	{
+		return SNew(SHorizontalBox)
+
+			+ SHorizontalBox::Slot()
+			.Padding(FMargin(0.f, 0.f, 4.f, 0.f))
+			.VAlign(VAlign_Center)
+			.AutoWidth()
+			[
+				SNew(SOverlay)
+
+				+ SOverlay::Slot()
+				[
+					SNew(SImage)
+					.Image(this, &FOutlinerItemModel::GetIconBrush)
+					.ColorAndOpacity(this, &FOutlinerItemModel::GetIconTint)
+				]
+
+				+ SOverlay::Slot()
+				.VAlign(VAlign_Top)
+				.HAlign(HAlign_Right)
+				[
+					SNew(SImage)
+					.Image(this, &FOutlinerItemModel::GetIconOverlayBrush)
+				]
+
+				+ SOverlay::Slot()
+				[
+					SNew(SSpacer)
+					.Visibility(EVisibility::Visible)
+					.ToolTipText(this, &FOutlinerItemModel::GetIconToolTipText)
+				]
+			]
+
+			+ SHorizontalBox::Slot()
+			.VAlign(VAlign_Center)
+			.Padding(FMargin(0.f, 4.f, 0.f, 4.f))
+			[
+				SNew(STextBlock)
+				.Text(this, &FOutlinerItemModel::GetLabel)
+				.Font(this, &FOutlinerItemModel::GetLabelFont)
+				.HighlightText_Static(GetCurveEditorHighlightText, InCurveEditor)
+				.ToolTipText(this, &FOutlinerItemModel::GetLabelToolTipText)
+			];
+	}
+	else if (InColumnName == ColumnNames.SelectHeader)
+	{
+		return SNew(SCurveEditorTreeSelect, InCurveEditor, InTreeItemID, TableRow);
+	}
+	else if (InColumnName == ColumnNames.PinHeader)
+	{
+		return SNew(SCurveEditorTreePin, InCurveEditor, InTreeItemID, TableRow);
+	}
+
+	return nullptr;
+}
+
+void FOutlinerItemModel::CreateCurveModels(TArray<TUniquePtr<FCurveModel>>& OutCurveModels)
+{
+}
+
+bool FOutlinerItemModel::PassesFilter(const FCurveEditorTreeFilter* InFilter) const
+{
+	if (InFilter->GetType() == ECurveEditorTreeFilterType::Text)
+	{
+		const FCurveEditorTreeTextFilter* Filter = static_cast<const FCurveEditorTreeTextFilter*>(InFilter);
+
+		TSharedPtr<const IOutlinerExtension> This = CastThisShared<IOutlinerExtension>();
+		for (const FCurveEditorTreeTextFilterTerm& Term : Filter->GetTerms())
+		{
+			if (NodeMatchesTextFilterTerm(This, Term))
+			{
+				return true;
+			}
+		}
+
+		return false;
+	}
+	else if (InFilter->GetType() == ISequencerModule::GetSequencerSelectionFilterType())
+	{
+		const FSequencerSelectionCurveFilter* Filter = static_cast<const FSequencerSelectionCurveFilter*>(InFilter);
+		return Filter->Match(AsShared());
+	}
+	return false;
+}
+
+} // namespace Sequencer
+} // namespace UE
+
+#undef LOCTEXT_NAMESPACE
+
diff --git a/Engine/Source/Editor/Sequencer/Private/MVVM/ViewModels/TrackRowModel.cpp b/Engine/Source/Editor/Sequencer/Private/MVVM/ViewModels/TrackRowModel.cpp
new file mode 100644
--- /dev/null
+++ b/Engine/Source/Editor/Sequencer/Private/MVVM/ViewModels/TrackRowModel.cpp
@@ -0,0 +1,342 @@
+// Copyright Epic Games, Inc. All Rights Reserved.
+
+#include "MVVM/ViewModels/TrackRowModel.h"
+#include "MVVM/ViewModels/SectionModel.h"
+#include "MVVM/ViewModels/SequenceModel.h"
+#include "MVVM/ViewModels/ViewModelIterators.h"
+#include "MVVM/ViewModels/ChannelModel.h"
+#include "MVVM/ViewModels/FolderModel.h"
+#include "MVVM/ViewModels/TrackModelLayoutBuilder.h"
+#include "MVVM/Views/SOutlinerTrackView.h"
+#include "MVVM/ViewModels/SequencerEditorViewModel.h"
+
+#include "SequencerNodeTree.h"
+#include "ISequencer.h"
+#include "ISequencerSection.h"
+#include "ISequencerTrackEditor.h"
+#include "SequencerUtilities.h"
+#include "SequencerCommonHelpers.h"
+#include "SSequencer.h"
+
+#include "MovieSceneTrack.h"
+#include "MovieSceneFolder.h"
+#include "MovieSceneSection.h"
+#include "MovieSceneTrack.h"
+#include "Tracks/MovieScene3DTransformTrack.h"
+#include "Tracks/MovieScenePrimitiveMaterialTrack.h"
+
+#include "Styling/AppStyle.h"
+
+#define LOCTEXT_NAMESPACE "TrackRowModel"
+
+namespace UE
+{
+namespace Sequencer
+{
+
+namespace LayoutConstants
+{
+	extern const float CommonPadding;
+}
+
+FTrackRowModel::FTrackRowModel(UMovieSceneTrack* Track, int32 InRowIndex)
+	: SectionList(EViewModelListType::TrackArea)
+	, TopLevelChannelList(FTrackModel::GetTopLevelChannelGroupType())
+	, WeakTrack(Track)
+	, RowIndex(InRowIndex)
+{
+	RegisterChildList(&SectionList);
+	RegisterChildList(&TopLevelChannelList);
+
+	FName Identifier = Track->GetFName();
+	Identifier.SetNumber(InRowIndex);
+	SetIdentifier(Identifier);
+}
+
+FTrackRowModel::~FTrackRowModel()
+{
+}
+
+void FTrackRowModel::Initialize()
+{
+	UMovieSceneTrack* Track = GetTrack();
+	TSharedPtr<FSequenceModel> SequenceModel = FindAncestorOfType<FSequenceModel>();
+	if (Track && SequenceModel)
+	{
+		TrackEditor = SequenceModel->GetSequencer()->GetTrackEditor(Track);
+		ensure(TrackEditor);
+		SetExpansion(TrackEditor->GetDefaultExpansionState(Track));
+	}
+}
+
+UMovieSceneTrack* FTrackRowModel::GetTrack() const
+{
+	return WeakTrack.Get();
+}
+
+int32 FTrackRowModel::GetRowIndex() const
+{
+	return RowIndex;
+}
+
+FViewModelChildren FTrackRowModel::GetSections()
+{
+	return GetChildrenForList(&SectionList);
+}
+
+FOutlinerSizing FTrackRowModel::GetOutlinerSizing() const
+{
+	float Height = SequencerLayoutConstants::SectionAreaDefaultHeight;
+	for (TSharedPtr<FSectionModel> Section : SectionList.Iterate<FSectionModel>())
+	{
+		Height = Section->GetSectionInterface()->GetSectionHeight();
+		break;
+	}
+	return FOutlinerSizing(Height + 2 * LayoutConstants::CommonPadding);
+}
+
+FTrackAreaParameters FTrackRowModel::GetTrackAreaParameters() const
+{
+	FTrackAreaParameters Params;
+	Params.LaneType = ETrackAreaLaneType::Nested;
+	return Params;
+}
+
+FViewModelVariantIterator FTrackRowModel::GetTrackAreaModelList() const
+{
+	return &SectionList;
+}
+
+void FTrackRowModel::CreateCurveModels(TArray<TUniquePtr<FCurveModel>>& OutCurveModels)
+{
+	TViewModelPtr<FChannelGroupModel> ChannelGroup = TopLevelChannelList.GetHead().ImplicitCast();
+	if (ChannelGroup)
+	{
+		ChannelGroup->CreateCurveModels(OutCurveModels);
+	}
+}
+
+bool FTrackRowModel::IsDimmed() const
+{
+	UMovieSceneTrack* Track = GetTrack();
+
+	if (Track && Track->IsRowEvalDisabled(GetRowIndex()))
+	{
+		return true;
+	}
+
+	return FOutlinerItemModel::IsDimmed();
+}
+
+FSlateFontInfo FTrackRowModel::GetLabelFont() const
+{
+	bool bAllAnimated = false;
+	TViewModelPtr<FChannelGroupModel> TopLevelChannel = TopLevelChannelList.GetHead().ImplicitCast();
+	if (TopLevelChannel)
+	{
+		for (const TViewModelPtr<FChannelModel>& ChannelModel : TopLevelChannel->GetDescendantsOfType<FChannelModel>())
+		{
+			FMovieSceneChannel* Channel = ChannelModel->GetChannel();
+			if (!Channel || Channel->GetNumKeys() == 0)
+			{
+				return FOutlinerItemModel::GetLabelFont();
+			}
+			else
+			{
+				bAllAnimated = true;
+			}
+		}
+		if (bAllAnimated == true)
+		{
+			return FAppStyle::GetFontStyle("Sequencer.AnimationOutliner.ItalicFont");
+		}
+	}
+	return FOutlinerItemModel::GetLabelFont();
+}
+
+const FSlateBrush* FTrackRowModel::GetIconBrush() const
+{
+	return TrackEditor ? TrackEditor->GetIconBrush() : nullptr;
+}
+
+FText FTrackRowModel::GetLabel() const
+{
+	UMovieSceneTrack* Track = GetTrack();
+	return Track ? Track->GetTrackRowDisplayName(RowIndex) : FText::GetEmpty();
+}
+
+FSlateColor FTrackRowModel::GetLabelColor() const
+{
+	UMovieSceneTrack* Track = GetTrack();
+	TSharedPtr<FSequenceModel> SequenceModel = FindAncestorOfType<FSequenceModel>();
+	if (!Track || !SequenceModel)
+	{
+		return FSlateColor::UseForeground();
+	}
+
+	TSharedPtr<ISequencer> Sequencer = SequenceModel->GetSequencer();
+	if (!Sequencer)
+	{
+		return FSlateColor::UseForeground();
+	}
+
+	// Display track node is red if the property track is not bound to valid property
+	if (UMovieScenePropertyTrack* PropertyTrack = Cast<UMovieScenePropertyTrack>(Track))
+	{
+		const bool bIsDimmed = IsDimmed();
+
+		// 3D transform tracks don't map to property bindings as below
+		if (Track->IsA<UMovieScene3DTransformTrack>() || Track->IsA<UMovieScenePrimitiveMaterialTrack>())
+		{
+			return bIsDimmed ? FSlateColor::UseSubduedForeground() : FSlateColor::UseForeground();
+		}
+
+		if (TSharedPtr<IObjectBindingExtension> ParentBinding = FindAncestorOfType<IObjectBindingExtension>())
+		{
+			for (TWeakObjectPtr<> WeakObject : Sequencer->FindBoundObjects(ParentBinding->GetObjectGuid(), Sequencer->GetFocusedTemplateID()))
+			{
+				if (UObject* Object = WeakObject.Get())
+				{
+					FTrackInstancePropertyBindings PropertyBinding(PropertyTrack->GetPropertyName(), PropertyTrack->GetPropertyPath().ToString());
+					if (PropertyBinding.GetProperty(*Object))
+					{
+						return bIsDimmed ? FSlateColor::UseSubduedForeground() : FSlateColor::UseForeground();
+					}
+				}
+			}
+
+			return bIsDimmed ? FSlateColor(FLinearColor::Red.Desaturate(0.6f)) : FLinearColor::Red;
+		}
+	}
+
+	return FOutlinerItemModel::GetLabelColor();
+}
+
+TSharedRef<SWidget> FTrackRowModel::CreateOutlinerView(const FCreateOutlinerViewParams& InParams)
+{
+	return SNew(SOutlinerTrackView, 
+			TWeakViewModelPtr<IOutlinerExtension>(SharedThis(this)), 
+			InParams.Editor->CastThisSharedChecked<FSequencerEditorViewModel>(), 
+			InParams.TreeViewRow);
+}
+
+bool FTrackRowModel::IsResizable() const
+{
+	UMovieSceneTrack* Track = GetTrack();
+	TSharedPtr<ISequencerTrackEditor> TrackEditorPtr = GetTrackEditor();
+	return Track && TrackEditorPtr && TrackEditorPtr->IsResizable(Track);
+}
+
+void FTrackRowModel::Resize(float NewSize)
+{
+	UMovieSceneTrack* Track = GetTrack();
+	float PaddingAmount = 2 * LayoutConstants::CommonPadding;
+	if (Track)
+	{
+		PaddingAmount *= (Track->GetMaxRowIndex() + 1);
+	}
+	
+	NewSize -= PaddingAmount;
+
+	TSharedPtr<ISequencerTrackEditor> TrackEditorPtr = GetTrackEditor();
+	if (Track && TrackEditorPtr && TrackEditorPtr->IsResizable(Track))
+	{
+		TrackEditorPtr->Resize(NewSize, Track);
+	}
+}
+
+void FTrackRowModel::BuildContextMenu(FMenuBuilder& MenuBuilder)
+{
+	UMovieSceneTrack* Track = GetTrack();
+	if (!Track)
+	{
+		return;
+	}
+
+	const int32 TrackRowIndex = GetRowIndex();
+
+	TSharedPtr<ISequencerTrackEditor> TrackEditorPtr = GetTrackEditor();
+	TrackEditorPtr->BuildTrackContextMenu(MenuBuilder, Track);
+
+	if (Track && Track->GetSupportedBlendTypes().Num() > 0)
+	{
+		TSharedPtr<ISequencer> WeakSequencer = GetEditor()->GetSequencer();
+
+		MenuBuilder.AddSubMenu(
+			LOCTEXT("AddSection", "Add Section"),
+			FText(),
+			FNewMenuDelegate::CreateLambda([=](FMenuBuilder& SubMenuBuilder){
+				FSequencerUtilities::PopulateMenu_CreateNewSection(SubMenuBuilder, TrackRowIndex + 1, Track, WeakSequencer);
+			})
+		);	
+	}
+
+	// Find sections in the track to add batch properties for
+	TArray<TWeakObjectPtr<UObject>> TrackSections;
+	for (TSharedPtr<FViewModel> TrackAreaModel : GetTrackAreaModelList())
+	{
+		constexpr bool bIncludeThis = true;
+		for (TSharedPtr<FSectionModel> Section : TParentFirstChildIterator<FSectionModel>(TrackAreaModel, bIncludeThis))
+		{
+			if (UMovieSceneSection* SectionObject = Section->GetSection())
+			{
+				TrackSections.Add(SectionObject);
+			}
+		}
+	}
+		
+	if (TrackSections.Num())
+	{
+		TSharedPtr<FSequencer> Sequencer = GetEditor()->GetSequencerImpl();
+		MenuBuilder.AddSubMenu(
+			TrackSections.Num() > 1 ? LOCTEXT("BatchEditSections", "Batch Edit Sections") : LOCTEXT("EditSection", "Edit Section"),
+			FText(),
+			FNewMenuDelegate::CreateLambda([=](FMenuBuilder& SubMenuBuilder){
+				SequencerHelpers::AddPropertiesMenu(*Sequencer, SubMenuBuilder, TrackSections);
+			})
+		);
+	}
+}
+
+bool FTrackRowModel::CanDelete(FText* OutErrorMessage) const
+{
+	return true;
+}
+
+void FTrackRowModel::Delete()
+{
+	UMovieSceneTrack* Track = GetTrack();
+	if (!Track)
+	{
+		return;
+	}
+
+	// Remove from a parent folder if necessary.
+	if (TViewModelPtr<FFolderModel> ParentFolder = CastParent<FFolderModel>())
+	{
+		ParentFolder->GetFolder()->Modify();
+		ParentFolder->GetFolder()->RemoveChildMasterTrack(Track);
+	}
+
+	// Remove sub tracks belonging to this row only
+	Track->Modify();
+	Track->SetFlags(RF_Transactional);
+
+	for (TSharedPtr<FSectionModel> SectionModel : SectionList.Iterate<FSectionModel>())
+	{
+		UMovieSceneSection* Section = SectionModel->GetSection();
+		if (Section)
+		{
+			Track->RemoveSection(*Section);
+		}
+	}
+
+	Track->UpdateEasing();
+	Track->FixRowIndices();
+}
+
+} // namespace Sequencer
+} // namespace UE
+
+#undef LOCTEXT_NAMESPACE
+
diff --git a/Engine/Source/Editor/Sequencer/Private/MVVM/Views/SOutlinerObjectBindingView.cpp b/Engine/Source/Editor/Sequencer/Private/MVVM/Views/SOutlinerObjectBindingView.cpp
new file mode 100644
--- /dev/null
+++ b/Engine/Source/Editor/Sequencer/Private/MVVM/Views/SOutlinerObjectBindingView.cpp
@@ -0,0 +1,70 @@
+// Copyright Epic Games, Inc. All Rights Reserved.
+
+#include "MVVM/Views/SOutlinerObjectBindingView.h"
+#include "MVVM/ViewModels/SequencerEditorViewModel.h"
+#include "MVVM/ViewModels/ObjectBindingModel.h"
+#include "MVVM/ViewModels/ViewModelIterators.h"
+#include "SequencerUtilities.h"
+
+#include "MovieScene.h"
+#include "MovieSceneSequence.h"
+#include "ISequencerModule.h"
+#include "ISequencerTrackEditor.h"
+
+#include "Sequencer.h"
+#include "EditorStyleSet.h"
+#include "GameFramework/Actor.h"
+#include "Styling/SlateIconFinder.h"
+#include "Widgets/Layout/SSpacer.h"
+#include "Framework/Commands/UIAction.h"
+#include "Framework/MultiBox/MultiBoxBuilder.h"
+#include "Containers/ArrayBuilder.h"
+#include "Modules/ModuleManager.h"
+
+#define LOCTEXT_NAMESPACE "SOutlinerObjectBindingView"
+
+namespace UE
+{
+namespace Sequencer
+{
+
+void SOutlinerObjectBindingView::Construct(
+		const FArguments& InArgs, 
+		TSharedPtr<FObjectBindingModel> InModel, 
+		TSharedPtr<FSequencerEditorViewModel> InEditor,
+		const TSharedRef<ISequencerTreeViewRow>& InTableRow)
+{
+	TSharedPtr<FSequencer> Sequencer = StaticCastSharedPtr<FSequencer>(InEditor->GetSequencer());
+
+	TAttribute<bool> HoverState = MakeAttributeSP(&InTableRow.Get(), &ISequencerTreeViewRow::IsHovered);
+
+	TSharedRef<SHorizontalBox> CustomBoxPanel =
+		SNew(SHorizontalBox)
+		+ SHorizontalBox::Slot()
+		[
+			SNew(SSpacer)
+		]
+
+		+ SHorizontalBox::Slot()
+		.AutoWidth()
+		.VAlign(VAlign_Center)
+		[
+			FSequencerUtilities::MakeAddButton(LOCTEXT("TrackText", "Track"), FOnGetContent::CreateSP(InModel.ToSharedRef(), &FObjectBindingModel::GetAddTrackMenuContent), HoverState, Sequencer->AsShared())
+		];
+
+	const UClass* ObjectClass = InModel->FindObjectClass();
+	Sequencer->BuildObjectBindingEditButtons(CustomBoxPanel, InModel->GetObjectGuid(), ObjectClass);
+
+	SOutlinerItemViewBase::Construct(
+		FArguments(InArgs)
+		.CustomContent()
+		[
+			CustomBoxPanel
+		], TWeakViewModelPtr<IOutlinerExtension>(InModel), InEditor, InTableRow
+	);
+}
+
+} // namespace Sequencer
+} // namespace UE
+
+#undef LOCTEXT_NAMESPACE
diff --git a/Engine/Source/Editor/Sequencer/Private/MVVM/Views/SOutlinerTrackView.cpp b/Engine/Source/Editor/Sequencer/Private/MVVM/Views/SOutlinerTrackView.cpp
new file mode 100644
--- /dev/null
+++ b/Engine/Source/Editor/Sequencer/Private/MVVM/Views/SOutlinerTrackView.cpp
@@ -0,0 +1,295 @@
+// Copyright Epic Games, Inc. All Rights Reserved.
+
+#include "MVVM/Views/SOutlinerTrackView.h"
+#include "MVVM/ViewModels/ViewModel.h"
+#include "MVVM/ICastable.h"
+#include "MVVM/Extensions/ITrackExtension.h"
+#include "MVVM/Extensions/IObjectBindingExtension.h"
+#include "MVVM/ViewModels/ChannelModel.h"
+#include "MVVM/ViewModels/TrackModel.h"
+#include "MVVM/ViewModels/ViewModelIterators.h"
+#include "MVVM/ViewModels/OutlinerViewModel.h"
+#include "MVVM/ViewModels/SequencerEditorViewModel.h"
+#include "MVVM/Views/STrackAreaView.h"
+#include "ISequencerTrackEditor.h"
+#include "SSequencerKeyNavigationButtons.h"
+#include "SKeyAreaEditorSwitcher.h"
+
+#include "MovieSceneTrack.h"
+#include "Tracks/MovieScenePropertyTrack.h"
+#include "Tracks/MovieScenePrimitiveMaterialTrack.h"
+#include "Tracks/MovieScene3DTransformTrack.h"
+
+#include "Widgets/DeclarativeSyntaxSupport.h"
+#include "Widgets/SBoxPanel.h"
+#include "Widgets/SOverlay.h"
+#include "Widgets/Images/SImage.h"
+#include "Widgets/Input/SButton.h"
+#include "Widgets/Layout/SSpacer.h"
+#include "Widgets/Layout/SBox.h"
+#include "Widgets/Colors/SColorPicker.h"
+
+#include "Engine/Engine.h"
+#include "Styling/AppStyle.h"
+
+#define LOCTEXT_NAMESPACE "SOutlinerTrackView"
+
+namespace UE
+{
+namespace Sequencer
+{
+
+void SOutlinerTrackView::Construct(
+		const FArguments& InArgs,
+		TWeakViewModelPtr<IOutlinerExtension> InWeakOutlinerExtension,
+		TWeakPtr<FSequencerEditorViewModel> InWeakEditor,
+		const TSharedRef<ISequencerTreeViewRow>& InTableRow)
+{
+	using namespace UE::Sequencer;
+
+	TSharedPtr<FViewModel> DataModel = InWeakOutlinerExtension.Pin().AsModel();
+	FSequencerEditorViewModel* EditorViewModel = ICastable::CastWeakPtr<FSequencerEditorViewModel>(InWeakEditor);
+	check(DataModel && EditorViewModel);
+
+	ITrackExtension* TrackExtension = DataModel->CastThis<ITrackExtension>();
+	checkf(TrackExtension != nullptr, TEXT("It is invalid to create a SOutlinerTrackView widget with a model that does not implement ITrackExtension"));
+
+	TSharedPtr<ISequencerTrackEditor> TrackEditor = EditorViewModel->GetSequencer()->GetTrackEditor(TrackExtension->GetTrack());
+
+	TOptional<FViewModelChildren> TopLevelChannels = DataModel->FindChildList(FTrackModel::GetTopLevelChannelGroupType());
+
+	TAttribute<bool> HoverState = MakeAttributeSP(&InTableRow.Get(), &ISequencerTreeViewRow::IsHovered);
+
+	TSharedRef<SHorizontalBox> BoxPanel = SNew(SHorizontalBox);
+
+	TSharedPtr<IObjectBindingExtension> ObjectBinding = DataModel->FindAncestorOfType<IObjectBindingExtension>();
+
+	UMovieSceneTrack* Track = TrackExtension->GetTrack();
+	int32 RowIndex = TrackExtension->GetRowIndex();
+
+	FBuildEditWidgetParams Params;
+	Params.NodeIsHovered = HoverState;
+	if (DataModel->FindAncestorOfType<ITrackExtension>())
+	{
+		Params.TrackInsertRowIndex = TrackExtension->GetRowIndex();
+	}
+	else if (Track && Track->SupportsMultipleRows())
+	{
+		Params.TrackInsertRowIndex = Track->GetMaxRowIndex()+1;
+	}
+
+	TSharedPtr<SWidget> CustomWidget = TrackEditor->BuildOutlinerEditWidget(ObjectBinding ? ObjectBinding->GetObjectGuid() : FGuid(), Track, Params);
+
+	TSharedPtr<FChannelGroupModel> TopLevelChannel = TopLevelChannels ? TopLevelChannels->FindFirstChildOfType<FChannelGroupModel>() : nullptr;
+	if (TopLevelChannel)
+	{
+		TSharedRef<SOverlay> Overlay = SNew(SOverlay);
+
+		Overlay->AddSlot()
+		.HAlign(HAlign_Left)
+		.VAlign(VAlign_Center)
+		[
+			SNew(SKeyAreaEditorSwitcher, TopLevelChannel, EditorViewModel->GetSequencer())
+		];
+
+		if (CustomWidget.IsValid())
+		{
+			Overlay->AddSlot()
+			.HAlign(HAlign_Right)
+			.VAlign(VAlign_Center)
+			[
+				CustomWidget.ToSharedRef()
+			];
+		}
+
+		BoxPanel->AddSlot()
+		.AutoWidth()
+		.VAlign(VAlign_Center)
+		[
+			Overlay
+		];
+
+		BoxPanel->AddSlot()
+		.AutoWidth()
+		.VAlign(VAlign_Center)
+		[
+			SNew(SSequencerKeyNavigationButtons, TopLevelChannel, EditorViewModel->GetSequencer())
+		];
+	}
+	else
+	{
+		if (CustomWidget.IsValid())
+		{
+			BoxPanel->AddSlot()
+			.AutoWidth()
+			.VAlign(VAlign_Center)
+			[
+				CustomWidget.ToSharedRef()
+			];
+		}
+
+		bool bHasKeyableAreas = false;
+		for (TSharedPtr<FChannelGroupModel> ChannelGroup : DataModel->GetDescendantsOfType<FChannelGroupModel>())
+		{
+			for (const TWeakViewModelPtr<FChannelModel>& WeakChannel : ChannelGroup->GetChannels())
+			{
+				if (TViewModelPtr<FChannelModel> Channel = WeakChannel.Pin())
+				{
+					if (Channel->GetKeyArea()->CanCreateKeyEditor())
+					{
+						bHasKeyableAreas = true;
+						break;
+					}
+				}
+			}
+			if (bHasKeyableAreas)
+			{
+				break;
+			}
+		}
+
+		if (bHasKeyableAreas)
+		{
+			BoxPanel->AddSlot()
+			.VAlign(VAlign_Center)
+			.AutoWidth()
+			[
+				SNew(SSequencerKeyNavigationButtons, DataModel, EditorViewModel->GetSequencer())
+			];
+		}
+	}
+
+	SOutlinerItemViewBase::Construct(
+		FArguments(InArgs)
+		.CustomContent()
+		[
+			SNew(SBox)
+			.VAlign(VAlign_Center)
+			.HAlign(HAlign_Right)
+			[
+				BoxPanel
+			]
+		]
+		.RightGutterContent()
+		[
+			SNew(SButton)
+			.ContentPadding(0)
+			.VAlign(VAlign_Fill)
+			.IsFocusable(false) // Intentionally false so that it's easier to tab to the next numeric input
+			.IsEnabled(!IsReadOnlyAttribute.Get())
+			.ButtonStyle(FAppStyle::Get(), "Sequencer.AnimationOutliner.ColorStrip")
+			.OnClicked(this, &SOutlinerTrackView::OnSetTrackColor)
+			.Content()
+			[
+				SNew(SBox)
+				.WidthOverride(6.f)
+				[
+					SNew(SImage)
+					.Image(FAppStyle::GetBrush("WhiteBrush"))
+					.ColorAndOpacity(this, &SOutlinerTrackView::GetTrackColorTint)
+				]
+			]
+		], InWeakOutlinerExtension, InWeakEditor, InTableRow);
+}
+
+// We store these for when the Color Picker is canceled so we can restore the old value.
+namespace AnimationOutlinerTreeNode
+{
+	FLinearColor InitialTrackColor;
+	bool bFolderPickerWasCancelled;
+}
+
+FReply SOutlinerTrackView::OnSetTrackColor()
+{
+	AnimationOutlinerTreeNode::InitialTrackColor = GetTrackColorTint().GetSpecifiedColor();
+	AnimationOutlinerTreeNode::bFolderPickerWasCancelled = false;
+
+	FColorPickerArgs PickerArgs;
+	PickerArgs.bUseAlpha = false;
+	PickerArgs.DisplayGamma = TAttribute<float>::Create(TAttribute<float>::FGetter::CreateUObject(GEngine, &UEngine::GetDisplayGamma));
+	PickerArgs.InitialColorOverride = AnimationOutlinerTreeNode::InitialTrackColor;
+	PickerArgs.ParentWidget = GetParentWidget();
+	PickerArgs.OnColorCommitted = FOnLinearColorValueChanged::CreateSP(this, &SOutlinerTrackView::OnColorPickerPicked);
+	PickerArgs.OnColorPickerWindowClosed = FOnWindowClosed::CreateSP(this, &SOutlinerTrackView::OnColorPickerClosed);
+	PickerArgs.OnColorPickerCancelled = FOnColorPickerCancelled::CreateSP(this, &SOutlinerTrackView::OnColorPickerCancelled);
+
+	OpenColorPicker(PickerArgs);
+	return FReply::Handled();
+}
+
+void SOutlinerTrackView::OnColorPickerPicked(FLinearColor NewFolderColor)
+{
+	UMovieSceneTrack* Track = GetTrackFromNode();
+	if (Track)
+	{
+		// This is called every time the user adjusts the UI so we don't want to create a transaction for it, just directly
+		// modify the track so we can see the change immediately.
+		Track->SetColorTint(NewFolderColor.ToFColor(true));
+	}
+}
+
+void SOutlinerTrackView::OnColorPickerClosed(const TSharedRef<SWindow>& Window)
+{
+	// Under Unreal UX terms, closing the Color Picker (via the UI) is the same as confirming it since we've been live updating
+	// the color. The track already has the latest color change so we undo the change before calling Modify so that Undo sets us
+	// to the original color. This is also called in the event of pressing cancel so we need to detect if it was canceled or not.
+	if (!AnimationOutlinerTreeNode::bFolderPickerWasCancelled)
+	{
+		UMovieSceneTrack* Track = GetTrackFromNode();
+		if(Track)
+		{
+			const FScopedTransaction Transaction(LOCTEXT("SetTrackColor", "Set Track Color"));
+			FSlateColor CurrentColor = GetTrackColorTint();
+			Track->SetColorTint(AnimationOutlinerTreeNode::InitialTrackColor.ToFColor(true));
+			Track->Modify();
+			Track->SetColorTint(CurrentColor.GetSpecifiedColor().ToFColor(true));
+		}
+	}
+}
+
+void SOutlinerTrackView::OnColorPickerCancelled(FLinearColor NewFolderColor)
+{
+	AnimationOutlinerTreeNode::bFolderPickerWasCancelled = true;
+
+	// Restore the original color of the track. No transaction will be created when the OnColorPickerClosed callback is called.
+	UMovieSceneTrack* Track = GetTrackFromNode();
+	if (Track)
+	{
+		Track->SetColorTint(AnimationOutlinerTreeNode::InitialTrackColor.ToFColor(true));
+	}
+}
+
+UMovieSceneTrack* SOutlinerTrackView::GetTrackFromNode() const
+{
+	TViewModelPtr<IOutlinerExtension> Outliner = WeakOutlinerExtension.Pin();
+	if (Outliner)
+	{
+		constexpr bool bIncludeThis = true;
+		for (TSharedPtr<ITrackExtension> TrackExtension : Outliner.AsModel()->GetAncestorsOfType<ITrackExtension>(bIncludeThis))
+		{
+			if (UMovieSceneTrack* Track = TrackExtension->GetTrack())
+			{
+				return Track;
+			}
+		}
+	}
+
+	return nullptr;
+}
+
+FSlateColor SOutlinerTrackView::GetTrackColorTint() const
+{
+	UMovieSceneTrack* Track = GetTrackFromNode();
+	if (Track)
+	{
+		return STrackAreaView::BlendDefaultTrackColor(Track->GetColorTint());
+	}
+
+	return FLinearColor::Transparent;
+}
+
+} // namespace Sequencer
+} // namespace UE
+
+#undef LOCTEXT_NAMESPACE
+
diff --git a/Engine/Source/Editor/Sequencer/Private/MovieSceneTrackEditor.cpp b/Engine/Source/Editor/Sequencer/Private/MovieSceneTrackEditor.cpp
--- a/Engine/Source/Editor/Sequencer/Private/MovieSceneTrackEditor.cpp
+++ b/Engine/Source/Editor/Sequencer/Private/MovieSceneTrackEditor.cpp
@@ -250,11 +250,6 @@
 { 
 }
 
-int32 FMovieSceneTrackEditor::PaintTrackArea(const FPaintArgs& Args, const FGeometry& AllottedGeometry, const FSlateRect& MyCullingRect, FSlateWindowElementList& OutDrawElements, int32 LayerId, const FWidgetStyle& InWidgetStyle)
-{
-	return LayerId;
-}
-
 void FMovieSceneTrackEditor::Tick(float DeltaTime) 
 { 
 }
diff --git a/Engine/Source/Editor/Sequencer/Private/SKeyAreaEditorSwitcher.cpp b/Engine/Source/Editor/Sequencer/Private/SKeyAreaEditorSwitcher.cpp
--- a/Engine/Source/Editor/Sequencer/Private/SKeyAreaEditorSwitcher.cpp
+++ b/Engine/Source/Editor/Sequencer/Private/SKeyAreaEditorSwitcher.cpp
@@ -5,15 +5,21 @@
 #include "Widgets/Layout/SBox.h"
 #include "Widgets/SBoxPanel.h"
 #include "Sequencer.h"
-#include "SKeyNavigationButtons.h"
+#include "SequencerCommonHelpers.h"
 #include "IKeyArea.h"
-#include "DisplayNodes/SequencerSectionKeyAreaNode.h"
+#include "MVVM/ViewModels/ChannelModel.h"
+#include "MVVM/Extensions/IObjectBindingExtension.h"
+#include "MVVM/Views/SKeyNavigationButtons.h"
 
-void SKeyAreaEditorSwitcher::Construct(const FArguments& InArgs, TSharedRef<FSequencerSectionKeyAreaNode> InKeyAreaNode)
+namespace UE
 {
-	WeakKeyAreaNode = InKeyAreaNode;
+namespace Sequencer
+{
 
-	Rebuild();
+void SKeyAreaEditorSwitcher::Construct(const FArguments& InArgs, TSharedPtr<FChannelGroupModel> InModel, TWeakPtr<ISequencer> InWeakSequencer)
+{
+	WeakModel = InModel;
+	WeakSequencer = InWeakSequencer;
 }
 
 int32 SKeyAreaEditorSwitcher::GetWidgetIndex() const
@@ -23,8 +29,9 @@
 
 void SKeyAreaEditorSwitcher::Rebuild()
 {
-	TSharedPtr<FSequencerSectionKeyAreaNode> KeyAreaNode = WeakKeyAreaNode.Pin();
-	if (!KeyAreaNode.IsValid())
+	TSharedPtr<FChannelGroupModel> Model     = WeakModel.Pin();
+	TSharedPtr<ISequencer>         Sequencer = WeakSequencer.Pin();
+	if (!Model || !Sequencer)
 	{
 		// Empty our cache so we don't persistently rebuild
 		CachedKeyAreas.Empty();
@@ -37,11 +44,8 @@
 		return;
 	}
 
-	const bool bIsEnabled = !KeyAreaNode->GetSequencer().IsReadOnly();
+	const bool bIsEnabled = !Sequencer->IsReadOnly();
 
-	// Update the cached list so we know when to rebuild next
-	CachedKeyAreas = KeyAreaNode->GetAllKeyAreas();
-
 	// Index 0 is always the spacer node
 	VisibleIndex = 0;
 
@@ -54,8 +58,8 @@
 			SNullWidget::NullWidget
 		];
 
-	TSharedPtr<FSequencerObjectBindingNode> ParentObjectBinding = KeyAreaNode->FindParentObjectBindingNode();
-	FGuid ObjectBindingID = ParentObjectBinding.IsValid() ? ParentObjectBinding->GetObjectBinding() : FGuid();
+	TSharedPtr<IObjectBindingExtension> ParentObjectBinding = Model->FindAncestorOfType<IObjectBindingExtension>();
+	FGuid ObjectBindingID = ParentObjectBinding.IsValid() ? ParentObjectBinding->GetObjectGuid() : FGuid();
 
 	for (TSharedRef<IKeyArea> KeyArea : CachedKeyAreas)
 	{
@@ -76,7 +80,7 @@
 				.MinDesiredWidth(100)
 				.HAlign(HAlign_Left)
 				[
-					KeyArea->CreateKeyEditor(KeyAreaNode->GetSequencer().AsShared(), ObjectBindingID)
+					KeyArea->CreateKeyEditor(Sequencer, ObjectBindingID)
 				]
 			];
 		}
@@ -90,11 +94,14 @@
 
 void SKeyAreaEditorSwitcher::Tick( const FGeometry& AllottedGeometry, const double InCurrentTime, const float InDeltaTime )
 {
-	TSharedPtr<FSequencerSectionKeyAreaNode> KeyAreaNode = WeakKeyAreaNode.Pin();
-	if (!KeyAreaNode.IsValid())
+	TSharedPtr<FChannelGroupModel> Model = WeakModel.Pin();
+	TSharedPtr<ISequencer>         Sequencer = WeakSequencer.Pin();
+	if (!Model || !Sequencer)
 	{
 		if (CachedKeyAreas.Num() != 0)
 		{
+			CachedKeyAreas.Empty();
+
 			// Node is not valid but we have a valid cache - we need to rebuild the switcher now
 			Rebuild();
 		}
@@ -102,8 +109,11 @@
 	}
 	else
 	{
-		if (CachedKeyAreas != KeyAreaNode->GetAllKeyAreas())
+		TArray<TSharedRef<IKeyArea>> NewCachedKeyAreas = Model->GetAllKeyAreas();
+		if (CachedKeyAreas != NewCachedKeyAreas)
 		{
+			Swap(CachedKeyAreas, NewCachedKeyAreas);
+
 			// Node is valid but now has a different set of key areas. Must rebuild the widgets.
 			Rebuild();
 		}
@@ -114,7 +124,7 @@
 			AllSections.Add(KeyArea->GetOwningSection());
 		}
 
-		const int32 ActiveKeyArea = SequencerHelpers::GetSectionFromTime(AllSections, KeyAreaNode->GetSequencer().GetLocalTime().Time.FrameNumber);
+		const int32 ActiveKeyArea = SequencerHelpers::GetSectionFromTime(AllSections, Sequencer->GetLocalTime().Time.FrameNumber);
 		if (ActiveKeyArea != INDEX_NONE)
 		{
 			// Index 0 is the spacer node, so add 1 to the key area index to get the widget index
@@ -126,3 +136,6 @@
 		}
 	}
 }
+
+} // namespace Sequencer
+} // namespace UE
diff --git a/Engine/Source/Editor/Sequencer/Private/SKeyAreaEditorSwitcher.h b/Engine/Source/Editor/Sequencer/Private/SKeyAreaEditorSwitcher.h
--- a/Engine/Source/Editor/Sequencer/Private/SKeyAreaEditorSwitcher.h
+++ b/Engine/Source/Editor/Sequencer/Private/SKeyAreaEditorSwitcher.h
@@ -5,10 +5,17 @@
 #include "Widgets/SCompoundWidget.h"
 #include "Widgets/DeclarativeSyntaxSupport.h"
 
-class FSequencerSectionKeyAreaNode;
 class SOverlay;
 class IKeyArea;
+class ISequencer;
 
+namespace UE
+{
+namespace Sequencer
+{
+
+class FChannelGroupModel;
+
 class SKeyAreaEditorSwitcher : public SCompoundWidget
 {
 public:
@@ -16,7 +23,7 @@
 	SLATE_END_ARGS()
 
 	/** Construct the widget */
-	void Construct(const FArguments& InArgs, TSharedRef<FSequencerSectionKeyAreaNode> InKeyAreaNode);
+	void Construct(const FArguments& InArgs, TSharedPtr<FChannelGroupModel> InModel, TWeakPtr<ISequencer> InWeakSequencer);
 
 	/** Rebuild this widget from its cached key area node */
 	void Rebuild();
@@ -35,7 +42,13 @@
 	/** Index of the currently visible key editor */
 	int32 VisibleIndex;
 	/** The key area to which we relate */
-	TWeakPtr<FSequencerSectionKeyAreaNode> WeakKeyAreaNode;
+	TWeakPtr<FChannelGroupModel> WeakModel;
+	/** The sequencer we belong to */
+	TWeakPtr<ISequencer> WeakSequencer;
 	/** Key areas cached from the node */
 	TArray<TSharedRef<IKeyArea>> CachedKeyAreas;
 };
+
+
+} // namespace Sequencer
+} // namespace UE
diff --git a/Engine/Source/Editor/Sequencer/Private/SSequencer.cpp b/Engine/Source/Editor/Sequencer/Private/SSequencer.cpp
--- a/Engine/Source/Editor/Sequencer/Private/SSequencer.cpp
+++ b/Engine/Source/Editor/Sequencer/Private/SSequencer.cpp
@@ -3,6 +3,12 @@
 #include "SSequencer.h"
 #include "Engine/Blueprint.h"
 #include "MovieSceneSequence.h"
+#include "MVVM/ViewModels/CategoryModel.h"
+#include "MVVM/ViewModels/ChannelModel.h"
+#include "MVVM/ViewModels/ViewModelIterators.h"
+#include "MVVM/Extensions/IRenameableExtension.h"
+#include "MVVM/ViewModels/SequencerEditorViewModel.h"
+#include "MVVM/ViewModels/SequencerOutlinerViewModel.h"
 #include "Sections/MovieSceneSubSection.h"
 #include "MovieScene.h"
 #include "Framework/MultiBox/MultiBoxDefs.h"
@@ -13,6 +19,7 @@
 #include "Modules/ModuleManager.h"
 #include "Editor.h"
 #include "IDetailsView.h"
+#include "IKeyArea.h"
 #include "Widgets/Layout/SBorder.h"
 #include "ISequencerEditTool.h"
 #include "Layout/WidgetPath.h"
@@ -39,8 +46,6 @@
 #include "LevelEditorViewport.h"
 #include "Widgets/Navigation/SBreadcrumbTrail.h"
 #include "SequencerCommands.h"
-#include "DisplayNodes/SequencerObjectBindingNode.h"
-#include "DisplayNodes/SequencerTrackNode.h"
 #include "Widgets/Input/SNumericDropDown.h"
 #include "SequencerCommonHelpers.h"
 #include "ISequencerWidgetsModule.h"
@@ -48,8 +53,8 @@
 #include "SequencerTimeSliderController.h"
 #include "SequencerToolMenuContext.h"
 #include "SSequencerSectionOverlay.h"
-#include "SSequencerTrackArea.h"
-#include "SSequencerTrackOutliner.h"
+#include "MVVM/Views/STrackAreaView.h"
+#include "MVVM/Views/SSequencerTrackAreaView.h"
 #include "DragAndDrop/AssetDragDropOp.h"
 #include "DragAndDrop/ActorDragDropOp.h"
 #include "DragAndDrop/ActorDragDropGraphEdOp.h"
@@ -57,12 +62,12 @@
 #include "DragAndDrop/FolderDragDropOp.h"
 #include "DragAndDrop/CompositeDragDropOp.h"
 #include "Widgets/Input/SSearchBox.h"
-#include "SSequencerTreeView.h"
+#include "MVVM/Views/SOutlinerView.h"
 #include "MovieSceneTrackEditor.h"
 #include "SSequencerSplitterOverlay.h"
 #include "SequencerHotspots.h"
 #include "SSequencerTimePanel.h"
-#include "VirtualTrackArea.h"
+#include "MVVM/ViewModels/VirtualTrackArea.h"
 #include "Framework/Commands/GenericCommands.h"
 #include "SequencerContextMenus.h"
 #include "Math/UnitConversion.h"
@@ -325,6 +330,8 @@
 PRAGMA_DISABLE_OPTIMIZATION
 void SSequencer::Construct(const FArguments& InArgs, TSharedRef<FSequencer> InSequencer)
 {
+	using namespace UE::Sequencer;
+
 	SequencerPtr = InSequencer;
 	bIsActiveTimerRegistered = false;
 	bUserIsSelecting = false;
@@ -446,8 +453,6 @@
 		SequencerWidgets.CreateTimeRangeSlider(TimeSliderControllerRef)
 	);
 
-	OnGetAddMenuContent = InArgs._OnGetAddMenuContent;
-	OnBuildCustomContextMenuForGuid = InArgs._OnBuildCustomContextMenuForGuid;
 	OnGetPlaybackSpeeds = InArgs._OnGetPlaybackSpeeds;
 
 	RootCustomization.AddMenuExtender = InArgs._AddMenuExtender;
@@ -467,27 +472,25 @@
 
 	TSharedRef<SScrollBar> PinnedAreaScrollBar = SNew(SScrollBar)
 		.Thickness(FVector2D(9.0f, 9.0f));
-	
-	SAssignNew(TrackOutliner, SSequencerTrackOutliner);
 
-	SAssignNew(PinnedTrackArea, SSequencerTrackArea, TimeSliderControllerRef, InSequencer);
-	SAssignNew(PinnedTreeView, SSequencerTreeView, InSequencer->GetNodeTree(), PinnedTrackArea.ToSharedRef())
+	SAssignNew(PinnedTrackArea, SSequencerTrackAreaView, InSequencer->GetViewModel()->GetTrackArea(), TimeSliderControllerRef);
+	SAssignNew(PinnedTreeView, SOutlinerView, InSequencer->GetViewModel()->GetOutliner(), PinnedTrackArea.ToSharedRef())
+		.SelectionHandler(InArgs._SelectionHandler)
 		.ExternalScrollbar(PinnedAreaScrollBar)
-		.Clipping(EWidgetClipping::ClipToBounds)
-		.OnGetContextMenuContent(FOnGetContextMenuContent::CreateSP(this, &SSequencer::GetContextMenuContent));
+		.Clipping(EWidgetClipping::ClipToBounds);
 
-	PinnedTrackArea->SetTreeView(PinnedTreeView);
+	PinnedTrackArea->SetOutliner(PinnedTreeView);
 	PinnedTrackArea->SetShowPinned(true);
 	PinnedTrackArea->SetIsSlave(true);
 	PinnedTreeView->SetShowPinned(true);
 
-	SAssignNew(TrackArea, SSequencerTrackArea, TimeSliderControllerRef, InSequencer);
-	SAssignNew(TreeView, SSequencerTreeView, InSequencer->GetNodeTree(), TrackArea.ToSharedRef())
+	SAssignNew(TrackArea, SSequencerTrackAreaView, InSequencer->GetViewModel()->GetTrackArea(), TimeSliderControllerRef);
+	SAssignNew(TreeView, SOutlinerView, InSequencer->GetViewModel()->GetOutliner(), TrackArea.ToSharedRef())
+		.SelectionHandler(InArgs._SelectionHandler)
 		.ExternalScrollbar(ScrollBar)
-		.Clipping(EWidgetClipping::ClipToBounds)
-		.OnGetContextMenuContent(FOnGetContextMenuContent::CreateSP(this, &SSequencer::GetContextMenuContent));
+		.Clipping(EWidgetClipping::ClipToBounds);
 
-	TrackArea->SetTreeView(TreeView);
+	TrackArea->SetOutliner(TreeView);
 
 	TreeView->AddSlaveTreeView(PinnedTreeView);
 
@@ -1097,7 +1100,6 @@
 	ApplySequencerCustomization(RootCustomization);
 
 	InSequencer->GetSelection().GetOnKeySelectionChanged().AddSP(this, &SSequencer::HandleKeySelectionChanged);
-	InSequencer->GetSelection().GetOnSectionSelectionChanged().AddSP(this, &SSequencer::HandleSectionSelectionChanged);
 	InSequencer->GetSelection().GetOnOutlinerNodeSelectionChanged().AddSP(this, &SSequencer::HandleOutlinerNodeSelectionChanged);
 
 	ResetBreadcrumbs();
@@ -1234,12 +1236,7 @@
 	WeakTickResolutionOptionsWindow = TickResolutionOptionsWindow;
 }
 
-const ISequencerEditTool* SSequencer::GetEditTool() const
-{
-	return TrackArea->GetEditTool();
-}
 
-
 void SSequencer::NotifyPostChange(const FPropertyChangedEvent& PropertyChangedEvent, FEditPropertyChain* PropertyThatChanged)
 {
 	// @todo sequencer: is this still needed?
@@ -1294,20 +1291,17 @@
 }
 
 
-void SSequencer::HandleSectionSelectionChanged()
-{
-}
-
-
 void SSequencer::HandleOutlinerNodeSelectionChanged()
 {
+	using namespace UE::Sequencer;
+
 	TSharedPtr<FSequencer> Sequencer = SequencerPtr.Pin();
 	if (!Sequencer.IsValid())
 	{
 		return;
 	}
 
-	const TSet<TSharedRef<FSequencerDisplayNode>>& SelectedDisplayNodes = Sequencer->GetSelection().GetSelectedOutlinerNodes();
+	const TSet<TWeakPtr<FViewModel>>& SelectedDisplayNodes = Sequencer->GetSelection().GetSelectedOutlinerItems();
 
 	TSharedPtr<FCurveEditor> CurveEditor = Sequencer->GetCurveEditor();
 	if (CurveEditor && CurveEditorTree)
@@ -1321,7 +1315,7 @@
 				SequencerSelectionCurveEditorFilter = MakeShared<FSequencerSelectionCurveFilter>();
 			}
 
-			SequencerSelectionCurveEditorFilter->Update(Sequencer->GetSelection().GetSelectedOutlinerNodes(), SequencerSettings->GetAutoExpandTreeView());
+			SequencerSelectionCurveEditorFilter->Update(Sequencer->GetSelection().GetSelectedOutlinerItems(), SequencerSettings->GetAutoExpandTreeView());
 
 			CurveEditor->GetTree()->AddFilter(SequencerSelectionCurveEditorFilter);
 		}
@@ -1334,20 +1328,15 @@
 
 		if (GetSequencerSettings()->ShouldSyncCurveEditorSelection())
 		{
-			//since Selection ticks before any filter is applied and the curve editor tree is udpated in FSequencer::Tick we need to sync 
-			//the curve editor on the next tick so the filter is applied first, otherwise the curve tree hierarchy won't be up to date
-			//and the curve may not be found.
-			GEditor->GetTimerManager()->SetTimerForNextTick([Sequencer, this]()
+			TSharedRef<FSequencerNodeTree> NodeTree = Sequencer->GetNodeTree();
+			CurveEditor->SuspendBroadcast();
+			// Clear the tree selection
+			CurveEditorTree->ClearSelection();
+			for (TWeakPtr<FViewModel> Node : SelectedDisplayNodes)
 			{
-				const TSet<TSharedRef<FSequencerDisplayNode>>& SelectedDisplayNodes = Sequencer->GetSelection().GetSelectedOutlinerNodes();
-				TSharedPtr<FCurveEditor> CurveEditor = Sequencer->GetCurveEditor();
-				TSharedRef<FSequencerNodeTree> NodeTree = Sequencer->GetNodeTree();
-				CurveEditor->SuspendBroadcast();
-				// Clear the tree selection
-				CurveEditorTree->ClearSelection();
-				for (TSharedRef<FSequencerDisplayNode> Node : SelectedDisplayNodes)
+				if (ICurveEditorTreeItemExtension* CurveEditorItem = ICastable::CastWeakPtr<ICurveEditorTreeItemExtension>(Node))
 				{
-					FCurveEditorTreeItemID CurveEditorTreeItem = NodeTree->FindCurveEditorTreeItem(Node);
+					FCurveEditorTreeItemID CurveEditorTreeItem = CurveEditorItem->GetCurveEditorItemID();
 					if (CurveEditorTreeItem != FCurveEditorTreeItemID::Invalid())
 					{
 						if (!CurveEditorTree->IsItemSelected(CurveEditorTreeItem))
@@ -1357,8 +1346,8 @@
 						}
 					}
 				}
-				CurveEditor->ResumeBroadcast();
-			});
+			}
+			CurveEditor->ResumeBroadcast();
 		}
 	}
 
@@ -1742,40 +1731,15 @@
 	}
 }
 
-void SSequencer::GetContextMenuContent(FMenuBuilder& MenuBuilder)
-{
-	// let toolkits populate the menu
-	MenuBuilder.BeginSection("MainMenu");
-	OnGetAddMenuContent.ExecuteIfBound(MenuBuilder, SequencerPtr.Pin().ToSharedRef());
-	MenuBuilder.EndSection();
-
-	// let track editors & object bindings populate the menu
-	TSharedPtr<FSequencer> Sequencer = SequencerPtr.Pin();
-
-	// Always create the section so that we afford extension
-	MenuBuilder.BeginSection("ObjectBindings");
-	if (Sequencer.IsValid())
-	{
-		Sequencer->BuildAddObjectBindingsMenu(MenuBuilder);
-	}
-	MenuBuilder.EndSection();
-
-	// Always create the section so that we afford extension
-	MenuBuilder.BeginSection("AddTracks");
-	if (Sequencer.IsValid())
-	{
-		Sequencer->BuildAddTrackMenu(MenuBuilder);
-	}
-	MenuBuilder.EndSection();
-}
-
-
 TSharedRef<SWidget> SSequencer::MakeAddMenu()
 {
+	using namespace UE::Sequencer;
+
 	TSharedPtr<FExtender> Extender = FExtender::Combine(AddMenuExtenders);
 	FMenuBuilder MenuBuilder(true, nullptr, Extender);
 	{
-		GetContextMenuContent(MenuBuilder);
+		TSharedPtr<FSequencer> Sequencer = SequencerPtr.Pin();
+		Sequencer->GetViewModel()->GetOutliner()->CastThisChecked<FSequencerOutlinerViewModel>()->BuildContextMenu(MenuBuilder);
 	}
 
 	return MenuBuilder.MakeWidget();
@@ -2050,7 +2014,6 @@
 	TSharedPtr<FSequencer> Sequencer = SequencerPtr.Pin();
 	Sequencer->GetNodeTree()->RemoveAllFilters();
 	OnEnableAllNodeGroupFilters(false);
-	OnEnableAllLevelFilters(false);
 }
 
 void SSequencer::OnEnableAllFilters()
@@ -3067,97 +3030,38 @@
 	return EActiveTimerReturnType::Stop;
 }
 
-
-void RestoreSelectionState(const TArray<TSharedRef<FSequencerDisplayNode>>& DisplayNodes, TSet<FString>& SelectedPathNames, FSequencerSelection& SequencerSelection)
-{
-	for (TSharedRef<FSequencerDisplayNode> DisplayNode : DisplayNodes)
-	{
-		if (SelectedPathNames.Contains(DisplayNode->GetPathName()))
-		{
-			SequencerSelection.AddToSelection(DisplayNode);
-		}
-
-		RestoreSelectionState(DisplayNode->GetChildNodes(), SelectedPathNames, SequencerSelection);
-	}
-}
-
 void SSequencer::UpdateLayoutTree()
 {
+	using namespace UE::Sequencer;
+
 	TrackArea->Empty();
 	PinnedTrackArea->Empty();
 
 	TSharedPtr<FSequencer> Sequencer = SequencerPtr.Pin();
 	if ( Sequencer.IsValid() )
 	{
-		// Cache the selected path names so selection can be restored after the update.
-		TSet<FString> SelectedPathNames;
-		// Cache selected keys
-		TSet<FSequencerSelectedKey> SelectedKeys = Sequencer->GetSelection().GetSelectedKeys();
-		TSet<TWeakObjectPtr<UMovieSceneSection> > SelectedSections = Sequencer->GetSelection().GetSelectedSections();
-
-		for (TSharedRef<const FSequencerDisplayNode> SelectedDisplayNode : Sequencer->GetSelection().GetSelectedOutlinerNodes().Array())
-		{
-			FString PathName = SelectedDisplayNode->GetPathName();
-			if ( FName(*PathName).IsNone() == false )
-			{
-				SelectedPathNames.Add(PathName);
-			}
-		}
-
-		// Add any additional paths that have been added by the user for nodes that may not exist yet but we want them to be selected
-		// after the node tree is updated and we restore selections.
-		SelectedPathNames.Append(AdditionalSelectionsToAdd);
-
-		// Suspend broadcasting selection changes because we don't want unnecessary rebuilds.
-		Sequencer->GetSelection().SuspendBroadcast();
-
 		// Update the node tree
 		Sequencer->GetNodeTree()->Update();
 
-		// Restore the selection state. This is still needed to apply the AdditionalSelectionsToAdd hack.
-		RestoreSelectionState(Sequencer->GetNodeTree()->GetRootNodes(), SelectedPathNames, SequencerPtr.Pin()->GetSelection());	// Update to actor selection.
-
 		// This must come after the selection state has been restored so that the tree and curve editor are populated with the correctly selected nodes
 		TreeView->Refresh();
-		
-		// If we've manually specified an additional selection to add it's because the item was newly created.
-		// Now that the treeview has been refreshed and selection restored, we'll try to focus the first item
-		// so that the view scrolls down when things are added to the bottom.
-		if (AdditionalSelectionsToAdd.Num() > 0)
-		{
-			FString NodePath = AdditionalSelectionsToAdd[0];
 
-			for (TSharedRef<FSequencerDisplayNode> Node : Sequencer->GetNodeTree()->GetAllNodes())
-			{
-				if (Node->GetPathName() == NodePath)
-				{
-					auto Parent = Node->GetParent();
-					while (Parent.IsValid())
-					{
-						TreeView->SetItemExpansion(Parent->AsShared(), true);
-						PinnedTreeView->SetItemExpansion(Parent->AsShared(), true);
-						Parent = Parent->GetParent();
-					}
-
-					TreeView->RequestScrollIntoView(Node);
-					break;
-				}
-			}
-		}
-
-		AdditionalSelectionsToAdd.Empty();
-
 		if (!NodePathToRename.IsEmpty())
 		{
-			for (TSharedRef<FSequencerDisplayNode> Node : Sequencer->GetNodeTree()->GetAllNodes())
+			TArray<TSharedPtr<FViewModel>> AllNodes;
+			Sequencer->GetNodeTree()->GetAllNodes(AllNodes);
+			for (TSharedPtr<FViewModel> Node : AllNodes)
 			{
-				if (Node->GetPathName() == NodePathToRename)
+				const FString CurNodePath = IOutlinerExtension::GetPathName(Node);
+				if (CurNodePath == NodePathToRename)
 				{
-					GEditor->GetTimerManager()->SetTimerForNextTick([Node]()
+					GEditor->GetTimerManager()->SetTimerForNextTick([Node]
 					{
-						Node->OnRenameRequested().Broadcast();
-					}
-					);
+						if (IRenameableExtension* Rename = Node->CastThis<IRenameableExtension>())
+						{
+							Rename->OnRenameRequested().Broadcast();
+						}
+					});
 					break;
 				}
 			}
@@ -3170,19 +3074,21 @@
 
 			TSharedPtr<FSequencerNodeTree> NodeTree = Sequencer->GetNodeTree();
 			const bool bHasSoloNodes = NodeTree->HasSoloNodes();
-			for (TSharedRef<FSequencerDisplayNode> Node : NodeTree->GetAllNodes())
+
+			for (TViewModelPtr<IOutlinerExtension> OutlinerItem
+				: Sequencer->GetNodeTree()->GetRootNode()->GetDescendantsOfType<IOutlinerExtension>())
 			{
-				if (Node->GetType() == ESequencerNode::Track)
+				TViewModelPtr<ITrackExtension> TrackNode = OutlinerItem.ImplicitCast();
+				if (TrackNode)
 				{
-					FSequencerTrackNode& TrackNode = static_cast<FSequencerTrackNode&>(Node.Get());
-					UMovieSceneTrack* Track = TrackNode.GetTrack();
-					bool bDisableEval = NodeTree->IsNodeMute(&Node.Get()) || (bHasSoloNodes && !NodeTree->IsNodeSolo(&Node.Get()));
-					if (TrackNode.GetSubTrackMode() == FSequencerTrackNode::ESubTrackMode::SubTrack)
+					UMovieSceneTrack* Track = TrackNode->GetTrack();
+					bool bDisableEval = NodeTree->IsNodeMute(OutlinerItem) || (bHasSoloNodes && !NodeTree->IsNodeSolo(OutlinerItem));
+					if (TrackNode->GetRowIndex() > 0)
 					{
-						if (bDisableEval != Track->IsRowEvalDisabled(TrackNode.GetRowIndex()))
+						if (bDisableEval != Track->IsRowEvalDisabled(TrackNode->GetRowIndex()))
 						{
 							Track->Modify();
-							Track->SetRowEvalDisabled(bDisableEval, TrackNode.GetRowIndex());
+							Track->SetRowEvalDisabled(bDisableEval, TrackNode->GetRowIndex());
 							bAnyChanged = true;
 						}
 					}
@@ -3202,9 +3108,6 @@
 				Sequencer->NotifyMovieSceneDataChanged(EMovieSceneDataChangeType::TrackValueChanged);
 			}
 		}
-
-		// Continue broadcasting selection changes
-		Sequencer->GetSelection().ResumeBroadcast();
 		
 		if (NodeGroupManager.IsValid())
 		{
@@ -3421,6 +3324,8 @@
 
 void SSequencer::OnAssetsDropped( const FAssetDragDropOp& DragDropOp )
 {
+	using namespace UE::Sequencer;
+
 	FSequencer& SequencerRef = *SequencerPtr.Pin();
 
 	TArray< UObject* > DroppedObjects;
@@ -3465,20 +3370,19 @@
 		GWarn->EndSlowTask();
 	}
 
-	const TSet< TSharedRef<FSequencerDisplayNode> >& SelectedNodes = SequencerPtr.Pin()->GetSelection().GetSelectedOutlinerNodes();
+	const TSet<TWeakPtr<FViewModel>>& SelectedNodes = SequencerPtr.Pin()->GetSelection().GetSelectedOutlinerItems();
 	FGuid TargetObjectGuid;
 	// if exactly one object node is selected, we have a target object guid
-	TSharedPtr<const FSequencerDisplayNode> DisplayNode;
+	TSharedPtr<const FViewModel> DisplayNode;
 	if (SelectedNodes.Num() == 1)
 	{
-		for (TSharedRef<const FSequencerDisplayNode> SelectedNode : SelectedNodes )
+		for (TWeakPtr<const FViewModel> SelectedNode : SelectedNodes )
 		{
-			DisplayNode = SelectedNode;
+			DisplayNode = SelectedNode.Pin();
 		}
-		if (DisplayNode.IsValid() && DisplayNode->GetType() == ESequencerNode::Object)
+		if (const FObjectBindingModel* ObjectBindingNode = DisplayNode->CastThis<FObjectBindingModel>())
 		{
-			TSharedPtr<const FSequencerObjectBindingNode> ObjectBindingNode = StaticCastSharedPtr<const FSequencerObjectBindingNode>(DisplayNode);
-			TargetObjectGuid = ObjectBindingNode->GetObjectBinding();
+			TargetObjectGuid = ObjectBindingNode->GetObjectGuid();
 		}
 	}
 
@@ -3676,12 +3580,18 @@
 }
 
 
-TSharedPtr<SSequencerTreeView> SSequencer::GetTreeView() const
+TSharedPtr<UE::Sequencer::SOutlinerView> SSequencer::GetTreeView() const
 {
 	return TreeView;
 }
 
 
+TSharedPtr<UE::Sequencer::SOutlinerView> SSequencer::GetPinnedTreeView() const
+{
+	return PinnedTreeView;
+}
+
+
 void SSequencer::OnSaveMovieSceneClicked()
 {
 	SequencerPtr.Pin()->SaveCurrentMovieScene();
@@ -3719,14 +3629,16 @@
 
 void SSequencer::StepToKey(bool bStepToNextKey, bool bCameraOnly)
 {
+	using namespace UE::Sequencer;
+
 	TSharedPtr<FSequencer> Sequencer = SequencerPtr.Pin();
 	if ( Sequencer.IsValid() )
 	{
-		TSet< TSharedRef<FSequencerDisplayNode> > Nodes;
+		TSet<TViewModelPtr<IOutlinerExtension>> Nodes;
 
 		if ( bCameraOnly )
 		{
-			TSet<TSharedRef<FSequencerDisplayNode>> RootNodes( Sequencer->GetNodeTree()->GetRootNodes() );
+			TSet<TViewModelPtr<IOutlinerExtension>> RootNodes(Sequencer->GetNodeTree()->GetRootNodes());
 
 			TSet<TWeakObjectPtr<AActor> > LockedActors;
 			for (FLevelEditorViewportClient* LevelVC : GEditor->GetLevelViewportClients())
@@ -3741,33 +3653,43 @@
 				}
 			}
 
-			for ( auto RootNode : RootNodes )
+			for (TViewModelPtr<IOutlinerExtension> RootNode : RootNodes)
 			{
-				TSharedRef<FSequencerObjectBindingNode> ObjectBindingNode = StaticCastSharedRef<FSequencerObjectBindingNode>( RootNode );
-
-				for (TWeakObjectPtr<>& Object : Sequencer->FindObjectsInCurrentSequence(ObjectBindingNode->GetObjectBinding()))
+				TSharedPtr<FObjectBindingModel> ObjectBindingNode = RootNode.ImplicitCast();
+				if (ObjectBindingNode)
 				{
-					AActor* RuntimeActor = Cast<AActor>( Object.Get() );
-					if ( RuntimeActor != nullptr && LockedActors.Contains( RuntimeActor ) )
+					for (TWeakObjectPtr<>& Object : Sequencer->FindObjectsInCurrentSequence(ObjectBindingNode->GetObjectGuid()))
 					{
-						Nodes.Add( RootNode );
+						AActor* RuntimeActor = Cast<AActor>( Object.Get() );
+						if ( RuntimeActor != nullptr && LockedActors.Contains( RuntimeActor ) )
+						{
+							Nodes.Add(RootNode);
+						}
 					}
 				}
 			}
 		}
 		else
 		{
-			const TSet< TSharedRef<FSequencerDisplayNode> >& SelectedNodes = Sequencer->GetSelection().GetSelectedOutlinerNodes();
-			Nodes = SelectedNodes;
-
+			Nodes.Reserve(Sequencer->GetSelection().GetSelectedOutlinerItems().Num());
+			for (const TWeakPtr<FViewModel>& WeakModel : Sequencer->GetSelection().GetSelectedOutlinerItems())
+			{
+				TViewModelPtr<IOutlinerExtension> Model = WeakModel.Pin();
+				if (Model)
+				{
+					Nodes.Add(Model);
+				}
+			}
 			if ( Nodes.Num() == 0 )
 			{
-				TSet<TSharedRef<FSequencerDisplayNode>> RootNodes( Sequencer->GetNodeTree()->GetRootNodes() );
-				for ( auto RootNode : RootNodes )
+				TSet<TViewModelPtr<IOutlinerExtension>> RootNodes(Sequencer->GetNodeTree()->GetRootNodes());
+				for (TViewModelPtr<IOutlinerExtension> RootNode : RootNodes)
 				{
-					Nodes.Add( RootNode );
-
-					SequencerHelpers::GetDescendantNodes( RootNode, Nodes );
+					Nodes.Add(RootNode);
+					for (TViewModelPtr<IOutlinerExtension> Child : RootNode.AsModel()->GetDescendantsOfType<IOutlinerExtension>())
+					{
+						Nodes.Add(Child);
+					}
 				}
 			}
 		}
@@ -3781,14 +3703,14 @@
 			TOptional<FFrameNumber> StepToTime;
 
 			auto It = Nodes.CreateConstIterator();
-			bool bExpand = !( *It ).Get().IsExpanded();
+			bool bExpand = !( *It )->IsExpanded();
 
-			for ( auto Node : Nodes )
+			for (TViewModelPtr<IOutlinerExtension> Node : Nodes)
 			{
 				TArray<FFrameNumber> AllTimes;
 
 				TSet<TSharedPtr<IKeyArea>> KeyAreas;
-				SequencerHelpers::GetAllKeyAreas( Node, KeyAreas );
+				SequencerHelpers::GetAllKeyAreas( Node.AsModel(), KeyAreas );
 				for ( TSharedPtr<IKeyArea> KeyArea : KeyAreas )
 				{
 					if (KeyArea->GetOwningSection())
@@ -3798,7 +3720,7 @@
 				}
 
 				TSet<TWeakObjectPtr<UMovieSceneSection> > Sections;
-				SequencerHelpers::GetAllSections( Node, Sections );
+				SequencerHelpers::GetAllSections( Node.AsModel(), Sections );
 
 				for ( TWeakObjectPtr<UMovieSceneSection> Section : Sections )
 				{
@@ -4037,22 +3959,29 @@
 }
 
 
-FVirtualTrackArea SSequencer::GetVirtualTrackArea(const SSequencerTrackArea* InTrackArea) const
+UE::Sequencer::FVirtualTrackArea SSequencer::GetVirtualTrackArea(const UE::Sequencer::STrackAreaView* InTrackArea) const
 {
-	const SSequencerTrackArea* TargetTrackArea = TrackArea.Get();
-	TSharedPtr<SSequencerTreeView> TargetTreeView = TreeView;
+	using namespace UE::Sequencer;
+
+	TSharedPtr<UE::Sequencer::FEditorViewModel> EditorViewModel = SequencerPtr.Pin()->GetViewModel();
+	FTrackAreaViewModel& TrackAreaViewModel = *EditorViewModel->GetTrackArea().Get();
+
+	const STrackAreaView* TargetTrackArea = TrackArea.Get();
+	TSharedPtr<SOutlinerView> TargetTreeView = TreeView;
 	
 	if (InTrackArea != nullptr)
 	{
 		TargetTrackArea = InTrackArea;
-		TargetTreeView = TargetTrackArea->GetTreeView().Pin();
+		//TargetTreeView = TargetTrackArea->GetTreeView().Pin();
 	}
 
-	return FVirtualTrackArea(*SequencerPtr.Pin(), *TargetTreeView.Get(), TargetTrackArea->GetCachedGeometry());
+	return FVirtualTrackArea(TrackAreaViewModel, *TargetTreeView.Get(), TargetTrackArea->GetCachedGeometry());
 }
 
 FPasteContextMenuArgs SSequencer::GeneratePasteArgs(FFrameNumber PasteAtTime, TSharedPtr<FMovieSceneClipboard> Clipboard)
 {
+	using namespace UE::Sequencer;
+
 	TSharedPtr<FSequencer> Sequencer = SequencerPtr.Pin();
 	if (GetSequencerSettings()->GetIsSnapEnabled())
 	{
@@ -4070,7 +3999,7 @@
 	FVirtualTrackArea VirtualTrackArea = GetVirtualTrackArea();
 
 	// Paste into the currently selected sections, or hit test the mouse position as a last resort
-	TArray<TSharedRef<FSequencerDisplayNode>> PasteIntoNodes;
+	TArray<TViewModelPtr<IOutlinerExtension>> PasteIntoNodes;
 	{
 		TSet<TWeakObjectPtr<UMovieSceneSection>> Sections = Sequencer->GetSelection().GetSelectedSections();
 		
@@ -4081,28 +4010,37 @@
 
 		for (TWeakObjectPtr<UMovieSceneSection> WeakSection : Sections)
 		{
-			if (TOptional<FSectionHandle> Handle = Sequencer->GetNodeTree()->GetSectionHandle(WeakSection.Get()))
+			if (TSharedPtr<FSectionModel> Handle = Sequencer->GetNodeTree()->GetSectionModel(WeakSection.Get()))
 			{
-				PasteIntoNodes.Add(Handle->GetTrackNode());
+				TViewModelPtr<IOutlinerExtension> TrackModel = Handle->GetParentTrackModel().ImplicitCast();
+				if (TrackModel)
+				{
+					PasteIntoNodes.Add(TrackModel);
+				}
 			}
 		}
 
-		auto SelectedNodes = Sequencer->GetSelection().GetSelectedOutlinerNodes();
-		for (const TSharedRef<FSequencerDisplayNode>& SelectedNode : SelectedNodes)
+		TArray<TSharedRef<FViewModel>> SelectedNodes;
+		Sequencer->GetSelection().GetSelectedOutlinerItems(SelectedNodes);
+		for (const TSharedRef<FViewModel>& SelectedNode : SelectedNodes)
 		{
-			if (SelectedNode->GetType() == ESequencerNode::Category || SelectedNode->GetType() == ESequencerNode::Track || SelectedNode->GetType() == ESequencerNode::KeyArea)
+			if (SelectedNode->IsA<FCategoryGroupModel>() || SelectedNode->IsA<ITrackExtension>() || SelectedNode->IsA<FChannelGroupModel>())
 			{
-				PasteIntoNodes.Add(SelectedNode);
+				TViewModelPtr<IOutlinerExtension> TrackModel(SelectedNode);
+				if (TrackModel)
+				{
+					PasteIntoNodes.Add(TrackModel);
+				}
 			}
 		}
 	}
 
 	if (PasteIntoNodes.Num() == 0)
 	{
-		TSharedPtr<FSequencerDisplayNode> Node = VirtualTrackArea.HitTestNode(LocalMousePosition.Y);
+		TSharedPtr<FViewModel> Node = VirtualTrackArea.HitTestNode(LocalMousePosition.Y);
 		if (Node.IsValid())
 		{
-			PasteIntoNodes.Add(Node.ToSharedRef());
+			PasteIntoNodes.Add(FViewModelPtr(Node).ImplicitCast());
 		}
 	}
 
@@ -4111,8 +4049,10 @@
 
 void SSequencer::OnPaste()
 {
+	using namespace UE::Sequencer;
+
 	TSharedPtr<FSequencer> Sequencer = SequencerPtr.Pin();
-	TSet<TSharedRef<FSequencerDisplayNode>> SelectedNodes = Sequencer->GetSelection().GetSelectedOutlinerNodes();
+	const TSet<TWeakPtr<FViewModel>>& SelectedNodes = Sequencer->GetSelection().GetSelectedOutlinerItems();
 	if (SelectedNodes.Num() == 0)
 	{
 		if (!OpenPasteMenu())
@@ -4127,8 +4067,6 @@
 			OpenPasteMenu();
 		}
 	}
-
-	
 }
 
 bool SSequencer::CanPaste()
@@ -4229,11 +4167,6 @@
 	return GetSequencerSettings()->ShouldShowDebugVisualization() ? EVisibility::Visible : EVisibility::Collapsed;
 }
 
-void SSequencer::BuildCustomContextMenuForGuid(FMenuBuilder& MenuBuilder, FGuid ObjectBinding)
-{
-	OnBuildCustomContextMenuForGuid.ExecuteIfBound(MenuBuilder, ObjectBinding);
-}
-
 double SSequencer::GetSpinboxDelta() const
 {
 	TSharedPtr<FSequencer> Sequencer = SequencerPtr.Pin();
diff --git a/Engine/Source/Editor/Sequencer/Private/SSequencer.h b/Engine/Source/Editor/Sequencer/Private/SSequencer.h
--- a/Engine/Source/Editor/Sequencer/Private/SSequencer.h
+++ b/Engine/Source/Editor/Sequencer/Private/SSequencer.h
@@ -26,14 +26,9 @@
 class FClassDragDropOp;
 class FMovieSceneClipboard;
 class FSequencerTimeSliderController;
-class FVirtualTrackArea;
-class ISequencerEditTool;
 class SCurveEditorTree;
-class SSequencerTrackArea;
-class SSequencerTrackOutliner;
 class SSequencerTransformBox;
 class SSequencerStretchBox;
-class SSequencerTreeView;
 class SCurveEditorPanel;
 class SDockTab;
 class SWindow;
@@ -45,6 +40,19 @@
 struct FSequencerCustomizationInfo;
 struct FSequencerSelectionCurveFilter;
 
+namespace UE
+{
+namespace Sequencer
+{
+
+	class IOutlinerSelectionHandler;
+	class SOutlinerView;
+	class STrackAreaView;
+	class FVirtualTrackArea;
+
+} // namespace Sequencer
+} // namespace UE
+
 namespace SequencerLayoutConstants
 {
 	/** The amount to indent child nodes of the layout tree */
@@ -242,12 +250,6 @@
 		/** Called when the user changes the scrub position */
 		SLATE_EVENT( FOnScrubPositionChanged, OnScrubPositionChanged )
 
-		/** Called to populate the add combo button in the toolbar. */
-		SLATE_EVENT( FOnGetAddMenuContent, OnGetAddMenuContent )
-
-		/** Called when object is clicked. */
-		SLATE_EVENT(FOnBuildCustomContextMenuForGuid, OnBuildCustomContextMenuForGuid)
-			
 		/** Called when any widget contained within sequencer has received focus */
 		SLATE_EVENT( FSimpleDelegate, OnReceivedFocus )
 
@@ -276,8 +278,11 @@
 		SLATE_ARGUMENT( TSharedPtr<FExtender>, AddMenuExtender )
 
 		/** Extender to use for the toolbar. */
-		SLATE_ARGUMENT(TSharedPtr<FExtender>, ToolbarExtender)
+		SLATE_ARGUMENT( TSharedPtr<FExtender>, ToolbarExtender )
 
+		/** Selection handler to pass to outliner view */
+		SLATE_ATTRIBUTE( TSharedPtr<UE::Sequencer::IOutlinerSelectionHandler>, SelectionHandler )
+
 	SLATE_END_ARGS()
 
 
@@ -326,23 +331,23 @@
 	void OnCurveEditorVisibilityChanged(bool bShouldBeVisible);
 
 	/** Access the tree view for this sequencer */
-	TSharedPtr<SSequencerTreeView> GetTreeView() const;
+	TSharedPtr<UE::Sequencer::SOutlinerView> GetTreeView() const;
 
+	/** Access the pinned tree view for this sequencer */
+	TSharedPtr<UE::Sequencer::SOutlinerView> GetPinnedTreeView() const;
+
 	/** 
 	 * Generate a helper structure that can be used to transform between phsyical space and virtual space in the track area
 	 *
 	 * @param InTrackArea	(optional) The track area to generate helper structure for, if not specified the main track area will be used.
 	 */
-	FVirtualTrackArea GetVirtualTrackArea(const SSequencerTrackArea* InTrackArea = nullptr) const;
+	UE::Sequencer::FVirtualTrackArea GetVirtualTrackArea(const UE::Sequencer::STrackAreaView* InTrackArea = nullptr) const;
 
 	/** Access this widget's track area widget */
-	TSharedPtr<SSequencerTrackArea> GetTrackAreaWidget() const { return TrackArea; }
+	TSharedPtr<UE::Sequencer::STrackAreaView> GetTrackAreaWidget() const { return TrackArea; }
 
 	/** @return a numeric type interface that will parse and display numbers as frames and times correctly */
 	TSharedRef<INumericTypeInterface<double>> GetNumericTypeInterface() const;
-	
-	/** Access the currently active track area edit tool */
-	const ISequencerEditTool* GetEditTool() const;
 
 	void OpenTickResolutionOptions();
 
@@ -393,9 +398,6 @@
 	/** Handles key selection changes. */
 	void HandleKeySelectionChanged();
 
-	/** Handles section selection changes. */
-	void HandleSectionSelectionChanged();
-
 	/** Handles changes to the selected outliner nodes. */
 	void HandleOutlinerNodeSelectionChanged();
 
@@ -418,6 +420,7 @@
 	/** Makes add button. */
 	TSharedRef<SWidget> MakeAddButton();
 
+	/** Makes filter button */
 	TSharedRef<SWidget> MakeFilterButton();
 
 	/** Makes the add menu for the toolbar. */
@@ -624,9 +627,6 @@
 	/** Generate a paste menu args structure */
 	struct FPasteContextMenuArgs GeneratePasteArgs(FFrameNumber PasteAtTime, TSharedPtr<FMovieSceneClipboard> Clipboard = nullptr);
 
-	/** Execute custom context menu if passed in the FSequencerViewParams  */
-	void BuildCustomContextMenuForGuid(FMenuBuilder& MenuBuilder, FGuid ObjectBinding);
-
 	/** This adds the specified path to the selection set to be restored the next time the tree view is refreshed. */
 	void AddAdditionalPathToSelectionSet(const FString& Path) { AdditionalSelectionsToAdd.Add(Path); }
 
@@ -653,16 +653,13 @@
 
 	/** Filter Status Bar */
 	TSharedPtr<SSequencerTreeFilterStatusBar> SequencerTreeFilterStatusBar;
-	
+
 	/** Section area widget */
-	TSharedPtr<SSequencerTrackArea> TrackArea;
+	TSharedPtr<UE::Sequencer::STrackAreaView> TrackArea;
 
 	/** Section area widget for pinned tracks*/
-	TSharedPtr<SSequencerTrackArea> PinnedTrackArea;
+	TSharedPtr<UE::Sequencer::STrackAreaView> PinnedTrackArea;
 
-	/** Outliner widget */
-	TSharedPtr<SSequencerTrackOutliner> TrackOutliner;
-
 	/** Curve editor tree widget */
 	TSharedPtr<SCurveEditorTree> CurveEditorTree;
 
@@ -685,10 +682,10 @@
 	TSharedPtr<STextBlock> LoopIndexDisplay;
 
 	/** The sequencer tree view responsible for the outliner and track areas */
-	TSharedPtr<SSequencerTreeView> TreeView;
+	TSharedPtr<UE::Sequencer::SOutlinerView> TreeView;
 
 	/** The sequencer tree view for pinned tracks */
-	TSharedPtr<SSequencerTreeView> PinnedTreeView;
+	TSharedPtr<UE::Sequencer::SOutlinerView> PinnedTreeView;
 
 	/** Dropdown for selecting breadcrumbs */
 	TSharedPtr<class SComboButton> BreadcrumbPickerButton;
@@ -729,11 +726,6 @@
 	/** Time slider controller for this sequencer */
 	TSharedPtr<FSequencerTimeSliderController> TimeSliderController;
 
-	FOnGetAddMenuContent OnGetAddMenuContent;
-
-	/** Called when object is clicked in track list */
-	FOnBuildCustomContextMenuForGuid OnBuildCustomContextMenuForGuid;
-
 	/** Called when the user has begun dragging the selection selection range */
 	FSimpleDelegate OnSelectionRangeBeginDrag;
 
diff --git a/Engine/Source/Editor/Sequencer/Private/SSequencerDebugVisualizer.cpp b/Engine/Source/Editor/Sequencer/Private/SSequencerDebugVisualizer.cpp
--- a/Engine/Source/Editor/Sequencer/Private/SSequencerDebugVisualizer.cpp
+++ b/Engine/Source/Editor/Sequencer/Private/SSequencerDebugVisualizer.cpp
@@ -16,6 +16,7 @@
 #include "EntitySystem/MovieSceneEntitySystemLinker.h"
 #include "EntitySystem/EntityAllocationIterator.h"
 #include "EntitySystem/MovieSceneComponentRegistry.h"
+#include "TimeToPixel.h"
 
 #define LOCTEXT_NAMESPACE "SSequencerDebugVisualizer"
 
diff --git a/Engine/Source/Editor/Sequencer/Private/SSequencerGroupManager.cpp b/Engine/Source/Editor/Sequencer/Private/SSequencerGroupManager.cpp
--- a/Engine/Source/Editor/Sequencer/Private/SSequencerGroupManager.cpp
+++ b/Engine/Source/Editor/Sequencer/Private/SSequencerGroupManager.cpp
@@ -2,11 +2,12 @@
 
 #include "SSequencerGroupManager.h"
 
+#include "MVVM/Extensions/IOutlinerExtension.h"
+#include "MVVM/Extensions/IGroupableExtension.h"
 #include "Sequencer.h"
 #include "MovieSceneSequence.h"
 #include "MovieScene.h"
-#include "DisplayNodes/SequencerDisplayNode.h"
-#include "SequencerDisplayNodeDragDropOp.h"
+#include "SequencerOutlinerItemDragDropOp.h"
 #include "SequencerUtilities.h"
 #include "SlateOptMacros.h"
 #include "Widgets/SNullWidget.h"
@@ -245,7 +246,9 @@
 
 	TOptional<EItemDropZone> OnCanAcceptDrop(const FDragDropEvent& DragDropEvent, EItemDropZone InItemDropZone, TSharedPtr<FSequencerNodeGroupTreeNode> SequencerGroupTreeNode)
 	{
-		TSharedPtr<FSequencerDisplayNodeDragDropOp> DragDropOp = DragDropEvent.GetOperationAs<FSequencerDisplayNodeDragDropOp>();
+		using namespace UE::Sequencer;
+
+		TSharedPtr<FSequencerOutlinerDragDropOp> DragDropOp = DragDropEvent.GetOperationAs<FSequencerOutlinerDragDropOp>();
 		if (DragDropOp.IsValid())
 		{
 			DragDropOp->ResetToDefaultToolTip();
@@ -254,7 +257,7 @@
 			if (SequencerGroupTreeNode->GetType() == FSequencerNodeGroupTreeNode::Type::GroupNode)
 			{
 				AllowedDropZone = EItemDropZone::OntoItem;
-				DragDropOp->CurrentHoverText = FText::Format(LOCTEXT("DragDropAddItemsFormat", "Add {0} item(s)"), FText::AsNumber(DragDropOp->GetDraggedNodes().Num()));
+				DragDropOp->CurrentHoverText = FText::Format(LOCTEXT("DragDropAddItemsFormat", "Add {0} item(s)"), FText::AsNumber(DragDropOp->GetDraggedViewModels().Num()));
 			}
 
 			if (AllowedDropZone.IsSet() == false)
@@ -268,7 +271,9 @@
 
 	FReply OnAcceptDrop(const FDragDropEvent& DragDropEvent, EItemDropZone InItemDropZone, TSharedPtr<FSequencerNodeGroupTreeNode> SequencerGroupTreeNode)
 	{
-		TSharedPtr<FSequencerDisplayNodeDragDropOp> DragDropOp = DragDropEvent.GetOperationAs<FSequencerDisplayNodeDragDropOp>();
+		using namespace UE::Sequencer;
+
+		TSharedPtr<FSequencerOutlinerDragDropOp> DragDropOp = DragDropEvent.GetOperationAs<FSequencerOutlinerDragDropOp>();
 		if (DragDropOp.IsValid())
 		{
 			TSharedPtr<SSequencerGroupManager> SequencerGroupManager = WeakSequencerGroupManager.Pin();
@@ -278,7 +283,12 @@
 				if (SequencerGroupTreeNode->GetType() == FSequencerNodeGroupTreeNode::Type::GroupNode)
 				{
 					TSharedPtr<FSequencerNodeGroupNode> NodeGroupNode = StaticCastSharedPtr<FSequencerNodeGroupNode>(SequencerGroupTreeNode);
-					Sequencer->AddNodesToExistingNodeGroup(DragDropOp->GetDraggedNodes(), NodeGroupNode->Group);
+					TArray<TWeakPtr<FViewModel>> DraggedNodes;
+					for (const TWeakViewModelPtr<IOutlinerExtension>& DraggedNode : DragDropOp->GetDraggedViewModels())
+					{
+						DraggedNodes.Add(DraggedNode);
+					}
+					Sequencer->AddNodesToExistingNodeGroup(DraggedNodes, NodeGroupNode->Group);
 
 					return FReply::Handled();
 				}
@@ -356,6 +366,8 @@
 
 void SSequencerGroupManager::UpdateTree()
 {
+	using namespace UE::Sequencer;
+
 	TSharedPtr<FSequencer> Sequencer = WeakSequencer.Pin(); 
 	UMovieScene* MovieScene = GetMovieScene();
 
@@ -387,10 +399,10 @@
 		NodeGroupsTree.Add(SequencerGroupNode);
 		for (const FString& NodePath : NodeGroup->GetNodes())
 		{
-			const FSequencerDisplayNode* Node = NodeTree->GetNodeAtPath(NodePath);
+			TViewModelPtr<IOutlinerExtension> Node = NodeTree->GetNodeAtPath(NodePath);
 			if (Node)
 			{
-				TSharedPtr<FSequencerGroupItemNode> SequencerGroupItemNode = MakeShared<FSequencerGroupItemNode>(Node->GetDisplayName(), NodePath, SequencerGroupNode);
+				TSharedPtr<FSequencerGroupItemNode> SequencerGroupItemNode = MakeShared<FSequencerGroupItemNode>(Node->GetLabel(), NodePath, SequencerGroupNode);
 				SequencerGroupNode->Children.Add(SequencerGroupItemNode);
 				AllNodeGroupItems.Add(NodePath);
 			}
@@ -456,6 +468,8 @@
 
 const FSlateBrush* SSequencerGroupManager::GetIconBrush(TSharedPtr<FSequencerNodeGroupTreeNode> NodeGroupTreeNode) const
 {
+	using namespace UE::Sequencer;
+
 	if (NodeGroupTreeNode->GetType() == FSequencerNodeGroupTreeNode::Type::ItemNode)
 	{
 		TSharedPtr<FSequencerGroupItemNode> SequencerGroupItemNode = StaticCastSharedPtr<FSequencerGroupItemNode>(NodeGroupTreeNode);
@@ -464,10 +478,16 @@
 		{
 			return nullptr;
 		}
+
 		TSharedRef<FSequencerNodeTree> NodeTree = Sequencer->GetNodeTree();
-		const FSequencerDisplayNode* Node = NodeTree->GetNodeAtPath(SequencerGroupItemNode->Path);
-		
-		return Node ? Node->GetIconBrush() : nullptr;
+
+		// @todo_sequencer_mvvm: This is literally walking the entire tree to find a node by its path.
+		//                       Worse still, it is doing so every frame, for every group node :/
+		TViewModelPtr<IOutlinerExtension> OutlinerNode = NodeTree->GetNodeAtPath(SequencerGroupItemNode->Path);
+		if (OutlinerNode)
+		{
+			return OutlinerNode->GetIconBrush();
+		}
 	}
 
 	return nullptr;
@@ -621,6 +641,8 @@
 
 void SSequencerGroupManager::SelectItemsSelectedInSequencer()
 {
+	using namespace UE::Sequencer;
+
 	if (bSynchronizingSelection)
 	{
 		return;
@@ -634,21 +656,29 @@
 		return;
 	}
 
-	const TSet<TSharedRef<FSequencerDisplayNode>>& SequencerSelectedNodes = Sequencer->GetSelection().GetSelectedOutlinerNodes();
+	const TSet<TWeakPtr<FViewModel>>& SelectedModels = Sequencer->GetSelection().GetSelectedOutlinerItems();
 
+	TStringBuilder<128> TempString;
+
 	// Build a list of the nodepaths that we want to consider for selection
 	TSet<FString> NodesPathsToSelect;
-	for (const TSharedRef<const FSequencerDisplayNode> Node : SequencerSelectedNodes)
+	for (TWeakPtr<FViewModel> WeakModel : SelectedModels)
 	{
-		const FSequencerDisplayNode* BaseNode = Node->GetBaseNode();
-		ESequencerNode::Type NodeType = BaseNode->GetType();
+		TSharedPtr<FViewModel> Model = WeakModel.Pin();
+		if (!Model)
+		{
+			continue;
+		}
 
-		if (NodeType == ESequencerNode::Track || NodeType == ESequencerNode::Object || NodeType == ESequencerNode::Folder)
+		TViewModelPtr<IGroupableExtension> Groupable = Model->FindAncestorOfType<IGroupableExtension>(true);
+		if (Groupable)
 		{
-			FString NodePath = BaseNode->GetPathName();
-			if (AllNodeGroupItems.Contains(NodePath))
+			TempString.Reset();
+			Groupable->GetIdentifierForGrouping(TempString);
+
+			if (AllNodeGroupItems.Contains(TempString.ToString()))
 			{
-				NodesPathsToSelect.Add(NodePath);
+				NodesPathsToSelect.Add(TempString.ToString());
 			}
 		}
 	}
