Sequencer: MVVM2 branch and Layer Bars

diff --git a/Engine/Source/Editor/Sequencer/Private/SSequencerSection.cpp b/Engine/Source/Editor/Sequencer/Private/SSequencerSection.cpp
--- a/Engine/Source/Editor/Sequencer/Private/SSequencerSection.cpp
+++ b/Engine/Source/Editor/Sequencer/Private/SSequencerSection.cpp
@@ -1,6 +1,14 @@
 // Copyright Epic Games, Inc. All Rights Reserved.
 
 #include "SSequencerSection.h"
+#include "MVVM/ViewModels/SectionModel.h"
+#include "MVVM/ViewModels/TrackModel.h"
+#include "MVVM/ViewModels/CategoryModel.h"
+#include "MVVM/ViewModels/ViewModelIterators.h"
+#include "MVVM/ViewModels/SequencerEditorViewModel.h"
+#include "MVVM/Views/ITrackAreaHotspot.h"
+#include "MVVM/Views/STrackAreaView.h"
+#include "MVVM/Extensions/IObjectBindingExtension.h"
 #include "Rendering/DrawElements.h"
 #include "Styling/AppStyle.h"
 #include "SequencerSelectionPreview.h"
@@ -13,20 +21,23 @@
 #include "CommonMovieSceneTools.h"
 #include "ISequencerEditTool.h"
 #include "ISequencerSection.h"
-#include "ISequencerHotspot.h"
 #include "SequencerHotspots.h"
 #include "Widgets/SOverlay.h"
 #include "SequencerAddKeyOperation.h"
-#include "SequencerObjectBindingNode.h"
-#include "SequencerSectionCategoryNode.h"
-#include "SequencerSectionKeyAreaNode.h"
 #include "MovieScene.h"
 #include "Fonts/FontCache.h"
 #include "Framework/Application/SlateApplication.h"
 #include "KeyDrawParams.h"
 #include "MovieSceneTimeHelpers.h"
+#include "Tracks/MovieScenePropertyTrack.h"
 #include "Sections/MovieSceneSubSection.h"
+#include "Generators/MovieSceneEasingFunction.h"
 
+namespace UE
+{
+namespace Sequencer
+{
+
 double SSequencerSection::SectionSelectionThrobEndTime = 0;
 double SSequencerSection::KeySelectionThrobEndTime = 0;
 
@@ -56,26 +67,26 @@
 
 struct FSequencerSectionPainterImpl : FSequencerSectionPainter
 {
-	FSequencerSectionPainterImpl(FSequencer& InSequencer, UMovieSceneSection& InSection, FSlateWindowElementList& _OutDrawElements, const FGeometry& InSectionGeometry, const SSequencerSection& InSectionWidget)
+	FSequencerSectionPainterImpl(FSequencer& InSequencer, TSharedPtr<FSectionModel> InSection, FSlateWindowElementList& _OutDrawElements, const FGeometry& InSectionGeometry, const SSequencerSection& InSectionWidget)
 		: FSequencerSectionPainter(_OutDrawElements, InSectionGeometry, InSection)
 		, Sequencer(InSequencer)
 		, SectionWidget(InSectionWidget)
-		, TimeToPixelConverter(ConstructTimeConverterForSection(SectionGeometry, InSection, Sequencer))
+		, TimeToPixelConverter(ConstructTimeConverterForSection(SectionGeometry, *InSection->GetSection(), Sequencer))
 	{
 		CalculateSelectionColor();
 
-		const ISequencerEditTool* EditTool = InSequencer.GetEditTool();
+		const ISequencerEditTool* EditTool = InSequencer.GetViewModel()->GetTrackArea()->GetEditTool();
 		Hotspot = EditTool ? EditTool->GetDragHotspot() : nullptr;
 		if (!Hotspot)
 		{
-			Hotspot = Sequencer.GetHotspot().Get();
+			Hotspot = InSequencer.GetViewModel()->GetTrackArea()->GetHotspot();
 		}
 	}
 
 	FLinearColor GetFinalTintColor(const FLinearColor& Tint) const
 	{
-		FLinearColor FinalTint = FSequencerSectionPainter::BlendColor(Tint);
-		if (bIsHighlighted && Section.GetRange() != TRange<FFrameNumber>::All())
+		FLinearColor FinalTint = STrackAreaView::BlendDefaultTrackColor(Tint);
+		if (bIsHighlighted && SectionModel->GetRange() != TRange<FFrameNumber>::All())
 		{
 			float Lum = FinalTint.GetLuminance() * 0.2f;
 			FinalTint = FinalTint + FLinearColor(Lum, Lum, Lum, 0.f);
@@ -88,6 +99,8 @@
 
 	virtual int32 PaintSectionBackground(const FLinearColor& Tint) override
 	{
+		UMovieSceneSection* SectionObject = SectionModel->GetSection();
+
 		const ESlateDrawEffect DrawEffects = bParentEnabled
 			? ESlateDrawEffect::None
 			: ESlateDrawEffect::DisabledEffect;
@@ -99,8 +112,8 @@
 		FLinearColor FinalTint = GetFinalTintColor(Tint);
 
 		// Offset lower bounds and size for infinte sections so we don't draw the rounded border on the visible area
-		const float InfiniteLowerOffset = Section.HasStartFrame() ? 0.f : 100.f;
-		const float InfiniteSizeOffset  = InfiniteLowerOffset + (Section.HasEndFrame() ? 0.f : 100.f);
+		const float InfiniteLowerOffset = SectionObject->HasStartFrame() ? 0.f : 100.f;
+		const float InfiniteSizeOffset  = InfiniteLowerOffset + (SectionObject->HasEndFrame() ? 0.f : 100.f);
 
 		FPaintGeometry PaintGeometry = SectionGeometry.ToPaintGeometry(
 			SectionGeometry.GetLocalSize() + FVector2D(InfiniteSizeOffset, 0.f),
@@ -115,10 +128,10 @@
 			static const FSlateBrush* PreRollBrush = FAppStyle::GetBrush("Sequencer.Section.PreRoll");
 			float BrushHeight = 16.f, BrushWidth = 10.f;
 
-			if (Section.HasStartFrame())
+			if (SectionObject->HasStartFrame())
 			{
-				FFrameNumber SectionStartTime = Section.GetInclusiveStartFrame();
-				FFrameNumber PreRollStartTime = SectionStartTime - Section.GetPreRollFrames();
+				FFrameNumber SectionStartTime = SectionObject->GetInclusiveStartFrame();
+				FFrameNumber PreRollStartTime = SectionStartTime - SectionObject->GetPreRollFrames();
 
 				const float PreRollPx = TimeToPixelConverter.FrameToPixel(SectionStartTime) - TimeToPixelConverter.FrameToPixel(PreRollStartTime);
 				if (PreRollPx > 0)
@@ -141,10 +154,10 @@
 				}
 			}
 
-			if (Section.HasEndFrame())
+			if (SectionObject->HasEndFrame())
 			{
-				FFrameNumber SectionEndTime  = Section.GetExclusiveEndFrame();
-				FFrameNumber PostRollEndTime = SectionEndTime + Section.GetPostRollFrames();
+				FFrameNumber SectionEndTime  = SectionObject->GetExclusiveEndFrame();
+				FFrameNumber PostRollEndTime = SectionEndTime + SectionObject->GetPostRollFrames();
 
 				const float PostRollPx = TimeToPixelConverter.FrameToPixel(PostRollEndTime) - TimeToPixelConverter.FrameToPixel(SectionEndTime);
 				if (PostRollPx > 0)
@@ -244,7 +257,7 @@
 		FSequencerSelection& Selection = Sequencer.GetSelection();
 		FSequencerSelectionPreview& SelectionPreview = Sequencer.GetSelectionPreview();
 
-		ESelectionPreviewState SelectionPreviewState = SelectionPreview.GetSelectionState(&Section);
+		ESelectionPreviewState SelectionPreviewState = SelectionPreview.GetSelectionState(SectionWidget.WeakSectionModel);
 
 		if (SelectionPreviewState == ESelectionPreviewState::NotSelected)
 		{
@@ -252,7 +265,7 @@
 			return;
 		}
 		
-		if (SelectionPreviewState == ESelectionPreviewState::Undefined && !Selection.IsSelected(&Section))
+		if (SelectionPreviewState == ESelectionPreviewState::Undefined && !Selection.IsSelected(SectionModel))
 		{
 			// No preview selection for this section, and it's not selected
 			return;
@@ -276,8 +289,9 @@
 	void DrawBlendType()
 	{
 		// Draw the blend type text if necessary
+		UMovieSceneSection* SectionObject = SectionModel->GetSection();
 		UMovieSceneTrack* Track = GetTrack();
-		if (!Track || Track->GetSupportedBlendTypes().Num() <= 1 || !Section.GetBlendType().IsValid() || !bIsHighlighted || Section.GetBlendType().Get() == EMovieSceneBlendType::Absolute)
+		if (!Track || Track->GetSupportedBlendTypes().Num() <= 1 || !SectionObject->GetBlendType().IsValid() || !bIsHighlighted || SectionObject->GetBlendType().Get() == EMovieSceneBlendType::Absolute)
 		{
 			return;
 		}
@@ -285,7 +299,7 @@
 		TSharedRef<FSlateFontCache> FontCache = FSlateApplication::Get().GetRenderer()->GetFontCache();
 
 		UEnum* Enum = FindObjectChecked<UEnum>(nullptr, TEXT("/Script/MovieScene.EMovieSceneBlendType"), true);
-		FText DisplayText = Enum->GetDisplayNameTextByValue((int64)Section.GetBlendType().Get());
+		FText DisplayText = Enum->GetDisplayNameTextByValue((int64)SectionObject->GetBlendType().Get());
 
 		FSlateFontInfo FontInfo = FAppStyle::GetFontStyle("Sequencer.Section.BackgroundText");
 		FontInfo.Size = 24;
@@ -299,7 +313,7 @@
 			FontInfo.Size = FMath::Max(FMath::FloorToInt(FontInfo.Size - 6.f), 11);
 		}
 
-		FVector2D TextOffset = Section.GetRange() == TRange<FFrameNumber>::All() ? FVector2D(0.f, -1.f) : FVector2D(1.f, -1.f);
+		FVector2D TextOffset = SectionModel->GetRange() == TRange<FFrameNumber>::All() ? FVector2D(0.f, -1.f) : FVector2D(1.f, -1.f);
 		FVector2D BottomLeft = SectionGeometry.AbsoluteToLocal(SectionClippingRect.GetBottomLeft()) + TextOffset;
 
 		FSlateDrawElement::MakeText(
@@ -316,22 +330,16 @@
 		);
 	}
 
-	float GetEaseHighlightAmount(FSectionHandle Handle, float EaseInInterp, float EaseOutInterp) const
+	float GetEaseHighlightAmount(UMovieSceneSection* InSection, float EaseInInterp, float EaseOutInterp) const
 	{
-		if (!Hotspot)
-		{
-			return 0.f;
-		}
+		using namespace UE::Sequencer;
 
-		const bool bEaseInHandle = Hotspot->GetType() == ESequencerHotspot::EaseInHandle;
-		const bool bEaseOutHandle = Hotspot->GetType() == ESequencerHotspot::EaseOutHandle;
-
 		float EaseInScale = 0.f, EaseOutScale = 0.f;
-		if (bEaseInHandle || bEaseOutHandle)
+		if (TSharedPtr<FSectionEasingHandleHotspot> EasingHandleHotspot = HotspotCast<FSectionEasingHandleHotspot>(Hotspot))
 		{
-			if (static_cast<const FSectionEasingHandleHotspot*>(Hotspot)->WeakSection.Get() == Handle.GetSectionObject())
+			if (EasingHandleHotspot->GetSection() == InSection)
 			{
-				if (bEaseInHandle)
+				if (EasingHandleHotspot->HandleType == ESequencerEasingType::In)
 				{
 					EaseInScale = 1.f;
 				}
@@ -341,11 +349,11 @@
 				}
 			}
 		}
-		else if (Hotspot->GetType() == ESequencerHotspot::EasingArea)
+		else if (TSharedPtr<FSectionEasingAreaHotspot> EasingAreaHotspot = HotspotCast<FSectionEasingAreaHotspot>(Hotspot))
 		{
-			for (const FEasingAreaHandle& Easing : static_cast<const FSectionEasingAreaHotspot*>(Hotspot)->Easings)
+			for (const FEasingAreaHandle& Easing : EasingAreaHotspot->Easings)
 			{
-				if (Easing.WeakSection.Get() == Handle.GetSectionObject())
+				if (Easing.WeakSectionModel.Pin()->GetSection() == InSection)
 				{
 					if (Easing.EasingType == ESequencerEasingType::In)
 					{
@@ -358,33 +366,35 @@
 				}
 			}
 		}
+		else
+		{
+			return 0.f;
+		}
 
 		const float TotalScale = EaseInScale + EaseOutScale;
 		return TotalScale > 0.f ? EaseInInterp * (EaseInScale/TotalScale) + ((1.f-EaseOutInterp) * (EaseOutScale/TotalScale)) : 0.f;
 	}
 
-	FEasingCurvePoint MakeCurvePoint(FSectionHandle SectionHandle, FFrameTime Time, const FLinearColor& FinalTint, const FLinearColor& EaseSelectionColor) const
+	FEasingCurvePoint MakeCurvePoint(UMovieSceneSection* InSection, FFrameTime Time, const FLinearColor& FinalTint, const FLinearColor& EaseSelectionColor) const
 	{
 		TOptional<float> EaseInValue, EaseOutValue;
 		float EaseInInterp = 0.f, EaseOutInterp = 1.f;
-		SectionHandle.GetSectionObject()->EvaluateEasing(Time, EaseInValue, EaseOutValue, &EaseInInterp, &EaseOutInterp);
+		InSection->EvaluateEasing(Time, EaseInValue, EaseOutValue, &EaseInInterp, &EaseOutInterp);
 
 		return FEasingCurvePoint(
 			FVector2D(Time / TimeToPixelConverter.GetTickResolution(), EaseInValue.Get(1.f) * EaseOutValue.Get(1.f)),
-			FMath::Lerp(FinalTint, EaseSelectionColor, GetEaseHighlightAmount(SectionHandle, EaseInInterp, EaseOutInterp))
+			FMath::Lerp(FinalTint, EaseSelectionColor, GetEaseHighlightAmount(InSection, EaseInInterp, EaseOutInterp))
 		);
 	}
 
 	/** Adds intermediate control points for the specified section's easing up to a given threshold */
-	void RefineCurvePoints(FSectionHandle SectionHandle, const FLinearColor& FinalTint, const FLinearColor& EaseSelectionColor, TArray<FEasingCurvePoint>& InOutPoints)
+	void RefineCurvePoints(UMovieSceneSection* SectionObject, const FLinearColor& FinalTint, const FLinearColor& EaseSelectionColor, TArray<FEasingCurvePoint>& InOutPoints)
 	{
 		static float GradientThreshold = .05f;
 		static float ValueThreshold = .05f;
 
 		float MinTimeSize = FMath::Max(0.0001, TimeToPixelConverter.PixelToSeconds(2.5) - TimeToPixelConverter.PixelToSeconds(0));
 
-		UMovieSceneSection* SectionObject = SectionHandle.GetSectionObject();
-
 		for (int32 Index = 0; Index < InOutPoints.Num() - 1; ++Index)
 		{
 			const FEasingCurvePoint& Lower = InOutPoints[Index];
@@ -404,14 +414,14 @@
 					!FMath::IsNearlyEqual(LinearValue, NewPointValue, ValueThreshold))
 				{
 					// Add the point
-					InOutPoints.Insert(MakeCurvePoint(SectionHandle, FrameTime, FinalTint, EaseSelectionColor), Index+1);
+					InOutPoints.Insert(MakeCurvePoint(SectionObject, FrameTime, FinalTint, EaseSelectionColor), Index+1);
 					--Index;
 				}
 			}
 		}
 	}
 
-	void DrawEasingForSegment(const FSequencerOverlapRange& Segment, const FGeometry& InnerSectionGeometry, const FLinearColor& FinalTint)
+	void DrawEasingForSegment(const FOverlappingSections& Segment, const FGeometry& InnerSectionGeometry, const FLinearColor& FinalTint)
 	{
 		// @todo: sequencer-timecode: Test that start offset is not required here
 		const float RangeStartPixel = TimeToPixelConverter.FrameToPixel(UE::MovieScene::DiscreteInclusiveLower(Segment.Range));
@@ -424,7 +434,7 @@
 			return;
 		}
 
-		UMovieSceneTrack* Track = Section.GetTypedOuter<UMovieSceneTrack>();
+		UMovieSceneTrack* Track = GetTrack();
 		if (!Track)
 		{
 			return;
@@ -452,20 +462,20 @@
 		// Segment.Impls are already sorted bottom to top
 		for (int32 CurveIndex = 0; CurveIndex < Segment.Sections.Num(); ++CurveIndex)
 		{
-			FSectionHandle Handle = Segment.Sections[CurveIndex];
+			UMovieSceneSection* CurveSection = Segment.Sections[CurveIndex].Pin()->GetSection();
 
 			// Make the points for the curve
 			CurvePoints.Reset(20);
 			{
-				CurvePoints.Add(MakeCurvePoint(Handle, Segment.Range.GetLowerBoundValue(), FinalTint, EaseSelectionColor));
-				CurvePoints.Add(MakeCurvePoint(Handle, Segment.Range.GetUpperBoundValue(), FinalTint, EaseSelectionColor));
+				CurvePoints.Add(MakeCurvePoint(CurveSection, Segment.Range.GetLowerBoundValue(), FinalTint, EaseSelectionColor));
+				CurvePoints.Add(MakeCurvePoint(CurveSection, Segment.Range.GetUpperBoundValue(), FinalTint, EaseSelectionColor));
 
 				// Refine the control points
 				int32 LastNumPoints;
 				do
 				{
 					LastNumPoints = CurvePoints.Num();
-					RefineCurvePoints(Handle, FinalTint, EaseSelectionColor, CurvePoints);
+					RefineCurvePoints(CurveSection, FinalTint, EaseSelectionColor, CurvePoints);
 				} while(LastNumPoints != CurvePoints.Num());
 			}
 
@@ -536,14 +546,14 @@
 
 	void DrawEasing(const FLinearColor& FinalTint)
 	{
-		if (!Section.GetBlendType().IsValid())
+		if (!SectionModel->GetSection()->GetBlendType().IsValid())
 		{
 			return;
 		}
 
 		// Compute easing geometry by insetting from the current section geometry by 1px
 		FGeometry InnerSectionGeometry = SectionGeometry.MakeChild(SectionGeometry.Size - FVector2D(2.f, 2.f), FSlateLayoutTransform(FVector2D(1.f, 1.f)));
-		for (const FSequencerOverlapRange& Segment : SectionWidget.UnderlappingEasingSegments)
+		for (const FOverlappingSections& Segment : SectionWidget.UnderlappingEasingSegments)
 		{
 			DrawEasingForSegment(Segment, InnerSectionGeometry, FinalTint);
 		}
@@ -553,9 +563,11 @@
 
 	void DrawOverlaps(const FLinearColor& FinalTint)
 	{
+		using namespace UE::Sequencer;
+
 		FGeometry InnerSectionGeometry = SectionGeometry.MakeChild(SectionGeometry.Size - FVector2D(2.f, 2.f), FSlateLayoutTransform(FVector2D(1.f, 1.f)));
 
-		UMovieSceneTrack* Track = Section.GetTypedOuter<UMovieSceneTrack>();
+		UMovieSceneTrack* Track = GetTrack();
 		if (!Track)
 		{
 			return;
@@ -566,11 +578,11 @@
 
 		const ESlateDrawEffect DrawEffects = bParentEnabled ? ESlateDrawEffect::None : ESlateDrawEffect::DisabledEffect;
 
-		const float StartTimePixel = Section.HasStartFrame() ? TimeToPixelConverter.FrameToPixel(Section.GetInclusiveStartFrame()) : 0.f;
+		const float StartTimePixel = SectionModel->GetSection()->HasStartFrame() ? TimeToPixelConverter.FrameToPixel(SectionModel->GetSection()->GetInclusiveStartFrame()) : 0.f;
 
 		for (int32 SegmentIndex = 0; SegmentIndex < SectionWidget.UnderlappingSegments.Num(); ++SegmentIndex)
 		{
-			const FSequencerOverlapRange& Segment = SectionWidget.UnderlappingSegments[SegmentIndex];
+			const FOverlappingSections& Segment = SectionWidget.UnderlappingSegments[SegmentIndex];
 
 			const float RangeStartPixel	= Segment.Range.GetLowerBound().IsOpen() ? 0.f							: TimeToPixelConverter.FrameToPixel(UE::MovieScene::DiscreteInclusiveLower(Segment.Range));
 			const float RangeEndPixel	= Segment.Range.GetUpperBound().IsOpen() ? InnerSectionGeometry.Size.X	: TimeToPixelConverter.FrameToPixel(UE::MovieScene::DiscreteExclusiveUpper(Segment.Range));
@@ -582,7 +594,7 @@
 				continue;
 			}
 
-			const FSequencerOverlapRange* NextSegment = SegmentIndex < SectionWidget.UnderlappingSegments.Num() - 1 ? &SectionWidget.UnderlappingSegments[SegmentIndex+1] : nullptr;
+			const FOverlappingSections* NextSegment = SegmentIndex < SectionWidget.UnderlappingSegments.Num() - 1 ? &SectionWidget.UnderlappingSegments[SegmentIndex+1] : nullptr;
 			const bool bDrawRightMostBound = !NextSegment || !Segment.Range.Adjoins(NextSegment->Range);
 
 			FSlateDrawElement::MakeBox(
@@ -610,6 +622,8 @@
 
 	void DrawEmptySpace()
 	{
+		using namespace UE::Sequencer;
+
 		const ESlateDrawEffect DrawEffects = bParentEnabled ? ESlateDrawEffect::None : ESlateDrawEffect::DisabledEffect;
 		static const FSlateBrush* EmptySpaceBrush = FAppStyle::GetBrush("Sequencer.Section.EmptySpace");
 
@@ -618,7 +632,7 @@
 
 		for (const FSectionLayoutElement& Element : SectionWidget.Layout->GetElements())
 		{
-			const bool bIsEmptySpace = Element.GetDisplayNode()->GetType() == ESequencerNode::KeyArea && Element.GetKeyAreas().Num() == 0;
+			const bool bIsEmptySpace = Element.GetModel()->IsA<FChannelModel>() && Element.GetChannels().Num() == 0;
 			const bool bExistingEmptySpace = CurrentArea.IsSet();
 
 			if (bIsEmptySpace && bExistingEmptySpace && FMath::IsNearlyEqual(CurrentArea->Bottom, Element.GetOffset()))
@@ -653,42 +667,174 @@
 	FSequencer& Sequencer;
 	const SSequencerSection& SectionWidget;
 	FTimeToPixel TimeToPixelConverter;
-	const ISequencerHotspot* Hotspot;
+	TSharedPtr<ITrackAreaHotspot> Hotspot;
 
 	/** The clipping rectangle of the parent widget */
 	FSlateRect ParentClippingRect;
 };
 
-void SSequencerSection::Construct( const FArguments& InArgs, TSharedRef<FSequencerTrackNode> SectionNode, int32 InSectionIndex )
+void SSequencerSection::Construct( const FArguments& InArgs, TSharedPtr<FSequencer> InSequencer, TSharedPtr<FSectionModel> InSectionModel)
 {
-	SectionIndex = InSectionIndex;
-	ParentSectionArea = SectionNode;
-	SectionInterface = SectionNode->GetSections()[InSectionIndex];
-	Layout = FSectionLayout(*SectionNode, InSectionIndex);
+	using namespace UE::MovieScene;
+
+	Sequencer = InSequencer;
+	WeakSectionModel = InSectionModel;
+	SectionInterface = InSectionModel->GetSectionInterface();
+	Layout = FSectionLayout(InSectionModel);
 	HandleOffsetPx = 0.f;
 
+	SetEnabled(MakeAttributeSP(this, &SSequencerSection::IsEnabled));
+	SetToolTipText(MakeAttributeSP(this, &SSequencerSection::GetToolTipText));
+
+	UMovieSceneSection* Section = InSectionModel->GetSection();
+	UMovieSceneTrack*   Track   = Section ? Section->GetTypedOuter<UMovieSceneTrack>() : nullptr;
+	if (ensure(Track))
+	{
+		Track->EventHandlers.Link(TrackModifiedBinding, this);
+	}
+
+	UpdateUnderlappingSegments();
+
 	ChildSlot
 	[
 		SectionInterface->GenerateSectionWidget()
 	];
 }
 
+SSequencerSection::~SSequencerSection()
+{
+}
 
+FText SSequencerSection::GetToolTipText() const
+{
+	const UMovieSceneSection* SectionObject = SectionInterface->GetSectionObject();
+	const UMovieScene* MovieScene = SectionObject ? SectionObject->GetTypedOuter<UMovieScene>() : nullptr;
+
+	// Optional section specific content to add to tooltip
+	FText SectionToolTipContent = SectionInterface->GetSectionToolTip();
+
+	FText SectionTitleText = SectionInterface->GetSectionTitle();
+	if (!SectionTitleText.IsEmpty())
+	{
+		SectionTitleText = FText::Format(FText::FromString(TEXT("{0}\n")), SectionTitleText);
+	}
+
+	// If the objects are valid and the section is not unbounded, add frame information to the tooltip
+	if (SectionObject && MovieScene && SectionObject->HasStartFrame() && SectionObject->HasEndFrame())
+	{
+		FFrameRate TickResolution = MovieScene->GetTickResolution();
+		FFrameRate DisplayRate = MovieScene->GetDisplayRate();
+
+		int32 StartFrame = ConvertFrameTime(SectionObject->GetInclusiveStartFrame(), MovieScene->GetTickResolution(), MovieScene->GetDisplayRate()).RoundToFrame().Value;
+		int32 EndFrame = ConvertFrameTime(SectionObject->GetExclusiveEndFrame(), MovieScene->GetTickResolution(), MovieScene->GetDisplayRate()).RoundToFrame().Value;
+	
+		FText SectionToolTip;
+		if (SectionToolTipContent.IsEmpty())
+		{
+			SectionToolTip = FText::Format(NSLOCTEXT("SequencerSection", "TooltipFormat", "{0}{1} - {2} ({3} frames)"), SectionTitleText,
+				StartFrame,
+				EndFrame,
+				EndFrame - StartFrame);
+		}
+		else
+		{
+			SectionToolTip = FText::Format(NSLOCTEXT("SequencerSection", "TooltipFormatWithSectionContent", "{0}{1} - {2} ({3} frames)\n{4}"), SectionTitleText,
+				StartFrame,
+				EndFrame,
+				EndFrame - StartFrame,
+				SectionToolTipContent);
+		}
+	
+		if (SectionObject->Easing.EaseIn.GetObject() && SectionObject->Easing.GetEaseInDuration() > 0)
+		{
+			int32 EaseInFrames = ConvertFrameTime(SectionObject->Easing.GetEaseInDuration(), TickResolution, DisplayRate).RoundToFrame().Value;
+			FText EaseInText = FText::Format(NSLOCTEXT("SequencerSection", "EaseInFormat", "Ease In: {0} ({1} frames)"), SectionObject->Easing.EaseIn->GetDisplayName(), EaseInFrames);
+			SectionToolTip = FText::Join(FText::FromString("\n"), SectionToolTip, EaseInText);
+		}
+
+		if (SectionObject->Easing.EaseOut.GetObject() && SectionObject->Easing.GetEaseOutDuration() > 0)
+		{
+			int32 EaseOutFrames = ConvertFrameTime(SectionObject->Easing.GetEaseOutDuration(), TickResolution, DisplayRate).RoundToFrame().Value;
+			FText EaseOutText = FText::Format(NSLOCTEXT("SequencerSection", "EaseOutFormat", "Ease Out: {0} ({1} frames)"), SectionObject->Easing.EaseOut->GetDisplayName(), EaseOutFrames);
+			SectionToolTip = FText::Join(FText::FromString("\n"), SectionToolTip, EaseOutText);
+		}
+		
+		return SectionToolTip;
+	}
+	else
+	{
+		if (SectionToolTipContent.IsEmpty())
+		{
+			return SectionInterface->GetSectionTitle();
+		}
+		else
+		{
+			return FText::Format(NSLOCTEXT("SequencerSection", "TooltipSectionContentFormat", "{0}{1}"), SectionTitleText, SectionToolTipContent);
+		}
+	}
+}
+
+bool SSequencerSection::IsEnabled() const
+{
+	return !SectionInterface->IsReadOnly();
+}
+
 FVector2D SSequencerSection::ComputeDesiredSize(float) const
 {
-	return FVector2D(100, Layout->GetTotalHeight());
+	return FVector2D(0.f, 0.f);
 }
 
+void SSequencerSection::ReportParentGeometry(const FGeometry& InParentGeometry)
+{
+	ParentGeometry = InParentGeometry;
+}
 
+FTrackLaneScreenAlignment SSequencerSection::GetAlignment(const FTimeToPixel& InTimeToPixel, const FGeometry& InParentGeometry) const
+{
+	using namespace UE::MovieScene;
+	TSharedPtr<FSectionModel> SectionModel = WeakSectionModel.Pin();
+	if (!SectionModel)
+	{
+		return FTrackLaneScreenAlignment();
+	}
+
+	FTrackLaneVirtualAlignment VirtualAlignment = SectionModel->ArrangeVirtualTrackLaneView();
+	FTrackLaneScreenAlignment  ScreenAlignment  = VirtualAlignment.ToScreen(InTimeToPixel, InParentGeometry);
+
+	if (TOptional<FFrameNumber> FiniteLength = VirtualAlignment.GetFiniteLength())
+	{
+		constexpr float MinSectionWidth = 1.f;
+
+		const float FinalSectionWidth = FMath::Max(MinSectionWidth, ScreenAlignment.WidthPx);
+		const float GripOffset        = (FinalSectionWidth - ScreenAlignment.WidthPx) / 2.f;
+
+		ScreenAlignment.LeftPosPx -= GripOffset;
+		ScreenAlignment.WidthPx    = FMath::Max(FinalSectionWidth, MinSectionWidth + SectionInterface->GetSectionGripSize() * 2.f);
+	}
+
+	return ScreenAlignment;
+}
+
+int32 SSequencerSection::GetOverlapPriority() const
+{
+	if (TSharedPtr<FSectionModel> SectionModel = WeakSectionModel.Pin())
+	{
+		if (UMovieSceneSection* Section = SectionModel->GetSection())
+		{
+			return Section->GetOverlapPriority();
+		}
+	}
+	return 0;
+}
+
 void SSequencerSection::GetKeysUnderMouse( const FVector2D& MousePosition, const FGeometry& AllottedGeometry, TArray<FSequencerSelectedKey>& OutKeys, float KeyHeightFraction ) const
 {
 	FGeometry SectionGeometry = MakeSectionGeometryWithoutHandles( AllottedGeometry, SectionInterface );
 
 	UMovieSceneSection& Section = *SectionInterface->GetSectionObject();
+	FTimeToPixel TimeToPixelConverter = ConstructTimeConverterForSection(SectionGeometry, Section, GetSequencer());
+	const FVector2D MousePixel  = SectionGeometry.AbsoluteToLocal( MousePosition );
 
-	FTimeToPixel TimeToPixelConverter = ConstructTimeConverterForSection(SectionGeometry, Section, GetSequencer());
-	const FVector2D MousePixel           = SectionGeometry.AbsoluteToLocal( MousePosition );
-
 	// HitTest 
 	const FFrameTime HalfKeySizeFrames = TimeToPixelConverter.PixelDeltaToFrame( SequencerSectionConstants::KeySize.X*.5f );
 	const FFrameTime MouseFrameTime    = TimeToPixelConverter.PixelToFrame( MousePixel.X );
@@ -715,17 +861,17 @@
 			continue;
 		}
 
-		for (TSharedPtr<IKeyArea> KeyArea : Element.GetKeyAreas())
+		for (TWeakPtr<FChannelModel> WeakChannel : Element.GetChannels())
 		{
-			if (KeyArea.IsValid())
+			if (TSharedPtr<FChannelModel> Channel = WeakChannel.Pin())
 			{
 				TArray<FKeyHandle> KeyHandles;
-				KeyArea->GetKeyHandles(KeyHandles, HitTestRange);
+				Channel->GetKeyArea()->GetKeyHandles(KeyHandles, HitTestRange);
 
 				// Only ever select one key from any given key area
 				if (KeyHandles.Num())
 				{
-					OutKeys.Add( FSequencerSelectedKey( Section, KeyArea, KeyHandles[0] ) );
+					OutKeys.Add( FSequencerSelectedKey( Section, Channel, KeyHandles[0] ) );
 				}
 			}
 		}
@@ -738,6 +884,12 @@
 
 void SSequencerSection::CreateKeysUnderMouse( const FVector2D& MousePosition, const FGeometry& AllottedGeometry, TArrayView<const FSequencerSelectedKey> InPressedKeys, TArray<FSequencerSelectedKey>& OutKeys )
 {
+	TSharedPtr<FSectionModel> SectionModel = WeakSectionModel.Pin();
+	if (!SectionModel)
+	{
+		return;
+	}
+
 	UMovieSceneSection& Section = *SectionInterface->GetSectionObject();
 
 	if (Section.IsReadOnly())
@@ -745,36 +897,36 @@
 		return;
 	}
 
+	FGeometry SectionGeometry = MakeSectionGeometryWithoutHandles( AllottedGeometry, SectionInterface );
+	FTimeToPixel TimeToPixelConverter = ConstructTimeConverterForSection(SectionGeometry, Section, GetSequencer());
+
 	// If the pressed key exists, offset the new key and look for it in the newly laid out key areas
 	if (InPressedKeys.Num())
 	{
 		Section.Modify();
 
 		// Offset by 1 pixel worth of time if possible
-		const FFrameTime TimeFuzz = (GetSequencer().GetViewRange().Size<double>() / ParentGeometry.GetLocalSize().X) * Section.GetTypedOuter<UMovieScene>()->GetTickResolution();
+		const FFrameTime TimeFuzz = TimeToPixelConverter.PixelDeltaToFrame(1.f);
 
 		for (const FSequencerSelectedKey& PressedKey : InPressedKeys)
 		{
-			const FFrameNumber CurrentTime = PressedKey.KeyArea->GetKeyTime(PressedKey.KeyHandle.GetValue());
-			const FKeyHandle NewHandle = PressedKey.KeyArea->DuplicateKey(PressedKey.KeyHandle.GetValue());
+			const FFrameNumber CurrentTime = PressedKey.WeakChannel.Pin()->GetKeyArea()->GetKeyTime(PressedKey.KeyHandle);
+			const FKeyHandle NewHandle = PressedKey.WeakChannel.Pin()->GetKeyArea()->DuplicateKey(PressedKey.KeyHandle);
 
-			PressedKey.KeyArea->SetKeyTime(NewHandle, CurrentTime + TimeFuzz.FrameNumber);
-			OutKeys.Add(FSequencerSelectedKey(Section, PressedKey.KeyArea, NewHandle));
+			PressedKey.WeakChannel.Pin()->GetKeyArea()->SetKeyTime(NewHandle, CurrentTime + TimeFuzz.FrameNumber);
+			OutKeys.Add(FSequencerSelectedKey(Section, PressedKey.WeakChannel, NewHandle));
 		}
 	}
 	else
 	{
-		TSharedPtr<FSequencerObjectBindingNode> ObjectBindingNode = ParentSectionArea.IsValid() ? ParentSectionArea->FindParentObjectBindingNode() : nullptr;
-		FGuid ObjectBinding = ObjectBindingNode.IsValid() ? ObjectBindingNode->GetObjectBinding() : FGuid();
+		TSharedPtr<FTrackModel>             TrackModel             = SectionModel->FindAncestorOfType<FTrackModel>();
+		TSharedPtr<IObjectBindingExtension> ObjectBindingExtension = SectionModel->FindAncestorOfType<IObjectBindingExtension>();
 
-		FGeometry SectionGeometry = MakeSectionGeometryWithoutHandles( AllottedGeometry, SectionInterface );
+		FGuid ObjectBinding = ObjectBindingExtension ? ObjectBindingExtension->GetObjectGuid() : FGuid();
 
-		FTimeToPixel TimeToPixelConverter = ConstructTimeConverterForSection(SectionGeometry, Section, GetSequencer());
 		FVector2D LocalSpaceMousePosition = SectionGeometry.AbsoluteToLocal( MousePosition );
 		const FFrameTime CurrentTime = TimeToPixelConverter.PixelToFrame(LocalSpaceMousePosition.X);
 
-		ISequencerTrackEditor& TrackEditor = ParentSectionArea->GetTrackEditor();
-
 		TArray<TSharedRef<IKeyArea>> ValidKeyAreasUnderCursor;
 
 		// Search every key area until we find the one under the mouse
@@ -789,19 +941,17 @@
 				continue;
 			}
 
-			for (TSharedPtr<IKeyArea> KeyArea : Element.GetKeyAreas())
+			for (TWeakPtr<FChannelModel> WeakChannel : Element.GetChannels())
 			{
-				if (KeyArea.IsValid())
+				if (TSharedPtr<FChannelModel> Channel = WeakChannel.Pin())
 				{
-					ValidKeyAreasUnderCursor.Add(KeyArea.ToSharedRef());
+					ValidKeyAreasUnderCursor.Add(Channel->GetKeyArea().ToSharedRef());
 				}
 			}
 		}
 
-		using namespace UE::Sequencer;
-
 		FScopedTransaction Transaction(NSLOCTEXT("Sequencer", "CreateKeysUnderMouse", "Create keys under mouse"));
-		FAddKeyOperation::FromKeyAreas(&TrackEditor, ValidKeyAreasUnderCursor).Commit(CurrentTime.FrameNumber, GetSequencer());
+		FAddKeyOperation::FromKeyAreas(TrackModel->GetTrackEditor().Get(), ValidKeyAreasUnderCursor).Commit(CurrentTime.FrameNumber, GetSequencer());
 
 		// Get the keys under the mouse as the newly created keys. Check with the full height of the key track area.
 		const float KeyHeightFraction = 1.f;
@@ -810,7 +960,7 @@
 
 	if (OutKeys.Num())
 	{
-		Layout = FSectionLayout(*ParentSectionArea, SectionIndex);
+		Layout = FSectionLayout(SectionModel);
 	}
 }
 
@@ -850,20 +1000,19 @@
 	}
 
 	// Gather all underlapping sections
-	TArray<FSectionHandle> AllUnderlappingSections;
-	AllUnderlappingSections.Add(FSectionHandle(ParentSectionArea.ToSharedRef(), SectionIndex));
-	for (const FSequencerOverlapRange& Segment : UnderlappingSegments)
+	TArray<TSharedPtr<FSectionModel>> AllUnderlappingSections;
+	AllUnderlappingSections.Add(WeakSectionModel.Pin());
+	for (const FOverlappingSections& Segment : UnderlappingSegments)
 	{
-		for (FSectionHandle Section : Segment.Sections)
+		for (const TWeakPtr<FSectionModel>& Section : Segment.Sections)
 		{
-			AllUnderlappingSections.AddUnique(Section);
+			AllUnderlappingSections.AddUnique(Section.Pin());
 		}
 	}
 
-	for (FSectionHandle Handle : AllUnderlappingSections)
+	for (const TSharedPtr<FSectionModel>& SectionModel : AllUnderlappingSections)
 	{
-		TSharedRef<ISequencerSection> EasingSection    =  Handle.GetSectionInterface();
-		UMovieSceneSection*           EasingSectionObj = EasingSection->GetSectionObject();
+		UMovieSceneSection* EasingSectionObj = SectionModel->GetSection();
 
 		if (EasingSectionObj->HasStartFrame() && EnumHasAllFlags(EasingFlags, EMovieSceneTrackEasingSupportFlags::ManualEaseIn))
 		{
@@ -872,7 +1021,7 @@
 
 			if (FMath::IsNearlyEqual(MouseTime, HandlePositionIn, HalfHandleSizeX))
 			{
-				GetSequencer().SetHotspot(MakeShared<FSectionEasingHandleHotspot>(ESequencerEasingType::In, EasingSectionObj));
+				GetSequencer().GetViewModel()->GetTrackArea()->SetHotspot(MakeShared<FSectionEasingHandleHotspot>(ESequencerEasingType::In, SectionModel, Sequencer));
 				return true;
 			}
 		}
@@ -884,7 +1033,7 @@
 
 			if (FMath::IsNearlyEqual(MouseTime, HandlePositionOut, HalfHandleSizeX))
 			{
-				GetSequencer().SetHotspot(MakeShared<FSectionEasingHandleHotspot>(ESequencerEasingType::Out, EasingSectionObj));
+				GetSequencer().GetViewModel()->GetTrackArea()->SetHotspot(MakeShared<FSectionEasingHandleHotspot>(ESequencerEasingType::Out, SectionModel, Sequencer));
 				return true;
 			}
 		}
@@ -902,30 +1051,30 @@
 		return false;
 	}
 
-	TArray<FSectionHandle> AllUnderlappingSections;
-	AllUnderlappingSections.Add(FSectionHandle(ParentSectionArea.ToSharedRef(), SectionIndex));
-	for (const FSequencerOverlapRange& Segment : UnderlappingSegments)
+	TArray<TSharedPtr<FSectionModel>> AllUnderlappingSections;
+	AllUnderlappingSections.Add(WeakSectionModel.Pin());
+	for (const FOverlappingSections& Segment : UnderlappingSegments)
 	{
-		for (FSectionHandle Section : Segment.Sections)
+		for (const TWeakPtr<FSectionModel>& Section : Segment.Sections)
 		{
-			AllUnderlappingSections.AddUnique(Section);
+			AllUnderlappingSections.AddUnique(Section.Pin());
 		}
 	}
 
-	FGeometry    SectionGeometryWithoutHandles = MakeSectionGeometryWithoutHandles(SectionGeometry, SectionInterface);
-	FTimeToPixel TimeToPixelConverter          = ConstructTimeConverterForSection(SectionGeometryWithoutHandles, *ThisSection, GetSequencer());
+	FGeometry SectionGeometryWithoutHandles = MakeSectionGeometryWithoutHandles(SectionGeometry, SectionInterface);
+	FTimeToPixel TimeToPixelConverter       = ConstructTimeConverterForSection(SectionGeometryWithoutHandles, *ThisSection, GetSequencer());
 
-	for (FSectionHandle Handle : AllUnderlappingSections)
+	for (const TSharedPtr<FSectionModel>& UnderlappingSection : AllUnderlappingSections)
 	{
-		TSharedRef<ISequencerSection> UnderlappingSection =  Handle.GetSectionInterface();
-		UMovieSceneSection* UnderlappingSectionObj = UnderlappingSection->GetSectionObject();
-		if (!UnderlappingSection->SectionIsResizable())
+		UMovieSceneSection*           UnderlappingSectionObj       = UnderlappingSection->GetSection();
+		TSharedPtr<ISequencerSection> UnderlappingSectionInterface = UnderlappingSection->GetSectionInterface();
+		if (!UnderlappingSectionInterface->SectionIsResizable())
 		{
 			continue;
 		}
 
 		const float ThisHandleOffset = UnderlappingSectionObj == ThisSection ? HandleOffsetPx : 0.f;
-		FVector2D GripSize( UnderlappingSection->GetSectionGripSize(), SectionGeometry.Size.Y );
+		FVector2D GripSize( UnderlappingSectionInterface->GetSectionGripSize(), SectionGeometry.Size.Y );
 
 		if (UnderlappingSectionObj->HasStartFrame())
 		{
@@ -937,7 +1086,7 @@
 
 			if( SectionRectLeft.IsUnderLocation( MouseEvent.GetScreenSpacePosition() ) )
 			{
-				GetSequencer().SetHotspot(MakeShareable( new FSectionResizeHotspot(FSectionResizeHotspot::Left, UnderlappingSectionObj)) );
+				GetSequencer().GetViewModel()->GetTrackArea()->SetHotspot(MakeShareable( new FSectionResizeHotspot(FSectionResizeHotspot::Left, UnderlappingSection, Sequencer)) );
 				return true;
 			}
 		}
@@ -945,13 +1094,13 @@
 		if (UnderlappingSectionObj->HasEndFrame())
 		{
 			FGeometry SectionRectRight = SectionGeometryWithoutHandles.MakeChild(
-				FVector2D( TimeToPixelConverter.FrameToPixel(UnderlappingSectionObj->GetExclusiveEndFrame()) - UnderlappingSection->GetSectionGripSize() + ThisHandleOffset, 0 ), 
+				FVector2D( TimeToPixelConverter.FrameToPixel(UnderlappingSectionObj->GetExclusiveEndFrame()) - UnderlappingSectionInterface->GetSectionGripSize() + ThisHandleOffset, 0 ), 
 				GripSize
 			);
 
 			if( SectionRectRight.IsUnderLocation( MouseEvent.GetScreenSpacePosition() ) )
 			{
-				GetSequencer().SetHotspot(MakeShareable( new FSectionResizeHotspot(FSectionResizeHotspot::Right, UnderlappingSectionObj)) );
+				GetSequencer().GetViewModel()->GetTrackArea()->SetHotspot(MakeShareable( new FSectionResizeHotspot(FSectionResizeHotspot::Right, UnderlappingSection, Sequencer)) );
 				return true;
 			}
 		}
@@ -962,16 +1111,11 @@
 bool SSequencerSection::CheckForEasingAreaInteraction( const FPointerEvent& MouseEvent, const FGeometry& SectionGeometry )
 {
 	UMovieSceneSection* ThisSection = SectionInterface->GetSectionObject();
-	if (!ThisSection)
-	{
-		return false;
-	}
-
 	FTimeToPixel TimeToPixelConverter = ConstructTimeConverterForSection(MakeSectionGeometryWithoutHandles(SectionGeometry, SectionInterface), *ThisSection, GetSequencer());
 	const FFrameNumber MouseTime = TimeToPixelConverter.PixelToFrame(SectionGeometry.AbsoluteToLocal(MouseEvent.GetScreenSpacePosition()).X).FrameNumber;
 
 	// First off, set the hotspot to an easing area if necessary
-	for (const FSequencerOverlapRange& Segment : UnderlappingEasingSegments)
+	for (const FOverlappingSections& Segment : UnderlappingEasingSegments)
 	{
 		if (!Segment.Range.Contains(MouseTime))
 		{
@@ -979,22 +1123,22 @@
 		}
 
 		TArray<FEasingAreaHandle> EasingAreas;
-		for (FSectionHandle Handle : Segment.Sections)
+		for (const TWeakPtr<FSectionModel>& SectionModel : Segment.Sections)
 		{
-			UMovieSceneSection* Section = Handle.GetSectionObject();
+			UMovieSceneSection* Section = SectionModel.Pin()->GetSection();
 			if (Section->GetEaseInRange().Contains(MouseTime))
 			{
-				EasingAreas.Add(FEasingAreaHandle{ Section, ESequencerEasingType::In });
+				EasingAreas.Add(FEasingAreaHandle{ SectionModel, ESequencerEasingType::In });
 			}
 			if (Section->GetEaseOutRange().Contains(MouseTime))
 			{
-				EasingAreas.Add(FEasingAreaHandle{ Section, ESequencerEasingType::Out });
+				EasingAreas.Add(FEasingAreaHandle{ SectionModel, ESequencerEasingType::Out });
 			}
 		}
 
 		if (EasingAreas.Num())
 		{
-			GetSequencer().SetHotspot(MakeShared<FSectionEasingAreaHotspot>(EasingAreas, ThisSection));
+			GetSequencer().GetViewModel()->GetTrackArea()->SetHotspot(MakeShared<FSectionEasingAreaHotspot>(EasingAreas, WeakSectionModel, Sequencer));
 			return true;
 		}
 	}
@@ -1003,19 +1147,23 @@
 
 FSequencer& SSequencerSection::GetSequencer() const
 {
-	return ParentSectionArea->GetSequencer();
+	return *Sequencer.Pin().Get();
 }
 
-
 int32 SSequencerSection::OnPaint( const FPaintArgs& Args, const FGeometry& AllottedGeometry, const FSlateRect& MyCullingRect, FSlateWindowElementList& OutDrawElements, int32 LayerId, const FWidgetStyle& InWidgetStyle, bool bParentEnabled ) const
 {
-	UMovieSceneSection* SectionObject = SectionInterface->GetSectionObject();
+	TSharedPtr<FSectionModel> SectionModel = WeakSectionModel.Pin();
+	UMovieSceneSection* SectionObject = SectionModel ? SectionModel->GetSection() : nullptr;
+	if (!SectionModel || !SectionObject)
+	{
+		return LayerId;
+	}
 
-	const ISequencerEditTool* EditTool = GetSequencer().GetEditTool();
-	const ISequencerHotspot* Hotspot = EditTool ? EditTool->GetDragHotspot() : nullptr;
+	const ISequencerEditTool* EditTool = GetSequencer().GetViewModel()->GetTrackArea()->GetEditTool();
+	TSharedPtr<ITrackAreaHotspot> Hotspot = EditTool ? EditTool->GetDragHotspot() : nullptr;
 	if (!Hotspot)
 	{
-		Hotspot = GetSequencer().GetHotspot().Get();
+		Hotspot = GetSequencer().GetViewModel()->GetTrackArea()->GetHotspot();
 	}
 
 	UMovieSceneTrack* Track = SectionObject->GetTypedOuter<UMovieSceneTrack>();
@@ -1034,7 +1182,7 @@
 
 	FGeometry SectionGeometry = MakeSectionGeometryWithoutHandles( AllottedGeometry, SectionInterface );
 
-	FSequencerSectionPainterImpl Painter(ParentSectionArea->GetSequencer(), *SectionObject, OutDrawElements, SectionGeometry, *this);
+	FSequencerSectionPainterImpl Painter(GetSequencer(), SectionModel, OutDrawElements, SectionGeometry, *this);
 
 	FGeometry PaintSpaceParentGeometry = ParentGeometry;
 	PaintSpaceParentGeometry.AppendTransform(FSlateLayoutTransform(Inverse(Args.GetWindowToDesktopTransform())));
@@ -1057,26 +1205,24 @@
 			Painter.GhostAlpha = .3f;
 		}
 	}
-	
-	auto& Selection = ParentSectionArea->GetSequencer().GetSelection();
-	Painter.bIsSelected = Selection.IsSelected(SectionObject);
+
+	Painter.bIsSelected = GetSequencer().GetSelection().IsSelected(SectionModel);
 
 	for (const FSectionLayoutElement& Element : Layout->GetElements())
 	{
-		TSharedPtr<FSequencerDisplayNode> DisplayNode = Element.GetDisplayNode();
-		if (DisplayNode->GetType() == ESequencerNode::Category)
+		TSharedPtr<FViewModel> Model = Element.GetModel();
+		if (FCategoryModel* Category = Model->CastThis<FCategoryModel>())
 		{
 			TArray<TSharedRef<IKeyArea>> ChildKeyAreas;
-			TSharedPtr<FSequencerSectionCategoryNode> CategoryNode = StaticCastSharedPtr<FSequencerSectionCategoryNode>(DisplayNode);
-			for (TSharedRef<FSequencerDisplayNode> ChildNode : CategoryNode->GetChildNodes())
+
+			for (TSharedPtr<FChannelModel> Channel : Category->GetDescendantsOfType<FChannelModel>())
 			{
-				TSharedRef<FSequencerSectionKeyAreaNode> KeyAreaNode = StaticCastSharedRef<FSequencerSectionKeyAreaNode>(ChildNode);
-				for (TSharedRef<IKeyArea> KeyArea : KeyAreaNode.Get().GetAllKeyAreas())
+				if (TSharedPtr<IKeyArea> KeyArea = Channel->GetKeyArea())
 				{
-					ChildKeyAreas.Add(KeyArea);
+					ChildKeyAreas.Add(KeyArea.ToSharedRef());
 				}
 			}
-			
+
 			FKeyAreaElement KeyAreaElement;
 			KeyAreaElement.KeyAreas = ChildKeyAreas;
 			KeyAreaElement.KeyAreaGeometry = Element.ComputeGeometry(AllottedGeometry);
@@ -1086,7 +1232,13 @@
 		else
 		{
 			FKeyAreaElement KeyAreaElement;
-			KeyAreaElement.KeyAreas = Element.GetKeyAreas();
+			for (TWeakPtr<FChannelModel> WeakChannel : Element.GetChannels())
+			{
+				if (TSharedPtr<FChannelModel> Channel = WeakChannel.Pin())
+				{
+					KeyAreaElement.KeyAreas.Add(Channel->GetKeyArea().ToSharedRef());
+				}
+			}
 			KeyAreaElement.KeyAreaGeometry = Element.ComputeGeometry(AllottedGeometry);
 			KeyAreaElement.Type = (FKeyAreaElement::EType)Element.GetType();
 			Painter.KeyAreaElements.Add(KeyAreaElement);
@@ -1108,7 +1260,7 @@
 	PaintEasingHandles( Painter, SelectionColor, Hotspot );
 
 	{
-		UE::Sequencer::FKeyRendererPaintArgs KeyRenderArgs;
+		FKeyRendererPaintArgs KeyRenderArgs;
 		KeyRenderArgs.KeyThrobValue = GetKeySelectionThrobValue();
 		KeyRenderArgs.SectionThrobValue = GetSectionSelectionThrobValue();
 
@@ -1203,34 +1355,35 @@
 }
 
 
-void SSequencerSection::PaintEasingHandles( FSequencerSectionPainter& InPainter, FLinearColor SelectionColor, const ISequencerHotspot* Hotspot ) const
+void SSequencerSection::PaintEasingHandles( FSequencerSectionPainter& InPainter, FLinearColor SelectionColor, TSharedPtr<ITrackAreaHotspot> Hotspot ) const
 {
 	if (!SectionInterface->GetSectionObject()->GetBlendType().IsValid())
 	{
 		return;
 	}
 
-	TArray<FSectionHandle> AllUnderlappingSections;
+	TArray<TSharedPtr<FSectionModel>> AllUnderlappingSections;
 	if (IsSectionHighlighted(SectionInterface->GetSectionObject(), Hotspot))
 	{
-		AllUnderlappingSections.Add(FSectionHandle(ParentSectionArea.ToSharedRef(), SectionIndex));
+		AllUnderlappingSections.Add(WeakSectionModel.Pin());
 	}
 
-	for (const FSequencerOverlapRange& Segment : UnderlappingSegments)
+	for (const FOverlappingSections& Segment : UnderlappingSegments)
 	{
-		for (FSectionHandle Section : Segment.Sections)
+		for (const TWeakPtr<FSectionModel>& Section : Segment.Sections)
 		{
-			if (IsSectionHighlighted(Section.GetSectionObject(), Hotspot) && !AllUnderlappingSections.Contains(Section))
+			UMovieSceneSection* SectionObject = Section.Pin()->GetSection();
+			if (IsSectionHighlighted(SectionObject, Hotspot))
 			{
-				AllUnderlappingSections.Add(Section);
+				AllUnderlappingSections.AddUnique(Section.Pin());
 			}
 		}
 	}
 
 	FTimeToPixel TimeToPixelConverter = InPainter.GetTimeConverter();
-	for (FSectionHandle Handle : AllUnderlappingSections)
+	for (const TSharedPtr<FSectionModel>& SectionModel : AllUnderlappingSections)
 	{
-		UMovieSceneSection* UnderlappingSectionObj = Handle.GetSectionInterface()->GetSectionObject();
+		UMovieSceneSection* UnderlappingSectionObj = SectionModel->GetSection();
 		if (UnderlappingSectionObj->GetRange() == TRange<FFrameNumber>::All())
 		{
 			continue;
@@ -1243,20 +1396,17 @@
 		// Get the hovered/selected state for the section handles from the hotspot
 		if (Hotspot)
 		{
-			if (Hotspot->GetType() == ESequencerHotspot::EaseInHandle || Hotspot->GetType() == ESequencerHotspot::EaseOutHandle)
+			if (const FSectionEasingHandleHotspot* EasingHotspot = Hotspot->CastThis<FSectionEasingHandleHotspot>())
 			{
-				const FSectionEasingHandleHotspot* EasingHotspot = static_cast<const FSectionEasingHandleHotspot*>(Hotspot);
-
-				bDrawThisSectionsHandles = (EasingHotspot->WeakSection == Handle.GetSectionObject());
-				bLeftHandleActive = Hotspot->GetType() == ESequencerHotspot::EaseInHandle;
-				bRightHandleActive = Hotspot->GetType() == ESequencerHotspot::EaseOutHandle;
+				bDrawThisSectionsHandles = (EasingHotspot->WeakSectionModel.Pin() == SectionModel);
+				bLeftHandleActive = EasingHotspot->HandleType == ESequencerEasingType::In;
+				bRightHandleActive = EasingHotspot->HandleType == ESequencerEasingType::Out;
 			}
-			else if (Hotspot->GetType() == ESequencerHotspot::EasingArea)
+			else if (const FSectionEasingAreaHotspot* EasingAreaHotspot = Hotspot->CastThis<FSectionEasingAreaHotspot>())
 			{
-				const FSectionEasingAreaHotspot* EasingAreaHotspot = static_cast<const FSectionEasingAreaHotspot*>(Hotspot);
 				for (const FEasingAreaHandle& Easing : EasingAreaHotspot->Easings)
 				{
-					if (Easing.WeakSection == Handle.GetSectionObject())
+					if (Easing.WeakSectionModel.Pin()->GetSection() == UnderlappingSectionObj)
 					{
 						if (Easing.EasingType == ESequencerEasingType::In)
 						{
@@ -1337,7 +1487,7 @@
 }
 
 
-void SSequencerSection::DrawSectionHandles( const FGeometry& AllottedGeometry, FSlateWindowElementList& OutDrawElements, int32 LayerId, ESlateDrawEffect DrawEffects, FLinearColor SelectionColor, const ISequencerHotspot* Hotspot ) const
+void SSequencerSection::DrawSectionHandles( const FGeometry& AllottedGeometry, FSlateWindowElementList& OutDrawElements, int32 LayerId, ESlateDrawEffect DrawEffects, FLinearColor SelectionColor, TSharedPtr<ITrackAreaHotspot> Hotspot ) const
 {
 	UMovieSceneSection* ThisSection = SectionInterface->GetSectionObject();
 	if (!ThisSection)
@@ -1350,23 +1500,31 @@
 
 	OutDrawElements.PushClip(FSlateClippingZone(AllottedGeometry.GetLayoutBoundingRect()));
 
-	TArray<FSectionHandle> AllUnderlappingSections;
-	AllUnderlappingSections.Add(FSectionHandle(ParentSectionArea.ToSharedRef(), SectionIndex));
-	for (const FSequencerOverlapRange& Segment : UnderlappingSegments)
+	TArray<TSharedPtr<FSectionModel>> AllUnderlappingSections;
+	if (IsSectionHighlighted(SectionInterface->GetSectionObject(), Hotspot))
 	{
-		for (FSectionHandle Section : Segment.Sections)
+		AllUnderlappingSections.Add(WeakSectionModel.Pin());
+	}
+
+	for (const FOverlappingSections& Segment : UnderlappingSegments)
+	{
+		for (const TWeakPtr<FSectionModel>& Section : Segment.Sections)
 		{
-			AllUnderlappingSections.AddUnique(Section);
+			UMovieSceneSection* SectionObject = Section.Pin()->GetSection();
+			if (IsSectionHighlighted(SectionObject, Hotspot))
+			{
+				AllUnderlappingSections.AddUnique(Section.Pin());
+			}
 		}
 	}
 
 	FGeometry SectionGeometryWithoutHandles = MakeSectionGeometryWithoutHandles(AllottedGeometry, SectionInterface);
 	FTimeToPixel TimeToPixelConverter = ConstructTimeConverterForSection(SectionGeometryWithoutHandles, *ThisSection, GetSequencer());
 
-	for (FSectionHandle Handle : AllUnderlappingSections)
+	for (TSharedPtr<FSectionModel> SectionModel : AllUnderlappingSections)
 	{
-		TSharedRef<ISequencerSection> UnderlappingSection =  Handle.GetSectionInterface();
-		UMovieSceneSection* UnderlappingSectionObj = UnderlappingSection->GetSectionObject();
+		UMovieSceneSection*           UnderlappingSectionObj = SectionModel->GetSection();
+		TSharedPtr<ISequencerSection> UnderlappingSection    = SectionModel->GetSectionInterface();
 		if (!UnderlappingSection->SectionIsResizable() || UnderlappingSectionObj->GetRange() == TRange<FFrameNumber>::All())
 		{
 			continue;
@@ -1377,16 +1535,13 @@
 		bool bRightHandleActive = false;
 
 		// Get the hovered/selected state for the section handles from the hotspot
-		if (Hotspot && (
-			Hotspot->GetType() == ESequencerHotspot::SectionResize_L ||
-			Hotspot->GetType() == ESequencerHotspot::SectionResize_R))
+		if (TSharedPtr<FSectionResizeHotspot> ResizeHotspot = HotspotCast<FSectionResizeHotspot>(Hotspot))
 		{
-			const FSectionResizeHotspot* ResizeHotspot = static_cast<const FSectionResizeHotspot*>(Hotspot);
-			if (ResizeHotspot->WeakSection == Handle.GetSectionObject())
+			if (ResizeHotspot->WeakSectionModel.Pin() == SectionModel)
 			{
 				bDrawThisSectionsHandles = true;
-				bLeftHandleActive = Hotspot->GetType() == ESequencerHotspot::SectionResize_L;
-				bRightHandleActive = Hotspot->GetType() == ESequencerHotspot::SectionResize_R;
+				bLeftHandleActive = ResizeHotspot->HandleType == FSectionResizeHotspot::Left;
+				bRightHandleActive = ResizeHotspot->HandleType == FSectionResizeHotspot::Right;
 			}
 			else
 			{
@@ -1456,10 +1611,10 @@
 }
 
 void SSequencerSection::Tick( const FGeometry& AllottedGeometry, const double InCurrentTime, const float InDeltaTime )
-{	
+{
 	if( GetVisibility() == EVisibility::Visible )
 	{
-		Layout = FSectionLayout(*ParentSectionArea, SectionIndex);
+		Layout = FSectionLayout(WeakSectionModel.Pin());
 
 		UMovieSceneSection* Section = SectionInterface->GetSectionObject();
 		if (Section && Section->HasStartFrame() && Section->HasEndFrame())
@@ -1482,22 +1637,30 @@
 
 		FGeometry SectionGeometry = MakeSectionGeometryWithoutHandles( AllottedGeometry, SectionInterface );
 		SectionInterface->Tick(SectionGeometry, ParentGeometry, InCurrentTime, InDeltaTime);
+	}
+}
 
+void SSequencerSection::AddChildLane(TSharedPtr<ITrackLaneWidget> ChildWidget)
+{
+	
+}
+
+void SSequencerSection::OnModifiedIndirectly(UMovieSceneSignedObject* Object)
+{
+	if (Object->IsA<UMovieSceneSection>())
+	{
 		UpdateUnderlappingSegments();
 	}
 }
 
 FReply SSequencerSection::OnMouseButtonDown( const FGeometry& MyGeometry, const FPointerEvent& MouseEvent )
 {
-	FSequencer& Sequencer = GetSequencer();
-
 	TArrayView<const FSequencerSelectedKey> HoveredKeys;
 	
 	// The hovered key is defined from the sequencer hotspot
-	TSharedPtr<ISequencerHotspot> Hotspot = GetSequencer().GetHotspot();
-	if (Hotspot.IsValid() && Hotspot->GetType() == ESequencerHotspot::Key)
+	if (TSharedPtr<FKeyHotspot> KeyHotspot = HotspotCast<FKeyHotspot>(GetSequencer().GetViewModel()->GetTrackArea()->GetHotspot()))
 	{
-		HoveredKeys = static_cast<FKeyHotspot*>(Hotspot.Get())->Keys;
+		HoveredKeys = KeyHotspot->Keys;
 	}
 
 	if (MouseEvent.GetEffectingButton() == EKeys::MiddleMouseButton)
@@ -1510,14 +1673,14 @@
 
 		if (NewKeys.Num())
 		{
-			Sequencer.GetSelection().EmptySelectedKeys();
+			GetSequencer().GetSelection().EmptySelectedKeys();
 			for (const FSequencerSelectedKey& NewKey : NewKeys)
 			{
-				Sequencer.GetSelection().AddToSelection(NewKey);
+				GetSequencer().GetSelection().AddToSelection(NewKey);
 			}
 
 			// Pass the event to the tool to copy the hovered key and move it
-			GetSequencer().SetHotspot( MakeShared<FKeyHotspot>(MoveTemp(NewKeys)) );
+			GetSequencer().GetViewModel()->GetTrackArea()->SetHotspot( MakeShared<FKeyHotspot>(MoveTemp(NewKeys), Sequencer) );
 
 			// Return unhandled so that the EditTool can handle the mouse down based on the newly created keyframe and prepare to move it
 			return FReply::Unhandled();
@@ -1538,23 +1701,27 @@
 
 void SSequencerSection::UpdateUnderlappingSegments()
 {
-	UMovieSceneSection* ThisSection = SectionInterface->GetSectionObject();
-	UMovieSceneTrack* Track = ThisSection ? ThisSection->GetTypedOuter<UMovieSceneTrack>() : nullptr;
-	if (!Track)
+	TSharedPtr<FSectionModel> SectionModel = WeakSectionModel.Pin();
+	if (SectionModel)
 	{
+		UnderlappingSegments = SectionModel->GetUnderlappingSections();
+		UnderlappingEasingSegments = SectionModel->GetEasingSegments();
+	}
+	else
+	{
 		UnderlappingSegments.Reset();
 		UnderlappingEasingSegments.Reset();
 	}
-	else if (Track->GetSignature() != CachedTrackSignature)
-	{
-		UnderlappingSegments = ParentSectionArea->GetUnderlappingSections(ThisSection);
-		UnderlappingEasingSegments = ParentSectionArea->GetEasingSegmentsForSection(ThisSection);
-		CachedTrackSignature = Track->GetSignature();
-	}
 }
 
 FReply SSequencerSection::OnMouseButtonDoubleClick( const FGeometry& MyGeometry, const FPointerEvent& MouseEvent )
 {
+	TSharedPtr<FSectionModel> SectionModel = WeakSectionModel.Pin();
+	if (!SectionModel)
+	{
+		return FReply::Handled();
+	}
+
 	if( MouseEvent.GetEffectingButton() == EKeys::LeftMouseButton )
 	{
 		TArray<FSequencerSelectedKey> Keys;
@@ -1562,10 +1729,7 @@
 		TArray<FKeyHandle> KeyHandles;
 		for (FSequencerSelectedKey Key : Keys)
 		{
-			if (Key.KeyHandle.IsSet())
-			{
-				KeyHandles.Add(Key.KeyHandle.GetValue());
-			}
+			KeyHandles.Add(Key.KeyHandle);
 		}
 		if (KeyHandles.Num() > 0)
 		{
@@ -1577,13 +1741,9 @@
 		{
 			// Find the object binding this node is underneath
 			FGuid ObjectBinding;
-			if (ParentSectionArea.IsValid())
+			if (TSharedPtr<IObjectBindingExtension> ObjectBindingExtension = SectionModel->FindAncestorOfType<IObjectBindingExtension>())
 			{
-				TSharedPtr<FSequencerObjectBindingNode> ObjectBindingNode = ParentSectionArea->FindParentObjectBindingNode();
-				if (ObjectBindingNode.IsValid())
-				{
-					ObjectBinding = ObjectBindingNode->GetObjectBinding();
-				}
+				ObjectBinding = ObjectBindingExtension->GetObjectGuid();
 			}
 
 			Reply = SectionInterface->OnSectionDoubleClicked(MyGeometry, MouseEvent, ObjectBinding);
@@ -1610,7 +1770,7 @@
 	GetKeysUnderMouse( MouseEvent.GetScreenSpacePosition(), MyGeometry, KeysUnderMouse );
 	if ( KeysUnderMouse.Num() )
 	{
-		GetSequencer().SetHotspot( MakeShared<FKeyHotspot>( MoveTemp(KeysUnderMouse) ) );
+		GetSequencer().GetViewModel()->GetTrackArea()->SetHotspot( MakeShared<FKeyHotspot>( MoveTemp(KeysUnderMouse), Sequencer ) );
 	}
 	// Check other interaction points in order of importance
 	else if (
@@ -1618,9 +1778,8 @@
 		!CheckForEdgeInteraction(MouseEvent, MyGeometry) &&
 		!CheckForEasingAreaInteraction(MouseEvent, MyGeometry))
 	{
-		UMovieSceneSection* ThisSection = SectionInterface->GetSectionObject();
 		// If nothing was hit, we just hit the section
-		GetSequencer().SetHotspot( MakeShareable( new FSectionHotspot(ThisSection)) );
+		GetSequencer().GetViewModel()->GetTrackArea()->SetHotspot( MakeShareable( new FSectionHotspot(WeakSectionModel, Sequencer)) );
 	}
 
 	return FReply::Unhandled();
@@ -1638,7 +1797,7 @@
 
 			for (const FSequencerSelectedKey& SelectedKey : GetSequencer().GetSelection().GetSelectedKeys())
 			{
-				const FFrameNumber CurrentTime = SelectedKey.KeyArea->GetKeyTime(SelectedKey.KeyHandle.GetValue());
+				const FFrameNumber CurrentTime = SelectedKey.WeakChannel.Pin()->GetKeyArea()->GetKeyTime(SelectedKey.KeyHandle);
 				GetSequencer().SetLocalTime(CurrentTime, ESnapTimeMode::STM_Interval);
 				break;
 			}
@@ -1654,7 +1813,7 @@
 void SSequencerSection::OnMouseLeave( const FPointerEvent& MouseEvent )
 {
 	SCompoundWidget::OnMouseLeave( MouseEvent );
-	GetSequencer().SetHotspot(nullptr);
+	GetSequencer().GetViewModel()->GetTrackArea()->SetHotspot(nullptr);
 }
 
 static float SectionThrobDurationSeconds = 1.f;
@@ -1700,33 +1859,31 @@
 	return 0.f;
 }
 
-bool SSequencerSection::IsSectionHighlighted(UMovieSceneSection* InSection, const ISequencerHotspot* Hotspot)
+bool SSequencerSection::IsSectionHighlighted(UMovieSceneSection* InSection, TSharedPtr<ITrackAreaHotspot> Hotspot)
 {
 	if (!Hotspot)
 	{
 		return false;
 	}
 
-	switch(Hotspot->GetType())
+	if (FKeyHotspot* KeyHotspot = Hotspot->CastThis<FKeyHotspot>())
 	{
-	case ESequencerHotspot::Key:
-		return static_cast<const FKeyHotspot*>(Hotspot)->Keys.ContainsByPredicate([InSection](const FSequencerSelectedKey& Key){ return Key.Section == InSection; });
+		return KeyHotspot->Keys.ContainsByPredicate([InSection](const FSequencerSelectedKey& Key){ return Key.Section == InSection; });
+	}
+	else if (FSectionEasingAreaHotspot* EasingAreaHotspot = Hotspot->CastThis<FSectionEasingAreaHotspot>())
+	{
+		return EasingAreaHotspot->Contains(InSection);
+	}
+	else if (FSectionHotspotBase* SectionHotspot = Hotspot->CastThis<FSectionHotspotBase>())
+	{
+		if (TSharedPtr<FSectionModel> SectionModel = SectionHotspot->WeakSectionModel.Pin())
+		{
+			return SectionModel->GetSection() == InSection;
+		}
+	}
 
-	case ESequencerHotspot::Section:
-		return static_cast<const FSectionHotspot*>(Hotspot)->WeakSection == InSection;
+	return false;
+}
 
-	case ESequencerHotspot::SectionResize_L:
-	case ESequencerHotspot::SectionResize_R:
-		return static_cast<const FSectionResizeHotspot*>(Hotspot)->WeakSection == InSection;
-
-	case ESequencerHotspot::EaseInHandle:
-	case ESequencerHotspot::EaseOutHandle:
-		return static_cast<const FSectionEasingHandleHotspot*>(Hotspot)->WeakSection == InSection;
-
-	case ESequencerHotspot::EasingArea:
-		return static_cast<const FSectionEasingAreaHotspot*>(Hotspot)->Contains(InSection);
-
-	default:
-		return false;
-	}
-}
+} // namespace Sequencer
+} // namespace UE
diff --git a/Engine/Source/Editor/Sequencer/Private/SSequencerSection.h b/Engine/Source/Editor/Sequencer/Private/SSequencerSection.h
--- a/Engine/Source/Editor/Sequencer/Private/SSequencerSection.h
+++ b/Engine/Source/Editor/Sequencer/Private/SSequencerSection.h
@@ -9,31 +9,39 @@
 #include "Widgets/SCompoundWidget.h"
 #include "SequencerSelectedKey.h"
 #include "Rendering/RenderingCommon.h"
-#include "DisplayNodes/SequencerTrackNode.h"
 #include "SectionLayout.h"
 #include "SequencerKeyRenderer.h"
-#include "SequencerKeyTimeCache.h"
+#include "DisplayNodes/SequencerKeyTimeCache.h"
 
+#include "EventHandlers/ISignedObjectEventHandler.h"
+#include "MVVM/ViewModels/SectionModel.h"
+#include "MVVM/Extensions/ITrackLaneExtension.h"
+
 class FPaintArgs;
 class FSequencer;
 class FSequencerSectionPainter;
 class FSlateWindowElementList;
-struct ISequencerHotspot;
 
-class SSequencerSection : public SCompoundWidget
+namespace UE
+{
+namespace Sequencer
+{
+
+class FSectionModel;
+struct ITrackAreaHotspot;
+
+class SSequencerSection : public SCompoundWidget, public ITrackLaneWidget, public UE::MovieScene::ISignedObjectEventHandler
 {
 public:
 	SLATE_BEGIN_ARGS( SSequencerSection )
 	{}
 	SLATE_END_ARGS()
 
-	void Construct( const FArguments& InArgs, TSharedRef<FSequencerTrackNode> SectionNode, int32 SectionIndex );
+	void Construct( const FArguments& InArgs, TSharedPtr<FSequencer> Sequencer, TSharedPtr<FSectionModel> InSectionModel);
+	~SSequencerSection();
 
 	TSharedPtr<ISequencerSection> GetSectionInterface() const { return SectionInterface; }
 
-	/** Caches the parent geometry to be given to section interfaces that need it on tick */
-	void CacheParentGeometry(const FGeometry& InParentGeometry) { ParentGeometry = InParentGeometry; }
-
 	virtual FVector2D ComputeDesiredSize(float) const override;
 
 private:
@@ -80,15 +88,23 @@
 	virtual FReply OnMouseButtonUp(const FGeometry& MyGeometry, const FPointerEvent& MouseEvent) override;
 	virtual void OnMouseLeave( const FPointerEvent& MouseEvent ) override;
 
+	/*~ ITrackLaneWidget Interface */
+	virtual FTrackLaneScreenAlignment GetAlignment(const FTimeToPixel& TimeToPixel, const FGeometry& InParentGeometry) const override;
+	virtual int32 GetOverlapPriority() const override;
+	virtual void ReportParentGeometry(const FGeometry& InParentGeometry) override;
+	virtual TSharedRef<const SWidget> AsWidget() const override { return AsShared(); }
+	virtual bool AcceptsChildren() const override { return true; }
+	virtual void AddChildLane(TSharedPtr<ITrackLaneWidget> ChildWidget) override;
+
 	/**
 	 * Paint the easing handles for this section
 	 */
-	void PaintEasingHandles( FSequencerSectionPainter& InPainter, FLinearColor SelectionColor, const ISequencerHotspot* Hotspot ) const;
+	void PaintEasingHandles( FSequencerSectionPainter& InPainter, FLinearColor SelectionColor, TSharedPtr<ITrackAreaHotspot> Hotspot ) const;
 
 	/**
 	 * Draw the section resize handles.
 	 */
-	void DrawSectionHandles( const FGeometry& AllottedGeometry, FSlateWindowElementList& OutDrawElements, int32 LayerId, ESlateDrawEffect DrawEffects, FLinearColor SelectionColor, const ISequencerHotspot* Hotspot ) const;
+	void DrawSectionHandles( const FGeometry& AllottedGeometry, FSlateWindowElementList& OutDrawElements, int32 LayerId, ESlateDrawEffect DrawEffects, FLinearColor SelectionColor, TSharedPtr<ITrackAreaHotspot> Hotspot ) const;
 
 	/** @return the sequencer interface */
 	FSequencer& GetSequencer() const;
@@ -103,6 +119,18 @@
 	 */
 	void UpdateUnderlappingSegments();
 
+	/**
+	 * Retrieve the tooltip text for this section
+	 */
+	FText GetToolTipText() const;
+
+	/**
+	 * Check whether this section widget is enabled or not
+	 */
+	bool IsEnabled() const;
+
+	void OnModifiedIndirectly(UMovieSceneSignedObject* Object) override;
+
 public:
 
 	/** Indicate that the current section selection should throb the specified number of times. A single throb takes 0.2s. */
@@ -120,17 +148,16 @@
 	/**
 	 * Check to see whether the specified section is highlighted
 	 */
-	static bool IsSectionHighlighted(UMovieSceneSection* InSection, const ISequencerHotspot* Hotspot);
+	static bool IsSectionHighlighted(UMovieSceneSection* InSection, TSharedPtr<ITrackAreaHotspot> Hotspot);
 
 private:
+	TWeakPtr<FSequencer> Sequencer;
 	/** Interface to section data */
 	TSharedPtr<ISequencerSection> SectionInterface;
-	/** Section area where this section resides */
-	TSharedPtr<FSequencerTrackNode> ParentSectionArea;
+	/** Section model */
+	TWeakPtr<FSectionModel> WeakSectionModel;
 	/** Cached layout generated each tick */
 	TOptional<FSectionLayout> Layout;
-	/** The index of this section in the parent section area */
-	int32 SectionIndex;
 	/** Cached parent geometry to pass down to any section interfaces that need it during tick */
 	FGeometry ParentGeometry;
 	/** The end time for a throbbing animation for selected sections */
@@ -140,13 +167,18 @@
 	/** Handle offset amount in pixels */
 	float HandleOffsetPx;
 	/** Array of segments that define other sections that reside below this one */
-	TArray<FSequencerOverlapRange> UnderlappingSegments;
+	TArray<FOverlappingSections> UnderlappingSegments;
 	/** Array of segments that define other sections that reside below this one */
-	TArray<FSequencerOverlapRange> UnderlappingEasingSegments;
-	/** The signature of the track last time the overlapping segments were updated */
-	FGuid CachedTrackSignature;
+	TArray<FOverlappingSections> UnderlappingEasingSegments;
 	/** structure responsible for rendering keys  */
-	UE::Sequencer::FKeyRenderer KeyRenderer;
+	FKeyRenderer KeyRenderer;
 
+	MovieScene::TNonIntrusiveEventHandler<MovieScene::ISignedObjectEventHandler> TrackModifiedBinding;
+
+
 	friend struct FSequencerSectionPainterImpl;
 };
+
+} // namespace Sequencer
+} // namespace UE
+
diff --git a/Engine/Source/Editor/Sequencer/Private/SSequencerShotFilterOverlay.cpp b/Engine/Source/Editor/Sequencer/Private/SSequencerShotFilterOverlay.cpp
--- a/Engine/Source/Editor/Sequencer/Private/SSequencerShotFilterOverlay.cpp
+++ b/Engine/Source/Editor/Sequencer/Private/SSequencerShotFilterOverlay.cpp
@@ -4,6 +4,7 @@
 #include "Rendering/DrawElements.h"
 #include "Styling/AppStyle.h"
 #include "CommonMovieSceneTools.h"
+#include "TimeToPixel.h"
 
 /* SSequencerShotFilterOverlay interface
  *****************************************************************************/
diff --git a/Engine/Source/Editor/Sequencer/Private/SSequencerTreeFilterStatusBar.cpp b/Engine/Source/Editor/Sequencer/Private/SSequencerTreeFilterStatusBar.cpp
--- a/Engine/Source/Editor/Sequencer/Private/SSequencerTreeFilterStatusBar.cpp
+++ b/Engine/Source/Editor/Sequencer/Private/SSequencerTreeFilterStatusBar.cpp
@@ -70,6 +70,8 @@
 
 void SSequencerTreeFilterStatusBar::UpdateText()
 {
+	using namespace UE::Sequencer;
+
 	TSharedPtr<FSequencer> Sequencer = WeakSequencer.Pin();
 	if (!ensureAlways(Sequencer))
 	{
@@ -80,18 +82,18 @@
 	FLinearColor NewColor = FLinearColor::White;
 
 	const TSharedRef<FSequencerNodeTree> NodeTree = Sequencer->GetNodeTree();
-	const TSet<TSharedRef<FSequencerDisplayNode>>& SelectedNodes = Sequencer->GetSelection().GetSelectedOutlinerNodes();
+	const TSet<TWeakPtr<FViewModel>>& SelectedOutlinerItems = Sequencer->GetSelection().GetSelectedOutlinerItems();
 
 	FFormatNamedArguments NamedArgs;
 	NamedArgs.Add("Total", NodeTree->GetTotalDisplayNodeCount());
 
-	const bool bHasSelection = SelectedNodes.Num() != 0;
+	const bool bHasSelection = SelectedOutlinerItems.Num() != 0;
 	const bool bHasFilter = NodeTree->HasActiveFilter();
 	const int32 NumFiltered = NodeTree->GetFilteredDisplayNodeCount();
 
 	if (bHasSelection)
 	{
-		NamedArgs.Add("NumSelected", SelectedNodes.Num());
+		NamedArgs.Add("NumSelected", SelectedOutlinerItems.Num());
 	}
 
 	if (bHasFilter)
diff --git a/Engine/Source/Editor/Sequencer/Private/SectionLayout.cpp b/Engine/Source/Editor/Sequencer/Private/SectionLayout.cpp
--- a/Engine/Source/Editor/Sequencer/Private/SectionLayout.cpp
+++ b/Engine/Source/Editor/Sequencer/Private/SectionLayout.cpp
@@ -1,212 +1,224 @@
-// Copyright Epic Games, Inc. All Rights Reserved.
-
-#include "SectionLayout.h"
-#include "Layout/Geometry.h"
-#include "Sequencer.h"
-#include "SequencerSettings.h"
-
-bool FSectionLayoutElementKeyFuncs::Matches(const FSectionLayoutElement& A, const FSectionLayoutElement& B)
-{
-	if (A.GetDisplayNode() != B.GetDisplayNode())
-	{
-		return false;
-	}
-	TArrayView<const TSharedRef<IKeyArea>> KeyAreasA = A.GetKeyAreas(), KeyAreasB = B.GetKeyAreas();
-	if (KeyAreasA.Num() != KeyAreasB.Num())
-	{
-		return false;
-	}
-	for (int32 Index = 0; Index < KeyAreasA.Num(); ++Index)
-	{
-		if (KeyAreasA[Index] != KeyAreasB[Index])
-		{
-			return false;
-		}
-	}
-	return true;
-}
-
-uint32 FSectionLayoutElementKeyFuncs::GetKeyHash(const FSectionLayoutElement& Key)
-{
-	uint32 Hash = GetTypeHash(Key.GetDisplayNode()) ;
-	for (TSharedRef<IKeyArea> KeyArea : Key.GetKeyAreas())
-	{
-		Hash = HashCombine(GetTypeHash(KeyArea), Hash);
-	}
-	return Hash;
-}
-
-FSectionLayoutElement FSectionLayoutElement::FromGroup(const TSharedRef<FSequencerDisplayNode>& InNode, UMovieSceneSection* InSection, float InOffset)
-{
-	TArray< TSharedRef<FSequencerSectionKeyAreaNode> > ChildKeyAreaNodes;
-	const bool bJustVisible = true;
-	InNode->GetChildKeyAreaNodesRecursively(ChildKeyAreaNodes, bJustVisible);
-
-	FSectionLayoutElement Tmp;
-	Tmp.Type = Group;
-
-	for (TSharedRef<FSequencerSectionKeyAreaNode> KeyAreaNode : ChildKeyAreaNodes)
-	{
-		TSharedPtr<IKeyArea> KeyArea = KeyAreaNode->GetKeyArea(InSection);
-		if (KeyArea.IsValid())
-		{
-			Tmp.KeyAreas.Add(KeyArea.ToSharedRef());
-		}
-	}
-	Tmp.LocalOffset = InOffset;
-	Tmp.Height = InNode->GetNodeHeight();
-	Tmp.DisplayNode = InNode;
-	return Tmp;
-}
-
-FSectionLayoutElement FSectionLayoutElement::FromKeyAreaNode(const TSharedRef<FSequencerSectionKeyAreaNode>& InKeyAreaNode, UMovieSceneSection* InSection, float InOffset)
-{
-	FSectionLayoutElement Tmp;
-	Tmp.Type = Single;
-	TSharedPtr<IKeyArea> KeyArea = InKeyAreaNode->GetKeyArea(InSection);
-	if (KeyArea.IsValid())
-	{
-		Tmp.KeyAreas.Add(KeyArea.ToSharedRef());
-	}
-	Tmp.LocalOffset = InOffset;
-	Tmp.DisplayNode = InKeyAreaNode;
-	Tmp.Height = InKeyAreaNode->GetNodeHeight();
-	return Tmp;
-}
-
-FSectionLayoutElement FSectionLayoutElement::FromTrack(const TSharedRef<FSequencerTrackNode>& InTrackNode, UMovieSceneSection* InSection, float InOffset)
-{
-	FSectionLayoutElement Tmp;
-	Tmp.Type = Single;
-	TSharedPtr<IKeyArea> KeyArea = InTrackNode->GetTopLevelKeyNode()->GetKeyArea(InSection);
-	if (KeyArea.IsValid())
-	{
-		Tmp.KeyAreas.Add(KeyArea.ToSharedRef());
-	}
-	Tmp.LocalOffset = InOffset;
-	Tmp.DisplayNode = InTrackNode;
-	Tmp.Height = InTrackNode->GetNodeHeight();
-	return Tmp;
-}
-
-FSectionLayoutElement FSectionLayoutElement::EmptySpace(const TSharedRef<FSequencerDisplayNode>& InNode, float InOffset)
-{
-	FSectionLayoutElement Tmp;
-	Tmp.Type = Single;
-	Tmp.LocalOffset = InOffset;
-	Tmp.DisplayNode = InNode;
-	Tmp.Height = InNode->GetNodeHeight();
-	return Tmp;
-}
-
-FSectionLayoutElement::EType FSectionLayoutElement::GetType() const
-{
-	return Type;
-}
-
-float FSectionLayoutElement::GetOffset() const
-{
-	return LocalOffset;
-}
-
-float FSectionLayoutElement::GetHeight() const
-{
-	return Height;
-}
-
-TArrayView<const TSharedRef<IKeyArea>> FSectionLayoutElement::GetKeyAreas() const
-{
-	return KeyAreas;
-}
-
-TSharedPtr<FSequencerDisplayNode> FSectionLayoutElement::GetDisplayNode() const
-{
-	return DisplayNode;
-}
-
-FGeometry FSectionLayoutElement::ComputeGeometry(const FGeometry& SectionAreaGeometry) const
-{
-	return SectionAreaGeometry.MakeChild(
-		FVector2D(0, LocalOffset),
-		FVector2D(SectionAreaGeometry.GetLocalSize().X, Height)
-	);
-}
-
-FSectionLayout::FSectionLayout(FSequencerTrackNode& TrackNode, int32 InSectionIndex)
-{
-	UMovieSceneSection* Section = TrackNode.GetSections()[InSectionIndex]->GetSectionObject();
-
-	auto SetupKeyArea = [this, Section](FSequencerDisplayNode& Node, float Offset){
-
-		if (Node.GetType() == ESequencerNode::KeyArea && static_cast<FSequencerSectionKeyAreaNode&>(Node).GetKeyArea(Section).IsValid())
-		{
-			Elements.Add(FSectionLayoutElement::FromKeyAreaNode(
-				StaticCastSharedRef<FSequencerSectionKeyAreaNode>(Node.AsShared()),
-				Section,
-				Offset
-			));
-		}
-		else if (Node.GetType() == ESequencerNode::Track && static_cast<FSequencerTrackNode&>(Node).GetTopLevelKeyNode().IsValid())
-		{
-			Elements.Add(FSectionLayoutElement::FromTrack(
-				StaticCastSharedRef<FSequencerTrackNode>(Node.AsShared()),
-				Section,
-				Offset
-			));
-		}
-		else if (Node.GetChildNodes().Num() && !Node.IsExpanded())
-		{
-			Elements.Add(FSectionLayoutElement::FromGroup(
-				Node.AsShared(),
-				Section,
-				Offset
-			));
-		}
-		else
-		{
-			// It's benign space
-			Elements.Add(FSectionLayoutElement::EmptySpace(
-				Node.AsShared(),
-				Offset
-			));
-		}
-
-	};
-
-	float VerticalOffset = 0.f;
-
-	// First, layout the parent
-	{
-		VerticalOffset += TrackNode.GetNodePadding().Top;
-
-		SetupKeyArea(TrackNode, VerticalOffset);
-
-		VerticalOffset += TrackNode.GetNodeHeight() + TrackNode.GetNodePadding().Bottom;
-	}
-
-	// Then any children
-	TrackNode.TraverseVisible_ParentFirst([&](FSequencerDisplayNode& Node){
-		
-		VerticalOffset += Node.GetNodePadding().Top;
-
-		SetupKeyArea(Node, VerticalOffset);
-
-		VerticalOffset += Node.GetNodeHeight() + Node.GetNodePadding().Bottom;
-		return true;
-
-	}, false);
-}
-
-const TArray<FSectionLayoutElement>& FSectionLayout::GetElements() const
-{
-	return Elements;
-}
-
-float FSectionLayout::GetTotalHeight() const
-{
-	if (Elements.Num())
-	{
-		return Elements.Last().GetOffset() + Elements.Last().GetDisplayNode()->GetNodePadding().Combined() + Elements.Last().GetHeight();
-	}
-	return 0.f;
-}
+// Copyright Epic Games, Inc. All Rights Reserved.
+
+#include "SectionLayout.h"
+#include "Layout/Geometry.h"
+#include "MVVM/ViewModels/ViewModelIterators.h"
+#include "MVVM/ViewModels/SectionModel.h"
+#include "MVVM/ViewModels/CategoryModel.h"
+#include "MVVM/ViewModels/TrackModel.h"
+#include "MVVM/ViewModels/ChannelModel.h"
+#include "Algo/Sort.h"
+
+namespace UE
+{
+namespace Sequencer
+{
+
+bool FSectionLayoutElementKeyFuncs::Matches(const FSectionLayoutElement& A, const FSectionLayoutElement& B)
+{
+	if (A.GetModel() != B.GetModel())
+	{
+		return false;
+	}
+	TArrayView<const TWeakPtr<FChannelModel>> ChannelsA = A.GetChannels(), ChannelsB = B.GetChannels();
+	if (ChannelsA.Num() != ChannelsB.Num())
+	{
+		return false;
+	}
+	for (int32 Index = 0; Index < ChannelsA.Num(); ++Index)
+	{
+		if (ChannelsA[Index] != ChannelsB[Index])
+		{
+			return false;
+		}
+	}
+	return true;
+}
+
+uint32 FSectionLayoutElementKeyFuncs::GetKeyHash(const FSectionLayoutElement& Key)
+{
+	uint32 Hash = GetTypeHash(Key.GetModel()) ;
+	for (TWeakPtr<FChannelModel> Channel : Key.GetChannels())
+	{
+		Hash = HashCombine(GetTypeHash(Channel), Hash);
+	}
+	return Hash;
+}
+
+FSectionLayoutElement FSectionLayoutElement::FromGroup(const TSharedPtr<FViewModel>& InGroup, const TSharedPtr<FViewModel>& InChannelRoot, float InOffset, float InHeight)
+{
+	FSectionLayoutElement Tmp;
+	Tmp.Type = Group;
+
+	for (TSharedPtr<FChannelModel> Channel : InChannelRoot->GetDescendantsOfType<FChannelModel>())
+	{
+		Tmp.WeakChannels.Add(Channel);
+	}
+	Tmp.LocalOffset = InOffset;
+	Tmp.Height = InHeight;
+	Tmp.DataModel = InGroup;
+	return Tmp;
+}
+
+FSectionLayoutElement FSectionLayoutElement::FromChannel(const TSharedPtr<FChannelModel>& InChannel, float InOffset, float InHeight)
+{
+	FSectionLayoutElement Tmp;
+	Tmp.Type = Single;
+	Tmp.WeakChannels.Add(InChannel);
+	Tmp.LocalOffset = InOffset;
+	Tmp.Height = InHeight;
+	Tmp.DataModel = InChannel;
+	return Tmp;
+}
+
+FSectionLayoutElement FSectionLayoutElement::EmptySpace(const TSharedPtr<FViewModel>& InModel, float InOffset, float InHeight)
+{
+	FSectionLayoutElement Tmp;
+	Tmp.Type = Single;
+	Tmp.LocalOffset = InOffset;
+	Tmp.Height = InHeight;
+	Tmp.DataModel = InModel;
+	return Tmp;
+}
+
+FSectionLayoutElement::EType FSectionLayoutElement::GetType() const
+{
+	return Type;
+}
+
+float FSectionLayoutElement::GetOffset() const
+{
+	return LocalOffset;
+}
+
+float FSectionLayoutElement::GetHeight() const
+{
+	return Height;
+}
+
+TArrayView<const TWeakPtr<FChannelModel>> FSectionLayoutElement::GetChannels() const
+{
+	return WeakChannels;
+}
+
+TSharedPtr<FViewModel> FSectionLayoutElement::GetModel() const
+{
+	return DataModel.Pin();
+}
+
+FGeometry FSectionLayoutElement::ComputeGeometry(const FGeometry& SectionAreaGeometry) const
+{
+	return SectionAreaGeometry.MakeChild(
+		FVector2D(0, LocalOffset),
+		FVector2D(SectionAreaGeometry.GetLocalSize().X, Height)
+	);
+}
+
+FSectionLayout::FSectionLayout(TSharedPtr<UE::Sequencer::FSectionModel> SectionModel)
+	: Height(0.f)
+{
+	using namespace UE::Sequencer;
+
+	float VerticalOffset = 0.f;
+	bool bLayoutChildren = true;
+
+	// First layout the parent
+	if (TSharedPtr<FTrackModel> ParentTrackModel = StaticCastSharedPtr<FTrackModel>(SectionModel->GetParent()))
+	{
+		const bool bIsExpanded = ParentTrackModel->IsExpanded();
+		const FOutlinerSizing ParentSizing = ParentTrackModel->GetOutlinerSizing();
+		VerticalOffset += ParentSizing.PaddingTop;
+		const float ParentModelHeight = ParentSizing.Height;
+
+		if (!bIsExpanded)
+		{
+			Elements.Add(FSectionLayoutElement::FromGroup(ParentTrackModel, SectionModel, VerticalOffset, ParentModelHeight));
+		}
+		else
+		{
+			Elements.Add(FSectionLayoutElement::EmptySpace(ParentTrackModel, VerticalOffset, ParentModelHeight));
+		}
+
+		VerticalOffset += ParentModelHeight + ParentSizing.PaddingBottom;
+		Height = VerticalOffset;
+
+		// Don't layout children if the parent is collapsed or filtered out.
+		bLayoutChildren = bIsExpanded && !ParentTrackModel->IsFilteredOut();
+	}
+
+	if (!bLayoutChildren)
+	{
+		// No need to sort out layout elements since there's only one.
+		return;
+	}
+
+	// Then layout the children
+	for (FParentFirstChildIterator ChildIt = SectionModel->GetDescendants(); ChildIt; ++ChildIt)
+	{
+		TSharedPtr<FViewModel> Model = *ChildIt;
+
+		bLayoutChildren = true;
+		FOutlinerSizing ChildSizing;
+
+		if (FLinkedOutlinerExtension* LinkedOutlinerExtension = Model->CastThis<FLinkedOutlinerExtension>())
+		{
+			TSharedPtr<IOutlinerExtension> LinkedOutlinerItem = LinkedOutlinerExtension->GetLinkedOutlinerItem();
+			const bool bIsFilteredOut = !LinkedOutlinerItem || LinkedOutlinerItem->IsFilteredOut();
+			if (bIsFilteredOut)
+			{
+				continue;
+			}
+
+			bLayoutChildren = LinkedOutlinerItem->IsExpanded();
+			ChildSizing = LinkedOutlinerItem->GetOutlinerSizing();
+		}
+		else if (IGeometryExtension* GeometryExtension = Model->CastThis<IGeometryExtension>())
+		{
+			FVirtualGeometry Geometry = GeometryExtension->GetVirtualGeometry();
+			ChildSizing = FOutlinerSizing(Geometry.GetHeight());
+		}
+		else
+		{
+			ensure(false);
+		}
+
+		VerticalOffset += ChildSizing.PaddingTop;
+		const float ChildModelHeight = ChildSizing.Height;
+
+		if (TSharedPtr<FChannelModel> Channel = Model->CastThisShared<FChannelModel>())
+		{
+			Elements.Add(FSectionLayoutElement::FromChannel(Channel, VerticalOffset, ChildModelHeight));
+		}
+		else if (Model->HasChildren() && !bLayoutChildren)
+		{
+			Elements.Add(FSectionLayoutElement::FromGroup(Model, Model, VerticalOffset, ChildModelHeight));
+		}
+		else
+		{
+			Elements.Add(FSectionLayoutElement::EmptySpace(Model, VerticalOffset, ChildModelHeight));
+		}
+
+		VerticalOffset += ChildModelHeight + ChildSizing.PaddingBottom;
+		Height = VerticalOffset;
+
+		if (!bLayoutChildren)
+		{
+			ChildIt.IgnoreCurrentChildren();
+		}
+	}
+
+	Algo::SortBy(Elements, &FSectionLayoutElement::GetOffset);
+}
+
+const TArray<FSectionLayoutElement>& FSectionLayout::GetElements() const
+{
+	return Elements;
+}
+
+float FSectionLayout::GetTotalHeight() const
+{
+	return Height;
+}
+
+} // namespace Sequencer
+} //namespace UE
diff --git a/Engine/Source/Editor/Sequencer/Private/SectionLayout.h b/Engine/Source/Editor/Sequencer/Private/SectionLayout.h
--- a/Engine/Source/Editor/Sequencer/Private/SectionLayout.h
+++ b/Engine/Source/Editor/Sequencer/Private/SectionLayout.h
@@ -3,12 +3,19 @@
 #pragma once
 
 #include "CoreMinimal.h"
-#include "DisplayNodes/SequencerDisplayNode.h"
-#include "DisplayNodes/SequencerSectionKeyAreaNode.h"
-#include "DisplayNodes/SequencerTrackNode.h"
 
 struct FGeometry;
 
+namespace UE
+{
+namespace Sequencer
+{
+
+class FViewModel;
+class FChannelModel;
+class FSectionModel;
+class FTrackModel;
+
 /** A layout element specifying the geometry required to render a key area */
 struct FSectionLayoutElement
 {
@@ -16,17 +23,14 @@
 	enum EType { Single, Group };
 
 	/** Construct this element from a grouped key area */
-	static FSectionLayoutElement FromGroup(const TSharedRef<FSequencerDisplayNode>& InNode, UMovieSceneSection* InSection, float InOffset);
+	static FSectionLayoutElement FromGroup(const TSharedPtr<FViewModel>& InGroup, const TSharedPtr<FViewModel>& InChannelRoot, float InOffset, float InHeight);
 
 	/** Construct this element from a single Key area node */
-	static FSectionLayoutElement FromKeyAreaNode(const TSharedRef<FSequencerSectionKeyAreaNode>& InKeyAreaNode, UMovieSceneSection* InSection, float InOffset);
+	static FSectionLayoutElement FromChannel(const TSharedPtr<FChannelModel>& InChannel, float InOffset, float InHeight);
 
 	/** Construct this element from a single Key area node */
-	static FSectionLayoutElement FromTrack(const TSharedRef<FSequencerTrackNode>& InTrackNode, UMovieSceneSection* InSection, float InOffset);
+	static FSectionLayoutElement EmptySpace(const TSharedPtr<FViewModel>& InNode, float InOffset, float InHeight);
 
-	/** Construct this element from a single Key area node */
-	static FSectionLayoutElement EmptySpace(const TSharedRef<FSequencerDisplayNode>& InNode, float InOffset);
-
 	/** Retrieve the type of this layout element */
 	EType GetType() const;
 
@@ -37,10 +41,10 @@
 	float GetHeight() const;
 
 	/** Access all the key areas that this layout element represents */
-	TArrayView<const TSharedRef<IKeyArea>> GetKeyAreas() const;
+	TArrayView<const TWeakPtr<FChannelModel>> GetChannels() const;
 
 	/** Access the display node that this layout element was generated for */
-	TSharedPtr<FSequencerDisplayNode> GetDisplayNode() const;
+	TSharedPtr<FViewModel> GetModel() const;
 
 	/** Computes the geometry for this layout as a child of the specified section area geometry */
 	FGeometry ComputeGeometry(const FGeometry& SectionAreaGeometry) const;
@@ -48,10 +52,10 @@
 private:
 
 	/** Pointer to the key area that we were generated from */
-	TArray<TSharedRef<IKeyArea>, TInlineAllocator<1>> KeyAreas;
+	TArray<TWeakPtr<FChannelModel>, TInlineAllocator<1>> WeakChannels;
 
 	/** The specific node that this key area relates to */
-	TSharedPtr<FSequencerDisplayNode> DisplayNode;
+	TWeakPtr<FViewModel> DataModel;
 
 	/** The type of this layout element */
 	EType Type;
@@ -69,7 +73,7 @@
 public:
 
 	/** Constructor that takes a display node, and the index of the section to layout */
-	FSectionLayout(FSequencerTrackNode& TrackNode, int32 SectionIndex);
+	FSectionLayout(TSharedPtr<UE::Sequencer::FSectionModel> SectionModel);
 
 	/** Get all layout elements that we generated */
 	const TArray<FSectionLayoutElement>& GetElements() const;
@@ -80,6 +84,8 @@
 private:
 	/** Array of layout elements that we generated */
 	TArray<FSectionLayoutElement> Elements;
+
+	float Height;
 };
 
 /** Key funcs for using a section layout element as a key. Intentionally not supported implicitly due to performance reasons. */
@@ -94,3 +100,7 @@
 	static bool Matches(const FSectionLayoutElement& A, const FSectionLayoutElement& B);
 	static uint32 GetKeyHash(const FSectionLayoutElement& Key);
 };
+
+} // namespace Sequencer
+} // namespace UE
+
