Sequencer: MVVM2 branch and Layer Bars

diff --git a/Engine/Source/Editor/Sequencer/Private/Sequencer.cpp b/Engine/Source/Editor/Sequencer/Private/Sequencer.cpp
--- a/Engine/Source/Editor/Sequencer/Private/Sequencer.cpp
+++ b/Engine/Source/Editor/Sequencer/Private/Sequencer.cpp
@@ -4,6 +4,26 @@
 #include "Engine/EngineTypes.h"
 #include "GameFramework/Actor.h"
 #include "Engine/World.h"
+#include "Algo/RemoveIf.h"
+#include "MVVM/SharedViewModelData.h"
+#include "MVVM/ViewModels/SequenceModel.h"
+#include "MVVM/ViewModels/SectionModel.h"
+#include "MVVM/ViewModels/FolderModel.h"
+#include "MVVM/ViewModels/ChannelModel.h"
+#include "MVVM/ViewModels/CategoryModel.h"
+#include "MVVM/ViewModels/TrackModel.h"
+#include "MVVM/ViewModels/TrackRowModel.h"
+#include "MVVM/ViewModels/ViewModelIterators.h"
+#include "MVVM/ViewModels/SequencerEditorViewModel.h"
+#include "MVVM/ViewModels/SequencerOutlinerViewModel.h"
+#include "MVVM/ViewModels/SequencerTrackAreaViewModel.h"
+#include "MVVM/ViewModels/SequencerEditorViewModel.h"
+#include "MVVM/Views/SOutlinerView.h"
+#include "MVVM/Views/IOutlinerSelectionHandler.h"
+#include "MVVM/Extensions/IDeletableExtension.h"
+#include "MVVM/CurveEditorIntegrationExtension.h"
+#include "MVVM/ObjectBindingModelStorageExtension.h"
+#include "MVVM/SectionModelStorageExtension.h"
 #include "Camera/PlayerCameraManager.h"
 #include "Misc/MessageDialog.h"
 #include "Containers/ArrayBuilder.h"
@@ -43,10 +63,7 @@
 #include "FileHelpers.h"
 #include "UnrealEdGlobals.h"
 #include "SequencerCommands.h"
-#include "DisplayNodes/SequencerFolderNode.h"
-#include "DisplayNodes/SequencerObjectBindingNode.h"
 #include "ISequencerSection.h"
-#include "DisplayNodes/SequencerTrackNode.h"
 #include "MovieSceneClipboard.h"
 #include "SequencerCommonHelpers.h"
 #include "SequencerMarkedFrameHelper.h"
@@ -63,7 +80,7 @@
 #include "EditorWidgetsModule.h"
 #include "IAssetViewport.h"
 #include "EditorSupportDelegates.h"
-#include "SSequencerTreeView.h"
+#include "MVVM/Views/SOutlinerView.h"
 #include "ScopedTransaction.h"
 #include "Tracks/MovieScene3DTransformTrack.h"
 #include "Tracks/MovieScene3DAttachTrack.h"
@@ -89,7 +106,6 @@
 #include "Framework/Notifications/NotificationManager.h"
 #include "Widgets/Notifications/SNotificationList.h"
 #include "Widgets/Input/STextEntryPopup.h"
-#include "ISequencerHotspot.h"
 #include "SequencerHotspots.h"
 #include "MovieSceneCaptureDialogModule.h"
 #include "AutomatedLevelSequenceCapture.h"
@@ -99,7 +115,6 @@
 #include "IContentBrowserSingleton.h"
 #include "ContentBrowserModule.h"
 #include "PackageTools.h"
-#include "VirtualTrackArea.h"
 #include "SequencerUtilities.h"
 #include "Tracks/MovieSceneCinematicShotTrack.h"
 #include "CineCameraActor.h"
@@ -147,7 +162,8 @@
 #include "ActorTreeItem.h"
 #include "Widgets/Layout/SSpacer.h"
 #include "CurveEditorCommands.h"
-
+#include "Tools/SequencerEditTool_Movement.h"
+#include "Tools/SequencerEditTool_Selection.h"
 #include "EntitySystem/MovieSceneEntitySystemLinker.h"
 #include "EntitySystem/MovieScenePreAnimatedStateSystem.h"
 #include "Systems/MovieSceneMotionVectorSimulationSystem.h"
@@ -228,6 +244,87 @@
 	TWeakPtr<FSequencer> WeakSequencer;
 };
 
+class FSequencerOutlinerSelectionHandler : 
+	public UE::Sequencer::IOutlinerSelectionHandler
+{
+public:
+
+	bool bForwardSelectionToView;
+	/** True if we are waiting for menus to be closed to synchronize external selection. */
+	bool bPendingSelectionChangedNotification;
+
+	FSequencerOutlinerSelectionHandler(FSequencerSelection& InSelection)
+		: Selection(InSelection)
+	{
+		bForwardSelectionToView = true;
+		bPendingSelectionChangedNotification = false;
+		Selection.GetOnOutlinerNodeSelectionChanged().AddRaw(this, &FSequencerOutlinerSelectionHandler::OnSelectedOutlinerNodesChanged);
+	}
+
+	void SetTreeViews(TSharedPtr<UE::Sequencer::SOutlinerView> InTreeView, TSharedPtr<UE::Sequencer::SOutlinerView> InPinnedTreeView)
+	{
+		WeakTreeView = InTreeView;
+		WeakPinnedTreeView = InPinnedTreeView;
+	}
+
+	void OnSelectedOutlinerNodesChanged()
+	{
+		using namespace UE::Sequencer;
+
+		if (bForwardSelectionToView)
+		{
+			TGuardValue<bool> Guard(bForwardSelectionToView, false);
+
+			// Called when our selection has changed programatically from code
+			TSharedPtr<SOutlinerView> TreeView = WeakTreeView.Pin();
+			if (ensure(TreeView))
+			{
+				TreeView->ForceSetSelectedItems(Selection.GetSelectedOutlinerItems());
+			}
+
+			TSharedPtr<SOutlinerView> PinnedTreeView = WeakPinnedTreeView.Pin();
+			if (ensure(PinnedTreeView))
+			{
+				PinnedTreeView->ForceSetSelectedItems(Selection.GetSelectedOutlinerItems());
+			}
+		}
+	}
+
+	void SelectOutlinerItems(const TArrayView<UE::Sequencer::TWeakViewModelPtr<UE::Sequencer::IOutlinerExtension>>& Items, bool bRightMouseButtonDown) override
+	{
+		using namespace UE::Sequencer;
+
+		TGuardValue<bool> Guard(bForwardSelectionToView, false);
+
+		Selection.SuspendBroadcast();
+		Selection.EmptySelectedOutlinerNodes();
+		for (const TWeakViewModelPtr<IOutlinerExtension>& Item : Items)
+		{
+			if (TViewModelPtr<IOutlinerExtension> ItemPtr = Item.Pin())
+			{
+				Selection.AddToOutlinerSelection(ItemPtr);
+			}
+		}
+
+		Selection.ResumeBroadcast();
+
+		if (bRightMouseButtonDown)
+		{
+			bPendingSelectionChangedNotification = true;
+		}
+		else
+		{
+			Selection.GetOnOutlinerNodeSelectionChanged().Broadcast();
+		}
+	}
+
+private:
+
+	FSequencerSelection& Selection;
+	TWeakPtr<UE::Sequencer::SOutlinerView> WeakTreeView;
+	TWeakPtr<UE::Sequencer::SOutlinerView> WeakPinnedTreeView;
+};
+
 class FSequencerCurveEditor : public FCurveEditor
 {
 public:
@@ -265,6 +362,8 @@
 
 void FSequencer::InitSequencer(const FSequencerInitParams& InitParams, const TSharedRef<ISequencerObjectChangeListener>& InObjectChangeListener, const TArray<FOnCreateTrackEditor>& TrackEditorDelegates, const TArray<FOnCreateEditorObjectBinding>& EditorObjectBindingDelegates)
 {
+	using namespace UE::Sequencer;
+
 	bIsEditingWithinLevelEditor = InitParams.bEditWithinLevelEditor;
 	ScrubStyle = InitParams.ViewParams.ScrubberStyle;
 	HostCapabilities = InitParams.HostCapabilities;
@@ -432,6 +531,34 @@
 	InitialValueCache = UE::MovieScene::FInitialValueCache::GetGlobalInitialValues();
 	RootTemplateInstance.GetEntitySystemLinker()->AddExtension(InitialValueCache.Get());
 
+	// Create tools and bind them to this sequencer
+	for( int32 DelegateIndex = 0; DelegateIndex < TrackEditorDelegates.Num(); ++DelegateIndex )
+	{
+		check( TrackEditorDelegates[DelegateIndex].IsBound() );
+		// Tools may exist in other modules, call a delegate that will create one for us 
+		TSharedRef<ISequencerTrackEditor> TrackEditor = TrackEditorDelegates[DelegateIndex].Execute( SharedThis( this ) );
+
+		if (TrackEditor->SupportsSequence(InitParams.RootSequence))
+		{
+			TrackEditors.Add( TrackEditor );
+		}
+	}
+
+	ViewModel = MakeShared<FSequencerEditorViewModel>(SharedThis(this));
+	ViewModel->InitializeEditor();
+	ViewModel->SetSequence(InitParams.RootSequence);
+
+	FSequencerOutlinerViewModel* OutlinerViewModel = ViewModel->GetOutliner()->CastThisChecked<FSequencerOutlinerViewModel>();
+	{
+		OutlinerViewModel->OnGetAddMenuContent = InitParams.ViewParams.OnGetAddMenuContent;
+		OutlinerViewModel->OnBuildCustomContextMenuForGuid = InitParams.ViewParams.OnBuildCustomContextMenuForGuid;
+	}
+
+	NodeTree->SetRootNode(ViewModel->GetRootModel());
+	Selection.SetRootModel(ViewModel->GetRootModel());
+
+	SelectionHandler = MakeShared<FSequencerOutlinerSelectionHandler>(Selection);
+
 	ResetTimeController();
 
 	UpdateTimeBases();
@@ -473,30 +600,21 @@
 		.OnViewRangeChanged( this, &FSequencer::SetViewRange )
 		.OnClampRangeChanged( this, &FSequencer::OnClampRangeChanged )
 		.OnGetNearestKey( this, &FSequencer::OnGetNearestKey )
-		.OnGetAddMenuContent(InitParams.ViewParams.OnGetAddMenuContent)
-		.OnBuildCustomContextMenuForGuid(InitParams.ViewParams.OnBuildCustomContextMenuForGuid)
 		.OnGetPlaybackSpeeds(InitParams.ViewParams.OnGetPlaybackSpeeds)
 		.OnReceivedFocus(InitParams.ViewParams.OnReceivedFocus)
 		.OnInitToolMenuContext(InitParams.ViewParams.OnInitToolMenuContext)
 		.AddMenuExtender(InitParams.ViewParams.AddMenuExtender)
-		.ToolbarExtender(InitParams.ViewParams.ToolbarExtender);
+		.ToolbarExtender(InitParams.ViewParams.ToolbarExtender)
+		.SelectionHandler(SelectionHandler);
 
+	SelectionHandler->SetTreeViews(SequencerWidget->GetTreeView(), SequencerWidget->GetPinnedTreeView());
+
 	// When undo occurs, get a notification so we can make sure our view is up to date
 	GEditor->RegisterForUndo(this);
 
-	// Create tools and bind them to this sequencer
-	for( int32 DelegateIndex = 0; DelegateIndex < TrackEditorDelegates.Num(); ++DelegateIndex )
-	{
-		check( TrackEditorDelegates[DelegateIndex].IsBound() );
-		// Tools may exist in other modules, call a delegate that will create one for us 
-		TSharedRef<ISequencerTrackEditor> TrackEditor = TrackEditorDelegates[DelegateIndex].Execute( SharedThis( this ) );
+	ViewModel->GetTrackArea()->AddEditTool(MakeShared<FSequencerEditTool_Selection>(*this, *SequencerWidget->GetTrackAreaWidget().Get()));
+	ViewModel->GetTrackArea()->AddEditTool(MakeShared<FSequencerEditTool_Movement>(*this));
 
-		if (TrackEditor->SupportsSequence(GetFocusedMovieSceneSequence()))
-		{
-			TrackEditors.Add( TrackEditor );
-		}
-	}
-
 	for (int32 DelegateIndex = 0; DelegateIndex < EditorObjectBindingDelegates.Num(); ++DelegateIndex)
 	{
 		check(EditorObjectBindingDelegates[DelegateIndex].IsBound());
@@ -571,7 +689,7 @@
 	, bHasPreAnimatedInfo(false)
 {
 	Selection.GetOnOutlinerNodeSelectionChanged().AddRaw(this, &FSequencer::OnSelectedOutlinerNodesChanged);
-	Selection.GetOnNodesWithSelectedKeysOrSectionsChanged().AddRaw(this, &FSequencer::OnSelectedOutlinerNodesChanged);
+	Selection.GetOnKeySelectionChanged().AddRaw(this, &FSequencer::OnSelectedOutlinerNodesChanged);
 	Selection.GetOnOutlinerNodeSelectionChangedObjectGuids().AddRaw(this, &FSequencer::OnSelectedOutlinerNodesChanged);
 
 	// Exposes the sequencer and curve editor command lists to subscribers from other systems
@@ -600,7 +718,12 @@
 	}
 
 	AcquiredResources.Release();
+
+	SelectionHandler.Reset();
+
+	ViewModel.Reset();
 	SequencerWidget.Reset();
+
 	TrackEditors.Empty();
 }
 
@@ -648,16 +771,58 @@
 	OnCloseEventDelegate.Broadcast(AsShared());
 }
 
+TSharedPtr<ISequencerTrackEditor> FSequencer::GetTrackEditor(UMovieSceneTrack* InTrack)
+{
+	UClass* TrackClass = InTrack->GetClass();
+	FObjectKey TrackClassKey(TrackClass);
 
+	TSharedPtr<ISequencerTrackEditor> TrackEditor = TrackEditorsByType.FindRef(TrackClassKey);
+	if (!TrackEditor)
+	{
+		for (TSharedPtr<ISequencerTrackEditor> OtherTrackEditor : TrackEditors)
+		{
+			if (OtherTrackEditor->SupportsType(TrackClass))
+			{
+				TrackEditor = OtherTrackEditor;
+				TrackEditorsByType.Add(TrackClassKey, TrackEditor);
+				break;
+			}
+		}
+	}
+
+	checkf(TrackEditor, TEXT("Unable to find a track editor for track type %s"), *TrackClass->GetName());
+	return TrackEditor;
+}
+
 void FSequencer::Tick(float InDeltaTime)
 {
+	using namespace UE::Sequencer;
+
 	static bool bEnableRefCountCheck = true;
-	if (bEnableRefCountCheck && !FSlateApplication::Get().AnyMenusVisible())
+
+	if (!FSlateApplication::Get().AnyMenusVisible())
 	{
-		const int32 SequencerRefCount = AsShared().GetSharedReferenceCount() - 1;
-		ensureAlwaysMsgf(SequencerRefCount == 1, TEXT("Multiple persistent shared references detected for Sequencer. There should only be one persistent authoritative reference. Found %d additional references which will result in FSequencer not being released correctly."), SequencerRefCount - 1);
+		if (SelectionHandler->bPendingSelectionChangedNotification)
+		{
+			SelectionHandler->bPendingSelectionChangedNotification = false;
+			HandleSelectedOutlinerNodesChanged();
+		}
+
+		if (bEnableRefCountCheck)
+		{
+			const int32 SequencerRefCount = AsShared().GetSharedReferenceCount() - 1;
+			ensureAlwaysMsgf(SequencerRefCount == 1, TEXT("Multiple persistent shared references detected for Sequencer. There should only be one persistent authoritative reference. Found %d additional references which will result in FSequencer not being released correctly."), SequencerRefCount - 1);
+		}
 	}
 
+	TSharedPtr<FSharedViewModelData> SharedModelData = ViewModel->GetRootModel()->GetSharedData();
+	if (SharedModelData)
+	{
+		SharedModelData->ReportLatentHierarchicalOperations();
+	}
+
+	UMovieSceneSignedObject::ResetImplicitScopedModifyDefer();
+
 	if (bNeedsInvalidateCachedData)
 	{
 		InvalidateCachedData();
@@ -1083,6 +1248,10 @@
 	bGlobalMarkedFramesCached = false;
 }
 
+TSharedPtr<UE::Sequencer::FEditorViewModel> FSequencer::GetViewModel() const
+{
+	return ViewModel;
+}
 
 void FSequencer::SuppressAutoEvaluation(UMovieSceneSequence* Sequence, const FGuid& InSequenceSignature)
 {
@@ -1129,7 +1298,7 @@
 			FMovieScenePossessable* ChildPossessable = OwnerMovieScene->FindPossessable(PossessableGuid);
 			if (ensure(ChildPossessable))
 			{
-				ChildPossessable->SetParent(ParentGuid);
+				ChildPossessable->SetParent(ParentGuid, OwnerMovieScene);
 			}
 
 			FMovieSceneSpawnable* ParentSpawnable = OwnerMovieScene->FindSpawnable(ParentGuid);
@@ -1178,16 +1347,15 @@
 
 void FSequencer::GetKeysFromSelection(TUniquePtr<FSequencerKeyCollection>& KeyCollection, float DuplicateThresholdSeconds)
 {
+	using namespace UE::Sequencer;
+
 	if (!KeyCollection.IsValid())
 	{
 		KeyCollection.Reset(new FSequencerKeyCollection);
 	}
 
-	TArray<FSequencerDisplayNode*> SelectedNodes;
-	for (const TSharedRef<FSequencerDisplayNode>& Node : Selection.GetSelectedOutlinerNodes())
-	{
-		SelectedNodes.Add(&Node.Get());
-	}
+	TArray<TSharedRef<FViewModel>> SelectedItems;
+	Selection.GetSelectedOutlinerItems(SelectedItems);
 
 	int64 TotalMaxSeconds = static_cast<int64>(TNumericLimits<int32>::Max() / GetFocusedTickResolution().AsDecimal());
 
@@ -1201,21 +1369,20 @@
 		ThresholdFrames.Value = TotalMaxSeconds;
 	}
 
-	KeyCollection->Update(FSequencerKeyCollectionSignature::FromNodesRecursive(SelectedNodes, ThresholdFrames));
+	KeyCollection->Update(FSequencerKeyCollectionSignature::FromNodesRecursive(SelectedItems, ThresholdFrames));
 }
 
 void FSequencer::GetAllKeys(TUniquePtr<FSequencerKeyCollection>& KeyCollection, float DuplicateThresholdSeconds)
 {
+	using namespace UE::Sequencer;
+
 	if (!KeyCollection.IsValid())
 	{
 		KeyCollection.Reset(new FSequencerKeyCollection);
 	}
 
-	TArray<FSequencerDisplayNode*> AllNodes;
-	for (TSharedRef<FSequencerDisplayNode> Node : NodeTree->GetAllNodes())
-	{
-		AllNodes.Add(&Node.Get());
-	}
+	TArray<TSharedRef<FViewModel>> AllNodes;
+	NodeTree->GetAllNodes(AllNodes);
 
 	FFrameNumber ThresholdFrames = (DuplicateThresholdSeconds * GetFocusedTickResolution()).FloorToFrame();
 	KeyCollection->Update(FSequencerKeyCollectionSignature::FromNodesRecursive(AllNodes, ThresholdFrames));
@@ -1520,8 +1687,7 @@
 		NotifyMovieSceneDataChanged( EMovieSceneDataChangeType::MovieSceneStructureItemRemoved );
 	}
 
-	Selection.EmptySelectedSections();
-	SequencerHelpers::ValidateNodesWithSelectedKeysOrSections(*this);
+	Selection.EmptySelectedTrackAreaItems();
 }
 
 
@@ -1559,7 +1725,6 @@
 		NotifyMovieSceneDataChanged( EMovieSceneDataChangeType::TrackValueChanged );
 
 		Selection.EmptySelectedKeys();
-		SequencerHelpers::ValidateNodesWithSelectedKeysOrSections(*this);
 	}
 }
 
@@ -2345,7 +2510,13 @@
 
 void FSequencer::RecreateCurveEditor()
 {
-	GetNodeTree()->RecreateCurveEditor();
+	using namespace UE::Sequencer;
+
+	FCurveEditorIntegrationExtension* CurveEditorIntegration = ViewModel->GetRootModel()->CastDynamic<FCurveEditorIntegrationExtension>();
+	if (ensure(CurveEditorIntegration))
+	{
+		CurveEditorIntegration->RecreateCurveEditor();
+	}
 }
 
 FAnimatedRange FSequencer::GetViewRange() const
@@ -2436,9 +2607,12 @@
 }
 
 
-void FSequencer::SelectInSelectionRange(const TSharedRef<FSequencerDisplayNode>& DisplayNode, const TRange<FFrameNumber>& SelectionRange, bool bSelectKeys, bool bSelectSections)
+void FSequencer::SelectInSelectionRange(const TSharedPtr<UE::Sequencer::FViewModel>& Item, const TRange<FFrameNumber>& SelectionRange, bool bSelectKeys, bool bSelectSections)
 {
-	if (DisplayNode->IsHidden())
+	using namespace UE::Sequencer;
+
+	IOutlinerExtension* Outliner = Item->CastThis<IOutlinerExtension>();
+	if (Outliner && Outliner->IsFilteredOut())
 	{
 		return;
 	}
@@ -2447,13 +2621,11 @@
 	{
 		TArray<FKeyHandle> HandlesScratch;
 
-		TSet<TSharedPtr<IKeyArea>> KeyAreas;
-		SequencerHelpers::GetAllKeyAreas(DisplayNode, KeyAreas);
+		for (TSharedPtr<FChannelModel> Channel : Item->GetDescendantsOfType<FChannelModel>())
+		{
+			TSharedPtr<IKeyArea> KeyArea = Channel->GetKeyArea();
+			UMovieSceneSection*  Section = KeyArea->GetOwningSection();
 
-		for (TSharedPtr<IKeyArea> KeyArea : KeyAreas)
-		{
-			UMovieSceneSection* Section = KeyArea->GetOwningSection();
-
 			if (Section)
 			{
 				HandlesScratch.Reset();
@@ -2461,7 +2633,7 @@
 
 				for (int32 Index = 0; Index < HandlesScratch.Num(); ++Index)
 				{
-					Selection.AddToSelection(FSequencerSelectedKey(*Section, KeyArea, HandlesScratch[Index]));
+					Selection.AddToSelection(FSequencerSelectedKey(*Section, Channel, HandlesScratch[Index]));
 				}
 			}
 		}
@@ -2469,21 +2641,19 @@
 
 	if (bSelectSections)
 	{
-		TSet<TWeakObjectPtr<UMovieSceneSection>> OutSections;
-		SequencerHelpers::GetAllSections(DisplayNode, OutSections);
-
-		for (auto Section : OutSections)
+		for (TSharedPtr<FSectionModel> SectionModel : Item->GetDescendantsOfType<FSectionModel>())
 		{
-			if (Section.IsValid() && Section->GetRange().Overlaps(SelectionRange) && Section->HasStartFrame() && Section->HasEndFrame())
+			TRange<FFrameNumber> SectionRange = SectionModel->GetRange();
+			if (SectionRange.Overlaps(SelectionRange) && SectionRange.GetLowerBound().IsClosed() && SectionRange.GetUpperBound().IsClosed())
 			{
-				Selection.AddToSelection(Section.Get());
+				Selection.AddToSelection(SectionModel);
 			}
 		}
 	}
 
-	for (const auto& ChildNode : DisplayNode->GetChildNodes())
+	for (TSharedPtr<FViewModel> Child : Item->GetChildren())
 	{
-		SelectInSelectionRange(ChildNode, SelectionRange, bSelectKeys, bSelectSections);
+		SelectInSelectionRange(Child, SelectionRange, bSelectKeys, bSelectSections);
 	}
 }
 
@@ -2494,6 +2664,8 @@
 
 void FSequencer::SelectInSelectionRange(bool bSelectKeys, bool bSelectSections)
 {
+	using namespace UE::Sequencer;
+
 	UMovieSceneSequence* Sequence = GetFocusedMovieSceneSequence();
 	UMovieScene* MovieScene = Sequence->GetMovieScene();
 	TRange<FFrameNumber> SelectionRange = MovieScene->GetSelectionRange();
@@ -2501,9 +2673,9 @@
 	// Don't empty all selection, just keys and sections
 	Selection.SuspendBroadcast();
 	Selection.EmptySelectedKeys();
-	Selection.EmptySelectedSections();
+	Selection.EmptySelectedTrackAreaItems();
 
-	for (const TSharedRef<FSequencerDisplayNode>& DisplayNode : NodeTree->GetRootNodes())
+	for (const TViewModelPtr<IOutlinerExtension>& DisplayNode : NodeTree->GetRootNodes())
 	{
 		SelectInSelectionRange(DisplayNode, SelectionRange, bSelectKeys, bSelectSections);
 	}
@@ -2512,28 +2684,38 @@
 
 void FSequencer::SelectForward()
 {
+	using namespace UE::Sequencer;
+
 	FFrameRate TickResolution = GetFocusedTickResolution();
 	FFrameNumber CurrentFrame = GetLocalTime().ConvertTo(TickResolution).CeilToFrame();
 	TRange<FFrameNumber> SelectionRange(CurrentFrame, TNumericLimits<FFrameNumber>::Max());
 
-	TSet<TSharedRef<FSequencerDisplayNode> > DisplayNodes = Selection.GetNodesWithSelectedKeysOrSections();
-	if (DisplayNodes.Num() == 0)
+	TSet<TWeakPtr<FViewModel>> SelectedItems = Selection.GetNodesWithSelectedKeysOrSections();
+	if (SelectedItems.Num() == 0)
 	{
-		DisplayNodes = Selection.GetSelectedOutlinerNodes();
+		SelectedItems = Selection.GetSelectedOutlinerItems();
 	}
-	if (DisplayNodes.Num() == 0)
+	if (SelectedItems.Num() == 0)
 	{
-		DisplayNodes.Append(NodeTree->GetAllNodes());
+		TArray<TSharedPtr<FViewModel>> AllNodes;
+		NodeTree->GetAllNodes(AllNodes);
+		for (TSharedPtr<FViewModel> Node : AllNodes)
+		{
+			SelectedItems.Add(Node);
+		}
 	}
 
-	if (DisplayNodes.Num() > 0)
+	if (SelectedItems.Num() > 0)
 	{
 		Selection.SuspendBroadcast();
 		Selection.EmptySelectedKeys();
-		Selection.EmptySelectedSections();
-		for (TSharedRef<FSequencerDisplayNode>& DisplayNode : DisplayNodes)
+		Selection.EmptySelectedTrackAreaItems();
+		for (const TWeakPtr<FViewModel>& WeakItem : SelectedItems)
 		{
-			SelectInSelectionRange(DisplayNode, SelectionRange, true, true);
+			if (TSharedPtr<FViewModel> Item = WeakItem.Pin())
+			{
+				SelectInSelectionRange(Item, SelectionRange, true, true);
+			}
 		}
 		Selection.ResumeBroadcast();
 	}
@@ -2542,28 +2724,38 @@
 
 void FSequencer::SelectBackward()
 {
+	using namespace UE::Sequencer;
+
 	FFrameRate TickResolution = GetFocusedTickResolution();
 	FFrameNumber CurrentFrame = GetLocalTime().ConvertTo(TickResolution).CeilToFrame();
-	TRange<FFrameNumber> SelectionRange(TNumericLimits<FFrameNumber>::Min(), CurrentFrame);
+	TRange<FFrameNumber> SelectionRange(CurrentFrame, TNumericLimits<FFrameNumber>::Max());
 
-	TSet<TSharedRef<FSequencerDisplayNode> > DisplayNodes = Selection.GetNodesWithSelectedKeysOrSections();
-	if (DisplayNodes.Num() == 0)
+	TSet<TWeakPtr<FViewModel>> SelectedItems = Selection.GetNodesWithSelectedKeysOrSections();
+	if (SelectedItems.Num() == 0)
 	{
-		DisplayNodes = Selection.GetSelectedOutlinerNodes();
+		SelectedItems = Selection.GetSelectedOutlinerItems();
 	}
-	if (DisplayNodes.Num() == 0)
+	if (SelectedItems.Num() == 0)
 	{
-		DisplayNodes.Append(NodeTree->GetAllNodes());
+		TArray<TSharedPtr<FViewModel>> AllNodes;
+		NodeTree->GetAllNodes(AllNodes);
+		for (const TSharedPtr<FViewModel>& Node : AllNodes)
+		{
+			SelectedItems.Add(Node);
+		}
 	}
 
-	if (DisplayNodes.Num() > 0)
+	if (SelectedItems.Num() > 0)
 	{
 		Selection.SuspendBroadcast();
 		Selection.EmptySelectedKeys();
-		Selection.EmptySelectedSections();
-		for (TSharedRef<FSequencerDisplayNode>& DisplayNode : DisplayNodes)
+		Selection.EmptySelectedTrackAreaItems();
+		for (const TWeakPtr<FViewModel>& WeakItem : SelectedItems)
 		{
-			SelectInSelectionRange(DisplayNode, SelectionRange, true, true);
+			if (TSharedPtr<FViewModel> Item = WeakItem.Pin())
+			{
+				SelectInSelectionRange(Item, SelectionRange, true, true);
+			}
 		}
 		Selection.ResumeBroadcast();
 	}
@@ -3665,7 +3857,7 @@
 	return FGuid();
 }
 
-UMovieSceneFolder* FSequencer::CreateFoldersRecursively(const TArray<FName>& FolderPath, int32 FolderPathIndex, UMovieScene* OwningMovieScene, UMovieSceneFolder* ParentFolder, const TArray<UMovieSceneFolder*>& FoldersToSearch)
+UMovieSceneFolder* FSequencer::CreateFoldersRecursively(const TArray<FName>& FolderPath, int32 FolderPathIndex, UMovieScene* OwningMovieScene, UMovieSceneFolder* ParentFolder, TArrayView<UMovieSceneFolder* const> FoldersToSearch)
 {
 	// An empty folder path won't create a folder
 	if (FolderPath.Num() == 0)
@@ -3702,7 +3894,7 @@
 		{
 			// If we have no parent folder then we must be at the root so we add it to the root of the movie scene
 			OwningMovieScene->Modify();
-			OwningMovieScene->GetRootFolders().Add(FolderToUse);
+			OwningMovieScene->AddRootFolder(FolderToUse);
 		}
 	}
 
@@ -4207,9 +4399,16 @@
 
 void FSequencer::ResetPerMovieSceneData()
 {
+	using namespace UE::Sequencer;
+
 	//@todo Sequencer - We may want to preserve selections when moving between movie scenes
 	Selection.Empty();
 
+	// This will reinitialize all extensions to rebuild the view-model hierarchy
+	UMovieSceneSequence* CurrentSequence = GetFocusedMovieSceneSequence();
+	TSharedPtr<FSequenceModel> RootSequenceModel = ViewModel->GetRootModel().ImplicitCast();
+	RootSequenceModel->SetSequence(CurrentSequence, ActiveTemplateIDs.Top());
+
 	RefreshTree();
 
 	UpdateTimeBoundsToFocusedMovieScene();
@@ -4705,7 +4904,7 @@
 
 FReply FSequencer::JumpToPreviousKey()
 {
-	if (Selection.GetSelectedOutlinerNodes().Num())
+	if (Selection.GetSelectedOutlinerItems().Num())
 	{
 		GetKeysFromSelection(SelectedKeyCollection, SMALL_NUMBER);
 	}
@@ -4734,7 +4933,7 @@
 
 FReply FSequencer::JumpToNextKey()
 {
-	if (Selection.GetSelectedOutlinerNodes().Num())
+	if (Selection.GetSelectedOutlinerItems().Num())
 	{
 		GetKeysFromSelection(SelectedKeyCollection, SMALL_NUMBER);
 	}
@@ -5329,11 +5528,6 @@
 	return bReadOnly || (GetFocusedMovieSceneSequence() && GetFocusedMovieSceneSequence()->GetMovieScene()->IsReadOnly());
 }
 
-void FSequencer::VerticalScroll(float ScrollAmountUnits)
-{
-	SequencerWidget->GetTreeView()->ScrollByDelta(ScrollAmountUnits);
-}
-
 FGuid FSequencer::AddSpawnable(UObject& Object, UActorFactory* ActorFactory)
 {
 	UMovieSceneSequence* Sequence = GetFocusedMovieSceneSequence();
@@ -5459,109 +5653,34 @@
 	return bWasConsumed;
 }
 
+bool FSequencer::OnRequestNodeDeleted( TSharedRef<FViewModel> NodeToBeDeleted, const bool bKeepState )
+{
+	using namespace UE::Sequencer;
+	using namespace UE::MovieScene;
 
-// Takes a display node and traverses it's parents to find the nearest track node if any.  Also collects the names of the nodes which make
-// up the path from the track node to the display node being checked.  The name path includes the name of the node being checked, but not
-// the name of the track node.
-void GetParentTrackNodeAndNamePath(TSharedRef<const FSequencerDisplayNode> DisplayNode, TSharedPtr<FSequencerTrackNode>& OutParentTrack, TArray<FName>& OutNamePath )
-{
-	TArray<FName> PathToTrack;
-	PathToTrack.Add( DisplayNode->GetNodeName() );
-	TSharedPtr<FSequencerDisplayNode> CurrentParent = DisplayNode->GetParent();
-
-	while ( CurrentParent.IsValid() && CurrentParent->GetType() != ESequencerNode::Track )
+	// Find out which moviescene this is in
+	TViewModelPtr<FSequenceModel> OwnerModel = NodeToBeDeleted->FindAncestorOfType<FSequenceModel>();
+	if (!OwnerModel)
 	{
-		PathToTrack.Add( CurrentParent->GetNodeName() );
-		CurrentParent = CurrentParent->GetParent();
+		return false;
 	}
 
-	if ( CurrentParent.IsValid() )
-	{
-		OutParentTrack = StaticCastSharedPtr<FSequencerTrackNode>( CurrentParent );
-		for ( int32 i = PathToTrack.Num() - 1; i >= 0; i-- )
-		{
-			OutNamePath.Add( PathToTrack[i] );
-		}
-	}
-}
-
-
-bool FSequencer::OnRequestNodeDeleted( TSharedRef<const FSequencerDisplayNode> NodeToBeDeleted, const bool bKeepState )
-{
-	bool bAnythingRemoved = false;
-	
-	UMovieSceneSequence* Sequence = GetFocusedMovieSceneSequence();
-	UMovieScene* OwnerMovieScene = Sequence->GetMovieScene();
-
+	UMovieScene* OwnerMovieScene = OwnerModel->GetMovieScene();
 	if (OwnerMovieScene->IsReadOnly())
 	{
 		FSequencerUtilities::ShowReadOnlyError();
-		return bAnythingRemoved;
+		return false;
 	}
 
-	// Remove the selected object from our selection otherwise invisible objects are still selected and it causes confusion with
-	// things that are based on having a selection or not.
-	TSharedRef<FSequencerDisplayNode> SelectionNodeToRemove = ConstCastSharedRef<FSequencerDisplayNode>(NodeToBeDeleted);
-	Selection.RemoveFromSelection(SelectionNodeToRemove);
+	if (bKeepState)
+	{
+		FMovieSceneSequenceID SequenceID = OwnerModel->GetSequenceID();
 
-	SelectionNodeToRemove->DeleteNode();
-
-	if ( NodeToBeDeleted->GetType() == ESequencerNode::Folder )
-	{
-		// Delete Children
-		for ( const TSharedRef<FSequencerDisplayNode>& ChildNode : NodeToBeDeleted->GetChildNodes() )
+		for (const TViewModelPtr<IObjectBindingExtension>& ObjectBinding :
+			NodeToBeDeleted->GetDescendantsOfType<IObjectBindingExtension>(true, EViewModelListType::Outliner))
 		{
-			OnRequestNodeDeleted( ChildNode, bKeepState );
-		}
-
-		// Delete from parent, or root.
-		TSharedRef<const FSequencerFolderNode> FolderToBeDeleted = StaticCastSharedRef<const FSequencerFolderNode>(NodeToBeDeleted);
-		if ( NodeToBeDeleted->GetParent().IsValid() )
-		{
-			TSharedPtr<FSequencerFolderNode> ParentFolder = StaticCastSharedPtr<FSequencerFolderNode>( NodeToBeDeleted->GetParent() );
-			ParentFolder->GetFolder().RemoveChildFolder( &FolderToBeDeleted->GetFolder() );
-		}
-		else
-		{
-			UMovieScene* FocusedMovieScene = GetFocusedMovieSceneSequence()->GetMovieScene();
-			if (FocusedMovieScene)
+			for (TWeakObjectPtr<> WeakObject : FindBoundObjects(ObjectBinding->GetObjectGuid(), SequenceID))
 			{
-				FocusedMovieScene->Modify();
-				FocusedMovieScene->GetRootFolders().Remove(&FolderToBeDeleted->GetFolder());
-			}
-		}
-
-		bAnythingRemoved = true;
-	}
-	else if (NodeToBeDeleted->GetType() == ESequencerNode::Object)
-	{
-		// Delete any child object bindings
-		for (const TSharedRef<FSequencerDisplayNode>& ChildNode : NodeToBeDeleted->GetChildNodes())
-		{
-			if (ChildNode->GetType() == ESequencerNode::Object)
-			{
-				OnRequestNodeDeleted(ChildNode, bKeepState);
-			}
-		}
-
-		const FGuid& BindingToRemove = StaticCastSharedRef<const FSequencerObjectBindingNode>( NodeToBeDeleted )->GetObjectBinding();
-
-		// Remove from a parent folder if necessary.
-		if ( NodeToBeDeleted->GetParent().IsValid() && NodeToBeDeleted->GetParent()->GetType() == ESequencerNode::Folder )
-		{
-			TSharedPtr<FSequencerFolderNode> ParentFolder = StaticCastSharedPtr<FSequencerFolderNode>( NodeToBeDeleted->GetParent() );
-			ParentFolder->GetFolder().RemoveChildObjectBinding( BindingToRemove );
-		}
-		
-		if (bKeepState)
-		{
-			using namespace UE::MovieScene;
-
-			UMovieSceneEntitySystemLinker* EntitySystemLinker = RootTemplateInstance.GetEntitySystemLinker();
-			check(EntitySystemLinker);
-
-			for (TWeakObjectPtr<> WeakObject : FindBoundObjects(BindingToRemove, ActiveTemplateIDs.Top()))
-			{
 				TArray<UObject*> SubObjects;
 				GetObjectsWithOuter(WeakObject.Get(), SubObjects);
 
@@ -5576,99 +5695,15 @@
 				}
 			}
 		}
-				
-		// Try to remove as a spawnable first
-		if (OwnerMovieScene->RemoveSpawnable(BindingToRemove))
-		{
-			SpawnRegister->DestroySpawnedObject(BindingToRemove, ActiveTemplateIDs.Top(), *this);
-		}
-		// The guid should be associated with a possessable if it wasnt a spawnable
-		else if (OwnerMovieScene->RemovePossessable(BindingToRemove))
-		{
-			Sequence->Modify();
-			Sequence->UnbindPossessableObjects( BindingToRemove );
-		}
-
-		bAnythingRemoved = true;
-	}
-	else if( NodeToBeDeleted->GetType() == ESequencerNode::Track  )
-	{
-		TSharedRef<const FSequencerTrackNode> SectionAreaNode = StaticCastSharedRef<const FSequencerTrackNode>( NodeToBeDeleted );
-		UMovieSceneTrack* Track = SectionAreaNode->GetTrack();
-
-		// Remove from a parent folder if necessary.
-		if ( NodeToBeDeleted->GetParent().IsValid() && NodeToBeDeleted->GetParent()->GetType() == ESequencerNode::Folder )
-		{
-			TSharedPtr<FSequencerFolderNode> ParentFolder = StaticCastSharedPtr<FSequencerFolderNode>( NodeToBeDeleted->GetParent() );
-			ParentFolder->GetFolder().RemoveChildMasterTrack( Track );
-		}
-
-		if (Track != nullptr)
-		{
-			// Remove sub tracks belonging to this row only
-			if (SectionAreaNode->GetSubTrackMode() == FSequencerTrackNode::ESubTrackMode::SubTrack)
-			{
-				SectionAreaNode->GetTrack()->Modify();
-				TSet<TWeakObjectPtr<UMovieSceneSection> > SectionsToDelete;
-				for (TSharedRef<ISequencerSection> SectionToDelete : SectionAreaNode->GetSections())
-				{
-					UMovieSceneSection* Section = SectionToDelete->GetSectionObject();
-					if (Section)
-					{
-						SectionsToDelete.Add(Section);
-					}
-				}
-				DeleteSections(SectionsToDelete);
-				SectionAreaNode->GetTrack()->FixRowIndices();
-			}
-			else
-			{
-				OwnerMovieScene->Modify();
-				if (OwnerMovieScene->IsAMasterTrack(*Track))
-				{
-					OwnerMovieScene->RemoveMasterTrack(*Track);
-				}
-				else if (OwnerMovieScene->GetCameraCutTrack() == Track)
-				{
-					OwnerMovieScene->RemoveCameraCutTrack();
-				}
-				else
-				{
-					OwnerMovieScene->RemoveTrack(*Track);
-				}
-			}
-		
-			bAnythingRemoved = true;
-		}
-	}
-	else if ( NodeToBeDeleted->GetType() == ESequencerNode::Category )
-	{
-		TSharedPtr<FSequencerTrackNode> ParentTrackNode;
-		TArray<FName> PathFromTrack;
-		GetParentTrackNodeAndNamePath(NodeToBeDeleted, ParentTrackNode, PathFromTrack);
-		if ( ParentTrackNode.IsValid() )
-		{
-			for ( TSharedRef<ISequencerSection> Section : ParentTrackNode->GetSections() )
-			{
-				bAnythingRemoved |= Section->RequestDeleteCategory( PathFromTrack );
-			}
-		}
-	}
-	else if ( NodeToBeDeleted->GetType() == ESequencerNode::KeyArea )
-	{
-		TSharedPtr<FSequencerTrackNode> ParentTrackNode;
-		TArray<FName> PathFromTrack;
-		GetParentTrackNodeAndNamePath( NodeToBeDeleted, ParentTrackNode, PathFromTrack );
-		if ( ParentTrackNode.IsValid() )
-		{
-			for ( TSharedRef<ISequencerSection> Section : ParentTrackNode->GetSections() )
-			{
-				bAnythingRemoved |= Section->RequestDeleteKeyArea( PathFromTrack );
-			}
-		}
 	}
 
-	return bAnythingRemoved;
+	if (IDeletableExtension* Deletable = NodeToBeDeleted->CastThis<IDeletableExtension>())
+	{
+		Deletable->Delete();
+		return true;
+	}
+
+	return false;
 }
 
 bool FSequencer::MatchesContext(const FTransactionContext& InContext, const TArray<TPair<UObject*, FTransactionObjectEvent>>& TransactionObjects) const
@@ -6108,6 +6143,8 @@
 
 TArray<FGuid> FSequencer::AddActors(const TArray<TWeakObjectPtr<AActor> >& InActors, bool bSelectActors)
 {
+	using namespace UE::Sequencer;
+
 	TArray<FGuid> PossessableGuids;
 
 	if (GetFocusedMovieSceneSequence()->GetMovieScene()->IsReadOnly())
@@ -6146,28 +6183,29 @@
 		// Check if a folder is selected so we can add the actors to the selected folder.
 		TArray<UMovieSceneFolder*> SelectedParentFolders;
 		FString NewNodePath;
-		if (Selection.GetSelectedOutlinerNodes().Num() > 0)
+		if (Selection.GetSelectedOutlinerItems().Num() > 0)
 		{
-			for (TSharedRef<FSequencerDisplayNode> SelectedNode : Selection.GetSelectedOutlinerNodes())
+			for (TWeakPtr<FViewModel> WeakItem : Selection.GetSelectedOutlinerItems())
 			{
-				TSharedPtr<FSequencerDisplayNode> CurrentNode = SelectedNode;
-				while (CurrentNode.IsValid() && CurrentNode->GetType() != ESequencerNode::Folder)
+				TSharedPtr<FViewModel> CurrentItem = WeakItem.Pin();
+				if (!CurrentItem)
 				{
-					CurrentNode = CurrentNode->GetParent();
+					continue;
 				}
-				if (CurrentNode.IsValid())
+
+				if (TSharedPtr<FFolderModel> Folder = CurrentItem->FindAncestorOfType<FFolderModel>(true))
 				{
-					SelectedParentFolders.Add(&StaticCastSharedPtr<FSequencerFolderNode>(CurrentNode)->GetFolder());
+					SelectedParentFolders.Add(Folder->GetFolder());
 
 					// The first valid folder we find will be used to put the new actors into, so it's the node that we
 					// want to know the path from.
 					if (NewNodePath.Len() == 0)
 					{
 						// Add an extra delimiter (".") as we know that the new objects will be appended onto the end of this.
-						NewNodePath = FString::Printf(TEXT("%s."), *CurrentNode->GetPathName());
+						NewNodePath = FString::Printf(TEXT("%s."), *IOutlinerExtension::GetPathName(*Folder));
 
 						// Make sure the folder is expanded too so that adding objects to hidden folders become visible.
-						CurrentNode->SetExpansionState(true);
+						Folder->SetExpansion(true);
 					}
 				}
 			}
@@ -6226,6 +6264,11 @@
 
 void FSequencer::OnSelectedOutlinerNodesChanged()
 {
+	HandleSelectedOutlinerNodesChanged();
+}
+
+void FSequencer::HandleSelectedOutlinerNodesChanged()
+{
 	SynchronizeExternalSelectionWithSequencerSelection();
 
 	FSequencerEdMode* SequencerEdMode = (FSequencerEdMode*)(GLevelEditorModeTools().GetActiveMode(FSequencerEdMode::EM_SequencerMode));
@@ -6288,6 +6331,8 @@
 
 void FSequencer::AddSelectedNodesToNewNodeGroup()
 {
+	using namespace UE::Sequencer;
+
 	UMovieScene* MovieScene = GetFocusedMovieSceneSequence()->GetMovieScene();
 
 	if (MovieScene->IsReadOnly())
@@ -6295,22 +6340,21 @@
 		return;
 	}
 
-	const TSet<TSharedRef<FSequencerDisplayNode> >& SelectedNodes = GetSelection().GetSelectedOutlinerNodes();
-	if (SelectedNodes.Num() == 0)
+	const TSet<TWeakPtr<FViewModel>>& SelectedItems = GetSelection().GetSelectedOutlinerItems();
+	if (SelectedItems.Num() == 0)
 	{
 		return;
 	}
 
 	TSet<FString> NodesToAdd;
-	for (const TSharedRef<const FSequencerDisplayNode> Node : SelectedNodes)
+	for (const TWeakPtr<FViewModel>& WeakItem : SelectedItems)
 	{
-		const FSequencerDisplayNode* BaseNode = Node->GetBaseNode();
-		ESequencerNode::Type NodeType = BaseNode->GetType();
+		TSharedPtr<FViewModel>  Item      = WeakItem.Pin();
+		TSharedPtr<IGroupableExtension> Groupable = Item ? Item->FindAncestorOfType<IGroupableExtension>(true) : nullptr;
 
-		if (NodeType == ESequencerNode::Track || NodeType == ESequencerNode::Object || NodeType == ESequencerNode::Folder)
+		if (Groupable)
 		{
-			FString NodePath = BaseNode->GetPathName();
-			NodesToAdd.Add(NodePath);
+			NodesToAdd.Add(Groupable->GetIdentifierForGrouping());
 		}
 	}
 
@@ -6343,11 +6387,13 @@
 
 void FSequencer::AddSelectedNodesToExistingNodeGroup(UMovieSceneNodeGroup* NodeGroup)
 {
-	AddNodesToExistingNodeGroup(GetSelection().GetSelectedOutlinerNodes().Array(), NodeGroup);
+	AddNodesToExistingNodeGroup(GetSelection().GetSelectedOutlinerItems().Array(), NodeGroup);
 }
 
-void FSequencer::AddNodesToExistingNodeGroup(const TArray<TSharedRef<FSequencerDisplayNode>>& InNodes, UMovieSceneNodeGroup* InNodeGroup)
+void FSequencer::AddNodesToExistingNodeGroup(const TArray<TWeakPtr<UE::Sequencer::FViewModel>>& InItems, UMovieSceneNodeGroup* InNodeGroup)
 {
+	using namespace UE::Sequencer;
+
 	UMovieScene* MovieScene = GetFocusedMovieSceneSequence()->GetMovieScene();
 
 	if (MovieScene->IsReadOnly())
@@ -6361,15 +6407,14 @@
 	}
 
 	TSet<FString> NodesToAdd;
-	for (const TSharedRef<const FSequencerDisplayNode> Node : InNodes)
+	for (const TWeakPtr<FViewModel>& WeakItem : InItems)
 	{
-		const FSequencerDisplayNode* BaseNode = Node->GetBaseNode();
-		ESequencerNode::Type NodeType = BaseNode->GetType();
+		TSharedPtr<FViewModel>  Item      = WeakItem.Pin();
+		TSharedPtr<IGroupableExtension> Groupable = Item ? Item->FindAncestorOfType<IGroupableExtension>(true) : nullptr;
 
-		if (NodeType == ESequencerNode::Track || NodeType == ESequencerNode::Object || NodeType == ESequencerNode::Folder)
+		if (Groupable)
 		{
-			FString NodePath = BaseNode->GetPathName();
-			NodesToAdd.Add(NodePath);
+			NodesToAdd.Add(Groupable->GetIdentifierForGrouping());
 		}
 	}
 
@@ -6413,6 +6458,8 @@
 
 void FSequencer::SynchronizeExternalSelectionWithSequencerSelection()
 {
+	using namespace UE::Sequencer;
+
 	if ( bUpdatingSequencerSelection || !IsLevelEditorSequencer() )
 	{
 		return;
@@ -6423,38 +6470,30 @@
 	TSet<AActor*> SelectedSequencerActors;
 	TSet<UActorComponent*> SelectedSequencerComponents;
 
-	TSet<TSharedRef<FSequencerDisplayNode> > DisplayNodes = Selection.GetNodesWithSelectedKeysOrSections();
-	DisplayNodes.Append(Selection.GetSelectedOutlinerNodes());
+	TSet<TWeakPtr<FViewModel>> Items = Selection.GetNodesWithSelectedKeysOrSections();
+	Items.Append(Selection.GetSelectedOutlinerItems());
 
-	for ( TSharedRef<FSequencerDisplayNode> DisplayNode : DisplayNodes)
+	const FName ControlRigEditModeModeName("EditMode.ControlRig");
+	const bool bHACK_ControlRigEditMode = GLevelEditorModeTools().GetActiveMode(ControlRigEditModeModeName) != nullptr;
+
+	for ( TWeakPtr<FViewModel> WeakItem : Items)
 	{
-		// Get the closest object binding node.
-		TSharedPtr<FSequencerDisplayNode> CurrentNode = DisplayNode;
-		TSharedPtr<FSequencerObjectBindingNode> ObjectBindingNode;
-		while ( CurrentNode.IsValid() )
+		TSharedPtr<FViewModel> SelectedItem = WeakItem.Pin();
+		if (!SelectedItem)
 		{
-			if ( CurrentNode->GetType() == ESequencerNode::Object )
-			{
-				ObjectBindingNode = StaticCastSharedPtr<FSequencerObjectBindingNode>(CurrentNode);
-				break;
-			}
-			//HACK for DHI, if we have an active control rig then one is selected so don't find a parent actor or compomonent to select	
-			//but if we do select the actor/compoent directly we still select it.
-			const FName ControlRigEditModeModeName("EditMode.ControlRig");
-			if (GLevelEditorModeTools().GetActiveMode(ControlRigEditModeModeName) == nullptr)
-			{
-				CurrentNode = CurrentNode->GetParent();
-			}
-			else
-			{
-				break;
-			}
+			continue;
 		}
 
+		//HACK for DHI, if we have an active control rig then one is selected so don't find a parent actor or compomonent to select	
+		//but if we do select the actor/compoent directly we still select it.
+		TViewModelPtr<IObjectBindingExtension> ObjectBinding = bHACK_ControlRigEditMode
+			? FViewModelPtr(SelectedItem)
+			: SelectedItem->FindAncestorOfType(IObjectBindingExtension::ID, true);
+
 		// If the closest node is an object node, try to get the actor/component nodes from it.
-		if ( ObjectBindingNode.IsValid() )
+		if (ObjectBinding)
 		{
-			for (auto RuntimeObject : FindBoundObjects(ObjectBindingNode->GetObjectBinding(), ActiveTemplateIDs.Top()) )
+			for (auto RuntimeObject : FindBoundObjects(ObjectBinding->GetObjectGuid(), ActiveTemplateIDs.Top()) )
 			{
 				AActor* Actor = Cast<AActor>(RuntimeObject.Get());
 				if ( Actor != nullptr )
@@ -6596,25 +6635,10 @@
 }
 
 
-void GetRootObjectBindingNodes(const TArray<TSharedRef<FSequencerDisplayNode>>& DisplayNodes, TArray<TSharedRef<FSequencerObjectBindingNode>>& RootObjectBindings )
-{
-	for ( TSharedRef<FSequencerDisplayNode> DisplayNode : DisplayNodes )
-	{
-		switch ( DisplayNode->GetType() )
-		{
-		case ESequencerNode::Folder:
-			GetRootObjectBindingNodes( DisplayNode->GetChildNodes(), RootObjectBindings );
-			break;
-		case ESequencerNode::Object:
-			RootObjectBindings.Add( StaticCastSharedRef<FSequencerObjectBindingNode>( DisplayNode ) );
-			break;
-		}
-	}
-}
-
-
 void FSequencer::SynchronizeSequencerSelectionWithExternalSelection()
 {
+	using namespace UE::Sequencer;
+
 	if ( bUpdatingExternalSelection )
 	{
 		return;
@@ -6646,16 +6670,18 @@
 	TArray<ASequencerKeyActor*> SelectedSequencerKeyActors;
 	ActorSelection->GetSelectedObjects<ASequencerKeyActor>(SelectedSequencerKeyActors);
 
-	TSet<TSharedRef<FSequencerDisplayNode>> NodesToSelect;
+	FObjectBindingModelStorageExtension* ObjectModelStorage = ViewModel->GetRootModel()->CastDynamic<FObjectBindingModelStorageExtension>();
+	check(ObjectModelStorage);
+
+	TSet<TSharedPtr<FObjectBindingModel>> NodesToSelect;
 	for (const FMovieSceneBinding& Binding : Sequence->GetMovieScene()->GetBindings())
 	{
-		TSharedPtr<FSequencerObjectBindingNode> NodePtr = NodeTree->FindObjectBindingNode(Binding.GetObjectGuid());
-		if (!NodePtr)
+		TSharedPtr<FObjectBindingModel> ObjectModel = ObjectModelStorage->FindModelForObjectBinding(Binding.GetObjectGuid());
+		if (!ObjectModel)
 		{
 			continue;
 		}
 
-		TSharedRef<FSequencerObjectBindingNode> ObjectBindingNode = NodePtr.ToSharedRef();
 		for ( TWeakObjectPtr<> WeakObject : FindBoundObjects(Binding.GetObjectGuid(), ActiveTemplateIDs.Top()) )
 		{
 			UObject* RuntimeObject = WeakObject.Get();
@@ -6671,7 +6697,7 @@
 					AActor* TrailActor = KeyActor->GetAssociatedActor();
 					if (TrailActor != nullptr && RuntimeObject == TrailActor)
 					{
-						NodesToSelect.Add(ObjectBindingNode);
+						NodesToSelect.Add(ObjectModel);
 						bAllAlreadySelected = false;
 						break;
 					}
@@ -6686,25 +6712,29 @@
 
 			if (bActorSelected || bComponentSelected)
 			{
-				NodesToSelect.Add( ObjectBindingNode );
+				NodesToSelect.Add( ObjectModel );
 
-				if (bAllAlreadySelected && !Selection.IsSelected(ObjectBindingNode))
+				if (bAllAlreadySelected && !Selection.IsSelected(ObjectModel))
 				{
-					// Traversal callback will exit prematurely if there are any selected children
-					auto Traverse_IsSelected = [this](FSequencerDisplayNode& InNode)
+					bool bAnyChildrenSelected = false;
+					TArray<TSharedPtr<FViewModel>> OutlinerChildren;
+					ObjectModel->GetDescendantsOfType(IOutlinerExtension::ID, OutlinerChildren);
+					for (TSharedPtr<FViewModel> Child : OutlinerChildren)
 					{
-						TSharedRef<FSequencerDisplayNode> SharedNode = InNode.AsShared();
-						return !this->Selection.IsSelected(SharedNode) && !this->Selection.NodeHasSelectedKeysOrSections(SharedNode);
-					};
+						if (Selection.IsSelected(Child) || Selection.NodeHasSelectedKeysOrSections(Child))
+						{
+							bAnyChildrenSelected = true;
+							break;
+						}
+					}
 
-					const bool bNoChildrenSelected = ObjectBindingNode->Traverse_ParentFirst(Traverse_IsSelected, false);
-					if (bNoChildrenSelected)
+					if (!bAnyChildrenSelected)
 					{
 						bAllAlreadySelected = false;
 					}
 				}
 			}
-			else if (Selection.IsSelected(ObjectBindingNode))
+			else if (Selection.IsSelected(ObjectModel))
 			{
 				bAllAlreadySelected = false;
 			}
@@ -6712,30 +6742,33 @@
 	}
 	//Only test if none are selected if we are not transacting, otherwise it will clear out control rig's incorrectly.
 
-	if (!bAllAlreadySelected || (!GIsTransacting && (NodesToSelect.Num() == 0 && Selection.GetSelectedOutlinerNodes().Num())))
+	if (!bAllAlreadySelected || (!GIsTransacting && (NodesToSelect.Num() == 0 && Selection.GetSelectedOutlinerItems().Num())))
 	{
 		Selection.SuspendBroadcast();
 		Selection.EmptySelectedOutlinerNodes();
-		for ( TSharedRef<FSequencerDisplayNode> NodeToSelect : NodesToSelect)
+		for (TSharedPtr<FViewModel> NodeToSelect : NodesToSelect)
 		{
 			Selection.AddToSelection( NodeToSelect );
 		}
 		
-		TSharedPtr<SSequencerTreeView> TreeView = SequencerWidget->GetTreeView();
-		const TSet<TSharedRef<FSequencerDisplayNode>>& OutlinerSelection = GetSelection().GetSelectedOutlinerNodes();
+		TSharedPtr<SOutlinerView> TreeView = SequencerWidget->GetTreeView();
+		const TSet<TWeakPtr<FViewModel>>& OutlinerSelection = GetSelection().GetSelectedOutlinerItems();
 		if (OutlinerSelection.Num() == 1)
 		{
-			for (auto& Node : OutlinerSelection)
+			for (TWeakPtr<FViewModel> WeakNode : OutlinerSelection)
 			{
-				auto Parent = Node->GetParent();
-				while (Parent.IsValid())
+				if (TSharedPtr<FViewModel> Node = WeakNode.Pin())
 				{
-					TreeView->SetItemExpansion(Parent->AsShared(), true);
-					Parent = Parent->GetParent();
+					TSharedPtr<FViewModel> Parent = Node->GetParent();
+					while (Parent.IsValid())
+					{
+						TreeView->SetItemExpansion(Parent->AsShared(), true);
+						Parent = Parent->GetParent();
+					}
+
+					TreeView->RequestScrollIntoView(Node);
+					break;
 				}
-
-				TreeView->RequestScrollIntoView(Node);
-				break;
 			}
 		}
 
@@ -6746,6 +6779,8 @@
 
 void FSequencer::SelectNodesByPath(const TSet<FString>& NodePaths)
 {
+	using namespace UE::Sequencer;
+
 	if (bUpdatingExternalSelection)
 	{
 		return;
@@ -6761,12 +6796,14 @@
 	// nodes are not cleared and reselected, which can cause issues with the curve editor auto-fitting 
 	// based on selection.
 	bool bAllAlreadySelected = true;
-	const TSet<TSharedRef<FSequencerDisplayNode>>& CurrentSelection = GetSelection().GetSelectedOutlinerNodes();
+	const TSet<TWeakPtr<FViewModel>>& CurrentSelection = GetSelection().GetSelectedOutlinerItems();
 
-	TSet<TSharedRef<FSequencerDisplayNode>> NodesToSelect;
-	for (TSharedRef<FSequencerDisplayNode> DisplayNode : NodeTree->GetAllNodes())
+	TArray<TSharedPtr<FViewModel>> AllNodes;
+	NodeTree->GetAllNodes(AllNodes);
+	TSet<TSharedPtr<FViewModel>> NodesToSelect;
+	for (TSharedPtr<FViewModel> DisplayNode : AllNodes)
 	{
-		if (NodePaths.Contains(DisplayNode->GetPathName()))
+		if (NodePaths.Contains(IOutlinerExtension::GetPathName(DisplayNode)))
 		{
 			NodesToSelect.Add(DisplayNode);
 			if (bAllAlreadySelected && !CurrentSelection.Contains(DisplayNode))
@@ -6780,24 +6817,27 @@
 	{
 		Selection.SuspendBroadcast();
 		Selection.EmptySelectedOutlinerNodes();
-		for (TSharedRef<FSequencerDisplayNode> NodeToSelect : NodesToSelect)
+		for (TSharedPtr<FViewModel> NodeToSelect : NodesToSelect)
 		{
-			Selection.AddToSelection(NodeToSelect);
+			Selection.AddToSelection( NodeToSelect );
 		}
 
-		TSharedPtr<SSequencerTreeView> TreeView = SequencerWidget->GetTreeView();
-		const TSet<TSharedRef<FSequencerDisplayNode>>& OutlinerSelection = GetSelection().GetSelectedOutlinerNodes();
-		for (auto& Node : OutlinerSelection)
+		TSharedPtr<SOutlinerView> TreeView = SequencerWidget->GetTreeView();
+		const TSet<TWeakPtr<FViewModel>>& OutlinerSelection = GetSelection().GetSelectedOutlinerItems();
+		for (const TWeakPtr<FViewModel>& WeakNode : OutlinerSelection)
 		{
-			auto Parent = Node->GetParent();
-			while (Parent.IsValid())
+			if (TSharedPtr<FViewModel> Node = WeakNode.Pin())
 			{
-				TreeView->SetItemExpansion(Parent->AsShared(), true);
-				Parent = Parent->GetParent();
+				TSharedPtr<FViewModel> Parent = Node->GetParent();
+				while (Parent.IsValid())
+				{
+					TreeView->SetItemExpansion(Parent->AsShared(), true);
+					Parent = Parent->GetParent();
+				}
+
+				TreeView->RequestScrollIntoView(Node);
+				break;
 			}
-
-			TreeView->RequestScrollIntoView(Node);
-			break;
 		}
 
 		Selection.ResumeBroadcast();
@@ -6913,7 +6953,7 @@
 	{
 		if (Key.IsValid())
 		{
-			FFrameNumber KeyTime = Key.KeyArea->GetKeyTime(Key.KeyHandle.GetValue());
+			FFrameNumber KeyTime = Key.WeakChannel.Pin()->GetKeyArea()->GetKeyTime(Key.KeyHandle);
 			if (!BoundsHull.HasLowerBound() || BoundsHull.GetLowerBoundValue() > KeyTime)
 			{
 				BoundsHull.SetLowerBound(TRange<FFrameNumber>::BoundsType::Inclusive(KeyTime));
@@ -7014,11 +7054,11 @@
 
 void FSequencer::GetSelectedSections(TArray<UMovieSceneSection*>& OutSelectedSections)
 {
-	for (TWeakObjectPtr<UMovieSceneSection> SelectedSection : Selection.GetSelectedSections())
+	for (TWeakObjectPtr<UMovieSceneSection> WeakSection : Selection.GetSelectedSections())
 	{
-		if (SelectedSection.IsValid())
+		if (UMovieSceneSection* Section = WeakSection.Get())
 		{
-			OutSelectedSections.Add(SelectedSection.Get());
+			OutSelectedSections.Add(Section);
 		}
 	}
 }
@@ -7036,26 +7076,33 @@
 
 void FSequencer::GetSelectedKeyAreas(TArray<const IKeyArea*>& OutSelectedKeyAreas, bool bIncludeSelectedKeys)
 {
-	TSet<TSharedRef<FSequencerDisplayNode>> NodesToKey = Selection.GetSelectedOutlinerNodes();
+	using namespace UE::Sequencer;
+
+	TArray<TSharedRef<FViewModel>> NodesToKey;
 	{
+		TSet<TSharedRef<FViewModel>> ChildNodes;
+		for (const TWeakPtr<FViewModel>& WeakNode : Selection.GetSelectedOutlinerItems())
+		{
+			TSharedPtr<FViewModel> Node = WeakNode.Pin();
+			if (!Node)
+			{
+				continue;
+			}
 
-		TSet<TSharedRef<FSequencerDisplayNode>> ChildNodes;
-		for (TSharedRef<FSequencerDisplayNode> Node : NodesToKey.Array())
-		{
+			NodesToKey.Add(Node.ToSharedRef());
+
 			// No need to gather key areas from binding/tracks because they have no key areas
-			if (Node->GetType()== ESequencerNode::Object || Node->GetType() == ESequencerNode::Track)
+			if (Node->IsA<IObjectBindingExtension>() || Node->IsA<ITrackExtension>())
 			{
 				continue;
 			}
-			else
+
+			ChildNodes.Reset();
+			SequencerHelpers::GetDescendantNodes(Node.ToSharedRef(), ChildNodes);
+
+			for (TSharedRef<FViewModel> ChildNode : ChildNodes)
 			{
-				ChildNodes.Reset();
-				SequencerHelpers::GetDescendantNodes(Node, ChildNodes);
-
-				for (TSharedRef<FSequencerDisplayNode> ChildNode : ChildNodes)
-				{
-					NodesToKey.Remove(ChildNode);
-				}
+				NodesToKey.Remove(ChildNode);
 			}
 		}
 	}
@@ -7063,10 +7110,10 @@
 	TSet<TSharedPtr<IKeyArea>> KeyAreas;
 	TSet<UMovieSceneSection*>  ModifiedSections;
 
-	for (TSharedRef<FSequencerDisplayNode> Node : NodesToKey)
+	for (TSharedRef<FViewModel> Node : NodesToKey)
 	{
 		//if object or track selected we don't want all of the children only if spefically selected.
-		if ((Node->GetType() != ESequencerNode::Track) && (Node->GetType() != ESequencerNode::Object))
+		if (!Node->IsA<ITrackExtension>() && !Node->IsA<IObjectBindingExtension>())
 		{
 			SequencerHelpers::GetAllKeyAreas(Node, KeyAreas);
 		}
@@ -7076,7 +7123,7 @@
 	{
 		for (FSequencerSelectedKey Key : Selection.GetSelectedKeys())
 		{
-			KeyAreas.Add(Key.KeyArea);
+			KeyAreas.Add(Key.WeakChannel.Pin()->GetKeyArea()); 
 		}
 	}
 	for (TSharedPtr<IKeyArea> KeyArea : KeyAreas)
@@ -7088,31 +7135,33 @@
 
 void FSequencer::SelectByNthCategoryNode(UMovieSceneSection* Section, int Index, bool bSelect)
 {
-	TSet<TSharedRef<FSequencerDisplayNode>> Nodes;
-	TArray<TSharedRef<FSequencerDisplayNode>> NodesToSelect;
+	using namespace UE::Sequencer;
 
-	TOptional<FSectionHandle> SectionHandle = NodeTree->GetSectionHandle(Section);
+	TSet<TSharedRef<FViewModel>> Nodes;
+	TArray<TSharedRef<FViewModel>> NodesToSelect;
+
+	TSharedPtr<FSectionModel> SectionHandle = NodeTree->GetSectionModel(Section);
 	int32 Count = 0;
-	if (SectionHandle.IsSet())
+	if (SectionHandle)
 	{
-		TSharedRef<FSequencerTrackNode> TrackNode = SectionHandle->GetTrackNode();
-		for (const TSharedRef<FSequencerDisplayNode>& Node : TrackNode->GetChildNodes())
+		TSharedPtr<FViewModel> TrackNode = SectionHandle->GetParentTrackModel();
+		for (TSharedPtr<FViewModel> Node : TrackNode->GetChildren())
 		{
-			if (Node->GetType() == ESequencerNode::Category && Count++ == Index)
+			if (Node->IsA<FCategoryGroupModel>() && Count++ == Index)
 			{
 				bool bAlreadySelected = false;
 				if (bSelect == true)
 				{
-					bAlreadySelected = Selection.GetSelectedOutlinerNodes().Contains(Node);					
+					bAlreadySelected = Selection.GetSelectedOutlinerItems().Contains(Node);					
 				}
 				if (bAlreadySelected == false)
 				{
-					NodesToSelect.Add(Node);
+					NodesToSelect.Add(Node.ToSharedRef());
 					if (bSelect == false) //make sure all children not selected
 					{
-						for (const TSharedRef<FSequencerDisplayNode>& ChildNode : Node->GetChildNodes())
+						for (TSharedPtr<FViewModel> ChildNode : Node->GetChildren())
 						{
-							NodesToSelect.Add(ChildNode);
+							NodesToSelect.Add(ChildNode.ToSharedRef());
 						}
 					}
 				}
@@ -7123,11 +7172,17 @@
 	{
 		if (Settings->GetAutoExpandTreeView())
 		{
-			for (const TSharedRef<FSequencerDisplayNode>& DisplayNode : NodesToSelect)
+			for (const TSharedRef<FViewModel>& DisplayNode : NodesToSelect)
 			{
-				if (DisplayNode->GetParent().IsValid() && DisplayNode->GetParent()->GetType() == ESequencerNode::Track && !DisplayNode->GetParent()->IsExpanded())
+				if (DisplayNode->GetParent().IsValid() && DisplayNode->GetParent()->IsA<ITrackExtension>())
 				{
-					DisplayNode->GetParent()->SetExpansionState(true);
+					if (IOutlinerExtension* ParentOutlinerItem = DisplayNode->GetParent()->CastThis<IOutlinerExtension>())
+					{
+						if (!ParentOutlinerItem->IsExpanded())
+						{
+							ParentOutlinerItem->SetExpansion(true);
+						}
+					}
 					break;
 				}
 			}
@@ -7143,10 +7198,9 @@
 	}
 	else if (NodesToSelect.Num() > 0)
 	{
-		for (const TSharedRef<FSequencerDisplayNode>& DisplayNode : NodesToSelect)
+		for (const TSharedRef<FViewModel>& DisplayNode : NodesToSelect)
 		{
 			Selection.RemoveFromSelection(DisplayNode);
-			Selection.RemoveFromNodesWithSelectedKeysOrSections(DisplayNode);
 		}
 		Selection.RequestOutlinerNodeSelectionChangedBroadcast();
 	}
@@ -7154,16 +7208,16 @@
 
 void FSequencer::SelectByChannels(UMovieSceneSection* Section, TArrayView<const FMovieSceneChannelHandle> InChannels, bool bSelectParentInstead, bool bSelect)
 {
-	TSet<TSharedRef<FSequencerDisplayNode>> Nodes;
-	TArray<TSharedRef<FSequencerDisplayNode>> NodesToSelect;
+	using namespace UE::Sequencer;
 
-	TOptional<FSectionHandle> SectionHandle = NodeTree->GetSectionHandle(Section);
-	if (SectionHandle.IsSet())
+	TSet<TSharedRef<FViewModel>> Nodes;
+	TArray<TSharedRef<FViewModel>> NodesToSelect;
+
+	TSharedPtr<FSectionModel> SectionHandle = NodeTree->GetSectionModel(Section);
+	if (SectionHandle)
 	{
-		TSharedRef<FSequencerTrackNode> TrackNode = SectionHandle->GetTrackNode();
-		TArray<TSharedRef<FSequencerSectionKeyAreaNode>> KeyAreaNodes;
-		TrackNode->GetChildKeyAreaNodesRecursively(KeyAreaNodes);
-		for (TSharedRef<FSequencerSectionKeyAreaNode> KeyAreaNode : KeyAreaNodes)
+		TSharedPtr<FViewModel> TrackNode = SectionHandle->GetParentTrackModel();
+		for (TSharedPtr<FChannelGroupModel> KeyAreaNode : TrackNode->GetDescendantsOfType<FChannelGroupModel>())
 		{
 			for (TSharedPtr<IKeyArea> KeyArea : KeyAreaNode->GetAllKeyAreas())
 			{
@@ -7172,24 +7226,24 @@
 				{
 					if (bSelectParentInstead || bSelect == false)
 					{
-						if(bSelect)
+						if (bSelect)
 						{
-							NodesToSelect.Add(KeyAreaNode->GetParent()->AsShared());
+							NodesToSelect.Add(KeyAreaNode->GetParent().ToSharedRef());
 						}
 						else
 						{
-							Nodes.Add(KeyAreaNode->GetParent()->AsShared());
+							Nodes.Add(KeyAreaNode->GetParent().ToSharedRef());
 						}
 					}
 					if (!bSelectParentInstead || bSelect == false)
 					{
-						if(bSelect)
+						if (bSelect)
 						{
-							NodesToSelect.Add(KeyAreaNode);
+							NodesToSelect.Add(KeyAreaNode.ToSharedRef());
 						}
 						else
 						{
-							Nodes.Add(KeyAreaNode);
+							Nodes.Add(KeyAreaNode.ToSharedRef());
 						}
 					}
 				}
@@ -7198,7 +7252,7 @@
 	}
 	
 	if (bSelect)
-	{	
+	{
 		if (NodesToSelect.Num() > 0)
 		{
 			//todo hide behind preference 
@@ -7210,10 +7264,9 @@
 	}
 	else if (Nodes.Num() > 0)
 	{
-		for (const TSharedRef<FSequencerDisplayNode>& DisplayNode : Nodes)
+		for (const TSharedRef<FViewModel>& DisplayNode : Nodes)
 		{
 			Selection.RemoveFromSelection(DisplayNode);
-			Selection.RemoveFromNodesWithSelectedKeysOrSections(DisplayNode);
 		}
 		Selection.RequestOutlinerNodeSelectionChangedBroadcast();
 	}
@@ -7221,20 +7274,20 @@
 
 void FSequencer::SelectByChannels(UMovieSceneSection* Section, const TArray<FName>& InChannelNames, bool bSelectParentInstead, bool bSelect)
 {
-	TSet<TSharedRef<FSequencerDisplayNode>> Nodes;
-	TArray<TSharedRef<FSequencerDisplayNode>> NodesToSelect;
+	using namespace UE::Sequencer;
 
-	TOptional<FSectionHandle> SectionHandle = NodeTree->GetSectionHandle(Section);
-	if (SectionHandle.IsSet())
+	TSet<TSharedRef<FViewModel>> Nodes;
+	TArray<TSharedRef<FViewModel>> NodesToSelect;
+
+	TSharedPtr<FSectionModel> SectionHandle = NodeTree->GetSectionModel(Section);
+	if (SectionHandle)
 	{
-		TSharedRef<FSequencerTrackNode> TrackNode = SectionHandle->GetTrackNode();
-		TArray<TSharedRef<FSequencerSectionKeyAreaNode>> KeyAreaNodes;
-		TrackNode->GetChildKeyAreaNodesRecursively(KeyAreaNodes);
-		for (TSharedRef<FSequencerSectionKeyAreaNode> KeyAreaNode : KeyAreaNodes)
+		TSharedPtr<FViewModel> TrackNode = SectionHandle->GetParentTrackModel();
+		for (TSharedPtr<FChannelGroupModel> KeyAreaNode : TrackNode->GetDescendantsOfType<FChannelGroupModel>())
 		{
-			if (KeyAreaNode->GetParent().IsValid() && InChannelNames.Contains(*KeyAreaNode->GetParent()->AsShared()->GetDisplayName().ToString()))
+			if (KeyAreaNode->GetParent().IsValid() && InChannelNames.Contains(*KeyAreaNode->GetParent()->CastThis<IOutlinerExtension>()->GetLabel().ToString()))
 			{
-				Nodes.Add(KeyAreaNode->GetParent()->AsShared());
+				Nodes.Add(KeyAreaNode->GetParent().ToSharedRef());
 			}
 
 			for (TSharedPtr<IKeyArea> KeyArea : KeyAreaNode->GetAllKeyAreas())
@@ -7247,11 +7300,11 @@
 				{
 					if (bSelectParentInstead || bSelect == false)
 					{
-						Nodes.Add(KeyAreaNode->GetParent()->AsShared());
+						Nodes.Add(KeyAreaNode->GetParent().ToSharedRef());
 					}
 					if (!bSelectParentInstead || bSelect == false)
 					{
-						Nodes.Add(KeyAreaNode);
+						Nodes.Add(KeyAreaNode.ToSharedRef());
 					}
 				}
 			}
@@ -7260,13 +7313,13 @@
 	
 	if (bSelect)
 	{
-		for (const TSharedRef<FSequencerDisplayNode>& DisplayNode : Nodes)
+		for (const TSharedRef<FViewModel>& DisplayNode : Nodes)
 		{
 			if (Settings->GetAutoExpandTreeView())
 			{
-				if (DisplayNode->GetParent().IsValid() && DisplayNode->GetParent()->GetType() == ESequencerNode::Track && !DisplayNode->GetParent()->IsExpanded())
+				if (DisplayNode->GetParent().IsValid() && DisplayNode->GetParent()->IsA<ITrackExtension>() && !DisplayNode->GetParent()->CastThisChecked<IOutlinerExtension>()->IsExpanded())
 				{
-					DisplayNode->GetParent()->SetExpansionState(true);
+					DisplayNode->GetParent()->CastThisChecked<IOutlinerExtension>()->SetExpansion(true);
 				}
 			}
 			NodesToSelect.Add(DisplayNode);
@@ -7282,10 +7335,9 @@
 	}
 	else if (Nodes.Num() > 0)
 	{
-		for (const TSharedRef<FSequencerDisplayNode>& DisplayNode : Nodes)
+		for (const TSharedRef<FViewModel>& DisplayNode : Nodes)
 		{
 			Selection.RemoveFromSelection(DisplayNode);
-			Selection.RemoveFromNodesWithSelectedKeysOrSections(DisplayNode);
 		}
 		Selection.RequestOutlinerNodeSelectionChangedBroadcast();
 	}
@@ -7293,24 +7345,31 @@
 
 void FSequencer::SelectObject(FGuid ObjectBinding)
 {
-	TSharedPtr<FSequencerObjectBindingNode> Node = NodeTree->FindObjectBindingNode(ObjectBinding);
-	if (Node.IsValid())
+	using namespace UE::Sequencer;
+
+	FObjectBindingModelStorageExtension* ObjectStorage = ViewModel->GetRootModel()->CastDynamic<FObjectBindingModelStorageExtension>();
+	check(ObjectStorage);
+
+	if (TSharedPtr<FObjectBindingModel> Model = ObjectStorage->FindModelForObjectBinding(ObjectBinding))
 	{
-		GetSelection().AddToSelection(Node.ToSharedRef());
+		Selection.AddToOutlinerSelection(Model);
 	}
 }
 
 void FSequencer::SelectTrack(UMovieSceneTrack* Track)
 {
-	for (TSharedRef<FSequencerDisplayNode> Node : NodeTree->GetAllNodes())
+	using namespace UE::Sequencer;
+	
+	TArray<TSharedPtr<FViewModel>> AllNodes;
+	NodeTree->GetAllNodes(AllNodes);
+	for (TSharedPtr<FViewModel> Node : AllNodes)
 	{
-		if (Node->GetType() == ESequencerNode::Track)
+		if (ITrackExtension* TrackNode = Node->CastThis<ITrackExtension>())
 		{
-			TSharedRef<FSequencerTrackNode> TrackNode = StaticCastSharedRef<FSequencerTrackNode>(Node);
 			UMovieSceneTrack* TrackForNode = TrackNode->GetTrack();
 			if (TrackForNode == Track)
 			{
-				Selection.AddToSelection(Node);
+				Selection.AddToOutlinerSelection(Node);
 				break;
 			}
 		}
@@ -7319,12 +7378,21 @@
 
 void FSequencer::SelectSection(UMovieSceneSection* Section)
 {
-	Selection.AddToSelection(Section);
+	using namespace UE::Sequencer;
+
+	FSectionModelStorageExtension* SectionModelStorage = ViewModel->GetRootModel()->CastDynamic<FSectionModelStorageExtension>();
+	check(SectionModelStorage);
+
+	TSharedPtr<UE::Sequencer::FSectionModel> SectionModel = SectionModelStorage->FindModelForSection(Section);
+	if (SectionModel)
+	{
+		Selection.AddToSelection(SectionModel);
+	}
 }
 
-void FSequencer::SelectKey(UMovieSceneSection* InSection, TSharedPtr<IKeyArea> KeyArea, FKeyHandle KeyHandle, bool bToggle)
+void FSequencer::SelectKey(UMovieSceneSection* InSection, TSharedPtr<UE::Sequencer::FChannelModel> InChannel, FKeyHandle KeyHandle, bool bToggle)
 {
-	FSequencerSelectedKey SelectedKey(*InSection, KeyArea, KeyHandle);
+	FSequencerSelectedKey SelectedKey(*InSection, InChannel, KeyHandle);
 
 	if (bToggle && Selection.IsSelected(SelectedKey))
 	{
@@ -7338,12 +7406,17 @@
 
 void FSequencer::SelectByPropertyPaths(const TArray<FString>& InPropertyPaths)
 {
-	TArray<TSharedRef<FSequencerDisplayNode>> NodesToSelect;
-	for (const TSharedRef<FSequencerDisplayNode>& Node : NodeTree->GetAllNodes())
+	using namespace UE::Sequencer;
+
+	TArray<TSharedRef<FViewModel>> AllNodes;
+	NodeTree->GetAllNodes(AllNodes);
+
+	TArray<TSharedRef<FViewModel>> NodesToSelect;
+	for (const TSharedRef<FViewModel>& Node : AllNodes)
 	{
-		if (Node->GetType() == ESequencerNode::Track)
+		if (ITrackExtension* TrackNode = Node->CastThis<ITrackExtension>())
 		{
-			if (UMovieScenePropertyTrack* PropertyTrack = Cast<UMovieScenePropertyTrack>(StaticCastSharedRef<FSequencerTrackNode>(Node)->GetTrack()))
+			if (UMovieScenePropertyTrack* PropertyTrack = Cast<UMovieScenePropertyTrack>(TrackNode->GetTrack()))
 			{
 				FString Path = PropertyTrack->GetPropertyPath().ToString();
 				for (const FString& PropertyPath : InPropertyPaths)
@@ -7371,12 +7444,15 @@
 
 void FSequencer::SelectFolder(UMovieSceneFolder* Folder)
 {
-	for (TSharedRef<FSequencerDisplayNode> Node : NodeTree->GetAllNodes())
+	using namespace UE::Sequencer;
+
+	TArray<TSharedRef<FViewModel>> AllNodes;
+	NodeTree->GetAllNodes(AllNodes);
+	for (TSharedRef<FViewModel> Node : AllNodes)
 	{
-		if (Node->GetType() == ESequencerNode::Folder)
+		if (FFolderModel* FolderNode = Node->CastThis<FFolderModel>())
 		{
-			TSharedRef<FSequencerFolderNode> FolderNode = StaticCastSharedRef<FSequencerFolderNode>(Node);
-			UMovieSceneFolder* FolderForNode = &FolderNode->GetFolder();
+			UMovieSceneFolder* FolderForNode = FolderNode->GetFolder();
 			if (FolderForNode == Folder)
 			{
 				Selection.AddToSelection(Node);
@@ -7394,11 +7470,15 @@
 
 void FSequencer::ThrobKeySelection()
 {
+	using namespace UE::Sequencer;
+
 	SSequencerSection::ThrobKeySelection();
 }
 
 void FSequencer::ThrobSectionSelection()
 {
+	using namespace UE::Sequencer;
+
 	// Scrub to the beginning of newly created sections if they're out of view
 	TOptional<FFrameNumber> ScrubFrame;
 	for (TWeakObjectPtr<UMovieSceneSection> SelectedSectionPtr : Selection.GetSelectedSections())
@@ -7452,7 +7532,7 @@
 	
 		DeleteSections(Selection.GetSelectedSections());
 	}
-	else if (Selection.GetSelectedOutlinerNodes().Num())
+	else if (Selection.GetSelectedOutlinerItems().Num())
 	{
 		DeleteSelectedNodes(false);
 	}
@@ -7646,7 +7726,7 @@
 	NotifyMovieSceneDataChanged(EMovieSceneDataChangeType::MovieSceneStructureItemsChanged);
 }
 
-void FSequencer::DeleteNode(TSharedRef<FSequencerDisplayNode> NodeToBeDeleted, const bool bKeepState)
+void FSequencer::DeleteNode(TSharedRef<FViewModel> NodeToBeDeleted, const bool bKeepState)
 {
 	// If this node is selected, delete all selected nodes
 	if (GetSelection().IsSelected(NodeToBeDeleted))
@@ -7667,8 +7747,11 @@
 
 void FSequencer::DeleteSelectedNodes(const bool bKeepState)
 {
-	TSet< TSharedRef<FSequencerDisplayNode> > SelectedNodesCopy = GetSelection().GetSelectedOutlinerNodes();
+	using namespace UE::Sequencer;
 
+	TArray<TSharedPtr<FViewModel>> SelectedNodesCopy;
+	GetSelection().GetSelectedOutlinerItems(SelectedNodesCopy);
+
 	if (SelectedNodesCopy.Num() == 0)
 	{
 		return;
@@ -7676,16 +7759,20 @@
 
 	const FScopedTransaction Transaction( NSLOCTEXT("Sequencer", "UndoDeletingObject", "Delete Node") );
 
+	Selection.EmptySelectedOutlinerNodes();
+
 	bool bAnythingDeleted = false;
 
-	for( const TSharedRef<FSequencerDisplayNode>& SelectedNode : SelectedNodesCopy )
+	for (TSharedPtr<FViewModel>& SelectedNode : SelectedNodesCopy)
 	{
-		if( !SelectedNode->IsHidden() )
+		if (!SelectedNode->CastThisChecked<IOutlinerExtension>()->IsFilteredOut())
 		{
 			// Delete everything in the entire node
-			TSharedRef<const FSequencerDisplayNode> NodeToBeDeleted = StaticCastSharedRef<const FSequencerDisplayNode>(SelectedNode);
-			bAnythingDeleted |= OnRequestNodeDeleted( NodeToBeDeleted, bKeepState );
+			bAnythingDeleted |= OnRequestNodeDeleted( SelectedNode.ToSharedRef(), bKeepState );
 		}
+
+		// Null out the node, which should have the effect of completely destroying it
+		SelectedNode = nullptr;
 	}
 
 	if ( bAnythingDeleted )
@@ -7694,9 +7781,11 @@
 	}
 }
 
-void FSequencer::MoveNodeToFolder(TSharedRef<FSequencerDisplayNode> NodeToMove, UMovieSceneFolder* DestinationFolder)
+void FSequencer::MoveNodeToFolder(TSharedRef<FViewModel> NodeToMove, UMovieSceneFolder* DestinationFolder)
 {
-	TSharedPtr<FSequencerDisplayNode> ParentNode = NodeToMove->GetParent();
+	using namespace UE::Sequencer;
+
+	TSharedPtr<FViewModel> ParentNode = NodeToMove->GetParent();
 	
 	if (DestinationFolder == nullptr)
 	{
@@ -7705,75 +7794,68 @@
 
 	DestinationFolder->Modify();
 
-	switch (NodeToMove->GetType())
+	if (FFolderModel* FolderNode = NodeToMove->CastThis<FFolderModel>())
 	{
-		case ESequencerNode::Folder:
+		if (ParentNode.IsValid() && ParentNode->IsA<FFolderModel>())
 		{
-			TSharedRef<FSequencerFolderNode> FolderNode = StaticCastSharedRef<FSequencerFolderNode>(NodeToMove);
-			if (ParentNode.IsValid())
+			FFolderModel* NodeParentFolder = ParentNode->CastThisChecked<FFolderModel>();
+			NodeParentFolder->GetFolder()->Modify();
+			NodeParentFolder->GetFolder()->RemoveChildFolder(FolderNode->GetFolder());
+		}
+		else
+		{
+			UMovieScene* FocusedMovieScene = GetFocusedMovieSceneSequence()->GetMovieScene();
+			if (FocusedMovieScene)
 			{
-				checkf(ParentNode->GetType() == ESequencerNode::Folder, TEXT("Can not remove from unsupported parent node."));
-				TSharedPtr<FSequencerFolderNode> NodeParentFolder = StaticCastSharedPtr<FSequencerFolderNode>(ParentNode);
-				NodeParentFolder->GetFolder().Modify();
-				NodeParentFolder->GetFolder().RemoveChildFolder(&FolderNode->GetFolder());
+				FocusedMovieScene->Modify();
+				FocusedMovieScene->RemoveRootFolder(FolderNode->GetFolder());
 			}
-			else
-			{
-				UMovieScene* FocusedMovieScene = GetFocusedMovieSceneSequence()->GetMovieScene();
-				if (FocusedMovieScene)
-				{
-					FocusedMovieScene->Modify();
-					FocusedMovieScene->GetRootFolders().Remove(&FolderNode->GetFolder());
-				}
-			}
+		}
 
-			DestinationFolder->AddChildFolder(&FolderNode->GetFolder());
+		DestinationFolder->AddChildFolder(FolderNode->GetFolder());
+	}
+	else if (ITrackExtension* TrackNode = NodeToMove->CastThis<ITrackExtension>())
+	{
+		if (ParentNode.IsValid() && ParentNode->IsA<FFolderModel>())
+		{
+			FFolderModel* NodeParentFolder = ParentNode->CastThisChecked<FFolderModel>();
+			NodeParentFolder->GetFolder()->Modify();
+			NodeParentFolder->GetFolder()->RemoveChildMasterTrack(TrackNode->GetTrack());
+		}
 
-			break;
+		DestinationFolder->AddChildMasterTrack(TrackNode->GetTrack());
+	}
+	else if (IObjectBindingExtension* ObjectBindingNode = NodeToMove->CastThis<IObjectBindingExtension>())
+	{
+		if (ParentNode.IsValid() && ParentNode->IsA<FFolderModel>())
+		{
+			FFolderModel* NodeParentFolder = ParentNode->CastThisChecked<FFolderModel>();
+			NodeParentFolder->GetFolder()->Modify();
+			NodeParentFolder->GetFolder()->RemoveChildObjectBinding(ObjectBindingNode->GetObjectGuid());
 		}
-		case ESequencerNode::Track:
-		{
-			TSharedRef<FSequencerTrackNode> TrackNode = StaticCastSharedRef<FSequencerTrackNode>(NodeToMove);
 
-			if (ParentNode.IsValid())
-			{
-				checkf(ParentNode->GetType() == ESequencerNode::Folder, TEXT("Can not remove from unsupported parent node."));
-				TSharedPtr<FSequencerFolderNode> NodeParentFolder = StaticCastSharedPtr<FSequencerFolderNode>(ParentNode);
-				NodeParentFolder->GetFolder().Modify();
-				NodeParentFolder->GetFolder().RemoveChildMasterTrack(TrackNode->GetTrack());
-			}
-
-			DestinationFolder->AddChildMasterTrack(TrackNode->GetTrack());
-			break;
-		}
-		case ESequencerNode::Object:
-		{
-			TSharedRef<FSequencerObjectBindingNode> ObjectBindingNode = StaticCastSharedRef<FSequencerObjectBindingNode>(NodeToMove);
-			if (ParentNode.IsValid())
-			{
-				checkf(ParentNode->GetType() == ESequencerNode::Folder, TEXT("Can not remove from unsupported parent node."));
-				TSharedPtr<FSequencerFolderNode> NodeParentFolder = StaticCastSharedPtr<FSequencerFolderNode>(ParentNode);
-				NodeParentFolder->GetFolder().Modify();
-				NodeParentFolder->GetFolder().RemoveChildObjectBinding(ObjectBindingNode->GetObjectBinding());
-			}
-
-			DestinationFolder->AddChildObjectBinding(ObjectBindingNode->GetObjectBinding());
-			break;
-		}
+		DestinationFolder->AddChildObjectBinding(ObjectBindingNode->GetObjectGuid());
 	}
 }
 
-TArray<TSharedRef<FSequencerDisplayNode> > FSequencer::GetSelectedNodesToMove()
+TArray<TSharedRef<UE::Sequencer::FViewModel>> FSequencer::GetSelectedNodesToMove()
 {
-	TArray<TSharedRef<FSequencerDisplayNode> > NodesToMove;
+	using namespace UE::Sequencer;
 
+	TArray<TSharedRef<FViewModel>> NodesToMove;
+
 	// Build a list of the nodes we want to move.
-	for (TSharedRef<FSequencerDisplayNode> Node : GetSelection().GetSelectedOutlinerNodes())
+	TArray<TSharedRef<FViewModel>> SelectedItems;
+	GetSelection().GetSelectedOutlinerItems(SelectedItems);
+	for (TSharedRef<FViewModel> Node : SelectedItems)
 	{
-		// Only nodes that can be dragged can be moved in to a folder. They must also either be in the root or in a folder.
-		if (Node->CanDrag() && (!Node->GetParent().IsValid() || Node->GetParent()->GetType() == ESequencerNode::Folder))
+		// Only nodes that can be dragged can be moved in to a folder.
+		if (IDraggableOutlinerExtension* DraggableModel = Node->CastThis<IDraggableOutlinerExtension>())
 		{
-			NodesToMove.Add(Node);
+			if (DraggableModel->CanDrag())
+			{
+				NodesToMove.Add(Node);
+			}
 		}
 	}
 
@@ -7787,18 +7869,22 @@
 	// Find nodes that are children of other nodes in the list
 	for (int32 NodeIndex = 0; NodeIndex < NodesToMove.Num(); ++NodeIndex)
 	{
-		TSharedPtr<FSequencerDisplayNode> Node = NodesToMove[NodeIndex];
+		TSharedPtr<FViewModel> Node = NodesToMove[NodeIndex];
 
-		for (TSharedRef<FSequencerDisplayNode> ParentNode : NodesToMove)
+		for (TSharedRef<FViewModel> ParentNode : NodesToMove)
 		{
 			if (ParentNode == Node)
 			{
 				continue;
 			}
 
-			if (!ParentNode->Traverse_ParentFirst([&](FSequencerDisplayNode& InNode) { return &InNode != Node.Get(); }))
+			for (TSharedPtr<FViewModel> Child : ParentNode->GetDescendants(true))
 			{
-				NodesToRemove.Add(NodeIndex);
+				if (Child == Node)
+				{
+					NodesToRemove.Add(NodeIndex);
+					break;
+				}
 			}
 		}
 	}
@@ -7813,29 +7899,31 @@
 	return NodesToMove;
 }
 
-TArray<TSharedRef<FSequencerDisplayNode> > FSequencer::GetSelectedNodesInFolders()
+TArray<TSharedRef<UE::Sequencer::FViewModel> > FSequencer::GetSelectedNodesInFolders()
 {
-	TArray<TSharedRef<FSequencerDisplayNode> > NodesToFolders;
+	using namespace UE::Sequencer;
 
-	for (TSharedRef<FSequencerDisplayNode> SelectedNode : GetSelection().GetSelectedOutlinerNodes())
+	TArray<TSharedRef<FViewModel>> SelectedItems;
+	GetSelection().GetSelectedOutlinerItems(SelectedItems);
+
+	TArray<TSharedRef<FViewModel> > NodesToFolders;
+	for (TSharedRef<FViewModel> SelectedNode : SelectedItems)
 	{
-		TSharedPtr<FSequencerFolderNode> Folder = SelectedNode->FindFolderNode();
+		TSharedPtr<FFolderModel> Folder = SelectedNode->FindAncestorOfType<FFolderModel>();
 		if (Folder.IsValid())
 		{
-			if (SelectedNode->GetType() == ESequencerNode::Object)
+			if (IObjectBindingExtension* ObjectBindingNode = SelectedNode->CastThis<IObjectBindingExtension>())
 			{
-				TSharedRef<FSequencerObjectBindingNode> ObjectBindingNode = StaticCastSharedRef<FSequencerObjectBindingNode>(SelectedNode);
-				if (Folder->GetFolder().GetChildObjectBindings().Contains(ObjectBindingNode->GetObjectBinding()))
+				if (Folder->GetFolder()->GetChildObjectBindings().Contains(ObjectBindingNode->GetObjectGuid()))
 				{
 					NodesToFolders.Add(SelectedNode);
 				}
 			}
-			else if (SelectedNode->GetType() == ESequencerNode::Track)
+			else if (ITrackExtension* TrackNode = SelectedNode->CastThis<ITrackExtension>())
 			{
-				TSharedRef<FSequencerTrackNode> TrackNode = StaticCastSharedRef<FSequencerTrackNode>(SelectedNode);
 				if (TrackNode->GetTrack())
 				{
-					if (Folder->GetFolder().GetChildMasterTracks().Contains(TrackNode->GetTrack()))
+					if (Folder->GetFolder()->GetChildMasterTracks().Contains(TrackNode->GetTrack()))
 					{
 						NodesToFolders.Add(SelectedNode);
 					}
@@ -7849,6 +7937,8 @@
 
 void FSequencer::MoveSelectedNodesToFolder(UMovieSceneFolder* DestinationFolder)
 {
+	using namespace UE::Sequencer;
+
 	if (!DestinationFolder)
 	{
 		return;
@@ -7867,14 +7957,14 @@
 		return;
 	}
 
-	TArray<TSharedRef<FSequencerDisplayNode> > NodesToMove = GetSelectedNodesToMove();
+	TArray<TSharedRef<FViewModel> > NodesToMove = GetSelectedNodesToMove();
 
-	for (TSharedRef<FSequencerDisplayNode> Node : NodesToMove)
+	for (TSharedRef<FViewModel> Node : NodesToMove)
 	{
 		// If this node is the destination folder, don't try to move it
-		if (Node->GetType() == ESequencerNode::Folder)
+		if (FFolderModel* FolderNode = Node->CastThis<FFolderModel>())
 		{
-			if (&StaticCastSharedRef<FSequencerFolderNode>(Node)->GetFolder() == DestinationFolder)
+			if (FolderNode->GetFolder() == DestinationFolder)
 			{
 				NodesToMove.Remove(Node);
 				break;
@@ -7891,11 +7981,11 @@
 	int32 SharedPathLength = TNumericLimits<int32>::Max();
 
 	// Build a list of the paths for each node, split in to folder names
-	for (TSharedRef<FSequencerDisplayNode> Node : NodesToMove)
+	for (TSharedRef<FViewModel> Node : NodesToMove)
 	{
 		// Split the node's path in to segments
 		TArray<FString>& NodePath = NodePathSplits.AddDefaulted_GetRef();
-		Node->GetPathName().ParseIntoArray(NodePath, TEXT("."));
+		IOutlinerExtension::GetPathName(*Node).ParseIntoArray(NodePath, TEXT("."));
 
 		// Shared path obviously won't be larger than the shortest path
 		SharedPathLength = FMath::Min(SharedPathLength, NodePath.Num() - 1);
@@ -7963,20 +8053,26 @@
 	FString DestinationFolderPath;
 	if (DestinationFolder)
 	{
-		for (TSharedRef<FSequencerDisplayNode> Node : NodeTree->GetAllNodes())
+		TArray<TSharedRef<FViewModel>> AllNodes;
+		NodeTree->GetAllNodes(AllNodes);
+		for (TSharedRef<FViewModel> Node : AllNodes)
 		{
 			// If this node is the destination folder, don't try to move it
-			if (Node->GetType() == ESequencerNode::Folder)
+			if (FFolderModel* FolderNode = Node->CastThis<FFolderModel>())
 			{
-				if (&StaticCastSharedRef<FSequencerFolderNode>(Node)->GetFolder() == DestinationFolder)
+				if (FolderNode->GetFolder() == DestinationFolder)
 				{
-					DestinationFolderPath = Node->GetPathName();
+					DestinationFolderPath = IOutlinerExtension::GetPathName(*Node);
 
 					// Expand the folders to our destination
-					TSharedPtr<FSequencerDisplayNode> ParentNode = Node;
+					TSharedPtr<FViewModel> ParentNode = Node;
 					while (ParentNode)
 					{
-						ParentNode->SetExpansionState(true);
+						if (ParentNode->IsA<IOutlinerExtension>())
+						{
+							ParentNode->CastThisChecked<IOutlinerExtension>()->SetExpansion(true);
+						}
+						
 						ParentNode = ParentNode->GetParent();
 					}
 					break;
@@ -7987,7 +8083,7 @@
 
 	for (int32 NodeIndex = 0; NodeIndex < NodesToMove.Num(); ++NodeIndex)
 	{
-		TSharedRef<FSequencerDisplayNode> Node = NodesToMove[NodeIndex];
+		TSharedRef<FViewModel> Node = NodesToMove[NodeIndex];
 		TArray<FString>& NodePathSplit = NodePathSplits[NodeIndex];
 
 		// Reset the relative path
@@ -8007,7 +8103,7 @@
 			NewPath += NodePathSplit[FolderPathIndex] + TEXT(".");
 		}
 
-		NewPath += Node->GetNodeName().ToString();
+		NewPath += Node->CastThis<IOutlinerExtension>()->GetIdentifier().ToString();
 
 		UMovieSceneFolder* NodeDestinationFolder = CreateFoldersRecursively(FolderPath, 0, FocusedMovieScene, DestinationFolder, DestinationFolder->GetChildFolders());
 		MoveNodeToFolder(Node, NodeDestinationFolder);
@@ -8021,6 +8117,8 @@
 
 void FSequencer::MoveSelectedNodesToNewFolder()
 {
+	using namespace UE::Sequencer;
+
 	UMovieScene* FocusedMovieScene = GetFocusedMovieSceneSequence()->GetMovieScene();
 
 	if (!FocusedMovieScene)
@@ -8034,7 +8132,7 @@
 		return;
 	}
 
-	TArray<TSharedRef<FSequencerDisplayNode> > NodesToMove = GetSelectedNodesToMove();
+	TArray<TSharedRef<FViewModel> > NodesToMove = GetSelectedNodesToMove();
 
 	if (!NodesToMove.Num())
 	{
@@ -8045,11 +8143,11 @@
 	int32 SharedPathLength = TNumericLimits<int32>::Max();
 
 	// Build a list of the paths for each node, split in to folder names
-	for (TSharedRef<FSequencerDisplayNode> Node : NodesToMove)
+	for (TSharedRef<FViewModel> Node : NodesToMove)
 	{
 		// Split the node's path in to segments
 		TArray<FString>& NodePath = NodePathSplits.AddDefaulted_GetRef();
-		Node->GetPathName().ParseIntoArray(NodePath, TEXT("."));
+		IOutlinerExtension::GetPathName(*Node).ParseIntoArray(NodePath, TEXT("."));
 
 		// Shared path obviously won't be larger than the shortest path
 		SharedPathLength = FMath::Min(SharedPathLength, NodePath.Num() - 1);
@@ -8146,7 +8244,7 @@
 	if (!ParentFolder)
 	{
 		FocusedMovieScene->Modify();
-		FocusedMovieScene->GetRootFolders().Add(SharedFolder);
+		FocusedMovieScene->AddRootFolder(SharedFolder);
 	}
 	else
 	{
@@ -8156,7 +8254,7 @@
 
 	for (int32 NodeIndex = 0; NodeIndex < NodesToMove.Num() ; ++NodeIndex)
 	{
-		TSharedRef<FSequencerDisplayNode> Node = NodesToMove[NodeIndex];
+		TSharedRef<FViewModel> Node = NodesToMove[NodeIndex];
 		TArray<FString>& NodePathSplit = NodePathSplits[NodeIndex];
 
 		// Reset to just the path to the shared folder
@@ -8184,6 +8282,8 @@
 
 void FSequencer::RemoveSelectedNodesFromFolders()
 {
+	using namespace UE::Sequencer;
+
 	UMovieScene* FocusedMovieScene = GetFocusedMovieSceneSequence()->GetMovieScene();
 
 	if (!FocusedMovieScene)
@@ -8197,7 +8297,7 @@
 		return;
 	}
 
-	TArray<TSharedRef<FSequencerDisplayNode> > NodesToFolders = GetSelectedNodesInFolders();
+	TArray<TSharedRef<FViewModel> > NodesToFolders = GetSelectedNodesInFolders();
 	if (!NodesToFolders.Num())
 	{
 		return;
@@ -8207,22 +8307,20 @@
 
 	FocusedMovieScene->Modify();
 
-	for (TSharedRef<FSequencerDisplayNode> NodeInFolder : NodesToFolders)
+	for (TSharedRef<FViewModel> NodeInFolder : NodesToFolders)
 	{
-		TSharedPtr<FSequencerFolderNode> Folder = NodeInFolder->FindFolderNode();
+		TSharedPtr<FFolderModel> Folder = NodeInFolder->FindAncestorOfType<FFolderModel>();
 		if (Folder.IsValid())
 		{
-			if (NodeInFolder->GetType() == ESequencerNode::Object)
+			if (IObjectBindingExtension* ObjectBindingNode = NodeInFolder->CastThis<IObjectBindingExtension>())
 			{
-				TSharedRef<FSequencerObjectBindingNode> ObjectBindingNode = StaticCastSharedRef<FSequencerObjectBindingNode>(NodeInFolder);
-				Folder->GetFolder().RemoveChildObjectBinding(ObjectBindingNode->GetObjectBinding());
+				Folder->GetFolder()->RemoveChildObjectBinding(ObjectBindingNode->GetObjectGuid());
 			}
-			else if (NodeInFolder->GetType() == ESequencerNode::Track)
+			else if (ITrackExtension* TrackNode = NodeInFolder->CastThis<ITrackExtension>())
 			{
-				TSharedRef<FSequencerTrackNode> TrackNode = StaticCastSharedRef<FSequencerTrackNode>(NodeInFolder);
 				if (TrackNode->GetTrack())
 				{
-					Folder->GetFolder().RemoveChildMasterTrack(TrackNode->GetTrack());
+					Folder->GetFolder()->RemoveChildMasterTrack(TrackNode->GetTrack());
 				}
 			}
 		}
@@ -8252,14 +8350,14 @@
 		// We can't use TextExportTransient on USTRUCTS (which our object contains) so we're going to manually null out some references before serializing them. These references are
 		// serialized manually into the archive, as the auto-serialization will only store a reference (to a privately owned object) which creates issues on deserialization. Attempting 
 		// to deserialize these private objects throws a superflous error in the console that makes it look like things went wrong when they're actually OK and expected.
-		TArray<UMovieSceneTrack*> OldTracks = ObjectToExport->Binding.StealTracks();
+		TArray<UMovieSceneTrack*> OldTracks = ObjectToExport->Binding.StealTracks(nullptr);
 		UObject* OldSpawnableTemplate = ObjectToExport->Spawnable.GetObjectTemplate();
 		ObjectToExport->Spawnable.SetObjectTemplate(nullptr);
 
 		UExporter::ExportToOutputDevice(&Context, ObjectToExport, nullptr, Archive, TEXT("copy"), 0, PPF_ExportsNotFullyQualified | PPF_Copy | PPF_Delimited, false, ThisOuter);
 
 		// Restore the references (as we don't want to modify the original in the event of a copy operation!)
-		ObjectToExport->Binding.SetTracks(OldTracks);
+		ObjectToExport->Binding.SetTracks(MoveTemp(OldTracks), nullptr);
 		ObjectToExport->Spawnable.SetObjectTemplate(OldSpawnableTemplate);
 
 		// We manually export the object template for the same private-ownership reason as above. Templates need to be re-created anyways as each Spawnable contains its own copy of the template.
@@ -8346,38 +8444,39 @@
 
 TArray<TSharedPtr<FMovieSceneClipboard>> GClipboardStack;
 
-void FSequencer::CopySelectedObjects(TArray<TSharedPtr<FSequencerObjectBindingNode>>& ObjectNodes, const TArray<UMovieSceneFolder*>& Folders, FString& ExportedText)
+void FSequencer::CopySelectedObjects(TArray<TSharedPtr<UE::Sequencer::FObjectBindingModel>>& ObjectNodes, const TArray<UMovieSceneFolder*>& Folders, FString& ExportedText)
 {
+	using namespace UE::Sequencer;
+
 	UMovieScene* MovieScene = GetFocusedMovieSceneSequence()->GetMovieScene();
 
 	// Gather guids for the object nodes and any child object nodes
 	TSet<FGuid> GuidsToCopy;
 	TMap<FGuid, UMovieSceneFolder*> GuidToFolder;
-	for (TSharedPtr<FSequencerObjectBindingNode> ObjectNode : ObjectNodes)
+	for (TSharedPtr<FObjectBindingModel> ObjectNode : ObjectNodes)
 	{
-		GuidsToCopy.Add(ObjectNode->GetObjectBinding());
+		GuidsToCopy.Add(ObjectNode->GetObjectGuid());
 
-		TSharedPtr<FSequencerFolderNode> FolderNode = ObjectNode->FindFolderNode();
-		if (FolderNode.IsValid() && Folders.Contains(&FolderNode->GetFolder()))
+		TSharedPtr<FFolderModel> FolderNode = ObjectNode->FindAncestorOfType<FFolderModel>();
+		if (FolderNode.IsValid() && Folders.Contains(FolderNode->GetFolder()))
 		{
-			GuidToFolder.Add(ObjectNode->GetObjectBinding(), &FolderNode->GetFolder());
+			GuidToFolder.Add(ObjectNode->GetObjectGuid(), FolderNode->GetFolder());
 		}
 
-		TSet<TSharedRef<FSequencerDisplayNode> > DescendantNodes;
+		TSet<TSharedRef<FViewModel> > DescendantNodes;
 
 		SequencerHelpers::GetDescendantNodes(ObjectNode.ToSharedRef(), DescendantNodes);
 
 		for (auto DescendantNode : DescendantNodes)
 		{
-			if (DescendantNode->GetType() == ESequencerNode::Object)
+			if (FObjectBindingModel* DescendantObjectNode = DescendantNode->CastThis<FObjectBindingModel>())
 			{
-				TSharedRef<FSequencerObjectBindingNode> DescendantObjectNode = StaticCastSharedRef<FSequencerObjectBindingNode>(DescendantNode);
-				GuidsToCopy.Add(DescendantObjectNode->GetObjectBinding());
+				GuidsToCopy.Add(DescendantObjectNode->GetObjectGuid());
 
-				TSharedPtr<FSequencerFolderNode> DescendantFolderNode = DescendantObjectNode->FindFolderNode();
+				TSharedPtr<FFolderModel> DescendantFolderNode = DescendantObjectNode->FindAncestorOfType<FFolderModel>();
 				if (DescendantFolderNode.IsValid())
 				{
-					GuidToFolder.Add(DescendantObjectNode->GetObjectBinding(), &DescendantFolderNode->GetFolder());
+					GuidToFolder.Add(DescendantObjectNode->GetObjectGuid(), DescendantFolderNode->GetFolder());
 				}
 			}
 		}
@@ -8437,18 +8536,20 @@
 	}
 }
 
-void FSequencer::CopySelectedTracks(TArray<TSharedPtr<FSequencerTrackNode>>& TrackNodes, const TArray<UMovieSceneFolder*>& Folders, FString& ExportedText)
+void FSequencer::CopySelectedTracks(TArray<TSharedPtr<FViewModel>>& TrackNodes, const TArray<UMovieSceneFolder*>& Folders, FString& ExportedText)
 {
+	using namespace UE::Sequencer;
+
 	UMovieScene* MovieScene = GetFocusedMovieSceneSequence()->GetMovieScene();
 
 	TArray<UObject*> CopyableObjects;
-	for (TSharedPtr<FSequencerTrackNode> TrackNode : TrackNodes)
+	for (TSharedPtr<FViewModel> TrackNode : TrackNodes)
 	{
 		bool bIsParentSelected = false;
-		TSharedPtr<FSequencerDisplayNode> ParentNode = TrackNode->GetParent();
-		while (ParentNode.IsValid() && ParentNode->GetType() != ESequencerNode::Folder)
+		TSharedPtr<FViewModel> ParentNode = TrackNode->GetParent();
+		while (ParentNode.IsValid() && !ParentNode->IsA<FFolderModel>())
 		{
-			if (Selection.GetSelectedOutlinerNodes().Contains(ParentNode.ToSharedRef()))
+			if (Selection.GetSelectedOutlinerItems().Contains(ParentNode.ToSharedRef()))
 			{
 				bIsParentSelected = true;
 				break;
@@ -8459,11 +8560,11 @@
 		if (!bIsParentSelected)
 		{
 			// If this is a subtrack, only copy the sections that belong to this row. otherwise copying the entire track will copy all the sections across all the rows
-			if (TrackNode->GetSubTrackMode() == FSequencerTrackNode::ESubTrackMode::SubTrack)
+			if (FTrackRowModel* TrackRowNode = TrackNode->CastThis<FTrackRowModel>())
 			{
-				for (UMovieSceneSection* Section : TrackNode->GetTrack()->GetAllSections())
+				for (UMovieSceneSection* Section : TrackRowNode->GetTrack()->GetAllSections())
 				{
-					if (Section && Section->GetRowIndex() == TrackNode->GetRowIndex())
+					if (Section && Section->GetRowIndex() == TrackRowNode->GetRowIndex())
 					{
 						CopyableObjects.Add(Section);
 					}
@@ -8474,15 +8575,15 @@
 				UMovieSceneCopyableTrack* CopyableTrack = NewObject<UMovieSceneCopyableTrack>(GetTransientPackage(), UMovieSceneCopyableTrack::StaticClass(), NAME_None, RF_Transient);
 				CopyableObjects.Add(CopyableTrack);
 
-				UMovieSceneTrack* DuplicatedTrack = Cast<UMovieSceneTrack>(StaticDuplicateObject(TrackNode->GetTrack(), CopyableTrack));
+				UMovieSceneTrack* DuplicatedTrack = Cast<UMovieSceneTrack>(StaticDuplicateObject(TrackNode->CastThisChecked<ITrackExtension>()->GetTrack(), CopyableTrack));
 				CopyableTrack->Track = DuplicatedTrack;
-				CopyableTrack->bIsAMasterTrack = MovieScene->IsAMasterTrack(*TrackNode->GetTrack());
-				CopyableTrack->bIsACameraCutTrack = TrackNode->GetTrack()->IsA<UMovieSceneCameraCutTrack>();
+				CopyableTrack->bIsAMasterTrack = MovieScene->IsAMasterTrack(*TrackNode->CastThisChecked<ITrackExtension>()->GetTrack());
+				CopyableTrack->bIsACameraCutTrack = TrackNode->CastThisChecked<ITrackExtension>()->GetTrack()->IsA<UMovieSceneCameraCutTrack>();
 
-				TSharedPtr<FSequencerFolderNode> FolderNode = TrackNode->FindFolderNode();
-				if (FolderNode.IsValid() && Folders.Contains(&FolderNode->GetFolder()))
+				TSharedPtr<FFolderModel> FolderNode = TrackNode->FindAncestorOfType<FFolderModel>();
+				if (FolderNode.IsValid() && Folders.Contains(FolderNode->GetFolder()))
 				{
-					UMovieSceneFolder::CalculateFolderPath(&FolderNode->GetFolder(), Folders, CopyableTrack->FolderPath);
+					UMovieSceneFolder::CalculateFolderPath(FolderNode->GetFolder(), Folders, CopyableTrack->FolderPath);
 				}
 			}
 		}
@@ -8659,7 +8760,7 @@
 			}
 			else
 			{
-				MovieScene->GetRootFolders().Add(CopiedFolder);
+				MovieScene->AddRootFolder(CopiedFolder);
 			}
 		}
 	}
@@ -8671,6 +8772,8 @@
 
 bool FSequencer::PasteObjectBindings(const FString& TextToImport, UMovieSceneFolder* InParentFolder, const TArray<UMovieSceneFolder*>& InFolders, TArray<FNotificationInfo>& PasteErrors, bool bClearSelection)
 {
+	using namespace UE::Sequencer;
+
 	UWorld* World = Cast<UWorld>(GetPlaybackContext());
 
 	UMovieSceneSequence* OwnerSequence = GetFocusedMovieSceneSequence();
@@ -8684,22 +8787,17 @@
 
 	TArray<FMovieSceneBinding> BindingsPasted;
 
-	TSet<TSharedRef<FSequencerDisplayNode>> SelectedNodes = Selection.GetSelectedOutlinerNodes();
+	TArray<TSharedRef<FViewModel>> SelectedNodes;
+	Selection.GetSelectedOutlinerItems(SelectedNodes);
 
 	TArray<FGuid> SelectedParentGuids;
 	if (!bClearSelection)
 	{
-		for (TSharedRef<FSequencerDisplayNode> Node : SelectedNodes)
+		for (TSharedRef<FViewModel> Node : SelectedNodes)
 		{
-			if (Node->GetType() != ESequencerNode::Object)
+			if (FObjectBindingModel* ObjectNode = Node->CastThis<FObjectBindingModel>())
 			{
-				continue;
-			}
-
-			TSharedPtr<FSequencerObjectBindingNode> ObjectNode = StaticCastSharedRef<FSequencerObjectBindingNode>(Node);
-			if (ObjectNode.IsValid())
-			{
-				SelectedParentGuids.Add(ObjectNode->GetObjectBinding());
+				SelectedParentGuids.Add(ObjectNode->GetObjectGuid());
 			}
 		}
 	}
@@ -8767,7 +8865,7 @@
 				{
 					if (TargetIndex < SelectedParentGuids.Num())
 					{
-						Possessable->SetParent(SelectedParentGuids[TargetIndex]);
+						Possessable->SetParent(SelectedParentGuids[TargetIndex], MovieScene);
 					}
 				}
 
@@ -8864,7 +8962,7 @@
 		FMovieScenePossessable* Possessable = MovieScene->FindPossessable(PossessableGuid);
 		if (Possessable && OldToNewGuidMap.Contains(Possessable->GetParent()) && PossessableGuid != OldToNewGuidMap[Possessable->GetParent()])
 		{
-			Possessable->SetParent(OldToNewGuidMap[Possessable->GetParent()]);
+			Possessable->SetParent(OldToNewGuidMap[Possessable->GetParent()], MovieScene);
 		}
 	}
 
@@ -8971,7 +9069,7 @@
 				FMovieSceneSpawnable* SpawnableParent = MovieScene->FindSpawnable(ParentGuid);
 				if (!PossessableParent && !SpawnableParent)
 				{
-					Possessable->SetParent(FGuid());
+					Possessable->SetParent(FGuid(), MovieScene);
 				}
 				else if (SpawnableParent)
 				{
@@ -8994,6 +9092,8 @@
 
 bool FSequencer::PasteTracks(const FString& TextToImport, UMovieSceneFolder* InParentFolder, const TArray<UMovieSceneFolder*>& InFolders, TArray<FNotificationInfo>& PasteErrors, bool bClearSelection)
 {
+	using namespace UE::Sequencer;
+
 	TArray<UMovieSceneCopyableTrack*> ImportedTracks;
 	FSequencer::ImportTracksFromText(TextToImport, ImportedTracks);
 
@@ -9005,22 +9105,17 @@
 	UMovieScene* MovieScene = GetFocusedMovieSceneSequence()->GetMovieScene();
 	UObject* BindingContext = GetPlaybackContext();
 
-	TSet<TSharedRef<FSequencerDisplayNode>> SelectedNodes = Selection.GetSelectedOutlinerNodes();
+	TArray<TSharedRef<FViewModel>> SelectedNodes;
+	Selection.GetSelectedOutlinerItems(SelectedNodes);
 
-	TArray<TSharedPtr<FSequencerObjectBindingNode>> ObjectNodes;
+	TArray<TSharedPtr<FObjectBindingModel>> ObjectNodes;
 
 	if (!bClearSelection)
 	{
-		for (TSharedRef<FSequencerDisplayNode> Node : SelectedNodes)
+		for (TSharedRef<FViewModel> Node : SelectedNodes)
 		{
-			if (Node->GetType() != ESequencerNode::Object)
+			if (TSharedPtr<FObjectBindingModel> ObjectNode = Node->CastThisShared<FObjectBindingModel>())
 			{
-				continue;
-			}
-
-			TSharedPtr<FSequencerObjectBindingNode> ObjectNode = StaticCastSharedRef<FSequencerObjectBindingNode>(Node);
-			if (ObjectNode.IsValid())
-			{
 				ObjectNodes.Add(ObjectNode);
 			}
 		}
@@ -9045,9 +9140,9 @@
 	int32 NumMasterOrCameraCutTracksPasted = 0;
 	if (ObjectNodes.Num())
 	{
-		for (TSharedPtr<FSequencerObjectBindingNode> ObjectNode : ObjectNodes)
+		for (TSharedPtr<FObjectBindingModel> ObjectNode : ObjectNodes)
 		{
-			FGuid ObjectGuid = ObjectNode->GetObjectBinding();
+			FGuid ObjectGuid = ObjectNode->GetObjectGuid();
 
 			TArray<UMovieSceneCopyableTrack*> NewTracks;
 			FSequencer::ImportTracksFromText(TextToImport, NewTracks);
@@ -9164,16 +9259,17 @@
 	return false;
 }
 
-void GetSupportedTracks(TSharedRef<FSequencerDisplayNode> DisplayNode, const TArray<UMovieSceneSection*>& ImportedSections, TArray<TSharedRef<FSequencerTrackNode>>& TracksToPasteOnto)
+void GetSupportedTracks(TSharedRef<UE::Sequencer::FViewModel> DisplayNode, const TArray<UMovieSceneSection*>& ImportedSections, TArray<TSharedRef<UE::Sequencer::FViewModel>>& TracksToPasteOnto)
 {
-	if (DisplayNode->GetType() != ESequencerNode::Track)
+	using namespace UE::Sequencer;
+
+	ITrackExtension* TrackNode = DisplayNode->CastThis<ITrackExtension>();
+	if (TrackNode == nullptr)
 	{
 		return;
 	}
 
-	TSharedRef<FSequencerTrackNode> TrackNode = StaticCastSharedRef<FSequencerTrackNode>(DisplayNode);
-
-	if (TracksToPasteOnto.Contains(TrackNode))
+	if (TracksToPasteOnto.Contains(DisplayNode))
 	{
 		return;
 	}
@@ -9187,7 +9283,7 @@
 
 			if (Track->SupportsType(Section->GetClass()))
 			{
-				TracksToPasteOnto.Add(TrackNode);
+				TracksToPasteOnto.Add(DisplayNode);
 				return;
 			}
 		}
@@ -9197,6 +9293,10 @@
 
 bool FSequencer::PasteSections(const FString& TextToImport, TArray<FNotificationInfo>& PasteErrors)
 {
+	using namespace UE::Sequencer;
+
+	FScopedTransaction Transaction(FGenericCommands::Get().Paste->GetDescription());
+
 	// First import as a track and extract sections to allow for copying track contents to another track
 	TArray<UMovieSceneCopyableTrack*> ImportedTracks;
 	FSequencer::ImportTracksFromText(TextToImport, ImportedTracks);
@@ -9221,19 +9321,20 @@
 		return false;
 	}
 
-	TSet<TSharedRef<FSequencerDisplayNode>> SelectedNodes = Selection.GetSelectedOutlinerNodes();
+	TArray<TSharedRef<FViewModel>> SelectedNodes;
+	Selection.GetSelectedOutlinerItems(SelectedNodes);
 
 	if (SelectedNodes.Num() == 0)
 	{
-		for (const TSharedRef<FSequencerDisplayNode>& DisplayNode : NodeTree->GetRootNodes())
+		for (const TViewModelPtr<IOutlinerExtension>& RootNode : NodeTree->GetRootNodes())
 		{
 			TSet<TWeakObjectPtr<UMovieSceneSection>> Sections;
-			SequencerHelpers::GetAllSections(DisplayNode, Sections);
+			SequencerHelpers::GetAllSections(RootNode.AsModel(), Sections);
 			for (TWeakObjectPtr<UMovieSceneSection> Section : Sections)
 			{
 				if (Selection.GetSelectedSections().Contains(Section.Get()))
 				{
-					SelectedNodes.Add(DisplayNode);
+					SelectedNodes.Add(RootNode.AsModel().ToSharedRef());
 					break;
 				}
 			}
@@ -9269,8 +9370,8 @@
 	}
 
 	// Check if any of the selected nodes supports pasting this type of section
-	TArray<TSharedRef<FSequencerTrackNode>> TracksToPasteOnto;
-	for (TSharedRef<FSequencerDisplayNode> Node : SelectedNodes)
+	TArray<TSharedRef<FViewModel>> TracksToPasteOnto;
+	for (TSharedRef<FViewModel> Node : SelectedNodes)
 	{
 		GetSupportedTracks(Node, ImportedSections, TracksToPasteOnto);
 	}
@@ -9278,21 +9379,17 @@
 	// Otherwise, look at all child nodes for supported tracks
 	if (TracksToPasteOnto.Num() == 0)
 	{
-		for (TSharedRef<FSequencerDisplayNode> Node : SelectedNodes)
+		for (TSharedRef<FViewModel> Node : SelectedNodes)
 		{
-			TSet<TSharedRef<FSequencerDisplayNode> > DescendantNodes;
+			TSet<TSharedRef<FViewModel> > DescendantNodes;
 			SequencerHelpers::GetDescendantNodes(Node, DescendantNodes);
 
-			for (TSharedRef<FSequencerDisplayNode> DescendantNode : DescendantNodes)
+			for (TSharedRef<FViewModel> DescendantNode : DescendantNodes)
 			{
 				// Don't automatically paste onto subtracks because that would lead to multiple paste destinations
-				if (DescendantNode->GetType() == ESequencerNode::Track)
+				if (DescendantNode->IsA<FTrackRowModel>())
 				{
-					TSharedPtr<FSequencerTrackNode> DescendantTrackNode = StaticCastSharedRef<FSequencerTrackNode>(DescendantNode);
-					if (DescendantTrackNode.IsValid() && DescendantTrackNode->GetSubTrackMode() == FSequencerTrackNode::ESubTrackMode::SubTrack)
-					{
-						continue;
-					}
+					continue;
 				}
 				GetSupportedTracks(DescendantNode, ImportedSections, TracksToPasteOnto);
 			}
@@ -9302,9 +9399,9 @@
 	TArray<UMovieSceneSection*> NewSections;
 	TArray<int32> SectionIndicesImported;
 
-	for (TSharedRef<FSequencerTrackNode> TrackNode : TracksToPasteOnto)
+	for (TSharedRef<FViewModel> TrackNode : TracksToPasteOnto)
 	{
-		UMovieSceneTrack* Track = TrackNode->GetTrack();
+		UMovieSceneTrack* Track = TrackNode->CastThisChecked<ITrackExtension>()->GetTrack();
 		for (int32 SectionIndex = 0; SectionIndex < ImportedSections.Num(); ++SectionIndex)
 		{
 			UMovieSceneSection* Section = ImportedSections[SectionIndex];
@@ -9328,9 +9425,9 @@
 
 			if (Track->SupportsMultipleRows())
 			{
-				if (TrackNode->GetSubTrackMode() == FSequencerTrackNode::ESubTrackMode::SubTrack)
+				if (FTrackRowModel* TrackRowNode = TrackNode->CastThis<FTrackRowModel>())
 				{
-					Section->SetRowIndex(TrackNode->GetRowIndex());
+					Section->SetRowIndex(TrackRowNode->GetRowIndex());
 				}
 			}
 			NewSections.Add(Section);
@@ -9396,11 +9493,26 @@
 	}
 
 	NotifyMovieSceneDataChanged(EMovieSceneDataChangeType::MovieSceneStructureItemAdded);
+
+	FSectionModelStorageExtension* SectionModelStorage = ViewModel->GetRootModel()->CastDynamic<FSectionModelStorageExtension>();
+	{
+		UE::MovieScene::FScopedSignedObjectModifyDefer ForceFlush(true);
+		check(SectionModelStorage);
+	}
+
 	EmptySelection();
+
+	Selection.SuspendBroadcast();
 	for (UMovieSceneSection* NewSection : NewSections)
 	{
-		SelectSection(NewSection);
+		TSharedPtr<FSectionModel> SectionModel = SectionModelStorage->FindModelForSection(NewSection);
+		if (ensure(SectionModel))
+		{
+			Selection.AddToSelection(SectionModel);
+		}
 	}
+	Selection.ResumeBroadcast();
+
 	ThrobSectionSelection();
 
 	return true;
@@ -9607,10 +9719,10 @@
 	bool bIsActive = !IsNodeActive();
 	const FScopedTransaction Transaction( NSLOCTEXT("Sequencer", "ToggleNodeActive", "Toggle Node Active") );
 
-	for (auto OutlinerNode : Selection.GetSelectedOutlinerNodes())
+	for (auto OutlinerNode : Selection.GetSelectedOutlinerItems())
 	{
 		TSet<TWeakObjectPtr<UMovieSceneSection> > Sections;
-		SequencerHelpers::GetAllSections(OutlinerNode, Sections);
+		SequencerHelpers::GetAllSections(OutlinerNode.Pin(), Sections);
 
 		for (auto Section : Sections)
 		{
@@ -9626,10 +9738,10 @@
 bool FSequencer::IsNodeActive() const
 {
 	// Active if ONE is active, changed in 4.20
-	for (auto OutlinerNode : Selection.GetSelectedOutlinerNodes())
+	for (auto OutlinerNode : Selection.GetSelectedOutlinerItems())
 	{
 		TSet<TWeakObjectPtr<UMovieSceneSection> > Sections;
-		SequencerHelpers::GetAllSections(OutlinerNode, Sections);
+		SequencerHelpers::GetAllSections(OutlinerNode.Pin(), Sections);
 		if (Sections.Num() > 0)
 		{
 			for (auto Section : Sections)
@@ -9652,10 +9764,10 @@
 
 	const FScopedTransaction Transaction( NSLOCTEXT("Sequencer", "ToggleNodeLocked", "Toggle Node Locked") );
 
-	for (auto OutlinerNode : Selection.GetSelectedOutlinerNodes())
+	for (auto OutlinerNode : Selection.GetSelectedOutlinerItems())
 	{
 		TSet<TWeakObjectPtr<UMovieSceneSection> > Sections;
-		SequencerHelpers::GetAllSections(OutlinerNode, Sections);
+		SequencerHelpers::GetAllSections(OutlinerNode.Pin(), Sections);
 
 		for (auto Section : Sections)
 		{
@@ -9670,10 +9782,10 @@
 {
 	// Locked only if all are locked
 	int NumSections = 0;
-	for (auto OutlinerNode : Selection.GetSelectedOutlinerNodes())
+	for (auto OutlinerNode : Selection.GetSelectedOutlinerItems())
 	{
 		TSet<TWeakObjectPtr<UMovieSceneSection> > Sections;
-		SequencerHelpers::GetAllSections(OutlinerNode, Sections);
+		SequencerHelpers::GetAllSections(OutlinerNode.Pin(), Sections);
 
 		for (auto Section : Sections)
 		{
@@ -9766,6 +9878,8 @@
 
 void FSequencer::SaveSelectedNodesSpawnableState()
 {
+	using namespace UE::Sequencer;
+
 	UMovieScene* MovieScene = GetFocusedMovieSceneSequence()->GetMovieScene();
 
 	if (MovieScene->IsReadOnly())
@@ -9780,11 +9894,11 @@
 
 	TArray<FMovieSceneSpawnable*> Spawnables;
 
-	for (const TSharedRef<FSequencerDisplayNode>& Node : Selection.GetSelectedOutlinerNodes())
+	for (const TWeakPtr<FViewModel>& Node : Selection.GetSelectedOutlinerItems())
 	{
-		if (Node->GetType() == ESequencerNode::Object)
+		if (IObjectBindingExtension* ObjectBindingNode = ICastable::CastWeakPtr<IObjectBindingExtension>(Node))
 		{
-			FMovieSceneSpawnable* Spawnable = MovieScene->FindSpawnable(StaticCastSharedRef<FSequencerObjectBindingNode>(Node)->GetObjectBinding());
+			FMovieSceneSpawnable* Spawnable = MovieScene->FindSpawnable(ObjectBindingNode->GetObjectGuid());
 			if (Spawnable)
 			{
 				Spawnables.Add(Spawnable);
@@ -9812,7 +9926,9 @@
 }
 
 void FSequencer::SetSelectedNodesSpawnableLevel(FName InLevelName)
-{							
+{						
+	using namespace UE::Sequencer;
+
 	UMovieScene* MovieScene = GetFocusedMovieSceneSequence()->GetMovieScene();
 
 	if (MovieScene->IsReadOnly())
@@ -9827,11 +9943,11 @@
 
 	TArray<FMovieSceneSpawnable*> Spawnables;
 
-	for (const TSharedRef<FSequencerDisplayNode>& Node : Selection.GetSelectedOutlinerNodes())
+	for (const TWeakPtr<FViewModel>& Node : Selection.GetSelectedOutlinerItems())
 	{
-		if (Node->GetType() == ESequencerNode::Object)
+		if (IObjectBindingExtension* ObjectBindingNode = ICastable::CastWeakPtr<IObjectBindingExtension>(Node))
 		{
-			FMovieSceneSpawnable* Spawnable = MovieScene->FindSpawnable(StaticCastSharedRef<FSequencerObjectBindingNode>(Node)->GetObjectBinding());
+			FMovieSceneSpawnable* Spawnable = MovieScene->FindSpawnable(ObjectBindingNode->GetObjectGuid());
 			if (Spawnable)
 			{
 				Spawnable->SetLevelName(InLevelName);
@@ -9840,8 +9956,10 @@
 	}
 }
 
-void FSequencer::ConvertToSpawnable(TSharedRef<FSequencerObjectBindingNode> NodeToBeConverted)
+void FSequencer::ConvertToSpawnable(TSharedRef<UE::Sequencer::FObjectBindingModel> NodeToBeConverted)
 {
+	using namespace UE::Sequencer;
+
 	if (GetFocusedMovieSceneSequence()->GetMovieScene()->IsReadOnly())
 	{
 		FSequencerUtilities::ShowReadOnlyError();
@@ -9854,7 +9972,7 @@
 	TGuardValue<bool> Guard(bUpdatingExternalSelection, true);
 	RestorePreAnimatedState();
 	GetFocusedMovieSceneSequence()->GetMovieScene()->Modify();
-	FMovieScenePossessable* Possessable = GetFocusedMovieSceneSequence()->GetMovieScene()->FindPossessable(NodeToBeConverted->GetObjectBinding());
+	FMovieScenePossessable* Possessable = GetFocusedMovieSceneSequence()->GetMovieScene()->FindPossessable(NodeToBeConverted->GetObjectGuid());
 	if (Possessable)
 	{
 		ConvertToSpawnableInternal(Possessable->GetGuid());
@@ -9879,6 +9997,8 @@
 
 void FSequencer::ConvertSelectedNodesToSpawnables()
 {
+	using namespace UE::Sequencer;
+
 	UMovieScene* MovieScene = GetFocusedMovieSceneSequence()->GetMovieScene();
 
 	if (MovieScene->IsReadOnly())
@@ -9895,16 +10015,14 @@
 	RestorePreAnimatedState();
 	MovieScene->Modify();
 
-	TArray<TSharedRef<FSequencerObjectBindingNode>> ObjectBindingNodes;
+	TArray<IObjectBindingExtension*> ObjectBindingNodes;
 
-	for (const TSharedRef<FSequencerDisplayNode>& Node : Selection.GetSelectedOutlinerNodes())
+	for (const TWeakPtr<FViewModel>& Node : Selection.GetSelectedOutlinerItems())
 	{
-		if (Node->GetType() == ESequencerNode::Object)
+		if (IObjectBindingExtension* ObjectBindingNode = ICastable::CastWeakPtr<IObjectBindingExtension>(Node))
 		{
-			auto ObjectBindingNode = StaticCastSharedRef<FSequencerObjectBindingNode>(Node);
-
 			// If we have a possessable for this node, and it has no parent, we can convert it to a spawnable
-			FMovieScenePossessable* Possessable = MovieScene->FindPossessable(ObjectBindingNode->GetObjectBinding());
+			FMovieScenePossessable* Possessable = MovieScene->FindPossessable(ObjectBindingNode->GetObjectGuid());
 			if (Possessable && !Possessable->GetParent().IsValid())
 			{
 				ObjectBindingNodes.Add(ObjectBindingNode);
@@ -9916,11 +10034,11 @@
 	SlowTask.MakeDialog(true);
 
 	TArray<AActor*> SpawnedActors;
-	for (const TSharedRef<FSequencerObjectBindingNode>& ObjectBindingNode : ObjectBindingNodes)
+	for (IObjectBindingExtension* ObjectBindingNode : ObjectBindingNodes)
 	{
 		SlowTask.EnterProgressFrame();
 	
-		FMovieScenePossessable* Possessable = MovieScene->FindPossessable(ObjectBindingNode->GetObjectBinding());
+		FMovieScenePossessable* Possessable = MovieScene->FindPossessable(ObjectBindingNode->GetObjectGuid());
 		if (Possessable)
 		{
 			TArray<FMovieSceneSpawnable*> Spawnables = ConvertToSpawnableInternal(Possessable->GetGuid());
@@ -9966,6 +10084,8 @@
 
 TArray<FGuid> FSequencer::ExpandMultiplePossessableBindings(FGuid PossessableGuid)
 {
+	using namespace UE::Sequencer;
+
 	UMovieSceneSequence* Sequence = GetFocusedMovieSceneSequence();
 	UMovieScene* MovieScene = Sequence->GetMovieScene();
 
@@ -10008,7 +10128,7 @@
 		}
 	}
 
-	TArray<UMovieSceneTrack* > Tracks = PossessableBinding->StealTracks();
+	TArray<UMovieSceneTrack* > Tracks = PossessableBinding->StealTracks(MovieScene);
 
 	// Remove binding to stop any children from claiming the old guid as their parent
 	if (MovieScene->RemovePossessable(PossessableGuid))
@@ -10046,7 +10166,7 @@
 			if (ParentObject)
 			{
 				FGuid ParentGuid = FindObjectId(*ParentObject, ActiveTemplateIDs.Top());
-				NewPossessable->SetParent(ParentGuid);
+				NewPossessable->SetParent(ParentGuid, MovieScene);
 			}
 
 			if (!NewPossessable->BindSpawnableObject(GetFocusedTemplateID(), FoundObject, this))
@@ -10060,7 +10180,7 @@
 			for (UMovieSceneTrack* Track : Tracks)
 			{
 				UMovieSceneTrack* DuplicatedTrack = Cast<UMovieSceneTrack>(StaticDuplicateObject(Track, MovieScene));
-				NewPossessableBinding->AddTrack(*DuplicatedTrack);
+				NewPossessableBinding->AddTrack(*DuplicatedTrack, MovieScene);
 			}
 		}
 	}
@@ -10078,6 +10198,8 @@
 
 TArray<FMovieSceneSpawnable*> FSequencer::ConvertToSpawnableInternal(FGuid PossessableGuid)
 {
+	using namespace UE::Sequencer;
+
 	UMovieSceneSequence* Sequence = GetFocusedMovieSceneSequence();
 	UMovieScene* MovieScene = Sequence->GetMovieScene();
 
@@ -10178,8 +10300,10 @@
 	return CreatedSpawnables;
 }
 
-void FSequencer::ConvertToPossessable(TSharedRef<FSequencerObjectBindingNode> NodeToBeConverted)
+void FSequencer::ConvertToPossessable(TSharedRef<UE::Sequencer::FObjectBindingModel> NodeToBeConverted)
 {
+	using namespace UE::Sequencer;
+
 	if (GetFocusedMovieSceneSequence()->GetMovieScene()->IsReadOnly())
 	{
 		FSequencerUtilities::ShowReadOnlyError();
@@ -10192,7 +10316,7 @@
 	TGuardValue<bool> Guard(bUpdatingExternalSelection, true);
 	RestorePreAnimatedState();
 	GetFocusedMovieSceneSequence()->GetMovieScene()->Modify();
-	FMovieSceneSpawnable* Spawnable = GetFocusedMovieSceneSequence()->GetMovieScene()->FindSpawnable(NodeToBeConverted->GetObjectBinding());
+	FMovieSceneSpawnable* Spawnable = GetFocusedMovieSceneSequence()->GetMovieScene()->FindSpawnable(NodeToBeConverted->GetObjectGuid());
 	if (Spawnable)
 	{
 		ConvertToPossessableInternal(Spawnable->GetGuid());
@@ -10202,6 +10326,8 @@
 
 void FSequencer::ConvertSelectedNodesToPossessables()
 {
+	using namespace UE::Sequencer;
+
 	UMovieScene* MovieScene = GetFocusedMovieSceneSequence()->GetMovieScene();
 
 	if (MovieScene->IsReadOnly())
@@ -10210,15 +10336,13 @@
 		return;
 	}
 
-	TArray<TSharedRef<FSequencerObjectBindingNode>> ObjectBindingNodes;
+	TArray<IObjectBindingExtension*> ObjectBindingNodes;
 
-	for (const TSharedRef<FSequencerDisplayNode>& Node : Selection.GetSelectedOutlinerNodes())
+	for (const TWeakPtr<FViewModel>& Node : Selection.GetSelectedOutlinerItems())
 	{
-		if (Node->GetType() == ESequencerNode::Object)
+		if (IObjectBindingExtension* ObjectBindingNode = ICastable::CastWeakPtr<IObjectBindingExtension>(Node))
 		{
-			auto ObjectBindingNode = StaticCastSharedRef<FSequencerObjectBindingNode>(Node);
-
-			FMovieSceneSpawnable* Spawnable = MovieScene->FindSpawnable(ObjectBindingNode->GetObjectBinding());
+			FMovieSceneSpawnable* Spawnable = MovieScene->FindSpawnable(ObjectBindingNode->GetObjectGuid());
 			if (Spawnable && SpawnRegister->CanConvertSpawnableToPossessable(*Spawnable))
 			{
 				ObjectBindingNodes.Add(ObjectBindingNode);
@@ -10235,11 +10359,11 @@
 		SlowTask.MakeDialog(true);
 
 		TArray<AActor*> PossessedActors;
-		for (const TSharedRef<FSequencerObjectBindingNode>& ObjectBindingNode : ObjectBindingNodes)
+		for (IObjectBindingExtension* ObjectBindingNode : ObjectBindingNodes)
 		{
 			SlowTask.EnterProgressFrame();
 
-			FMovieSceneSpawnable* Spawnable = MovieScene->FindSpawnable(ObjectBindingNode->GetObjectBinding());
+			FMovieSceneSpawnable* Spawnable = MovieScene->FindSpawnable(ObjectBindingNode->GetObjectGuid());
 			if (Spawnable)
 			{
 				FMovieScenePossessable* Possessable = ConvertToPossessableInternal(Spawnable->GetGuid());
@@ -10285,6 +10409,8 @@
 
 FMovieScenePossessable* FSequencer::ConvertToPossessableInternal(FGuid SpawnableGuid)
 {
+	using namespace UE::Sequencer;
+
 	UMovieSceneSequence* Sequence = GetFocusedMovieSceneSequence();
 	UMovieScene* MovieScene = Sequence->GetMovieScene();
 
@@ -10530,6 +10656,8 @@
 
 void FSequencer::OnAddFolder()
 {
+	using namespace UE::Sequencer;
+
 	UMovieScene* FocusedMovieScene = GetFocusedMovieSceneSequence()->GetMovieScene();
 	if (!FocusedMovieScene)
 	{
@@ -10582,7 +10710,7 @@
 	else
 	{
 		FocusedMovieScene->Modify();
-		FocusedMovieScene->GetRootFolders().Add( NewFolder );
+		FocusedMovieScene->AddRootFolder( NewFolder );
 	}
 
 	Selection.Empty();
@@ -10638,29 +10766,31 @@
 
 void FSequencer::CalculateSelectedFolderAndPath(TArray<UMovieSceneFolder*>& OutSelectedParentFolders, FString& OutNewNodePath)
 {
+	using namespace UE::Sequencer;
+
 	// Check if a folder, or child of a folder is currently selected.
-	if (Selection.GetSelectedOutlinerNodes().Num() > 0)
+	if (Selection.GetSelectedOutlinerItems().Num() > 0)
 	{
-		for (TSharedRef<FSequencerDisplayNode> SelectedNode : Selection.GetSelectedOutlinerNodes())
+		for (TWeakPtr<FViewModel> SelectedNode : Selection.GetSelectedOutlinerItems())
 		{
-			TSharedPtr<FSequencerDisplayNode> CurrentNode = SelectedNode;
-			while (CurrentNode.IsValid() && CurrentNode->GetType() != ESequencerNode::Folder)
+			TSharedPtr<FViewModel> CurrentNode = SelectedNode.Pin();
+			while (CurrentNode.IsValid() && !CurrentNode->IsA<FFolderModel>())
 			{
 				CurrentNode = CurrentNode->GetParent();
 			}
 			if (CurrentNode.IsValid())
 			{
-				OutSelectedParentFolders.Add(&StaticCastSharedPtr<FSequencerFolderNode>(CurrentNode)->GetFolder());
+				OutSelectedParentFolders.Add(CurrentNode->CastThisChecked<FFolderModel>()->GetFolder());
 
 				// The first valid folder we find will be used to put the new folder into, so it's the node that we
 				// want to know the path from.
 				if (OutNewNodePath.Len() == 0)
 				{
 					// Add an extra delimiter (".") as we know that the new folder will be appended onto the end of this.
-					OutNewNodePath = FString::Printf(TEXT("%s."), *CurrentNode->GetPathName());
+					OutNewNodePath = FString::Printf(TEXT("%s."), *IOutlinerExtension::GetPathName(CurrentNode));
 
 					// Make sure this folder is expanded too so that adding objects to hidden folders become visible.
-					CurrentNode->SetExpansionState(true);
+					CurrentNode->CastThisChecked<IOutlinerExtension>()->SetExpansion(true);
 				}
 			}
 		}
@@ -11081,22 +11211,25 @@
 void FSequencer::SortAllNodesAndDescendants()
 {
 	FScopedTransaction SortAllNodesTransaction(NSLOCTEXT("Sequencer", "SortAllNodes_Transaction", "Sort Tracks"));
-	SequencerWidget->GetTreeView()->GetNodeTree()->SortAllNodesAndDescendants();
+	NodeTree->SortAllNodesAndDescendants();
 }
 
 void FSequencer::ToggleExpandCollapseNodes()
 {
+	using namespace UE::Sequencer;
 	SequencerWidget->GetTreeView()->ToggleExpandCollapseNodes(ETreeRecursion::NonRecursive);
 }
 
 
 void FSequencer::ToggleExpandCollapseNodesAndDescendants()
 {
+	using namespace UE::Sequencer;
 	SequencerWidget->GetTreeView()->ToggleExpandCollapseNodes(ETreeRecursion::Recursive);
 }
 
 void FSequencer::ExpandAllNodes()
 {
+	using namespace UE::Sequencer;
 	const bool bExpandAll = true;
 	const bool bCollapseAll = false;
 	SequencerWidget->GetTreeView()->ToggleExpandCollapseNodes(ETreeRecursion::Recursive, bExpandAll, bCollapseAll);
@@ -11104,6 +11237,7 @@
 
 void FSequencer::CollapseAllNodes()
 {
+	using namespace UE::Sequencer;
 	const bool bExpandAll = false;
 	const bool bCollapseAll = true;
 	SequencerWidget->GetTreeView()->ToggleExpandCollapseNodes(ETreeRecursion::Recursive, bExpandAll, bCollapseAll);
@@ -11132,7 +11266,7 @@
 
 void FSequencer::SetKey()
 {
-	if (Selection.GetSelectedOutlinerNodes().Num() > 0)
+	if (Selection.GetSelectedOutlinerItems().Num() > 0)
 	{
 		using namespace UE::Sequencer;
 
@@ -11140,7 +11274,7 @@
 
 		const FFrameNumber KeyTime = GetLocalTime().Time.FrameNumber;
 
-		FAddKeyOperation::FromNodes(Selection.GetSelectedOutlinerNodes()).Commit(KeyTime, *this);
+		FAddKeyOperation::FromNodes(Selection.GetSelectedOutlinerItems()).Commit(KeyTime, *this);
 	}
 }
 
@@ -11160,7 +11294,7 @@
 	{
 		if (Key.IsValid())
 		{
-			KeyTime = Key.KeyArea->GetKeyTime(Key.KeyHandle.GetValue());
+			KeyTime = Key.WeakChannel.Pin()->GetKeyArea()->GetKeyTime(Key.KeyHandle);
 			break;
 		}
 	}
@@ -11194,7 +11328,7 @@
 			{
 	 			if (Key.Section->TryModify())
 	 			{
-	 				Key.KeyArea->SetKeyTime(Key.KeyHandle.GetValue(), NewFrame);
+	 				Key.WeakChannel.Pin()->GetKeyArea()->SetKeyTime(Key.KeyHandle, NewFrame);
 	 				bAnythingChanged = true;
 
 					Key.Section->ExpandToFrame(NewFrame);
@@ -11230,7 +11364,7 @@
 		{
 	 		if (Key.Section->TryModify())
 	 		{
-	 			Key.KeyArea->SetKeyTime(Key.KeyHandle.GetValue(), CurrentTime.Time.FrameNumber);
+	 			Key.WeakChannel.Pin()->GetKeyArea()->SetKeyTime(Key.KeyHandle, CurrentTime.Time.FrameNumber);
 	 			bAnythingChanged = true;
 
 				Key.Section->ExpandToFrame(CurrentTime.Time.FrameNumber);
@@ -11519,41 +11653,39 @@
 	}
 }
 
-void GatherTracksAndObjectsToCopy(TSharedRef<FSequencerDisplayNode> Node, TArray<TSharedPtr<FSequencerTrackNode>>& TracksToCopy, TArray<TSharedPtr<FSequencerObjectBindingNode>>& ObjectsToCopy, TArray<UMovieSceneFolder*>& FoldersToCopy)
+void GatherTracksAndObjectsToCopy(TSharedRef<UE::Sequencer::FViewModel> Node, TArray<TSharedPtr<UE::Sequencer::FViewModel>>& TracksToCopy, TArray<TSharedPtr<UE::Sequencer::FObjectBindingModel>>& ObjectsToCopy, TArray<UMovieSceneFolder*>& FoldersToCopy)
 {
-	if (Node->GetType() == ESequencerNode::Track)
+	using namespace UE::Sequencer;
+
+	if (ITrackExtension* TrackNode = Node->CastThis<ITrackExtension>())
 	{
-		TSharedPtr<FSequencerTrackNode> TrackNode = StaticCastSharedRef<FSequencerTrackNode>(Node);
-		if (TrackNode.IsValid() && !TracksToCopy.Contains(TrackNode))
+		if (!TracksToCopy.Contains(Node))
 		{
-			TracksToCopy.Add(TrackNode);
+			TracksToCopy.Add(Node);
 		}
 	}
-	else if (Node->GetType() == ESequencerNode::Object)
+	else if (TSharedPtr<FObjectBindingModel> ObjectNode = Node->CastThisShared<FObjectBindingModel>())
 	{
-		TSharedPtr<FSequencerObjectBindingNode> ObjectNode = StaticCastSharedRef<FSequencerObjectBindingNode>(Node);
-		if (ObjectNode.IsValid() && !ObjectsToCopy.Contains(ObjectNode))
+		if (!ObjectsToCopy.Contains(ObjectNode))
 		{
 			ObjectsToCopy.Add(ObjectNode);
 		}
 	}
-	else if (Node->GetType() == ESequencerNode::Folder)
+	else if (FFolderModel* FolderNode = Node->CastThis<FFolderModel>())
 	{
-		TSharedPtr<FSequencerFolderNode> FolderNode = StaticCastSharedRef<FSequencerFolderNode>(Node);
-		if (FolderNode.IsValid())
+		FoldersToCopy.Add(FolderNode->GetFolder());
+
+		for (TSharedPtr<FViewModel> ChildNode : FolderNode->GetChildren())
 		{
-			FoldersToCopy.Add(&FolderNode->GetFolder());
-
-			for (TSharedRef<FSequencerDisplayNode> ChildNode : FolderNode->GetChildNodes())
-			{
-				GatherTracksAndObjectsToCopy(ChildNode, TracksToCopy, ObjectsToCopy, FoldersToCopy);
-			}
+			GatherTracksAndObjectsToCopy(ChildNode.ToSharedRef(), TracksToCopy, ObjectsToCopy, FoldersToCopy);
 		}
 	}
 }
 
 void FSequencer::CopySelection()
 {
+	using namespace UE::Sequencer;
+
 	if (Selection.GetSelectedKeys().Num() != 0)
 	{
 		CopySelectedKeys();
@@ -11564,15 +11696,16 @@
 	}
 	else
 	{
-		TArray<TSharedPtr<FSequencerTrackNode>> TracksToCopy;
-		TArray<TSharedPtr<FSequencerObjectBindingNode>> ObjectsToCopy;
+		TArray<TSharedPtr<FViewModel>> TracksToCopy;
+		TArray<TSharedPtr<FObjectBindingModel>> ObjectsToCopy;
 		TArray<UMovieSceneFolder*> FoldersToCopy;
-		TSet<TSharedRef<FSequencerDisplayNode>> SelectedNodes = Selection.GetNodesWithSelectedKeysOrSections();
+		TArray<TSharedRef<FViewModel>> SelectedNodes;
+		Selection.GetNodesWithSelectedKeysOrSections(SelectedNodes);
 		if (SelectedNodes.Num() == 0)
 		{
-			SelectedNodes = Selection.GetSelectedOutlinerNodes();
+			Selection.GetSelectedOutlinerItems(SelectedNodes);
 		}
-		for (TSharedRef<FSequencerDisplayNode> Node : SelectedNodes)
+		for (TSharedRef<FViewModel> Node : SelectedNodes)
 		{
 			GatherTracksAndObjectsToCopy(Node, TracksToCopy, ObjectsToCopy, FoldersToCopy);
 		}
@@ -11651,13 +11784,13 @@
 		{
 			if (Key.IsValid())
 			{
-				TSharedPtr<IKeyArea> KeyArea = Key.KeyArea;
-				FKeyHandle KeyHandle = Key.KeyHandle.GetValue();
+				TSharedPtr<IKeyArea> KeyArea = Key.WeakChannel.Pin()->GetKeyArea();
+				FKeyHandle KeyHandle = Key.KeyHandle;
 
 				FKeyHandle NewKeyHandle = KeyArea->DuplicateKey(KeyHandle);
 				KeyArea->SetKeyTime(NewKeyHandle, KeyArea->GetKeyTime(KeyHandle) + FrameOffset);
 
-				NewSelection.Add(FSequencerSelectedKey(*KeyArea->GetOwningSection(), KeyArea, NewKeyHandle));
+				NewSelection.Add(FSequencerSelectedKey(*KeyArea->GetOwningSection(), Key.WeakChannel, NewKeyHandle));
 			}
 		}
 
@@ -11689,30 +11822,32 @@
 
 void FSequencer::CopySelectedKeys()
 {
+	using namespace UE::Sequencer;
+
 	TOptional<FFrameNumber> CopyRelativeTo;
-	
+
 	// Copy relative to the current key hotspot, if applicable
-	if (Hotspot.IsValid() && Hotspot->GetType() == ESequencerHotspot::Key)
+	if (TSharedPtr<FKeyHotspot> KeyHotspot = HotspotCast<FKeyHotspot>(GetViewModel()->GetTrackArea()->GetHotspot()))
 	{
-		CopyRelativeTo = StaticCastSharedPtr<FKeyHotspot>(Hotspot)->GetTime();
+		CopyRelativeTo = KeyHotspot->GetTime();
 	}
 
 	FMovieSceneClipboardBuilder Builder;
 
 	// Map selected keys to their key areas
-	TMap<TSharedPtr<IKeyArea>, TArray<FKeyHandle>> KeyAreaMap;
+	TMap<TSharedPtr<FChannelModel>, TArray<FKeyHandle>> KeyAreaMap;
 	for (const FSequencerSelectedKey& Key : Selection.GetSelectedKeys())
 	{
-		if (Key.KeyHandle.IsSet())
+		if (Key.IsValid())
 		{
-			KeyAreaMap.FindOrAdd(Key.KeyArea).Add(Key.KeyHandle.GetValue());
+			KeyAreaMap.FindOrAdd(Key.WeakChannel.Pin()).Add(Key.KeyHandle);
 		}
 	}
 
 	// Serialize each key area to the clipboard
-	for (auto& Pair : KeyAreaMap)
+	for (const TPair<TSharedPtr<FChannelModel>, TArray<FKeyHandle>>& Pair : KeyAreaMap)
 	{
-		Pair.Key->CopyKeys(Builder, Pair.Value);
+		Pair.Key->GetKeyArea()->CopyKeys(Builder, Pair.Value);
 	}
 
 	TSharedRef<FMovieSceneClipboard> Clipboard = MakeShareable( new FMovieSceneClipboard(Builder.Commit(CopyRelativeTo)) );
@@ -11999,16 +12134,20 @@
 
 void FSequencer::ImportFBX()
 {
+	using namespace UE::Sequencer;
+
 	TMap<FGuid, FString> ObjectBindingNameMap;
 
-	TArray<TSharedRef<FSequencerObjectBindingNode>> RootObjectBindingNodes;
-	GetRootObjectBindingNodes(NodeTree->GetRootNodes(), RootObjectBindingNodes);
+	TParentFirstChildIterator<IObjectBindingExtension> ObjectBindingIt = NodeTree->GetRootNode()->GetDescendantsOfType<IObjectBindingExtension>();
 
-	for (auto RootObjectBindingNode : RootObjectBindingNodes)
+	// Only visit the first object binding in a given sub-hierarchy so we get top-level object bindings.
+	// This is done by skipping the branch on each iteration
+	for ( ; ObjectBindingIt; ++ObjectBindingIt)
 	{
-		FGuid ObjectBinding = RootObjectBindingNode.Get().GetObjectBinding();
+		FGuid ObjectBinding = ObjectBindingIt->GetObjectGuid();
+		ObjectBindingNameMap.Add(ObjectBinding, (*ObjectBindingIt).AsModel()->CastThisChecked<IOutlinerExtension>()->GetLabel().ToString());
 
-		ObjectBindingNameMap.Add(ObjectBinding, RootObjectBindingNode.Get().GetDisplayName().ToString());
+		ObjectBindingIt.IgnoreCurrentChildren();
 	}
 
 	MovieSceneToolHelpers::ImportFBXWithDialog(GetFocusedMovieSceneSequence(), *this, ObjectBindingNameMap, TOptional<bool>());
@@ -12016,18 +12155,18 @@
 
 void FSequencer::ImportFBXOntoSelectedNodes()
 {
+	using namespace UE::Sequencer;
+
 	// The object binding and names to match when importing from fbx
 	TMap<FGuid, FString> ObjectBindingNameMap;
 
-	for (const TSharedRef<FSequencerDisplayNode>& Node : Selection.GetSelectedOutlinerNodes())
+	for (const TWeakPtr<FViewModel>& Node : Selection.GetSelectedOutlinerItems())
 	{
-		if (Node->GetType() == ESequencerNode::Object)
+		if (FObjectBindingModel* ObjectBindingNode = ICastable::CastWeakPtr<FObjectBindingModel>(Node))
 		{
-			auto ObjectBindingNode = StaticCastSharedRef<FSequencerObjectBindingNode>(Node);
+			FGuid ObjectBinding = ObjectBindingNode->GetObjectGuid();
 
-			FGuid ObjectBinding = ObjectBindingNode.Get().GetObjectBinding();
-
-			ObjectBindingNameMap.Add(ObjectBinding, ObjectBindingNode.Get().GetDisplayName().ToString());
+			ObjectBindingNameMap.Add(ObjectBinding, ObjectBindingNode->GetLabel().ToString());
 		}
 	}
 
@@ -12036,6 +12175,8 @@
 
 void FSequencer::ExportFBX()
 {
+	using namespace UE::Sequencer;
+
 	TArray<UExporter*> Exporters;
 	TArray<FString> SaveFilenames;
 	IDesktopPlatform* DesktopPlatform = FDesktopPlatformModule::Get();
@@ -12096,27 +12237,25 @@
 		TArray<FGuid> Bindings;
 		TArray<UMovieSceneTrack*> MasterTracks;
 		UMovieScene* MovieScene = GetFocusedMovieSceneSequence()->GetMovieScene();
-		for (const TSharedRef<FSequencerDisplayNode>& Node : Selection.GetSelectedOutlinerNodes())
+		for (const TWeakPtr<FViewModel>& Node : Selection.GetSelectedOutlinerItems())
 		{
-			if (Node->GetType() == ESequencerNode::Object)
+			if (IObjectBindingExtension* ObjectBindingNode = ICastable::CastWeakPtr<IObjectBindingExtension>(Node))
 			{
-				auto ObjectBindingNode = StaticCastSharedRef<FSequencerObjectBindingNode>(Node);
-				Bindings.Add(ObjectBindingNode.Get().GetObjectBinding());
+				Bindings.Add(ObjectBindingNode->GetObjectGuid());
 
-				TSet<TSharedRef<FSequencerDisplayNode> > DescendantNodes;
-				SequencerHelpers::GetDescendantNodes(Node, DescendantNodes);
-				for (auto DescendantNode : DescendantNodes)
+				TSet<TSharedRef<FViewModel> > DescendantNodes;
+				SequencerHelpers::GetDescendantNodes(Node.Pin().ToSharedRef(), DescendantNodes);
+				for (TSharedRef<FViewModel> DescendantNode : DescendantNodes)
 				{
-					if (!Selection.IsSelected(DescendantNode) && DescendantNode->GetType() == ESequencerNode::Object)
+					if (!Selection.IsSelected(DescendantNode) && DescendantNode->IsA<IObjectBindingExtension>())
 					{
-						auto DescendantObjectBindingNode = StaticCastSharedRef<FSequencerObjectBindingNode>(DescendantNode);
-						Bindings.Add(DescendantObjectBindingNode.Get().GetObjectBinding());
+						IObjectBindingExtension* DescendantObjectBindingNode = DescendantNode->CastThisChecked<IObjectBindingExtension>();
+						Bindings.Add(DescendantObjectBindingNode->GetObjectGuid());
 					}
 				}
 			}
-			else if (Node->GetType() == ESequencerNode::Track)
+			else if (ITrackExtension* TrackNode = ICastable::CastWeakPtr<ITrackExtension>(Node))
 			{
-				TSharedRef<const FSequencerTrackNode> TrackNode = StaticCastSharedRef<const FSequencerTrackNode>(Node);
 				UMovieSceneTrack* Track = TrackNode->GetTrack();
 				if (Track && MovieScene->IsAMasterTrack(*Track))
 				{
@@ -12212,17 +12351,16 @@
 
 void FSequencer::ExportToCameraAnim()
 {
-	for (const TSharedRef<FSequencerDisplayNode>& Node : Selection.GetSelectedOutlinerNodes())
+	using namespace UE::Sequencer;
+
+	for (const TWeakPtr<FViewModel>& Node : Selection.GetSelectedOutlinerItems())
 	{
-		if (Node->GetType() != ESequencerNode::Object)
+		if (FObjectBindingModel* ObjectBindingNode = ICastable::CastWeakPtr<FObjectBindingModel>(Node))
 		{
-			continue;
+			FGuid Guid = ObjectBindingNode->GetObjectGuid();
+			
+			MovieSceneToolHelpers::ExportToCameraAnim(GetFocusedMovieSceneSequence()->GetMovieScene(), Guid);
 		}
-		auto ObjectBindingNode = StaticCastSharedRef<FSequencerObjectBindingNode>(Node);
-
-		FGuid Guid = ObjectBindingNode->GetObjectBinding();
-		
-		MovieSceneToolHelpers::ExportToCameraAnim(GetFocusedMovieSceneSequence()->GetMovieScene(), Guid);
 	}
 }
 
@@ -12267,6 +12405,8 @@
 
 void FSequencer::TrimOrExtendSection(bool bTrimOrExtendLeft)
 {
+	using namespace UE::Sequencer;
+
 	UMovieSceneSequence* FocusedMovieSceneSequence = GetFocusedMovieSceneSequence();
 	UMovieScene* MovieScene = FocusedMovieSceneSequence ? FocusedMovieSceneSequence->GetMovieScene() : nullptr;
 	if (!MovieScene)
@@ -12276,25 +12416,24 @@
 
 	FScopedTransaction TrimOrExtendSectionTransaction( NSLOCTEXT("Sequencer", "TrimOrExtendSection_Transaction", "Trim or Extend Section") );
 
-	if (Selection.GetSelectedOutlinerNodes().Num() > 0)
+	if (Selection.GetSelectedOutlinerItems().Num() > 0)
 	{
-		const TSet<TSharedRef<FSequencerDisplayNode> >& SelectedNodes = GetSelection().GetSelectedOutlinerNodes();
+		TArray<TSharedRef<FViewModel>> SelectedNodes;
+		GetSelection().GetSelectedOutlinerItems(SelectedNodes);
 	
-		for (const TSharedRef<const FSequencerDisplayNode> Node : SelectedNodes)
+		for (const TSharedRef<FViewModel>& Node : SelectedNodes)
 		{
-			if (Node->GetType() == ESequencerNode::Track  )
+			if (ITrackExtension* TrackNode = Node->CastThis<ITrackExtension>())
 			{
-				TSharedRef<const FSequencerTrackNode> TrackNode = StaticCastSharedRef<const FSequencerTrackNode>( Node );
 				UMovieSceneTrack* Track = TrackNode->GetTrack();
 				if (Track)
 				{
-					MovieSceneToolHelpers::TrimOrExtendSection(Track, TrackNode->GetSubTrackMode() == FSequencerTrackNode::ESubTrackMode::SubTrack ? TrackNode->GetRowIndex() : TOptional<int32>(), GetLocalTime(), bTrimOrExtendLeft, Settings->GetDeleteKeysWhenTrimming());
+					MovieSceneToolHelpers::TrimOrExtendSection(Track, Node->IsA<FTrackRowModel>() ? TrackNode->GetRowIndex() : TOptional<int32>(), GetLocalTime(), bTrimOrExtendLeft, Settings->GetDeleteKeysWhenTrimming());
 				}
 			}
-			else if (Node->GetType() == ESequencerNode::Object)
+			else if (IObjectBindingExtension* ObjectBindingNode = Node->CastThis<IObjectBindingExtension>())
 			{
-				TSharedRef<const FSequencerObjectBindingNode> ObjectBindingNode = StaticCastSharedRef<const FSequencerObjectBindingNode>( Node );
-				const FMovieSceneBinding* Binding = MovieScene->FindBinding(ObjectBindingNode->GetObjectBinding());
+				const FMovieSceneBinding* Binding = MovieScene->FindBinding(ObjectBindingNode->GetObjectGuid());
 				if (Binding)
 				{
 					for (UMovieSceneTrack* Track : Binding->GetTracks())
@@ -12331,32 +12470,10 @@
 	NotifyMovieSceneDataChanged( EMovieSceneDataChangeType::RefreshAllImmediately );
 }
 
-const ISequencerEditTool* FSequencer::GetEditTool() const
-{
-	return SequencerWidget->GetEditTool();
-}
-
-TSharedPtr<ISequencerHotspot> FSequencer::GetHotspot() const
-{
-	return Hotspot;
-}
-
-void FSequencer::SetHotspot(TSharedPtr<ISequencerHotspot> NewHotspot)
-{
-	if (!Hotspot.IsValid() || !Hotspot->bIsLocked)
-	{
-		Hotspot = MoveTemp(NewHotspot);
-	}
-
-	// Simulate an update-on-hover for the new hotspot to ensure that any hover behavior doesn't have to wait until the next frame
-	if (Hotspot)
-	{
-		Hotspot->UpdateOnHover(*SequencerWidget->GetTrackAreaWidget(), *this);
-	}
-}
-
 void FSequencer::BindCommands()
 {
+	using namespace UE::Sequencer;
+
 	const FSequencerCommands& Commands = FSequencerCommands::Get();
 
 	SequencerCommandBindings->MapAction(
@@ -12474,11 +12591,11 @@
 
 		UMovieScene* MovieScene = Sequence->GetMovieScene();
 
-		for (const TSharedRef<FSequencerDisplayNode>& Node : Selection.GetSelectedOutlinerNodes())
+		for (const TWeakPtr<FViewModel>& Node : Selection.GetSelectedOutlinerItems())
 		{
-			if (Node->GetType() == ESequencerNode::Object)
+			if (IObjectBindingExtension* ObjectBindingNode = ICastable::CastWeakPtr<IObjectBindingExtension>(Node))
 			{
-				FMovieScenePossessable* Possessable = MovieScene->FindPossessable(static_cast<FSequencerObjectBindingNode&>(*Node).GetObjectBinding());
+				FMovieScenePossessable* Possessable = MovieScene->FindPossessable(ObjectBindingNode->GetObjectGuid());
 				if (Possessable && !Possessable->GetParent().IsValid())
 				{
 					return true;
@@ -12496,11 +12613,11 @@
 	auto AreConvertableSpawnablesSelected = [this] {
 		UMovieScene* MovieScene = GetFocusedMovieSceneSequence()->GetMovieScene();
 
-		for (const TSharedRef<FSequencerDisplayNode>& Node : Selection.GetSelectedOutlinerNodes())
+		for (const TWeakPtr<FViewModel>& Node : Selection.GetSelectedOutlinerItems())
 		{
-			if (Node->GetType() == ESequencerNode::Object)
+			if (IObjectBindingExtension* ObjectBindingNode = ICastable::CastWeakPtr<IObjectBindingExtension>(Node))
 			{
-				FMovieSceneSpawnable* Spawnable = MovieScene->FindSpawnable(static_cast<FSequencerObjectBindingNode&>(*Node).GetObjectBinding());
+				FMovieSceneSpawnable* Spawnable = MovieScene->FindSpawnable(ObjectBindingNode->GetObjectGuid());
 				if (Spawnable && SpawnRegister->CanConvertSpawnableToPossessable(*Spawnable))
 				{
 					return true;
@@ -12519,11 +12636,11 @@
 	auto AreSpawnablesSelected = [this] {
 		UMovieScene* MovieScene = GetFocusedMovieSceneSequence()->GetMovieScene();
 
-		for (const TSharedRef<FSequencerDisplayNode>& Node : Selection.GetSelectedOutlinerNodes())
+		for (const TWeakPtr<FViewModel>& Node : Selection.GetSelectedOutlinerItems())
 		{
-			if (Node->GetType() == ESequencerNode::Object)
+			if (IObjectBindingExtension* ObjectBindingNode = ICastable::CastWeakPtr<IObjectBindingExtension>(Node))
 			{
-				FMovieSceneSpawnable* Spawnable = MovieScene->FindSpawnable(static_cast<FSequencerObjectBindingNode&>(*Node).GetObjectBinding());
+				FMovieSceneSpawnable* Spawnable = MovieScene->FindSpawnable(ObjectBindingNode->GetObjectGuid());
 				if (Spawnable)
 				{
 					return true;
@@ -12784,20 +12901,21 @@
 
 	auto CanCutOrCopy = [this]{
 		// For copy tracks
-		TSet<TSharedRef<FSequencerDisplayNode>> SelectedNodes = Selection.GetNodesWithSelectedKeysOrSections();
+		TSet<TWeakPtr<FViewModel>> SelectedNodes = Selection.GetNodesWithSelectedKeysOrSections();
 		// If this is empty then we are selecting display nodes
 		if (SelectedNodes.Num() == 0)
 		{
-			SelectedNodes = Selection.GetSelectedOutlinerNodes();
-			for (TSharedRef<FSequencerDisplayNode> Node : SelectedNodes)
+			SelectedNodes = Selection.GetSelectedOutlinerItems();
+			for (TWeakPtr<FViewModel> WeakNode : SelectedNodes)
 			{
-				if (Node->GetType() == ESequencerNode::Track || Node->GetType() == ESequencerNode::Object || Node->GetType() == ESequencerNode::Folder)
+				TSharedPtr<FViewModel> Node = WeakNode.Pin();
+				if (Node->IsA<ITrackExtension>() || Node->IsA<IObjectBindingExtension>() || Node->IsA<FFolderModel>())
 				{
 					// if contains one node that can be copied we allow the action
 					// later on we will filter out the invalid nodes in CopySelection() or CutSelection()
 					return true;
 				}
-				else if (Node->GetParent().IsValid() && Node->GetParent()->GetType() == ESequencerNode::Track && Node->GetType() != ESequencerNode::Category)
+				else if (Node->GetParent().IsValid() && Node->GetParent()->IsA<ITrackExtension>() && !Node->IsA<FCategoryGroupModel>())
 				{
 					return true;
 				}
@@ -12821,7 +12939,7 @@
 	};
 
 	auto CanDelete = [this]{
-		return Selection.GetSelectedKeys().Num() || Selection.GetSelectedSections().Num() || Selection.GetSelectedOutlinerNodes().Num();
+		return Selection.GetSelectedKeys().Num() || Selection.GetSelectedSections().Num() || Selection.GetSelectedOutlinerItems().Num();
 	};
 
 	auto CanDuplicate = [this]{
@@ -12832,13 +12950,14 @@
 		}
 
 		// For duplicate object tracks
-		TSet<TSharedRef<FSequencerDisplayNode>> SelectedNodes = Selection.GetNodesWithSelectedKeysOrSections();
+		TArray<TSharedRef<FViewModel>> SelectedNodes;
+		Selection.GetNodesWithSelectedKeysOrSections(SelectedNodes);
 		if (SelectedNodes.Num() == 0)
 		{
-			SelectedNodes = Selection.GetSelectedOutlinerNodes();
-			for (TSharedRef<FSequencerDisplayNode> Node : SelectedNodes)
+			Selection.GetSelectedOutlinerItems(SelectedNodes);
+			for (TSharedRef<FViewModel> Node : SelectedNodes)
 			{
-				if (Node->GetType() == ESequencerNode::Object)
+				if (Node->IsA<IObjectBindingExtension>())
 				{
 					// if contains one node that can be copied we allow the action
 					return true;
@@ -12863,11 +12982,26 @@
 		FGenericCommands::Get().Rename,
 		FExecuteAction::CreateLambda([this]
 		{
-			Selection.GetSelectedOutlinerNodes().Array()[0]->OnRenameRequested().Broadcast();
+			for (TWeakPtr<FViewModel> Item : Selection.GetSelectedOutlinerItems())
+			{
+				IRenameableExtension* Renamable = ICastable::CastWeakPtr<IRenameableExtension>(Item);
+				if (Renamable && Renamable->CanRename())
+				{
+					Renamable->OnRenameRequested().Broadcast();
+				}
+			}
 		}),
 		FCanExecuteAction::CreateLambda([this]
 		{
-			return (Selection.GetSelectedOutlinerNodes().Num() > 0) && (Selection.GetSelectedOutlinerNodes().Array()[0]->CanRenameNode());
+			for (TWeakPtr<FViewModel> Item : Selection.GetSelectedOutlinerItems())
+			{
+				IRenameableExtension* Renamable = ICastable::CastWeakPtr<IRenameableExtension>(Item);
+				if (Renamable && Renamable->CanRename())
+				{
+					return true;
+				}
+			}
+			return false;
 		})
 	);
 
@@ -13200,7 +13334,25 @@
 		FExecuteAction::CreateSP( this, &FSequencer::SetPlaybackRangeToAllShots ),
 		FCanExecuteAction::CreateSP( this, &FSequencer::IsViewingMasterSequence ) );
 
+	SequencerCommandBindings->MapAction(
+		Commands.RefreshUI,
+		FExecuteAction::CreateLambda(
+			[this]
+			{
+				this->Selection.Empty();
 
+				UMovieSceneSequence*  FocusedSequence = GetFocusedMovieSceneSequence();
+				FMovieSceneSequenceID SequenceID      = GetFocusedTemplateID();
+
+				TSharedPtr<FSequenceModel> RootSequenceModel = this->ViewModel->GetRootModel().ImplicitCast();
+				RootSequenceModel->SetSequence(nullptr, MovieSceneSequenceID::Root);
+				RootSequenceModel->SetSequence(FocusedSequence, SequenceID);
+
+				RefreshTree();
+			}
+		),
+		FCanExecuteAction());
+
 	// We want a subset of the commands to work in the Curve Editor too, but bound to our functions. This minimizes code duplication
 	// while also freeing us up from issues that result from Sequencer already using two lists (for which our commands might be spread
 	// across both lists which makes a direct copy like it already uses difficult).
@@ -13258,7 +13410,6 @@
 	}
 }
 
-
 void FSequencer::BuildAddObjectBindingsMenu(class FMenuBuilder& MenuBuilder)
 {
 	for (int32 i = 0; i < ObjectBindings.Num(); ++i)
@@ -13289,6 +13440,8 @@
 
 void FSequencer::BuildAddSelectedToFolderMenu(FMenuBuilder& MenuBuilder)
 {
+	using namespace UE::Sequencer;
+
 	MenuBuilder.AddMenuEntry(
 		LOCTEXT("MoveNodesToNewFolder", "New Folder"),
 		LOCTEXT("MoveNodesToNewFolderTooltip", "Create a new folder and adds the selected nodes"),
@@ -13302,15 +13455,21 @@
 	if (MovieScene)
 	{
 		TSharedRef<TArray<UMovieSceneFolder*>> ExcludedFolders = MakeShared<TArray<UMovieSceneFolder*> >();
-		for (TSharedRef<FSequencerDisplayNode> Node : GetSelection().GetSelectedOutlinerNodes())
+		for (TWeakPtr<FViewModel> Node : GetSelection().GetSelectedOutlinerItems())
 		{
-			if (Node->GetType() == ESequencerNode::Folder && Node->CanDrag())
+			if (FFolderModel* FolderNode = ICastable::CastWeakPtr<FFolderModel>(Node))
 			{
-				ExcludedFolders->Add(&StaticCastSharedRef<FSequencerFolderNode>(Node)->GetFolder());
+				if (FolderNode->CanDrag())
+				{
+					ExcludedFolders->Add(FolderNode->GetFolder());
+				}
 			}
 		}
 
-		TArray<UMovieSceneFolder*> ChildFolders = MovieScene->GetRootFolders();
+		// Copy the list of root folders and remove any currently dragged folders. The rest represents the
+		// list of possible folders to move the selection into.
+		TArray<UMovieSceneFolder*> ChildFolders;
+		MovieScene->GetRootFolders(ChildFolders);
 		for (int32 Index = 0; Index < ChildFolders.Num(); ++Index)
 		{
 			if (ExcludedFolders->Contains(ChildFolders[Index]))
@@ -13353,14 +13512,13 @@
 
 void FSequencer::BuildAddSelectedToFolderMenuEntry(FMenuBuilder& InMenuBuilder, TSharedRef<TArray<UMovieSceneFolder*> > InExcludedFolders, UMovieSceneFolder* InFolder)
 {
-	TArray<UMovieSceneFolder*> ChildFolders = InFolder->GetChildFolders();
+	TArray<UMovieSceneFolder*> ChildFolders;
 
-	for (int32 Index = 0; Index < ChildFolders.Num(); ++Index)
+	for (UMovieSceneFolder* Folder : InFolder->GetChildFolders())
 	{
-		if (InExcludedFolders->Contains(ChildFolders[Index]))
+		if (!InExcludedFolders->Contains(Folder))
 		{
-			ChildFolders.RemoveAt(Index);
-			--Index;
+			ChildFolders.Add(Folder);
 		}
 	}
 
@@ -13453,10 +13611,12 @@
 	TimeController->PlayerStatusChanged(PlaybackState, GetGlobalTime());
 }
 
-
 void FSequencer::BuildCustomContextMenuForGuid(FMenuBuilder& MenuBuilder, FGuid ObjectBinding)
 {
-	SequencerWidget->BuildCustomContextMenuForGuid(MenuBuilder, ObjectBinding);
+	using namespace UE::Sequencer;
+
+	FSequencerOutlinerViewModel* OutlinerViewModel = ViewModel->GetOutliner()->CastThisChecked<FSequencerOutlinerViewModel>();
+	OutlinerViewModel->BuildCustomContextMenuForGuid(MenuBuilder, ObjectBinding);
 }
 
 FKeyAttributes FSequencer::GetDefaultKeyAttributes() const
@@ -13542,43 +13702,40 @@
 
 void FSequencer::SetDisplayName(FGuid Binding, const FText& InDisplayName)
 {
-	for (const TSharedRef<FSequencerDisplayNode>& Node : Selection.GetSelectedOutlinerNodes())
+	using namespace UE::Sequencer;
+
+	for (const TWeakPtr<FViewModel>& Node : Selection.GetSelectedOutlinerItems())
 	{
-		if (Node->GetType() != ESequencerNode::Object)
+		if (FObjectBindingModel* ObjectBindingNode = ICastable::CastWeakPtr<FObjectBindingModel>(Node))
 		{
-			continue;
-		}
-
-		auto ObjectBindingNode = StaticCastSharedRef<FSequencerObjectBindingNode>(Node);
-		FGuid Guid = ObjectBindingNode->GetObjectBinding();
-		if (Guid == Binding)
-		{
-			ObjectBindingNode->SetDisplayName(InDisplayName);
-			break;
+			FGuid Guid = ObjectBindingNode->GetObjectGuid();
+			if (Guid == Binding)
+			{
+				ObjectBindingNode->Rename(InDisplayName);
+				break;
+			}
 		}
 	}
 }
 
 FText FSequencer::GetDisplayName(FGuid Binding)
 {
-	for (const TSharedRef<FSequencerDisplayNode>& Node : Selection.GetSelectedOutlinerNodes())
+	using namespace UE::Sequencer;
+
+	for (const TWeakPtr<FViewModel>& Node : Selection.GetSelectedOutlinerItems())
 	{
-		if (Node->GetType() != ESequencerNode::Object)
+		if (FObjectBindingModel* ObjectBindingNode = ICastable::CastWeakPtr<FObjectBindingModel>(Node))
 		{
-			continue;
-		}
-
-		auto ObjectBindingNode = StaticCastSharedRef<FSequencerObjectBindingNode>(Node);
-		FGuid Guid = ObjectBindingNode->GetObjectBinding();
-		if (Guid == Binding)
-		{
-			return ObjectBindingNode->GetDisplayName();
+			FGuid Guid = ObjectBindingNode->GetObjectGuid();
+			if (Guid == Binding)
+			{
+				return ObjectBindingNode->GetLabel();
+			}
 		}
 	}
 	return FText();
 }
 
-
 void FSequencer::OnCurveModelDisplayChanged(FCurveModel *InCurveModel, bool bDisplayed, const FCurveEditor* InCurveEditor)
 {
 	OnCurveDisplayChanged.Broadcast(InCurveModel, bDisplayed, InCurveEditor);
