Sequencer: MVVM2 branch and Layer Bars

diff --git a/Engine/Source/Editor/Sequencer/Private/Sequencer.h b/Engine/Source/Editor/Sequencer/Private/Sequencer.h
--- a/Engine/Source/Editor/Sequencer/Private/Sequencer.h
+++ b/Engine/Source/Editor/Sequencer/Private/Sequencer.h
@@ -11,7 +11,6 @@
 #include "UObject/WeakInterfacePtr.h"
 #include "Widgets/SWidget.h"
 #include "SequencerNodeTree.h"
-#include "DisplayNodes/SequencerDisplayNode.h"
 #include "UObject/GCObject.h"
 #include "MovieSceneSequenceID.h"
 #include "IMovieScenePlayer.h"
@@ -46,8 +45,6 @@
 class FLevelEditorViewportClient;
 class FMenuBuilder;
 class FMovieSceneClipboard;
-class FSequencerObjectBindingNode;
-class FSequencerTrackNode;
 class FViewportClient;
 class IDetailKeyframeHandler;
 class IAssetViewport;
@@ -55,6 +52,7 @@
 class FCurveEditor;
 class ISequencerEditTool;
 class FSequencerKeyCollection;
+class FSequencerOutlinerSelectionHandler;
 class FObjectBindingTagCache;
 class ISequencerTrackEditor;
 class ISequencerEditorObjectBinding;
@@ -70,7 +68,6 @@
 struct FMovieSceneSequencePlaybackParams;
 struct FMovieScenePossessable;
 struct FTransformData;
-struct ISequencerHotspot;
 struct FKeyAttributes;
 struct FNotificationInfo;
 struct FEditorViewportViewModifierParams;
@@ -79,12 +76,22 @@
 
 namespace UE
 {
-namespace MovieScene
-{
+	namespace MovieScene
+	{
 
-	struct FInitialValueCache;
+		struct FInitialValueCache;
 
-} // namespace MovieScene
+	} // namespace MovieScene
+
+	namespace Sequencer
+	{
+
+		struct FViewModelTypeID;
+		class FSequenceModel;
+		class FSequencerEditorViewModel;
+		class FViewModel;
+
+	} // namespace Sequencer
 } // namespace UE
 
 
@@ -97,6 +104,8 @@
 	, public FEditorUndoClient
 	, public FTickableEditorObject
 {
+	using FViewModel = UE::Sequencer::FViewModel;
+
 public:
 
 	/** Constructor */
@@ -205,6 +214,7 @@
 	void ZoomOutViewRange();
 
 public:
+
 	/** Gets the tree of nodes which is used to populate the animation outliner. */
 	TSharedRef<FSequencerNodeTree> GetNodeTree()
 	{
@@ -321,14 +331,14 @@
 	 *
 	 * @param	PossessableGuid		The guid of the possessable to convert
 	 */
-	void ConvertToSpawnable(TSharedRef<FSequencerObjectBindingNode> NodeToBeConverted);
+	void ConvertToSpawnable(TSharedRef<UE::Sequencer::FObjectBindingModel> NodeToBeConverted);
 
 	/**
 	 * Converts the specified spawnable GUID to a possessable
 	 *
 	 * @param	SpawnableGuid		The guid of the spawnable to convert
 	 */
-	void ConvertToPossessable(TSharedRef<FSequencerObjectBindingNode> NodeToBeConverted);
+	void ConvertToPossessable(TSharedRef<UE::Sequencer::FObjectBindingModel> NodeToBeConverted);
 
 	/**
 	 * Converts all the currently selected nodes to be spawnables, if possible
@@ -384,7 +394,7 @@
 	 * @param NodeToBeDeleted	Node with data that should be deleted
 	 * @return true if anything was deleted, otherwise false.
 	 */
-	bool OnRequestNodeDeleted( TSharedRef<const FSequencerDisplayNode> NodeToBeDeleted, const bool bKeepState );
+	bool OnRequestNodeDeleted( TSharedRef<FViewModel> NodeToBeDeleted, const bool bKeepState );
 
 	/** Zooms to the edges of all currently selected sections and keys. */
 	void ZoomToFit();
@@ -524,22 +534,22 @@
 	void RemoveActorsFromBinding(FGuid ObjectBinding, const TArray<AActor*>& InActors);
 
 	/** Called when a user executes the delete node menu item */
-	void DeleteNode(TSharedRef<FSequencerDisplayNode> NodeToBeDeleted, const bool bKeepState);
+	void DeleteNode(TSharedRef<FViewModel> NodeToBeDeleted, const bool bKeepState);
 	void DeleteSelectedNodes(const bool bKeepState);
 
 	/** @return The list of nodes which must be moved to move the current selected nodes */
-	TArray<TSharedRef<FSequencerDisplayNode> > GetSelectedNodesToMove();
-	TArray<TSharedRef<FSequencerDisplayNode> > GetSelectedNodesInFolders();
+	TArray<TSharedRef<FViewModel>> GetSelectedNodesToMove();
+	TArray<TSharedRef<FViewModel>> GetSelectedNodesInFolders();
 
 	/** Called when a user executes the move to new folder menu item */
 	void MoveSelectedNodesToNewFolder();
 	void RemoveSelectedNodesFromFolders();
-	void MoveNodeToFolder(TSharedRef<FSequencerDisplayNode> NodeToMove, UMovieSceneFolder* DestinationFolder);
+	void MoveNodeToFolder(TSharedRef<FViewModel> NodeToMove, UMovieSceneFolder* DestinationFolder);
 	void MoveSelectedNodesToFolder(UMovieSceneFolder* DestinationFolder);
 
 	/** Called when a user executes the copy track menu item */
-	void CopySelectedObjects(TArray<TSharedPtr<FSequencerObjectBindingNode>>& ObjectNodes, const TArray<UMovieSceneFolder*>& Folders, /*out*/ FString& ExportedText);
-	void CopySelectedTracks(TArray<TSharedPtr<FSequencerTrackNode>>& TrackNodes, const TArray<UMovieSceneFolder*>& Folders, /*out*/ FString& ExportedText);
+	void CopySelectedObjects(TArray<TSharedPtr<UE::Sequencer::FObjectBindingModel>>& ObjectNodes, const TArray<UMovieSceneFolder*>& Folders, /*out*/ FString& ExportedText);
+	void CopySelectedTracks(TArray<TSharedPtr<UE::Sequencer::FViewModel>>& TrackNodes, const TArray<UMovieSceneFolder*>& Folders, /*out*/ FString& ExportedText);
 	void CopySelectedFolders(const TArray<UMovieSceneFolder*>& Folders, /*out*/ FString& ExportedText);
 	void ExportObjectsToText(const TArray<UObject*>& ObjectsToExport, /*out*/ FString& ExportedText);
 
@@ -585,7 +595,7 @@
 	bool CanRekey() const;
 	void Rekey();
 
-	void SelectKey(UMovieSceneSection* InSection, TSharedPtr<IKeyArea> InKeyArea, FKeyHandle KeyHandle, bool bToggle);
+	void SelectKey(UMovieSceneSection* InSection, TSharedPtr<UE::Sequencer::FChannelModel> InChannel, FKeyHandle KeyHandle, bool bToggle);
 
 	/** Updates the external selection to match the current sequencer selection. */
 	void SynchronizeExternalSelectionWithSequencerSelection();
@@ -668,32 +678,13 @@
 	void ExportToCameraAnim();
 
 public:
-	
-	/** Access the currently active track area edit tool */
-	const ISequencerEditTool* GetEditTool() const;
 
-	/** Get the current active hotspot */
-	TSharedPtr<ISequencerHotspot> GetHotspot() const;
-
-	/** Set the hotspot to something else */
-	void SetHotspot(TSharedPtr<ISequencerHotspot> NewHotspot);
-
-protected:
-
-	/** The current hotspot that can be set from anywhere to initiate drags */
-	TSharedPtr<ISequencerHotspot> Hotspot;
-
-public:
-
 	/** Put the sequencer in a horizontally auto-scrolling state with the given rate */
 	void StartAutoscroll(float UnitsPerS);
 
 	/** Stop the sequencer from auto-scrolling */
 	void StopAutoscroll();
 
-	/** Scroll the sequencer vertically by the specified number of slate units */
-	void VerticalScroll(float ScrollAmountUnits);
-
 	/**
 	 * Update auto-scroll mechanics as a result of a new time position
 	 */
@@ -732,6 +723,7 @@
 	virtual FMovieSceneRootEvaluationTemplateInstance& GetEvaluationTemplate() override { return RootTemplateInstance; }
 	virtual void ResetToNewRootSequence(UMovieSceneSequence& NewSequence) override;
 	virtual void FocusSequenceInstance(UMovieSceneSubSection& InSubSection) override;
+	virtual TSharedPtr<UE::Sequencer::FEditorViewModel> GetViewModel() const override;
 	virtual void SuppressAutoEvaluation(UMovieSceneSequence* Sequence, const FGuid& InSequenceSignature) override;
 	virtual EAutoChangeMode GetAutoChangeMode() const override;
 	virtual void SetAutoChangeMode(EAutoChangeMode AutoChangeMode) override;
@@ -831,6 +823,7 @@
 	virtual FGuid MakeNewSpawnable(UObject& SourceObject, UActorFactory* ActorFactory = nullptr, bool bSetupDefaults = true) override;
 	virtual bool IsReadOnly() const override;
 	virtual void ExternalSelectionHasChanged() override { SynchronizeSequencerSelectionWithExternalSelection(); }
+	virtual TSharedPtr<ISequencerTrackEditor> GetTrackEditor(UMovieSceneTrack* InTrack) override;
 	virtual void ObjectImplicitlyAdded(UObject* InObject) const override;
 	virtual void ObjectImplicitlyRemoved(UObject* InObject) const override;
 
@@ -1056,6 +1049,8 @@
 	// End of FEditorUndoClient
 
 	void OnSelectedOutlinerNodesChanged();
+	// Called on Tick after OnSelectedOutlinerNodesChanged has been called
+	void HandleSelectedOutlinerNodesChanged();
 
 	void AddNodeGroupsCollectionChangedDelegate();
 	void RemoveNodeGroupsCollectionChangedDelegate();
@@ -1065,7 +1060,7 @@
 public:
 	void AddSelectedNodesToNewNodeGroup();
 	void AddSelectedNodesToExistingNodeGroup(UMovieSceneNodeGroup* NodeGroup);
-	void AddNodesToExistingNodeGroup(const TArray<TSharedRef<FSequencerDisplayNode>>& Nodes, UMovieSceneNodeGroup* NodeGroup);
+	void AddNodesToExistingNodeGroup(const TArray<TWeakPtr<UE::Sequencer::FViewModel>>& Nodes, UMovieSceneNodeGroup* NodeGroup);
 
 	void ClearFilters();
 
@@ -1105,7 +1100,7 @@
 	void CalculateSelectedFolderAndPath(TArray<UMovieSceneFolder*>& OutSelectedParentFolders, FString& OutNewNodePath);
 
 	/** Returns the tail folder from the given Folder Path, creating each folder if needed. */
-	UMovieSceneFolder* CreateFoldersRecursively(const TArray<FName>& FolderPath, int32 FolderPathIndex, UMovieScene* OwningMovieScene, UMovieSceneFolder* ParentFolder, const TArray<UMovieSceneFolder*>& FoldersToSearch);
+	UMovieSceneFolder* CreateFoldersRecursively(const TArray<FName>& FolderPaths, int32 FolderPathIndex, UMovieScene* OwningMovieScene, UMovieSceneFolder* ParentFolder, TArrayView<UMovieSceneFolder* const> FoldersToSearch);
 
 	/** Create set playback start transport control */
 	TSharedRef<SWidget> OnCreateTransportSetPlaybackStart();
@@ -1120,7 +1115,7 @@
 	TSharedRef<SWidget> OnCreateTransportSetPlaybackEnd();
 
 	/** Select keys and/or sections in a display node that fall into the current selection range. */
-	void SelectInSelectionRange(const TSharedRef<FSequencerDisplayNode>& DisplayNode, const TRange<FFrameNumber>& SelectionRange, bool bSelectKeys, bool bSelectSections);
+	void SelectInSelectionRange(const TSharedPtr<UE::Sequencer::FViewModel>& Item, const TRange<FFrameNumber>& SelectionRange, bool bSelectKeys, bool bSelectSections);
 	
 	/** Create loop mode transport control */
 	TSharedRef<SWidget> OnCreateTransportLoopMode();
@@ -1195,6 +1190,7 @@
 
 	/** List of tools we own */
 	TArray<TSharedPtr<ISequencerTrackEditor>> TrackEditors;
+	TMap<FObjectKey, TSharedPtr<ISequencerTrackEditor>> TrackEditorsByType;
 
 	/** List of object bindings we can use */
 	TArray<TSharedPtr<ISequencerEditorObjectBinding>> ObjectBindings;
@@ -1204,7 +1200,10 @@
 
 	/** Main sequencer widget */
 	TSharedPtr<SSequencer> SequencerWidget;
-	
+
+	/** Selection handler for interfacing with sequencer selection */
+	TSharedPtr<FSequencerOutlinerSelectionHandler> SelectionHandler;
+
 	/** Spawn register for keeping track of what is spawned */
 	TSharedPtr<FMovieSceneSpawnRegister> SpawnRegister;
 
@@ -1460,6 +1459,8 @@
 
 	TUniquePtr<FObjectBindingTagCache> ObjectBindingTagCache;
 
+	TSharedPtr<UE::Sequencer::FSequencerEditorViewModel> ViewModel;
+
 	struct FCachedViewState
 	{
 		FCachedViewState()
diff --git a/Engine/Source/Editor/Sequencer/Private/SequencerAddKeyOperation.cpp b/Engine/Source/Editor/Sequencer/Private/SequencerAddKeyOperation.cpp
--- a/Engine/Source/Editor/Sequencer/Private/SequencerAddKeyOperation.cpp
+++ b/Engine/Source/Editor/Sequencer/Private/SequencerAddKeyOperation.cpp
@@ -1,178 +1,164 @@
-// Copyright Epic Games, Inc. All Rights Reserved.
-
-#include "SequencerAddKeyOperation.h"
-
-#include "ISequencer.h"
-#include "ISequencerTrackEditor.h"
-#include "DisplayNodes/SequencerDisplayNode.h"
-#include "DisplayNodes/SequencerTrackNode.h"
-#include "DisplayNodes/SequencerSectionKeyAreaNode.h"
-
-#include "Algo/Find.h"
-
-
-namespace UE
-{
-namespace Sequencer
-{
-
-
-FAddKeyOperation FAddKeyOperation::FromNodes(const TSet<TSharedRef<FSequencerDisplayNode>>& InNodes, bool bJustVisible)
-{
-	FAddKeyOperation Operation;
-
-	TArray<TSharedRef<FSequencerDisplayNode>> FilteredNodes;
-
-	// Remove any child nodes that have a parent also included in the set
-	for (const TSharedRef<FSequencerDisplayNode>& ProspectiveNode : InNodes)
-	{
-		TSharedPtr<FSequencerDisplayNode> Parent = ProspectiveNode->GetParent();
-		while (Parent)
-		{
-			if (InNodes.Contains(Parent.ToSharedRef()))
-			{
-				goto Continue;
-			}
-			Parent = Parent->GetParent();
-		}
-
-		FilteredNodes.Add(ProspectiveNode);
-
-	Continue:
-		continue;
-	}
-
-	Operation.AddPreFilteredNodes(FilteredNodes, bJustVisible);
-	return Operation;
-}
-
-FAddKeyOperation FAddKeyOperation::FromNode(TSharedRef<FSequencerDisplayNode> InNode, bool bJustVisible)
-{
-	FAddKeyOperation Operation;
-	Operation.AddPreFilteredNodes(MakeArrayView(&InNode, 1), bJustVisible);
-	return Operation;
-}
-
-FAddKeyOperation FAddKeyOperation::FromKeyAreas(ISequencerTrackEditor* TrackEditor, const TArrayView<TSharedRef<IKeyArea>> InKeyAreas)
-{
-	FAddKeyOperation Operation;
-	if (ensure(TrackEditor))
-	{
-		for (const TSharedRef<IKeyArea>& KeyArea : InKeyAreas)
-		{
-			Operation.ProcessKeyArea(TrackEditor, KeyArea);
-		}
-	}
-	return Operation;
-}
-
-void FAddKeyOperation::AddPreFilteredNodes(TArrayView<const TSharedRef<FSequencerDisplayNode>> FilteredNodes, bool bJustVisible)
-{
-	auto KeyChildTrackArea = [this, bJustVisible](FSequencerDisplayNode& InNode)
-	{
-		if (InNode.GetType() == ESequencerNode::Track)
-		{
-			FSequencerTrackNode* TrackNode = static_cast<FSequencerTrackNode*>(&InNode);
-			if (TrackNode->GetSubTrackMode() != FSequencerTrackNode::ESubTrackMode::ParentTrack)
-			{
-				// Consider everything underneath this track for keying
-				this->ConsiderKeyableAreas(TrackNode, TrackNode, bJustVisible);
-			}
-		}
-		return true;
-	};
-
-	for (const TSharedRef<FSequencerDisplayNode>& Node : FilteredNodes)
-	{
-		if (TSharedPtr<FSequencerTrackNode> ParentTrack = Node->FindParentTrackNode())
-		{
-			ConsiderKeyableAreas(ParentTrack.Get(), &Node.Get(), bJustVisible);
-		}
-		else
-		{
-			Node->Traverse_ParentFirst(KeyChildTrackArea, true/*bIncludeThisNode*/, bJustVisible);
-		}
-	}
-}
-
-bool FAddKeyOperation::ConsiderKeyableAreas(FSequencerTrackNode* InTrackNode, FSequencerDisplayNode* KeyAnythingBeneath, bool bJustVisible)
-{
-	bool bKeyedAnything = false;
-
-	auto Traversal = [this, InTrackNode, &bKeyedAnything, bJustVisible](FSequencerDisplayNode& InNode)
-	{
-		if (InNode.GetType() == ESequencerNode::Track)
-		{
-			FSequencerTrackNode* ThisTrackNode = static_cast<FSequencerTrackNode*>(&InNode);
-
-			if (TSharedPtr<FSequencerSectionKeyAreaNode> KeyAreaNode = ThisTrackNode->GetTopLevelKeyNode())
-			{
-				bKeyedAnything |= this->ProcessKeyAreaNode(ThisTrackNode, KeyAreaNode.Get(), bJustVisible);
-			}
-		}
-
-		if (InNode.GetType() == ESequencerNode::KeyArea)
-		{
-			bKeyedAnything |= this->ProcessKeyAreaNode(InTrackNode, static_cast<FSequencerSectionKeyAreaNode*>(&InNode), bJustVisible);
-		}
-		return true;
-	};
-	KeyAnythingBeneath->Traverse_ParentFirst(Traversal, true/*bIncludeThisNode*/, bJustVisible);
-
-	return bKeyedAnything;
-}
-
-bool FAddKeyOperation::ProcessKeyAreaNode(FSequencerTrackNode* InTrackNode, const FSequencerSectionKeyAreaNode* KeyAreaNode, bool bJustVisible)
-{
-	bool bKeyedAnything = false;
-	if (bJustVisible == false || InTrackNode->IsVisible())
-	{
-		for (const TSharedRef<IKeyArea>& KeyArea : KeyAreaNode->GetAllKeyAreas())
-		{
-			bKeyedAnything |= ProcessKeyArea(InTrackNode, KeyArea);
-		}
-	}
-
-	return bKeyedAnything;
-}
-
-bool FAddKeyOperation::ProcessKeyArea(FSequencerTrackNode* InTrackNode, TSharedPtr<IKeyArea> InKeyArea)
-{
-	ISequencerTrackEditor* TrackEditor = &InTrackNode->GetTrackEditor();
-	return ProcessKeyArea(TrackEditor, InKeyArea);
-}
-
-bool FAddKeyOperation::ProcessKeyArea(ISequencerTrackEditor* InTrackEditor, TSharedPtr<IKeyArea> InKeyArea)
-{
-	TSharedPtr<ISequencerSection> Section       = InKeyArea->GetSectionInterface();
-	UMovieSceneSection*           SectionObject = Section       ? Section->GetSectionObject()                      : nullptr;
-	UMovieSceneTrack*             TrackObject   = SectionObject ? SectionObject->GetTypedOuter<UMovieSceneTrack>() : nullptr;
-
-	if (TrackObject)
-	{
-		GetTrackOperation(InTrackEditor).Populate(TrackObject, Section, InKeyArea);
-		return true;
-	}
-
-	return false;
-}
-
-void FAddKeyOperation::Commit(FFrameNumber KeyTime, ISequencer& InSequencer)
-{
-	for (TTuple<ISequencerTrackEditor*, FKeyOperation>& Pair : OperationsByTrackEditor)
-	{
-		Pair.Value.InitializeOperation(KeyTime);
-		Pair.Key->ProcessKeyOperation(KeyTime, Pair.Value, InSequencer);
-	}
-
-	InSequencer.UpdatePlaybackRange();
-	InSequencer.NotifyMovieSceneDataChanged(EMovieSceneDataChangeType::TrackValueChanged);
-}
-
-FKeyOperation& FAddKeyOperation::GetTrackOperation(ISequencerTrackEditor* TrackEditor)
-{
-	return OperationsByTrackEditor.FindOrAdd(TrackEditor);
-}
-
-} // namespace Sequencer
-} // namespace UE
+// Copyright Epic Games, Inc. All Rights Reserved.
+
+#include "SequencerAddKeyOperation.h"
+
+#include "ISequencer.h"
+#include "ISequencerTrackEditor.h"
+#include "IKeyArea.h"
+#include "MVVM/ViewModels/ViewModel.h"
+#include "MVVM/ViewModels/ChannelModel.h"
+#include "MVVM/ViewModels/ViewModelIterators.h"
+#include "MVVM/ViewModels/SequenceModel.h"
+#include "MVVM/ViewModels/OutlinerViewModel.h"
+#include "MVVM/ViewModels/SequencerEditorViewModel.h"
+#include "MVVM/Extensions/ITrackExtension.h"
+
+#include "Algo/Find.h"
+
+
+namespace UE
+{
+namespace Sequencer
+{
+
+FAddKeyOperation FAddKeyOperation::FromNodes(const TSet<TWeakPtr<FViewModel>>& InNodes)
+{
+	FAddKeyOperation Operation;
+
+	TArray<TWeakPtr<FViewModel>> FilteredNodes;
+
+	// Remove any child nodes that have a parent also included in the set
+	for (const TWeakPtr<FViewModel>& ProspectiveNode : InNodes)
+	{
+		TSharedPtr<FViewModel> Parent = ProspectiveNode.Pin()->GetParent();
+		while (Parent)
+		{
+			if (InNodes.Contains(Parent.ToSharedRef()))
+			{
+				goto Continue;
+			}
+			Parent = Parent->GetParent();
+		}
+
+		FilteredNodes.Add(ProspectiveNode);
+
+	Continue:
+		continue;
+	}
+
+	Operation.AddPreFilteredNodes(FilteredNodes);
+	return Operation;
+}
+
+FAddKeyOperation FAddKeyOperation::FromNode(TWeakPtr<FViewModel> InNode)
+{
+	FAddKeyOperation Operation;
+	Operation.AddPreFilteredNodes(MakeArrayView(&InNode, 1));
+	return Operation;
+}
+
+FAddKeyOperation FAddKeyOperation::FromKeyAreas(ISequencerTrackEditor* TrackEditor, const TArrayView<TSharedRef<IKeyArea>> InKeyAreas)
+{
+	FAddKeyOperation Operation;
+	if (ensure(TrackEditor))
+	{
+		for (const TSharedRef<IKeyArea>& KeyArea : InKeyAreas)
+		{
+			Operation.ProcessKeyArea(TrackEditor, KeyArea);
+		}
+	}
+	return Operation;
+}
+
+void FAddKeyOperation::AddPreFilteredNodes(TArrayView<const TWeakPtr<FViewModel>> FilteredNodes)
+{
+	TSharedPtr<FSequenceModel> SequenceModel;
+	const TArray<FViewModelTypeID> TypeIDs({ ITrackExtension::ID, IOutlinerExtension::ID });
+	for (const TWeakPtr<FViewModel>& FilteredNode : FilteredNodes)
+	{
+		TSharedPtr<FViewModel> Node = FilteredNode.Pin();
+		TSharedPtr<FViewModel> ParentModel = Node->FindAncestorOfTypes(MakeArrayView(TypeIDs));
+		if (ParentModel)
+		{
+			ConsiderKeyableAreas(ParentModel->CastThisShared<ITrackExtension>(), Node);
+		}
+		else
+		{
+			constexpr bool bIncludeThis = true;
+			for (FParentFirstChildIterator Child(FilteredNode.Pin(), bIncludeThis); Child; ++Child)
+			{
+				if (Child->IsA<ITrackExtension>() && Child->IsA<IOutlinerExtension>())
+				{
+					TSharedPtr<ITrackExtension> TrackExtension = Child->CastThisShared<ITrackExtension>();
+					ConsiderKeyableAreas(TrackExtension, *Child);
+				}
+			}
+		}
+	}
+}
+
+bool FAddKeyOperation::ConsiderKeyableAreas(TSharedPtr<ITrackExtension> InTrackModel, FViewModelPtr KeyAnythingBeneath)
+{
+	bool bKeyedAnything = false;
+
+	constexpr bool bIncludeThis = true;
+	for (TParentFirstChildIterator<FChannelGroupModel> ChannelGroupModelIt(KeyAnythingBeneath->AsShared(), bIncludeThis); ChannelGroupModelIt; ++ChannelGroupModelIt)
+	{
+		bKeyedAnything |= ProcessKeyArea(InTrackModel, *ChannelGroupModelIt);
+	}
+
+	return bKeyedAnything;
+}
+
+bool FAddKeyOperation::ProcessKeyArea(TSharedPtr<ITrackExtension> InTrackModel, TViewModelPtr<FChannelGroupModel> InChannelGroupModel)
+{
+	bool bKeyedAnything = false;
+	ISequencerTrackEditor* TrackEditor = InTrackModel->GetTrackEditor().Get();
+
+	constexpr bool bIncludeThis = true;
+	for (const TWeakViewModelPtr<FChannelModel>& WeakChannel : InChannelGroupModel->GetChannels())
+	{
+		if (TSharedPtr<FChannelModel> Channel = WeakChannel.Pin())
+		{
+			bKeyedAnything |= ProcessKeyArea(TrackEditor, Channel->GetKeyArea());
+		}
+	}
+
+	return bKeyedAnything;
+}
+
+bool FAddKeyOperation::ProcessKeyArea(ISequencerTrackEditor* InTrackEditor, TSharedPtr<IKeyArea> InKeyArea)
+{
+	TSharedPtr<ISequencerSection> Section       = InKeyArea->GetSectionInterface();
+	UMovieSceneSection*           SectionObject = Section       ? Section->GetSectionObject()                      : nullptr;
+	UMovieSceneTrack*             TrackObject   = SectionObject ? SectionObject->GetTypedOuter<UMovieSceneTrack>() : nullptr;
+
+	if (TrackObject)
+	{
+		GetTrackOperation(InTrackEditor).Populate(TrackObject, Section, InKeyArea);
+		return true;
+	}
+
+	return false;
+}
+
+void FAddKeyOperation::Commit(FFrameNumber KeyTime, ISequencer& InSequencer)
+{
+	for (TTuple<ISequencerTrackEditor*, FKeyOperation>& Pair : OperationsByTrackEditor)
+	{
+		Pair.Value.InitializeOperation(KeyTime);
+		Pair.Key->ProcessKeyOperation(KeyTime, Pair.Value, InSequencer);
+	}
+
+	InSequencer.UpdatePlaybackRange();
+	InSequencer.NotifyMovieSceneDataChanged(EMovieSceneDataChangeType::TrackValueChanged);
+}
+
+FKeyOperation& FAddKeyOperation::GetTrackOperation(ISequencerTrackEditor* TrackEditor)
+{
+	return OperationsByTrackEditor.FindOrAdd(TrackEditor);
+}
+
+} // namespace Sequencer
+} // namespace UE
+
diff --git a/Engine/Source/Editor/Sequencer/Private/SequencerClipboardReconciler.cpp b/Engine/Source/Editor/Sequencer/Private/SequencerClipboardReconciler.cpp
--- a/Engine/Source/Editor/Sequencer/Private/SequencerClipboardReconciler.cpp
+++ b/Engine/Source/Editor/Sequencer/Private/SequencerClipboardReconciler.cpp
@@ -2,6 +2,7 @@
 
 #include "SequencerClipboardReconciler.h"
 #include "IKeyArea.h"
+#include "MVVM/ViewModels/ChannelModel.h"
 
 TMap<FName, TArray<FName>> FSequencerClipboardReconciler::KeyAreaAliases;
 
@@ -70,7 +71,7 @@
 	for (auto& MetaDataPair : MetaData)
 	{
 		const TArray<FMovieSceneClipboardKeyTrack>& SrcArray = Clipboard->GetKeyTrackGroups()[MetaDataPair.Value.SourceGroup];
-		FKeyAreaArray& DstArray = PasteDestination[MetaDataPair.Key];
+		TArray<TSharedPtr<UE::Sequencer::FChannelModel>>& DstArray = PasteDestination[MetaDataPair.Key];
 
 		int32 SrcIndex = 0;
 		int32 DstIndex = 0;
@@ -81,18 +82,28 @@
 		{
 			for (auto& IndexPair : MetaDataPair.Value.DestToSrcMap)
 			{
-				if (UMovieSceneSection* Section = DstArray[DstIndex]->GetOwningSection())
+				if (UMovieSceneSection* Section = DstArray[DstIndex]->GetSection())
 				{
-					DstArray[IndexPair.Key]->PasteKeys(SrcArray[IndexPair.Value], Clipboard->GetEnvironment(), PasteEnvironment);
+					TArray<FKeyHandle> PastedKeys =  DstArray[IndexPair.Key]->GetKeyArea()->PasteKeys(SrcArray[IndexPair.Value], Clipboard->GetEnvironment(), PasteEnvironment);
+
+					for (FKeyHandle KeyHandle : PastedKeys)
+					{
+						PasteEnvironment.ReportPastedKey(KeyHandle, DstArray[IndexPair.Key]);
+					}
 					bAnythingPasted = true;
 				}
 			}
 		}
 		else for (; SrcIndex < SrcArray.Num() && DstIndex < DstArray.Num();)
 		{
-			if (UMovieSceneSection* Section = DstArray[DstIndex]->GetOwningSection())
+			if (UMovieSceneSection* Section = DstArray[DstIndex]->GetSection())
 			{
-				DstArray[DstIndex]->PasteKeys(SrcArray[SrcIndex], Clipboard->GetEnvironment(), PasteEnvironment);
+				TArray<FKeyHandle> PastedKeys = DstArray[DstIndex]->GetKeyArea()->PasteKeys(SrcArray[SrcIndex], Clipboard->GetEnvironment(), PasteEnvironment);
+				for (FKeyHandle KeyHandle : PastedKeys)
+				{
+					PasteEnvironment.ReportPastedKey(KeyHandle, DstArray[DstIndex]);
+				}
+
 				bAnythingPasted = true;
 			}
 
@@ -127,13 +138,13 @@
 	return bAnythingPasted;
 }
 
-bool FSequencerClipboardReconciler::FindMatchingGroup(const FKeyAreaArray& Destination, const TArray<FMovieSceneClipboardKeyTrack>& Source, TMap<int32, int32>& Map, bool bAllowAliases)
+bool FSequencerClipboardReconciler::FindMatchingGroup(const TArray<TSharedPtr<UE::Sequencer::FChannelModel>>& Destination, const TArray<FMovieSceneClipboardKeyTrack>& Source, TMap<int32, int32>& Map, bool bAllowAliases)
 {
 	// If we find an exact match, we use custom
 	bool bFoundMatch = false;
 	for (int32 DstIndex = 0; DstIndex < Destination.Num(); ++DstIndex)
 	{
-		FName DstName = Destination[DstIndex]->GetName();
+		FName DstName = Destination[DstIndex]->GetKeyArea()->GetName();
 
 		const int32 SourceIndex = Source.IndexOfByPredicate([&](const FMovieSceneClipboardKeyTrack& Track){
 
diff --git a/Engine/Source/Editor/Sequencer/Private/SequencerCommands.cpp b/Engine/Source/Editor/Sequencer/Private/SequencerCommands.cpp
--- a/Engine/Source/Editor/Sequencer/Private/SequencerCommands.cpp
+++ b/Engine/Source/Editor/Sequencer/Private/SequencerCommands.cpp
@@ -175,6 +175,8 @@
 	UI_COMMAND(AddScaleKey, "Add Scale Key", "Add a scale key at the current time for the selected actor.", EUserInterfaceActionType::Button, FInputChord(EModifierKey::Shift, EKeys::R));
 
 	UI_COMMAND(TogglePilotCamera, "Pilot Camera", "Toggle piloting the last camera or the camera cut camera.", EUserInterfaceActionType::ToggleButton, FInputChord(EModifierKey::Shift, EKeys::P));
+
+	UI_COMMAND(RefreshUI, "Refresh UI", "Forcibly refresh the UI from source data.", EUserInterfaceActionType::Button, FInputChord(EKeys::F5));
 }
 
 #undef LOCTEXT_NAMESPACE
diff --git a/Engine/Source/Editor/Sequencer/Private/SequencerCommonHelpers.cpp b/Engine/Source/Editor/Sequencer/Private/SequencerCommonHelpers.cpp
--- a/Engine/Source/Editor/Sequencer/Private/SequencerCommonHelpers.cpp
+++ b/Engine/Source/Editor/Sequencer/Private/SequencerCommonHelpers.cpp
@@ -1,510 +1,351 @@
-// Copyright Epic Games, Inc. All Rights Reserved.
-
-#include "SequencerCommonHelpers.h"
-#include "SequencerSelectedKey.h"
-#include "DisplayNodes/SequencerSectionKeyAreaNode.h"
-#include "DisplayNodes/SequencerObjectBindingNode.h"
-#include "DisplayNodes/SequencerTrackNode.h"
-#include "Sequencer.h"
-#include "SSequencer.h"
-#include "ISequencerHotspot.h"
-#include "SSequencerSection.h"
-#include "SSequencerTreeView.h"
-#include "VirtualTrackArea.h"
-#include "SequencerContextMenus.h"
-#include "IDetailsView.h"
-#include "IStructureDetailsView.h"
-#include "FrameNumberDetailsCustomization.h"
-#include "MovieSceneSectionDetailsCustomization.h"
-#include "MovieSceneSequence.h"
-#include "MovieScene.h"
-#include "Modules/ModuleManager.h"
-#include "PropertyEditorModule.h"
-#include "Styling/CoreStyle.h"
-
-void SequencerHelpers::GetAllKeyAreas(TSharedPtr<FSequencerDisplayNode> DisplayNode, TSet<TSharedPtr<IKeyArea>>& KeyAreas)
-{
-	TArray<TSharedPtr<FSequencerDisplayNode>> NodesToCheck;
-	NodesToCheck.Add(DisplayNode);
-	while (NodesToCheck.Num() > 0)
-	{
-		TSharedPtr<FSequencerDisplayNode> NodeToCheck = NodesToCheck[0];
-		NodesToCheck.RemoveAt(0);
-
-		if (NodeToCheck->GetType() == ESequencerNode::Track)
-		{
-			TSharedPtr<FSequencerTrackNode> TrackNode = StaticCastSharedPtr<FSequencerTrackNode>(NodeToCheck);
-			TArray<TSharedRef<FSequencerSectionKeyAreaNode>> KeyAreaNodes;
-			TrackNode->GetChildKeyAreaNodesRecursively(KeyAreaNodes);
-			for (TSharedRef<FSequencerSectionKeyAreaNode> KeyAreaNode : KeyAreaNodes)
-			{
-				for (TSharedPtr<IKeyArea> KeyArea : KeyAreaNode->GetAllKeyAreas())
-				{
-					KeyAreas.Add(KeyArea);
-				}
-			}
-		}
-		else
-		{
-			if (NodeToCheck->GetType() == ESequencerNode::KeyArea)
-			{
-				TSharedPtr<FSequencerSectionKeyAreaNode> KeyAreaNode = StaticCastSharedPtr<FSequencerSectionKeyAreaNode>(NodeToCheck);
-				for (TSharedPtr<IKeyArea> KeyArea : KeyAreaNode->GetAllKeyAreas())
-				{
-					KeyAreas.Add(KeyArea);
-				}
-			}
-			for (TSharedRef<FSequencerDisplayNode> ChildNode : NodeToCheck->GetChildNodes())
-			{
-				NodesToCheck.Add(ChildNode);
-			}
-		}
-	}
-}
-
-int32 SequencerHelpers::GetSectionFromTime(TArrayView<UMovieSceneSection* const> InSections, FFrameNumber Time)
-{
-	FFrameNumber ClosestLowerBound = TNumericLimits<int32>::Max();
-	TOptional<int32> MaxOverlapPriority, MaxProximalPriority;
-
-	int32 MostRelevantIndex = INDEX_NONE;
-
-	for (int32 Index = 0; Index < InSections.Num(); ++Index)
-	{
-		const UMovieSceneSection* Section = InSections[Index];
-		if (Section)
-		{
-			const int32 ThisSectionPriority = Section->GetOverlapPriority();
-			TRange<FFrameNumber> SectionRange = Section->GetRange();
-
-			// If the specified time is within the section bounds
-			if (SectionRange.Contains(Time))
-			{
-				if (ThisSectionPriority >= MaxOverlapPriority.Get(ThisSectionPriority))
-				{
-					MaxOverlapPriority = ThisSectionPriority;
-					MostRelevantIndex = Index;
-				}
-			}
-			// Check for nearby sections if there is nothing overlapping
-			else if (!MaxOverlapPriority.IsSet() && SectionRange.HasLowerBound())
-			{
-				const FFrameNumber LowerBoundValue = SectionRange.GetLowerBoundValue();
-				// If this section exists beyond the current time, we can choose it if its closest to the time
-				if (LowerBoundValue >= Time)
-				{
-					if (
-						(LowerBoundValue < ClosestLowerBound) ||
-						(LowerBoundValue == ClosestLowerBound && ThisSectionPriority >= MaxProximalPriority.Get(ThisSectionPriority))
-						)
-					{
-						MostRelevantIndex = Index;
-						ClosestLowerBound = LowerBoundValue;
-						MaxProximalPriority = ThisSectionPriority;
-					}
-				}
-			}
-		}
-	}
-
-	// If we didn't find one, use the last one (or return -1)
-	if (MostRelevantIndex == -1)
-	{
-		MostRelevantIndex = InSections.Num() - 1;
-	}
-
-	return MostRelevantIndex;
-}
-
-void SequencerHelpers::GetDescendantNodes(TSharedRef<FSequencerDisplayNode> DisplayNode, TSet<TSharedRef<FSequencerDisplayNode>>& Nodes)
-{
-	for (auto ChildNode : DisplayNode.Get().GetChildNodes())
-	{
-		Nodes.Add(ChildNode);
-
-		GetDescendantNodes(ChildNode, Nodes);
-	}
-}
-
-void SequencerHelpers::GetAllSections(TSharedRef<FSequencerDisplayNode> DisplayNode, TSet<TWeakObjectPtr<UMovieSceneSection>>& Sections)
-{
-	TSet<TSharedRef<FSequencerDisplayNode> > AllNodes;
-	AllNodes.Add(DisplayNode);
-	GetDescendantNodes(DisplayNode, AllNodes);
-
-	for (auto NodeToCheck : AllNodes)
-	{
-		TSet<TSharedPtr<IKeyArea> > KeyAreas;
-		GetAllKeyAreas(NodeToCheck, KeyAreas);
-		
-		for (auto KeyArea : KeyAreas)
-		{
-			UMovieSceneSection* OwningSection = KeyArea->GetOwningSection();
-			if (OwningSection != nullptr)
-			{
-				Sections.Add(OwningSection);	
-			}
-		}
-
-		if (NodeToCheck->GetType() == ESequencerNode::Track)
-		{
-			TSharedRef<const FSequencerTrackNode> TrackNode = StaticCastSharedRef<const FSequencerTrackNode>( NodeToCheck );
-			UMovieSceneTrack* Track = TrackNode->GetTrack();
-			if (Track != nullptr)
-			{
-				for (TSharedRef<ISequencerSection> TrackSection : TrackNode->GetSections())
-				{
-					if (UMovieSceneSection* Section = TrackSection->GetSectionObject())
-					{
-						Sections.Add(Section);
-					}
-				}
-			}
-		}
-	}
-}
-
-bool IsSectionSelectedInNode(FSequencer& Sequencer, TSharedRef<FSequencerDisplayNode> InNode)
-{
-	if (InNode->GetType() == ESequencerNode::Track)
-	{
-		TSharedRef<FSequencerTrackNode> TrackNode = StaticCastSharedRef<FSequencerTrackNode>(InNode);
-
-		for (auto Section : TrackNode->GetSections())
-		{
-			if (Sequencer.GetSelection().IsSelected(Section->GetSectionObject()))
-			{
-				return true;
-			}
-		}
-	}
-	return false;
-}
-
-bool AreKeysSelectedInNode(FSequencer& Sequencer, TSharedRef<FSequencerDisplayNode> InNode)
-{
-	TSet<TSharedPtr<IKeyArea>> KeyAreas;
-	SequencerHelpers::GetAllKeyAreas(InNode, KeyAreas);
-
-	for (const FSequencerSelectedKey& Key : Sequencer.GetSelection().GetSelectedKeys())
-	{
-		if (KeyAreas.Contains(Key.KeyArea))
-		{
-			return true;
-		}
-	}
-
-	return false;
-}
-
-void SequencerHelpers::ValidateNodesWithSelectedKeysOrSections(FSequencer& Sequencer)
-{
-	TArray<TSharedRef<FSequencerDisplayNode>> NodesToRemove;
-
-	for (auto Node : Sequencer.GetSelection().GetNodesWithSelectedKeysOrSections())
-	{
-		if (!IsSectionSelectedInNode(Sequencer, Node) && !AreKeysSelectedInNode(Sequencer, Node))
-		{
-			NodesToRemove.Add(Node);
-		}
-	}
-
-	for (auto Node : NodesToRemove)
-	{
-		Sequencer.GetSelection().RemoveFromNodesWithSelectedKeysOrSections(Node);
-	}
-}
-
-void SequencerHelpers::UpdateHoveredNodeFromSelectedSections(FSequencer& Sequencer)
-{
-	FSequencerSelection& Selection = Sequencer.GetSelection();
-
-	TSharedRef<SSequencer> SequencerWidget = StaticCastSharedRef<SSequencer>(Sequencer.GetSequencerWidget());
-	TSharedPtr<FSequencerDisplayNode> HoveredNode = SequencerWidget->GetTreeView()->GetNodeTree()->GetHoveredNode();
-	if (!HoveredNode.IsValid())
-	{
-		return;
-	}
-
-	if (IsSectionSelectedInNode(Sequencer, HoveredNode.ToSharedRef()))
-	{
-		Selection.AddToNodesWithSelectedKeysOrSections(HoveredNode.ToSharedRef());
-	}
-	else
-	{
-		Selection.RemoveFromNodesWithSelectedKeysOrSections(HoveredNode.ToSharedRef());
-	}
-}
-
-
-void SequencerHelpers::UpdateHoveredNodeFromSelectedKeys(FSequencer& Sequencer)
-{
-	FSequencerSelection& Selection = Sequencer.GetSelection();
-
-	TSharedRef<SSequencer> SequencerWidget = StaticCastSharedRef<SSequencer>(Sequencer.GetSequencerWidget());
-	TSharedPtr<FSequencerDisplayNode> HoveredNode = SequencerWidget->GetTreeView()->GetNodeTree()->GetHoveredNode();
-	if (!HoveredNode.IsValid())
-	{
-		return;
-	}
-
-	if (AreKeysSelectedInNode(Sequencer, HoveredNode.ToSharedRef()))
-	{
-		Selection.AddToNodesWithSelectedKeysOrSections(HoveredNode.ToSharedRef());
-	}
-	else
-	{
-		Selection.RemoveFromNodesWithSelectedKeysOrSections(HoveredNode.ToSharedRef());
-	}
-}
-
-
-void SequencerHelpers::PerformDefaultSelection(FSequencer& Sequencer, const FPointerEvent& MouseEvent)
-{
-	FSequencerSelection& Selection = Sequencer.GetSelection();
-	Selection.SuspendBroadcast();
-
-	// @todo: selection in transactions
-	auto ConditionallyClearSelection = [&]{
-		if (!MouseEvent.IsShiftDown() && !MouseEvent.IsControlDown())
-		{			
-			Selection.EmptySelectedSections();
-			Selection.EmptySelectedKeys();
-			Selection.EmptyNodesWithSelectedKeysOrSections();
-		}
-	};
-
-	TSharedPtr<ISequencerHotspot> Hotspot = Sequencer.GetHotspot();
-	if (!Hotspot.IsValid())
-	{
-		ConditionallyClearSelection();
-		Selection.ResumeBroadcast();
-		Selection.GetOnOutlinerNodeSelectionChanged().Broadcast();
-		return;
-	}
-
-	// Handle right-click selection separately since we never deselect on right click (except for clearing on exclusive selection)
-	if (MouseEvent.GetEffectingButton() == EKeys::RightMouseButton)
-	{
-		if (Hotspot->GetType() == ESequencerHotspot::Key)
-		{
-			bool bHasClearedSelection = false;
-			for (const FSequencerSelectedKey& Key : static_cast<FKeyHotspot*>(Hotspot.Get())->Keys)
-			{
-				if (!Selection.IsSelected(Key))
-				{
-					if (!bHasClearedSelection)
-					{
-						ConditionallyClearSelection();
-						bHasClearedSelection = true;
-					}
-					Selection.AddToSelection(Key);
-				}
-			}
-		}
-		else if (Hotspot->GetType() == ESequencerHotspot::Section || Hotspot->GetType() == ESequencerHotspot::EasingArea)
-		{
-			UMovieSceneSection* Section = static_cast<FSectionHotspot*>(Hotspot.Get())->WeakSection.Get();
-			if (!Selection.IsSelected(Section))
-			{
-				ConditionallyClearSelection();
-				Selection.AddToSelection(Section);
-			}
-		}
-		else if (Hotspot->GetType() == ESequencerHotspot::SectionResize_L || Hotspot->GetType() == ESequencerHotspot::SectionResize_R)
-		{
-			UMovieSceneSection* Section = static_cast<FSectionResizeHotspot*>(Hotspot.Get())->WeakSection.Get();
-			if (!Selection.IsSelected(Section))
-			{
-				ConditionallyClearSelection();
-				Selection.AddToSelection(Section);
-			}
-		}
-
-		if (Hotspot->GetType() == ESequencerHotspot::Key)
-		{
-			UpdateHoveredNodeFromSelectedKeys(Sequencer);
-		}
-		else
-		{
-			UpdateHoveredNodeFromSelectedSections(Sequencer);
-		}
-		
-		Selection.ResumeBroadcast();
-		return;
-	}
-
-	// Normal selection
-	ConditionallyClearSelection();
-
-	bool bForceSelect = !MouseEvent.IsControlDown();
-		
-	if (Hotspot->GetType() == ESequencerHotspot::Key)
-	{
-		for (const FSequencerSelectedKey& Key : static_cast<FKeyHotspot*>(Hotspot.Get())->Keys)
-		{
-			if (bForceSelect || !Selection.IsSelected(Key))
-			{
-				Selection.AddToSelection(Key);
-			}
-			else
-			{
-				Selection.RemoveFromSelection(Key);
-			}
-		}
-	}
-	else if (Hotspot->GetType() == ESequencerHotspot::Section || Hotspot->GetType() == ESequencerHotspot::EasingArea)
-	{
-		UMovieSceneSection* Section = static_cast<FSectionHotspot*>(Hotspot.Get())->WeakSection.Get();
-
-		// Never allow infinite sections to be selected through normal click (they're only selectable through right click)
-		if (Section->GetRange() != TRange<FFrameNumber>::All())
-		{
-			if (bForceSelect || !Selection.IsSelected(Section))
-			{
-				Selection.AddToSelection(Section);
-			}
-			else
-			{
-				Selection.RemoveFromSelection(Section);
-			}
-		}
-	}
-
-	if (Hotspot->GetType() == ESequencerHotspot::Key)
-	{
-		UpdateHoveredNodeFromSelectedKeys(Sequencer);
-	}
-	else
-	{
-		UpdateHoveredNodeFromSelectedSections(Sequencer);
-	}
-
-	Selection.ResumeBroadcast();
-	Selection.GetOnOutlinerNodeSelectionChanged().Broadcast();
-}
-
-TSharedPtr<SWidget> SequencerHelpers::SummonContextMenu(FSequencer& Sequencer, const FGeometry& MyGeometry, const FPointerEvent& MouseEvent)
-{
-	// @todo sequencer replace with UI Commands instead of faking it
-
-	// Attempt to paste into either the current node selection, or the clicked on track
-	TSharedRef<SSequencer> SequencerWidget = StaticCastSharedRef<SSequencer>(Sequencer.GetSequencerWidget());
-	const FFrameNumber PasteAtTime = Sequencer.GetLocalTime().Time.FrameNumber;
-
-	// The menu are generated through reflection and sometime the API exposes some recursivity (think about a Widget returning it parent which is also a Widget). Just by reflection
-	// it is not possible to determine when the root object is reached. It needs a kind of simulation which is not implemented. Also, even if the recursivity was correctly handled, the possible
-	// permutations tend to grow exponentially. Until a clever solution is found, the simple approach is to disable recursively searching those menus. User can still search the current one though.
-	// See UE-131257
-	const bool bInRecursivelySearchable = false;
-
-	const bool bShouldCloseWindowAfterMenuSelection = true;
-	FMenuBuilder MenuBuilder(bShouldCloseWindowAfterMenuSelection, Sequencer.GetCommandBindings(), nullptr, false, &FCoreStyle::Get(), true, NAME_None, bInRecursivelySearchable);
-
-	TSharedPtr<ISequencerHotspot> Hotspot = Sequencer.GetHotspot();
-
-	if (Hotspot.IsValid() && Hotspot->PopulateContextMenu(MenuBuilder, Sequencer, PasteAtTime))
-	{
-		return MenuBuilder.MakeWidget();
-	}
-	else if (Sequencer.GetClipboardStack().Num() != 0)
-	{
-		TSharedPtr<FPasteContextMenu> PasteMenu = FPasteContextMenu::CreateMenu(Sequencer, SequencerWidget->GeneratePasteArgs(PasteAtTime));
-		if (PasteMenu.IsValid() && PasteMenu->IsValidPaste())
-		{
-			PasteMenu->PopulateMenu(MenuBuilder);
-
-			return MenuBuilder.MakeWidget();
-		}
-	}
-
-	return nullptr;
-}
-
-/** A widget which wraps the section details view which is an FNotifyHook which is used to forward
-	changes to the section to sequencer. */
-class SSectionDetailsNotifyHookWrapper : public SCompoundWidget, public FNotifyHook
-{
-public:
-	SLATE_BEGIN_ARGS(SSectionDetailsNotifyHookWrapper) {}
-	SLATE_END_ARGS();
-
-	void Construct(FArguments InArgs) { }
-
-	void SetDetailsAndSequencer(TSharedRef<SWidget> InDetailsPanel, TSharedRef<ISequencer> InSequencer)
-	{
-		ChildSlot
-		[
-			InDetailsPanel
-		];
-		Sequencer = InSequencer;
-	}
-
-	//~ FNotifyHook interface
-	virtual void NotifyPostChange(const FPropertyChangedEvent& PropertyChangedEvent, FProperty* PropertyThatChanged) override
-	{
-		Sequencer->NotifyMovieSceneDataChanged(EMovieSceneDataChangeType::TrackValueChanged);
-	}
-
-private:
-	TSharedPtr<ISequencer> Sequencer;
-};
-
-
-void SequencerHelpers::AddPropertiesMenu(FSequencer& Sequencer, FMenuBuilder& MenuBuilder, const TArray<TWeakObjectPtr<UObject>>& Sections)
-{
-	TSharedRef<SSectionDetailsNotifyHookWrapper> DetailsNotifyWrapper = SNew(SSectionDetailsNotifyHookWrapper);
-	FDetailsViewArgs DetailsViewArgs;
-	{
-		DetailsViewArgs.bAllowSearch = false;
-		DetailsViewArgs.bCustomFilterAreaLocation = true;
-		DetailsViewArgs.bCustomNameAreaLocation = true;
-		DetailsViewArgs.bHideSelectionTip = true;
-		DetailsViewArgs.bLockable = false;
-		DetailsViewArgs.bSearchInitialKeyFocus = true;
-		DetailsViewArgs.bUpdatesFromSelection = false;
-		DetailsViewArgs.bShowOptions = false;
-		DetailsViewArgs.bShowModifiedPropertiesOption = false;
-		DetailsViewArgs.NotifyHook = &DetailsNotifyWrapper.Get();
-		DetailsViewArgs.ColumnWidth = 0.45f;
-	}
-
-	// We pass the current scene to the UMovieSceneSection customization so we can get the overall bounds of the section when we change a section from infinite->bounded.
-	UMovieScene* CurrentScene = Sequencer.GetFocusedMovieSceneSequence()->GetMovieScene();
-
-	TSharedRef<INumericTypeInterface<double>> NumericTypeInterface = Sequencer.GetNumericTypeInterface();
-
-	TSharedRef<IDetailsView> DetailsView = FModuleManager::GetModuleChecked<FPropertyEditorModule>("PropertyEditor").CreateDetailView(DetailsViewArgs);
-	DetailsView->RegisterInstancedCustomPropertyTypeLayout("FrameNumber", FOnGetPropertyTypeCustomizationInstance::CreateLambda([=]() {
-		return MakeShared<FFrameNumberDetailsCustomization>(NumericTypeInterface); }));
-	DetailsView->RegisterInstancedCustomPropertyLayout(UMovieSceneSection::StaticClass(), FOnGetDetailCustomizationInstance::CreateLambda([=]() {
-		return MakeShared<FMovieSceneSectionDetailsCustomization>(NumericTypeInterface, CurrentScene); }));
-
-	// Let section interfaces further customize the properties details view.
-	TSharedRef<FSequencerNodeTree> SequencerNodeTree = Sequencer.GetNodeTree();
-	for (TWeakObjectPtr<UObject> Section : Sections)
-	{
-		if (Section.IsValid())
-		{
-			TOptional<FSectionHandle> SectionHandle = SequencerNodeTree->GetSectionHandle(Cast<UMovieSceneSection>(Section));
-			if (SectionHandle)
-			{
-				TSharedRef<ISequencerSection> SectionInterface = SectionHandle->GetSectionInterface();
-				FSequencerSectionPropertyDetailsViewCustomizationParams CustomizationDetails(
-					SectionInterface, Sequencer.AsShared(), SectionHandle->GetTrackNode()->GetTrackEditor());
-				TSharedPtr<FSequencerObjectBindingNode> ParentObjectBindingNode = SectionHandle->GetTrackNode()->FindParentObjectBindingNode();
-				if (ParentObjectBindingNode.IsValid())
-				{
-					CustomizationDetails.ParentObjectBindingGuid = ParentObjectBindingNode->GetObjectBinding();
-				}
-				SectionInterface->CustomizePropertiesDetailsView(DetailsView, CustomizationDetails);
-			}
-		}
-	}
-
-	Sequencer.OnInitializeDetailsPanel().Broadcast(DetailsView, Sequencer.AsShared());
-	DetailsView->SetObjects(Sections);
-
-	DetailsNotifyWrapper->SetDetailsAndSequencer(DetailsView, Sequencer.AsShared());
-	MenuBuilder.AddWidget(DetailsNotifyWrapper, FText::GetEmpty(), true);
-}
+// Copyright Epic Games, Inc. All Rights Reserved.
+
+#include "SequencerCommonHelpers.h"
+#include "FrameNumberDetailsCustomization.h"
+#include "IDetailsView.h"
+#include "ISequencerSection.h"
+#include "IStructureDetailsView.h"
+#include "MVVM/Extensions/ITrackAreaExtension.h"
+#include "MVVM/ViewModels/ChannelModel.h"
+#include "MVVM/ViewModels/ViewModelIterators.h"
+#include "MVVM/ViewModels/ObjectBindingModel.h"
+#include "MVVM/ViewModels/SectionModel.h"
+#include "MVVM/ViewModels/TrackModel.h"
+#include "MVVM/ViewModels/ViewModel.h"
+#include "MVVM/ViewModels/OutlinerViewModel.h"
+#include "MVVM/ViewModels/SequencerEditorViewModel.h"
+#include "MVVM/ViewModels/VirtualTrackArea.h"
+#include "MVVM/Views/ITrackAreaHotspot.h"
+#include "Modules/ModuleManager.h"
+#include "MovieScene.h"
+#include "MovieSceneSectionDetailsCustomization.h"
+#include "MovieSceneSequence.h"
+#include "PropertyEditorModule.h"
+#include "SSequencer.h"
+#include "Sequencer.h"
+#include "SequencerContextMenus.h"
+#include "SequencerSelectedKey.h"
+#include "Styling/CoreStyle.h"
+
+void SequencerHelpers::GetAllChannels(TSharedPtr<FViewModel> DataModel, TSet<TSharedPtr<UE::Sequencer::FChannelModel>>& Channels)
+{
+	using namespace UE::Sequencer;
+
+	if (DataModel)
+	{
+		constexpr bool bIncludeThis = true;
+		for (const FViewModelPtr& Child : DataModel->GetDescendants(bIncludeThis))
+		{
+			if (TSharedPtr<ITrackAreaExtension> TrackArea = Child.ImplicitCast())
+			{
+				for (const FViewModelPtr& TrackAreaModel : TrackArea->GetTrackAreaModelList())
+				{
+					if (TViewModelPtr<FChannelModel> Channel = TrackAreaModel.ImplicitCast())
+					{
+						Channels.Add(Channel);
+					}
+				}
+			}
+			else if (TSharedPtr<FChannelModel> Channel = Child.ImplicitCast())
+			{
+				Channels.Add(Channel);
+			}
+		}
+	}
+}
+
+void SequencerHelpers::GetAllKeyAreas(TSharedPtr<FViewModel> DataModel, TSet<TSharedPtr<IKeyArea>>& Channels)
+{
+	using namespace UE::Sequencer;
+
+	if (DataModel)
+	{
+		constexpr bool bIncludeThis = true;
+		for (const FViewModelPtr& Child : DataModel->GetDescendants(bIncludeThis))
+		{
+			if (TSharedPtr<ITrackAreaExtension> TrackArea = Child.ImplicitCast())
+			{
+				for (const FViewModelPtr& TrackAreaModel : TrackArea->GetTrackAreaModelList())
+				{
+					if (TViewModelPtr<FChannelModel> Channel = TrackAreaModel.ImplicitCast())
+					{
+						Channels.Add(Channel->GetKeyArea());
+					}
+				}
+			}
+			else if (TSharedPtr<FChannelModel> Channel = Child.ImplicitCast())
+			{
+				Channels.Add(Channel->GetKeyArea());
+			}
+		}
+	}
+}
+
+void SequencerHelpers::GetAllSections(TSharedPtr<FViewModel> DataModel, TSet<TWeakObjectPtr<UMovieSceneSection>>& Sections)
+{
+	using namespace UE::Sequencer;
+
+	if (DataModel)
+	{
+		constexpr bool bIncludeThis = true;
+		for (TSharedPtr<FSectionModel> Section : TParentFirstChildIterator<FSectionModel>(DataModel, bIncludeThis))
+		{
+			Sections.Add(Section->GetSection());
+		}
+	}
+}
+
+int32 SequencerHelpers::GetSectionFromTime(TArrayView<UMovieSceneSection* const> InSections, FFrameNumber Time)
+{
+	FFrameNumber ClosestLowerBound = TNumericLimits<int32>::Max();
+	TOptional<int32> MaxOverlapPriority, MaxProximalPriority;
+
+	int32 MostRelevantIndex = INDEX_NONE;
+
+	for (int32 Index = 0; Index < InSections.Num(); ++Index)
+	{
+		const UMovieSceneSection* Section = InSections[Index];
+		if (Section)
+		{
+			const int32 ThisSectionPriority = Section->GetOverlapPriority();
+			TRange<FFrameNumber> SectionRange = Section->GetRange();
+
+			// If the specified time is within the section bounds
+			if (SectionRange.Contains(Time))
+			{
+				if (ThisSectionPriority >= MaxOverlapPriority.Get(ThisSectionPriority))
+				{
+					MaxOverlapPriority = ThisSectionPriority;
+					MostRelevantIndex = Index;
+				}
+			}
+			// Check for nearby sections if there is nothing overlapping
+			else if (!MaxOverlapPriority.IsSet() && SectionRange.HasLowerBound())
+			{
+				const FFrameNumber LowerBoundValue = SectionRange.GetLowerBoundValue();
+				// If this section exists beyond the current time, we can choose it if its closest to the time
+				if (LowerBoundValue >= Time)
+				{
+					if (
+						(LowerBoundValue < ClosestLowerBound) ||
+						(LowerBoundValue == ClosestLowerBound && ThisSectionPriority >= MaxProximalPriority.Get(ThisSectionPriority))
+						)
+					{
+						MostRelevantIndex = Index;
+						ClosestLowerBound = LowerBoundValue;
+						MaxProximalPriority = ThisSectionPriority;
+					}
+				}
+			}
+		}
+	}
+
+	// If we didn't find one, use the last one (or return -1)
+	if (MostRelevantIndex == -1)
+	{
+		MostRelevantIndex = InSections.Num() - 1;
+	}
+
+	return MostRelevantIndex;
+}
+
+void SequencerHelpers::GetDescendantNodes(TSharedRef<FViewModel> DataModel, TSet<TSharedRef<FViewModel>>& Nodes)
+{
+	using namespace UE::Sequencer;
+
+	for (TSharedPtr<FViewModel> ChildNode : DataModel->GetChildren())
+	{
+		if (ChildNode->IsA<IOutlinerExtension>())
+		{
+			Nodes.Add(ChildNode.ToSharedRef());
+		}
+
+		GetDescendantNodes(ChildNode.ToSharedRef(), Nodes);
+	}
+}
+
+bool IsSectionSelectedInNode(FSequencer& Sequencer, TSharedPtr<UE::Sequencer::FViewModel> InModel)
+{
+	using namespace UE::Sequencer;
+
+	if (ITrackAreaExtension* TrackArea = InModel->CastThis<ITrackAreaExtension>())
+	{
+		for (TSharedPtr<FViewModel> TrackAreaModel : TrackArea->GetTrackAreaModelList())
+		{
+			constexpr bool bIncludeThis = true;
+			for (TSharedPtr<FSectionModel> Section : TParentFirstChildIterator<FSectionModel>(TrackAreaModel, bIncludeThis))
+			{
+				if (Sequencer.GetSelection().IsSelected(Section))
+				{
+					return true;
+				}
+			}
+		}
+	}
+
+	return false;
+}
+
+bool AreKeysSelectedInNode(FSequencer& Sequencer, TSharedPtr<UE::Sequencer::FViewModel> InModel)
+{
+	TSet<TSharedPtr<UE::Sequencer::FChannelModel>> Channels;
+	SequencerHelpers::GetAllChannels(InModel, Channels);
+
+	for (const FSequencerSelectedKey& Key : Sequencer.GetSelection().GetSelectedKeys())
+	{
+		if (Channels.Contains(Key.WeakChannel.Pin()))
+		{
+			return true;
+		}
+	}
+
+	return false;
+}
+
+
+void SequencerHelpers::PerformDefaultSelection(FSequencer& Sequencer, const FPointerEvent& MouseEvent)
+{
+	using namespace UE::Sequencer;
+
+	// @todo: selection in transactions
+	FHotspotSelectionManager SelectionManager(&MouseEvent, &Sequencer);
+	TSharedPtr<ITrackAreaHotspot> Hotspot = Sequencer.GetViewModel()->GetTrackArea()->GetHotspot();
+	if (TSharedPtr<IMouseHandlerHotspot> MouseHandler = HotspotCast<IMouseHandlerHotspot>(Hotspot))
+	{
+		MouseHandler->HandleMouseSelection(SelectionManager);
+	}
+	else
+	{
+		// No hotspot so clear the selection if we're not adding to it
+		SelectionManager.ConditionallyClearSelection();
+	}
+}
+
+TSharedPtr<SWidget> SequencerHelpers::SummonContextMenu(FSequencer& Sequencer, const FGeometry& MyGeometry, const FPointerEvent& MouseEvent)
+{
+	using namespace UE::Sequencer;
+
+	// @todo sequencer replace with UI Commands instead of faking it
+
+	// Attempt to paste into either the current node selection, or the clicked on track
+	TSharedRef<SSequencer> SequencerWidget = StaticCastSharedRef<SSequencer>(Sequencer.GetSequencerWidget());
+	const FFrameNumber PasteAtTime = Sequencer.GetLocalTime().Time.FrameNumber;
+
+	// The menu are generated through reflection and sometime the API exposes some recursivity (think about a Widget returning it parent which is also a Widget). Just by reflection
+	// it is not possible to determine when the root object is reached. It needs a kind of simulation which is not implemented. Also, even if the recursivity was correctly handled, the possible
+	// permutations tend to grow exponentially. Until a clever solution is found, the simple approach is to disable recursively searching those menus. User can still search the current one though.
+	// See UE-131257
+	const bool bInRecursivelySearchable = false;
+
+	const bool bShouldCloseWindowAfterMenuSelection = true;
+	FMenuBuilder MenuBuilder(bShouldCloseWindowAfterMenuSelection, Sequencer.GetCommandBindings(), nullptr, false, &FCoreStyle::Get(), true, NAME_None, bInRecursivelySearchable);
+
+	TSharedPtr<ITrackAreaHotspot> Hotspot = Sequencer.GetViewModel()->GetTrackArea()->GetHotspot();
+
+	if (Hotspot.IsValid() && Hotspot->PopulateContextMenu(MenuBuilder, PasteAtTime))
+	{
+		return MenuBuilder.MakeWidget();
+	}
+	else if (Sequencer.GetClipboardStack().Num() != 0)
+	{
+		TSharedPtr<FPasteContextMenu> PasteMenu = FPasteContextMenu::CreateMenu(Sequencer, SequencerWidget->GeneratePasteArgs(PasteAtTime));
+		if (PasteMenu.IsValid() && PasteMenu->IsValidPaste())
+		{
+			PasteMenu->PopulateMenu(MenuBuilder);
+
+			return MenuBuilder.MakeWidget();
+		}
+	}
+
+	return nullptr;
+}
+
+/** A widget which wraps the section details view which is an FNotifyHook which is used to forward
+	changes to the section to sequencer. */
+class SSectionDetailsNotifyHookWrapper : public SCompoundWidget, public FNotifyHook
+{
+public:
+	SLATE_BEGIN_ARGS(SSectionDetailsNotifyHookWrapper) {}
+	SLATE_END_ARGS();
+
+	void Construct(FArguments InArgs) { }
+
+	void SetDetailsAndSequencer(TSharedRef<SWidget> InDetailsPanel, TSharedRef<ISequencer> InSequencer)
+	{
+		ChildSlot
+		[
+			InDetailsPanel
+		];
+		Sequencer = InSequencer;
+	}
+
+	//~ FNotifyHook interface
+	virtual void NotifyPostChange(const FPropertyChangedEvent& PropertyChangedEvent, FProperty* PropertyThatChanged) override
+	{
+		Sequencer->NotifyMovieSceneDataChanged(EMovieSceneDataChangeType::TrackValueChanged);
+	}
+
+private:
+	TSharedPtr<ISequencer> Sequencer;
+};
+
+
+void SequencerHelpers::AddPropertiesMenu(FSequencer& Sequencer, FMenuBuilder& MenuBuilder, const TArray<TWeakObjectPtr<UObject>>& Sections)
+{
+	using namespace UE::Sequencer;
+
+	TSharedRef<SSectionDetailsNotifyHookWrapper> DetailsNotifyWrapper = SNew(SSectionDetailsNotifyHookWrapper);
+	FDetailsViewArgs DetailsViewArgs;
+	{
+		DetailsViewArgs.bAllowSearch = false;
+		DetailsViewArgs.bCustomFilterAreaLocation = true;
+		DetailsViewArgs.bCustomNameAreaLocation = true;
+		DetailsViewArgs.bHideSelectionTip = true;
+		DetailsViewArgs.bLockable = false;
+		DetailsViewArgs.bSearchInitialKeyFocus = true;
+		DetailsViewArgs.bUpdatesFromSelection = false;
+		DetailsViewArgs.bShowOptions = false;
+		DetailsViewArgs.bShowModifiedPropertiesOption = false;
+		DetailsViewArgs.NotifyHook = &DetailsNotifyWrapper.Get();
+		DetailsViewArgs.ColumnWidth = 0.45f;
+	}
+
+	// We pass the current scene to the UMovieSceneSection customization so we can get the overall bounds of the section when we change a section from infinite->bounded.
+	UMovieScene* CurrentScene = Sequencer.GetFocusedMovieSceneSequence()->GetMovieScene();
+
+	TSharedRef<INumericTypeInterface<double>> NumericTypeInterface = Sequencer.GetNumericTypeInterface();
+
+	TSharedRef<IDetailsView> DetailsView = FModuleManager::GetModuleChecked<FPropertyEditorModule>("PropertyEditor").CreateDetailView(DetailsViewArgs);
+	DetailsView->RegisterInstancedCustomPropertyTypeLayout("FrameNumber", FOnGetPropertyTypeCustomizationInstance::CreateLambda([=]() {
+		return MakeShared<FFrameNumberDetailsCustomization>(NumericTypeInterface); }));
+	DetailsView->RegisterInstancedCustomPropertyLayout(UMovieSceneSection::StaticClass(), FOnGetDetailCustomizationInstance::CreateLambda([=]() {
+		return MakeShared<FMovieSceneSectionDetailsCustomization>(NumericTypeInterface, CurrentScene); }));
+
+	// Let section interfaces further customize the properties details view.
+	TSharedRef<FSequencerNodeTree> SequencerNodeTree = Sequencer.GetNodeTree();
+	for (TWeakObjectPtr<UObject> Section : Sections)
+	{
+		if (Section.IsValid())
+		{
+			TSharedPtr<FSectionModel> SectionHandle = SequencerNodeTree->GetSectionModel(Cast<UMovieSceneSection>(Section));
+			if (SectionHandle)
+			{
+				TSharedPtr<ISequencerSection> SectionInterface = SectionHandle->GetSectionInterface();
+				FSequencerSectionPropertyDetailsViewCustomizationParams CustomizationDetails(
+					SectionInterface.ToSharedRef(), Sequencer.AsShared(), *SectionHandle->GetParentTrackExtension()->GetTrackEditor().Get());
+				TSharedPtr<FObjectBindingModel> ParentObjectBindingNode = SectionHandle->FindAncestorOfType<FObjectBindingModel>();
+				if (ParentObjectBindingNode.IsValid())
+				{
+					CustomizationDetails.ParentObjectBindingGuid = ParentObjectBindingNode->GetObjectGuid();
+				}
+				SectionInterface->CustomizePropertiesDetailsView(DetailsView, CustomizationDetails);
+			}
+		}
+	}
+
+	Sequencer.OnInitializeDetailsPanel().Broadcast(DetailsView, Sequencer.AsShared());
+	DetailsView->SetObjects(Sections);
+
+	DetailsNotifyWrapper->SetDetailsAndSequencer(DetailsView, Sequencer.AsShared());
+	MenuBuilder.AddWidget(DetailsNotifyWrapper, FText::GetEmpty(), true);
+}
diff --git a/Engine/Source/Editor/Sequencer/Private/SequencerCommonHelpers.h b/Engine/Source/Editor/Sequencer/Private/SequencerCommonHelpers.h
--- a/Engine/Source/Editor/Sequencer/Private/SequencerCommonHelpers.h
+++ b/Engine/Source/Editor/Sequencer/Private/SequencerCommonHelpers.h
@@ -1,73 +1,76 @@
-// Copyright Epic Games, Inc. All Rights Reserved.
-
-#pragma once
-
-#include "CoreMinimal.h"
-#include "DisplayNodes/SequencerDisplayNode.h"
-#include "Widgets/Input/SNumericDropDown.h"
-#include "Containers/ArrayView.h"
-#include "Misc/FrameNumber.h"
-
-#define LOCTEXT_NAMESPACE "SequencerHelpers"
-
-class FSequencer;
-class IKeyArea;
-class UMovieSceneSection;
-
-class SequencerHelpers
-{
-public:
-	/**
-	 * Gets the key areas from the requested node
-	 */
-	static void GetAllKeyAreas(TSharedPtr<FSequencerDisplayNode> DisplayNode, TSet<TSharedPtr<IKeyArea>>& KeyAreas);
-
-	/**
-	 * Get the section index that relates to the specified time
-	 * @return the index corresponding to the highest overlapping section, or nearest section where no section overlaps the current time
-	 */
-	static int32 GetSectionFromTime(TArrayView<UMovieSceneSection* const> InSections, FFrameNumber Time);
-
-	/**
-	 * Get descendant nodes
-	 */
-	static void GetDescendantNodes(TSharedRef<FSequencerDisplayNode> DisplayNode, TSet<TSharedRef<FSequencerDisplayNode> >& Nodes);
-
-	/**
-	 * Gets all sections from the requested node
-	 */
-	static void GetAllSections(TSharedRef<FSequencerDisplayNode> DisplayNode, TSet<TWeakObjectPtr<UMovieSceneSection>>& Sections);
-
-	/**
-	 * Validate that the nodes with selected keys or sections actually are true
-	 */
-
-	static void ValidateNodesWithSelectedKeysOrSections(FSequencer& Sequencer);
-
-	/**
-	 * Update the nodes with selected sections from the hovered node
-	 */
-	static void UpdateHoveredNodeFromSelectedSections(FSequencer& Sequencer);
-
-	/**
-	 * Update the nodes with selected keys from the hovered node
-	 */
-	static void UpdateHoveredNodeFromSelectedKeys(FSequencer& Sequencer);
-
-	/**
-	 * Perform default selection for the specified mouse event, based on the current hotspot
-	 */
-	static void PerformDefaultSelection(FSequencer& Sequencer, const FPointerEvent& MouseEvent);
-	
-	/**
-	 * Attempt to summon a context menu for the current hotspot
-	 */
-	static TSharedPtr<SWidget> SummonContextMenu(FSequencer& Sequencer, const FGeometry& MyGeometry, const FPointerEvent& MouseEvent);
-
-	/*
-	 * Build a context menu for the sections
-	 */
-	static void AddPropertiesMenu(FSequencer& Sequencer, FMenuBuilder& MenuBuilder, const TArray<TWeakObjectPtr<UObject>>& Sections);
-};
-
-#undef LOCTEXT_NAMESPACE
+// Copyright Epic Games, Inc. All Rights Reserved.
+
+#pragma once
+
+#include "CoreMinimal.h"
+#include "Widgets/Input/SNumericDropDown.h"
+#include "Containers/ArrayView.h"
+#include "Misc/FrameNumber.h"
+#include "MVVM/ViewModels/ViewModel.h"
+#include "MVVM/ViewModels/ViewModelHierarchy.h"
+
+#define LOCTEXT_NAMESPACE "SequencerHelpers"
+
+class FSequencer;
+class IKeyArea;
+class UMovieSceneSection;
+
+namespace UE
+{
+namespace Sequencer
+{
+
+class FChannelModel;
+
+} // namespace Sequencer
+} // namespace UE
+
+
+class SequencerHelpers
+{
+public:
+	using FViewModel = UE::Sequencer::FViewModel;
+
+	/**
+	 * Gets the key areas from the requested node
+	 */
+	static void GetAllKeyAreas(TSharedPtr<FViewModel> DataModel, TSet<TSharedPtr<IKeyArea>>& KeyAreas);
+
+	/**
+	 * Gets the channels from the requested node
+	 */
+	static void GetAllChannels(TSharedPtr<FViewModel> DataModel, TSet<TSharedPtr<UE::Sequencer::FChannelModel>>& Channels);
+
+	/**
+	 * Get the section index that relates to the specified time
+	 * @return the index corresponding to the highest overlapping section, or nearest section where no section overlaps the current time
+	 */
+	static int32 GetSectionFromTime(TArrayView<UMovieSceneSection* const> InSections, FFrameNumber Time);
+
+	/**
+	 * Get descendant nodes
+	 */
+	static void GetDescendantNodes(TSharedRef<FViewModel> DataModel, TSet<TSharedRef<FViewModel> >& Nodes);
+
+	/**
+	 * Gets all sections from the requested node
+	 */
+	static void GetAllSections(TSharedPtr<FViewModel> DataModel, TSet<TWeakObjectPtr<UMovieSceneSection>>& Sections);
+
+	/**
+	 * Perform default selection for the specified mouse event, based on the current hotspot
+	 */
+	static void PerformDefaultSelection(FSequencer& Sequencer, const FPointerEvent& MouseEvent);
+	
+	/**
+	 * Attempt to summon a context menu for the current hotspot
+	 */
+	static TSharedPtr<SWidget> SummonContextMenu(FSequencer& Sequencer, const FGeometry& MyGeometry, const FPointerEvent& MouseEvent);
+
+	/*
+	 * Build a context menu for the sections
+	 */
+	static void AddPropertiesMenu(FSequencer& Sequencer, FMenuBuilder& MenuBuilder, const TArray<TWeakObjectPtr<UObject>>& Sections);
+};
+
+#undef LOCTEXT_NAMESPACE
diff --git a/Engine/Source/Editor/Sequencer/Private/SequencerContextMenus.cpp b/Engine/Source/Editor/Sequencer/Private/SequencerContextMenus.cpp
--- a/Engine/Source/Editor/Sequencer/Private/SequencerContextMenus.cpp
+++ b/Engine/Source/Editor/Sequencer/Private/SequencerContextMenus.cpp
@@ -1,18 +1,17 @@
 // Copyright Epic Games, Inc. All Rights Reserved.
 
 #include "SequencerContextMenus.h"
+#include "MVVM/ViewModels/ViewModel.h"
 #include "Modules/ModuleManager.h"
 #include "Styling/AppStyle.h"
-#include "DisplayNodes/SequencerSectionKeyAreaNode.h"
-#include "DisplayNodes/SequencerTrackNode.h"
-#include "DisplayNodes/SequencerObjectBindingNode.h"
 #include "SequencerCommonHelpers.h"
 #include "SequencerCommands.h"
 #include "SSequencer.h"
 #include "SectionLayout.h"
 #include "SSequencerSection.h"
 #include "SequencerSettings.h"
-#include "ISequencerHotspot.h"
+#include "MVVM/Views/ITrackAreaHotspot.h"
+#include "MVVM/ViewModels/ViewModel.h"
 #include "SequencerHotspots.h"
 #include "ScopedTransaction.h"
 #include "MovieSceneToolHelpers.h"
@@ -43,6 +42,10 @@
 #include "MovieSceneSequence.h"
 #include "MovieScene.h"
 #include "Channels/MovieSceneChannel.h"
+#include "MVVM/ViewModels/ViewModelIterators.h"
+#include "MVVM/ViewModels/SectionModel.h"
+#include "MVVM/ViewModels/TrackModel.h"
+#include "MVVM/ViewModels/SequencerEditorViewModel.h"
 #include "Tracks/MovieScenePropertyTrack.h"
 #include "Algo/AnyOf.h"
 
@@ -57,10 +60,10 @@
 	{
 		for (const FSequencerSelectedKey& Key : SelectedKeys)
 		{
-			if (Key.KeyArea.IsValid() && Key.KeyHandle.IsSet())
+			if (Key.IsValid())
 			{
-				OutKeyStruct = Key.KeyArea->GetKeyStruct(Key.KeyHandle.GetValue());
-				OutKeyStructSection = Key.KeyArea->GetOwningSection();
+				OutKeyStruct = Key.WeakChannel.Pin()->GetKeyArea()->GetKeyStruct(Key.KeyHandle);
+				OutKeyStructSection = Key.WeakChannel.Pin()->GetSection();
 				return;
 			}
 		}
@@ -71,15 +74,15 @@
 		UMovieSceneSection* CommonSection = nullptr;
 		for (const FSequencerSelectedKey& Key : SelectedKeys)
 		{
-			if (Key.KeyArea.IsValid() && Key.KeyHandle.IsSet())
+			if (Key.IsValid())
 			{
-				KeyHandles.Add(Key.KeyHandle.GetValue());
+				KeyHandles.Add(Key.KeyHandle);
 
 				if (!CommonSection)
 				{
-					CommonSection = Key.KeyArea->GetOwningSection();
+					CommonSection = Key.WeakChannel.Pin()->GetSection();
 				}
-				else if (CommonSection != Key.KeyArea->GetOwningSection())
+				else if (CommonSection != Key.WeakChannel.Pin()->GetSection())
 				{
 					CommonSection = nullptr;
 					return;
@@ -103,6 +106,8 @@
 
 void FKeyContextMenu::PopulateMenu(FMenuBuilder& MenuBuilder)
 {
+	using namespace UE::Sequencer;
+
 	FSequencer* SequencerPtr = &Sequencer.Get();
 	TSharedRef<FKeyContextMenu> Shared = AsShared();
 
@@ -152,9 +157,7 @@
 
 	MenuBuilder.BeginSection("SequencerKeyEdit", LOCTEXT("EditMenu", "Edit"));
 	{
-		TSharedPtr<ISequencerHotspot> Hotspot = SequencerPtr->GetHotspot();
-
-		if (Hotspot.IsValid() && Hotspot->GetType() == ESequencerHotspot::Key)
+		if (HotspotCast<FKeyHotspot>(SequencerPtr->GetViewModel()->GetTrackArea()->GetHotspot()))
 		{
 			MenuBuilder.AddMenuEntry(FGenericCommands::Get().Cut);
 			MenuBuilder.AddMenuEntry(FGenericCommands::Get().Copy);
@@ -620,27 +623,32 @@
 
 void FSectionContextMenu::SelectAllKeys()
 {
+	using namespace UE::Sequencer;
+
 	for (const TWeakObjectPtr<UMovieSceneSection>& WeakSection : Sequencer->GetSelection().GetSelectedSections())
 	{
 		UMovieSceneSection* Section = WeakSection.Get();
-		TOptional<FSectionHandle> SectionHandle = Sequencer->GetNodeTree()->GetSectionHandle(Section);
+		TSharedPtr<FSectionModel> SectionHandle = Sequencer->GetNodeTree()->GetSectionModel(Section);
 		if (!SectionHandle)
 		{
 			continue;
 		}
 
-		FSectionLayout Layout(*SectionHandle->GetTrackNode(), SectionHandle->GetSectionIndex());
+		FSectionLayout Layout(SectionHandle);
 		for (const FSectionLayoutElement& Element : Layout.GetElements())
 		{
-			for (TSharedRef<IKeyArea> KeyArea : Element.GetKeyAreas())
+			for (const TWeakPtr<FChannelModel>& WeakChannel : Element.GetChannels())
 			{
-				TArray<FKeyHandle> Handles;
-				KeyArea->GetKeyHandles(Handles);
+				if (TSharedPtr<FChannelModel> Channel = WeakChannel.Pin())
+				{
+					TArray<FKeyHandle> Handles;
+					Channel->GetKeyArea()->GetKeyHandles(Handles);
 
-				for (FKeyHandle KeyHandle : Handles)
-				{
-					FSequencerSelectedKey SelectKey(*Section, KeyArea, KeyHandle);
-					Sequencer->GetSelection().AddToSelection(SelectKey);
+					for (FKeyHandle KeyHandle : Handles)
+					{
+						FSequencerSelectedKey SelectKey(*Section, Channel, KeyHandle);
+						Sequencer->GetSelection().AddToSelection(SelectKey);
+					}
 				}
 			}
 		}
@@ -738,24 +746,24 @@
 
 void FSectionContextMenu::ReduceKeys()
 {
+	using namespace UE::Sequencer;
+
 	FScopedTransaction ReduceKeysTransaction(LOCTEXT("ReduceKeys_Transaction", "Reduce Keys"));
 
 	TSet<TSharedPtr<IKeyArea> > KeyAreas;
-	for (const TSharedRef<FSequencerDisplayNode>& DisplayNode : Sequencer->GetSelection().GetSelectedOutlinerNodes())
+	for (const TWeakPtr<FViewModel>& WeakItem : Sequencer->GetSelection().GetSelectedOutlinerItems())
 	{
-		SequencerHelpers::GetAllKeyAreas(DisplayNode, KeyAreas);
+		SequencerHelpers::GetAllKeyAreas(WeakItem.Pin(), KeyAreas);
 	}
 
 	if (KeyAreas.Num() == 0)
 	{
-		const TSet<TSharedRef<FSequencerDisplayNode>>& SelectedNodes = Sequencer->GetSelection().GetNodesWithSelectedKeysOrSections();
-		for (const TSharedRef<FSequencerDisplayNode>& DisplayNode : SelectedNodes)
+		for (const TWeakPtr<FViewModel>& DisplayNode : Sequencer->GetSelection().GetNodesWithSelectedKeysOrSections())
 		{
-			SequencerHelpers::GetAllKeyAreas(DisplayNode, KeyAreas);
+			SequencerHelpers::GetAllKeyAreas(DisplayNode.Pin(), KeyAreas);
 		}
 	}
 
-
 	FKeyDataOptimizationParams Params;
 	Params.bAutoSetInterpolation = true;
 	Params.Tolerance = Sequencer->GetSequencerSettings()->GetReduceKeysTolerance();
@@ -797,18 +805,19 @@
 
 bool FSectionContextMenu::CanReduceKeys() const
 {
+	using namespace UE::Sequencer;
+
 	TSet<TSharedPtr<IKeyArea> > KeyAreas;
-	for (const TSharedRef<FSequencerDisplayNode>& DisplayNode : Sequencer->GetSelection().GetSelectedOutlinerNodes())
+	for (const TWeakPtr<FViewModel>& WeakItem : Sequencer->GetSelection().GetSelectedOutlinerItems())
 	{
-		SequencerHelpers::GetAllKeyAreas(DisplayNode, KeyAreas);
+		SequencerHelpers::GetAllKeyAreas(WeakItem.Pin(), KeyAreas);
 	}
 
 	if (KeyAreas.Num() == 0)
 	{
-		const TSet<TSharedRef<FSequencerDisplayNode>>& SelectedNodes = Sequencer->GetSelection().GetNodesWithSelectedKeysOrSections();
-		for (const TSharedRef<FSequencerDisplayNode>& DisplayNode : SelectedNodes)
+		for (const TWeakPtr<FViewModel>& DisplayNode : Sequencer->GetSelection().GetNodesWithSelectedKeysOrSections())
 		{
-			SequencerHelpers::GetAllKeyAreas(DisplayNode, KeyAreas);
+			SequencerHelpers::GetAllKeyAreas(DisplayNode.Pin(), KeyAreas);
 		}
 	}
 
@@ -817,20 +826,21 @@
 
 void FSectionContextMenu::SetInterpTangentMode(ERichCurveInterpMode InterpMode, ERichCurveTangentMode TangentMode)
 {
+	using namespace UE::Sequencer;
+
 	FScopedTransaction SetInterpTangentModeTransaction(LOCTEXT("SetInterpTangentMode_Transaction", "Set Interpolation and Tangent Mode"));
 
 	TSet<TSharedPtr<IKeyArea> > KeyAreas;
-	for (const TSharedRef<FSequencerDisplayNode>& DisplayNode : Sequencer->GetSelection().GetSelectedOutlinerNodes())
+	for (const TWeakPtr<FViewModel>& WeakItem : Sequencer->GetSelection().GetSelectedOutlinerItems())
 	{
-		SequencerHelpers::GetAllKeyAreas(DisplayNode, KeyAreas);
+		SequencerHelpers::GetAllKeyAreas(WeakItem.Pin(), KeyAreas);
 	}
 
 	if (KeyAreas.Num() == 0)
 	{
-		const TSet<TSharedRef<FSequencerDisplayNode>>& SelectedNodes = Sequencer->GetSelection().GetNodesWithSelectedKeysOrSections();
-		for (const TSharedRef<FSequencerDisplayNode>& DisplayNode : SelectedNodes)
+		for (const TWeakPtr<FViewModel>& DisplayNode : Sequencer->GetSelection().GetNodesWithSelectedKeysOrSections())
 		{
-			SequencerHelpers::GetAllKeyAreas(DisplayNode, KeyAreas);
+			SequencerHelpers::GetAllKeyAreas(DisplayNode.Pin(), KeyAreas);
 		}
 	}
 
@@ -886,18 +896,19 @@
 
 bool FSectionContextMenu::CanSetInterpTangentMode() const
 {
+	using namespace UE::Sequencer;
+
 	TSet<TSharedPtr<IKeyArea> > KeyAreas;
-	for (const TSharedRef<FSequencerDisplayNode>& DisplayNode : Sequencer->GetSelection().GetSelectedOutlinerNodes())
+	for (const TWeakPtr<FViewModel>& WeakItem : Sequencer->GetSelection().GetSelectedOutlinerItems())
 	{
-		SequencerHelpers::GetAllKeyAreas(DisplayNode, KeyAreas);
+		SequencerHelpers::GetAllKeyAreas(WeakItem.Pin(), KeyAreas);
 	}
 
 	if (KeyAreas.Num() == 0)
 	{
-		const TSet<TSharedRef<FSequencerDisplayNode>>& SelectedNodes = Sequencer->GetSelection().GetNodesWithSelectedKeysOrSections();
-		for (const TSharedRef<FSequencerDisplayNode>& DisplayNode : SelectedNodes)
+		for (const TWeakPtr<FViewModel>& DisplayNode : Sequencer->GetSelection().GetNodesWithSelectedKeysOrSections())
 		{
-			SequencerHelpers::GetAllKeyAreas(DisplayNode, KeyAreas);
+			SequencerHelpers::GetAllKeyAreas(DisplayNode.Pin(), KeyAreas);
 		}
 	}
 
@@ -1301,18 +1312,16 @@
 }
 
 
-TArray<TSharedRef<FSequencerSectionKeyAreaNode>> KeyAreaNodesBuffer;
+TArray<TSharedPtr<UE::Sequencer::FChannelGroupModel>> KeyAreaNodesBuffer;
 
-void FPasteContextMenu::GatherPasteDestinationsForNode(FSequencerDisplayNode& InNode, UMovieSceneSection* InSection, const FName& CurrentScope, TMap<FName, FSequencerClipboardReconciler>& Map)
+void FPasteContextMenu::GatherPasteDestinationsForNode(const UE::Sequencer::TViewModelPtr<UE::Sequencer::IOutlinerExtension>& InNode, UMovieSceneSection* InSection, const FName& CurrentScope, TMap<FName, FSequencerClipboardReconciler>& Map)
 {
+	using namespace UE::Sequencer;
+
 	KeyAreaNodesBuffer.Reset();
-	if (InNode.GetType() == ESequencerNode::KeyArea)
+	for (const TViewModelPtr<FChannelGroupModel>& ChannelNode : InNode.AsModel()->GetDescendantsOfType<FChannelGroupModel>(true))
 	{
-		KeyAreaNodesBuffer.Add(StaticCastSharedRef<FSequencerSectionKeyAreaNode>(InNode.AsShared()));
-	}
-	else
-	{
-		InNode.GetChildKeyAreaNodesRecursively(KeyAreaNodesBuffer);
+		KeyAreaNodesBuffer.Add(ChannelNode);
 	}
 
 	if (!KeyAreaNodesBuffer.Num())
@@ -1328,7 +1337,7 @@
 			ThisScopeString.Append(CurrentScope.ToString());
 			ThisScopeString.AppendChar('.');
 		}
-		ThisScopeString.Append(InNode.GetDisplayName().ToString());
+		ThisScopeString.Append(InNode->GetIdentifier().ToString());
 		ThisScope = *ThisScopeString;
 	}
 
@@ -1339,58 +1348,43 @@
 	}
 
 	FSequencerClipboardPasteGroup Group = Reconciler->AddDestinationGroup();
-	for (const TSharedRef<FSequencerSectionKeyAreaNode>& KeyAreaNode : KeyAreaNodesBuffer)
+	for (const TSharedPtr<FChannelGroupModel>& KeyAreaNode : KeyAreaNodesBuffer)
 	{
-		TSharedPtr<IKeyArea> KeyArea = KeyAreaNode->GetKeyArea(InSection);
-		if (KeyArea.IsValid())
+		TSharedPtr<FChannelModel> Channel = KeyAreaNode->GetChannel(InSection);
+		if (Channel)
 		{
-			Group.Add(*KeyArea.Get());
+			Group.Add(Channel);
 		}
 	}
 
 	// Add children
-	for (const TSharedPtr<FSequencerDisplayNode> Child : InNode.GetChildNodes())
+	for (const TViewModelPtr<IOutlinerExtension>& Child : InNode.AsModel()->GetChildrenOfType<IOutlinerExtension>())
 	{
-		GatherPasteDestinationsForNode(*Child, InSection, ThisScope, Map);
+		GatherPasteDestinationsForNode(Child, InSection, ThisScope, Map);
 	}
 }
 
 
-void GetFullNodePath(FSequencerDisplayNode& InNode, FString& Path)
+TSharedPtr<UE::Sequencer::FTrackModel> GetTrackFromNode(const UE::Sequencer::TViewModelPtr<UE::Sequencer::IOutlinerExtension>& InNode, FString& Scope)
 {
-	TSharedPtr<FSequencerDisplayNode> Parent = InNode.GetParent();
-	if (Parent.IsValid())
+	using namespace UE::Sequencer;
+
+	if (TSharedPtr<FTrackModel> TrackModel = InNode.ImplicitCast())
 	{
-		GetFullNodePath(*Parent, Path);
+		return TrackModel;
 	}
-
-	if (!Path.IsEmpty())
-	{
-		Path.AppendChar('.');
-	}
-
-	Path.Append(InNode.GetDisplayName().ToString());
-}
-
-
-TSharedPtr<FSequencerTrackNode> GetTrackFromNode(FSequencerDisplayNode& InNode, FString& Scope)
-{
-	if (InNode.GetType() == ESequencerNode::Track)
-	{
-		return StaticCastSharedRef<FSequencerTrackNode>(InNode.AsShared());
-	}
-	else if (InNode.GetType() == ESequencerNode::Object)
+	if (TSharedPtr<FObjectBindingModel> ObjectBindingModel = InNode.ImplicitCast())
 	{
 		return nullptr;
 	}
 
-	TSharedPtr<FSequencerDisplayNode> Parent = InNode.GetParent();
-	if (Parent.IsValid())
+	TViewModelPtr<IOutlinerExtension> Parent = InNode.AsModel()->FindAncestorOfType<IOutlinerExtension>();
+	if (Parent)
 	{
-		TSharedPtr<FSequencerTrackNode> Track = GetTrackFromNode(*Parent, Scope);
+		TSharedPtr<FTrackModel> Track = GetTrackFromNode(Parent, Scope);
 		if (Track.IsValid())
 		{
-			FString ThisScope = InNode.GetDisplayName().ToString();
+			FString ThisScope = Track->GetLabel().ToString();
 			if (!Scope.IsEmpty())
 			{
 				ThisScope.AppendChar('.');
@@ -1407,6 +1401,8 @@
 
 void FPasteContextMenu::Setup()
 {
+	using namespace UE::Sequencer;
+
 	if (!Args.Clipboard.IsValid())
 	{
 		if (Sequencer->GetClipboardStack().Num() != 0)
@@ -1420,22 +1416,22 @@
 	}
 
 	// Gather a list of sections we want to paste into
-	TArray<FSectionHandle> SectionHandles;
+	TArray<TSharedPtr<FSectionModel>> SectionModels;
 
 	if (Args.DestinationNodes.Num())
 	{
 		// If we have exactly one channel to paste, first check if we have exactly one valid target channel selected to support copying between channels e.g. from Tranform.x to Transform.y
 		if (Args.Clipboard->GetKeyTrackGroups().Num() == 1)
 		{
-			for (const TSharedRef<FSequencerDisplayNode>& Node : Args.DestinationNodes)
+			for (const TViewModelPtr<IOutlinerExtension>& Node : Args.DestinationNodes)
 			{
-				if (Node->GetType() != ESequencerNode::KeyArea && Node->GetType() != ESequencerNode::Category)
+				if (!Node.AsModel()->IsA<FChannelModel>() && !Node.AsModel()->IsA<FChannelGroupModel>())
 				{
 					continue;
 				}
 
 				FString Scope;
-				TSharedPtr<FSequencerTrackNode> TrackNode = GetTrackFromNode(*Node, Scope);
+				TSharedPtr<FTrackModel> TrackNode = GetTrackFromNode(Node, Scope);
 				if (!TrackNode.IsValid())
 				{
 					continue;
@@ -1444,11 +1440,11 @@
 				FPasteDestination& Destination = PasteDestinations[PasteDestinations.AddDefaulted()];
 
 				TArray<UMovieSceneSection*> Sections;
-				for (TSharedRef<ISequencerSection> Section : TrackNode->GetSections())
+				for (const TViewModelPtr<FSectionModel>& Section : TrackNode->GetSections().IterateSubList<FSectionModel>())
 				{
-					if (Section.Get().GetSectionObject())
+					if (Section->GetSection())
 					{
-						GatherPasteDestinationsForNode(*Node, Section.Get().GetSectionObject(), NAME_None, Destination.Reconcilers);
+						GatherPasteDestinationsForNode(Node, Section->GetSection(), NAME_None, Destination.Reconcilers);
 					}
 				}
 
@@ -1487,34 +1483,29 @@
 		}
 
 		// Build a list of sections based on selected tracks
-		for (const TSharedRef<FSequencerDisplayNode>& Node : Args.DestinationNodes)
+		for (const TViewModelPtr<IOutlinerExtension>& Node : Args.DestinationNodes)
 		{
 			FString Scope;
-			TSharedPtr<FSequencerTrackNode> TrackNode = GetTrackFromNode(*Node, Scope);
+			TSharedPtr<FTrackModel> TrackNode = GetTrackFromNode(Node, Scope);
 			if (!TrackNode.IsValid())
 			{
 				continue;
 			}
 
 			TArray<UMovieSceneSection*> Sections;
-			for (TSharedRef<ISequencerSection> Section : TrackNode->GetSections())
+			for (const TViewModelPtr<FSectionModel>& Section : TrackNode->GetSections().IterateSubList<FSectionModel>())
 			{
-				if (Section.Get().GetSectionObject())
+				if (Section->GetSection())
 				{
-					Sections.Add(Section.Get().GetSectionObject());
+					Sections.Add(Section->GetSection());
 				}
 			}
 
 			UMovieSceneSection* Section = MovieSceneHelpers::FindNearestSectionAtTime(Sections, Args.PasteAtTime);
-			int32 SectionIndex = INDEX_NONE;
-			if (Section)
+			TSharedPtr<FSectionModel> SectionModel = Sequencer->GetNodeTree()->GetSectionModel(Section);
+			if (SectionModel)
 			{
-				SectionIndex = Sections.IndexOfByKey(Section);
-			}
-
-			if (SectionIndex != INDEX_NONE)
-			{
-				SectionHandles.Add(FSectionHandle(TrackNode.ToSharedRef(), SectionIndex));
+				SectionModels.Add(SectionModel);
 			}
 		}
 	}
@@ -1523,40 +1514,44 @@
 		// Use the selected sections
 		for (TWeakObjectPtr<UMovieSceneSection> WeakSection : Sequencer->GetSelection().GetSelectedSections())
 		{
-			if (TOptional<FSectionHandle> SectionHandle = Sequencer->GetNodeTree()->GetSectionHandle(WeakSection.Get()))
+			if (TSharedPtr<FSectionModel> SectionHandle = Sequencer->GetNodeTree()->GetSectionModel(WeakSection.Get()))
 			{
-				SectionHandles.Add(SectionHandle.GetValue());
+				SectionModels.Add(SectionHandle);
 			}
 		}
 	}
 
-	TMap<FName, TArray<FSectionHandle>> SectionsByType;
-	for (const FSectionHandle& Section : SectionHandles)
+	TMap<FName, TArray<TSharedPtr<FSectionModel>>> SectionsByType;
+	for (TSharedPtr<FSectionModel> SectionModel : SectionModels)
 	{
-		UMovieSceneTrack* Track = Section.GetTrackNode()->GetTrack();
+		UMovieSceneTrack* Track = SectionModel->GetParentTrackExtension()->GetTrack();
 		if (Track)
 		{
-			SectionsByType.FindOrAdd(Track->GetClass()->GetFName()).Add(Section);
+			SectionsByType.FindOrAdd(Track->GetClass()->GetFName()).Add(SectionModel);
 		}
 	}
 
-	for (const TTuple<FName, TArray<FSectionHandle>>& Pair : SectionsByType)
+	for (const TTuple<FName, TArray<TSharedPtr<FSectionModel>>>& Pair : SectionsByType)
 	{
 		FPasteDestination& Destination = PasteDestinations[PasteDestinations.AddDefaulted()];
 		if (Pair.Value.Num() == 1)
 		{
-			FString Path;
-			GetFullNodePath(*Pair.Value[0].GetTrackNode(), Path);
-			Destination.Name = FText::FromString(Path);
+			TSharedPtr<FViewModel> Model = Pair.Value[0]->FindAncestorOfTypes({ITrackExtension::ID, IOutlinerExtension::ID});
+			if (ensure(Model))
+			{
+				FString Path = IOutlinerExtension::GetPathName(Model);
+				Destination.Name = FText::FromString(Path);
+			}
 		}
 		else
 		{
 			Destination.Name = FText::Format(LOCTEXT("PasteMenuHeaderFormat", "{0} ({1} tracks)"), FText::FromName(Pair.Key), FText::AsNumber(Pair.Value.Num()));
 		}
 
-		for (const FSectionHandle& Section : Pair.Value)
+		for (TSharedPtr<FSectionModel> Section : Pair.Value)
 		{
-			GatherPasteDestinationsForNode(*Section.GetTrackNode(), Section.GetSectionObject(), NAME_None, Destination.Reconcilers);
+			FViewModelPtr Model = Section->FindAncestorOfTypes({ITrackExtension::ID, IOutlinerExtension::ID});
+			GatherPasteDestinationsForNode(Model.ImplicitCast(), Section->GetSection(), NAME_None, Destination.Reconcilers);
 		}
 
 		// Reconcile and remove invalid pastes
@@ -1675,11 +1670,11 @@
 
 	GEditor->EndTransaction();
 
-	SSequencerSection::ThrobKeySelection();
+	UE::Sequencer::SSequencerSection::ThrobKeySelection();
 
 	FSequencerSelection& Selection = Sequencer->GetSelection();
 	Selection.SuspendBroadcast();
-	Selection.EmptySelectedSections();
+	Selection.EmptySelectedTrackAreaItems();
 	Selection.EmptySelectedKeys();
 
 	for (const FSequencerSelectedKey& Key : NewSelection)
@@ -1688,7 +1683,6 @@
 	}
 	Selection.ResumeBroadcast();
 	Selection.GetOnKeySelectionChanged().Broadcast();
-	Selection.GetOnSectionSelectionChanged().Broadcast();
 
 	Sequencer->OnClipboardUsed(Args.Clipboard);
 	Sequencer->NotifyMovieSceneDataChanged(EMovieSceneDataChangeType::TrackValueChanged);
@@ -1701,8 +1695,8 @@
 	FSequencerPasteEnvironment PasteEnvironment;
 	PasteEnvironment.TickResolution = Sequencer->GetFocusedTickResolution();
 	PasteEnvironment.CardinalTime = Args.PasteAtTime;
-	PasteEnvironment.OnKeyPasted = [&](FKeyHandle Handle, IKeyArea& KeyArea){
-		NewSelection.Add(FSequencerSelectedKey(*KeyArea.GetOwningSection(), KeyArea.AsShared(), Handle));
+	PasteEnvironment.OnKeyPasted = [&](FKeyHandle Handle, TSharedPtr<UE::Sequencer::FChannelModel> Channel){
+		NewSelection.Add(FSequencerSelectedKey(*Channel->GetSection(), Channel, Handle));
 	};
 
 	return Reconciler.Paste(PasteEnvironment);
@@ -1763,7 +1757,7 @@
 	MenuBuilder.EndSection();
 }
 
-void FEasingContextMenu::BuildMenu(FMenuBuilder& MenuBuilder, const TArray<FEasingAreaHandle>& InEasings, FSequencer& Sequencer, FFrameTime InMouseDownTime)
+void FEasingContextMenu::BuildMenu(FMenuBuilder& MenuBuilder, const TArray<UE::Sequencer::FEasingAreaHandle>& InEasings, FSequencer& Sequencer, FFrameTime InMouseDownTime)
 {
 	TSharedRef<FEasingContextMenu> EasingMenu = MakeShareable(new FEasingContextMenu(InEasings, Sequencer));
 	EasingMenu->PopulateMenu(MenuBuilder);
@@ -1775,10 +1769,12 @@
 
 void FEasingContextMenu::PopulateMenu(FMenuBuilder& MenuBuilder)
 {
+	using namespace UE::Sequencer;
+
 	FText SectionText = Easings.Num() == 1 ? LOCTEXT("EasingCurve", "Easing Curve") : FText::Format(LOCTEXT("EasingCurvesFormat", "Easing Curves ({0} curves)"), FText::AsNumber(Easings.Num()));
 	const bool bReadOnly = Algo::AnyOf(Easings, [](const FEasingAreaHandle& Handle) -> bool
 		{
-			const UMovieSceneSection* Section = Handle.WeakSection.Get();
+			const UMovieSceneSection* Section = Handle.WeakSectionModel.Pin()->GetSection();
 			const UMovieSceneTrack* SectionTrack = Section->GetTypedOuter<UMovieSceneTrack>();
 			FMovieSceneSupportsEasingParams Params(Section);
 			return !EnumHasAllFlags(SectionTrack->SupportsEasing(Params), EMovieSceneTrackEasingSupportFlags::ManualEasing);
@@ -1878,11 +1874,13 @@
 
 TOptional<int32> FEasingContextMenu::GetCurrentLength() const
 {
+	using namespace UE::Sequencer;
+
 	TOptional<int32> Value;
 
 	for (const FEasingAreaHandle& Handle : Easings)
 	{
-		UMovieSceneSection* Section = Handle.WeakSection.Get();
+		UMovieSceneSection* Section = Handle.WeakSectionModel.Pin()->GetSection();
 		if (Section)
 		{
 			if (Handle.EasingType == ESequencerEasingType::In && Section->Easing.GetEaseInDuration() == Value.Get(Section->Easing.GetEaseInDuration()))
@@ -1905,9 +1903,11 @@
 
 void FEasingContextMenu::OnUpdateLength(int32 NewLength)
 {
+	using namespace UE::Sequencer;
+
 	for (const FEasingAreaHandle& Handle : Easings)
 	{
-		if (UMovieSceneSection* Section = Handle.WeakSection.Get())
+		if (UMovieSceneSection* Section = Handle.WeakSectionModel.Pin()->GetSection())
 		{
 			Section->Modify();
 			if (Handle.EasingType == ESequencerEasingType::In)
@@ -1926,10 +1926,12 @@
 
 ECheckBoxState FEasingContextMenu::GetAutoEasingCheckState() const
 {
+	using namespace UE::Sequencer;
+
 	TOptional<bool> IsChecked;
 	for (const FEasingAreaHandle& Handle : Easings)
 	{
-		if (UMovieSceneSection* Section = Handle.WeakSection.Get())
+		if (UMovieSceneSection* Section = Handle.WeakSectionModel.Pin()->GetSection())
 		{
 			if (Handle.EasingType == ESequencerEasingType::In)
 			{
@@ -1954,13 +1956,15 @@
 
 void FEasingContextMenu::SetAutoEasing(bool bAutoEasing)
 {
+	using namespace UE::Sequencer;
+
 	FScopedTransaction Transaction(LOCTEXT("SetAutoEasingText", "Set Automatic Easing"));
 
 	TArray<UMovieSceneTrack*> AllTracks;
 
 	for (const FEasingAreaHandle& Handle : Easings)
 	{
-		if (UMovieSceneSection* Section = Handle.WeakSection.Get())
+		if (UMovieSceneSection* Section = Handle.WeakSectionModel.Pin()->GetSection())
 		{
 			AllTracks.AddUnique(Section->GetTypedOuter<UMovieSceneTrack>());
 
@@ -1984,11 +1988,13 @@
 
 FText FEasingContextMenu::GetEasingTypeText() const
 {
+	using namespace UE::Sequencer;
+
 	FText CurrentText;
 	UClass* ClassType = nullptr;
 	for (const FEasingAreaHandle& Handle : Easings)
 	{
-		if (UMovieSceneSection* Section = Handle.WeakSection.Get())
+		if (UMovieSceneSection* Section = Handle.WeakSectionModel.Pin()->GetSection())
 		{
 			UObject* Object = Handle.EasingType == ESequencerEasingType::In ? Section->Easing.EaseIn.GetObject() : Section->Easing.EaseOut.GetObject();
 			if (Object)
@@ -2048,11 +2054,13 @@
 
 void FEasingContextMenu::OnEasingTypeChanged(UClass* NewClass)
 {
+	using namespace UE::Sequencer;
+
 	FScopedTransaction Transaction(LOCTEXT("SetEasingType", "Set Easing Method"));
 
 	for (const FEasingAreaHandle& Handle : Easings)
 	{
-		UMovieSceneSection* Section = Handle.WeakSection.Get();
+		UMovieSceneSection* Section = Handle.WeakSectionModel.Pin()->GetSection();
 		if (!Section)
 		{
 			continue;
@@ -2073,6 +2081,8 @@
 
 void FEasingContextMenu::EasingOptionsMenu(FMenuBuilder& MenuBuilder)
 {
+	using namespace UE::Sequencer;
+
 	FPropertyEditorModule& EditModule = FModuleManager::Get().GetModuleChecked<FPropertyEditorModule>("PropertyEditor");
 
 	FDetailsViewArgs DetailsViewArgs;
@@ -2087,7 +2097,7 @@
 	TArray<UObject*> Objects;
 	for (const FEasingAreaHandle& Handle : Easings)
 	{
-		if (UMovieSceneSection* Section = Handle.WeakSection.Get())
+		if (UMovieSceneSection* Section = Handle.WeakSectionModel.Pin()->GetSection())
 		{
 			if (Handle.EasingType == ESequencerEasingType::In)
 			{
diff --git a/Engine/Source/Editor/Sequencer/Private/SequencerContextMenus.h b/Engine/Source/Editor/Sequencer/Private/SequencerContextMenus.h
--- a/Engine/Source/Editor/Sequencer/Private/SequencerContextMenus.h
+++ b/Engine/Source/Editor/Sequencer/Private/SequencerContextMenus.h
@@ -3,14 +3,15 @@
 #pragma once
 
 #include "CoreMinimal.h"
-#include "DisplayNodes/SequencerDisplayNode.h"
+#include "SequencerHotspots.h"
 #include "Sequencer.h"
 #include "SequencerClipboardReconciler.h"
 #include "ScopedTransaction.h"
 #include "Channels/MovieSceneChannelHandle.h"
 #include "Curves/RealCurve.h"
+#include "MVVM/ViewModelPtr.h"
+#include "MVVM/Extensions/IOutlinerExtension.h"
 
-struct FEasingAreaHandle;
 struct FSequencerSelectedKey;
 class FMenuBuilder;
 class UMovieSceneSection;
@@ -98,7 +99,7 @@
 struct FPasteContextMenuArgs
 {
 	/** Paste the clipboard into the specified array of sequencer nodes, at the given time */
-	static FPasteContextMenuArgs PasteInto(TArray<TSharedRef<FSequencerDisplayNode>> InNodes, FFrameNumber InTime, TSharedPtr<FMovieSceneClipboard> InClipboard = nullptr)
+	static FPasteContextMenuArgs PasteInto(TArray<UE::Sequencer::TViewModelPtr<UE::Sequencer::IOutlinerExtension>> InNodes, FFrameNumber InTime, TSharedPtr<FMovieSceneClipboard> InClipboard = nullptr)
 	{
 		FPasteContextMenuArgs Args;
 		Args.Clipboard = InClipboard;
@@ -123,7 +124,7 @@
 	FFrameNumber PasteAtTime;
 
 	/** Optional user-supplied nodes to paste into */
-	TArray<TSharedRef<FSequencerDisplayNode>> DestinationNodes;
+	TArray<UE::Sequencer::TViewModelPtr<UE::Sequencer::IOutlinerExtension>> DestinationNodes;
 };
 
 struct FPasteContextMenu : TSharedFromThis<FPasteContextMenu>
@@ -156,7 +157,7 @@
 	bool PasteInto(int32 DestinationIndex, FName KeyAreaName, TSet<FSequencerSelectedKey>& NewSelection);
 	void EndPasteInto(bool bAnythingPasted, const TSet<FSequencerSelectedKey>& NewSelection);
 
-	void GatherPasteDestinationsForNode(FSequencerDisplayNode& InNode, UMovieSceneSection* InSection, const FName& CurrentScope, TMap<FName, FSequencerClipboardReconciler>& Map);
+	void GatherPasteDestinationsForNode(const UE::Sequencer::TViewModelPtr<UE::Sequencer::IOutlinerExtension>& InNode, UMovieSceneSection* InSection, const FName& CurrentScope, TMap<FName, FSequencerClipboardReconciler>& Map);
 
 	/** The sequencer */
 	TSharedRef<FSequencer> Sequencer;
@@ -239,11 +240,11 @@
  */
 struct FEasingContextMenu : TSharedFromThis<FEasingContextMenu>
 {
-	static void BuildMenu(FMenuBuilder& MenuBuilder, const TArray<FEasingAreaHandle>& InEasings, FSequencer& Sequencer, FFrameTime InMouseDownTime);
+	static void BuildMenu(FMenuBuilder& MenuBuilder, const TArray<UE::Sequencer::FEasingAreaHandle>& InEasings, FSequencer& Sequencer, FFrameTime InMouseDownTime);
 
 private:
 
-	FEasingContextMenu(const TArray<FEasingAreaHandle>& InEasings, FSequencer& InSequencer)
+	FEasingContextMenu(const TArray<UE::Sequencer::FEasingAreaHandle>& InEasings, FSequencer& InSequencer)
 		: Easings(InEasings)
 		, Sequencer(StaticCastSharedRef<FSequencer>(InSequencer.AsShared()))
 
@@ -270,7 +271,7 @@
 
 	void SetAutoEasing(bool bAutoEasing);
 
-	TArray<FEasingAreaHandle> Easings;
+	TArray<UE::Sequencer::FEasingAreaHandle> Easings;
 
 	/** The sequencer */
 	TSharedRef<FSequencer> Sequencer;
diff --git a/Engine/Source/Editor/Sequencer/Private/SequencerEdMode.cpp b/Engine/Source/Editor/Sequencer/Private/SequencerEdMode.cpp
--- a/Engine/Source/Editor/Sequencer/Private/SequencerEdMode.cpp
+++ b/Engine/Source/Editor/Sequencer/Private/SequencerEdMode.cpp
@@ -1,16 +1,21 @@
 // Copyright Epic Games, Inc. All Rights Reserved.
 
 #include "SequencerEdMode.h"
+#include "MVVM/ViewModels/ObjectBindingModel.h"
+#include "MVVM/ViewModels/SequenceModel.h"
+#include "MVVM/ViewModels/TrackModel.h"
+#include "MVVM/ObjectBindingModelStorageExtension.h"
+#include "MVVM/ViewModels/SequencerEditorViewModel.h"
+#include "MVVM/ViewModels/ViewModelIterators.h"
+#include "MVVM/Views/STrackAreaView.h"
 #include "EditorViewportClient.h"
 #include "Curves/KeyHandle.h"
 #include "ISequencer.h"
 #include "MovieSceneSequence.h"
 #include "MovieScene.h"
-#include "DisplayNodes/SequencerDisplayNode.h"
+#include "IKeyArea.h"
 #include "Sequencer.h"
 #include "Framework/Application/SlateApplication.h"
-#include "DisplayNodes/SequencerObjectBindingNode.h"
-#include "DisplayNodes/SequencerTrackNode.h"
 #include "Evaluation/MovieSceneEvaluationTrack.h"
 #include "SequencerCommonHelpers.h"
 #include "MovieSceneHitProxy.h"
@@ -484,6 +489,8 @@
 
 void FSequencerEdMode::OnKeySelected(FViewport* Viewport, HMovieSceneKeyProxy* KeyProxy)
 {
+	using namespace UE::Sequencer;
+
 	if (!KeyProxy)
 	{
 		return;
@@ -516,16 +523,14 @@
 			for (const FTrajectoryKey::FData& KeyData : KeyProxy->Key.KeyData)
 			{
 				UMovieSceneSection* Section = KeyData.Section.Get();
-				TOptional<FSectionHandle> SectionHandle = Sequencer->GetNodeTree()->GetSectionHandle(Section);
+				TSharedPtr<FSectionModel> SectionHandle = Sequencer->GetNodeTree()->GetSectionModel(Section);
 				if (SectionHandle && KeyData.KeyHandle.IsSet())
 				{
-					TArray<TSharedRef<FSequencerSectionKeyAreaNode>> KeyAreaNodes;
-					SectionHandle->GetTrackNode()->GetChildKeyAreaNodesRecursively(KeyAreaNodes);
-
-					for (TSharedRef<FSequencerSectionKeyAreaNode> KeyAreaNode : KeyAreaNodes)
+					TParentFirstChildIterator<FChannelGroupModel> KeyAreaNodes = SectionHandle->GetParentTrackModel().AsModel()->GetDescendantsOfType<FChannelGroupModel>();
+					for (const TViewModelPtr<FChannelGroupModel>& KeyAreaNode : KeyAreaNodes)
 					{
-						TSharedPtr<IKeyArea> KeyArea = KeyAreaNode->GetKeyArea(Section);
-						if (KeyArea.IsValid() && KeyArea->GetName() == KeyData.ChannelName)
+						TSharedPtr<FChannelModel> Channel = KeyAreaNode->GetChannel(Section);
+						if (Channel && Channel->GetKeyArea()->GetName() == KeyData.ChannelName)
 						{
 							if (!bChangedSelection)
 							{
@@ -533,7 +538,7 @@
 								bChangedSelection = true;
 							}
 
-							Sequencer->SelectKey(Section, KeyArea, KeyData.KeyHandle.GetValue(), bToggleSelection);
+							Sequencer->SelectKey(Section, Channel, KeyData.KeyHandle.GetValue(), bToggleSelection);
 							break;
 						}
 					}
@@ -597,6 +602,7 @@
 											UMovieScene3DTransformTrack* TransformTrack, TArrayView<const TWeakObjectPtr<>> BoundObjects, const bool bIsSelected)
 {
 	using namespace UE::MovieScene;
+	using namespace UE::Sequencer;
 
 	const bool bHitTesting = PDI && PDI->IsHitTesting();
 
@@ -640,7 +646,7 @@
 
 	TArray<UMovieScene3DTransformSection*> AllSectionsScratch;
 
-	FLinearColor TrackColor = FSequencerSectionPainter::BlendColor(TransformTrack->GetColorTint());
+	FLinearColor TrackColor = STrackAreaView::BlendDefaultTrackColor(TransformTrack->GetColorTint());
 	FColor       KeyColor   = TrackColor.ToFColor(true);
 
 	// Draw one line per-track (should only really ever be one)
@@ -769,6 +775,7 @@
 
 void FSequencerEdMode::DrawTracks3D(FPrimitiveDrawInterface* PDI)
 {
+	using namespace UE::Sequencer;
 
 	if (CVarUseOldSequencerMotionTrails->GetBool() == false)
 	{
@@ -791,42 +798,38 @@
 		// Gather a map of object bindings to their implict selection state
 		TMap<const FMovieSceneBinding*, bool> ObjectBindingNodesSelectionMap;
 
+		FObjectBindingModelStorageExtension* ObjectStorage = Sequencer->GetViewModel()->GetRootModel()->CastDynamic<FObjectBindingModelStorageExtension>();
+		check(ObjectStorage);
+
 		const FSequencerSelection& Selection = Sequencer->GetSelection();
-		const TSharedRef<FSequencerNodeTree>& NodeTree  = Sequencer->GetNodeTree();
+		const TSharedRef<FSequencerNodeTree>& NodeTree = Sequencer->GetNodeTree();
 		for (const FMovieSceneBinding& Binding : Sequence->GetMovieScene()->GetBindings())
 		{
-			TSharedPtr<FSequencerObjectBindingNode> ObjectBindingNode = NodeTree->FindObjectBindingNode(Binding.GetObjectGuid());
-			if (!ObjectBindingNode.IsValid())
+			TSharedPtr<FObjectBindingModel> ObjectBindingNode = ObjectStorage->FindModelForObjectBinding(Binding.GetObjectGuid());
+			if (!ObjectBindingNode)
 			{
 				continue;
 			}
 
 			bool bSelected = false;
-			auto Traverse_IsSelected = [&Selection, &bSelected](FSequencerDisplayNode& InNode)
+			for (TSharedPtr<FViewModel> Child : ObjectBindingNode->GetDescendants())
 			{
-				TSharedRef<FSequencerDisplayNode> Shared = InNode.AsShared();
-				if (Selection.IsSelected(Shared) || Selection.NodeHasSelectedKeysOrSections(Shared))
+				if (Selection.IsSelected(Child) || Selection.NodeHasSelectedKeysOrSections(Child))
 				{
 					bSelected = true;
 					// Stop traversing
-					return false;
+					break;
 				}
+			}
 
-				return true;
-			};
-
-			ObjectBindingNode->Traverse_ParentFirst(Traverse_IsSelected, true);
-
 			// If one of our parent is selected, we're considered selected
-			TSharedPtr<FSequencerDisplayNode> ParentNode = ObjectBindingNode->GetParent();
-			while (!bSelected && ParentNode.IsValid())
+			for (TSharedPtr<FViewModel> Parent : ObjectBindingNode->GetAncestors())
 			{
-				if (Selection.IsSelected(ParentNode.ToSharedRef()) || Selection.NodeHasSelectedKeysOrSections(ParentNode.ToSharedRef()))
+				if (Selection.IsSelected(Parent) || Selection.NodeHasSelectedKeysOrSections(Parent))
 				{
 					bSelected = true;
+					break;
 				}
-
-				ParentNode = ParentNode->GetParent();
 			}
 
 			ObjectBindingNodesSelectionMap.Add(&Binding, bSelected);
