Sequencer: MVVM2 branch and Layer Bars

diff --git a/Engine/Source/Editor/Sequencer/Private/SequencerHotspots.cpp b/Engine/Source/Editor/Sequencer/Private/SequencerHotspots.cpp
--- a/Engine/Source/Editor/Sequencer/Private/SequencerHotspots.cpp
+++ b/Engine/Source/Editor/Sequencer/Private/SequencerHotspots.cpp
@@ -1,214 +1,403 @@
-// Copyright Epic Games, Inc. All Rights Reserved.
-
-#include "SequencerHotspots.h"
-#include "DisplayNodes/SequencerObjectBindingNode.h"
-#include "SequencerCommonHelpers.h"
-#include "SSequencer.h"
-#include "Tools/EditToolDragOperations.h"
-#include "SequencerContextMenus.h"
-#include "SSequencerTrackArea.h"
-#include "Tools/SequencerEditTool_Movement.h"
-#include "Tools/SequencerEditTool_Selection.h"
-#include "SequencerTrackNode.h"
-#include "Widgets/Layout/SBox.h"
-#include "Styling/AppStyle.h"
-#include "Channels/MovieSceneChannel.h"
-#include "Channels/MovieSceneChannelProxy.h"
-#include "MovieSceneTimeHelpers.h"
-
-#define LOCTEXT_NAMESPACE "SequencerHotspots"
-
-
-void FKeyHotspot::UpdateOnHover(SSequencerTrackArea& InTrackArea, ISequencer& InSequencer) const
-{
-	InTrackArea.AttemptToActivateTool(FSequencerEditTool_Movement::Identifier);
-}
-
-TOptional<FFrameNumber> FKeyHotspot::GetTime() const
-{
-	FFrameNumber Time = 0;
-
-	if (Keys.Num())
-	{
-		TArrayView<const FSequencerSelectedKey> FirstKey(&Keys[0], 1);
-		TArrayView<FFrameNumber> FirstKeyTime(&Time, 1);
-		GetKeyTimes(FirstKey, FirstKeyTime);
-	}
-
-	return Time;
-}
-
-bool FKeyHotspot::PopulateContextMenu(FMenuBuilder& MenuBuilder, ISequencer& InSequencer, FFrameTime MouseDownTime)
-{
-	FSequencer& Sequencer = static_cast<FSequencer&>(InSequencer);
-	FKeyContextMenu::BuildMenu(MenuBuilder, Sequencer);
-	return true;
-}
-
-TOptional<FFrameNumber> FSectionHotspot::GetTime() const
-{
-	UMovieSceneSection* ThisSection = WeakSection.Get();
-	return ThisSection && ThisSection->HasStartFrame() ? ThisSection->GetInclusiveStartFrame() : TOptional<FFrameNumber>();
-}
-
-TOptional<FFrameTime> FSectionHotspot::GetOffsetTime() const
-{
-	UMovieSceneSection* ThisSection = WeakSection.Get();
-	return ThisSection ? ThisSection->GetOffsetTime() : TOptional<FFrameTime>();
-}
-
-void FSectionHotspot::UpdateOnHover(SSequencerTrackArea& InTrackArea, ISequencer& InSequencer) const
-{
-	UMovieSceneSection* ThisSection = WeakSection.Get();
-
-	// Move sections if they are selected
-	if (InSequencer.GetSelection().IsSelected(ThisSection))
-	{
-		InTrackArea.AttemptToActivateTool(FSequencerEditTool_Movement::Identifier);
-	}
-	else if (ThisSection)
-	{
-		// Activate selection mode if the section has keys
-		for (const FMovieSceneChannelEntry& Entry : ThisSection->GetChannelProxy().GetAllEntries())
-		{
-			for (const FMovieSceneChannel* Channel : Entry.GetChannels())
-			{
-				if (Channel && Channel->GetNumKeys() != 0)
-				{
-					InTrackArea.AttemptToActivateTool(FSequencerEditTool_Selection::Identifier);
-					return;
-				}
-			}
-		}
-
-		// Activate selection mode if the section is infinite, otherwise just move it
-		if (ThisSection->GetRange() == TRange<FFrameNumber>::All())
-		{
-			InTrackArea.AttemptToActivateTool(FSequencerEditTool_Selection::Identifier);
-		}
-		else
-		{
-			InTrackArea.AttemptToActivateTool(FSequencerEditTool_Movement::Identifier);
-		}
-	}
-}
-
-bool FSectionHotspot::PopulateContextMenu(FMenuBuilder& MenuBuilder, ISequencer& InSequencer, FFrameTime MouseDownTime)
-{
-	UMovieSceneSection* ThisSection = WeakSection.Get();
-	if (ThisSection)
-	{
-		FSequencer& Sequencer = static_cast<FSequencer&>(InSequencer);
-		FSectionContextMenu::BuildMenu(MenuBuilder, Sequencer, MouseDownTime);
-
-		TOptional<FSectionHandle> SectionHandle = Sequencer.GetNodeTree()->GetSectionHandle(ThisSection);
-		if (SectionHandle.IsSet())
-		{
-			FGuid ObjectBinding = SectionHandle->GetTrackNode()->GetObjectGuid();
-			SectionHandle->GetSectionInterface()->BuildSectionContextMenu(MenuBuilder, ObjectBinding);
-		}
-	}
-
-	return true;
-}
-
-TOptional<FFrameNumber> FSectionResizeHotspot::GetTime() const
-{
-	UMovieSceneSection* ThisSection = WeakSection.Get();
-	if (!ThisSection)
-	{
-		return TOptional<FFrameNumber>();
-	}
-	return HandleType == Left ? ThisSection->GetInclusiveStartFrame() : ThisSection->GetExclusiveEndFrame();
-}
-
-void FSectionResizeHotspot::UpdateOnHover(SSequencerTrackArea& InTrackArea, ISequencer& InSequencer) const
-{
-	InTrackArea.AttemptToActivateTool(FSequencerEditTool_Movement::Identifier);
-}
-
-TSharedPtr<ISequencerEditToolDragOperation> FSectionResizeHotspot::InitiateDrag(ISequencer& Sequencer)
-{
-	UMovieSceneSection* ThisSection = WeakSection.Get();
-
-	if (ThisSection && !Sequencer.GetSelection().GetSelectedSections().Contains(ThisSection))
-	{
-		Sequencer.GetSelection().Empty();
-		Sequencer.GetSelection().AddToSelection(ThisSection);
-		SequencerHelpers::UpdateHoveredNodeFromSelectedSections(static_cast<FSequencer&>(Sequencer));
-	}
-
-	const bool bIsSlipping = false;
-	return MakeShareable( new FResizeSection(static_cast<FSequencer&>(Sequencer), Sequencer.GetSelection().GetSelectedSections(), HandleType == Right, bIsSlipping) );
-}
-
-const FSlateBrush* FSectionResizeHotspot::GetCursorDecorator(const FGeometry& MyGeometry, const FPointerEvent& CursorEvent) const
-{
-	if (CursorEvent.IsControlDown())
-	{
-		return FAppStyle::Get().GetBrush(TEXT("Sequencer.CursorDecorator_Retime"));
-	}
-	else
-	{
-		return ISequencerHotspot::GetCursorDecorator(MyGeometry, CursorEvent);
-	}
-}
-
-TOptional<FFrameNumber> FSectionEasingHandleHotspot::GetTime() const
-{
-	UMovieSceneSection* ThisSection = WeakSection.Get();
-	if (ThisSection)
-	{
-		if (HandleType == ESequencerEasingType::In && !ThisSection->GetEaseInRange().IsEmpty())
-		{
-			return UE::MovieScene::DiscreteExclusiveUpper(ThisSection->GetEaseInRange());
-		}
-		else if (HandleType == ESequencerEasingType::Out && !ThisSection->GetEaseOutRange().IsEmpty())
-		{
-			return UE::MovieScene::DiscreteInclusiveLower(ThisSection->GetEaseOutRange());
-		}
-	}
-	return TOptional<FFrameNumber>();
-}
-
-void FSectionEasingHandleHotspot::UpdateOnHover(SSequencerTrackArea& InTrackArea, ISequencer& InSequencer) const
-{
-	InTrackArea.AttemptToActivateTool(FSequencerEditTool_Movement::Identifier);
-}
-
-bool FSectionEasingHandleHotspot::PopulateContextMenu(FMenuBuilder& MenuBuilder, ISequencer& Sequencer, FFrameTime MouseDownTime)
-{
-	FEasingContextMenu::BuildMenu(MenuBuilder, { FEasingAreaHandle{WeakSection, HandleType} }, static_cast<FSequencer&>(Sequencer), MouseDownTime);
-	return true;
-}
-
-TSharedPtr<ISequencerEditToolDragOperation> FSectionEasingHandleHotspot::InitiateDrag(ISequencer& Sequencer)
-{
-	return MakeShareable( new FManipulateSectionEasing(static_cast<FSequencer&>(Sequencer), WeakSection, HandleType == ESequencerEasingType::In) );
-}
-
-const FSlateBrush* FSectionEasingHandleHotspot::GetCursorDecorator(const FGeometry& MyGeometry, const FPointerEvent& CursorEvent) const
-{
-	return FAppStyle::Get().GetBrush(TEXT("Sequencer.CursorDecorator_EasingHandle"));
-}
-
-bool FSectionEasingAreaHotspot::PopulateContextMenu(FMenuBuilder& MenuBuilder, ISequencer& InSequencer, FFrameTime MouseDownTime)
-{
-	FSequencer& Sequencer = static_cast<FSequencer&>(InSequencer);
-	FEasingContextMenu::BuildMenu(MenuBuilder, Easings, Sequencer, MouseDownTime);
-
-	UMovieSceneSection* ThisSection = WeakSection.Get();
-	if (ThisSection)
-	{
-		TOptional<FSectionHandle> SectionHandle = Sequencer.GetNodeTree()->GetSectionHandle(ThisSection);
-		if (SectionHandle.IsSet())
-		{
-			FGuid ObjectBinding = SectionHandle->GetTrackNode()->GetObjectGuid();
-			SectionHandle->GetSectionInterface()->BuildSectionContextMenu(MenuBuilder, ObjectBinding);
-		}
-	}
-
-	return true;
-}
-
-#undef LOCTEXT_NAMESPACE
+// Copyright Epic Games, Inc. All Rights Reserved.
+
+#include "SequencerHotspots.h"
+#include "MVVM/Extensions/IObjectBindingExtension.h"
+#include "SequencerCommonHelpers.h"
+#include "SSequencer.h"
+#include "Tools/EditToolDragOperations.h"
+#include "SequencerContextMenus.h"
+#include "MVVM/Views/STrackAreaView.h"
+#include "Tools/SequencerEditTool_Movement.h"
+#include "Tools/SequencerEditTool_Selection.h"
+#include "Widgets/Layout/SBox.h"
+#include "Styling/AppStyle.h"
+#include "Channels/MovieSceneChannel.h"
+#include "Channels/MovieSceneChannelProxy.h"
+#include "MovieSceneTimeHelpers.h"
+#include "SequencerCommonHelpers.h"
+
+#include "MVVM/ViewModels/SectionModel.h"
+
+#define LOCTEXT_NAMESPACE "SequencerHotspots"
+
+namespace UE
+{
+namespace Sequencer
+{
+
+UE_SEQUENCER_DEFINE_CASTABLE(FKeyHotspot);
+UE_SEQUENCER_DEFINE_CASTABLE(FSectionEasingAreaHotspot);
+UE_SEQUENCER_DEFINE_CASTABLE(FSectionEasingHandleHotspot);
+UE_SEQUENCER_DEFINE_CASTABLE(FSectionHotspot);
+UE_SEQUENCER_DEFINE_CASTABLE(FSectionHotspotBase);
+UE_SEQUENCER_DEFINE_CASTABLE(FSectionResizeHotspot);
+
+UE_SEQUENCER_DEFINE_VIEW_MODEL_TYPE_ID(IMouseHandlerHotspot);
+
+FHotspotSelectionManager::FHotspotSelectionManager(const FPointerEvent* InMouseEvent, FSequencer* InSequencer)
+	: MouseEvent(InMouseEvent)
+	, Selection(&InSequencer->GetSelection())
+	, Sequencer(InSequencer)
+{
+	Selection->SuspendBroadcast();
+
+	bForceSelect = !MouseEvent->IsControlDown();
+	bAddingToSelection = MouseEvent->IsShiftDown() || MouseEvent->IsControlDown();
+
+	if (MouseEvent->GetEffectingButton() != EKeys::RightMouseButton)
+	{
+		// When single-clicking without the RMB, we always wipe the current selection
+		ConditionallyClearSelection();
+	}
+}
+
+FHotspotSelectionManager::~FHotspotSelectionManager()
+{
+	Selection->ResumeBroadcast();
+}
+
+void FHotspotSelectionManager::ConditionallyClearSelection()
+{
+	if (!bAddingToSelection)
+	{
+		Selection->EmptySelectedTrackAreaItems();
+		Selection->EmptySelectedKeys();
+
+		bAddingToSelection = true;
+	}
+}
+
+void FHotspotSelectionManager::ToggleKeys(TArrayView<const FSequencerSelectedKey> InKeys)
+{
+	for (const FSequencerSelectedKey& Key : InKeys)
+	{
+		const bool bIsSelected = Selection->IsSelected(Key);
+		if (bIsSelected && bForceSelect)
+		{
+			continue;
+		}
+
+		if (!bIsSelected)
+		{
+			Selection->AddToSelection(Key);
+		}
+		else
+		{
+			Selection->RemoveFromSelection(Key);
+		}
+	}
+}
+
+void FHotspotSelectionManager::ToggleModel(TSharedPtr<FViewModel> InModel)
+{
+	const bool bIsSelected = Selection->IsSelected(InModel);
+	if (bIsSelected && bForceSelect)
+	{
+		return;
+	}
+
+	TSharedPtr<ISelectableExtension> Selectable = InModel->CastThisShared<ISelectableExtension>();
+	if (!Selectable)
+	{
+		return;
+	}
+	else if (MouseEvent->GetEffectingButton() == EKeys::RightMouseButton && !EnumHasAnyFlags(Selectable->IsSelectable(), ESelectionIntent::ContextMenu))
+	{
+		return;
+	}
+	else if (MouseEvent->GetEffectingButton() == EKeys::LeftMouseButton && !EnumHasAnyFlags(Selectable->IsSelectable(), ESelectionIntent::PersistentSelection))
+	{
+		return;
+	}
+
+	if (!bIsSelected)
+	{
+		Selection->AddToSelection(InModel);
+	}
+	else
+	{
+		Selection->RemoveFromSelection(InModel);
+	}
+}
+
+void FHotspotSelectionManager::SelectKeysExclusive(TArrayView<const FSequencerSelectedKey> InKeys)
+{
+	for (const FSequencerSelectedKey& Key : InKeys)
+	{
+		if (!Selection->IsSelected(Key))
+		{
+			ConditionallyClearSelection();
+			Selection->AddToSelection(Key);
+		}
+	}
+}
+
+void FHotspotSelectionManager::SelectModelExclusive(TSharedPtr<FViewModel> InModel)
+{
+	if (!Selection->IsSelected(InModel))
+	{
+		ConditionallyClearSelection();
+		Selection->AddToSelection(InModel);
+	}
+}
+
+void FKeyHotspot::HandleMouseSelection(FHotspotSelectionManager& SelectionManager)
+{
+	if (SelectionManager.MouseEvent->GetEffectingButton() == EKeys::RightMouseButton)
+	{
+		SelectionManager.SelectKeysExclusive(Keys);
+	}
+	else
+	{
+		SelectionManager.ToggleKeys(Keys);
+	}
+}
+
+void FKeyHotspot::UpdateOnHover(FTrackAreaViewModel& InTrackArea) const
+{
+	InTrackArea.AttemptToActivateTool(FSequencerEditTool_Movement::Identifier);
+}
+
+TOptional<FFrameNumber> FKeyHotspot::GetTime() const
+{
+	FFrameNumber Time = 0;
+
+	if (Keys.Num())
+	{
+		TArrayView<const FSequencerSelectedKey> FirstKey(&Keys[0], 1);
+		TArrayView<FFrameNumber> FirstKeyTime(&Time, 1);
+		GetKeyTimes(FirstKey, FirstKeyTime);
+	}
+
+	return Time;
+}
+
+bool FKeyHotspot::PopulateContextMenu(FMenuBuilder& MenuBuilder, FFrameTime MouseDownTime)
+{
+	if (TSharedPtr<FSequencer> Sequencer = WeakSequencer.Pin())
+	{
+		FKeyContextMenu::BuildMenu(MenuBuilder, *Sequencer);
+	}
+	return true;
+}
+
+UMovieSceneSection* FSectionHotspotBase::GetSection() const
+{
+	if (TSharedPtr<FSectionModel> Model = WeakSectionModel.Pin())
+	{
+		return Model->GetSection();
+	}
+	return nullptr;
+}
+
+void FSectionHotspotBase::HandleMouseSelection(FHotspotSelectionManager& SelectionManager)
+{
+	// Base-class only handles RMB selection so that the other handles and interactive controls
+	// that act as hotspots and still operate correctly with Left click
+	TSharedPtr<FSectionModel> Section = WeakSectionModel.Pin();
+	if (Section && SelectionManager.MouseEvent->GetEffectingButton() == EKeys::RightMouseButton)
+	{
+		SelectionManager.SelectModelExclusive(Section);
+	}
+}
+
+TOptional<FFrameNumber> FSectionHotspotBase::GetTime() const
+{
+	UMovieSceneSection* ThisSection = GetSection();
+	return ThisSection && ThisSection->HasStartFrame() ? ThisSection->GetInclusiveStartFrame() : TOptional<FFrameNumber>();
+}
+
+TOptional<FFrameTime> FSectionHotspotBase::GetOffsetTime() const
+{
+	UMovieSceneSection* ThisSection = GetSection();
+	return ThisSection ? ThisSection->GetOffsetTime() : TOptional<FFrameTime>();
+}
+
+void FSectionHotspotBase::UpdateOnHover(FTrackAreaViewModel& InTrackArea) const
+{
+	UMovieSceneSection* ThisSection = GetSection();
+
+	// Move sections if they are selected
+	TSharedPtr<ISequencer> Sequencer = WeakSequencer.Pin();
+	if (Sequencer->GetSelection().IsSelected(WeakSectionModel))
+	{
+		InTrackArea.AttemptToActivateTool(FSequencerEditTool_Movement::Identifier);
+	}
+	else if (ThisSection)
+	{
+		// Activate selection mode if the section has keys
+		for (const FMovieSceneChannelEntry& Entry : ThisSection->GetChannelProxy().GetAllEntries())
+		{
+			for (const FMovieSceneChannel* Channel : Entry.GetChannels())
+			{
+				if (Channel && Channel->GetNumKeys() != 0)
+				{
+					InTrackArea.AttemptToActivateTool(FSequencerEditTool_Selection::Identifier);
+					return;
+				}
+			}
+		}
+
+		// Activate selection mode if the section is infinite, otherwise just move it
+		if (ThisSection->GetRange() == TRange<FFrameNumber>::All())
+		{
+			InTrackArea.AttemptToActivateTool(FSequencerEditTool_Selection::Identifier);
+		}
+		else
+		{
+			InTrackArea.AttemptToActivateTool(FSequencerEditTool_Movement::Identifier);
+		}
+	}
+}
+
+bool FSectionHotspotBase::PopulateContextMenu(FMenuBuilder& MenuBuilder, FFrameTime MouseDownTime)
+{
+	TSharedPtr<FSectionModel> SectionModel = WeakSectionModel.Pin();
+	UMovieSceneSection*       ThisSection  = SectionModel ? SectionModel->GetSection() : nullptr;
+	if (ThisSection)
+	{
+		TSharedPtr<FSequencer> Sequencer = WeakSequencer.Pin();
+		FSectionContextMenu::BuildMenu(MenuBuilder, *Sequencer, MouseDownTime);
+
+		TSharedPtr<IObjectBindingExtension> ObjectBinding = SectionModel->FindAncestorOfType<IObjectBindingExtension>();
+		SectionModel->GetSectionInterface()->BuildSectionContextMenu(MenuBuilder, ObjectBinding ? ObjectBinding->GetObjectGuid() : FGuid());
+	}
+
+	return true;
+}
+
+void FSectionHotspot::HandleMouseSelection(FHotspotSelectionManager& SelectionManager)
+{
+	TSharedPtr<FSectionModel> Section = WeakSectionModel.Pin();
+	if (Section && SelectionManager.MouseEvent->GetEffectingButton() == EKeys::LeftMouseButton)
+	{
+		SelectionManager.ToggleModel(Section);
+	}
+	else
+	{
+		FSectionHotspotBase::HandleMouseSelection(SelectionManager);
+	}
+}
+
+TOptional<FFrameNumber> FSectionResizeHotspot::GetTime() const
+{
+	UMovieSceneSection* ThisSection = GetSection();
+	if (!ThisSection)
+	{
+		return TOptional<FFrameNumber>();
+	}
+	return HandleType == Left ? ThisSection->GetInclusiveStartFrame() : ThisSection->GetExclusiveEndFrame();
+}
+
+void FSectionResizeHotspot::UpdateOnHover(FTrackAreaViewModel& InTrackArea) const
+{
+	InTrackArea.AttemptToActivateTool(FSequencerEditTool_Movement::Identifier);
+}
+
+TSharedPtr<ISequencerEditToolDragOperation> FSectionResizeHotspot::InitiateDrag(const FPointerEvent& MouseEvent)
+{
+	TSharedPtr<FSequencer> Sequencer = WeakSequencer.Pin();
+	if (!Sequencer->GetSelection().IsSelected(WeakSectionModel))
+	{
+		Sequencer->GetSelection().Empty();
+		Sequencer->GetSelection().AddToSelection(WeakSectionModel.Pin());
+	}
+
+	const bool bIsSlipping = false;
+	return MakeShareable( new FResizeSection(*Sequencer, Sequencer->GetSelection().GetSelectedSections(), HandleType == Right, bIsSlipping) );
+}
+
+const FSlateBrush* FSectionResizeHotspot::GetCursorDecorator(const FGeometry& MyGeometry, const FPointerEvent& CursorEvent) const
+{
+	if (CursorEvent.IsControlDown())
+	{
+		return FAppStyle::Get().GetBrush(TEXT("Sequencer.CursorDecorator_Retime"));
+	}
+	else
+	{
+		return ITrackAreaHotspot::GetCursorDecorator(MyGeometry, CursorEvent);
+	}
+}
+
+TOptional<FFrameNumber> FSectionEasingHandleHotspot::GetTime() const
+{
+	UMovieSceneSection* ThisSection = GetSection();
+	if (ThisSection)
+	{
+		if (HandleType == ESequencerEasingType::In && !ThisSection->GetEaseInRange().IsEmpty())
+		{
+			return UE::MovieScene::DiscreteExclusiveUpper(ThisSection->GetEaseInRange());
+		}
+		else if (HandleType == ESequencerEasingType::Out && !ThisSection->GetEaseOutRange().IsEmpty())
+		{
+			return UE::MovieScene::DiscreteInclusiveLower(ThisSection->GetEaseOutRange());
+		}
+	}
+	return TOptional<FFrameNumber>();
+}
+
+void FSectionEasingHandleHotspot::UpdateOnHover(FTrackAreaViewModel& InTrackArea) const
+{
+	InTrackArea.AttemptToActivateTool(FSequencerEditTool_Movement::Identifier);
+}
+
+bool FSectionEasingHandleHotspot::PopulateContextMenu(FMenuBuilder& MenuBuilder, FFrameTime MouseDownTime)
+{
+	TSharedPtr<FSequencer> Sequencer = WeakSequencer.Pin();
+
+	FEasingContextMenu::BuildMenu(MenuBuilder, { FEasingAreaHandle{WeakSectionModel, HandleType} }, *Sequencer, MouseDownTime);
+	return true;
+}
+
+TSharedPtr<ISequencerEditToolDragOperation> FSectionEasingHandleHotspot::InitiateDrag(const FPointerEvent& MouseEvent)
+{
+	UMovieSceneSection* Section = GetSection();
+	TSharedPtr<ISequencer> Sequencer = WeakSequencer.Pin();
+	return MakeShareable( new FManipulateSectionEasing(*static_cast<FSequencer*>(Sequencer.Get()), Section, HandleType == ESequencerEasingType::In) );
+}
+
+const FSlateBrush* FSectionEasingHandleHotspot::GetCursorDecorator(const FGeometry& MyGeometry, const FPointerEvent& CursorEvent) const
+{
+	return FAppStyle::Get().GetBrush(TEXT("Sequencer.CursorDecorator_EasingHandle"));
+}
+
+bool FSectionEasingAreaHotspot::PopulateContextMenu(FMenuBuilder& MenuBuilder, FFrameTime MouseDownTime)
+{
+	using namespace UE::Sequencer;
+
+	TSharedPtr<FSequencer> Sequencer = WeakSequencer.Pin();
+	FEasingContextMenu::BuildMenu(MenuBuilder, Easings, *Sequencer, MouseDownTime);
+
+	TSharedPtr<FSectionModel> SectionModel = WeakSectionModel.Pin();
+	UMovieSceneSection*       ThisSection  = SectionModel ? SectionModel->GetSection() : nullptr;
+	if (ThisSection)
+	{
+		TSharedPtr<IObjectBindingExtension> ObjectBinding = SectionModel->FindAncestorOfType<IObjectBindingExtension>();
+		SectionModel->GetSectionInterface()->BuildSectionContextMenu(MenuBuilder, ObjectBinding ? ObjectBinding->GetObjectGuid() : FGuid());
+	}
+
+	return true;
+}
+
+void FSectionEasingAreaHotspot::HandleMouseSelection(FHotspotSelectionManager& SelectionManager)
+{
+	TSharedPtr<FSectionModel> Section = WeakSectionModel.Pin();
+	if (Section && SelectionManager.MouseEvent->GetEffectingButton() == EKeys::LeftMouseButton)
+	{
+		SelectionManager.ToggleModel(Section);
+	}
+	else
+	{
+		FSectionHotspotBase::HandleMouseSelection(SelectionManager);
+	}
+}
+
+bool FSectionEasingAreaHotspot::Contains(UMovieSceneSection* InSection) const
+{
+	return Easings.ContainsByPredicate([=](const FEasingAreaHandle& InHandle){ return InHandle.WeakSectionModel.Pin()->GetSection() == InSection; });
+}
+
+} // namespace Sequencer
+} // namespace UE
+
+#undef LOCTEXT_NAMESPACE
diff --git a/Engine/Source/Editor/Sequencer/Private/SequencerHotspots.h b/Engine/Source/Editor/Sequencer/Private/SequencerHotspots.h
--- a/Engine/Source/Editor/Sequencer/Private/SequencerHotspots.h
+++ b/Engine/Source/Editor/Sequencer/Private/SequencerHotspots.h
@@ -1,130 +1,199 @@
-// Copyright Epic Games, Inc. All Rights Reserved.
-
-#pragma once
-
-#include "CoreMinimal.h"
-#include "Input/CursorReply.h"
-#include "Styling/SlateBrush.h"
-#include "SequencerSelectedKey.h"
-#include "ISequencerHotspot.h"
-
-class FMenuBuilder;
-class ISequencer;
-class ISequencerEditToolDragOperation;
-class FSequencerTrackNode;
-class ISequencerSection;
-
-enum class ESequencerEasingType
-{
-	In, Out
-};
-
-/** A hotspot representing a key */
-struct FKeyHotspot
-	: ISequencerHotspot
-{
-	FKeyHotspot(TArray<FSequencerSelectedKey> InKeys)
-		: Keys(MoveTemp(InKeys))
-	{ }
-
-	virtual ESequencerHotspot GetType() const override { return ESequencerHotspot::Key; }
-	virtual void UpdateOnHover(SSequencerTrackArea& InTrackArea, ISequencer& InSequencer) const override;
-	virtual TOptional<FFrameNumber> GetTime() const override;
-	virtual bool PopulateContextMenu(FMenuBuilder& MenuBuilder, ISequencer& Sequencer, FFrameTime MouseDownTime) override;
-
-	/** The keys that are part of this hotspot */
-	TArray<FSequencerSelectedKey> Keys;
-};
-
-
-/** A hotspot representing a section */
-struct FSectionHotspot
-	: ISequencerHotspot
-{
-	FSectionHotspot(UMovieSceneSection* InSection)
-		: WeakSection(InSection)
-	{ }
-
-	virtual ESequencerHotspot GetType() const override { return ESequencerHotspot::Section; }
-	virtual void UpdateOnHover(SSequencerTrackArea& InTrackArea, ISequencer& InSequencer) const override;
-	virtual TOptional<FFrameNumber> GetTime() const override;
-	virtual TOptional<FFrameTime> GetOffsetTime() const override;
-	virtual TSharedPtr<ISequencerEditToolDragOperation> InitiateDrag(ISequencer&) override { return nullptr; }
-	virtual bool PopulateContextMenu(FMenuBuilder& MenuBuilder, ISequencer& Sequencer, FFrameTime MouseDownTime) override;
-
-	/** The section */
-	TWeakObjectPtr<UMovieSceneSection> WeakSection;
-};
-
-
-/** A hotspot representing a resize handle on a section */
-struct FSectionResizeHotspot
-	: ISequencerHotspot
-{
-	enum EHandle
-	{
-		Left,
-		Right
-	};
-
-	FSectionResizeHotspot(EHandle InHandleType, UMovieSceneSection* InSection) : WeakSection(InSection), HandleType(InHandleType) {}
-
-	virtual ESequencerHotspot GetType() const override { return HandleType == Left ? ESequencerHotspot::SectionResize_L : ESequencerHotspot::SectionResize_R; }
-	virtual void UpdateOnHover(SSequencerTrackArea& InTrackArea, ISequencer& InSequencer) const override;
-	virtual TOptional<FFrameNumber> GetTime() const override;
-	virtual TSharedPtr<ISequencerEditToolDragOperation> InitiateDrag(ISequencer& Sequencer) override;
-	virtual FCursorReply GetCursor() const { return FCursorReply::Cursor( EMouseCursor::ResizeLeftRight ); }
-	virtual const FSlateBrush* GetCursorDecorator(const FGeometry& MyGeometry, const FPointerEvent& CursorEvent) const;
-
-	/** The section */
-	TWeakObjectPtr<UMovieSceneSection> WeakSection;
-
-private:
-
-	EHandle HandleType;
-};
-
-
-/** A hotspot representing a resize handle on a section's easing */
-struct FSectionEasingHandleHotspot
-	: ISequencerHotspot
-{
-	FSectionEasingHandleHotspot(ESequencerEasingType InHandleType, UMovieSceneSection* InSection) : WeakSection(InSection), HandleType(InHandleType) {}
-
-	virtual ESequencerHotspot GetType() const override { return HandleType == ESequencerEasingType::In ? ESequencerHotspot::EaseInHandle : ESequencerHotspot::EaseOutHandle; }
-	virtual void UpdateOnHover(SSequencerTrackArea& InTrackArea, ISequencer& InSequencer) const override;
-	virtual bool PopulateContextMenu(FMenuBuilder& MenuBuilder, ISequencer& Sequencer, FFrameTime MouseDownTime) override;
-	virtual TOptional<FFrameNumber> GetTime() const override;
-	virtual TSharedPtr<ISequencerEditToolDragOperation> InitiateDrag(ISequencer& Sequencer) override;
-	virtual FCursorReply GetCursor() const { return FCursorReply::Cursor( EMouseCursor::ResizeLeftRight ); }
-	virtual const FSlateBrush* GetCursorDecorator(const FGeometry& MyGeometry, const FPointerEvent& CursorEvent) const;
-
-	/** Handle to the section */
-	TWeakObjectPtr<UMovieSceneSection> WeakSection;
-
-private:
-
-	ESequencerEasingType HandleType;
-};
-
-
-struct FEasingAreaHandle
-{
-	TWeakObjectPtr<UMovieSceneSection> WeakSection;
-	ESequencerEasingType EasingType;
-};
-
-/** A hotspot representing an easing area for multiple sections */
-struct FSectionEasingAreaHotspot
-	: FSectionHotspot
-{
-	FSectionEasingAreaHotspot(const TArray<FEasingAreaHandle>& InEasings, UMovieSceneSection* InVisibleSection) : FSectionHotspot(InVisibleSection), Easings(InEasings) {}
-
-	virtual ESequencerHotspot GetType() const override { return ESequencerHotspot::EasingArea; }
-	virtual bool PopulateContextMenu(FMenuBuilder& MenuBuilder, ISequencer& Sequencer, FFrameTime MouseDownTime) override;
-
-	bool Contains(UMovieSceneSection* InSection) const { return Easings.ContainsByPredicate([=](const FEasingAreaHandle& InHandle){ return InHandle.WeakSection == InSection; }); }
-
-	/** Handles to the easings that exist on this hotspot */
-	TArray<FEasingAreaHandle> Easings;
-};
\ No newline at end of file
+// Copyright Epic Games, Inc. All Rights Reserved.
+
+#pragma once
+
+#include "CoreMinimal.h"
+#include "Input/CursorReply.h"
+#include "Styling/SlateBrush.h"
+#include "SequencerSelectedKey.h"
+#include "MVVM/ViewModels/TrackAreaViewModel.h"
+#include "MVVM/Views/ITrackAreaHotspot.h"
+#include "Sequencer.h"
+
+class FMenuBuilder;
+class FSequencerTrackNode;
+class ISequencerSection;
+
+namespace UE
+{
+namespace Sequencer
+{
+
+class FSectionModel;
+class ISequencerEditToolDragOperation;
+
+struct FHotspotSelectionManager
+{
+	const FPointerEvent* MouseEvent;
+	FSequencerSelection* Selection;
+	FSequencer* Sequencer;
+
+	bool bForceSelect;
+	bool bAddingToSelection;
+
+	FHotspotSelectionManager(const FPointerEvent* InMouseEvent, FSequencer* InSequencer);
+	~FHotspotSelectionManager();
+
+	void ConditionallyClearSelection();
+
+	void ToggleKeys(TArrayView<const FSequencerSelectedKey> InKeys);
+	void ToggleModel(TSharedPtr<FViewModel> InModel);
+
+	void SelectKeysExclusive(TArrayView<const FSequencerSelectedKey> InKeys);
+	void SelectModelExclusive(TSharedPtr<FViewModel> InModel);
+};
+
+struct IMouseHandlerHotspot
+{
+	UE_SEQUENCER_DECLARE_VIEW_MODEL_TYPE_ID(IMouseHandlerHotspot);
+
+	virtual ~IMouseHandlerHotspot(){}
+
+	virtual void HandleMouseSelection(FHotspotSelectionManager& SelectionManager) = 0;
+};
+
+enum class ESequencerEasingType
+{
+	In, Out
+};
+
+/** A hotspot representing a key */
+struct FKeyHotspot
+	: ITrackAreaHotspot, IMouseHandlerHotspot
+{
+	UE_SEQUENCER_DECLARE_CASTABLE(FKeyHotspot, ITrackAreaHotspot, IMouseHandlerHotspot);
+
+	FKeyHotspot(TArray<FSequencerSelectedKey> InKeys, TWeakPtr<FSequencer> InWeakSequencer)
+		: Keys(MoveTemp(InKeys))
+		, WeakSequencer(InWeakSequencer)
+	{ }
+
+	virtual void UpdateOnHover(FTrackAreaViewModel& InTrackArea) const override;
+	virtual TOptional<FFrameNumber> GetTime() const override;
+	virtual bool PopulateContextMenu(FMenuBuilder& MenuBuilder, FFrameTime MouseDownTime) override;
+
+	virtual void HandleMouseSelection(FHotspotSelectionManager& SelectionManager) override;
+
+	/** The keys that are part of this hotspot */
+	TArray<FSequencerSelectedKey> Keys;
+	TWeakPtr<FSequencer> WeakSequencer;
+};
+
+/** A hotspot representing a section */
+struct FSectionHotspotBase
+	: ITrackAreaHotspot, IMouseHandlerHotspot
+{
+	UE_SEQUENCER_DECLARE_CASTABLE(FSectionHotspotBase, ITrackAreaHotspot, IMouseHandlerHotspot);
+
+	FSectionHotspotBase(TWeakPtr<FSectionModel> InSectionModel, TWeakPtr<FSequencer> InWeakSequencer)
+		: WeakSectionModel(InSectionModel)
+		, WeakSequencer(InWeakSequencer)
+	{ }
+
+	virtual void UpdateOnHover(FTrackAreaViewModel& InTrackArea) const override;
+	virtual TOptional<FFrameNumber> GetTime() const override;
+	virtual TOptional<FFrameTime> GetOffsetTime() const override;
+	virtual TSharedPtr<ISequencerEditToolDragOperation> InitiateDrag(const FPointerEvent& MouseEvent) override { return nullptr; }
+	virtual bool PopulateContextMenu(FMenuBuilder& MenuBuilder, FFrameTime MouseDownTime) override;
+
+	/** IMouseHandlerHotspot */
+	virtual void HandleMouseSelection(FHotspotSelectionManager& SelectionManager) override;
+
+	UMovieSceneSection* GetSection() const;
+
+	/** The section model */
+	TWeakPtr<FSectionModel> WeakSectionModel;
+	TWeakPtr<FSequencer> WeakSequencer;
+};
+
+/** A hotspot representing a section */
+struct FSectionHotspot : FSectionHotspotBase
+{
+	UE_SEQUENCER_DECLARE_CASTABLE(FSectionHotspot, FSectionHotspotBase);
+
+	FSectionHotspot(TWeakPtr<FSectionModel> InSectionModel, TWeakPtr<FSequencer> InWeakSequencer)
+		: FSectionHotspotBase(InSectionModel, InWeakSequencer)
+	{ }
+
+	/** IMouseHandlerHotspot */
+	virtual void HandleMouseSelection(FHotspotSelectionManager& SelectionManager) override;
+};
+
+
+/** A hotspot representing a resize handle on a section */
+struct FSectionResizeHotspot : FSectionHotspotBase
+{
+	UE_SEQUENCER_DECLARE_CASTABLE(FSectionResizeHotspot, FSectionHotspotBase);
+
+	enum EHandle
+	{
+		Left,
+		Right
+	};
+
+	FSectionResizeHotspot(EHandle InHandleType, TWeakPtr<FSectionModel> InSectionModel, TWeakPtr<FSequencer> InWeakSequencer)
+		: FSectionHotspotBase(InSectionModel, InWeakSequencer)
+		, HandleType(InHandleType)
+	{}
+
+	virtual void UpdateOnHover(FTrackAreaViewModel& InTrackArea) const override;
+	virtual TOptional<FFrameNumber> GetTime() const override;
+	virtual TSharedPtr<ISequencerEditToolDragOperation> InitiateDrag(const FPointerEvent& MouseEvent) override;
+	virtual FCursorReply GetCursor() const { return FCursorReply::Cursor( EMouseCursor::ResizeLeftRight ); }
+	virtual const FSlateBrush* GetCursorDecorator(const FGeometry& MyGeometry, const FPointerEvent& CursorEvent) const;
+
+	EHandle HandleType;
+};
+
+
+/** A hotspot representing a resize handle on a section's easing */
+struct FSectionEasingHandleHotspot : FSectionHotspotBase
+{
+	UE_SEQUENCER_DECLARE_CASTABLE(FSectionEasingHandleHotspot, FSectionHotspotBase);
+
+	FSectionEasingHandleHotspot(ESequencerEasingType InHandleType, TWeakPtr<FSectionModel> InSectionModel, TWeakPtr<FSequencer> InWeakSequencer)
+		: FSectionHotspotBase(InSectionModel, InWeakSequencer)
+		, HandleType(InHandleType)
+	{}
+
+	virtual void UpdateOnHover(FTrackAreaViewModel& InTrackArea) const override;
+	virtual bool PopulateContextMenu(FMenuBuilder& MenuBuilder, FFrameTime MouseDownTime) override;
+	virtual TOptional<FFrameNumber> GetTime() const override;
+	virtual TSharedPtr<ISequencerEditToolDragOperation> InitiateDrag(const FPointerEvent& MouseEvent) override;
+	virtual FCursorReply GetCursor() const { return FCursorReply::Cursor( EMouseCursor::ResizeLeftRight ); }
+	virtual const FSlateBrush* GetCursorDecorator(const FGeometry& MyGeometry, const FPointerEvent& CursorEvent) const;
+
+	ESequencerEasingType HandleType;
+};
+
+
+struct FEasingAreaHandle
+{
+	TWeakPtr<FSectionModel> WeakSectionModel;
+	ESequencerEasingType EasingType;
+};
+
+/** A hotspot representing an easing area for multiple sections */
+struct FSectionEasingAreaHotspot : FSectionHotspotBase
+{
+	UE_SEQUENCER_DECLARE_CASTABLE(FSectionEasingAreaHotspot, FSectionHotspotBase);
+
+	FSectionEasingAreaHotspot(const TArray<FEasingAreaHandle>& InEasings, TWeakPtr<FSectionModel> InSectionModel, TWeakPtr<FSequencer> InWeakSequencer)
+		: FSectionHotspotBase(InSectionModel, InWeakSequencer)
+		, Easings(InEasings)
+	{}
+
+	virtual bool PopulateContextMenu(FMenuBuilder& MenuBuilder, FFrameTime MouseDownTime) override;
+
+	/** IMouseHandlerHotspot */
+	virtual void HandleMouseSelection(FHotspotSelectionManager& SelectionManager) override;
+
+	bool Contains(UMovieSceneSection* InSection) const;
+
+	/** Handles to the easings that exist on this hotspot */
+	TArray<FEasingAreaHandle> Easings;
+};
+
+
+} // namespace Sequencer
+} // namespace UE
diff --git a/Engine/Source/Editor/Sequencer/Private/SequencerKeyCollection.cpp b/Engine/Source/Editor/Sequencer/Private/SequencerKeyCollection.cpp
--- a/Engine/Source/Editor/Sequencer/Private/SequencerKeyCollection.cpp
+++ b/Engine/Source/Editor/Sequencer/Private/SequencerKeyCollection.cpp
@@ -3,34 +3,21 @@
 #include "SequencerKeyCollection.h"
 #include "MovieSceneSection.h"
 #include "IKeyArea.h"
-#include "DisplayNodes/SequencerDisplayNode.h"
-#include "DisplayNodes/SequencerTrackNode.h"
-#include "DisplayNodes/SequencerSectionKeyAreaNode.h"
-#include "Sequencer.h"
-#include "SequencerSettings.h"
+#include "MVVM/ViewModels/ChannelModel.h"
+#include "MVVM/ViewModels/ViewModelIterators.h"
 
-FSequencerKeyCollectionSignature FSequencerKeyCollectionSignature::FromNodes(const TArray<FSequencerDisplayNode*>& InNodes, FFrameNumber InDuplicateThresholdTime)
+FSequencerKeyCollectionSignature FSequencerKeyCollectionSignature::FromNodes(const TArray<TSharedRef<FViewModel>>& InNodes, FFrameNumber InDuplicateThresholdTime)
 {
+	using namespace UE::Sequencer;
+
 	FSequencerKeyCollectionSignature Result;
 	Result.DuplicateThresholdTime = InDuplicateThresholdTime;
 
-	for (const FSequencerDisplayNode* Node : InNodes)
+	for (const TSharedRef<FViewModel>& Node : InNodes)
 	{
-		const FSequencerSectionKeyAreaNode* KeyAreaNode = nullptr;
-
-		check(Node);
-		if (Node->GetType() == ESequencerNode::KeyArea)
+		for (TViewModelPtr<const FChannelGroupModel> ChannelGroupModel : Node->GetChildrenOfType<FChannelGroupModel>())
 		{
-			KeyAreaNode = static_cast<const FSequencerSectionKeyAreaNode*>(Node);
-		}
-		else if (Node->GetType() == ESequencerNode::Track)
-		{
-			KeyAreaNode = static_cast<const FSequencerTrackNode*>(Node)->GetTopLevelKeyNode().Get();
-		}
-
-		if (KeyAreaNode)
-		{
-			for (const TSharedRef<IKeyArea>& KeyArea : KeyAreaNode->GetAllKeyAreas())
+			for (const TSharedRef<IKeyArea>& KeyArea : ChannelGroupModel->GetAllKeyAreas())
 			{
 				const UMovieSceneSection* Section = KeyArea->GetOwningSection();
 				Result.KeyAreaToSignature.Add(KeyArea, Section ? Section->GetSignature() : FGuid());
@@ -41,26 +28,25 @@
 	return Result;
 }
 
-FSequencerKeyCollectionSignature FSequencerKeyCollectionSignature::FromNodesRecursive(const TArray<FSequencerDisplayNode*>& InNodes, FFrameNumber InDuplicateThresholdTime)
+FSequencerKeyCollectionSignature FSequencerKeyCollectionSignature::FromNodesRecursive(const TArray<TSharedRef<FViewModel>>& InNodes, FFrameNumber InDuplicateThresholdTime)
 {
+	using namespace UE::Sequencer;
+
 	FSequencerKeyCollectionSignature Result;
 	Result.DuplicateThresholdTime = InDuplicateThresholdTime;
 
-	TArray<TSharedRef<FSequencerSectionKeyAreaNode>> AllKeyAreaNodes;
+	TArray<TSharedPtr<FChannelGroupModel>> AllKeyAreaNodes;
 	AllKeyAreaNodes.Reserve(36);
-	const bool bJustVisible = true;
-
-	for (FSequencerDisplayNode* Node : InNodes)
+	for (const TSharedRef<FViewModel>& Node : InNodes)
 	{
-		if (Node->GetType() == ESequencerNode::KeyArea)
+		const bool bIncludeThis = true;
+		for (const TViewModelPtr<FChannelGroupModel>& KeyAreaNode : Node->GetDescendantsOfType<FChannelGroupModel>(bIncludeThis))
 		{
-			AllKeyAreaNodes.Add(StaticCastSharedRef<FSequencerSectionKeyAreaNode>(Node->AsShared()));
+			AllKeyAreaNodes.Add(KeyAreaNode);
 		}
-
-		Node->GetChildKeyAreaNodesRecursively(AllKeyAreaNodes, bJustVisible);
 	}
 
-	for (const TSharedRef<FSequencerSectionKeyAreaNode>& Node : AllKeyAreaNodes)
+	for (const TSharedPtr<FChannelGroupModel>& Node : AllKeyAreaNodes)
 	{
 		for (const TSharedRef<IKeyArea>& KeyArea : Node->GetAllKeyAreas())
 		{
@@ -72,18 +58,21 @@
 	return Result;
 }
 
-FSequencerKeyCollectionSignature FSequencerKeyCollectionSignature::FromNodeRecursive(FSequencerDisplayNode& InNode, UMovieSceneSection* InSection, FFrameNumber InDuplicateThresholdTime)
+FSequencerKeyCollectionSignature FSequencerKeyCollectionSignature::FromNodeRecursive(TSharedRef<FViewModel> InNode, UMovieSceneSection* InSection, FFrameNumber InDuplicateThresholdTime)
 {
+	using namespace UE::Sequencer;
+
 	FSequencerKeyCollectionSignature Result;
 	Result.DuplicateThresholdTime = InDuplicateThresholdTime;
 
-	TArray<TSharedRef<FSequencerSectionKeyAreaNode>> AllKeyAreaNodes;
+	TArray<TSharedPtr<FChannelGroupModel>> AllKeyAreaNodes;
 	AllKeyAreaNodes.Reserve(36);
+	for (TSharedPtr<FChannelGroupModel> KeyAreaNode : InNode->GetDescendantsOfType<FChannelGroupModel>(true))
+	{
+		AllKeyAreaNodes.Add(KeyAreaNode);
+	}
 
-	const bool bJustVisible = true;
-	InNode.GetChildKeyAreaNodesRecursively(AllKeyAreaNodes, bJustVisible);
-
-	for (const auto& Node : AllKeyAreaNodes)
+	for (const TSharedPtr<FChannelGroupModel>& Node : AllKeyAreaNodes)
 	{
 		TSharedPtr<IKeyArea> KeyArea = Node->GetKeyArea(InSection);
 		if (KeyArea.IsValid())
diff --git a/Engine/Source/Editor/Sequencer/Private/SequencerKeyRenderer.cpp b/Engine/Source/Editor/Sequencer/Private/SequencerKeyRenderer.cpp
--- a/Engine/Source/Editor/Sequencer/Private/SequencerKeyRenderer.cpp
+++ b/Engine/Source/Editor/Sequencer/Private/SequencerKeyRenderer.cpp
@@ -8,6 +8,13 @@
 #include "Sequencer.h"
 #include "MovieSceneSequence.h"
 #include "MovieScene.h"
+#include "ISequencerSection.h"
+#include "SequencerHotspots.h"
+#include "MVVM/Extensions/IHoveredExtension.h"
+#include "MVVM/Extensions/IOutlinerExtension.h"
+#include "MVVM/Extensions/LinkedOutlinerExtension.h"
+#include "MVVM/ViewModels/SequencerEditorViewModel.h"
+#include "MVVM/ViewModels/SequencerTrackAreaViewModel.h"
 
 namespace UE
 {
@@ -38,7 +45,7 @@
 
 	// Gather keys for a region larger than the view range to ensure we draw keys that are only just offscreen.
 	// Compute visible range taking into account a half-frame offset for keys, plus half a key width for keys that are partially offscreen
-	TRange<FFrameNumber> SectionRange   = InPainter.Section.GetRange();
+	TRange<FFrameNumber> SectionRange   = InPainter.SectionModel->GetRange();
 	const double         HalfKeyWidth   = 0.5f * (TimeToPixelConverter.PixelToSeconds(SequencerSectionConstants::KeySize.X) - TimeToPixelConverter.PixelToSeconds(0));
 	TRange<double>       VisibleRange   = UE::MovieScene::DilateRange(Sequencer->GetViewRange(), -HalfKeyWidth, HalfKeyWidth);
 	TRange<FFrameNumber> ValidKeyRange  = Sequencer->GetSubSequenceRange().Get(MovieScene->GetPlaybackRange());
@@ -84,15 +91,18 @@
 
 FKeyRenderer::FKeyDrawBatch::FKeyDrawBatch(const FSectionLayoutElement& LayoutElement)
 {
-	for (TSharedRef<IKeyArea> KeyArea : LayoutElement.GetKeyAreas())
+	for (TWeakPtr<FChannelModel> WeakChannel : LayoutElement.GetChannels())
 	{
-		KeyDrawInfo.Emplace(KeyArea);
+		if (TSharedPtr<FChannelModel> Channel = WeakChannel.Pin())
+		{
+			KeyDrawInfo.Emplace(Channel);
+		}
 	}
 }
 
 void FKeyRenderer::FCachedKeyDrawInformation::DrawExtra(FSequencerSectionPainter& Painter, const FGeometry& KeyGeometry) const
 {
-	CachedKeyPositions.GetKeyArea()->DrawExtra(Painter, KeyGeometry);
+	CachedKeyPositions.GetChannel()->GetKeyArea()->DrawExtra(Painter, KeyGeometry);
 }
 
 FKeyRenderer::ECacheFlags FKeyRenderer::FCachedKeyDrawInformation::UpdateViewIndependentData(FFrameRate TickResolution)
@@ -127,7 +137,7 @@
 				if (HeadNum > 0)
 				{
 					NewDrawParams.SetNum(HeadNum);
-					CachedKeyPositions.GetKeyArea()->DrawKeys(HandlesInRange.Slice(0, HeadNum), NewDrawParams);
+					CachedKeyPositions.GetChannel()->GetKeyArea()->DrawKeys(HandlesInRange.Slice(0, HeadNum), NewDrawParams);
 				}
 
 				NewDrawParams.Append(DrawParams.GetData() + PreserveStart, PreserveNum);
@@ -138,7 +148,7 @@
 				if (TailNum > 0)
 				{
 					NewDrawParams.SetNum(FramesInRange.Num());
-					CachedKeyPositions.GetKeyArea()->DrawKeys(HandlesInRange.Slice(TailStart, TailNum), TArrayView<FKeyDrawParams>(NewDrawParams).Slice(TailStart, TailNum));
+					CachedKeyPositions.GetChannel()->GetKeyArea()->DrawKeys(HandlesInRange.Slice(TailStart, TailNum), TArrayView<FKeyDrawParams>(NewDrawParams).Slice(TailStart, TailNum));
 				}
 
 				DrawParams = MoveTemp(NewDrawParams);
@@ -153,7 +163,7 @@
 			if (TimesInRange.Num())
 			{
 				// Draw these keys
-				CachedKeyPositions.GetKeyArea()->DrawKeys(HandlesInRange, DrawParams);
+				CachedKeyPositions.GetChannel()->GetKeyArea()->DrawKeys(HandlesInRange, DrawParams);
 			}
 		}
 
@@ -197,12 +207,12 @@
 	FFrameRate TickResolution = Sequencer->GetFocusedTickResolution();
 	const FTimeToPixel& TimeToPixelConverter = InPainter.GetTimeConverter();
 
-	TSharedPtr<ISequencerHotspot> Hotspot = Sequencer->GetHotspot();
+	TSharedPtr<FKeyHotspot> KeyHotspot = HotspotCast<FKeyHotspot>(Sequencer->GetViewModel()->GetTrackArea()->GetHotspot());
 	TArrayView<const FSequencerSelectedKey> HoveredKeys;
 
-	if (Hotspot.IsValid() && Hotspot->GetType() == ESequencerHotspot::Key)
+	if (KeyHotspot)
 	{
-		HoveredKeys = static_cast<FKeyHotspot*>(Hotspot.Get())->Keys;
+		HoveredKeys = KeyHotspot->Keys;
 	}
 
 	const TSet<FSequencerSelectedKey>& SelectedKeys = Sequencer->GetSelection().GetSelectedKeys();
@@ -364,11 +374,11 @@
 
 				// Avoid creating FSequencerSelectedKeys unless absolutely necessary
 				FKeyHandle ThisKeyHandle = Info.HandlesInRange[Info.NextUnhandledIndex];
-				const TSharedPtr<IKeyArea>& ThisKeyArea = Info.CachedKeyPositions.GetKeyArea();
+				TSharedPtr<FChannelModel> ThisChannel = Info.CachedKeyPositions.GetChannel();
 
 				if (bHasAnySelection)
 				{
-					FSequencerSelectedKey TestKey(InPainter.Section, ThisKeyArea, ThisKeyHandle);
+					FSequencerSelectedKey TestKey(*InPainter.SectionModel->GetSection(), ThisChannel, ThisKeyHandle);
 					if (SelectedKeys.Contains(TestKey))
 					{
 						++NumSelected;
@@ -376,7 +386,7 @@
 				}
 				if (bHasAnySelectionPreview)
 				{
-					FSequencerSelectedKey TestKey(InPainter.Section, ThisKeyArea, ThisKeyHandle);
+					FSequencerSelectedKey TestKey(*InPainter.SectionModel->GetSection(), ThisChannel, ThisKeyHandle);
 
 					if (const ESelectionPreviewState* SelectionPreviewState = SelectionPreview.Find(TestKey))
 					{
@@ -386,7 +396,7 @@
 				}
 				if (bHasAnyHoveredKeys)
 				{
-					FSequencerSelectedKey TestKey(InPainter.Section, ThisKeyArea, ThisKeyHandle);
+					FSequencerSelectedKey TestKey(*InPainter.SectionModel->GetSection(), ThisChannel, ThisKeyHandle);
 					NumHovered += int32(HoveredKeys.Contains(TestKey));
 				}
 
@@ -458,8 +468,8 @@
 
 	for (FCachedKeyDrawInformation& Info : KeyDrawInfo)
 	{
-		const TSharedPtr<IKeyArea>& ThisKeyArea = Info.CachedKeyPositions.GetKeyArea();
-		const FMovieSceneChannelHandle& Channel = ThisKeyArea.Get()->GetChannel();
+		const TSharedPtr<IKeyArea>& ThisKeyArea = Info.CachedKeyPositions.GetChannel()->GetKeyArea();
+		const FMovieSceneChannelHandle& Channel = ThisKeyArea->GetChannel();
 
 		const FName ChannelTypeName = Channel.GetChannelTypeName();
 
@@ -471,7 +481,7 @@
 
 		if (ChannelTypeName == FloatChannelTypeName)
 		{
-			FMovieSceneFloatChannel* FloatChannel = ThisKeyArea.Get()->GetChannel().Cast<FMovieSceneFloatChannel>().Get();
+			FMovieSceneFloatChannel* FloatChannel = ThisKeyArea->GetChannel().Cast<FMovieSceneFloatChannel>().Get();
 			if (!FloatChannel || !FloatChannel->GetShowCurve())
 			{
 				continue;
@@ -511,7 +521,7 @@
 		}
 		else if (ChannelTypeName == DoubleChannelTypeName)
 		{
-			FMovieSceneDoubleChannel* DoubleChannel = ThisKeyArea.Get()->GetChannel().Cast<FMovieSceneDoubleChannel>().Get();
+			FMovieSceneDoubleChannel* DoubleChannel = ThisKeyArea->GetChannel().Cast<FMovieSceneDoubleChannel>().Get();
 			if (!DoubleChannel || !DoubleChannel->GetShowCurve())
 			{
 				continue;
@@ -761,12 +771,12 @@
 {
 	FGeometry KeyAreaGeometry = LayoutElement.ComputeGeometry(SectionPainter.SectionGeometry);
 
-	TArrayView<const TSharedRef<IKeyArea>> KeyAreas = LayoutElement.GetKeyAreas();
+	TArrayView<const TWeakPtr<FChannelModel>> WeakChannels = LayoutElement.GetChannels();
 
 	TOptional<FLinearColor> ChannelColor;
-	if (KeyAreas.Num() == 1 && Sequencer->GetSequencerSettings()->GetShowChannelColors())
+	if (WeakChannels.Num() == 1 && Sequencer->GetSequencerSettings()->GetShowChannelColors())
 	{
-		ChannelColor = KeyAreas[0]->GetColor();
+		ChannelColor = WeakChannels[0].Pin()->GetKeyArea()->GetColor();
 	}
 
 	FSequencerSelection& Selection = Sequencer->GetSelection();
@@ -793,19 +803,18 @@
 		); 
 	}
 
-	TSharedPtr<FSequencerDisplayNode> DisplayNode = LayoutElement.GetDisplayNode();
-	if (DisplayNode.IsValid())
+	FLinkedOutlinerExtension* LinkedOutlinerItem = LayoutElement.GetModel() ? LayoutElement.GetModel()->CastThis<FLinkedOutlinerExtension>() : nullptr;
+	TSharedPtr<FViewModel> Model = LinkedOutlinerItem ? LinkedOutlinerItem->GetLinkedOutlinerItem().AsModel() : LayoutElement.GetModel();
+	if (Model)
 	{
-		TSharedRef<FSequencerDisplayNode> DisplayNodeRef = DisplayNode.ToSharedRef();
-
 		FLinearColor HighlightColor;
 		bool bDrawHighlight = false;
-		if (Sequencer->GetSelection().NodeHasSelectedKeysOrSections(DisplayNodeRef))
+		if (Sequencer->GetSelection().NodeHasSelectedKeysOrSections(Model))
 		{
 			bDrawHighlight = true;
 			HighlightColor = FLinearColor(1.0f, 1.0f, 1.0f, 0.15f);
 		}
-		else if (DisplayNode->IsHovered())
+		else if (Model->CastThis<IHoveredExtension>())
 		{
 			bDrawHighlight = true;
 			HighlightColor = FLinearColor(1.0f, 1.0f, 1.0f, 0.05f);
@@ -827,7 +836,7 @@
 
 		// --------------------------------------------
 		// Draw display node selection tint
-		if (Selection.IsSelected(DisplayNodeRef))
+		if (Selection.IsSelected(Model))
 		{
 			FSlateDrawElement::MakeBox(
 				SectionPainter.DrawElements,
@@ -842,7 +851,7 @@
 
 	// --------------------------------------------
 	// Draw section selection tint
-	const bool bSectionSelected = Selection.IsSelected(&SectionPainter.Section);
+	const bool bSectionSelected = Selection.IsSelected(SectionPainter.SectionModel);
 	if (bSectionSelected && Args.SectionThrobValue != 0.f)
 	{
 		FSlateDrawElement::MakeBox(
@@ -894,7 +903,7 @@
 		{
 			break;
 		}
-		if (LayoutElement.GetKeyAreas().Num() == 0)
+		if (LayoutElement.GetChannels().Num() == 0)
 		{
 			continue;
 		}
diff --git a/Engine/Source/Editor/Sequencer/Private/SequencerKeyRenderer.h b/Engine/Source/Editor/Sequencer/Private/SequencerKeyRenderer.h
--- a/Engine/Source/Editor/Sequencer/Private/SequencerKeyRenderer.h
+++ b/Engine/Source/Editor/Sequencer/Private/SequencerKeyRenderer.h
@@ -119,10 +119,10 @@
 	struct FCachedKeyDrawInformation
 	{
 		/** Construction from the key area this represents */
-		FCachedKeyDrawInformation(TSharedRef<IKeyArea> KeyArea)
+		FCachedKeyDrawInformation(TSharedPtr<UE::Sequencer::FChannelModel> Channel)
 			: NextUnhandledIndex(0)
 			, PreserveToIndex(0)
-			, CachedKeyPositions(KeyArea)
+			, CachedKeyPositions(Channel)
 		{}
 
 
diff --git a/Engine/Source/Editor/Sequencer/Private/SequencerModule.cpp b/Engine/Source/Editor/Sequencer/Private/SequencerModule.cpp
--- a/Engine/Source/Editor/Sequencer/Private/SequencerModule.cpp
+++ b/Engine/Source/Editor/Sequencer/Private/SequencerModule.cpp
@@ -19,7 +19,16 @@
 #include "IDetailsView.h"
 #include "Tree/CurveEditorTreeFilter.h"
 #include "AnimatedPropertyKey.h"
+#include "MovieSceneSignedObject.h"
 
+#include "MVVM/CurveEditorIntegrationExtension.h"
+#include "MVVM/FolderModelStorageExtension.h"
+#include "MVVM/ObjectBindingModelStorageExtension.h"
+#include "MVVM/SectionModelStorageExtension.h"
+#include "MVVM/TrackModelStorageExtension.h"
+#include "MVVM/TrackRowModelStorageExtension.h"
+#include "MVVM/ViewModels/SequenceModel.h"
+
 #include "ToolMenus.h"
 #include "ContentBrowserMenuContexts.h"
 #include "SequencerUtilities.h"
@@ -27,6 +36,11 @@
 #include "LevelSequence.h"
 
 
+#include "Misc/CoreDelegates.h"
+#include "UnrealEdGlobals.h"
+#include "Editor/UnrealEdEngine.h"
+#include "Editor/TransBuffer.h"
+
 #if !IS_MONOLITHIC
 	UE::MovieScene::FEntityManager*& GEntityManagerForDebugging = UE::MovieScene::GEntityManagerForDebuggingVisualizers;
 #endif
@@ -35,6 +49,118 @@
 
 #define LOCTEXT_NAMESPACE "SequencerEditor"
 
+
+namespace UE
+{
+namespace Sequencer
+{
+
+struct FDeferredSignedObjectChangeHandler : UE::MovieScene::IDeferredSignedObjectChangeHandler, FGCObject
+{
+	FDeferredSignedObjectChangeHandler()
+	{
+		Init();
+	}
+
+	~FDeferredSignedObjectChangeHandler()
+	{
+		if (UTransBuffer* TransBuffer = WeakBuffer.Get())
+		{
+			TransBuffer->OnTransactionStateChanged().RemoveAll(this);
+		}
+	}
+
+	void Init()
+	{
+		UTransBuffer* TransBuffer = GUnrealEd ? Cast<UTransBuffer>(GUnrealEd->Trans) : nullptr;
+		if (TransBuffer)
+		{
+			WeakBuffer = TransBuffer;
+			TransBuffer->OnTransactionStateChanged().AddRaw(this, &FDeferredSignedObjectChangeHandler::OnTransactionStateChanged);
+			if (TransBuffer->IsActive())
+			{
+				DeferTransactionChanges.Emplace();
+			}
+		}
+		else
+		{
+			FCoreDelegates::OnPostEngineInit.AddLambda([this]{ this->Init(); });
+		}
+	}
+
+	void OnTransactionStateChanged(const FTransactionContext& TransactionContext, ETransactionStateEventType TransactionState)
+	{
+		/** A transaction has been started. This will be followed by a TransactionCanceled or TransactionFinalized event. */
+		switch (TransactionState)
+		{
+		case ETransactionStateEventType::TransactionStarted:
+		case ETransactionStateEventType::UndoRedoStarted:
+			DeferTransactionChanges.Emplace();
+			break;
+
+		case ETransactionStateEventType::TransactionCanceled:
+		case ETransactionStateEventType::TransactionFinalized:
+		case ETransactionStateEventType::UndoRedoFinalized:
+			DeferTransactionChanges.Reset();
+			break;
+		}
+	}
+
+	void Flush() override
+	{
+		for (TWeakObjectPtr<UMovieSceneSignedObject> WeakObject : SignedObjects)
+		{
+			if (UMovieSceneSignedObject* Object = WeakObject.Get())
+			{
+				Object->BroadcastChanged();
+			}
+		}
+		SignedObjects.Empty();
+	}
+
+	void DeferMarkAsChanged(UMovieSceneSignedObject* SignedObject) override
+	{
+		SignedObjects.Add(SignedObject);
+	}
+
+	void AddReferencedObjects( FReferenceCollector& Collector ) override
+	{
+		for (TWeakObjectPtr<UMovieSceneSignedObject> WeakObject : SignedObjects)
+		{
+			if (UMovieSceneSignedObject* Object = WeakObject.Get())
+			{
+				Collector.AddReferencedObject(Object);
+			}
+		}
+	}
+
+	bool CreateImplicitScopedModifyDefer() override
+	{
+		ensure(!DeferImplicitChanges.IsSet());
+		DeferImplicitChanges.Emplace();
+		return true;
+	}
+
+	void ResetImplicitScopedModifyDefer() override
+	{
+		DeferImplicitChanges.Reset();
+	}
+
+	FString GetReferencerName() const override
+	{
+		return TEXT("FDeferredSignedObjectChangeHandler");
+	}
+
+	TSet<TWeakObjectPtr<UMovieSceneSignedObject>> SignedObjects;
+	TWeakObjectPtr<UTransBuffer>   WeakBuffer;
+	TOptional<UE::MovieScene::FScopedSignedObjectModifyDefer> DeferTransactionChanges;
+	TOptional<UE::MovieScene::FScopedSignedObjectModifyDefer> DeferImplicitChanges;
+};
+
+} // namespace Sequencer
+} // namespace UE
+
+
 // Destructor defined in CPP to avoid having to #include SequencerChannelInterface.h in the main module definition
 ISequencerModule::~ISequencerModule()
 {
@@ -187,6 +313,17 @@
 		}
 	}
 
+	virtual FDelegateHandle RegisterTrackModel(FOnCreateTrackModel InCreator) override
+	{
+		TrackModelDelegates.Add(InCreator);
+		return TrackModelDelegates.Last().GetHandle();
+	}
+
+	virtual void UnregisterTrackModel(FDelegateHandle InHandle) override
+	{
+		TrackModelDelegates.RemoveAll([=](const FOnCreateTrackModel& Delegate) { return Delegate.GetHandle() == InHandle; });
+	}
+
 	virtual FDelegateHandle RegisterOnSequencerCreated(FOnSequencerCreated::FDelegate InOnSequencerCreated) override
 	{
 		return OnSequencerCreated.Add(InOnSequencerCreated);
@@ -277,6 +414,9 @@
 
 	virtual void StartupModule() override
 	{
+		using namespace UE::Sequencer;
+		using namespace UE::MovieScene;
+
 		if (GIsEditor)
 		{
 			FEditorModeRegistry::Get().RegisterMode<FSequencerEdMode>(
@@ -296,10 +436,24 @@
 				FCoreDelegates::OnPostEngineInit.AddRaw(this, &FSequencerModule::RegisterMenus);
 			}
 
+			UMovieSceneSignedObject::SetDeferredHandler(MakeUnique<FDeferredSignedObjectChangeHandler>());
+
 			FPropertyEditorModule& EditModule = FModuleManager::Get().GetModuleChecked<FPropertyEditorModule>("PropertyEditor");
 			OnGetGlobalRowExtensionHandle = EditModule.GetGlobalRowExtensionDelegate().AddStatic(&RegisterKeyframeExtensionHandler);
 		}
 
+		FSequenceModel::CreateExtensionsEvent.AddLambda(
+			[&](FSequenceModel* InModel)
+			{
+				InModel->AddDynamicExtension(FCurveEditorIntegrationExtension::ID);
+				InModel->AddDynamicExtension(FFolderModelStorageExtension::ID);
+				InModel->AddDynamicExtension(FObjectBindingModelStorageExtension::ID);
+				InModel->AddDynamicExtension(FTrackModelStorageExtension::ID, TrackModelDelegates);
+				InModel->AddDynamicExtension(FTrackRowModelStorageExtension::ID);
+				InModel->AddDynamicExtension(FSectionModelStorageExtension::ID);
+			}
+		);
+
 		ObjectBindingContextMenuExtensibilityManager = MakeShareable( new FExtensibilityManager );
 		AddTrackMenuExtensibilityManager = MakeShareable( new FExtensibilityManager );
 		ToolBarExtensibilityManager = MakeShareable(new FExtensibilityManager);
@@ -312,6 +466,8 @@
 	{
 		if (GIsEditor)
 		{
+			UMovieSceneSignedObject::SetDeferredHandler(nullptr);
+
 			FSequencerCommands::Unregister();
 
 			if (FPropertyEditorModule* EditModulePtr = FModuleManager::Get().GetModulePtr<FPropertyEditorModule>("PropertyEditor"))
@@ -426,6 +582,9 @@
 	/** List of object binding handler delegates sequencers will execute when they are created */
 	TArray< FOnCreateEditorObjectBinding > EditorObjectBindingDelegates;
 
+	/** List of track model creators */
+	TArray<FOnCreateTrackModel> TrackModelDelegates;
+
 	/** Global details row extension delegate; */
 	FDelegateHandle OnGetGlobalRowExtensionHandle;
 
diff --git a/Engine/Source/Editor/Sequencer/Private/SequencerNodeSortingMethods.h b/Engine/Source/Editor/Sequencer/Private/SequencerNodeSortingMethods.h
--- a/Engine/Source/Editor/Sequencer/Private/SequencerNodeSortingMethods.h
+++ b/Engine/Source/Editor/Sequencer/Private/SequencerNodeSortingMethods.h
@@ -1,101 +1,100 @@
-// Copyright Epic Games, Inc. All Rights Reserved.
-
-#pragma once
-#include "CoreMinimal.h"
-#include "DisplayNodes/SequencerDisplayNode.h"
-#include "Widgets/Views/STableRow.h"
-
-/**
-* This sorting delegate sorts based on the vertical position of the node represented by where it falls
-* in the TreeView, ignoring collapsed nodes. This allows you to get a sorting order that matches the visual
-* representation which is more logical for the user.
-*/
-struct FDisplayNodeTreePositionSorter
-{
-	bool operator()(const TSharedRef<FSequencerDisplayNode>& A, const TSharedRef <FSequencerDisplayNode>& B) const
-	{
-		return A->GetVirtualTop() < B->GetVirtualTop();
-	}
-};
-
-
-/**
-* Sorts the supplied unsorted nodes and inserts them into the existing sorted nodes before assigning
-* a sorting order to the newly combined array. Inserted nodes can optionally be relative to an existing
-* node that comes from the ExistingSortedNodes array; this allows you to sort them and then insert them
-* mid array.
-*
-* This function modifies the sorting order of all nodes passed in from both lists, but does not change
-* their actual hierarchy.
-*
-* @param UnsortedNodesToInsert A list of nodes to be sorted by Predicate
-* @param ExistingSortedNdoes A list of nodes that are already sorted that you want to have their indexes
-* adjusted for having the unsorted nodes inserted among them.
-* @param ItemDropZone Specifies the offset for the RelativeToNode (above, below or at the end). This is
-* useful as many of the places this gets called from have this information already.
-* @param Predicate Sorting predicate to sort UnsortedNodesToSort before insertion.
-* @param RelativeToNode An optional node from the ExistingSortedNodes list to insert before/after. If null
-* unsorted nodes will be inserted at the start and all nodes in ExistingSortedNodes will have their sorting
-* order offset by that many.
-*/
-template <class PREDICATE_CLASS>
-void SortAndSetSortingOrder(const TArray<TSharedRef<FSequencerDisplayNode>>& UnsortedNodesToInsert, const TArray<TSharedRef<FSequencerDisplayNode>>& ExistingSortedNodes, TOptional<EItemDropZone> ItemDropZone, const PREDICATE_CLASS& Predicate, TSharedPtr<FSequencerDisplayNode> RelativeToNode)
-{
-	TArray<TSharedRef<FSequencerDisplayNode>> NewSortingOrder = ExistingSortedNodes;
-
-	// ExistingSortedNodes may contain nodes that we do not consider for sorting order (spacers) or nodes we wish to re-sort.
-	TArray<TSharedRef<FSequencerDisplayNode>> ChildrenToRemove;
-	for (TSharedRef<FSequencerDisplayNode> Node : ExistingSortedNodes)
-	{
-		// Discard any of our children that aren't folders, objects or tracks (this ignores spacers, etc.)
-		if (Node->GetType() != ESequencerNode::Folder && Node->GetType() != ESequencerNode::Object && Node->GetType() != ESequencerNode::Track)
-		{
-			ChildrenToRemove.Add(Node);
-			continue;
-		}
-
-		// If we're trying to sort this child remove it from our searchable list as well.
-		if (UnsortedNodesToInsert.Contains(Node))
-		{
-			ChildrenToRemove.Add(Node);
-		}
-	}
-
-	for (TSharedRef<FSequencerDisplayNode> Node : ChildrenToRemove)
-	{
-		NewSortingOrder.Remove(Node);
-	}
-
-	// Now get our index and insert the dragged nodes either before/after/at the end of that index depending on the ItemDropZone.
-	int32 DropAdjustedIndex = NewSortingOrder.Num();
-	int32 RelativeToIndex = NewSortingOrder.IndexOfByKey(RelativeToNode);
-	if (RelativeToIndex >= 0 && ItemDropZone.IsSet())
-	{
-		switch (ItemDropZone.GetValue())
-		{
-		case EItemDropZone::AboveItem:
-			DropAdjustedIndex = RelativeToIndex;
-			break;
-		case EItemDropZone::BelowItem:
-			DropAdjustedIndex = RelativeToIndex + 1;
-			break;
-			// Folders will emit this drop zone which means we want to put it at the end of the folder.
-		case EItemDropZone::OntoItem:
-			DropAdjustedIndex = NewSortingOrder.Num();
-			break;
-		}
-	}
-
-	// Allow our caller to specify the sorting order of new nodes. This allows us to sort incoming nodes by category or vertical position on the treeview, etc.
-	TArray<TSharedRef<FSequencerDisplayNode>> SortedNodes = UnsortedNodesToInsert;
-	SortedNodes.Sort(Predicate);
-
-	// Insert our sorted nodes into our child list so that we have the new absolute order for all items.
-	NewSortingOrder.Insert(SortedNodes, DropAdjustedIndex);
-
-	// And then re-assign the sorting order index of everything in the folder according to the new list.
-	for (int i = 0; i < NewSortingOrder.Num(); i++)
-	{
-		NewSortingOrder[i]->ModifyAndSetSortingOrder(i);
-	}
-};
\ No newline at end of file
+// Copyright Epic Games, Inc. All Rights Reserved.
+
+#pragma once
+#include "CoreMinimal.h"
+#include "Widgets/Views/STableRow.h"
+
+/**
+* This sorting delegate sorts based on the vertical position of the node represented by where it falls
+* in the TreeView, ignoring collapsed nodes. This allows you to get a sorting order that matches the visual
+* representation which is more logical for the user.
+*/
+//struct FDisplayNodeTreePositionSorter
+//{
+//	bool operator()(const TSharedRef<FSequencerDisplayNode>& A, const TSharedRef <FSequencerDisplayNode>& B) const
+//	{
+//		return A->GetVirtualTop() < B->GetVirtualTop();
+//	}
+//};
+
+
+/**
+* Sorts the supplied unsorted nodes and inserts them into the existing sorted nodes before assigning
+* a sorting order to the newly combined array. Inserted nodes can optionally be relative to an existing
+* node that comes from the ExistingSortedNodes array; this allows you to sort them and then insert them
+* mid array.
+*
+* This function modifies the sorting order of all nodes passed in from both lists, but does not change
+* their actual hierarchy.
+*
+* @param UnsortedNodesToInsert A list of nodes to be sorted by Predicate
+* @param ExistingSortedNdoes A list of nodes that are already sorted that you want to have their indexes
+* adjusted for having the unsorted nodes inserted among them.
+* @param ItemDropZone Specifies the offset for the RelativeToNode (above, below or at the end). This is
+* useful as many of the places this gets called from have this information already.
+* @param Predicate Sorting predicate to sort UnsortedNodesToSort before insertion.
+* @param RelativeToNode An optional node from the ExistingSortedNodes list to insert before/after. If null
+* unsorted nodes will be inserted at the start and all nodes in ExistingSortedNodes will have their sorting
+* order offset by that many.
+*/
+//template <class PREDICATE_CLASS>
+//void SortAndSetSortingOrder(const TArray<TSharedRef<FSequencerDisplayNode>>& UnsortedNodesToInsert, const TArray<TSharedRef<FSequencerDisplayNode>>& ExistingSortedNodes, TOptional<EItemDropZone> ItemDropZone, const PREDICATE_CLASS& Predicate, TSharedPtr<FSequencerDisplayNode> RelativeToNode)
+//{
+//	TArray<TSharedRef<FSequencerDisplayNode>> NewSortingOrder = ExistingSortedNodes;
+//
+//	// ExistingSortedNodes may contain nodes that we do not consider for sorting order (spacers) or nodes we wish to re-sort.
+//	TArray<TSharedRef<FSequencerDisplayNode>> ChildrenToRemove;
+//	for (TSharedRef<FSequencerDisplayNode> Node : ExistingSortedNodes)
+//	{
+//		// Discard any of our children that aren't folders, objects or tracks (this ignores spacers, etc.)
+//		if (Node->GetType() != ESequencerNode::Folder && Node->GetType() != ESequencerNode::Object && Node->GetType() != ESequencerNode::Track)
+//		{
+//			ChildrenToRemove.Add(Node);
+//			continue;
+//		}
+//
+//		// If we're trying to sort this child remove it from our searchable list as well.
+//		if (UnsortedNodesToInsert.Contains(Node))
+//		{
+//			ChildrenToRemove.Add(Node);
+//		}
+//	}
+//
+//	for (TSharedRef<FSequencerDisplayNode> Node : ChildrenToRemove)
+//	{
+//		NewSortingOrder.Remove(Node);
+//	}
+//
+//	// Now get our index and insert the dragged nodes either before/after/at the end of that index depending on the ItemDropZone.
+//	int32 DropAdjustedIndex = NewSortingOrder.Num();
+//	int32 RelativeToIndex = NewSortingOrder.IndexOfByKey(RelativeToNode);
+//	if (RelativeToIndex >= 0 && ItemDropZone.IsSet())
+//	{
+//		switch (ItemDropZone.GetValue())
+//		{
+//		case EItemDropZone::AboveItem:
+//			DropAdjustedIndex = RelativeToIndex;
+//			break;
+//		case EItemDropZone::BelowItem:
+//			DropAdjustedIndex = RelativeToIndex + 1;
+//			break;
+//			// Folders will emit this drop zone which means we want to put it at the end of the folder.
+//		case EItemDropZone::OntoItem:
+//			DropAdjustedIndex = NewSortingOrder.Num();
+//			break;
+//		}
+//	}
+//
+//	// Allow our caller to specify the sorting order of new nodes. This allows us to sort incoming nodes by category or vertical position on the treeview, etc.
+//	TArray<TSharedRef<FSequencerDisplayNode>> SortedNodes = UnsortedNodesToInsert;
+//	SortedNodes.Sort(Predicate);
+//
+//	// Insert our sorted nodes into our child list so that we have the new absolute order for all items.
+//	NewSortingOrder.Insert(SortedNodes, DropAdjustedIndex);
+//
+//	// And then re-assign the sorting order index of everything in the folder according to the new list.
+//	for (int i = 0; i < NewSortingOrder.Num(); i++)
+//	{
+//		NewSortingOrder[i]->ModifyAndSetSortingOrder(i);
+//	}
+//};
