Sequencer: MVVM2 branch and Layer Bars

diff --git a/Engine/Source/Editor/Sequencer/Private/SequencerNodeTree.cpp b/Engine/Source/Editor/Sequencer/Private/SequencerNodeTree.cpp
--- a/Engine/Source/Editor/Sequencer/Private/SequencerNodeTree.cpp
+++ b/Engine/Source/Editor/Sequencer/Private/SequencerNodeTree.cpp
@@ -1,1687 +1,1108 @@
-// Copyright Epic Games, Inc. All Rights Reserved.
-
-#include "SequencerNodeTree.h"
-#include "MovieSceneBinding.h"
-#include "GameFramework/Actor.h"
-#include "MovieScene.h"
-#include "DisplayNodes/SequencerDisplayNode.h"
-#include "DisplayNodes/SequencerFolderNode.h"
-#include "DisplayNodes/SequencerObjectBindingNode.h"
-#include "ISequencerSection.h"
-#include "DisplayNodes/SequencerTrackNode.h"
-#include "MovieSceneSequence.h"
-#include "Tracks/MovieSceneCinematicShotTrack.h"
-#include "Tracks/MovieSceneSubTrack.h"
-#include "Sequencer.h"
-#include "MovieSceneFolder.h"
-#include "ISequencerTrackEditor.h"
-#include "DisplayNodes/SequencerRootNode.h"
-#include "Widgets/Views/STableRow.h"
-#include "CurveEditor.h"
-#include "SequencerNodeSortingMethods.h"
-#include "SequencerTrackFilters.h"
-#include "Channels/MovieSceneChannel.h"
-#include "ScopedTransaction.h"
-#include "SequencerUtilities.h"
-#include "SequencerLog.h"
-#include "SequencerCommonHelpers.h"
-
-
-/**
- * A spacer node that is unselectable
- */
-class FSequencerSpacerNode : public FSequencerDisplayNode
-{
-public:
-
-	/**
-	 * Create and initialize a new instance.
-	 * 
-	 * @param InParentTree The tree this node is in.
-	 */
-	FSequencerSpacerNode(float InSize, int32 InSortingOrder, FSequencerNodeTree& InParentTree)
-		: FSequencerDisplayNode(NAME_None, InParentTree)
-		, Size(InSize)
-		, SortingOrder(InSortingOrder)
-	{ }
-
-public:
-
-	// FSequencerDisplayNode interface
-	virtual bool CanRenameNode() const override { return false; }
-	virtual int32 GetSortingOrder() const override { return SortingOrder; }
-	virtual FText GetDisplayName() const override { return FText(); }
-	virtual float GetNodeHeight() const override { return Size; }
-	virtual FNodePadding GetNodePadding() const override { return FNodePadding(0.f); }
-	virtual ESequencerNode::Type GetType() const override { return ESequencerNode::Spacer; }
-	virtual void SetDisplayName(const FText& NewDisplayName) override { }
-	virtual TSharedRef<SWidget> GenerateContainerWidgetForOutliner(const TSharedRef<SSequencerTreeViewRow>& InRow) override { return SNew(SBox).HeightOverride(Size); }
-	virtual bool IsSelectable() const override { return false; }
-
-private:
-
-	/** The size of the spacer */
-	float Size;
-
-	/** The sorting order for the spacer */
-	int32 SortingOrder;
-};
-
-FSequencerNodeTree::~FSequencerNodeTree()
-{
-	if (TrackFilters.IsValid())
-	{
-		TrackFilters->OnChanged().RemoveAll(this);
-	}
-	if (TrackFilterLevelFilter.IsValid())
-	{
-		TrackFilterLevelFilter->OnChanged().RemoveAll(this);
-	}
-}
-
-
-FSequencerNodeTree::FSequencerNodeTree(FSequencer& InSequencer)
-	: RootNode(MakeShared<FSequencerRootNode>(*this))
-	, BottomSpacerNode(MakeShared<FSequencerSpacerNode>(30.f, INT_MAX, *this))
-	, SerialNumber(0)
-	, Sequencer(InSequencer)
-	, bFilterUpdateRequested(false)
-{
-	TrackFilters = MakeShared<FSequencerTrackFilterCollection>();
-	TrackFilters->OnChanged().AddRaw(this, &FSequencerNodeTree::RequestFilterUpdate);
-	TrackFilterLevelFilter = MakeShared< FSequencerTrackFilter_LevelFilter>();
-	TrackFilterLevelFilter->OnChanged().AddRaw(this, &FSequencerNodeTree::RequestFilterUpdate);
-}
-
-TSharedPtr<FSequencerObjectBindingNode> FSequencerNodeTree::FindObjectBindingNode(const FGuid& BindingID) const
-{
-	return ObjectBindingToNode.FindRef(BindingID);
-}
-
-void FSequencerNodeTree::RefreshNodes(UMovieScene* MovieScene)
-{
-	check(MovieScene);
-
-	++SerialNumber;
-
-	TSortedMap<FGuid, FGuid> ChildToParentBinding;
-	TSortedMap<FGuid, const FMovieSceneBinding*> AllBindings;
-
-	// Gather all object bindings in the sequence
-	for (const FMovieSceneBinding& Binding : MovieScene->GetBindings())
-	{
-		AllBindings.Add(Binding.GetObjectGuid(), &Binding);
-	}
-
-	// Populate the binding hierarchy
-	for (int32 PossessableIndex = 0; PossessableIndex < MovieScene->GetPossessableCount(); ++PossessableIndex)
-	{
-		const FMovieScenePossessable& Possessable = MovieScene->GetPossessable(PossessableIndex);
-
-		FGuid ThisID   = Possessable.GetGuid();
-		FGuid ParentID = Possessable.GetParent();
-
-		if (ParentID.IsValid())
-		{
-			ChildToParentBinding.Add(ThisID, ParentID);
-		}
-	}
-
-	TMap<FSequencerDisplayNode*, TSharedPtr<FSequencerDisplayNode>> ChildToParentMap;
-
-	// Object Bindings
-	{
-		for (const FMovieSceneBinding& Binding : MovieScene->GetBindings())
-		{
-			TSharedPtr<FSequencerObjectBindingNode> ObjectBindingNode = FindOrCreateObjectBinding(Binding.GetObjectGuid(), AllBindings, ChildToParentBinding, &ChildToParentMap);
-			if (!ObjectBindingNode)
-			{
-				continue;
-			}
-
-			if (!ChildToParentMap.Contains(ObjectBindingNode.Get()))
-			{
-				ChildToParentMap.Add(ObjectBindingNode.Get(), RootNode);
-			}
-
-			// Create nodes for the object binding's tracks
-			for (UMovieSceneTrack* Track : Binding.GetTracks())
-			{
-				if (ensureAlwaysMsgf(Track, TEXT("MovieScene binding '%s' data contains a null track. This should never happen."), *Binding.GetName()))
-				{
-					TSharedPtr<FSequencerTrackNode> TrackNode = CreateOrUpdateTrack(Track, ETrackType::Object);
-					if (TrackNode.IsValid())
-					{
-						ChildToParentMap.Add(TrackNode.Get(), ObjectBindingNode);
-					}
-				}
-			}
-		}
-	}
-
-	// Master tracks
-	{
-		UMovieSceneTrack* CameraCutTrack = MovieScene->GetCameraCutTrack();
-		if (CameraCutTrack)
-		{
-			TSharedPtr<FSequencerTrackNode> TrackNode = CreateOrUpdateTrack(CameraCutTrack, ETrackType::Master);
-			if (TrackNode.IsValid())
-			{
-				ChildToParentMap.Add(TrackNode.Get(), RootNode);
-			}
-		}
-
-		// Iterate all master tracks and generate nodes if necessary
-		for (UMovieSceneTrack* Track : MovieScene->GetMasterTracks())
-		{
-			if (ensureAlwaysMsgf(Track, TEXT("MovieScene data contains a null master track. This should never happen.")))
-			{
-				TSharedPtr<FSequencerTrackNode> TrackNode = CreateOrUpdateTrack(Track, ETrackType::Master);
-				if (TrackNode.IsValid())
-				{
-					ChildToParentMap.Add(TrackNode.Get(), RootNode);
-				}
-			}
-		}
-	}
-
-	// Folders may also create hierarchy items for tracks and object bindings
-	{
-		for (UMovieSceneFolder* Folder : MovieScene->GetRootFolders())
-		{
-			if (ensureAlwaysMsgf(Folder, TEXT("MovieScene data contains a null folder. This should never happen.")))
-			{
-				TSharedRef<FSequencerFolderNode> RootFolderNode = CreateOrUpdateFolder(Folder, AllBindings, ChildToParentBinding, &ChildToParentMap);
-				RootFolderNode->SetParent(RootNode);
-			}
-		}
-	}
-
-	for (TTuple<FSequencerDisplayNode*, TSharedPtr<FSequencerDisplayNode>> Pair : ChildToParentMap)
-	{
-		Pair.Key->SetParent(Pair.Value);
-	}
-
-	// Remove anything that is no longer relevant (ie serial number is out of date)
-	for (auto It = FolderToNode.CreateIterator(); It; ++It)
-	{
-		if (It->Value->TreeSerialNumber != SerialNumber)
-		{
-			It->Value->SetParent(nullptr);
-			It.RemoveCurrent();
-		}
-	}
-	for (auto It = TrackToNode.CreateIterator(); It; ++It)
-	{
-		if (It->Value->TreeSerialNumber != SerialNumber)
-		{
-			It->Value->SetParent(nullptr);
-			It.RemoveCurrent();
-		}
-		else
-		{
-			// Update after the above SetParent() because track sections need to have valid parent object bindings set up
-			It->Value->UpdateInnerHierarchy();
-		}
-	}
-	for (auto It = ObjectBindingToNode.CreateIterator(); It; ++It)
-	{
-		if (It->Value->TreeSerialNumber != SerialNumber)
-		{
-			It->Value->SetParent(nullptr);
-			It.RemoveCurrent();
-		}
-	}
-
-	// Remove mute/solo markers for any nodes that no longer exist
-	if(!MovieScene->IsReadOnly())
-	{
-		for (auto It = MovieScene->GetSoloNodes().CreateIterator(); It; ++It)
-		{
-			if (!GetNodeAtPath(*It))
-			{
-				It.RemoveCurrent();
-			}
-		}
-
-		for (auto It = MovieScene->GetMuteNodes().CreateIterator(); It; ++It)
-		{
-			if (!GetNodeAtPath(*It))
-			{
-				It.RemoveCurrent();
-			}
-		}
-	}
-
-	// Always add the bottom spacer node before counting filtered nodes
-	BottomSpacerNode->SetParent(RootNode);
-
-	// Re-filter the tree after updating 
-	// @todo sequencer: Newly added sections may need to be visible even when there is a filter
-	bFilterUpdateRequested = true;
-	UpdateFilters();
-
-	// Always show the bottom spacer node
-	FilteredNodes.Add(BottomSpacerNode);
-}
-
-TSharedPtr<FSequencerTrackNode> FSequencerNodeTree::CreateOrUpdateTrack(UMovieSceneTrack* Track, ETrackType TrackType)
-{
-	check(Track);
-
-	FObjectKey TrackKey(Track);
-	TSharedPtr<FSequencerTrackNode> TrackNode = TrackToNode.FindRef(TrackKey);
-	if (TrackNode.IsValid())
-	{
-		// Should be implemented as a filter
-		if (!Sequencer.IsTrackVisible(Track))
-		{
-			TrackNode->SetParent(nullptr);
-			TrackToNode.Remove(TrackKey);
-			return nullptr;
-		}
-	}
-	else
-	{
-		const bool bIsDraggable = TrackType == ETrackType::Master;
-		if (HasValidTypeEditor(Track))
-		{
-			TrackNode = MakeShared<FSequencerTrackNode>(Track, *FindOrAddTypeEditor(Track), bIsDraggable, *this);
-			TrackToNode.Add(TrackKey, TrackNode);
-		}
-		else
-		{
-			ensureAlwaysMsgf(Track, TEXT("'%s' is not currently supported."));
-			return nullptr;
-		}
-	}
-
-	// Assign the serial number for this node to indicate that it is still relevant
-	TrackNode->TreeSerialNumber = SerialNumber;
-	return TrackNode;
-}
-
-TSharedRef<FSequencerFolderNode> FSequencerNodeTree::CreateOrUpdateFolder(UMovieSceneFolder* Folder, const TSortedMap<FGuid, const FMovieSceneBinding*>& AllBindings, const TSortedMap<FGuid, FGuid>& ChildToParentBinding, TMap<FSequencerDisplayNode*, TSharedPtr<FSequencerDisplayNode>>* OutChildToParentMap)
-{
-	check(Folder);
-
-	FObjectKey FolderKey(Folder);
-
-	TSharedPtr<FSequencerFolderNode> FolderNode = FolderToNode.FindRef(FolderKey);
-	if (!FolderNode.IsValid())
-	{
-		FolderNode = MakeShared<FSequencerFolderNode>(*Folder, *this);
-		FolderToNode.Add(FolderKey, FolderNode.ToSharedRef());
-	}
-
-	// Assign the serial number for this node to indicate that it is still relevant
-	FolderNode->TreeSerialNumber = SerialNumber;
-
-	// Create the hierarchy for any child bindings
-	TArray<FGuid> ChildObjectBindingsToRemove;
-	for (const FGuid& ID : Folder->GetChildObjectBindings())
-	{
-		if (!AllBindings.Contains(ID))
-		{
-			ChildObjectBindingsToRemove.Add(ID);
-			continue;
-		}
-
-		TSharedPtr<FSequencerObjectBindingNode> Binding = FindOrCreateObjectBinding(ID, AllBindings, ChildToParentBinding, OutChildToParentMap);
-		if (Binding.IsValid())
-		{
-			OutChildToParentMap->Add(Binding.Get(), FolderNode);
-		}
-	}
-
-	for (const FGuid& ID : ChildObjectBindingsToRemove)
-	{
-		Folder->RemoveChildObjectBinding(ID);
-	}
-
-	// Create the hierarchy for any master tracks
-	for (UMovieSceneTrack* Track : Folder->GetChildMasterTracks())
-	{
-		if (ensureAlwaysMsgf(Track, TEXT("MovieScene folder '%s' data contains a null track. This should never happen."), *Folder->GetName()))
-		{
-			TSharedPtr<FSequencerTrackNode> TrackNode = CreateOrUpdateTrack(Track, ETrackType::Master);
-			if (TrackNode.IsValid())
-			{
-				OutChildToParentMap->Add(TrackNode.Get(), FolderNode);
-			}
-		}
-	}
-
-	// Add child folders
-	for (UMovieSceneFolder* ChildFolder : Folder->GetChildFolders())
-	{
-		if (ensureAlwaysMsgf(ChildFolder, TEXT("MovieScene folder '%s' data contains a null child folder. This should never happen."), *Folder->GetName()))
-		{
-			TSharedRef<FSequencerFolderNode> ChildFolderNode = CreateOrUpdateFolder(ChildFolder, AllBindings, ChildToParentBinding, OutChildToParentMap);
-			OutChildToParentMap->Add(&ChildFolderNode.Get(), FolderNode);
-		}
-	}
-
-	return FolderNode.ToSharedRef();
-}
-
-bool FSequencerNodeTree::HasActiveFilter() const
-{
-	return (!FilterString.IsEmpty()
-		|| TrackFilters->Num() > 0
-		|| TrackFilterLevelFilter->IsActive()
-		|| Sequencer.GetSequencerSettings()->GetShowSelectedNodesOnly()
-		|| Sequencer.GetFocusedMovieSceneSequence()->GetMovieScene()->GetNodeGroups().HasAnyActiveFilter());
-}
-
-bool FSequencerNodeTree::UpdateFiltersOnTrackValueChanged()
-{
-	// If filters are already scheduled for update, we can defer until the next update
-	if (bFilterUpdateRequested)
-	{
-		return false;
-	}
-
-	for (TSharedPtr< FSequencerTrackFilter > TrackFilter : *TrackFilters)
-	{
-		if (TrackFilter->ShouldUpdateOnTrackValueChanged())
-		{
-			// UpdateFilters will only run if bFilterUpdateRequested is true
-			bFilterUpdateRequested = true;
-			bool bFiltersUpdated = UpdateFilters();
-
-			// If the filter list was modified, set bFilterUpdateRequested to suppress excessive re-filters between tree update
-			bFilterUpdateRequested = bFiltersUpdated;
-			return bFiltersUpdated;
-		}
-	}
-
-	return false;
-}
-
-TSharedPtr<FSequencerObjectBindingNode> FSequencerNodeTree::FindOrCreateObjectBinding(const FGuid& BindingID, const TSortedMap<FGuid, const FMovieSceneBinding*>& AllBindings, const TSortedMap<FGuid, FGuid>& ChildToParentBinding, TMap<FSequencerDisplayNode*, TSharedPtr<FSequencerDisplayNode>>* OutChildToParentMap)
-{
-	if (!ensureAlwaysMsgf(AllBindings.Contains(BindingID), TEXT("Attempting to add a binding that does not exist.")))
-	{
-		return nullptr;
-	}
-
-	TSharedPtr<FSequencerObjectBindingNode> ObjectBindingNode = ObjectBindingToNode.FindRef(BindingID);
-	if (!ObjectBindingNode.IsValid())
-	{
-		// The node name is the object guid
-		FName ObjectNodeName = *BindingID.ToString();
-
-		ObjectBindingNode = MakeShared<FSequencerObjectBindingNode>(ObjectNodeName, BindingID, *this);
-		ObjectBindingToNode.Add(BindingID, ObjectBindingNode);
-	}
-
-	ObjectBindingNode->TreeSerialNumber = SerialNumber;
-
-	// Create its parent and make the association
-	if (const FGuid* ParentGuid = ChildToParentBinding.Find(BindingID))
-	{
-		TSharedPtr<FSequencerObjectBindingNode> ParentBinding = FindOrCreateObjectBinding(*ParentGuid, AllBindings, ChildToParentBinding, OutChildToParentMap);
-		if (ParentBinding.IsValid())
-		{
-			OutChildToParentMap->Add(ObjectBindingNode.Get(), ParentBinding);
-		}
-	}
-
-	return ObjectBindingNode;
-}
-
-void FSequencerNodeTree::Update()
-{
-	Sequencer.GetSelection().EmptySelectedOutlinerNodes();
-
-	EditorMap.Empty();
-	FilteredNodes.Empty();
-	SectionToHandle.Empty();
-	HoveredNode = nullptr;
-
-	UMovieSceneSequence* CurrentSequence = Sequencer.GetFocusedMovieSceneSequence();
-	if (WeakCurrentSequence != CurrentSequence)
-	{
-		DestroyAllNodes();
-		WeakCurrentSequence = CurrentSequence;
-	}
-
-	UMovieScene* MovieScene = CurrentSequence->GetMovieScene();
-	RefreshNodes(MovieScene);
-
-	// Re-filter the tree after updating 
-	// @todo sequencer: Newly added sections may need to be visible even when there is a filter
-	FilterNodes( FilterString );
-
-	// Sort root nodes
-	RootNode->SortImmediateChildren();
-
-	// Set up virtual offsets, expansion states, and tints
-	float VerticalOffset = 0.f;
-
-	auto Traverse_OnTreeRefreshed = [this, &VerticalOffset](FSequencerDisplayNode& InNode)
-	{
-		// Set up the virtual node position
-		float VerticalTop = VerticalOffset;
-		VerticalOffset += InNode.GetNodeHeight() + InNode.GetNodePadding().Combined();
-		InNode.OnTreeRefreshed(VerticalTop, VerticalOffset);
-
-		if (InNode.GetType() == ESequencerNode::Track)
-		{
-			this->UpdateSectionHandles(StaticCastSharedRef<FSequencerTrackNode>(InNode.AsShared()));
-		}
-		return true;
-	};
-
-	const bool bIncludeRootNode = false;
-	RootNode->Traverse_ParentFirst(Traverse_OnTreeRefreshed, bIncludeRootNode);
-
-	// Cache pinned state of nodes, needs to happen after OnTreeRefreshed
-	RootNode->UpdateCachedPinnedState();
-
-	// Ensure that the curve editor tree is up to date for our tree layout
-	UpdateCurveEditorTree();
-
-	bHasSoloNodes = false;
-	
-	TArray<FString>& SoloNodes = Sequencer.GetFocusedMovieSceneSequence()->GetMovieScene()->GetSoloNodes();
-
-	// Muting overrides soloing, so a solo node only counts if it's not muted.
-	for (const FString& NodePath : SoloNodes)
-	{
-		FSequencerDisplayNode* Node = GetNodeAtPath(NodePath);
-		if (Node)
-		{
-			if (!IsNodeMute(Node))
-			{
-				bHasSoloNodes = true;
-				break;
-			}
-		}
-	}
-
-	OnUpdatedDelegate.Broadcast();
-}
-
-FSequencerDisplayNode* FindNodeWithPath(FSequencerDisplayNode* InNode, const FString& NodePath)
-{
-	if (!InNode)
-	{
-		return nullptr;
-	}
-
-	if (InNode->GetPathName() == NodePath)
-	{
-		return InNode;
-	}
-
-	for (TSharedRef<FSequencerDisplayNode> ChildNode: InNode->GetChildNodes())
-	{
-		FSequencerDisplayNode* FoundNode = FindNodeWithPath(&ChildNode.Get(), NodePath);
-		if (FoundNode)
-		{
-			return FoundNode;
-		}
-	}
-	return nullptr;
-}
-
-
-FSequencerDisplayNode* FSequencerNodeTree::GetNodeAtPath(const FString& NodePath) const
-{
-	return FindNodeWithPath(&RootNode.Get(), NodePath);
-}
-
-bool FSequencerNodeTree::HasValidTypeEditor(UMovieSceneTrack* InTrack) const
-{
-	TSharedPtr<ISequencerTrackEditor> Editor = EditorMap.FindRef(InTrack);
-
-	if (!Editor.IsValid())
-	{
-		const TArray<TSharedPtr<ISequencerTrackEditor>>& TrackEditors = Sequencer.GetTrackEditors();
-
-		// Get a tool for each track
-		// @todo sequencer: Should probably only need to get this once and it shouldn't be done here. It depends on when movie scene tool modules are loaded
-		TSharedPtr<ISequencerTrackEditor> SupportedTool;
-
-		for (const auto& TrackEditor : TrackEditors)
-		{
-			if (TrackEditor->SupportsType(InTrack->GetClass()))
-			{
-				return true;
-			}
-		}
-	}
-
-	return false;
-}
-
-TSharedRef<ISequencerTrackEditor> FSequencerNodeTree::FindOrAddTypeEditor( UMovieSceneTrack* InTrack )
-{
-	TSharedPtr<ISequencerTrackEditor> Editor = EditorMap.FindRef( InTrack );
-
-	if( !Editor.IsValid() )
-	{
-		const TArray<TSharedPtr<ISequencerTrackEditor>>& TrackEditors = Sequencer.GetTrackEditors();
-
-		// Get a tool for each track
-		// @todo sequencer: Should probably only need to get this once and it shouldn't be done here. It depends on when movie scene tool modules are loaded
-		TSharedPtr<ISequencerTrackEditor> SupportedTool;
-
-		for (const auto& TrackEditor : TrackEditors)
-		{
-			if (TrackEditor->SupportsType(InTrack->GetClass()))
-			{
-				EditorMap.Add(InTrack, TrackEditor);
-				Editor = TrackEditor;
-
-				break;
-			}
-		}
-	}
-
-	return Editor.ToSharedRef();
-}
-
-TSharedRef<FSequencerDisplayNode> FSequencerNodeTree::GetRootNode() const
-{
-	return RootNode;
-}
-
-const TArray<TSharedRef<FSequencerDisplayNode>>& FSequencerNodeTree::GetRootNodes() const
-{
-	return RootNode->GetChildNodes();
-}
-
-void FSequencerNodeTree::MoveDisplayNodeToRoot(TSharedRef<FSequencerDisplayNode>& Node)
-{
-	// Objects that exist at the root level in a sequence are just removed from the folder they reside in.
-	// When the treeview is refreshed this will cause the regenerated nodes to show up at the root level.
-	TSharedPtr<FSequencerDisplayNode> ParentSeqNode = Node->GetParent();
-	switch (Node->GetType())
-	{
-		case ESequencerNode::Folder:
-		{
-			TSharedRef<FSequencerFolderNode> FolderNode = StaticCastSharedRef<FSequencerFolderNode>(Node);
-			UMovieScene* FocusedMovieScene = GetSequencer().GetFocusedMovieSceneSequence()->GetMovieScene();
-
-			if (ParentSeqNode.IsValid())
-			{
-				checkf(ParentSeqNode->GetType() == ESequencerNode::Folder, TEXT("Can not remove from unsupported parent node."));
-				TSharedPtr<FSequencerFolderNode> ParentFolder = StaticCastSharedPtr<FSequencerFolderNode>(ParentSeqNode);
-				ParentFolder->GetFolder().RemoveChildFolder(&FolderNode->GetFolder());
-			}
-			else
-			{
-				FocusedMovieScene->GetRootFolders().Remove(&FolderNode->GetFolder());
-			}
-
-			FocusedMovieScene->GetRootFolders().Add(&FolderNode->GetFolder());
-			break;
-		}
-		case ESequencerNode::Track:
-		{
-			TSharedRef<FSequencerTrackNode> DraggedTrackNode = StaticCastSharedRef<FSequencerTrackNode>(Node);
-			UMovieScene* FocusedMovieScene = GetSequencer().GetFocusedMovieSceneSequence()->GetMovieScene();
-
-			if (ParentSeqNode.IsValid())
-			{
-				checkf(ParentSeqNode->GetType() == ESequencerNode::Folder, TEXT("Can not remove from unsupported parent node."));
-				TSharedPtr<FSequencerFolderNode> ParentFolder = StaticCastSharedPtr<FSequencerFolderNode>(ParentSeqNode);
-				ParentFolder->GetFolder().RemoveChildMasterTrack(DraggedTrackNode->GetTrack());
-			}
-			break;
-		}
-		case ESequencerNode::Object:
-		{
-			TSharedRef<FSequencerObjectBindingNode> DraggedObjectBindingNode = StaticCastSharedRef<FSequencerObjectBindingNode>(Node);
-			UMovieScene* FocusedMovieScene = GetSequencer().GetFocusedMovieSceneSequence()->GetMovieScene();
-
-			if (ParentSeqNode.IsValid())
-			{
-				checkf(ParentSeqNode->GetType() == ESequencerNode::Folder, TEXT("Can not remove from unsupported parent node."));
-				TSharedPtr<FSequencerFolderNode> ParentFolder = StaticCastSharedPtr<FSequencerFolderNode>(ParentSeqNode);
-				ParentFolder->GetFolder().RemoveChildObjectBinding(DraggedObjectBindingNode->GetObjectBinding());
-			}
-			break;
-		}
-	}
-
-	// Reset to the root node
-	Node->SetParent(RootNode);
-
-	// Our children have changed parents which means that on subsequent creation they will retrieve their expansion state
-	// from the map using their new path. If the new path already exists the object goes to the state stored at that path.
-	// If the new path does not exist, the object returns to default state and not what is currently displayed. Either way 
-	// causes unexpected user behavior as nodes appear to randomly change expansion state as they are moved around the sequencer.
-
-	// To solve this, we update a node's parent when the node is moved, and then we update their expansion state here
-	// while we still have the current expansion state and the new node path. When the UI is regenerated on the subsequent
-	// refresh call, it will now retrieve the state the node was just in, instead of the state the node was in the last time it
-	// was in that location. This is done recursively as children store absolute paths so they need to be updated too.
-	Node->Traverse_ParentFirst([](FSequencerDisplayNode& TraversalNode)
-	{
-		TraversalNode.GetParentTree().SaveExpansionState(TraversalNode, TraversalNode.IsExpanded());
-		return true;
-	}, true);
-}
-
-void FSequencerNodeTree::SortAllNodesAndDescendants()
-{
-	auto Traverse_ResetSortOrder = [](FSequencerDisplayNode& Node)
-	{
-		Node.ResortImmediateChildren();
-		return true;
-	};
-
-	const bool bIncludeRootNode = true;
-	RootNode->Traverse_ParentFirst(Traverse_ResetSortOrder, bIncludeRootNode);
-
-	// Refresh the tree so that our changes are visible.
-	// @todo: Is this necessary any more?
-	GetSequencer().RefreshTree();
-}
-
-void FSequencerNodeTree::AddFilter(TSharedRef<FSequencerTrackFilter> TrackFilter)
-{
-	TrackFilters->Add(TrackFilter);
-}
-
-int32 FSequencerNodeTree::RemoveFilter(TSharedRef<FSequencerTrackFilter> TrackFilter)
-{
-	return TrackFilters->Remove(TrackFilter);
-}
-
-void FSequencerNodeTree::RemoveAllFilters()
-{
-	TrackFilters->RemoveAll();
-	TrackFilterLevelFilter->ResetFilter();
-}
-
-bool FSequencerNodeTree::IsTrackFilterActive(TSharedRef<FSequencerTrackFilter> TrackFilter) const
-{
-	return TrackFilters->Contains(TrackFilter);
-}
-
-void FSequencerNodeTree::AddLevelFilter(const FString& LevelName)
-{
-	TrackFilterLevelFilter->UnhideLevel(LevelName);
-}
-
-void FSequencerNodeTree::RemoveLevelFilter(const FString& LevelName)
-{
-	TrackFilterLevelFilter->HideLevel(LevelName);
-}
-
-bool FSequencerNodeTree::IsTrackLevelFilterActive(const FString& LevelName) const
-{
-	return !TrackFilterLevelFilter->IsLevelHidden(LevelName);
-}
-
-bool FSequencerNodeTree::IsNodeSolo(const FSequencerDisplayNode* InNode) const
-{
-	const TArray<FString>& SoloNodes = Sequencer.GetFocusedMovieSceneSequence()->GetMovieScene()->GetSoloNodes();
-	const FString NodePath = InNode->GetPathName();
-
-	if (SoloNodes.Contains(NodePath))
-	{
-		return true;
-	}
-
-	// Children should follow their parent's behavior unless told otherwise.
-	TSharedPtr<FSequencerDisplayNode> ParentNode = InNode->GetParent();
-	if (ParentNode.IsValid())
-	{
-		return IsNodeSolo(ParentNode.Get());
-	}
-
-	return false;
-}
-
-bool FSequencerNodeTree::HasSoloNodes() const
-{
-	return bHasSoloNodes;
-}
-
-bool FSequencerNodeTree::IsSelectedNodesSolo() const
-{
-	const TSet<TSharedRef<FSequencerDisplayNode> > SelectedNodes = Sequencer.GetSelection().GetSelectedOutlinerNodes();
-	
-	if (SelectedNodes.Num() == 0)
-	{
-		return false;
-	}
-
-	const TArray<FString>& SoloNodes = Sequencer.GetFocusedMovieSceneSequence()->GetMovieScene()->GetSoloNodes();
-
-	bool bIsSolo = true;
-	for (const TSharedRef<const FSequencerDisplayNode> Node : SelectedNodes)
-	{
-		if (!SoloNodes.Contains(Node->GetPathName()))
-		{
-			bIsSolo = false;
-			break;
-		}
-	}
-
-	return bIsSolo;
-}
-
-void FSequencerNodeTree::ToggleSelectedNodesSolo()
-{
-	UMovieScene* MovieScene = Sequencer.GetFocusedMovieSceneSequence()->GetMovieScene();
-	
-	if (MovieScene->IsReadOnly())
-	{
-		return;
-	}
-	
-	TArray<FString>& SoloNodes = MovieScene->GetSoloNodes();
-
-	const TSet<TSharedRef<FSequencerDisplayNode> >& SelectedNodes = Sequencer.GetSelection().GetSelectedOutlinerNodes();
-	if (SelectedNodes.Num() == 0)
-	{
-		return;
-	}
-
-	// First, determine if any of the selected nodes are not marked as solo
-	// If we have a mix, we should default to setting them all as solo
-	bool bIsSolo = true;
-	for (const TSharedRef<const FSequencerDisplayNode> Node : Sequencer.GetSelection().GetSelectedOutlinerNodes())
-	{
-		if (!SoloNodes.Contains(Node->GetPathName()))
-		{
-			bIsSolo = false;
-			break;
-		}
-	}
-
-	const FScopedTransaction Transaction(NSLOCTEXT("Sequencer", "ToggleSolo", "Toggle Solo"));
-
-	MovieScene->Modify();
-
-	for (const TSharedRef<const FSequencerDisplayNode> Node : Sequencer.GetSelection().GetSelectedOutlinerNodes())
-	{
-		FString NodePath = Node->GetPathName();
-		if (bIsSolo)
-		{
-			// If we're currently solo, unsolo
-			SoloNodes.Remove(NodePath);
-		}
-		else
-		{
-			// Mark solo, being careful as we might be re-marking an already solo node
-			SoloNodes.AddUnique(NodePath);
-		}
-	}
-	
-	GetSequencer().RefreshTree();
-}
-
-bool FSequencerNodeTree::IsNodeMute(const FSequencerDisplayNode* InNode) const
-{
-	const TArray<FString>& MuteNodes = Sequencer.GetFocusedMovieSceneSequence()->GetMovieScene()->GetMuteNodes();
-	const FString NodePath = InNode->GetPathName();
-
-	if (MuteNodes.Contains(NodePath))
-	{
-		return true;
-	}
-	
-	// Children should follow their parent's behavior unless told otherwise.
-	const TSharedPtr<FSequencerDisplayNode> ParentNode = InNode->GetParent();
-	if (ParentNode.IsValid())
-	{
-		return IsNodeMute(ParentNode.Get());
-	}
-	
-	return false;
-}
-
-bool FSequencerNodeTree::IsSelectedNodesMute() const
-{
-	const TSet<TSharedRef<FSequencerDisplayNode> > SelectedNodes = Sequencer.GetSelection().GetSelectedOutlinerNodes();
-
-	if (SelectedNodes.Num() == 0)
-	{
-		return false;
-	}
-
-	const TArray<FString>& MuteNodes = Sequencer.GetFocusedMovieSceneSequence()->GetMovieScene()->GetMuteNodes();
-
-	bool bIsMute = true;
-	for (const TSharedRef<const FSequencerDisplayNode> Node : SelectedNodes)
-	{
-		if (!MuteNodes.Contains(Node->GetPathName()))
-		{
-			bIsMute = false;
-			break;
-		}
-	}
-
-	return bIsMute;
-}
-
-void FSequencerNodeTree::ToggleSelectedNodesMute()
-{
-	UMovieScene* MovieScene = Sequencer.GetFocusedMovieSceneSequence()->GetMovieScene();
-	
-	if (MovieScene->IsReadOnly())
-	{
-		return;
-	}
-
-	TArray<FString>& MuteNodes = MovieScene->GetMuteNodes();
-
-	const TSet<TSharedRef<FSequencerDisplayNode> >& SelectedNodes = Sequencer.GetSelection().GetSelectedOutlinerNodes();
-	if (SelectedNodes.Num() == 0)
-	{
-		return;
-	}
-
-	// First, determine if any of the selected nodes are not marked as Mute
-	// If we have a mix, we should default to setting them all as Mute
-	bool bIsMute = true;
-	for (const TSharedRef<const FSequencerDisplayNode> Node : SelectedNodes)
-	{
-		if (!MuteNodes.Contains(Node->GetPathName()))
-		{
-			bIsMute = false;
-			break;
-		}
-	}
-
-	const FScopedTransaction Transaction(NSLOCTEXT("Sequencer", "ToggleMute", "Toggle Mute"));
-
-	MovieScene->Modify();
-
-	for (const TSharedRef<const FSequencerDisplayNode> Node : SelectedNodes)
-	{
-		FString NodePath = Node->GetPathName();
-		if (bIsMute)
-		{
-			// If we're currently Mute, unMute
-			MuteNodes.Remove(NodePath);
-		}
-		else
-		{
-			// Mark Mute, being careful as we might be re-marking an already Mute node
-			MuteNodes.AddUnique(NodePath);
-		}
-	}
-
-	GetSequencer().RefreshTree();
-}
-
-void FSequencerNodeTree::SaveExpansionState(const FSequencerDisplayNode& Node, bool bExpanded)
-{	
-	// @todo Sequencer - This should be moved to the sequence level
-	UMovieScene* MovieScene = Sequencer.GetFocusedMovieSceneSequence()->GetMovieScene();
-	FMovieSceneEditorData& EditorData = MovieScene->GetEditorData();
-
-	EditorData.ExpansionStates.Add(Node.GetPathName(), FMovieSceneExpansionState(bExpanded));
-}
-
-
-bool FSequencerNodeTree::GetSavedExpansionState(const FSequencerDisplayNode& Node) const
-{
-	// @todo Sequencer - This should be moved to the sequence level
-	UMovieScene* MovieScene = Sequencer.GetFocusedMovieSceneSequence()->GetMovieScene();
-	FMovieSceneEditorData& EditorData = MovieScene->GetEditorData();
-	FMovieSceneExpansionState* ExpansionState = EditorData.ExpansionStates.Find( Node.GetPathName() );
-
-	return ExpansionState ? ExpansionState->bExpanded : GetDefaultExpansionState(Node);
-}
-
-
-bool FSequencerNodeTree::GetDefaultExpansionState( const FSequencerDisplayNode& Node ) const
-{
-	// Object nodes, and track nodes that are parent tracks are expanded by default.
-	if (Node.GetType() == ESequencerNode::Object)
-	{
-		return true;
-	}
-
-	else if (Node.GetType() == ESequencerNode::Track)
-	{
-		const FSequencerTrackNode& TrackNode = static_cast<const FSequencerTrackNode&>(Node);
-
-		if (TrackNode.GetSubTrackMode() == FSequencerTrackNode::ESubTrackMode::ParentTrack)
-		{
-			return true;
-		}
-
-		if (TrackNode.GetTrackEditor().GetDefaultExpansionState(TrackNode.GetTrack()))
-		{
-			return true;
-		}
-	}
-
-	return false;
-}
-
-void FSequencerNodeTree::SavePinnedState(const FSequencerDisplayNode& Node, bool bPinned)
-{
-	UMovieScene* MovieScene = Sequencer.GetFocusedMovieSceneSequence()->GetMovieScene();
-	FMovieSceneEditorData& EditorData = MovieScene->GetEditorData();
-
-	if (bPinned)
-	{
-		EditorData.PinnedNodes.AddUnique(Node.GetPathName());
-	}
-	else
-	{
-		EditorData.PinnedNodes.RemoveSingle(Node.GetPathName());
-	}
-}
-
-
-bool FSequencerNodeTree::GetSavedPinnedState(const FSequencerDisplayNode& Node) const
-{
-	UMovieScene* MovieScene = Sequencer.GetFocusedMovieSceneSequence()->GetMovieScene();
-	FMovieSceneEditorData& EditorData = MovieScene->GetEditorData();
-	bool bPinned = EditorData.PinnedNodes.Contains(Node.GetPathName());
-
-	return bPinned;
-}
-
-bool FSequencerNodeTree::IsNodeFiltered(const TSharedRef<const FSequencerDisplayNode> Node) const
-{
-	for (auto It = FilteredNodes.CreateConstIterator(); It; ++It)
-	{
-		if ((*It) == Node)
-		{
-			return true;
-		}
-	}
-	return false;
-}
-
-void FSequencerNodeTree::SetHoveredNode(const TSharedPtr<FSequencerDisplayNode>& InHoveredNode)
-{
-	if (InHoveredNode.IsValid() && !InHoveredNode->IsSelectable())
-	{
-		return;
-	}
-
-	if (InHoveredNode != HoveredNode)
-	{
-		HoveredNode = InHoveredNode;
-	}
-}
-
-const TSharedPtr<FSequencerDisplayNode>& FSequencerNodeTree::GetHoveredNode() const
-{
-	return HoveredNode;
-}
-
-void FSequencerNodeTree::UpdateSectionHandles(TSharedRef<FSequencerTrackNode> TrackNode)
-{
-	const TArray<TSharedRef<ISequencerSection>>& Sections = TrackNode->GetSections();
-	for (int32 SectionIndex = 0; SectionIndex < Sections.Num(); ++SectionIndex)
-	{
-		UMovieSceneSection* Section = Sections[SectionIndex]->GetSectionObject();
-		if (Section)
-		{
-			SectionToHandle.Add(Section, FSectionHandle(TrackNode, SectionIndex));
-		}
-	}
-}
-
-TOptional<FSectionHandle> FSequencerNodeTree::GetSectionHandle(const UMovieSceneSection* Section) const
-{
-	const FSectionHandle* Found = SectionToHandle.Find(Section);
-	if (Found)
-	{
-		return *Found;
-	}
-
-	return TOptional<FSectionHandle>();
-}
-
-static void AddChildNodes(const TSharedRef<FSequencerDisplayNode>& StartNode, TSet<TSharedRef<FSequencerDisplayNode>>& OutFilteredNodes)
-{
-	for (TSharedRef<FSequencerDisplayNode> ChildNode : StartNode->GetChildNodes())
-	{
-		OutFilteredNodes.Add(ChildNode);
-		AddChildNodes(ChildNode, OutFilteredNodes);
-	}
-}
-
-static void AddParentNodes(const TSharedRef<FSequencerDisplayNode>& StartNode, TSet<TSharedRef<FSequencerDisplayNode>>& OutFilteredNodes)
-{
-	// Gather parent folders up the chain
-	TSharedPtr<FSequencerDisplayNode> ParentNode = StartNode->GetParent();
-	while (ParentNode.IsValid())
-	{
-		OutFilteredNodes.Add(ParentNode.ToSharedRef());
-		ParentNode = ParentNode->GetParent();
-	}
-}
-
-static bool PassesFilterStrings(FSequencer& Sequencer, const TSharedRef<FSequencerDisplayNode>& StartNode, const TArray<FString>& FilterStrings)
-{
-	// If we have a filter string, make sure we match
-	if (FilterStrings.Num() > 0)
-	{
-		// check each string in the filter strings list against 
-		for (const FString& String : FilterStrings)
-		{
-			if (!StartNode->GetDisplayName().ToString().Contains(String))
-			{
-				return false;
-			}
-		}
-
-		TSharedPtr<FSequencerDisplayNode> ParentNode = StartNode->GetParent();
-		while (ParentNode.IsValid())
-		{
-			if (!ParentNode.Get()->IsExpanded())
-			{
-				ParentNode.Get()->SetExpansionState(true);
-			}
-
-			ParentNode = ParentNode->GetParent();
-		}
-	}
-
-	return true;
-}
-
-/**
- * Recursively filters nodes
- *
- * @param StartNode			The node to start from
- * @param Filters			The filter collection to test against
- * @param OutFilteredNodes	The list of all filtered nodes
- *
- * @return Whether this node passed filtering
-  */
-
-static bool FilterNodesRecursive(FSequencer& Sequencer, const TSharedRef<FSequencerDisplayNode>& StartNode, TSharedPtr<FSequencerTrackFilterCollection> Filters, const TArray<FString>& FilterStrings, TSharedPtr<FSequencerTrackFilter_LevelFilter> LevelTrackFilter, TSet<TSharedRef<FSequencerDisplayNode>>& OutFilteredNodes)
-{
-	bool bAnyChildPassed = false;
-
-	// Special case: If a parent node matches an active text search filter, the children should also pass the text search filter
-	// so we stop filtering them based on text to eliminate special case conflicts with non-object binding nodes
-	// with potential child object bindings nodes when only showing selected bindings. This is also faster.
-	bool bPassedFilterStrings = PassesFilterStrings(Sequencer, StartNode, FilterStrings);
-	const TArray<FString>& ChildFilterStrings = bPassedFilterStrings ? TArray<FString>() : FilterStrings;
-
-	// Special case: Child nodes should always be processed, as they may force their parents to pass
-	for (TSharedRef<FSequencerDisplayNode> Node : StartNode->GetChildNodes())
-	{
-		if (FilterNodesRecursive(Sequencer, Node, Filters, ChildFilterStrings, LevelTrackFilter, OutFilteredNodes))
-		{
-			bAnyChildPassed = true;
-		}
-	}
-
-	// After child nodes are processed, if this node didn't pass text filtering, fail it
-	if (!bPassedFilterStrings)
-	{
-		return bAnyChildPassed;
-	}
-
-	bool bPassedAnyFilters = false;
-	bool bIsTrackOrObjectBinding = false;
-	bool bIsSubTrack = false;
-
-	switch(StartNode->GetType())
-	{
-		case ESequencerNode::Track:
-		{
-			bIsTrackOrObjectBinding = true;
-
-			UMovieSceneTrack* Track = static_cast<const FSequencerTrackNode&>(StartNode.Get()).GetTrack();
-
-			// Always show subsequence tracks so that the user can navigate up and down the hierarchy with the filter
-			if (Track && Track->IsA<UMovieSceneSubTrack>())
-			{
-				bPassedAnyFilters = true;
-				bIsSubTrack = true;
-				break;
-			}
-
-			TSharedPtr<FSequencerSectionKeyAreaNode> TopLevelKeyArea = static_cast<const FSequencerTrackNode&>(StartNode.Get()).GetTopLevelKeyNode();
-			if (TopLevelKeyArea.IsValid())
-			{
-				for (const TSharedRef<IKeyArea>& KeyArea : TopLevelKeyArea->GetAllKeyAreas())
-				{
-					FMovieSceneChannel* Channel = KeyArea->ResolveChannel();
-					if (Channel)
-					{
-						if (Filters->Num() == 0 || Filters->PassesAnyFilters(Channel))
-						{
-							bPassedAnyFilters = true;
-							break;
-						}
-					}
-				}
-			}
-		
-			if (bPassedAnyFilters || Filters->Num() == 0 || Filters->PassesAnyFilters(Track, StartNode->GetDisplayName()))
-			{
-				bPassedAnyFilters = true;
-
-				// Track nodes do not belong to a level, but might be a child of an objectbinding node that does
-				if (LevelTrackFilter->IsActive())
-				{
-					TSharedPtr<const FSequencerDisplayNode> ParentNode = StartNode->GetParent();
-					while (ParentNode.IsValid())
-					{
-						if (ParentNode->GetType() == ESequencerNode::Object)
-						{
-							// The track belongs to an objectbinding node, start by assuming it doesn't match the level filter
-							bPassedAnyFilters = false;
-
-							const FSequencerObjectBindingNode* ObjectNode = static_cast<const FSequencerObjectBindingNode*>(ParentNode.Get());
-							for (TWeakObjectPtr<>& Object : Sequencer.FindObjectsInCurrentSequence(ObjectNode->GetObjectBinding()))
-							{
-								if (Object.IsValid() && LevelTrackFilter->PassesFilter(Object.Get()))
-								{
-									// If at least one of the objects on the objectbinding node pass the level filter, show the track
-									bPassedAnyFilters = true;
-									break;
-								}
-							}
-
-							break;
-						}
-						ParentNode = ParentNode->GetParent();
-					}
-				}
-
-				if (bPassedAnyFilters && Sequencer.GetSequencerSettings()->GetShowSelectedNodesOnly())
-				{
-					TSharedPtr<const FSequencerDisplayNode> ParentNode = StartNode->GetParent();
-					while (ParentNode.IsValid())
-					{
-						// Pinned tracks should be visible whether selected or not
-						if (ParentNode->IsPinned())
-						{
-							break;
-						}
-
-						if (ParentNode->GetType() == ESequencerNode::Object)
-						{
-							const FSequencerObjectBindingNode* ObjectNode = static_cast<const FSequencerObjectBindingNode*>(ParentNode.Get());
-							const FMovieSceneBinding* Binding = Sequencer.GetFocusedMovieSceneSequence()->GetMovieScene()->FindBinding(ObjectNode->GetObjectBinding());
-							if (!(Binding && Sequencer.IsBindingVisible(*Binding)))
-							{
-								return bAnyChildPassed;
-							}
-
-							break;
-						}
-						ParentNode = ParentNode->GetParent();
-					}
-				}
-			}
-			break;
-		}
-		case ESequencerNode::Object:
-		{
-			bIsTrackOrObjectBinding = true;
-
-			const FSequencerObjectBindingNode& ObjectNode = static_cast<const FSequencerObjectBindingNode&>(StartNode.Get());
-			for (TWeakObjectPtr<>& Object : Sequencer.FindObjectsInCurrentSequence(ObjectNode.GetObjectBinding()))
-			{
-				if (Object.IsValid() && (Filters->Num() == 0 || Filters->PassesAnyFilters(Object.Get(), StartNode->GetDisplayName()))
-					&& LevelTrackFilter->PassesFilter(Object.Get()))
-				{
-					bPassedAnyFilters = true;
-					break;
-				}
-			}
-
-			if (bPassedAnyFilters && Sequencer.GetSequencerSettings()->GetShowSelectedNodesOnly() && !StartNode->IsPinned())
-			{
-				UMovieScene* MovieScene = Sequencer.GetFocusedMovieSceneSequence()->GetMovieScene();
-				const FMovieSceneBinding* Binding = Sequencer.GetFocusedMovieSceneSequence()->GetMovieScene()->FindBinding(ObjectNode.GetObjectBinding());
-				if (Binding && !Sequencer.IsBindingVisible(*Binding))
-				{
-					return bAnyChildPassed;
-				}
-			}
-			break;
-		}
-		case ESequencerNode::Category:
-		{
-			if (TSharedPtr<FSequencerTrackNode> TrackNode = StartNode->FindParentTrackNode())
-			{
-				UMovieSceneTrack* Track = TrackNode->GetTrack();
-				if (Filters->Num() == 0 || Filters->PassesAnyFilters(Track, StartNode->GetDisplayName()))
-				{
-					bPassedAnyFilters = true;
-				}
-			}
-			if (bPassedAnyFilters && Sequencer.GetSequencerSettings()->GetShowSelectedNodesOnly())
-			{
-				TSharedPtr<const FSequencerDisplayNode> ParentNode = StartNode->GetParent();
-				while (ParentNode.IsValid())
-				{
-					// Pinned tracks should be visible whether selected or not
-					if (ParentNode->IsPinned())
-					{
-						break;
-					}
-
-					if (ParentNode->GetType() == ESequencerNode::Object)
-					{
-						const FSequencerObjectBindingNode* ObjectNode = static_cast<const FSequencerObjectBindingNode*>(ParentNode.Get());
-						const FMovieSceneBinding* Binding = Sequencer.GetFocusedMovieSceneSequence()->GetMovieScene()->FindBinding(ObjectNode->GetObjectBinding());
-						if (!(Binding && Sequencer.IsBindingVisible(*Binding)))
-						{
-							bPassedAnyFilters = false;
-						}
-
-						break;
-					}
-					ParentNode = ParentNode->GetParent();
-				}
-			}
-			break;
-		}
-		case ESequencerNode::KeyArea:
-		{
-			const FSequencerSectionKeyAreaNode& KeyAreaNode = static_cast<const FSequencerSectionKeyAreaNode&>(StartNode.Get());
-			for (const TSharedRef<IKeyArea>& KeyArea : KeyAreaNode.GetAllKeyAreas())
-			{
-				FMovieSceneChannel* Channel = KeyArea->ResolveChannel();
-				if (Channel)
-				{
-					if (Filters->Num() == 0 || Filters->PassesAnyFilters(Channel))
-					{
-						bPassedAnyFilters = true;
-						break;
-					}
-				}
-			}
-			if (bPassedAnyFilters && Sequencer.GetSequencerSettings()->GetShowSelectedNodesOnly())
-			{
-				TSharedPtr<const FSequencerDisplayNode> ParentNode = StartNode->GetParent();
-				while (ParentNode.IsValid())
-				{
-					// Pinned tracks should be visible whether selected or not
-					if (ParentNode->IsPinned())
-					{
-						break;
-					}
-
-					if (ParentNode->GetType() == ESequencerNode::Object)
-					{
-						const FSequencerObjectBindingNode* ObjectNode = static_cast<const FSequencerObjectBindingNode*>(ParentNode.Get());
-						const FMovieSceneBinding* Binding = Sequencer.GetFocusedMovieSceneSequence()->GetMovieScene()->FindBinding(ObjectNode->GetObjectBinding());
-						if (!(Binding && Sequencer.IsBindingVisible(*Binding)))
-						{
-							bPassedAnyFilters = false;
-						}
-
-						break;
-					}
-					ParentNode = ParentNode->GetParent();
-				}
-			}
-			break;
-		}
-		case ESequencerNode::Folder:
-		{
-			// Special case: If we're pinned, then we should pass regardless
-			if (StartNode->IsPinned())
-			{
-				bPassedAnyFilters = true;
-				break;
-			}
-
-			// Special case: If we're only filtering on text search, include folders and key areas in the search
-			if (Filters->Num() == 0 && FilterStrings.Num() > 0)
-			{
-				bPassedAnyFilters = true;
-
-				// Special case: but don't include if only showing selected bindings and we don't have child that passed
-				if (Sequencer.GetSequencerSettings()->GetShowSelectedNodesOnly() && !bAnyChildPassed)
-				{
-					bPassedAnyFilters = false;
-
-					// Special case: unless we're the child of a node that is a selected binding
-					TSharedPtr<const FSequencerDisplayNode> ParentNode = StartNode->GetParent();
-					while (ParentNode.IsValid())
-					{
-						// Pinned tracks should be visible whether selected or not
-						if (ParentNode->IsPinned())
-						{
-							bPassedAnyFilters = true;
-							break;
-						}
-
-						if (ParentNode->GetType() == ESequencerNode::Object)
-						{
-							const FSequencerObjectBindingNode* ObjectNode = static_cast<const FSequencerObjectBindingNode*>(ParentNode.Get());
-							const FMovieSceneBinding* Binding = Sequencer.GetFocusedMovieSceneSequence()->GetMovieScene()->FindBinding(ObjectNode->GetObjectBinding());
-							if (Binding && Sequencer.IsBindingVisible(*Binding))
-							{
-								bPassedAnyFilters = true;
-							}
-
-							break;
-						}
-						ParentNode = ParentNode->GetParent();
-					}
-				}
-			}
-			break;
-		}
-	}
-
-	if (bPassedAnyFilters && !bIsSubTrack)
-	{
-		// If filtering on selection set is enabled, we need to run another pass to verify we're in an enabled node group
-		UMovieSceneNodeGroupCollection& NodeGroups = Sequencer.GetFocusedMovieSceneSequence()->GetMovieScene()->GetNodeGroups();
-		if (NodeGroups.HasAnyActiveFilter())
-		{
-			bPassedAnyFilters = false;
-
-			TSharedPtr<const FSequencerDisplayNode> CurrentNode = StartNode;
-			while (CurrentNode.IsValid() && !bPassedAnyFilters)
-			{
-				// Special case: Pinned tracks should be visible whether in the node group or not
-				if (CurrentNode->IsPinned())
-				{
-					bPassedAnyFilters = true;
-					break;
-				}
-
-				for (const UMovieSceneNodeGroup* NodeGroup : NodeGroups)
-				{
-					if (NodeGroup->GetEnableFilter() && NodeGroup->ContainsNode(CurrentNode->GetPathName()))
-					{
-						bPassedAnyFilters = true;
-						break;
-					}
-				}
-
-				CurrentNode = CurrentNode->GetParent();
-			}
-		}
-	}
-
-	if (bPassedAnyFilters)
-	{
-		OutFilteredNodes.Add(StartNode);
-		AddParentNodes(StartNode, OutFilteredNodes);
-
-		// Special case: When only showing selected bindings, and a non-object passes text filtering
-		// don't add it's children, as they may be a binding that is not seleceted. Selected child nodes will add themselves.
-		if (!(Sequencer.GetSequencerSettings()->GetShowSelectedNodesOnly() && FilterStrings.Num() > 0 && !bIsTrackOrObjectBinding))
-		{
-			AddChildNodes(StartNode, OutFilteredNodes);
-		}
-
-		return true;
-	}
-
-	return bAnyChildPassed;
-}
-
-bool FSequencerNodeTree::UpdateFilters()
-{
-	if (!bFilterUpdateRequested)
-	{
-		return false;
-	}
-
-	TSet< TSharedRef<FSequencerDisplayNode> > PreviousFilteredNodes(FilteredNodes);
-
-	FilteredNodes.Empty();
-
-
-	UObject* PlaybackContext = Sequencer.GetPlaybackContext();
-	UWorld* World = PlaybackContext ? PlaybackContext->GetWorld() : nullptr;
-	TrackFilterLevelFilter->UpdateWorld(World);
-
-	if (HasActiveFilter())
-	{
-		// Build a list of strings that must be matched
-		TArray<FString> FilterStrings;
-
-		// Remove whitespace from the front and back of the string
-		FilterString.TrimStartAndEndInline();
-		FilterString.ParseIntoArray(FilterStrings, TEXT(" "), true /*bCullEmpty*/);
-
-		for (auto It = GetRootNodes().CreateConstIterator(); It; ++It)
-		{
-			// Recursively filter all nodes, matching them against the list of filter strings.  All filter strings must be matched
-			FilterNodesRecursive(Sequencer, *It, TrackFilters, FilterStrings, TrackFilterLevelFilter, FilteredNodes);
-		}
-	}
-
-	bFilteringOnNodeGroups = Sequencer.GetFocusedMovieSceneSequence()->GetMovieScene()->GetNodeGroups().HasAnyActiveFilter();
-	bFilterUpdateRequested = false;
-
-	// Count the total number of display nodes
-	DisplayNodeCount = 0;
-
-	RootNode->Traverse_ParentFirst(
-		[this](FSequencerDisplayNode& InNode)
-		{
-			++this->DisplayNodeCount;
-			return true;
-		}
-		, false
-	);
-
-	// Return whether the new list of FilteredNodes is different than the previous list
-	return (PreviousFilteredNodes.Num() != FilteredNodes.Num() || !PreviousFilteredNodes.Includes(FilteredNodes));
-}
-
-int32 FSequencerNodeTree::GetTotalDisplayNodeCount() const 
-{ 
-	// Subtract 1 for the spacer node which is always added
-	return DisplayNodeCount - 1; 
-}
-
-int32 FSequencerNodeTree::GetFilteredDisplayNodeCount() const 
-{ 
-	// Subtract 1 for the spacer node which is always added
-	return FilteredNodes.Num() - 1; 
-}
-
-void FSequencerNodeTree::FilterNodes(const FString& InFilter)
-{
-	if (InFilter != FilterString)
-	{
-		FilterString = InFilter;
-		bFilterUpdateRequested = true;
-	}
-}
-
-void FSequencerNodeTree::NodeGroupsCollectionChanged()
-{
-	if (Sequencer.GetFocusedMovieSceneSequence()->GetMovieScene()->GetNodeGroups().HasAnyActiveFilter() || bFilteringOnNodeGroups)
-	{
-		RequestFilterUpdate();
-	}
-}
-
-TArray< TSharedRef<FSequencerDisplayNode> > FSequencerNodeTree::GetAllNodes() const
-{
-	TArray< TSharedRef<FSequencerDisplayNode> > AllNodes;
-
-	const bool bIncludeRootNode = false;
-	RootNode->Traverse_ParentFirst([&AllNodes](FSequencerDisplayNode& InNode) 
-	{
-		AllNodes.Add(InNode.AsShared());
-		return true;
-	}, bIncludeRootNode);
-
-	return AllNodes;
-}
-
-// this will just clear out the curve editor, it will be recreated on next tick which is safer
-void FSequencerNodeTree::RecreateCurveEditor()
-{
-	FCurveEditor* CurveEditor = Sequencer.GetCurveEditor().Get();
-	if (CurveEditor)
-	{
-		for (auto It = CurveEditorTreeItemIDs.CreateIterator(); It; ++It)
-		{
-			CurveEditor->RemoveTreeItem(It->Value);
-		}
-
-		CurveEditorTreeItemIDs.Empty();
-	}
-}
-
-void FSequencerNodeTree::UpdateCurveEditorTree()
-{
-	FCurveEditor*     CurveEditor     = Sequencer.GetCurveEditor().Get();
-	FCurveEditorTree* CurveEditorTree = CurveEditor->GetTree();
-	// Guard against multiple broadcasts here and defer them until the end of this function
-	FScopedCurveEditorTreeEventGuard ScopedEventGuard = CurveEditorTree->ScopedEventGuard();
-
-	// Remove any curve editor tree items which are now parented incorrectly - we remove invalid entries before adding new ones below
-	// to ensure that we do not add new tree items as children of stale tree items
-	for (auto It = CurveEditorTreeItemIDs.CreateIterator(); It; ++It)
-	{
-		TSharedPtr<FSequencerDisplayNode> Node = It->Key.Pin();
-
-		bool bIsStillRelevant = Node.IsValid() && Node->TreeSerialNumber == SerialNumber && Node->IsVisible();
-		if (bIsStillRelevant)
-		{
-			// It's possible that the item was removed by a previous iteration of this loop, so we have to handle that case here
-			TSharedPtr<FSequencerDisplayNode> Parent = Node->GetParent();
-			FCurveEditorTreeItemID CachedParentID = Parent ? CurveEditorTreeItemIDs.FindRef(Parent) : FCurveEditorTreeItemID::Invalid();
-
-			const FCurveEditorTreeItem* TreeItem = CurveEditorTree->FindItem(It->Value);
-
-			bIsStillRelevant = TreeItem && TreeItem->GetParentID() == CachedParentID;
-		}
-
-		// Remove this item and all its children if it is no longer relevant, or it needs reparenting
-		if (!bIsStillRelevant)
-		{
-			CurveEditor->RemoveTreeItem(It->Value);
-			It.RemoveCurrent();
-		}
-	}
-
-	// Do a second pass to remove any items that were removed recursively above
-	for (auto It = CurveEditorTreeItemIDs.CreateIterator(); It; ++It)
-	{
-		if (CurveEditorTree->FindItem(It->Value) == nullptr)
-		{
-			It.RemoveCurrent();
-		}
-	}
-
-	auto Traverse_AddToCurveEditor = [this, CurveEditor](FSequencerDisplayNode& InNode)
-	{
-		if (!InNode.IsVisible())
-		{
-			return true;
-		}
-
-		if (InNode.GetType() == ESequencerNode::Track)
-		{
-			// Track nodes with top level key area's must be added
-			TSharedPtr<FSequencerSectionKeyAreaNode> TopLevelKeyArea = static_cast<const FSequencerTrackNode&>(InNode).GetTopLevelKeyNode();
-			if (TopLevelKeyArea.IsValid() && KeyAreaHasCurves(*TopLevelKeyArea))
-			{
-				this->AddToCurveEditor(InNode.AsShared(), CurveEditor);
-			}
-		}
-		else if (InNode.GetType() == ESequencerNode::KeyArea && KeyAreaHasCurves(static_cast<const FSequencerSectionKeyAreaNode&>(InNode)))
-		{
-			// Key area nodes are always added
-			this->AddToCurveEditor(InNode.AsShared(), CurveEditor);
-		}
-		return true;
-	};
-
-	static const bool bIncludeThisNode = false;
-	RootNode->Traverse_ChildFirst(Traverse_AddToCurveEditor, bIncludeThisNode);
-
-	CurveEditorTreeItemIDs.Compact();
-}
-
-bool FSequencerNodeTree::KeyAreaHasCurves(const FSequencerSectionKeyAreaNode& KeyAreaNode) const
-{
-	for (const TSharedRef<IKeyArea>& KeyArea : KeyAreaNode.GetAllKeyAreas())
-	{
-		const ISequencerChannelInterface* EditorInterface = KeyArea->FindChannelEditorInterface();
-		if (EditorInterface && EditorInterface->SupportsCurveEditorModels_Raw(KeyArea->GetChannel()))
-		{
-			return true;
-		}
-	}
-	return false;
-}
-
-FCurveEditorTreeItemID FSequencerNodeTree::AddToCurveEditor(TSharedRef<FSequencerDisplayNode> DisplayNode, FCurveEditor* CurveEditor)
-{
-	if (FCurveEditorTreeItemID* Existing = CurveEditorTreeItemIDs.Find(DisplayNode))
-	{
-		if (CurveEditor->GetTree()->FindItem(*Existing) != nullptr)
-		{
-			return *Existing;
-		}
-	}
-
-	TSharedPtr<FSequencerDisplayNode> Parent = DisplayNode->GetParent();
-
-	FCurveEditorTreeItemID ParentID = Parent.IsValid() ? AddToCurveEditor(Parent.ToSharedRef(), CurveEditor) : FCurveEditorTreeItemID::Invalid();
-
-	FCurveEditorTreeItem* NewItem = CurveEditor->AddTreeItem(ParentID);
-	NewItem->SetWeakItem(StaticCastSharedRef<ICurveEditorTreeItem>(DisplayNode));
-
-	CurveEditorTreeItemIDs.Add(DisplayNode, NewItem->GetID());
-	return NewItem->GetID();
-}
-
-FCurveEditorTreeItemID FSequencerNodeTree::FindCurveEditorTreeItem(TSharedRef<FSequencerDisplayNode> DisplayNode) const
-{
-	const FCurveEditorTreeItemID* FoundID = CurveEditorTreeItemIDs.Find(DisplayNode);
-	if (FoundID && Sequencer.GetCurveEditor()->GetTree()->FindItem(*FoundID) != nullptr)
-	{
-		return *FoundID;
-	}
-	return FCurveEditorTreeItemID::Invalid();
-}
-
-void FSequencerNodeTree::UnpinAllNodes()
-{
-	const bool bIncludeRootNode = false;
-	RootNode->Traverse_ParentFirst([](FSequencerDisplayNode& InNode)
-	{
-		InNode.Unpin();
-		return true;
-	}, bIncludeRootNode);
-}
-
-void FSequencerNodeTree::DestroyAllNodes()
-{
-	for (TSharedRef<FSequencerDisplayNode> RootChild : CopyTemp(RootNode->GetChildNodes()))
-	{
-		RootChild->SetParent(nullptr);
-	}
-
-	ObjectBindingToNode.Empty();
-	TrackToNode.Empty();
-	FolderToNode.Empty();
-	SectionToHandle.Empty();
-	FilteredNodes.Empty();
-	HoveredNode = nullptr;
-
-	if (FCurveEditor* CurveEditor = Sequencer.GetCurveEditor().Get())
-	{
-		FScopedCurveEditorTreeEventGuard ScopedEventGuard = CurveEditor->GetTree()->ScopedEventGuard();
-
-		for (auto It = CurveEditorTreeItemIDs.CreateIterator(); It; ++It)
-		{
-			CurveEditor->RemoveTreeItem(It->Value);
-		}
-	}
-	CurveEditorTreeItemIDs.Empty();
-}
+// Copyright Epic Games, Inc. All Rights Reserved.
+
+#include "SequencerNodeTree.h"
+#include "MVVM/ViewModels/ViewModelHierarchy.h"
+#include "MovieSceneBinding.h"
+#include "GameFramework/Actor.h"
+#include "MovieScene.h"
+#include "MVVM/ViewModels/ViewModel.h"
+#include "MVVM/Extensions/IOutlinerExtension.h"
+#include "MVVM/Extensions/ICurveEditorTreeItemExtension.h"
+#include "MVVM/Extensions/IPinnableExtension.h"
+#include "MVVM/Extensions/ISoloableExtension.h"
+#include "MVVM/Extensions/IMutableExtension.h"
+#include "MVVM/ViewModels/SequencerEditorViewModel.h"
+#include "MVVM/ViewModels/OutlinerViewModel.h"
+#include "MVVM/ViewModels/CategoryModel.h"
+#include "MVVM/ViewModels/FolderModel.h"
+#include "MVVM/ViewModels/ViewModelIterators.h"
+#include "MVVM/ViewModels/ObjectBindingModel.h"
+#include "MVVM/ViewModels/OutlinerSpacer.h"
+#include "MVVM/ViewModels/SectionModel.h"
+#include "MVVM/ViewModels/SequenceModel.h"
+#include "MVVM/ViewModels/TrackModel.h"
+#include "MVVM/CurveEditorIntegrationExtension.h"
+#include "MVVM/ObjectBindingModelStorageExtension.h"
+#include "MVVM/SectionModelStorageExtension.h"
+#include "MVVM/TrackModelStorageExtension.h"
+#include "IKeyArea.h"
+#include "ISequencerSection.h"
+#include "MovieSceneSequence.h"
+#include "Tracks/MovieSceneCinematicShotTrack.h"
+#include "Tracks/MovieSceneSubTrack.h"
+#include "Sequencer.h"
+#include "MovieSceneFolder.h"
+#include "ISequencerTrackEditor.h"
+#include "Widgets/Views/STableRow.h"
+#include "CurveEditor.h"
+#include "SequencerTrackFilters.h"
+#include "Channels/MovieSceneChannel.h"
+#include "ScopedTransaction.h"
+#include "SequencerUtilities.h"
+#include "SequencerLog.h"
+#include "SequencerCommonHelpers.h"
+
+FSequencerNodeTree::~FSequencerNodeTree()
+{
+	if (TrackFilters.IsValid())
+	{
+		TrackFilters->OnChanged().RemoveAll(this);
+	}
+	if (TrackFilterLevelFilter.IsValid())
+	{
+		TrackFilterLevelFilter->OnChanged().RemoveAll(this);
+	}
+}
+
+FSequencerNodeTree::FSequencerNodeTree(FSequencer& InSequencer)
+	: Sequencer(InSequencer)
+	, bFilterUpdateRequested(false)
+{
+	TrackFilters = MakeShared<FSequencerTrackFilterCollection>();
+	TrackFilters->OnChanged().AddRaw(this, &FSequencerNodeTree::RequestFilterUpdate);
+	TrackFilterLevelFilter = MakeShared< FSequencerTrackFilter_LevelFilter>();
+	TrackFilterLevelFilter->OnChanged().AddRaw(this, &FSequencerNodeTree::RequestFilterUpdate);
+}
+
+TSharedPtr<UE::Sequencer::FObjectBindingModel> FSequencerNodeTree::FindObjectBindingNode(const FGuid& BindingID) const
+{
+	using namespace UE::Sequencer;
+
+	const FObjectBindingModelStorageExtension* ObjectBindingStorage = RootNode->CastDynamic<FObjectBindingModelStorageExtension>();
+	if (ObjectBindingStorage)
+	{
+		TSharedPtr<FObjectBindingModel> Model = ObjectBindingStorage->FindModelForObjectBinding(BindingID);
+		return Model;
+	}
+
+	return nullptr;
+}
+
+bool FSequencerNodeTree::HasActiveFilter() const
+{
+	return (!FilterString.IsEmpty()
+		|| TrackFilters->Num() > 0
+		|| TrackFilterLevelFilter->IsActive()
+		|| Sequencer.GetSequencerSettings()->GetShowSelectedNodesOnly()
+		|| Sequencer.GetFocusedMovieSceneSequence()->GetMovieScene()->GetNodeGroups().HasAnyActiveFilter());
+}
+
+bool FSequencerNodeTree::UpdateFiltersOnTrackValueChanged()
+{
+	// If filters are already scheduled for update, we can defer until the next update
+	if (bFilterUpdateRequested)
+	{
+		return false;
+	}
+
+	for (TSharedPtr< FSequencerTrackFilter > TrackFilter : *TrackFilters)
+	{
+		if (TrackFilter->ShouldUpdateOnTrackValueChanged())
+		{
+			// UpdateFilters will only run if bFilterUpdateRequested is true
+			bFilterUpdateRequested = true;
+			bool bFiltersUpdated = UpdateFilters();
+
+			// If the filter list was modified, set bFilterUpdateRequested to suppress excessive re-filters between tree update
+			bFilterUpdateRequested = bFiltersUpdated;
+			return bFiltersUpdated;
+		}
+	}
+
+	return false;
+}
+
+void FSequencerNodeTree::Update()
+{
+	using namespace UE::Sequencer;
+
+	FViewModelHierarchyOperation UpdateOp(RootNode);
+
+	FilteredNodes.Empty();
+
+	if (!ensure(RootNode))
+	{
+		return;
+	}
+
+	TSharedPtr<FSequenceModel> SequenceModel = RootNode->CastThisShared<FSequenceModel>();
+	check(SequenceModel);
+
+	UMovieSceneSequence* CurrentSequence = SequenceModel->GetSequence();
+	check(CurrentSequence);
+
+	UMovieScene* MovieScene = CurrentSequence->GetMovieScene();
+	CleanupMuteSolo(MovieScene);
+
+	// Re-filter the tree after updating 
+	// @todo sequencer: Newly added sections may need to be visible even when there is a filter
+	FilterNodes(FilterString);
+	bFilterUpdateRequested = true;
+	UpdateFilters();
+
+	// Sort all nodes
+	const bool bIncludeRootNode = true;
+	for (TSharedPtr<ISortableExtension> SortableChild : RootNode->GetDescendantsOfType<ISortableExtension>(bIncludeRootNode))
+	{
+		SortableChild->SortChildren();
+	}
+
+	// Update all virtual geometries
+	// This must happen after the sorting
+	IGeometryExtension::UpdateVirtualGeometry(0.f, RootNode);
+
+	// Cache pinned state of nodes, needs to happen after OnTreeRefreshed
+	FPinnableExtensionShim::UpdateCachedPinnedState(RootNode);
+
+	// Update mute/solo markers and cache final state
+	UpdateMuteSolo(MovieScene);
+
+	// Update curve editor tree based on new filtered hierarchy
+	auto CurveEditorIntegration = SequenceModel->CastDynamic<FCurveEditorIntegrationExtension>();
+	if (CurveEditorIntegration)
+	{
+		CurveEditorIntegration->UpdateCurveEditor();
+	}
+
+	OnUpdatedDelegate.Broadcast();
+}
+
+UE::Sequencer::TViewModelPtr<UE::Sequencer::IOutlinerExtension> FindNodeWithPath(UE::Sequencer::TViewModelPtr<UE::Sequencer::IOutlinerExtension> InNode, const FString& NodePath)
+{
+	using namespace UE::Sequencer;
+
+	if (!InNode)
+	{
+		return nullptr;
+	}
+
+	FString HeadPath, TailPath;
+	const bool bHasDelimiter = NodePath.Split(".", &HeadPath, &TailPath);
+	const FString NodeIdentifier = InNode->GetIdentifier().ToString();
+
+	if (bHasDelimiter)
+	{
+		if (NodeIdentifier != HeadPath)
+		{
+			// The node we're looking for is not in this sub-branch.
+			return nullptr;
+		}
+	}
+	else
+	{
+		// NodePath is just a name, so simply check with our node's name.
+		return (NodeIdentifier == NodePath) ? InNode : nullptr;
+	}
+
+	check(bHasDelimiter && !TailPath.IsEmpty());
+
+	for (TViewModelPtr<IOutlinerExtension> Child : InNode.AsModel()->GetChildrenOfType<IOutlinerExtension>())
+	{
+		TViewModelPtr<IOutlinerExtension> FoundNode = FindNodeWithPath(Child, TailPath);
+		if (FoundNode)
+		{
+			return FoundNode;
+		}
+	}
+
+	return nullptr;
+}
+
+
+UE::Sequencer::TViewModelPtr<UE::Sequencer::IOutlinerExtension> FSequencerNodeTree::GetNodeAtPath(const FString& NodePath) const
+{
+	using namespace UE::Sequencer;
+
+	for (const TViewModelPtr<IOutlinerExtension>& RootChild : RootNode->GetChildrenOfType<IOutlinerExtension>())
+	{
+		TViewModelPtr<IOutlinerExtension> NodeAtPath = FindNodeWithPath(RootChild, NodePath);
+		if (NodeAtPath)
+		{
+			return NodeAtPath;
+		}
+	}
+	return nullptr;
+}
+
+void FSequencerNodeTree::SetRootNode(const UE::Sequencer::FViewModelPtr& InRootNode)
+{
+	ensureMsgf(!RootNode, TEXT("Re-assinging the root node is currently an undefined behavior"));
+	RootNode = InRootNode;
+}
+
+UE::Sequencer::FViewModelPtr FSequencerNodeTree::GetRootNode() const
+{
+	return RootNode;
+}
+
+TArray<UE::Sequencer::TViewModelPtr<UE::Sequencer::IOutlinerExtension>> FSequencerNodeTree::GetRootNodes() const
+{
+	using namespace UE::Sequencer;
+
+	TArray<TViewModelPtr<IOutlinerExtension>> RootNodes;
+	for (const TViewModelPtr<IOutlinerExtension>& Child : RootNode->GetChildrenOfType<IOutlinerExtension>())
+	{
+		RootNodes.Add(Child);
+	}
+	return RootNodes;
+}
+
+void FSequencerNodeTree::SortAllNodesAndDescendants()
+{
+	using namespace UE::Sequencer;
+
+	const bool bIncludeRootNode = true;
+	TArray<ISortableExtension*> SortableChildren;
+	for (TSharedPtr<FViewModel> Child : RootNode->GetDescendants(bIncludeRootNode))
+	{
+		if (ISortableExtension* SortableExtension = Child->CastThis<ISortableExtension>())
+		{
+			SortableExtension->SetCustomOrder(-1);
+			SortableChildren.Add(SortableExtension);
+		}
+	}
+	for (ISortableExtension* SortableChild : SortableChildren)
+	{
+		SortableChild->SortChildren();
+	}
+
+	// Refresh the tree so that our changes are visible.
+	// @todo: Is this necessary any more?
+	//GetSequencer().RefreshTree();
+}
+
+void FSequencerNodeTree::AddFilter(TSharedPtr<FSequencerTrackFilter> TrackFilter)
+{
+	TrackFilters->Add(TrackFilter);
+}
+
+int32 FSequencerNodeTree::RemoveFilter(TSharedPtr<FSequencerTrackFilter> TrackFilter)
+{
+	return TrackFilters->Remove(TrackFilter);
+}
+
+void FSequencerNodeTree::RemoveAllFilters()
+{
+	TrackFilters->RemoveAll();
+	TrackFilterLevelFilter->ResetFilter();
+}
+
+bool FSequencerNodeTree::IsTrackFilterActive(TSharedPtr<FSequencerTrackFilter> TrackFilter) const
+{
+	return TrackFilters->Contains(TrackFilter);
+}
+
+void FSequencerNodeTree::AddLevelFilter(const FString& LevelName)
+{
+	TrackFilterLevelFilter->UnhideLevel(LevelName);
+}
+
+void FSequencerNodeTree::RemoveLevelFilter(const FString& LevelName)
+{
+	TrackFilterLevelFilter->HideLevel(LevelName);
+}
+
+bool FSequencerNodeTree::IsTrackLevelFilterActive(const FString& LevelName) const
+{
+	return !TrackFilterLevelFilter->IsLevelHidden(LevelName);
+}
+
+bool FSequencerNodeTree::IsNodeSolo(const UE::Sequencer::TViewModelPtr<UE::Sequencer::IOutlinerExtension>& InNode) const
+{
+	using namespace UE::Sequencer;
+
+	const TArray<FString>& SoloNodes = Sequencer.GetFocusedMovieSceneSequence()->GetMovieScene()->GetSoloNodes();
+	const FString NodePath = IOutlinerExtension::GetPathName(InNode);
+
+	if (SoloNodes.Contains(NodePath))
+	{
+		return true;
+	}
+
+	// Children should follow their parent's behavior unless told otherwise.
+	TViewModelPtr<IOutlinerExtension> ParentNode = InNode.AsModel()->FindAncestorOfType<IOutlinerExtension>();
+	if (ParentNode)
+	{
+		return IsNodeSolo(ParentNode);
+	}
+
+	return false;
+}
+
+bool FSequencerNodeTree::HasSoloNodes() const
+{
+	return bHasSoloNodes;
+}
+
+bool FSequencerNodeTree::IsSelectedNodesSolo() const
+{
+	using namespace UE::Sequencer;
+
+	const TSet<TWeakPtr<FViewModel>> SelectedNodes = Sequencer.GetSelection().GetSelectedOutlinerItems();
+	
+	if (SelectedNodes.Num() == 0)
+	{
+		return false;
+	}
+
+	const TArray<FString>& SoloNodes = Sequencer.GetFocusedMovieSceneSequence()->GetMovieScene()->GetSoloNodes();
+
+	bool bIsSolo = true;
+	for (const TWeakPtr<FViewModel>& Node : SelectedNodes)
+	{
+		if (!SoloNodes.Contains(IOutlinerExtension::GetPathName(Node)))
+		{
+			bIsSolo = false;
+			break;
+		}
+	}
+
+	return bIsSolo;
+}
+
+void FSequencerNodeTree::ToggleSelectedNodesSolo()
+{
+	using namespace UE::Sequencer;
+
+	UMovieScene* MovieScene = Sequencer.GetFocusedMovieSceneSequence()->GetMovieScene();
+	
+	if (MovieScene->IsReadOnly())
+	{
+		return;
+	}
+	
+	TArray<FString>& SoloNodes = MovieScene->GetSoloNodes();
+
+	const TSet<TWeakPtr<FViewModel>> SelectedNodes = Sequencer.GetSelection().GetSelectedOutlinerItems();
+	if (SelectedNodes.Num() == 0)
+	{
+		return;
+	}
+
+	// First, determine if any of the selected nodes are not marked as solo
+	// If we have a mix, we should default to setting them all as solo
+	bool bIsSolo = true;
+	for (const TWeakPtr<const FViewModel> Node : SelectedNodes)
+	{
+		if (!SoloNodes.Contains(IOutlinerExtension::GetPathName(Node)))
+		{
+			bIsSolo = false;
+			break;
+		}
+	}
+
+	const FScopedTransaction Transaction(NSLOCTEXT("Sequencer", "ToggleSolo", "Toggle Solo"));
+
+	MovieScene->Modify();
+
+	for (const TWeakPtr<FViewModel>& Node : SelectedNodes)
+	{
+		FString NodePath = IOutlinerExtension::GetPathName(Node);
+		if (bIsSolo)
+		{
+			// If we're currently solo, unsolo
+			SoloNodes.Remove(NodePath);
+		}
+		else
+		{
+			// Mark solo, being careful as we might be re-marking an already solo node
+			SoloNodes.AddUnique(NodePath);
+		}
+	}
+	
+	GetSequencer().RefreshTree();
+}
+
+bool FSequencerNodeTree::IsNodeMute(const UE::Sequencer::TViewModelPtr<UE::Sequencer::IOutlinerExtension>& InNode) const
+{
+	using namespace UE::Sequencer;
+
+	const TArray<FString>& MuteNodes = Sequencer.GetFocusedMovieSceneSequence()->GetMovieScene()->GetMuteNodes();
+	const FString NodePath = IOutlinerExtension::GetPathName(InNode);
+
+	if (MuteNodes.Contains(NodePath))
+	{
+		return true;
+	}
+	
+	// Children should follow their parent's behavior unless told otherwise.
+	TViewModelPtr<IOutlinerExtension> ParentNode(InNode.AsModel()->GetParent());
+	if (ParentNode)
+	{
+		return IsNodeMute(ParentNode);
+	}
+	
+	return false;
+}
+
+bool FSequencerNodeTree::IsSelectedNodesMute() const
+{
+	using namespace UE::Sequencer;
+
+	const TSet<TWeakPtr<FViewModel>> SelectedNodes = Sequencer.GetSelection().GetSelectedOutlinerItems();
+
+	if (SelectedNodes.Num() == 0)
+	{
+		return false;
+	}
+
+	const TArray<FString>& MuteNodes = Sequencer.GetFocusedMovieSceneSequence()->GetMovieScene()->GetMuteNodes();
+
+	bool bIsMute = true;
+	for (const TWeakPtr<const FViewModel> Node : SelectedNodes)
+	{
+		if (!MuteNodes.Contains(IOutlinerExtension::GetPathName(Node)))
+		{
+			bIsMute = false;
+			break;
+		}
+	}
+
+	return bIsMute;
+}
+
+void FSequencerNodeTree::ToggleSelectedNodesMute()
+{
+	using namespace UE::Sequencer;
+
+	UMovieScene* MovieScene = Sequencer.GetFocusedMovieSceneSequence()->GetMovieScene();
+	
+	if (MovieScene->IsReadOnly())
+	{
+		return;
+	}
+
+	TArray<FString>& MuteNodes = MovieScene->GetMuteNodes();
+
+	const TSet<TWeakPtr<FViewModel>> SelectedNodes = Sequencer.GetSelection().GetSelectedOutlinerItems();
+	if (SelectedNodes.Num() == 0)
+	{
+		return;
+	}
+
+	// First, determine if any of the selected nodes are not marked as Mute
+	// If we have a mix, we should default to setting them all as Mute
+	bool bIsMute = true;
+	for (const TWeakPtr<const FViewModel> Node : SelectedNodes)
+	{
+		if (!MuteNodes.Contains(IOutlinerExtension::GetPathName(Node)))
+		{
+			bIsMute = false;
+			break;
+		}
+	}
+
+	const FScopedTransaction Transaction(NSLOCTEXT("Sequencer", "ToggleMute", "Toggle Mute"));
+
+	MovieScene->Modify();
+
+	for (const TWeakPtr<FViewModel>& Node : SelectedNodes)
+	{
+		FString NodePath = IOutlinerExtension::GetPathName(Node);
+		if (bIsMute)
+		{
+			// If we're currently Mute, unMute
+			MuteNodes.Remove(NodePath);
+		}
+		else
+		{
+			// Mark Mute, being careful as we might be re-marking an already Mute node
+			MuteNodes.AddUnique(NodePath);
+		}
+	}
+
+	GetSequencer().RefreshTree();
+}
+
+void FSequencerNodeTree::SaveExpansionState(const UE::Sequencer::FViewModel& Node, bool bExpanded)
+{
+	using namespace UE::Sequencer;
+
+	// @todo Sequencer - This should be moved to the sequence level
+	UMovieScene* MovieScene = Sequencer.GetFocusedMovieSceneSequence()->GetMovieScene();
+	FMovieSceneEditorData& EditorData = MovieScene->GetEditorData();
+
+	EditorData.ExpansionStates.Add(IOutlinerExtension::GetPathName(Node), FMovieSceneExpansionState(bExpanded));
+}
+
+TOptional<bool> FSequencerNodeTree::GetSavedExpansionState( const FViewModel& Node )
+{
+	using namespace UE::Sequencer;
+
+	UMovieScene* MovieScene = Sequencer.GetFocusedMovieSceneSequence()->GetMovieScene();
+	FMovieSceneEditorData& EditorData = MovieScene->GetEditorData();
+	FMovieSceneExpansionState* ExpansionState = EditorData.ExpansionStates.Find(IOutlinerExtension::GetPathName(Node));
+
+	return ExpansionState ? ExpansionState->bExpanded : TOptional<bool>();
+}
+
+void FSequencerNodeTree::SavePinnedState(const UE::Sequencer::FViewModel& Node, bool bPinned)
+{
+	using namespace UE::Sequencer;
+
+	UMovieScene* MovieScene = Sequencer.GetFocusedMovieSceneSequence()->GetMovieScene();
+	FMovieSceneEditorData& EditorData = MovieScene->GetEditorData();
+
+	if (bPinned)
+	{
+		EditorData.PinnedNodes.AddUnique(IOutlinerExtension::GetPathName(Node));
+	}
+	else
+	{
+		EditorData.PinnedNodes.RemoveSingle(IOutlinerExtension::GetPathName(Node));
+	}
+}
+
+
+bool FSequencerNodeTree::GetSavedPinnedState(const UE::Sequencer::FViewModel& Node) const
+{
+	using namespace UE::Sequencer;
+
+	UMovieScene* MovieScene = Sequencer.GetFocusedMovieSceneSequence()->GetMovieScene();
+	FMovieSceneEditorData& EditorData = MovieScene->GetEditorData();
+	bool bPinned = EditorData.PinnedNodes.Contains(IOutlinerExtension::GetPathName(Node));
+
+	return bPinned;
+}
+
+bool FSequencerNodeTree::IsNodeFiltered(const TSharedPtr<UE::Sequencer::FViewModel>& Node) const
+{
+	return FilteredNodes.Contains(Node);
+}
+
+TSharedPtr<UE::Sequencer::FSectionModel> FSequencerNodeTree::GetSectionModel(const UMovieSceneSection* Section) const
+{
+	using namespace UE::Sequencer;
+
+	FSectionModelStorageExtension* SectionStorage = RootNode->CastThis<FSectionModelStorageExtension>();
+	if (ensure(SectionStorage))
+	{
+		return SectionStorage->FindModelForSection(Section);
+	}
+	return nullptr;
+}
+
+static void AddChildNodes(const UE::Sequencer::TViewModelPtr<UE::Sequencer::IOutlinerExtension>& StartNode, TSet<UE::Sequencer::TWeakViewModelPtr<UE::Sequencer::IOutlinerExtension>>& OutFilteredNodes)
+{
+	using namespace UE::Sequencer;
+
+	for (TViewModelPtr<IOutlinerExtension> ChildNode : StartNode.AsModel()->GetDescendantsOfType<IOutlinerExtension>())
+	{
+		OutFilteredNodes.Add(ChildNode);
+	}
+}
+
+static void AddParentNodes(const UE::Sequencer::TViewModelPtr<UE::Sequencer::IOutlinerExtension>& StartNode, TSet<UE::Sequencer::TWeakViewModelPtr<UE::Sequencer::IOutlinerExtension>>& OutFilteredNodes)
+{
+	using namespace UE::Sequencer;
+
+	// Gather parent folders up the chain
+	for (TViewModelPtr<IOutlinerExtension> ParentNode : StartNode.AsModel()->GetAncestorsOfType<IOutlinerExtension>())
+	{
+		OutFilteredNodes.Add(ParentNode);
+	}
+}
+
+static bool PassesFilterStrings(FSequencer& Sequencer, const UE::Sequencer::TViewModelPtr<UE::Sequencer::IOutlinerExtension>& StartNode, const TArray<FString>& FilterStrings)
+{
+	using namespace UE::Sequencer;
+
+	// If we have a filter string, make sure we match
+	if (FilterStrings.Num() > 0)
+	{
+		TSharedPtr<FOutlinerViewModel> Outliner = Sequencer.GetViewModel()->GetOutliner();
+
+		if (!Outliner)
+		{
+			return false;
+		}
+
+		FString DisplayLabel = StartNode->GetLabel().ToString();
+
+		// check each string in the filter strings list against 
+		for (const FString& String : FilterStrings)
+		{
+			if (!DisplayLabel.Contains(String))
+			{
+				return false;
+			}
+		}
+
+		for (TViewModelPtr<IOutlinerExtension> Parent : StartNode.AsModel()->GetAncestorsOfType<IOutlinerExtension>())
+		{
+			if (!Parent->IsExpanded())
+			{
+				Parent->SetExpansion(true);
+			}
+		}
+	}
+
+	return true;
+}
+
+/**
+ * Recursively filters nodes
+ *
+ * @param StartNode			The node to start from
+ * @param Filters			The filter collection to test against
+ * @param OutFilteredNodes	The list of all filtered nodes
+ *
+ * @return Whether this node passed filtering
+  */
+
+static bool FilterNodesRecursive(
+	  FSequencer& Sequencer
+	, const UE::Sequencer::TViewModelPtr<UE::Sequencer::IOutlinerExtension>& StartNode
+	, TSharedPtr<FSequencerTrackFilterCollection> Filters, const TArray<FString>& FilterStrings
+	, TSharedPtr<FSequencerTrackFilter_LevelFilter> LevelTrackFilter
+	, TSet<UE::Sequencer::TWeakViewModelPtr<UE::Sequencer::IOutlinerExtension>>& OutFilteredNodes)
+{
+	using namespace UE::Sequencer;
+
+	bool bAnyChildPassed = false;
+
+	// Special case: If a parent node matches an active text search filter, the children should also pass the text search filter
+	// so we stop filtering them based on text to eliminate special case conflicts with non-object binding nodes
+	// with potential child object bindings nodes when only showing selected bindings. This is also faster.
+	bool bPassedFilterStrings = PassesFilterStrings(Sequencer, StartNode, FilterStrings);
+	const TArray<FString>& ChildFilterStrings = bPassedFilterStrings ? TArray<FString>() : FilterStrings;
+
+	// Special case: Child nodes should always be processed, as they may force their parents to pass
+	for (TViewModelPtr<IOutlinerExtension> Node : StartNode.AsModel()->GetChildrenOfType<IOutlinerExtension>())
+	{
+		if (FilterNodesRecursive(Sequencer, Node, Filters, ChildFilterStrings, LevelTrackFilter, OutFilteredNodes))
+		{
+			bAnyChildPassed = true;
+		}
+	}
+
+	// After child nodes are processed, if this node didn't pass text filtering, fail it
+	if (!bPassedFilterStrings)
+	{
+		return bAnyChildPassed;
+	}
+
+	// TODO: we can probably completely rewrite this logic by taking advantage of the new modular extensions.
+
+	bool bPassedAnyFilters = false;
+	bool bIsTrackOrObjectBinding = false;
+	bool bIsSubTrack = false;
+
+	TSharedPtr<IPinnableExtension> Pinnable = StartNode.ImplicitCast();
+	const bool bIsPinned = Pinnable && Pinnable->IsPinned();
+
+	if (TSharedPtr<FTrackModel> TrackModel = StartNode.ImplicitCast())
+	{
+		bIsTrackOrObjectBinding = true;
+
+		UMovieSceneTrack* Track = TrackModel->GetTrack();
+
+		// Always show subsequence tracks so that the user can navigate up and down the hierarchy with the filter
+		if (Track && Track->IsA<UMovieSceneSubTrack>())
+		{
+			bPassedAnyFilters = true;
+			bIsSubTrack = true;
+			goto NextStep;
+		}
+
+		for (TSharedPtr<FChannelGroupModel> ChannelGroupModel : StartNode.AsModel()->GetDescendantsOfType<FChannelGroupModel>())
+		{
+			for (const TWeakViewModelPtr<FChannelModel>& ChannelModel : ChannelGroupModel->GetChannels())
+			{
+				const TSharedPtr<IKeyArea>& KeyArea = ChannelModel.Pin()->GetKeyArea();
+				FMovieSceneChannel* Channel = KeyArea->ResolveChannel();
+				if (Channel)
+				{
+					if (Filters->Num() == 0 || Filters->PassesAnyFilters(Channel))
+					{
+						bPassedAnyFilters = true;
+						break;
+					}
+				}
+			}
+		}
+	
+		if (bPassedAnyFilters || Filters->Num() == 0 || Filters->PassesAnyFilters(Track, TrackModel->GetLabel()))
+		{
+			bPassedAnyFilters = true;
+
+			// Track nodes do not belong to a level, but might be a child of an objectbinding node that does
+			if (LevelTrackFilter->IsActive())
+			{
+				TSharedPtr<IObjectBindingExtension> ObjectNode = StartNode.AsModel()->FindAncestorOfType<IObjectBindingExtension>();
+				if (ObjectNode)
+				{
+					// The track belongs to an objectbinding node, start by assuming it doesn't match the level filter
+					bPassedAnyFilters = false;
+
+					for (TWeakObjectPtr<>& Object : Sequencer.FindObjectsInCurrentSequence(ObjectNode->GetObjectGuid()))
+					{
+						if (Object.IsValid() && LevelTrackFilter->PassesFilter(Object.Get()))
+						{
+							// If at least one of the objects on the objectbinding node pass the level filter, show the track
+							bPassedAnyFilters = true;
+							break;
+						}
+					}
+				}
+			}
+
+			if (bPassedAnyFilters && Sequencer.GetSequencerSettings()->GetShowSelectedNodesOnly())
+			{
+				TSharedPtr<IObjectBindingExtension> ObjectNode = StartNode.AsModel()->FindAncestorOfType<IObjectBindingExtension>();
+				// Always show pinned items
+				if (!bIsPinned && ObjectNode)
+				{
+					const FMovieSceneBinding* Binding = Sequencer.GetFocusedMovieSceneSequence()->GetMovieScene()->FindBinding(ObjectNode->GetObjectGuid());
+					if (!(Binding && Sequencer.IsBindingVisible(*Binding)))
+					{
+						return bAnyChildPassed;
+					}
+				}
+			}
+		}
+	}
+	else if (TViewModelPtr<FObjectBindingModel> ObjectNode = StartNode.ImplicitCast())
+	{
+		bIsTrackOrObjectBinding = true;
+
+		for (TWeakObjectPtr<>& Object : Sequencer.FindObjectsInCurrentSequence(ObjectNode->GetObjectGuid()))
+		{
+			if (Object.IsValid() && (Filters->Num() == 0 || Filters->PassesAnyFilters(Object.Get(), StartNode->GetLabel()))
+				&& LevelTrackFilter->PassesFilter(Object.Get()))
+			{
+				bPassedAnyFilters = true;
+				break;
+			}
+		}
+
+		if (bPassedAnyFilters && Sequencer.GetSequencerSettings()->GetShowSelectedNodesOnly() && !bIsPinned)
+		{
+			UMovieScene* MovieScene = Sequencer.GetFocusedMovieSceneSequence()->GetMovieScene();
+			const FMovieSceneBinding* Binding = Sequencer.GetFocusedMovieSceneSequence()->GetMovieScene()->FindBinding(ObjectNode->GetObjectGuid());
+			if (Binding && !Sequencer.IsBindingVisible(*Binding))
+			{
+				return bAnyChildPassed;
+			}
+		}
+	}
+	else if (TViewModelPtr<FCategoryModel> CategoryModel = StartNode.ImplicitCast())
+	{
+		if (TSharedPtr<FTrackModel> TrackNode = CategoryModel->FindAncestorOfType<FTrackModel>())
+		{
+			UMovieSceneTrack* Track = TrackNode->GetTrack();
+			if (Filters->Num() == 0 || Filters->PassesAnyFilters(Track, StartNode->GetLabel()))
+			{
+				bPassedAnyFilters = true;
+			}
+		}
+		if (bPassedAnyFilters && Sequencer.GetSequencerSettings()->GetShowSelectedNodesOnly())
+		{
+			TSharedPtr<IObjectBindingExtension> ParentObjectNode = StartNode.AsModel()->FindAncestorOfType<IObjectBindingExtension>();
+			// Always show pinned items
+			if (!bIsPinned && ParentObjectNode)
+			{
+				const FMovieSceneBinding* Binding = Sequencer.GetFocusedMovieSceneSequence()->GetMovieScene()->FindBinding(ParentObjectNode->GetObjectGuid());
+				if (!(Binding && Sequencer.IsBindingVisible(*Binding)))
+				{
+					return bAnyChildPassed;
+				}
+			}
+		}
+	}
+	else if (TViewModelPtr<FCategoryGroupModel> CategoryGroupModel = StartNode.ImplicitCast())
+	{
+		if (TSharedPtr<FTrackModel> TrackNode = CategoryGroupModel->FindAncestorOfType<FTrackModel>())
+		{
+			UMovieSceneTrack* Track = TrackNode->GetTrack();
+			if (Filters->Num() == 0 || Filters->PassesAnyFilters(Track, FText::FromName(CategoryGroupModel->GetCategoryName())))
+			{
+				bPassedAnyFilters = true;
+			}
+		}
+		if (bPassedAnyFilters && Sequencer.GetSequencerSettings()->GetShowSelectedNodesOnly())
+		{
+			TSharedPtr<IObjectBindingExtension> ParentObjectNode = StartNode.AsModel()->FindAncestorOfType<IObjectBindingExtension>();
+			// Always show pinned items
+			if (!bIsPinned && ParentObjectNode)
+			{
+				const FMovieSceneBinding* Binding = Sequencer.GetFocusedMovieSceneSequence()->GetMovieScene()->FindBinding(ParentObjectNode->GetObjectGuid());
+				if (!(Binding && Sequencer.IsBindingVisible(*Binding)))
+				{
+					return bAnyChildPassed;
+				}
+			}
+		}
+	}
+	else if (TViewModelPtr<FChannelGroupModel> ChannelGroupModel = StartNode.ImplicitCast())
+	{
+		for (const TWeakViewModelPtr<FChannelModel>& ChannelModel : ChannelGroupModel->GetChannels())
+		{
+			const TSharedPtr<IKeyArea>& KeyArea = ChannelModel.Pin()->GetKeyArea();
+			FMovieSceneChannel* Channel = KeyArea->ResolveChannel();
+			if (Channel)
+			{
+				if (Filters->Num() == 0 || Filters->PassesAnyFilters(Channel))
+				{
+					bPassedAnyFilters = true;
+					break;
+				}
+			}
+		}
+		if (bPassedAnyFilters && Sequencer.GetSequencerSettings()->GetShowSelectedNodesOnly())
+		{
+			TSharedPtr<IObjectBindingExtension> ParentObjectNode = StartNode.AsModel()->FindAncestorOfType<IObjectBindingExtension>();
+			// Always show pinned items
+			if (!bIsPinned && ParentObjectNode)
+			{
+				const FMovieSceneBinding* Binding = Sequencer.GetFocusedMovieSceneSequence()->GetMovieScene()->FindBinding(ParentObjectNode->GetObjectGuid());
+				if (!(Binding && Sequencer.IsBindingVisible(*Binding)))
+				{
+					bPassedAnyFilters = false;
+				}
+			}
+		}
+	}
+	else if (TViewModelPtr<FFolderModel> FolderModel = StartNode.ImplicitCast())
+	{
+		// Special case: If we're pinned, then we should pass regardless
+		if (bIsPinned)
+		{
+			bPassedAnyFilters = true;
+		}
+
+		// Special case: If we're only filtering on text search, include folders and key areas in the search
+		if (!bPassedAnyFilters && Filters->Num() == 0 && FilterStrings.Num() > 0)
+		{
+			bPassedAnyFilters = true;
+
+			// Special case: but don't include if only showing selected bindings and we don't have child that passed
+			if (Sequencer.GetSequencerSettings()->GetShowSelectedNodesOnly() && !bAnyChildPassed)
+			{
+				bPassedAnyFilters = false;
+
+				// Special case: unless we're the child of a node that is a selected binding
+				TSharedPtr<IObjectBindingExtension> ParentObjectNode = StartNode.AsModel()->FindAncestorOfType<IObjectBindingExtension>();
+				// Always show pinned items
+				if (ParentObjectNode)
+				{
+					const FMovieSceneBinding* Binding = Sequencer.GetFocusedMovieSceneSequence()->GetMovieScene()->FindBinding(ParentObjectNode->GetObjectGuid());
+					if (Binding && Sequencer.IsBindingVisible(*Binding))
+					{
+						bPassedAnyFilters = true;
+					}
+				}
+			}
+		}
+	}
+
+NextStep:
+
+	if (bPassedAnyFilters && !bIsSubTrack)
+	{
+		// If filtering on selection set is enabled, we need to run another pass to verify we're in an enabled node group
+		UMovieSceneNodeGroupCollection& NodeGroups = Sequencer.GetFocusedMovieSceneSequence()->GetMovieScene()->GetNodeGroups();
+		if (NodeGroups.HasAnyActiveFilter())
+		{
+			bPassedAnyFilters = false;
+			
+			constexpr bool bIncludeThis = true;
+			for (TViewModelPtr<IOutlinerExtension> Parent : StartNode.AsModel()->GetAncestorsOfType<IOutlinerExtension>(bIncludeThis))
+			{
+				// Special case: Pinned tracks should be visible whether in the node group or not
+				if (bIsPinned)
+				{
+					bPassedAnyFilters = true;
+					break;
+				}
+
+				for (const UMovieSceneNodeGroup* NodeGroup : NodeGroups)
+				{
+					if (NodeGroup->GetEnableFilter() && NodeGroup->ContainsNode(IOutlinerExtension::GetPathName(Parent.AsModel())))
+					{
+						bPassedAnyFilters = true;
+						break;
+					}
+				}
+			}
+		}
+	}
+
+	if (bPassedAnyFilters)
+	{
+		OutFilteredNodes.Add(StartNode);
+		AddParentNodes(StartNode, OutFilteredNodes);
+
+		// Special case: When only showing selected bindings, and a non-object passes text filtering
+		// don't add it's children, as they may be a binding that is not seleceted. Selected child nodes will add themselves.
+		if (!(Sequencer.GetSequencerSettings()->GetShowSelectedNodesOnly() && FilterStrings.Num() > 0 && !bIsTrackOrObjectBinding))
+		{
+			AddChildNodes(StartNode, OutFilteredNodes);
+		}
+
+		return true;
+	}
+
+	return bAnyChildPassed;
+}
+
+bool FSequencerNodeTree::UpdateFilters()
+{
+	using namespace UE::Sequencer;
+
+	if (!bFilterUpdateRequested)
+	{
+		return false;
+	}
+
+	TSet<TWeakViewModelPtr<IOutlinerExtension>> PreviousFilteredNodes(FilteredNodes);
+
+	FilteredNodes.Empty();
+	const bool bHasActiveFilter = HasActiveFilter();
+
+	UObject* PlaybackContext = Sequencer.GetPlaybackContext();
+	UWorld* World = PlaybackContext ? PlaybackContext->GetWorld() : nullptr;
+	TrackFilterLevelFilter->UpdateWorld(World);
+
+	if (bHasActiveFilter)
+	{
+		// Build a list of strings that must be matched
+		TArray<FString> FilterStrings;
+
+		// Remove whitespace from the front and back of the string
+		FilterString.TrimStartAndEndInline();
+		FilterString.ParseIntoArray(FilterStrings, TEXT(" "), true /*bCullEmpty*/);
+
+		for (const TViewModelPtr<IOutlinerExtension>& Node : GetRootNodes())
+		{
+			// Recursively filter all nodes, matching them against the list of filter strings.  All filter strings must be matched
+			FilterNodesRecursive(Sequencer, Node, TrackFilters, FilterStrings, TrackFilterLevelFilter, FilteredNodes);
+		}
+	}
+
+	bFilteringOnNodeGroups = Sequencer.GetFocusedMovieSceneSequence()->GetMovieScene()->GetNodeGroups().HasAnyActiveFilter();
+	bFilterUpdateRequested = false;
+
+	// Always include the bottom spacer
+	if (FSequenceModel* SequenceModel = RootNode->CastThis<FSequenceModel>())
+	{
+		FilteredNodes.Add(SequenceModel->GetBottomSpacer());
+	}
+
+	// Count the total number of display nodes, and update filtered state.
+	DisplayNodeCount = 0;
+	for (const TViewModelPtr<IOutlinerExtension>& Item : RootNode->GetDescendantsOfType<IOutlinerExtension>())
+	{
+		const bool bIsNodeFilteredIn = !bHasActiveFilter || FilteredNodes.Contains(Item);
+		Item->SetFilteredOut(!bIsNodeFilteredIn);
+
+		if (FOutlinerItemModel* OutlinerItemModel = Item.AsModel()->CastThis<FOutlinerItemModel>())
+		{
+			if (OutlinerItemModel->IsForceFilteredOut())
+			{
+				// Skip things that are force-filtered-out because we don't want to see or count them
+				// in the outliner hierarchy.
+				continue;
+			}
+		}
+
+		++DisplayNodeCount;
+	}
+
+	// Return whether the new list of FilteredNodes is different than the previous list
+	return (PreviousFilteredNodes.Num() != FilteredNodes.Num() || !PreviousFilteredNodes.Includes(FilteredNodes));
+}
+
+void FSequencerNodeTree::CleanupMuteSolo(UMovieScene* MovieScene)
+{
+	// Remove mute/solo markers for any nodes that no longer exist
+	if (!MovieScene->IsReadOnly())
+	{
+		for (auto It = MovieScene->GetSoloNodes().CreateIterator(); It; ++It)
+		{
+			if (!GetNodeAtPath(*It))
+			{
+				It.RemoveCurrent();
+			}
+		}
+
+		for (auto It = MovieScene->GetMuteNodes().CreateIterator(); It; ++It)
+		{
+			if (!GetNodeAtPath(*It))
+			{
+				It.RemoveCurrent();
+			}
+		}
+	}
+}
+
+void FSequencerNodeTree::UpdateMuteSolo(UMovieScene* MovieScene)
+{
+	// Muting overrides soloing, so a solo node only counts if it's not muted.
+	bHasSoloNodes = false;
+	TArray<FString>& SoloNodes = Sequencer.GetFocusedMovieSceneSequence()->GetMovieScene()->GetSoloNodes();
+	for (const FString& NodePath : SoloNodes)
+	{
+		UE::Sequencer::TViewModelPtr<UE::Sequencer::IOutlinerExtension> Node = GetNodeAtPath(NodePath);
+		if (Node)
+		{
+			if (!IsNodeMute(Node))
+			{
+				bHasSoloNodes = true;
+				break;
+			}
+		}
+	}
+}
+
+int32 FSequencerNodeTree::GetTotalDisplayNodeCount() const 
+{ 
+	// Subtract 1 for the spacer node which is always added
+	return DisplayNodeCount - 1; 
+}
+
+int32 FSequencerNodeTree::GetFilteredDisplayNodeCount() const 
+{ 
+	// Subtract 1 for the spacer node which is always added
+	return FilteredNodes.Num() - 1;
+}
+
+void FSequencerNodeTree::FilterNodes(const FString& InFilter)
+{
+	if (InFilter != FilterString)
+	{
+		FilterString = InFilter;
+		bFilterUpdateRequested = true;
+	}
+}
+
+void FSequencerNodeTree::NodeGroupsCollectionChanged()
+{
+	if (Sequencer.GetFocusedMovieSceneSequence()->GetMovieScene()->GetNodeGroups().HasAnyActiveFilter() || bFilteringOnNodeGroups)
+	{
+		RequestFilterUpdate();
+	}
+}
+
+void FSequencerNodeTree::GetAllNodes(TArray<TSharedPtr<UE::Sequencer::FViewModel>>& OutNodes) const
+{
+	using namespace UE::Sequencer;
+
+	const bool bIncludeRootNode = false;
+	for (const FViewModelPtr& It : RootNode->GetDescendants(bIncludeRootNode))
+	{
+		OutNodes.Add(It.AsModel());
+	}
+}
+
+void FSequencerNodeTree::GetAllNodes(TArray<TSharedRef<UE::Sequencer::FViewModel>>& OutNodes) const
+{
+	using namespace UE::Sequencer;
+
+	const bool bIncludeRootNode = false;
+	for (const FViewModelPtr& It : RootNode->GetDescendants(bIncludeRootNode))
+	{
+		OutNodes.Add(It.AsModel().ToSharedRef());
+	}
+}
+
