Sequencer: MVVM2 branch and Layer Bars

diff --git a/Engine/Source/Editor/Sequencer/Private/SequencerNodeTree.h b/Engine/Source/Editor/Sequencer/Private/SequencerNodeTree.h
--- a/Engine/Source/Editor/Sequencer/Private/SequencerNodeTree.h
+++ b/Engine/Source/Editor/Sequencer/Private/SequencerNodeTree.h
@@ -1,368 +1,260 @@
-// Copyright Epic Games, Inc. All Rights Reserved.
-
-#pragma once
-
-#include "CoreMinimal.h"
-#include "Misc/Guid.h"
-#include "UObject/ObjectKey.h"
-#include "Tree/CurveEditorTree.h"
-#include "DisplayNodes/SequencerDisplayNode.h"
-#include "SectionHandle.h"
-#include "MovieSceneSequence.h"
-
-class FSequencer;
-class FCurveEditor;
-class FSequencerDisplayNode;
-class FSequencerObjectBindingNode;
-class FSequencerSectionKeyAreaNode;
-class FSequencerTrackNode;
-class FSequencerFolderNode;
-class ISequencerTrackEditor;
-class UMovieSceneFolder;
-class UMovieSceneTrack;
-class UMovieScene;
-struct FMovieSceneBinding;
-struct FCurveEditorTreeItemID;
-class FSequencerTrackFilter;
-class FSequencerTrackFilterCollection;
-class FSequencerTrackFilter_LevelFilter;
-
-/**
- * Represents a tree of sequencer display nodes, used to populate the Sequencer UI with MovieScene data
- */
-class FSequencerNodeTree : public TSharedFromThis<FSequencerNodeTree>
-{
-public:
-	DECLARE_MULTICAST_DELEGATE(FOnUpdated);
-
-public:
-
-	FSequencerNodeTree(FSequencer& InSequencer);
-	
-	~FSequencerNodeTree();
-
-	/**
-	 * Updates the tree with sections from a MovieScene
-	 */
-	void Update();
-
-	/**
-	 * Access this tree's symbolic root node
-	 */
-	TSharedRef<FSequencerDisplayNode> GetRootNode() const;
-
-	/**
-	 * @return The root nodes of the tree
-	 */
-	const TArray< TSharedRef<FSequencerDisplayNode> >& GetRootNodes() const;
-
-	/** @return Whether or not there is an active filter */
-	bool HasActiveFilter() const;
-
-	/** 
-	 * Checks if filters should be updated on track value changes, and if so updates them.
-	 * @return Whether the filtered node list was modified
-	 */
-	bool UpdateFiltersOnTrackValueChanged();
-
-	/**
-	 * Returns whether or not a node is filtered
-	 *
-	 * @param Node	The node to check if it is filtered
-	 */
-	bool IsNodeFiltered( const TSharedRef<const FSequencerDisplayNode> Node ) const;
-
-	/**
-	 * Schedules an update of all filters
-	 */
-	void RequestFilterUpdate() { bFilterUpdateRequested = true; }
-
-	/**
-	 * @return Whether there is a filter update scheduled
-	 */
-	bool NeedsFilterUpdate() const { return bFilterUpdateRequested; }
-
-	/**
-	 * Filters the nodes based on the passed in filter terms
-	 *
-	 * @param InFilter	The filter terms
-	 */
-	void FilterNodes( const FString& InFilter );
-
-	/** Called when the active MovieScene's node group colletion has been modifed */
-	void NodeGroupsCollectionChanged();
-
-	/**
-	 * Unpins any pinned nodes in this tree
-	 */
-	void UnpinAllNodes();
-
-	/**
-	 * @return All nodes in a flat array
-	 */
-	TArray< TSharedRef<FSequencerDisplayNode> > GetAllNodes() const;
-
-	/** Gets the parent sequencer of this tree */
-	FSequencer& GetSequencer() {return Sequencer;}
-
-	/**
-	 * Saves the expansion state of a display node
-	 *
-	 * @param Node		The node whose expansion state should be saved
-	 * @param bExpanded	The new expansion state of the node
-	 */
-	void SaveExpansionState( const FSequencerDisplayNode& Node, bool bExpanded );
-
-	/**
-	 * Gets the saved expansion state of a display node
-	 *
-	 * @param Node	The node whose expansion state may have been saved
-	 * @return true if the node should be expanded, false otherwise	
-	 */
-	bool GetSavedExpansionState( const FSequencerDisplayNode& Node ) const;
-
-	/**
-	 * Get the default expansion state for the specified node, where its state has not yet been saved
-	 *
-	 * @return true if the node is to be expanded, false otherwise
-	 */
-	bool GetDefaultExpansionState( const FSequencerDisplayNode& Node ) const;
-
-	/**
-	 * Saves the pinned state of a display node
-	 *
-	 * @param Node		The node whose pinned state should be saved
-	 * @param bPinned	The new pinned state of the node
-	 */
-	void SavePinnedState( const FSequencerDisplayNode& Node, bool bPinned );
-
-	/**
-	 * Gets the saved pinned state of a display node
-	 *
-	 * @param Node	The node whose pinned state may have been saved
-	 * @return true if the node should be pinned, false otherwise	
-	 */
-	bool GetSavedPinnedState( const FSequencerDisplayNode& Node ) const;
-
-	/**
-	 * Set the single hovered node in the tree
-	 */
-	void SetHoveredNode(const TSharedPtr<FSequencerDisplayNode>& InHoveredNode);
-
-	/**
-	 * Get the single hovered node in the tree, possibly nullptr
-	 */
-	const TSharedPtr<FSequencerDisplayNode>& GetHoveredNode() const;
-
-	/*
-	 * Find the object binding node with the specified GUID
-	 */
-	TSharedPtr<FSequencerObjectBindingNode> FindObjectBindingNode(const FGuid& BindingID) const;
-
-	/*
-	 * Gets a multicast delegate which is called whenever the node tree has been updated.
-	 */
-	FOnUpdated& OnUpdated() { return OnUpdatedDelegate; }
-
-	/** Make the contents of the given node have the root as their parent again instead of their current parent. */
-	void MoveDisplayNodeToRoot(TSharedRef<FSequencerDisplayNode>& Node);
-
-	/** Sorts all nodes and their descendants by category then alphabetically.*/
-	void SortAllNodesAndDescendants();
-
-	/**
-	 * Attempt to find a curve editor tree item ID for the specified display node
-	 *
-	 * @param DisplayNode The node to find a curve editor tree item for
-	 * @return The ID of the tree item for this display node in the curve editor, or FCurveEditorTreeItemID::Invalid() if one was not found.
-	 */
-	FCurveEditorTreeItemID FindCurveEditorTreeItem(TSharedRef<FSequencerDisplayNode> DisplayNode) const;
-
-	/**
-	 * Attempt to get a node at the specified path
-	 *
-	 * @param NodePath The path of the node to search for
-	 * @return The node located at NodePath, or nullptr if not found
-	 */
-	FSequencerDisplayNode* GetNodeAtPath(const FString& NodePath) const;
-
-public:
-
-	/**
-	 * Finds the section handle relating to the specified section object, or Null if one was not found (perhaps, if it's filtered away)
-	 */
-	TOptional<FSectionHandle> GetSectionHandle(const UMovieSceneSection* Section) const;
-
-	void AddFilter(TSharedRef<FSequencerTrackFilter> TrackFilter);
-	int32 RemoveFilter(TSharedRef<FSequencerTrackFilter> TrackFilter);
-	void RemoveAllFilters();
-	bool IsTrackFilterActive(TSharedRef<FSequencerTrackFilter> TrackFilter) const;
-
-	void AddLevelFilter(const FString& LevelName);
-	void RemoveLevelFilter(const FString& LevelName);
-	bool IsTrackLevelFilterActive(const FString& LevelName) const;
-
-	void ToggleSelectedNodesSolo();
-	bool IsNodeSolo(const FSequencerDisplayNode* Node) const;
-	
-	/** Returns whether any of the nodes in this tree are marked solo */
-	bool HasSoloNodes() const;
-
-	void ToggleSelectedNodesMute();
-	bool IsNodeMute(const FSequencerDisplayNode* Node) const;
-
-	/** Returns whether any of the currently selected nodes are marked solo */
-	bool IsSelectedNodesSolo() const;
-	
-	/** Returns whether any of the currently selected nodes are muted */
-	bool IsSelectedNodesMute() const;
-
-	/**Make sure associated curve editor get's recreated, usually due to an item getting added that needs to be sorted*/
-	void RecreateCurveEditor();
-
-private:
-
-	/** Returns whether this NodeTree should only display selected nodes */
-	bool ShowSelectedNodesOnly() const;
-
-	/** Population algorithm utilities */
-	void RefreshNodes(UMovieScene* MovieScene);
-
-	enum class ETrackType { Master, Object };
-
-	/**
-	 * Create or update a track node for the specified track object, updating its serial number.
-	 *
-	 * @param Track                   Pointer to the track to create or update a node for
-	 * @param TrackType               Whether this track is a master track or contained within an object
-	 * @return A shared pointer to a track node with its serial number updated. Will return nullptr for tracks that are forcibly hidden through FSequencer::IsTrackVisible.
-	 */
-	TSharedPtr<FSequencerTrackNode> CreateOrUpdateTrack(UMovieSceneTrack* Track, ETrackType TrackType);
-
-	/**
-	 * Create or update a folder node for the specified folder and all its decendent child folders and object bindings, updating their serial numbers in the process
-	 *
-	 * @param Folder                  Pointer to the movie scene folder to create a node for
-	 * @param AllBindings             A map from guid to binding pointer for all the bindings in the sequence
-	 * @param ChildToParentBinding    Child to parent GUID map used for creating parent items
-	 * @param OutChildToParentMap     Pointer to a map that should be populated with child->parent relationships to be set up after creation of all nodes
-	 * @return A shared reference to a folder node. The resulting node's serial number will always be up-to-date, as will all its child folders, and any immediate child object bindings.
-	 */
-	TSharedRef<FSequencerFolderNode> CreateOrUpdateFolder(UMovieSceneFolder* Folder, const TSortedMap<FGuid, const FMovieSceneBinding*>& AllBindings, const TSortedMap<FGuid, FGuid>& ChildToParentBinding, TMap<FSequencerDisplayNode*, TSharedPtr<FSequencerDisplayNode>>* OutChildToParentMap);
-
-	/**
-	 * Find an existing object binding node (or create a new one) for the specified binding ID without updating its tree serial number, creating any parent object binding nodes in the process.
-	 * @note: Will only update FSequencerDisplayNode::TreeSerialNumber for object bindings that have a known and valid parent object binding.
-	 *
-	 * @param BindingID               The Guid of the object binding within UMovieScene::GetBindings
-	 * @param AllBindings             A map from guid to binding pointer for all the bindings in the sequence
-	 * @param ChildToParentBinding    Child to parent GUID map used for creating parent items
-	 * @param OutChildToParentMap     Pointer to a map that should be populated with child->parent relationships to be set up after creation of all nodes
-	 * @return A shared pointer to an object binding node or nullptr if the supplied BindingID was not valid for this sequence. The resulting node's serial number will be up-to-date if it is a child of another binding, or it was previously added to a folder.
-	 */
-	TSharedPtr<FSequencerObjectBindingNode> FindOrCreateObjectBinding(const FGuid& BindingID, const TSortedMap<FGuid, const FMovieSceneBinding*>& AllBindings, const TSortedMap<FGuid, FGuid>& ChildToParentBinding, TMap<FSequencerDisplayNode*, TSharedPtr<FSequencerDisplayNode>>* OutChildToParentMap);
-
-	/**
-	 * Creates section handles for all the sections contained in the specified track
-	 */
-	void UpdateSectionHandles(TSharedRef<FSequencerTrackNode> TrackNode);
-
-private:
-
-	/**
-	 * Update the list of filters nodes based on current filter settings, if an update is scheduled
-	 * This is called by Update();
-	 * 
-	 * @return Whether the list of filtered nodes changed
-	 */
-	bool UpdateFilters();
-
-	/*
-	 * Return whether there is a valid type editor for the track 
-	 */
-	bool HasValidTypeEditor(UMovieSceneTrack* Track) const;
-
-	/**
-	 * Finds or adds a type editor for the track
-	 *
-	 * @param Track	The type to find an editor for
-	 * @return The editor for the type
-	 */
-	TSharedRef<ISequencerTrackEditor> FindOrAddTypeEditor( UMovieSceneTrack* Track );
-
-	/**
-	 * Update the curve editor tree to include anything of relevance from this tree
-	 */
-	void UpdateCurveEditorTree();
-
-	/**
-	 * Add the specified node to the curve editor, including all its parents if necessary
-	 */
-	FCurveEditorTreeItemID AddToCurveEditor(TSharedRef<FSequencerDisplayNode> Node, FCurveEditor* CurveEditor);
-
-	/**
-	 * Checks whether the specified key area node contains key areas that can create curve models
-	 */
-	bool KeyAreaHasCurves(const FSequencerSectionKeyAreaNode& KeyAreaNode) const;
-
-	/**
-	 * Destroys all nodes contained within this tree.
-	 * @note: Does not broadcast update notifications
-	 */
-	void DestroyAllNodes();
-
-public:
-	int32 GetTotalDisplayNodeCount() const;
-	int32 GetFilteredDisplayNodeCount() const;
-
-private:
-
-	/** Symbolic root node that contains the actual displayed root nodes as children */
-	TSharedRef<FSequencerDisplayNode> RootNode;
-
-	TSharedRef<FSequencerDisplayNode> BottomSpacerNode;
-
-	/** A serially incrementing integer that is increased each time the tree is refreshed to track node relevance */
-	uint32 SerialNumber;
-	/** Map from FMovieSceneBinding::GetObjectGuid to display node */
-	TMap<FGuid,      TSharedPtr<FSequencerObjectBindingNode>> ObjectBindingToNode;
-	/** Map from UMovieSceneTrack object key to display node for any track (object tracks or master tracks) */
-	TMap<FObjectKey, TSharedPtr<FSequencerTrackNode>>         TrackToNode;
-	/** Map from UMovieSceneFolder object key to display node for any folder (root or child folder) */
-	TMap<FObjectKey, TSharedPtr<FSequencerFolderNode>>        FolderToNode;
-
-	/** Map from UMovieSceneSection* to its UI handle */
-	TMap<FObjectKey, FSectionHandle> SectionToHandle;
-
-	/** Tools for building movie scene section layouts.  One tool for each track */
-	TMap< UMovieSceneTrack*, TSharedPtr<ISequencerTrackEditor> > EditorMap;
-	/** Set of all filtered nodes */
-	TSet< TSharedRef<FSequencerDisplayNode> > FilteredNodes;
-	/** Cardinal hovered node */
-	TSharedPtr<FSequencerDisplayNode> HoveredNode;
-	/** Active filter string if any */
-	FString FilterString;
-	/** Sequencer interface */
-	FSequencer& Sequencer;
-	/** Display node -> curve editor tree item ID mapping */
-	TMap<TWeakPtr<FSequencerDisplayNode>, FCurveEditorTreeItemID> CurveEditorTreeItemIDs;
-	/** A multicast delegate which is called whenever the node tree has been updated. */
-	FOnUpdated OnUpdatedDelegate;
-
-	/** Active track filters */
-	TSharedPtr<FSequencerTrackFilterCollection> TrackFilters;
-	
-	/** Level based track filtering */
-	TSharedPtr<FSequencerTrackFilter_LevelFilter> TrackFilterLevelFilter;
-
-	TWeakObjectPtr<UMovieSceneSequence> WeakCurrentSequence;
-
-	/** The total number of DisplayNodes in the tree, both displayed and hidden */
-	uint32 DisplayNodeCount;
-
-	bool bFilterUpdateRequested;
-	bool bFilteringOnNodeGroups;
-
-	/** Cached value of whether we have any nodes that should be treated as soloing */
-	bool bHasSoloNodes;
-};
+// Copyright Epic Games, Inc. All Rights Reserved.
+
+#pragma once
+
+#include "CoreMinimal.h"
+#include "MVVM/ViewModels/ViewModelHierarchy.h"
+#include "Misc/Guid.h"
+#include "UObject/ObjectKey.h"
+#include "MVVM/ViewModels/ViewModel.h"
+#include "MVVM/ViewModels/ObjectBindingModel.h"
+#include "MVVM/ViewModels/TrackModel.h"
+#include "MVVM/ViewModels/SectionModel.h"
+#include "MVVM/ViewModels/OutlinerItemModel.h"
+#include "MVVM/ViewModels/ChannelModel.h"
+#include "MVVM/ViewModels/SequenceModel.h"
+#include "MVVM/Extensions/IOutlinerExtension.h"
+#include "MovieSceneSequence.h"
+
+class FSequencer;
+class FSequencerTrackFilter;
+class FSequencerTrackFilterCollection;
+class FSequencerTrackFilter_LevelFilter;
+class ISequencerTrackEditor;
+class UMovieScene;
+class UMovieSceneFolder;
+class UMovieSceneTrack;
+struct FMovieSceneBinding;
+
+/**
+ * Represents a tree of sequencer display nodes, used to populate the Sequencer UI with MovieScene data
+ */
+class FSequencerNodeTree : public TSharedFromThis<FSequencerNodeTree>
+{
+public:
+
+	DECLARE_MULTICAST_DELEGATE(FOnUpdated);
+
+	using FSectionModel = UE::Sequencer::FSectionModel;
+	using FViewModel = UE::Sequencer::FViewModel;
+	using FObjectBindingModel = UE::Sequencer::FObjectBindingModel;
+
+public:
+
+	FSequencerNodeTree(FSequencer& InSequencer);
+	
+	~FSequencerNodeTree();
+
+	/**
+	 * Updates the tree with sections from a MovieScene
+	 */
+	void Update();
+
+	/**
+	 * Set this tree's symbolic root node
+	 */
+	void SetRootNode(const UE::Sequencer::FViewModelPtr& InRootNode);
+
+	/**
+	 * Access this tree's symbolic root node
+	 */
+	UE::Sequencer::FViewModelPtr GetRootNode() const;
+
+	/**
+	 * @return The root nodes of the tree
+	 */
+	TArray<UE::Sequencer::TViewModelPtr<UE::Sequencer::IOutlinerExtension>> GetRootNodes() const;
+
+	/** @return Whether or not there is an active filter */
+	bool HasActiveFilter() const;
+
+	/** 
+	 * Checks if filters should be updated on track value changes, and if so updates them.
+	 * @return Whether the filtered node list was modified
+	 */
+	bool UpdateFiltersOnTrackValueChanged();
+
+	/**
+	 * Returns whether or not a node is filtered
+	 *
+	 * @param Node	The node to check if it is filtered
+	 */
+	bool IsNodeFiltered(const TSharedPtr<FViewModel>& Node) const;
+
+	/**
+	 * Schedules an update of all filters
+	 */
+	void RequestFilterUpdate() { bFilterUpdateRequested = true; }
+
+	/**
+	 * @return Whether there is a filter update scheduled
+	 */
+	bool NeedsFilterUpdate() const { return bFilterUpdateRequested; }
+
+	/**
+	 * Filters the nodes based on the passed in filter terms
+	 *
+	 * @param InFilter	The filter terms
+	 */
+	void FilterNodes( const FString& InFilter );
+
+	/** Called when the active MovieScene's node group colletion has been modifed */
+	void NodeGroupsCollectionChanged();
+
+	/**
+	 * @return All nodes in a flat array
+	 */
+	void GetAllNodes(TArray<TSharedPtr<FViewModel>>& OutNodes) const;
+
+	/**
+	 * @return All nodes in a flat array
+	 */
+	void GetAllNodes(TArray<TSharedRef<FViewModel>>& OutNodes) const;
+
+	/** Gets the parent sequencer of this tree */
+	FSequencer& GetSequencer() {return Sequencer;}
+
+	/**
+	 * Saves the expansion state of a display node
+	 *
+	 * @param Node		The node whose expansion state should be saved
+	 * @param bExpanded	The new expansion state of the node
+	 */
+	void SaveExpansionState( const FViewModel& Node, bool bExpanded );
+
+	/**
+	 * Gets the saved expansion state of a display node
+	 *
+	 * @param Node	The node whose expansion state may have been saved
+	 * @return true if the node should be expanded, false otherwise, unset if it wasn't saved
+	 */
+	TOptional<bool> GetSavedExpansionState( const FViewModel& Node );
+
+	/**
+	 * Saves the pinned state of a display node
+	 *
+	 * @param Node		The node whose pinned state should be saved
+	 * @param bPinned	The new pinned state of the node
+	 */
+	void SavePinnedState( const FViewModel& Node, bool bPinned );
+
+	/**
+	 * Gets the saved pinned state of a display node
+	 *
+	 * @param Node	The node whose pinned state may have been saved
+	 * @return true if the node should be pinned, false otherwise	
+	 */
+	bool GetSavedPinnedState( const FViewModel& Node ) const;
+
+	/*
+	 * Find the object binding node with the specified GUID
+	 */
+	TSharedPtr<FObjectBindingModel> FindObjectBindingNode(const FGuid& BindingID) const;
+
+	/*
+	 * Gets a multicast delegate which is called whenever the node tree has been updated.
+	 */
+	FOnUpdated& OnUpdated() { return OnUpdatedDelegate; }
+
+	/** Sorts all nodes and their descendants by category then alphabetically.*/
+	void SortAllNodesAndDescendants();
+
+	/**
+	 * Attempt to get a node at the specified path
+	 *
+	 * @param NodePath The path of the node to search for
+	 * @return The node located at NodePath, or nullptr if not found
+	 */
+	UE::Sequencer::TViewModelPtr<UE::Sequencer::IOutlinerExtension> GetNodeAtPath(const FString& NodePath) const;
+
+public:
+
+	/**
+	 * Finds the section handle relating to the specified section object, or Null if one was not found (perhaps, if it's filtered away)
+	 */
+	TSharedPtr<FSectionModel> GetSectionModel(const UMovieSceneSection* Section) const;
+
+	void AddFilter(TSharedPtr<FSequencerTrackFilter> TrackFilter);
+	int32 RemoveFilter(TSharedPtr<FSequencerTrackFilter> TrackFilter);
+	void RemoveAllFilters();
+	bool IsTrackFilterActive(TSharedPtr<FSequencerTrackFilter> TrackFilter) const;
+
+	void AddLevelFilter(const FString& LevelName);
+	void RemoveLevelFilter(const FString& LevelName);
+	bool IsTrackLevelFilterActive(const FString& LevelName) const;
+
+	void ToggleSelectedNodesSolo();
+	bool IsNodeSolo(const UE::Sequencer::TViewModelPtr<UE::Sequencer::IOutlinerExtension>& InNode) const;
+	
+	/** Returns whether any of the nodes in this tree are marked solo */
+	bool HasSoloNodes() const;
+
+	void ToggleSelectedNodesMute();
+	bool IsNodeMute(const UE::Sequencer::TViewModelPtr<UE::Sequencer::IOutlinerExtension>& InNode) const;
+
+	/** Returns whether any of the currently selected nodes are marked solo */
+	bool IsSelectedNodesSolo() const;
+	
+	/** Returns whether any of the currently selected nodes are muted */
+	bool IsSelectedNodesMute() const;
+
+private:
+
+	/** Returns whether this NodeTree should only display selected nodes */
+	bool ShowSelectedNodesOnly() const;
+
+private:
+
+	/**
+	 * Update the list of filters nodes based on current filter settings, if an update is scheduled
+	 * This is called by Update();
+	 * 
+	 * @return Whether the list of filtered nodes changed
+	 */
+	bool UpdateFilters();
+
+	/**
+	 * Cleans outdated mute/solo markers
+	 */
+	void CleanupMuteSolo(UMovieScene* MovieScene);
+
+	/**
+	 * Updates mute/solo markers
+	 */
+	void UpdateMuteSolo(UMovieScene* MovieScene);
+
+public:
+
+	int32 GetTotalDisplayNodeCount() const;
+	int32 GetFilteredDisplayNodeCount() const;
+
+private:
+
+	/** Symbolic root node that contains the actual displayed root nodes as children */
+	UE::Sequencer::FViewModelPtr RootNode;
+
+	/** Set of all filtered nodes */
+	TSet<UE::Sequencer::TWeakViewModelPtr<UE::Sequencer::IOutlinerExtension>> FilteredNodes;
+	/** Active filter string if any */
+	FString FilterString;
+	/** Sequencer interface */
+	FSequencer& Sequencer;
+	/** A multicast delegate which is called whenever the node tree has been updated. */
+	FOnUpdated OnUpdatedDelegate;
+
+	/** Active track filters */
+	TSharedPtr<FSequencerTrackFilterCollection> TrackFilters;
+	
+	/** Level based track filtering */
+	TSharedPtr<FSequencerTrackFilter_LevelFilter> TrackFilterLevelFilter;
+
+	/** The total number of DisplayNodes in the tree, both displayed and hidden */
+	uint32 DisplayNodeCount;
+
+	bool bFilterUpdateRequested;
+	bool bFilteringOnNodeGroups;
+
+	/** Cached value of whether we have any nodes that should be treated as soloing */
+	bool bHasSoloNodes;
+};
+
diff --git a/Engine/Source/Editor/Sequencer/Private/SequencerSectionPainter.cpp b/Engine/Source/Editor/Sequencer/Private/SequencerSectionPainter.cpp
--- a/Engine/Source/Editor/Sequencer/Private/SequencerSectionPainter.cpp
+++ b/Engine/Source/Editor/Sequencer/Private/SequencerSectionPainter.cpp
@@ -3,9 +3,11 @@
 #include "SequencerSectionPainter.h"
 #include "MovieSceneSection.h"
 #include "MovieSceneTrack.h"
+#include "MVVM/ViewModels/SectionModel.h"
 
-FSequencerSectionPainter::FSequencerSectionPainter(FSlateWindowElementList& OutDrawElements, const FGeometry& InSectionGeometry, UMovieSceneSection& InSection)
-	: Section(InSection)
+FSequencerSectionPainter::FSequencerSectionPainter(FSlateWindowElementList& OutDrawElements, const FGeometry& InSectionGeometry, TSharedPtr<UE::Sequencer::FSectionModel> InSection)
+	: Section(*InSection->GetSection())
+	, SectionModel(InSection)
 	, DrawElements(OutDrawElements)
 	, SectionGeometry(InSectionGeometry)
 	, LayerId(0)
@@ -27,7 +29,7 @@
 
 UMovieSceneTrack* FSequencerSectionPainter::GetTrack() const
 {
-	return Section.GetTypedOuter<UMovieSceneTrack>();
+	return SectionModel->GetSection()->GetTypedOuter<UMovieSceneTrack>();
 }
 
 FLinearColor FSequencerSectionPainter::BlendColor(FLinearColor InColor)
diff --git a/Engine/Source/Editor/Sequencer/Private/SequencerSelectedKey.cpp b/Engine/Source/Editor/Sequencer/Private/SequencerSelectedKey.cpp
--- a/Engine/Source/Editor/Sequencer/Private/SequencerSelectedKey.cpp
+++ b/Engine/Source/Editor/Sequencer/Private/SequencerSelectedKey.cpp
@@ -3,29 +3,32 @@
 #include "SequencerSelectedKey.h"
 #include "Modules/ModuleManager.h"
 #include "IKeyArea.h"
+#include "MVVM/ViewModels/ChannelModel.h"
 #include "Channels/MovieSceneChannel.h"
 #include "ISequencerChannelInterface.h"
 
 FSelectedKeysByChannel::FSelectedKeysByChannel(TArrayView<const FSequencerSelectedKey> InSelectedKeys)
 {
+	using namespace UE::Sequencer;
+
 	TMap<const IKeyArea*, int32> KeyAreaToChannelIndex;
 
 	for (int32 Index = 0; Index < InSelectedKeys.Num(); ++Index)
 	{
 		FSequencerSelectedKey Key = InSelectedKeys[Index];
-		const IKeyArea* KeyArea = Key.KeyArea.Get();
+		TSharedPtr<FChannelModel> Channel = Key.WeakChannel.Pin();
 
-		if (KeyArea && Key.KeyHandle.IsSet())
+		if (Channel && Key.IsValid())
 		{
-			const int32* ChannelArrayIndex = KeyAreaToChannelIndex.Find(KeyArea);
+			const int32* ChannelArrayIndex = KeyAreaToChannelIndex.Find(Channel->GetKeyArea().Get());
 			if (!ChannelArrayIndex)
 			{
-				int32 NewIndex = SelectedChannels.Add(FSelectedChannelInfo(Key.KeyArea->GetChannel(), Key.KeyArea->GetOwningSection()));
-				ChannelArrayIndex = &KeyAreaToChannelIndex.Add(KeyArea, NewIndex);
+				int32 NewIndex = SelectedChannels.Add(FSelectedChannelInfo(Channel->GetKeyArea()->GetChannel(), Channel->GetSection()));
+				ChannelArrayIndex = &KeyAreaToChannelIndex.Add(Channel->GetKeyArea().Get(), NewIndex);
 			}
 
 			FSelectedChannelInfo& ThisChannelInfo = SelectedChannels[*ChannelArrayIndex];
-			ThisChannelInfo.KeyHandles.Add(Key.KeyHandle.GetValue());
+			ThisChannelInfo.KeyHandles.Add(Key.KeyHandle);
 			ThisChannelInfo.OriginalIndices.Add(Index);
 		}
 	}
diff --git a/Engine/Source/Editor/Sequencer/Private/SequencerSelectedKey.h b/Engine/Source/Editor/Sequencer/Private/SequencerSelectedKey.h
--- a/Engine/Source/Editor/Sequencer/Private/SequencerSelectedKey.h
+++ b/Engine/Source/Editor/Sequencer/Private/SequencerSelectedKey.h
@@ -12,6 +12,16 @@
 class IKeyArea;
 class UMovieSceneSection;
 
+namespace UE
+{
+namespace Sequencer
+{
+
+class FChannelModel;
+
+} // namespace Sequencer 
+} // namespace UE
+
 /**
  * Represents a selected key in the sequencer
  */
@@ -21,41 +31,43 @@
 	UMovieSceneSection* Section;
 
 	/** Key area providing the key */
-	TSharedPtr<IKeyArea> KeyArea;
+	TWeakPtr<UE::Sequencer::FChannelModel> WeakChannel;
 
 	/** Index of the key in the key area */
-	TOptional<FKeyHandle> KeyHandle;
+	FKeyHandle KeyHandle;
 
 public:
 
 	/** Create and initialize a new instance. */
-	FSequencerSelectedKey(UMovieSceneSection& InSection, TSharedPtr<IKeyArea> InKeyArea, FKeyHandle InKeyHandle)
+	FSequencerSelectedKey(UMovieSceneSection& InSection, TWeakPtr<UE::Sequencer::FChannelModel> InChannel, FKeyHandle InKeyHandle)
 		: Section(&InSection)
-		, KeyArea(MoveTemp(InKeyArea))
+		, WeakChannel(MoveTemp(InChannel))
 		, KeyHandle(InKeyHandle)
 	{}
 
 	/** Default constructor. */
 	FSequencerSelectedKey()
 		: Section(nullptr)
-		, KeyArea(nullptr)
-		, KeyHandle()
+		, KeyHandle(FKeyHandle::Invalid())
 	{}
 
 	/** @return Whether or not this is a valid selected key */
-	bool IsValid() const { return Section != nullptr && KeyArea.IsValid() && KeyHandle.IsSet(); }
+	bool IsValid() const
+	{
+		return Section != nullptr && WeakChannel.Pin().IsValid()
+			&& KeyHandle != FKeyHandle::Invalid();
+	}
 
 	friend uint32 GetTypeHash(const FSequencerSelectedKey& SelectedKey)
 	{
-		return GetTypeHash(SelectedKey.Section) ^ GetTypeHash(SelectedKey.KeyArea) ^ 
-			(SelectedKey.KeyHandle.IsSet() ? GetTypeHash(SelectedKey.KeyHandle.GetValue()) : 0);
+		return GetTypeHash(SelectedKey.Section) ^ GetTypeHash(SelectedKey.WeakChannel)
+			 ^ GetTypeHash(SelectedKey.KeyHandle);
 	} 
 
 	bool operator==(const FSequencerSelectedKey& OtherKey) const
 	{
-		return Section == OtherKey.Section && KeyArea == OtherKey.KeyArea &&
-			KeyHandle.IsSet() && OtherKey.KeyHandle.IsSet() &&
-			KeyHandle.GetValue() == OtherKey.KeyHandle.GetValue();
+		return Section == OtherKey.Section && WeakChannel == OtherKey.WeakChannel
+			&& KeyHandle == OtherKey.KeyHandle;
 	}
 };
 
diff --git a/Engine/Source/Editor/Sequencer/Private/SequencerSelection.cpp b/Engine/Source/Editor/Sequencer/Private/SequencerSelection.cpp
--- a/Engine/Source/Editor/Sequencer/Private/SequencerSelection.cpp
+++ b/Engine/Source/Editor/Sequencer/Private/SequencerSelection.cpp
@@ -1,443 +1,711 @@
-// Copyright Epic Games, Inc. All Rights Reserved.
-
-#include "SequencerSelection.h"
-#include "MovieSceneSection.h"
-#include "DisplayNodes/SequencerObjectBindingNode.h"
-#include "DisplayNodes/SequencerTrackNode.h"
-#include "SequencerCommonHelpers.h"
-
-FSequencerSelection::FSequencerSelection()
-	: SerialNumber(0)
-	, SuspendBroadcastCount(0)
-	, bOutlinerNodeSelectionChangedBroadcastPending(false)
-	, bEmptySelectedOutlinerNodesWithSectionsPending(false)
-{
-}
-
-const TSet<FSequencerSelectedKey>& FSequencerSelection::GetSelectedKeys() const
-{
-	return SelectedKeys;
-}
-
-const TSet<TWeakObjectPtr<UMovieSceneSection>>& FSequencerSelection::GetSelectedSections() const
-{
-	return SelectedSections;
-}
-
-const TSet<TSharedRef<FSequencerDisplayNode>>& FSequencerSelection::GetSelectedOutlinerNodes() const
-{
-	return SelectedOutlinerNodes;
-}
-
-const TSet<TSharedRef<FSequencerDisplayNode>>& FSequencerSelection::GetNodesWithSelectedKeysOrSections() const
-{
-	return NodesWithSelectedKeysOrSections;
-}
-
-FSequencerSelection::FOnSelectionChanged& FSequencerSelection::GetOnKeySelectionChanged()
-{
-	return OnKeySelectionChanged;
-}
-
-FSequencerSelection::FOnSelectionChanged& FSequencerSelection::GetOnSectionSelectionChanged()
-{
-	return OnSectionSelectionChanged;
-}
-
-FSequencerSelection::FOnSelectionChanged& FSequencerSelection::GetOnOutlinerNodeSelectionChanged()
-{
-	return OnOutlinerNodeSelectionChanged;
-}
-
-FSequencerSelection::FOnSelectionChanged& FSequencerSelection::GetOnNodesWithSelectedKeysOrSectionsChanged()
-{
-	return OnNodesWithSelectedKeysOrSectionsChanged;
-}
-
-FSequencerSelection::FOnSelectionChangedObjectGuids& FSequencerSelection::GetOnOutlinerNodeSelectionChangedObjectGuids()
-{
-	return OnOutlinerNodeSelectionChangedObjectGuids;
-}
-
-TArray<FGuid> FSequencerSelection::GetBoundObjectsGuids()
-{
-	TArray<FGuid> OutGuids;
-	TSet<TSharedRef<FSequencerDisplayNode>> SelectedNodes = GetNodesWithSelectedKeysOrSections();
-	if (SelectedNodes.Num() == 0)
-	{
-		SelectedNodes = GetSelectedOutlinerNodes();
-	}
-	for (TSharedRef<FSequencerDisplayNode> Node : SelectedNodes)
-	{
-		TSharedPtr<FSequencerObjectBindingNode> ObjectNode = Node->FindParentObjectBindingNode();
-		if (Node->GetType() == ESequencerNode::Object)
-		{
-			ObjectNode = StaticCastSharedRef<FSequencerObjectBindingNode>(Node);
-		}
-		else
-		{
-			ObjectNode = Node->FindParentObjectBindingNode();
-		}
-
-		if (ObjectNode.IsValid())
-		{
-			OutGuids.Add(ObjectNode->GetObjectBinding());
-		}
-	}
-
-	return OutGuids;
-}
-
-TArray<UMovieSceneTrack*> FSequencerSelection::GetSelectedTracks() const
-{
-	TArray<UMovieSceneTrack*> OutTracks;
-	
-	for (const TSharedRef<FSequencerDisplayNode>& SelectedNode : SelectedOutlinerNodes)
-	{
-		TSharedPtr<FSequencerDisplayNode> CurrentNode = SelectedNode;
-		while (CurrentNode.IsValid() && CurrentNode->GetType() != ESequencerNode::Track)
-		{
-			CurrentNode = CurrentNode->GetParent();
-		}
-
-		if (CurrentNode.IsValid())
-		{
-			UMovieSceneTrack* SelectedTrack = StaticCastSharedPtr<FSequencerTrackNode>(CurrentNode)->GetTrack();
-			if (SelectedTrack != nullptr)
-			{
-				OutTracks.Add(SelectedTrack);
-			}
-		}
-	}
-
-	return OutTracks;
-}
-
-void FSequencerSelection::AddToSelection(const FSequencerSelectedKey& Key)
-{
-	++SerialNumber;
-
-	SelectedKeys.Add(Key);
-	if ( IsBroadcasting() )
-	{
-		OnKeySelectionChanged.Broadcast();
-		OnOutlinerNodeSelectionChangedObjectGuids.Broadcast();
-
-		// Deselect any outliner nodes that aren't within the trunk of this key
-		TArray<UMovieSceneSection*> Sections;
-		Sections.Add(Key.Section);
-
-		EmptySelectedOutlinerNodesWithoutSections(Sections);
-	}
-	else
-	{	
-		bEmptySelectedOutlinerNodesWithSectionsPending = true;
-	}
-}
-
-void FSequencerSelection::AddToSelection(UMovieSceneSection* Section)
-{
-	++SerialNumber;
-
-	SelectedSections.Add(Section);
-	if ( IsBroadcasting() )
-	{
-		OnSectionSelectionChanged.Broadcast();
-		OnOutlinerNodeSelectionChangedObjectGuids.Broadcast();
-
-		// Deselect any outliner nodes that aren't within the trunk of this section
-		if (Section)
-		{
-			TArray<UMovieSceneSection*> Sections;
-			Sections.Add(Section);
-		
-			EmptySelectedOutlinerNodesWithoutSections(Sections);
-		}
-	}
-	else
-	{	
-		bEmptySelectedOutlinerNodesWithSectionsPending = true;
-	}
-}
-
-void FSequencerSelection::AddToSelection(TSharedRef<FSequencerDisplayNode> OutlinerNode)
-{
-	++SerialNumber;
-
-	SelectedOutlinerNodes.Add(OutlinerNode);
-	if ( IsBroadcasting() )
-	{
-		OnOutlinerNodeSelectionChanged.Broadcast();
-		OnOutlinerNodeSelectionChangedObjectGuids.Broadcast();
-	}
-	EmptySelectedKeys();
-	EmptySelectedSections();
-	EmptyNodesWithSelectedKeysOrSections();
-}
-
-void FSequencerSelection::AddToSelection(const TArray<TSharedRef<FSequencerDisplayNode>>& OutlinerNodes)
-{
-	++SerialNumber;
-
-	SelectedOutlinerNodes.Append(OutlinerNodes);
-	if (IsBroadcasting())
-	{
-		OnOutlinerNodeSelectionChanged.Broadcast();
-		OnOutlinerNodeSelectionChangedObjectGuids.Broadcast();
-	}
-	EmptySelectedKeys();
-	EmptySelectedSections();
-	EmptyNodesWithSelectedKeysOrSections();
-}
-
-void FSequencerSelection::AddToNodesWithSelectedKeysOrSections(TSharedRef<FSequencerDisplayNode> OutlinerNode)
-{
-	++SerialNumber;
-
-	NodesWithSelectedKeysOrSections.Add(OutlinerNode);
-	if ( IsBroadcasting() )
-	{
-		++SerialNumber;
-		OnNodesWithSelectedKeysOrSectionsChanged.Broadcast();
-		OnOutlinerNodeSelectionChangedObjectGuids.Broadcast();
-	}
-}
-
-
-void FSequencerSelection::RemoveFromSelection(const FSequencerSelectedKey& Key)
-{
-	++SerialNumber;
-
-	SelectedKeys.Remove(Key);
-	if ( IsBroadcasting() )
-	{
-		OnKeySelectionChanged.Broadcast();
-	}
-}
-
-void FSequencerSelection::RemoveFromSelection(UMovieSceneSection* Section)
-{
-	++SerialNumber;
-
-	SelectedSections.Remove(Section);
-	if ( IsBroadcasting() )
-	{
-		OnSectionSelectionChanged.Broadcast();
-	}
-}
-
-void FSequencerSelection::RemoveFromSelection(TSharedRef<FSequencerDisplayNode> OutlinerNode)
-{
-	++SerialNumber;
-
-	SelectedOutlinerNodes.Remove(OutlinerNode);
-	if ( IsBroadcasting() )
-	{
-		OnOutlinerNodeSelectionChanged.Broadcast();
-	}
-}
-
-void FSequencerSelection::RemoveFromNodesWithSelectedKeysOrSections(TSharedRef<FSequencerDisplayNode> OutlinerNode)
-{
-	++SerialNumber;
-
-	NodesWithSelectedKeysOrSections.Remove(OutlinerNode);
-	if ( IsBroadcasting() )
-	{
-		OnNodesWithSelectedKeysOrSectionsChanged.Broadcast();
-	}
-}
-
-bool FSequencerSelection::IsSelected(const FSequencerSelectedKey& Key) const
-{
-	return SelectedKeys.Contains(Key);
-}
-
-bool FSequencerSelection::IsSelected(UMovieSceneSection* Section) const
-{
-	return SelectedSections.Contains(Section);
-}
-
-bool FSequencerSelection::IsSelected(TSharedRef<FSequencerDisplayNode> OutlinerNode) const
-{
-	return SelectedOutlinerNodes.Contains(OutlinerNode);
-}
-
-bool FSequencerSelection::NodeHasSelectedKeysOrSections(TSharedRef<FSequencerDisplayNode> OutlinerNode) const
-{
-	return NodesWithSelectedKeysOrSections.Contains(OutlinerNode);
-}
-
-void FSequencerSelection::Empty()
-{
-	++SerialNumber;
-
-	EmptySelectedKeys();
-	EmptySelectedSections();
-	EmptySelectedOutlinerNodes();
-	EmptyNodesWithSelectedKeysOrSections();
-}
-
-void FSequencerSelection::EmptySelectedKeys()
-{
-	if (!SelectedKeys.Num())
-	{
-		return;
-	}
-
-	++SerialNumber;
-
-	SelectedKeys.Empty();
-	if ( IsBroadcasting() )
-	{
-		OnKeySelectionChanged.Broadcast();
-	}
-}
-
-void FSequencerSelection::EmptySelectedSections()
-{
-	if (!SelectedSections.Num())
-	{
-		return;
-	}
-
-	++SerialNumber;
-
-	SelectedSections.Empty();
-	if ( IsBroadcasting() )
-	{
-		OnSectionSelectionChanged.Broadcast();
-	}
-}
-
-void FSequencerSelection::EmptySelectedOutlinerNodes()
-{
-	if (!SelectedOutlinerNodes.Num())
-	{
-		return;
-	}
-
-	++SerialNumber;
-
-	SelectedOutlinerNodes.Empty();
-	if ( IsBroadcasting() )
-	{
-		OnOutlinerNodeSelectionChanged.Broadcast();
-	}
-}
-
-void FSequencerSelection::EmptyNodesWithSelectedKeysOrSections()
-{
-	if (!NodesWithSelectedKeysOrSections.Num())
-	{
-		return;
-	}
-
-	++SerialNumber;
-
-	NodesWithSelectedKeysOrSections.Empty();
-	if ( IsBroadcasting() )
-	{
-		OnNodesWithSelectedKeysOrSectionsChanged.Broadcast();
-	}
-}
-
-/** Suspend or resume broadcast of selection changing  */
-void FSequencerSelection::SuspendBroadcast()
-{
-	SuspendBroadcastCount++;
-}
-
-void FSequencerSelection::ResumeBroadcast()
-{
-	SuspendBroadcastCount--;
-	checkf(SuspendBroadcastCount >= 0, TEXT("Suspend/Resume broadcast mismatch!"));
-}
-
-bool FSequencerSelection::IsBroadcasting() 
-{
-	return SuspendBroadcastCount == 0;
-}
-
-void FSequencerSelection::EmptySelectedOutlinerNodesWithoutSections(const TArray<UMovieSceneSection*>& Sections)
-{
-	TSet<TSharedRef<FSequencerDisplayNode>> LocalSelectedOutlinerNodes = SelectedOutlinerNodes;
-
-	SuspendBroadcast();
-	bool bRemoved = false;
-	for (auto SelectedOutlinerNode : LocalSelectedOutlinerNodes)
-	{
-		TSet<TSharedRef<FSequencerDisplayNode> > TrunkNodes;
-		TrunkNodes.Add(SelectedOutlinerNode);
-		SequencerHelpers::GetDescendantNodes(SelectedOutlinerNode, TrunkNodes);
-		
-		bool bFoundMatch = false;
-		for (auto TrunkIt = TrunkNodes.CreateConstIterator(); TrunkIt && !bFoundMatch; ++TrunkIt)
-		{
-			TSet<TWeakObjectPtr<UMovieSceneSection>> AllSections;
-			SequencerHelpers::GetAllSections(*TrunkIt, AllSections);
-
-			for (auto SectionIt = AllSections.CreateConstIterator(); SectionIt && !bFoundMatch; ++SectionIt)
-			{
-				if (Sections.Contains(*SectionIt))
-				{
-					bFoundMatch = true;
-					break;
-				}
-			}
-		}
-
-		if (!bFoundMatch)
-		{
-			RemoveFromSelection(SelectedOutlinerNode);
-			bRemoved = true;
-		}
-	}
-	ResumeBroadcast();
-
-	if (bRemoved)
-	{
-		++SerialNumber;
-
-		OnOutlinerNodeSelectionChanged.Broadcast();
-	}
-}
-
-void FSequencerSelection::RequestOutlinerNodeSelectionChangedBroadcast()
-{
-	bOutlinerNodeSelectionChangedBroadcastPending = true;
-}
-
-void FSequencerSelection::Tick()
-{
-	if ( bOutlinerNodeSelectionChangedBroadcastPending && IsBroadcasting() )
-	{
-		bOutlinerNodeSelectionChangedBroadcastPending = false;
-		OnOutlinerNodeSelectionChanged.Broadcast();
-	}
-
-	if ( bEmptySelectedOutlinerNodesWithSectionsPending && IsBroadcasting() )
-	{
-		bEmptySelectedOutlinerNodesWithSectionsPending = false;
-
-		TArray<UMovieSceneSection*> Sections;
-		for (TWeakObjectPtr<UMovieSceneSection> SelectedSection : SelectedSections)
-		{
-			if (SelectedSection.IsValid())
-			{
-				Sections.Add(SelectedSection.Get());
-			}
-		}
-
-		for (FSequencerSelectedKey SelectedKey : SelectedKeys)
-		{
-			if (SelectedKey.IsValid() && !Sections.Contains(SelectedKey.Section))
-			{
-				Sections.Add(SelectedKey.Section);
-			}
-		}
-
-		EmptySelectedOutlinerNodesWithoutSections(Sections);
-	}
-}
-
+// Copyright Epic Games, Inc. All Rights Reserved.
+
+#include "SequencerSelection.h"
+#include "MovieSceneSection.h"
+#include "SequencerCommonHelpers.h"
+
+#include "MVVM/SharedViewModelData.h"
+#include "MVVM/ViewModels/SectionModel.h"
+#include "MVVM/ViewModels/TrackModel.h"
+#include "MVVM/ViewModels/ChannelModel.h"
+#include "MVVM/ViewModels/ViewModelIterators.h"
+#include "MVVM/Extensions/IObjectBindingExtension.h"
+#include "MVVM/Extensions/ISelectableExtension.h"
+#include "MVVM/SectionModelStorageExtension.h"
+
+FSequencerSelection::FSequencerSelection()
+	: SerialNumber(0)
+	, SuspendBroadcastCount(0)
+	, bOutlinerNodeSelectionChangedBroadcastPending(false)
+	, bEmptySelectedOutlinerNodesWithSectionsPending(false)
+	, bNodesWithSelectedKeysOrSectionsDirty(false)
+{
+}
+
+TSharedPtr<UE::Sequencer::FViewModel> FSequencerSelection::GetRootModel() const
+{
+	return RootModel;
+}
+
+void FSequencerSelection::SetRootModel(TSharedPtr<UE::Sequencer::FViewModel> InRootModel)
+{
+	if (RootModel && HierarchyChangedHandle.IsValid())
+	{
+		RootModel->GetSharedData()->UnsubscribeFromHierarchyChanged(RootModel, HierarchyChangedHandle);
+		HierarchyChangedHandle = FDelegateHandle();
+	}
+
+	RootModel = InRootModel;
+
+	if (RootModel)
+	{
+		FSimpleMulticastDelegate& HierarchyChanged = RootModel->GetSharedData()->SubscribeToHierarchyChanged(RootModel);
+		HierarchyChangedHandle = HierarchyChanged.AddRaw(this, &FSequencerSelection::OnHierarchyChanged);
+	}
+}
+
+const TSet<FSequencerSelectedKey>& FSequencerSelection::GetSelectedKeys() const
+{
+	return SelectedKeys;
+}
+
+TSet<TWeakObjectPtr<UMovieSceneSection>> FSequencerSelection::GetSelectedSections() const
+{
+	using namespace UE::Sequencer;
+
+	TSet<TWeakObjectPtr<UMovieSceneSection>> Sections;
+
+	for (const TWeakPtr<FViewModel>& TrackAreaModel : SelectedTrackAreaItems)
+	{
+		if (TSharedPtr<FViewModel> Pinned = TrackAreaModel.Pin())
+		{
+			FSectionModel* Section = Pinned->CastThis<FSectionModel>();
+			if (Section && Section->GetSection())
+			{
+				Sections.Add(Section->GetSection());
+			}
+		}
+	}
+
+	return Sections;
+}
+
+const TSet<TWeakPtr<UE::Sequencer::FViewModel>>& FSequencerSelection::GetSelectedOutlinerItems() const
+{
+	return SelectedOutlinerItems;
+}
+
+const TSet<TWeakPtr<UE::Sequencer::FViewModel>>& FSequencerSelection::GetSelectedTrackAreaItems() const
+{
+	return SelectedTrackAreaItems;
+}
+
+const TSet<TWeakPtr<UE::Sequencer::FViewModel>>& FSequencerSelection::GetNodesWithSelectedKeysOrSections() const
+{
+	using namespace UE::Sequencer;
+
+	if (!bNodesWithSelectedKeysOrSectionsDirty)
+	{
+		return NodesWithSelectedKeysOrSections;
+	}
+
+	bNodesWithSelectedKeysOrSectionsDirty = false;
+	NodesWithSelectedKeysOrSections.Empty();
+
+	// Handle selected keys
+	for (const FSequencerSelectedKey& Key : SelectedKeys)
+	{
+		if (TSharedPtr<FChannelModel> Channel = Key.WeakChannel.Pin())
+		{
+			TViewModelPtr<IOutlinerExtension> Outliner = Channel->GetLinkedOutlinerItem();
+			if (Outliner)
+			{
+				NodesWithSelectedKeysOrSections.Add(Outliner.AsModel());
+			}
+		}
+	}
+
+	// Handle selected track area items
+	for (TWeakPtr<FViewModel> WeakTrackAreaModel : SelectedTrackAreaItems)
+	{
+		TSharedPtr<FViewModel> TrackAreaModel = WeakTrackAreaModel.Pin();
+		if (TrackAreaModel)
+		{
+			TViewModelPtr<IOutlinerExtension> OutlinerItem = TrackAreaModel->FindAncestorOfType<IOutlinerExtension>();
+			if (OutlinerItem)
+			{
+				NodesWithSelectedKeysOrSections.Add(OutlinerItem.AsModel());
+			}
+		}
+	}
+
+	return NodesWithSelectedKeysOrSections;
+}
+
+FSequencerSelection::FOnSelectionChanged& FSequencerSelection::GetOnKeySelectionChanged()
+{
+	return OnKeySelectionChanged;
+}
+
+FSequencerSelection::FOnSelectionChanged& FSequencerSelection::GetOnOutlinerNodeSelectionChanged()
+{
+	return OnOutlinerNodeSelectionChanged;
+}
+
+FSequencerSelection::FOnSelectionChangedObjectGuids& FSequencerSelection::GetOnOutlinerNodeSelectionChangedObjectGuids()
+{
+	return OnOutlinerNodeSelectionChangedObjectGuids;
+}
+
+TArray<FGuid> FSequencerSelection::GetBoundObjectsGuids()
+{
+	using namespace UE::Sequencer;
+
+	TArray<FGuid> OutGuids;
+	TSet<TWeakPtr<FViewModel>> SelectedItems = GetNodesWithSelectedKeysOrSections();
+	if (SelectedItems.Num() == 0)
+	{
+		SelectedItems = GetSelectedOutlinerItems();
+	}
+	for (TWeakPtr<FViewModel> WeakModel : SelectedItems)
+	{
+		TSharedPtr<FViewModel> Model = WeakModel.Pin();
+		if (Model)
+		{
+			TSharedPtr<IObjectBindingExtension> ObjectBinding = Model->FindAncestorOfType<IObjectBindingExtension>(true);
+			if (ObjectBinding)
+			{
+				OutGuids.Add(ObjectBinding->GetObjectGuid());
+			}
+		}
+	}
+
+	return OutGuids;
+}
+
+TArray<UMovieSceneTrack*> FSequencerSelection::GetSelectedTracks() const
+{
+	using namespace UE::Sequencer;
+
+	TArray<UMovieSceneTrack*> OutTracks;
+
+	for (TWeakPtr<FViewModel> WeakModel : SelectedOutlinerItems)
+	{
+		TSharedPtr<FViewModel> Model = WeakModel.Pin();
+		if (Model)
+		{
+			TSharedPtr<FTrackModel> Track = Model->FindAncestorOfType<FTrackModel>(true);
+			if (Track)
+			{
+				OutTracks.Add(Track->GetTrack());
+			}
+		}
+	}
+
+	return OutTracks;
+}
+
+void FSequencerSelection::AddToSelection(TSharedPtr<UE::Sequencer::FViewModel> InModel)
+{
+	if (InModel->IsA<UE::Sequencer::IOutlinerExtension>())
+	{
+		AddToOutlinerSelection(InModel);
+	}
+	else
+	{
+		AddToTrackAreaSelection(InModel);
+	}
+}
+
+void FSequencerSelection::AddToTrackAreaSelection(TSharedPtr<UE::Sequencer::FViewModel> InModel)
+{
+	using namespace UE::Sequencer;
+
+	ISelectableExtension* Selectable = InModel->CastThis<ISelectableExtension>();
+	if (Selectable && Selectable->IsSelectable() == ESelectionIntent::Never)
+	{
+		return;
+	}
+
+	bNodesWithSelectedKeysOrSectionsDirty = true;
+	++SerialNumber;
+	SelectedTrackAreaItems.Add(InModel);
+
+	if (TSharedPtr<IOutlinerExtension> ParentItem = InModel->FindAncestorOfType<IOutlinerExtension>())
+	{
+		ParentItem->ToggleSelectionState(EOutlinerSelectionState::HasSelectedTrackAreaItems, true);
+	}
+
+	if (FSectionModel* SectionModel = InModel->CastThis<FSectionModel>())
+	{
+		if (UMovieSceneSection* Section = SectionModel->GetSection())
+		{
+			if ( IsBroadcasting() )
+			{
+				OnSectionSelectionChanged.Broadcast();
+				OnOutlinerNodeSelectionChangedObjectGuids.Broadcast();
+			}
+
+			// Deselect any outliner nodes that aren't within the trunk of this section
+			EmptySelectedOutlinerNodesWithoutSections(TArray<UMovieSceneSection*>{Section});
+		}
+	}
+}
+
+void FSequencerSelection::AddToSelection(const FSequencerSelectedKey& Key)
+{
+	using namespace UE::Sequencer;
+
+	bNodesWithSelectedKeysOrSectionsDirty = true;
+	++SerialNumber;
+
+	SelectedKeys.Add(Key);
+
+	if (TViewModelPtr<IOutlinerExtension> OutlinerItem = Key.WeakChannel.Pin()->GetLinkedOutlinerItem())
+	{
+		OutlinerItem->ToggleSelectionState(EOutlinerSelectionState::HasSelectedKeys, true);
+	}
+
+	if ( IsBroadcasting() )
+	{
+		OnKeySelectionChanged.Broadcast();
+		OnOutlinerNodeSelectionChangedObjectGuids.Broadcast();
+
+		// Deselect any outliner nodes that aren't within the trunk of this key
+		TArray<UMovieSceneSection*> Sections;
+		Sections.Add(Key.Section);
+
+		EmptySelectedOutlinerNodesWithoutSections(Sections);
+	}
+	else
+	{	
+		bEmptySelectedOutlinerNodesWithSectionsPending = true;
+	}
+}
+
+void FSequencerSelection::AddToOutlinerSelection(TSharedPtr<UE::Sequencer::FViewModel> InModel)
+{
+	using namespace UE::Sequencer;
+
+	ISelectableExtension* Selectable = InModel->CastThis<ISelectableExtension>();
+	if (Selectable && Selectable->IsSelectable() == ESelectionIntent::Never)
+	{
+		return;
+	}
+
+	++SerialNumber;
+
+	SelectedOutlinerItems.Add(InModel);
+	if (IOutlinerExtension* OutlinerItem = InModel->CastThis<IOutlinerExtension>())
+	{
+		OutlinerItem->SetSelectionState(EOutlinerSelectionState::SelectedDirectly);
+	}
+
+	if ( IsBroadcasting() )
+	{
+		OnOutlinerNodeSelectionChanged.Broadcast();
+		OnOutlinerNodeSelectionChangedObjectGuids.Broadcast();
+	}
+
+	EmptySelectedKeys();
+	EmptySelectedTrackAreaItems();
+}
+
+void FSequencerSelection::AddToSelection(const TArrayView<TSharedPtr<UE::Sequencer::FViewModel>>& InModels)
+{
+	using namespace UE::Sequencer;
+
+	++SerialNumber;
+
+	for (TSharedPtr<FViewModel> InModel : InModels)
+	{
+		SelectedOutlinerItems.Add(InModel);
+		if (IOutlinerExtension* OutlinerItem = InModel->CastThis<IOutlinerExtension>())
+		{
+			OutlinerItem->SetSelectionState(EOutlinerSelectionState::SelectedDirectly);
+		}
+	}
+
+	if (IsBroadcasting())
+	{
+		OnOutlinerNodeSelectionChanged.Broadcast();
+		OnOutlinerNodeSelectionChangedObjectGuids.Broadcast();
+	}
+
+	EmptySelectedKeys();
+	EmptySelectedTrackAreaItems();
+}
+
+void FSequencerSelection::AddToSelection(const TArrayView<TSharedRef<UE::Sequencer::FViewModel>>& InModels)
+{
+	using namespace UE::Sequencer;
+
+	++SerialNumber;
+
+	for (TSharedRef<FViewModel> InModel : InModels)
+	{
+		SelectedOutlinerItems.Add(InModel);
+		if (IOutlinerExtension* OutlinerItem = InModel->CastThis<IOutlinerExtension>())
+		{
+			OutlinerItem->SetSelectionState(EOutlinerSelectionState::SelectedDirectly);
+		}
+	}
+
+	if (IsBroadcasting())
+	{
+		OnOutlinerNodeSelectionChanged.Broadcast();
+		OnOutlinerNodeSelectionChangedObjectGuids.Broadcast();
+	}
+
+	EmptySelectedKeys();
+	EmptySelectedTrackAreaItems();
+}
+
+
+void FSequencerSelection::RemoveFromSelection(const FSequencerSelectedKey& Key)
+{
+	using namespace UE::Sequencer;
+
+	bNodesWithSelectedKeysOrSectionsDirty = true;
+	++SerialNumber;
+
+	SelectedKeys.Remove(Key);
+
+	if (TViewModelPtr<IOutlinerExtension> OutlinerItem = Key.WeakChannel.Pin()->GetLinkedOutlinerItem())
+	{
+		OutlinerItem->ToggleSelectionState(EOutlinerSelectionState::HasSelectedKeys, false);
+	}
+
+	if ( IsBroadcasting() )
+	{
+		OnKeySelectionChanged.Broadcast();
+	}
+}
+
+void FSequencerSelection::RemoveFromSelection(TWeakPtr<UE::Sequencer::FViewModel> InModel)
+{
+	if (TSharedPtr<UE::Sequencer::FViewModel> InModelPtr = InModel.Pin())
+	{
+		RemoveFromSelection(InModelPtr);
+	}
+}
+
+void FSequencerSelection::RemoveFromSelection(TSharedPtr<UE::Sequencer::FViewModel> OutlinerNode)
+{
+	if (OutlinerNode)
+	{
+		RemoveFromSelection(OutlinerNode.ToSharedRef());
+	}
+}
+
+void FSequencerSelection::RemoveFromSelection(TSharedRef<UE::Sequencer::FViewModel> OutlinerNode)
+{
+	using namespace UE::Sequencer;
+
+	bNodesWithSelectedKeysOrSectionsDirty = true;
+	++SerialNumber;
+
+	SelectedOutlinerItems.Remove(OutlinerNode);
+	SelectedTrackAreaItems.Remove(OutlinerNode);
+
+	if (IOutlinerExtension* OutlinerItem = OutlinerNode->CastThis<IOutlinerExtension>())
+	{
+		OutlinerItem->SetSelectionState(EOutlinerSelectionState::None);
+	}
+
+	if ( IsBroadcasting() )
+	{
+		OnOutlinerNodeSelectionChanged.Broadcast();
+	}
+}
+
+bool FSequencerSelection::IsSelected(const FSequencerSelectedKey& Key) const
+{
+	return SelectedKeys.Contains(Key);
+}
+
+bool FSequencerSelection::IsSelected(TWeakPtr<UE::Sequencer::FViewModel> InModel) const
+{
+	return SelectedOutlinerItems.Contains(InModel) || SelectedTrackAreaItems.Contains(InModel);
+}
+
+bool FSequencerSelection::NodeHasSelectedKeysOrSections(TWeakPtr<UE::Sequencer::FViewModel> InModel) const
+{
+	return GetNodesWithSelectedKeysOrSections().Contains(InModel);
+}
+
+void FSequencerSelection::Empty()
+{
+	using namespace UE::Sequencer;
+
+	bNodesWithSelectedKeysOrSectionsDirty = true;
+	++SerialNumber;
+
+	FSectionModelStorageExtension* SectionModelStorage = RootModel->CastDynamicChecked<FSectionModelStorageExtension>();
+	for (const FSequencerSelectedKey& Key : SelectedKeys)
+	{
+		if (TSharedPtr<FSectionModel> SectionModel = SectionModelStorage->FindModelForSection(Key.Section))
+		{
+			for (TSharedPtr<IOutlinerExtension> ParentItem : SectionModel->GetAncestorsOfType<IOutlinerExtension>(true))
+			{
+				ParentItem->SetSelectionState(EOutlinerSelectionState::None);
+			}
+		}
+	}
+	for (TWeakPtr<FViewModel> WeakModel : SelectedTrackAreaItems)
+	{
+		if (TSharedPtr<FViewModel> Model = WeakModel.Pin())
+		{
+			for (TSharedPtr<IOutlinerExtension> ParentOutlinerItem : Model->GetAncestorsOfType<IOutlinerExtension>())
+			{
+				ParentOutlinerItem->SetSelectionState(EOutlinerSelectionState::None);
+			}
+		}
+	}
+	for (TWeakPtr<FViewModel> WeakModel : SelectedOutlinerItems)
+	{
+		if (TSharedPtr<FViewModel> Model = WeakModel.Pin())
+		{
+			if (IOutlinerExtension* OutlinerItem = Model->CastThis<IOutlinerExtension>())
+			{
+				OutlinerItem->SetSelectionState(EOutlinerSelectionState::None);
+			}
+		}
+	}
+
+	EmptySelectedKeys();
+	EmptySelectedOutlinerNodes();
+	EmptySelectedTrackAreaItems();
+}
+
+void FSequencerSelection::EmptySelectedKeys()
+{
+	using namespace UE::Sequencer;
+
+	if (!SelectedKeys.Num())
+	{
+		return;
+	}
+
+	bNodesWithSelectedKeysOrSectionsDirty = true;
+	++SerialNumber;
+
+	FSectionModelStorageExtension* SectionModelStorage = RootModel->CastDynamicChecked<FSectionModelStorageExtension>();
+	for (const FSequencerSelectedKey& Key : SelectedKeys)
+	{
+		if (TViewModelPtr<IOutlinerExtension> OutlinerItem = Key.WeakChannel.Pin()->GetLinkedOutlinerItem())
+		{
+			OutlinerItem->ToggleSelectionState(EOutlinerSelectionState::HasSelectedKeys, false);
+		}
+	}
+	SelectedKeys.Empty();
+
+	if ( IsBroadcasting() )
+	{
+		OnKeySelectionChanged.Broadcast();
+	}
+}
+
+void FSequencerSelection::EmptySelectedTrackAreaItems()
+{
+	using namespace UE::Sequencer;
+
+	if (SelectedTrackAreaItems.Num() > 0)
+	{
+		bNodesWithSelectedKeysOrSectionsDirty = true;
+		++SerialNumber;
+
+		for (TWeakPtr<FViewModel> WeakModel : SelectedTrackAreaItems)
+		{
+			if (TSharedPtr<FViewModel> Model = WeakModel.Pin())
+			{
+				for (TSharedPtr<IOutlinerExtension> ParentOutlinerItem : Model->GetAncestorsOfType<IOutlinerExtension>())
+				{
+					ParentOutlinerItem->SetSelectionState(EOutlinerSelectionState::None);
+				}
+			}
+		}
+		SelectedTrackAreaItems.Empty();
+
+		if ( IsBroadcasting() )
+		{
+			OnSectionSelectionChanged.Broadcast();
+		}
+	}
+}
+
+void FSequencerSelection::EmptySelectedOutlinerNodes()
+{
+	using namespace UE::Sequencer;
+
+	if (!SelectedOutlinerItems.Num())
+	{
+		return;
+	}
+
+	++SerialNumber;
+
+	for (TWeakPtr<FViewModel> WeakModel : SelectedOutlinerItems)
+	{
+		if (TSharedPtr<FViewModel> Model = WeakModel.Pin())
+		{
+			if (IOutlinerExtension* OutlinerItem = Model->CastThis<IOutlinerExtension>())
+			{
+				OutlinerItem->SetSelectionState(EOutlinerSelectionState::None);
+			}
+		}
+	}
+	SelectedOutlinerItems.Empty();
+
+	if ( IsBroadcasting() )
+	{
+		OnOutlinerNodeSelectionChanged.Broadcast();
+	}
+}
+
+/** Suspend or resume broadcast of selection changing  */
+void FSequencerSelection::SuspendBroadcast()
+{
+	SuspendBroadcastCount++;
+}
+
+void FSequencerSelection::ResumeBroadcast()
+{
+	SuspendBroadcastCount--;
+	checkf(SuspendBroadcastCount >= 0, TEXT("Suspend/Resume broadcast mismatch!"));
+}
+
+bool FSequencerSelection::IsBroadcasting() 
+{
+	return SuspendBroadcastCount == 0;
+}
+
+void FSequencerSelection::EmptySelectedOutlinerNodesWithoutSections(const TArray<UMovieSceneSection*>& Sections)
+{
+	using namespace UE::Sequencer;
+
+	TSet<TWeakPtr<UE::Sequencer::FViewModel>> LocalSelectedOutlinerItems = SelectedOutlinerItems;
+
+	SuspendBroadcast();
+	bool bRemoved = false;
+	for (TWeakPtr<FViewModel> WeakSelectedModel : LocalSelectedOutlinerItems)
+	{
+		TSharedPtr<FViewModel> SelectedModel = WeakSelectedModel.Pin();
+		if (!SelectedModel)
+		{
+			continue;
+		}
+
+		bool bFoundMatch = false;
+
+		for (TSharedPtr<FSectionModel> SectionModel : SelectedModel->GetDescendantsOfType<FSectionModel>())
+		{
+			if (Sections.Contains(SectionModel->GetSection()))
+			{
+				bFoundMatch = true;
+				break;
+			}
+		}
+
+		if (!bFoundMatch)
+		{
+			RemoveFromSelection(WeakSelectedModel);
+			bRemoved = true;
+		}
+	}
+	ResumeBroadcast();
+
+	if (bRemoved)
+	{
+		++SerialNumber;
+
+		OnOutlinerNodeSelectionChanged.Broadcast();
+	}
+}
+
+void FSequencerSelection::RequestOutlinerNodeSelectionChangedBroadcast()
+{
+	bOutlinerNodeSelectionChangedBroadcastPending = true;
+}
+
+void FSequencerSelection::Tick()
+{
+	if ( bOutlinerNodeSelectionChangedBroadcastPending && IsBroadcasting() )
+	{
+		bOutlinerNodeSelectionChangedBroadcastPending = false;
+		OnOutlinerNodeSelectionChanged.Broadcast();
+	}
+
+	if ( bEmptySelectedOutlinerNodesWithSectionsPending && IsBroadcasting() )
+	{
+		bEmptySelectedOutlinerNodesWithSectionsPending = false;
+
+		TArray<UMovieSceneSection*> Sections;
+		for (TWeakObjectPtr<UMovieSceneSection> SelectedSection : GetSelectedSections())
+		{
+			if (SelectedSection.IsValid())
+			{
+				Sections.Add(SelectedSection.Get());
+			}
+		}
+
+		for (FSequencerSelectedKey SelectedKey : SelectedKeys)
+		{
+			if (SelectedKey.IsValid() && !Sections.Contains(SelectedKey.Section))
+			{
+				Sections.Add(SelectedKey.Section);
+			}
+		}
+
+		EmptySelectedOutlinerNodesWithoutSections(Sections);
+	}
+}
+
+void FSequencerSelection::OnHierarchyChanged()
+{
+	// This is an esoteric hack that ensures we re-synchronize external (ie Actor)
+	// selection when models are removed from the tree. Doing so ensures that
+	// FSequencer::SynchronizeExternalSelectionWithSequencerSelection is called within
+	// the scope of GIsTransacting being true, which prevents that function from creating new
+	// transactions for the selection synchronization. This is important because otherwise
+	// the undo/redo stack gets wiped by actor selections when undoing if the selection is
+	// not identical
+	RevalidateSelection();
+}
+
+void FSequencerSelection::RevalidateSelection()
+{
+	using namespace UE::Sequencer;
+
+	bNodesWithSelectedKeysOrSectionsDirty = true;
+
+	const int32 NumSelectedKeys = SelectedKeys.Num();
+	const int32 NumSelectedTrackAreaItems = SelectedTrackAreaItems.Num();
+	const int32 NumSelectedOutlinerItems = SelectedOutlinerItems.Num();
+
+	for (auto It = SelectedKeys.CreateIterator(); It; ++It)
+	{
+		if (!It->Section || !It->WeakChannel.Pin().Get())
+		{
+			It.RemoveCurrent();
+		}
+	}
+
+	for (auto It = SelectedTrackAreaItems.CreateIterator(); It; ++It)
+	{
+		if (It->Pin().Get() == nullptr)
+		{
+			It.RemoveCurrent();
+		}
+	}
+
+	for (auto It = SelectedOutlinerItems.CreateIterator(); It; ++It)
+	{
+		if (It->Pin().Get() == nullptr)
+		{
+			It.RemoveCurrent();
+		}
+	}
+
+	if (NumSelectedKeys != SelectedKeys.Num())
+	{
+		++SerialNumber;
+
+		OnKeySelectionChanged.Broadcast();
+	}
+
+	if (NumSelectedTrackAreaItems != SelectedTrackAreaItems.Num())
+	{
+		++SerialNumber;
+
+		OnSectionSelectionChanged.Broadcast();
+		OnOutlinerNodeSelectionChangedObjectGuids.Broadcast();
+	}
+
+	if (NumSelectedOutlinerItems != SelectedOutlinerItems.Num())
+	{
+		++SerialNumber;
+		OnOutlinerNodeSelectionChanged.Broadcast();
+	}
+}
+
diff --git a/Engine/Source/Editor/Sequencer/Private/SequencerSelection.h b/Engine/Source/Editor/Sequencer/Private/SequencerSelection.h
--- a/Engine/Source/Editor/Sequencer/Private/SequencerSelection.h
+++ b/Engine/Source/Editor/Sequencer/Private/SequencerSelection.h
@@ -3,12 +3,14 @@
 #pragma once
 
 #include "CoreMinimal.h"
+#include "MVVM/ViewModels/ViewModel.h"
+#include "MVVM/ViewModels/ViewModelHierarchy.h"
 #include "Misc/Guid.h"
 #include "UObject/WeakObjectPtr.h"
-#include "DisplayNodes/SequencerDisplayNode.h"
 #include "SequencerSelectedKey.h"
 
 class UMovieSceneSection;
+class UMovieSceneTrack;
 
 /**
  * Manages the selection of keys, sections, and outliner nodes for the sequencer.
@@ -17,55 +19,106 @@
 {
 public:
 
+	using FViewModel = UE::Sequencer::FViewModel;
+
 	DECLARE_MULTICAST_DELEGATE(FOnSelectionChanged)
-
 	DECLARE_MULTICAST_DELEGATE(FOnSelectionChangedObjectGuids)
 
 public:
 
 	FSequencerSelection();
 
+	TSharedPtr<FViewModel> GetRootModel() const;
+	void SetRootModel(TSharedPtr<FViewModel> InRootModel);
+
 	/** Gets a set of the selected keys. */
 	const TSet<FSequencerSelectedKey>& GetSelectedKeys() const;
 
 	/** Gets a set of the selected sections */
-	const TSet<TWeakObjectPtr<UMovieSceneSection>>& GetSelectedSections() const;
+	TSet<TWeakObjectPtr<UMovieSceneSection>> GetSelectedSections() const;
 
-	/** Gets a set of the selected outliner nodes. */
-	const TSet<TSharedRef<FSequencerDisplayNode>>& GetSelectedOutlinerNodes() const;
+	/** Gets a set of the selected models on the outliner */
+	const TSet<TWeakPtr<FViewModel>>& GetSelectedOutlinerItems() const;
 
+	/** Gets the selected models on the outliner */
+	void GetSelectedOutlinerItems(TArray<TSharedRef<FViewModel>>& OutItems) const
+	{
+		OutItems.Reserve(SelectedOutlinerItems.Num());
+
+		for (const TWeakPtr<FViewModel>& WeakItem : SelectedOutlinerItems)
+		{
+			if (TSharedPtr<FViewModel> Item = WeakItem.Pin())
+			{
+				OutItems.Add(Item.ToSharedRef());
+			}
+		}
+	}
+
+	void GetSelectedOutlinerItems(TArray<TSharedPtr<FViewModel>>& OutItems) const
+	{
+		OutItems.Reserve(SelectedOutlinerItems.Num());
+
+		for (const TWeakPtr<FViewModel>& WeakItem : SelectedOutlinerItems)
+		{
+			if (TSharedPtr<FViewModel> Item = WeakItem.Pin())
+			{
+				OutItems.Add(Item);
+			}
+		}
+	}
+
+	/** Gets a set of the selected models */
+	const TSet<TWeakPtr<FViewModel>>& GetSelectedTrackAreaItems() const;
+
 	/** Gets a set of the outliner nodes that have selected keys or sections */
-	const TSet<TSharedRef<FSequencerDisplayNode>>& GetNodesWithSelectedKeysOrSections() const;
+	const TSet<TWeakPtr<FViewModel>>& GetNodesWithSelectedKeysOrSections() const;
 
+	/** Gets the outliner nodes that have selected keys or sections */
+	void GetNodesWithSelectedKeysOrSections(TArray<TSharedRef<FViewModel>>& OutKeysOrSections) const
+	{
+		const TSet<TWeakPtr<FViewModel>>& NodesWithKeysOrSections = GetNodesWithSelectedKeysOrSections();
+
+		OutKeysOrSections.Reserve(OutKeysOrSections.Num() + NodesWithKeysOrSections.Num());
+
+		for (const TWeakPtr<FViewModel>& WeakItem : NodesWithKeysOrSections)
+		{
+			if (TSharedPtr<FViewModel> Item = WeakItem.Pin())
+			{
+				OutKeysOrSections.Add(Item.ToSharedRef());
+			}
+		}
+	}
+
 	/** Get the currently selected tracks as UMovieSceneTracks */
 	TArray<UMovieSceneTrack*> GetSelectedTracks() const;
 
 	/** Adds a key to the selection */
 	void AddToSelection(const FSequencerSelectedKey& Key);
 
-	/** Adds a section to the selection */
-	void AddToSelection(UMovieSceneSection* Section);
+	/** Adds a model to the selection */
+	void AddToSelection(TSharedPtr<FViewModel> InModel);
 
-	/** Adds an outliner node to the selection */
-	void AddToSelection(TSharedRef<FSequencerDisplayNode> OutlinerNode);
+	/** Adds models to the selection */
+	void AddToSelection(const TArrayView<TSharedPtr<FViewModel>>& InModels);
 
-	/** Adds an array of outliner nodes to the selection */
-	void AddToSelection(const TArray<TSharedRef<FSequencerDisplayNode>>& OutlinerNodes);
+	/** Adds models to the selection */
+	void AddToSelection(const TArrayView<TSharedRef<FViewModel>>& InModels);
 
-	/** Adds an outliner node that has selected keys or sections */
-	void AddToNodesWithSelectedKeysOrSections(TSharedRef<FSequencerDisplayNode> OutlinerNode);
+	/** Adds a model to the selection */
+	void AddToTrackAreaSelection(TSharedPtr<FViewModel> InModel);
+	void AddToOutlinerSelection(TSharedPtr<FViewModel> InModel);
 
 	/** Removes a key from the selection */
 	void RemoveFromSelection(const FSequencerSelectedKey& Key);
 
-	/** Removes a section from the selection */
-	void RemoveFromSelection(UMovieSceneSection* Section);
+	/** Removes a model from the selection */
+	void RemoveFromSelection(TWeakPtr<FViewModel> InModel);
 
 	/** Removes an outliner node from the selection */
-	void RemoveFromSelection(TSharedRef<FSequencerDisplayNode> OutlinerNode);
+	void RemoveFromSelection(TSharedPtr<FViewModel> OutlinerNode);
 
-	/** Removes an outliner node that has selected keys or sections */
-	void RemoveFromNodesWithSelectedKeysOrSections(TSharedRef<FSequencerDisplayNode> OutlinerNode);
+	/** Removes an outliner node from the selection */
+	void RemoveFromSelection(TSharedRef<FViewModel> OutlinerNode);
 
 	/** Removes any outliner nodes from the selection that do not relate to the given section */
 	void EmptySelectedOutlinerNodesWithoutSections(const TArray<UMovieSceneSection*>& Sections);
@@ -73,14 +126,11 @@
 	/** Returns whether or not the key is selected. */
 	bool IsSelected(const FSequencerSelectedKey& Key) const;
 
-	/** Returns whether or not the section is selected. */
-	bool IsSelected(UMovieSceneSection* Section) const;
+	/** Returns whether or not the model is selected. */
+	bool IsSelected(TWeakPtr<FViewModel> InModel) const;
 
-	/** Returns whether or not the outliner node is selected. */
-	bool IsSelected(TSharedRef<FSequencerDisplayNode> OutlinerNode) const;
-
 	/** Returns whether or not the outliner node has keys or sections selected. */
-	bool NodeHasSelectedKeysOrSections(TSharedRef<FSequencerDisplayNode> OutlinerNode) const;
+	bool NodeHasSelectedKeysOrSections(TWeakPtr<FViewModel> Model) const;
 
 	/** Empties all selections. */
 	void Empty();
@@ -88,15 +138,12 @@
 	/** Empties the key selection. */
 	void EmptySelectedKeys();
 
-	/** Empties the section selection. */
-	void EmptySelectedSections();
+	/** Empties the Track Area selection. */
+	void EmptySelectedTrackAreaItems();
 
 	/** Empties the outliner node selection. */
 	void EmptySelectedOutlinerNodes();
 
-	/** Empties the outliner nodes with selected keys or sections. */
-	void EmptyNodesWithSelectedKeysOrSections();
-
 	/** Gets a multicast delegate which is called when the key selection changes. */
 	FOnSelectionChanged& GetOnKeySelectionChanged();
 
@@ -106,9 +153,6 @@
 	/** Gets a multicast delegate which is called when the outliner node selection changes. */
 	FOnSelectionChanged& GetOnOutlinerNodeSelectionChanged();
 
-	/** Gets a multicast delegate which is called when the outliner node with selected keys or sections changes. */
-	FOnSelectionChanged& GetOnNodesWithSelectedKeysOrSectionsChanged();
-
 	/** Gets a multicast delegate with an array of FGuid of bound objects which is called when the outliner node selection changes. */
 	FOnSelectionChangedObjectGuids& GetOnOutlinerNodeSelectionChangedObjectGuids();
 
@@ -127,6 +171,9 @@
 	/** Updates the selection once per frame.  This is required for deferred selection broadcasts. */
 	void Tick();
 
+	/** Revalidates the selection by removing stale ptrs */
+	void RevalidateSelection();
+
 	/** When true, selection change notifications should be broadcasted. */
 	bool IsBroadcasting();
 
@@ -138,15 +185,20 @@
 
 private:
 
+	void OnHierarchyChanged();
+
+private:
+
+	TSharedPtr<FViewModel> RootModel;
+
 	TSet<FSequencerSelectedKey> SelectedKeys;
-	TSet<TWeakObjectPtr<UMovieSceneSection>> SelectedSections;
-	TSet<TSharedRef<FSequencerDisplayNode>> SelectedOutlinerNodes;
-	TSet<TSharedRef<FSequencerDisplayNode>> NodesWithSelectedKeysOrSections;
+	TSet<TWeakPtr<FViewModel>> SelectedTrackAreaItems;
+	TSet<TWeakPtr<FViewModel>> SelectedOutlinerItems;
+	mutable TSet<TWeakPtr<FViewModel>> NodesWithSelectedKeysOrSections;
 
 	FOnSelectionChanged OnKeySelectionChanged;
 	FOnSelectionChanged OnSectionSelectionChanged;
 	FOnSelectionChanged OnOutlinerNodeSelectionChanged;
-	FOnSelectionChanged OnNodesWithSelectedKeysOrSectionsChanged;
 
 	FOnSelectionChangedObjectGuids OnOutlinerNodeSelectionChangedObjectGuids;
 
@@ -156,9 +208,15 @@
 	/** The number of times the broadcasting of selection change notifications has been suspended. */
 	int32 SuspendBroadcastCount;
 
+	/** Delegate handle for the current root node on-hierarchy-changed event */
+	FDelegateHandle HierarchyChangedHandle;
+
 	/** When true there is a pending outliner node selection change which will be broadcast next tick. */
 	bool bOutlinerNodeSelectionChangedBroadcastPending;
 
 	/** When true there is a pending call to remove selected outliner nodes that don't have sections selected. */
 	bool bEmptySelectedOutlinerNodesWithSectionsPending;
+
+	mutable bool bNodesWithSelectedKeysOrSectionsDirty;
 };
+
diff --git a/Engine/Source/Editor/Sequencer/Private/SequencerSelectionCurveFilter.h b/Engine/Source/Editor/Sequencer/Private/SequencerSelectionCurveFilter.h
--- a/Engine/Source/Editor/Sequencer/Private/SequencerSelectionCurveFilter.h
+++ b/Engine/Source/Editor/Sequencer/Private/SequencerSelectionCurveFilter.h
@@ -1,58 +1,62 @@
-// Copyright Epic Games, Inc. All Rights Reserved.
-
-#pragma once
-
-#include "CoreMinimal.h"
-#include "Tree/CurveEditorTreeFilter.h"
-#include "DisplayNodes/SequencerDisplayNode.h"
-
-
-/**
- * A specialized filter for showing items in the curve editor selected from the sequencer panel.
- * This filter will store the selected nodes and all the parents of the selected nodes in a NodesToFilter set.
- * An item will pass if it is either directly selected or has a parent in the set.
- */
-struct FSequencerSelectionCurveFilter : FCurveEditorTreeFilter
-{
-	static const int32 FilterPass = -1000;
-
-	FSequencerSelectionCurveFilter()
-		: FCurveEditorTreeFilter(ISequencerModule::GetSequencerSelectionFilterType(), FilterPass)
-	{}
-
-	/**
-	 * Adds all selected nodes and their object parents to the NodesToFilter set
-	 */
-	void Update(const TSet<TSharedRef<FSequencerDisplayNode>>& SelectedNodes, const bool bExpandTreeToSelectedNodes = true)
-	{
-		NodesToFilter.Empty(SelectedNodes.Num());
-
-		for (const TSharedRef<FSequencerDisplayNode>& SelectedNode : SelectedNodes)
-		{
-			NodesToFilter.Add(SelectedNode);
-
-			TSharedPtr<FSequencerDisplayNode> Parent = SelectedNode->GetParent();
-			while (Parent.IsValid())
-			{
-				if (Parent->GetType() == ESequencerNode::Object)
-				{
-					NodesToFilter.Add(Parent.ToSharedRef());
-					break;
-				}
-
-				Parent = Parent->GetParent();
-			}
-		}
-
-		bExpandToMatchedItems = bExpandTreeToSelectedNodes;
-	}
-
-	bool Match(TSharedRef<const FSequencerDisplayNode> InNode) const
-	{
-		return NodesToFilter.Contains(InNode);
-	}
-
-private:
-
-	TSet<TSharedRef<const FSequencerDisplayNode>> NodesToFilter;
-};
\ No newline at end of file
+// Copyright Epic Games, Inc. All Rights Reserved.
+
+#pragma once
+
+#include "CoreMinimal.h"
+#include "Tree/CurveEditorTreeFilter.h"
+
+/**
+ * A specialized filter for showing items in the curve editor selected from the sequencer panel.
+ * This filter will store the selected nodes and all the parents of the selected nodes in a NodesToFilter set.
+ * An item will pass if it is either directly selected or has a parent in the set.
+ */
+struct FSequencerSelectionCurveFilter : FCurveEditorTreeFilter
+{
+	static const int32 FilterPass = -1000;
+
+	FSequencerSelectionCurveFilter()
+		: FCurveEditorTreeFilter(ISequencerModule::GetSequencerSelectionFilterType(), FilterPass)
+	{}
+
+	/**
+	 * Adds all selected nodes and their object parents to the NodesToFilter set
+	 */
+	void Update(const TSet<TWeakPtr<UE::Sequencer::FViewModel>>& SelectedNodes, const bool bExpandTreeToSelectedNodes = true)
+	{
+		using namespace UE::Sequencer;
+
+		NodesToFilter.Empty(SelectedNodes.Num());
+
+		for (const TWeakPtr<FViewModel>& WeakSelectedNode : SelectedNodes)
+		{
+			if (TSharedPtr<FViewModel> SelectedNode = WeakSelectedNode.Pin())
+			{
+				NodesToFilter.Add(WeakSelectedNode);
+
+				TSharedPtr<FViewModel> Parent = SelectedNode->GetParent();
+				while (Parent.IsValid())
+				{
+					if (Parent->IsA<IObjectBindingExtension>())
+					{
+						NodesToFilter.Add(Parent);
+						break;
+					}
+
+					Parent = Parent->GetParent();
+				}
+			}
+		}
+
+		bExpandToMatchedItems = bExpandTreeToSelectedNodes;
+	}
+
+	bool Match(TSharedRef<const UE::Sequencer::FViewModel> InNode) const
+	{
+		return NodesToFilter.Contains(InNode);
+	}
+
+private:
+
+	TSet<TWeakPtr<const UE::Sequencer::FViewModel>> NodesToFilter;
+};
+
diff --git a/Engine/Source/Editor/Sequencer/Private/SequencerSelectionPreview.cpp b/Engine/Source/Editor/Sequencer/Private/SequencerSelectionPreview.cpp
--- a/Engine/Source/Editor/Sequencer/Private/SequencerSelectionPreview.cpp
+++ b/Engine/Source/Editor/Sequencer/Private/SequencerSelectionPreview.cpp
@@ -1,123 +1,88 @@
-// Copyright Epic Games, Inc. All Rights Reserved.
-
-#include "SequencerSelectionPreview.h"
-#include "MovieSceneSection.h"
-
-void FSequencerSelectionPreview::SetSelectionState(FSequencerSelectedKey Key, ESelectionPreviewState InState)
-{
-	if (InState == ESelectionPreviewState::Undefined)
-	{
-		DefinedKeyStates.Remove(Key);
-	}
-	else
-	{
-		DefinedKeyStates.Add(Key, InState);
-	}
-
-	CachedSelectionHash.Reset();
-}
-
-void FSequencerSelectionPreview::SetSelectionState(UMovieSceneSection* Section, ESelectionPreviewState InState)
-{
-	if (InState == ESelectionPreviewState::Undefined)
-	{
-		DefinedSectionStates.Remove(Section);
-	}
-	else
-	{
-		DefinedSectionStates.Add(Section, InState);
-	}
-
-	CachedSelectionHash.Reset();
-}
-
-void FSequencerSelectionPreview::SetSelectionState(TSharedRef<FSequencerDisplayNode> OutlinerNode, ESelectionPreviewState InState)
-{
-	if (InState == ESelectionPreviewState::Undefined)
-	{
-		DefinedOutlinerNodeStates.Remove(OutlinerNode);
-	}
-	else
-	{
-		DefinedOutlinerNodeStates.Add(OutlinerNode, InState);
-	}
-
-	CachedSelectionHash.Reset();
-}
-
-ESelectionPreviewState FSequencerSelectionPreview::GetSelectionState(const FSequencerSelectedKey& Key) const
-{
-	if (auto* State = DefinedKeyStates.Find(Key))
-	{
-		return *State;
-	}
-	return ESelectionPreviewState::Undefined;
-}
-
-ESelectionPreviewState FSequencerSelectionPreview::GetSelectionState(UMovieSceneSection* Section) const
-{
-	if (auto* State = DefinedSectionStates.Find(Section))
-	{
-		return *State;
-	}
-	return ESelectionPreviewState::Undefined;
-}
-
-ESelectionPreviewState FSequencerSelectionPreview::GetSelectionState(TSharedRef<FSequencerDisplayNode> OutlinerNode) const
-{
-	if (auto* State = DefinedOutlinerNodeStates.Find(OutlinerNode))
-	{
-		return *State;
-	}
-	return ESelectionPreviewState::Undefined;
-}
-
-void FSequencerSelectionPreview::Empty()
-{
-	EmptyDefinedKeyStates();
-	EmptyDefinedSectionStates();
-	EmptyDefinedOutlinerNodeStates();
-}
-
-void FSequencerSelectionPreview::EmptyDefinedKeyStates()
-{
-	DefinedKeyStates.Reset();
-	CachedSelectionHash.Reset();
-}
-
-void FSequencerSelectionPreview::EmptyDefinedSectionStates()
-{
-	DefinedSectionStates.Reset();
-	CachedSelectionHash.Reset();
-}
-
-void FSequencerSelectionPreview::EmptyDefinedOutlinerNodeStates()
-{
-	DefinedOutlinerNodeStates.Reset();
-	CachedSelectionHash.Reset();
-}
-
-uint32 FSequencerSelectionPreview::GetSelectionHash() const
-{
-	if (!CachedSelectionHash.IsSet())
-	{
-		uint32 NewHash = 0;
-
-		for (TPair<FSequencerSelectedKey, ESelectionPreviewState> Pair : DefinedKeyStates)
-		{
-			NewHash = HashCombine(NewHash, HashCombine(GetTypeHash(Pair.Key), GetTypeHash(Pair.Value)));
-		}
-		for (TPair<TWeakObjectPtr<UMovieSceneSection>, ESelectionPreviewState> Pair : DefinedSectionStates)
-		{
-			NewHash = HashCombine(NewHash, HashCombine(GetTypeHash(Pair.Key), GetTypeHash(Pair.Value)));
-		}
-		for (TPair<TSharedRef<FSequencerDisplayNode>, ESelectionPreviewState> Pair : DefinedOutlinerNodeStates)
-		{
-			NewHash = HashCombine(NewHash, HashCombine(GetTypeHash(Pair.Key), GetTypeHash(Pair.Value)));
-		}
-
-		CachedSelectionHash = NewHash;
-	}
-
-	return CachedSelectionHash.GetValue();
-}
\ No newline at end of file
+// Copyright Epic Games, Inc. All Rights Reserved.
+
+#include "SequencerSelectionPreview.h"
+#include "MovieSceneSection.h"
+
+void FSequencerSelectionPreview::SetSelectionState(FSequencerSelectedKey Key, ESelectionPreviewState InState)
+{
+	if (InState == ESelectionPreviewState::Undefined)
+	{
+		DefinedKeyStates.Remove(Key);
+	}
+	else
+	{
+		DefinedKeyStates.Add(Key, InState);
+	}
+
+	CachedSelectionHash.Reset();
+}
+
+void FSequencerSelectionPreview::SetSelectionState(TWeakPtr<UE::Sequencer::FViewModel> InModel, ESelectionPreviewState InState)
+{
+	if (InState == ESelectionPreviewState::Undefined)
+	{
+		DefinedModelStates.Remove(InModel);
+	}
+	else
+	{
+		DefinedModelStates.Add(InModel, InState);
+	}
+	CachedSelectionHash.Reset();
+}
+
+ESelectionPreviewState FSequencerSelectionPreview::GetSelectionState(const FSequencerSelectedKey& Key) const
+{
+	if (const ESelectionPreviewState* State = DefinedKeyStates.Find(Key))
+	{
+		return *State;
+	}
+	return ESelectionPreviewState::Undefined;
+}
+
+ESelectionPreviewState FSequencerSelectionPreview::GetSelectionState(TWeakPtr<UE::Sequencer::FViewModel> InModel) const
+{
+	if (const ESelectionPreviewState* State = DefinedModelStates.Find(InModel))
+	{
+		return *State;
+	}
+	return ESelectionPreviewState::Undefined;
+}
+
+void FSequencerSelectionPreview::Empty()
+{
+	EmptyDefinedKeyStates();
+	EmptyDefinedModelStates();
+}
+
+void FSequencerSelectionPreview::EmptyDefinedKeyStates()
+{
+	DefinedKeyStates.Reset();
+	CachedSelectionHash.Reset();
+}
+
+void FSequencerSelectionPreview::EmptyDefinedModelStates()
+{
+	DefinedModelStates.Reset();
+	CachedSelectionHash.Reset();
+}
+
+uint32 FSequencerSelectionPreview::GetSelectionHash() const
+{
+	if (!CachedSelectionHash.IsSet())
+	{
+		uint32 NewHash = 0;
+
+		for (TPair<FSequencerSelectedKey, ESelectionPreviewState> Pair : DefinedKeyStates)
+		{
+			NewHash = HashCombine(NewHash, HashCombine(GetTypeHash(Pair.Key), GetTypeHash(Pair.Value)));
+		}
+		for (TPair<TWeakPtr<UE::Sequencer::FViewModel>, ESelectionPreviewState> Pair : DefinedModelStates)
+		{
+			NewHash = HashCombine(NewHash, HashCombine(GetTypeHash(Pair.Key), GetTypeHash(Pair.Value)));
+		}
+
+		CachedSelectionHash = NewHash;
+	}
+
+	return CachedSelectionHash.GetValue();
+}
diff --git a/Engine/Source/Editor/Sequencer/Private/SequencerSelectionPreview.h b/Engine/Source/Editor/Sequencer/Private/SequencerSelectionPreview.h
--- a/Engine/Source/Editor/Sequencer/Private/SequencerSelectionPreview.h
+++ b/Engine/Source/Editor/Sequencer/Private/SequencerSelectionPreview.h
@@ -4,9 +4,10 @@
 
 #include "CoreMinimal.h"
 #include "UObject/WeakObjectPtr.h"
-#include "DisplayNodes/SequencerDisplayNode.h"
 #include "SequencerSelectedKey.h"
 
+#include "MVVM/ViewModels/ViewModel.h"
+
 class UMovieSceneSection;
 
 enum class ESelectionPreviewState
@@ -27,50 +28,37 @@
 	/** Access the defined key states */
 	const TMap<FSequencerSelectedKey, ESelectionPreviewState>& GetDefinedKeyStates() const { return DefinedKeyStates; }
 
-	/** Access the defined section states */
-	const TMap<TWeakObjectPtr<UMovieSceneSection>, ESelectionPreviewState>& GetDefinedSectionStates() const { return DefinedSectionStates; }
+	/** Access the defined model states */
+	const TMap<TWeakPtr<UE::Sequencer::FViewModel>, ESelectionPreviewState>& GetDefinedModelStates() const { return DefinedModelStates; }
 
-	/** Access the defined display node states */
-	const TMap<TSharedRef<FSequencerDisplayNode>, ESelectionPreviewState>& GetDefinedOutlinerNodeStates() const { return DefinedOutlinerNodeStates; }
-
 	/** Adds a key to the selection */
 	void SetSelectionState(FSequencerSelectedKey Key, ESelectionPreviewState InState);
 
-	/** Adds a section to the selection */
-	void SetSelectionState(UMovieSceneSection* Section, ESelectionPreviewState InState);
+	/** Adds a model to the selection */
+	void SetSelectionState(TWeakPtr<UE::Sequencer::FViewModel> InModel, ESelectionPreviewState InState);
 
-	/** Adds an outliner node to the selection */
-	void SetSelectionState(TSharedRef<FSequencerDisplayNode> OutlinerNode, ESelectionPreviewState InState);
-
 	/** Returns the selection state for the the specified key. */
 	ESelectionPreviewState GetSelectionState(const FSequencerSelectedKey& Key) const;
 
-	/** Returns the selection state for the the specified section. */
-	ESelectionPreviewState GetSelectionState(UMovieSceneSection* Section) const;
+	/** Returns the selection state for the the specified model. */
+	ESelectionPreviewState GetSelectionState(TWeakPtr<UE::Sequencer::FViewModel> InModel) const;
 
-	/** Returns the selection state for the the specified section. */
-	ESelectionPreviewState GetSelectionState(TSharedRef<FSequencerDisplayNode> OutlinerNode) const;
-
 	/** Empties all selections. */
 	void Empty();
 
 	/** Empties the key selection. */
 	void EmptyDefinedKeyStates();
 
-	/** Empties the section selection. */
-	void EmptyDefinedSectionStates();
+	/** Empties the model selection. */
+	void EmptyDefinedModelStates();
 
-	/** Empties the outliner node selection. */
-	void EmptyDefinedOutlinerNodeStates();
-
 	/** Hash the contents of this selection preview. */
 	uint32 GetSelectionHash() const;
 
 private:
 
 	TMap<FSequencerSelectedKey, ESelectionPreviewState> DefinedKeyStates;
-	TMap<TWeakObjectPtr<UMovieSceneSection>, ESelectionPreviewState> DefinedSectionStates;
-	TMap<TSharedRef<FSequencerDisplayNode>, ESelectionPreviewState> DefinedOutlinerNodeStates;
+	TMap<TWeakPtr<UE::Sequencer::FViewModel>, ESelectionPreviewState> DefinedModelStates;
 
 	/** Cached hash of this whole selection preview state. */
 	mutable TOptional<uint32> CachedSelectionHash;
diff --git a/Engine/Source/Editor/Sequencer/Private/SequencerTimeSliderController.cpp b/Engine/Source/Editor/Sequencer/Private/SequencerTimeSliderController.cpp
--- a/Engine/Source/Editor/Sequencer/Private/SequencerTimeSliderController.cpp
+++ b/Engine/Source/Editor/Sequencer/Private/SequencerTimeSliderController.cpp
@@ -21,7 +21,6 @@
 #include "MovieSceneTimeHelpers.h"
 #include "CommonFrameRates.h"
 #include "Sequencer.h"
-#include "SequencerDisplayNode.h"
 #include "Modules/ModuleManager.h"
 #include "MovieSceneSequence.h"
 #include "MovieScene.h"
@@ -1748,6 +1747,8 @@
 
 FFrameTime FSequencerTimeSliderController::SnapTimeToNearestKey(const FPointerEvent& MouseEvent, const FScrubRangeToScreen& RangeToScreen, float CursorPos, FFrameTime InTime) const
 {
+	using namespace UE::Sequencer;
+
 	if (!WeakSequencer.IsValid())
 	{
 		return InTime;
@@ -1759,7 +1760,7 @@
 
 		// If there are any tracks selected we'll find the nearest key only on that track. If there are no keys selected,
 		// we will try to find the nearest keys on all tracks. This mirrors the behavior of the Jump to Next Keyframe commands.
-		const TSet< TSharedRef<FSequencerDisplayNode> >& SelectedNodes = WeakSequencer.Pin()->GetSelection().GetSelectedOutlinerNodes();
+		const TSet< TWeakPtr<FViewModel> >& SelectedNodes = WeakSequencer.Pin()->GetSelection().GetSelectedOutlinerItems();
 		if (SelectedNodes.Num() == 0)
 		{
 			EnumAddFlags(NearestKeyOption, ENearestKeyOption::NKO_SearchAllTracks);
diff --git a/Engine/Source/Editor/Sequencer/Private/SequencerTrackFilters.h b/Engine/Source/Editor/Sequencer/Private/SequencerTrackFilters.h
--- a/Engine/Source/Editor/Sequencer/Private/SequencerTrackFilters.h
+++ b/Engine/Source/Editor/Sequencer/Private/SequencerTrackFilters.h
@@ -7,7 +7,6 @@
 #include "CollectionManagerTypes.h"
 #include "Styling/SlateIconFinder.h"
 #include "Styling/AppStyle.h"
-#include "SequencerDisplayNode.h"
 #include "MovieSceneSequence.h"
 
 #include "Tracks/MovieSceneAudioTrack.h"
diff --git a/Engine/Source/Editor/Sequencer/Private/SequencerUtilities.cpp b/Engine/Source/Editor/Sequencer/Private/SequencerUtilities.cpp
--- a/Engine/Source/Editor/Sequencer/Private/SequencerUtilities.cpp
+++ b/Engine/Source/Editor/Sequencer/Private/SequencerUtilities.cpp
@@ -1,6 +1,7 @@
 // Copyright Epic Games, Inc. All Rights Reserved.
 
 #include "SequencerUtilities.h"
+#include "Misc/Attribute.h"
 #include "Misc/Paths.h"
 #include "Layout/Margin.h"
 #include "Fonts/SlateFontInfo.h"
@@ -12,7 +13,6 @@
 #include "Widgets/Input/SButton.h"
 #include "Styling/CoreStyle.h"
 #include "Styling/AppStyle.h"
-#include "SequencerTrackNode.h"
 #include "MovieSceneTrack.h"
 #include "MovieSceneSection.h"
 #include "MovieSceneTimeHelpers.h"
@@ -23,12 +23,14 @@
 #include "ISequencer.h"
 #include "Sequencer.h"
 #include "SequencerNodeTree.h"
-#include "DisplayNodes/SequencerObjectBindingNode.h"
 #include "ScopedTransaction.h"
 #include "UObject/Package.h"
 #include "AssetRegistry/AssetRegistryModule.h"
 #include "FileHelpers.h"
 #include "LevelSequence.h"
+#include "MVVM/Extensions/IObjectBindingExtension.h"
+#include "MVVM/ViewModels/TrackModel.h"
+#include "MVVM/Views/ViewUtilities.h"
 #include "Framework/Notifications/NotificationManager.h"
 #include "Widgets/Notifications/SNotificationList.h"
 
@@ -36,18 +38,6 @@
 
 #define LOCTEXT_NAMESPACE "FSequencerUtilities"
 
-static EVisibility GetRolloverVisibility(TAttribute<bool> HoverState, TWeakPtr<SComboButton> WeakComboButton)
-{
-	TSharedPtr<SComboButton> ComboButton = WeakComboButton.Pin();
-	if (HoverState.Get() || ComboButton->IsOpen())
-	{
-		return EVisibility::SelfHitTestInvisible;
-	}
-	else
-	{
-		return EVisibility::Collapsed;
-	}
-}
 
 static EVisibility GetRolloverVisibility(TAttribute<bool> HoverState, TWeakPtr<SButton> WeakButton)
 {
@@ -64,51 +54,8 @@
 
 TSharedRef<SWidget> FSequencerUtilities::MakeAddButton(FText HoverText, FOnGetContent MenuContent, const TAttribute<bool>& HoverState, TWeakPtr<ISequencer> InSequencer)
 {
-	FSlateFontInfo SmallLayoutFont = FCoreStyle::GetDefaultFontStyle("Regular", 8);
-
-	TSharedRef<STextBlock> ComboButtonText = SNew(STextBlock)
-		.Text(HoverText)
-		.Font(SmallLayoutFont)
-		.ColorAndOpacity( FSlateColor::UseForeground() );
-
-	TSharedRef<SComboButton> ComboButton =
-
-		SNew(SComboButton)
-		.HasDownArrow(false)
-		.IsFocusable(true)
-		.ButtonStyle(FAppStyle::Get(), "HoverHintOnly")
-		.ForegroundColor( FSlateColor::UseForeground() )
-		.IsEnabled_Lambda([=]() { return InSequencer.IsValid() ? !InSequencer.Pin()->IsReadOnly() : false; })
-		.OnGetMenuContent(MenuContent)
-		.ContentPadding(FMargin(5, 2))
-		.HAlign(HAlign_Center)
-		.VAlign(VAlign_Center)
-		.ButtonContent()
-		[
-			SNew(SHorizontalBox)
-
-			+ SHorizontalBox::Slot()
-			.AutoWidth()
-			.VAlign(VAlign_Center)
-			.Padding(FMargin(0,0,2,0))
-			[
-				SNew(SImage)
-				.ColorAndOpacity( FSlateColor::UseForeground() )
-				.Image(FAppStyle::GetBrush("Plus"))
-			]
-
-			+ SHorizontalBox::Slot()
-			.VAlign(VAlign_Center)
-			.AutoWidth()
-			[
-				ComboButtonText
-			]
-		];
-
-	TAttribute<EVisibility> Visibility = TAttribute<EVisibility>::Create(TAttribute<EVisibility>::FGetter::CreateStatic(GetRolloverVisibility, HoverState, TWeakPtr<SComboButton>(ComboButton)));
-	ComboButtonText->SetVisibility(Visibility);
-
-	return ComboButton;
+	TAttribute<bool> IsEnabled = MakeAttributeLambda([InSequencer]() -> bool { return InSequencer.IsValid() ? !InSequencer.Pin()->IsReadOnly() : false; });
+	return UE::Sequencer::MakeAddButton(HoverText, MenuContent, HoverState, IsEnabled);
 }
 
 TSharedRef<SWidget> FSequencerUtilities::MakeAddButton(FText HoverText, FOnClicked OnClicked, const TAttribute<bool>& HoverState, TWeakPtr<ISequencer> InSequencer)
@@ -284,6 +231,9 @@
 
 void FSequencerUtilities::PopulateMenu_SetBlendType(FMenuBuilder& MenuBuilder, const TArray<TWeakObjectPtr<UMovieSceneSection>>& InSections, TWeakPtr<ISequencer> InSequencer)
 {
+	using namespace UE::MovieScene;
+	using namespace UE::Sequencer;
+
 	auto Execute = [InSections, InSequencer](EMovieSceneBlendType BlendType)
 	{
 		FScopedTransaction Transaction(LOCTEXT("SetBlendType", "Set Blend Type"));
@@ -309,19 +259,19 @@
 				{
 					if (UMovieSceneSection* Section = WeakSection.Get())
 					{
-						TOptional<FSectionHandle> SectionHandle = SequencerNodeTree->GetSectionHandle(Section);
+						TSharedPtr<FSectionModel> SectionHandle = SequencerNodeTree->GetSectionModel(Section);
 						if (!SectionHandle)
 						{
 							continue;
 						}
 
-						TSharedPtr<FSequencerObjectBindingNode> ParentObjectBindingNode = SectionHandle->GetTrackNode()->FindParentObjectBindingNode();
+						TSharedPtr<IObjectBindingExtension> ParentObjectBindingNode = SectionHandle->FindAncestorOfType<IObjectBindingExtension>();
 						if (!ParentObjectBindingNode.IsValid())
 						{
 							continue;
 						}
 
-						for (TWeakObjectPtr<> BoundObject : Sequencer->FindObjectsInCurrentSequence(ParentObjectBindingNode->GetObjectBinding()))
+						for (TWeakObjectPtr<> BoundObject : Sequencer->FindObjectsInCurrentSequence(ParentObjectBindingNode->GetObjectGuid()))
 						{
 							if (AActor* BoundActor = Cast<AActor>(BoundObject))
 							{
@@ -529,7 +479,7 @@
 		FMovieScenePossessable* ThisPossessable = OwnerMovieScene->FindPossessable(NewPossessable.GetGuid());
 		if (ensure(ThisPossessable))
 		{
-			ThisPossessable->SetParent(NewGuid);
+			ThisPossessable->SetParent(NewGuid, OwnerMovieScene);
 		}
 	};
 
@@ -558,7 +508,7 @@
 					if (!bComponentWasUpdated)
 					{
 						FMovieScenePossessable* ThisPossessable = OwnerMovieScene->FindPossessable(ComponentGuid);
-						ThisPossessable->SetParent(FGuid());
+						ThisPossessable->SetParent(FGuid(), OwnerMovieScene);
 					}
 				}
 			}
@@ -594,7 +544,10 @@
 
 	// Try to fix up folders
 	TArray<UMovieSceneFolder*> FoldersToCheck;
-	FoldersToCheck.Append(InSequencerPtr->GetFocusedMovieSceneSequence()->GetMovieScene()->GetRootFolders());
+	for (UMovieSceneFolder* Folder : InSequencerPtr->GetFocusedMovieSceneSequence()->GetMovieScene()->GetRootFolders())
+	{
+		FoldersToCheck.Add(Folder);
+	}
 	bool bFolderFound = false;
 	while (FoldersToCheck.Num() > 0 && bFolderFound == false)
 	{
