Sequencer: MVVM2 branch and Layer Bars

diff --git a/Engine/Source/Editor/Sequencer/Private/Tools/EditToolDragOperations.cpp b/Engine/Source/Editor/Sequencer/Private/Tools/EditToolDragOperations.cpp
--- a/Engine/Source/Editor/Sequencer/Private/Tools/EditToolDragOperations.cpp
+++ b/Engine/Source/Editor/Sequencer/Private/Tools/EditToolDragOperations.cpp
@@ -7,8 +7,15 @@
 #include "Sequencer.h"
 #include "SequencerSettings.h"
 #include "SequencerCommonHelpers.h"
-#include "VirtualTrackArea.h"
-#include "SequencerTrackNode.h"
+#include "MVVM/ViewModels/ViewModel.h"
+#include "MVVM/SectionModelStorageExtension.h"
+#include "MVVM/Extensions/IDraggableTrackAreaExtension.h"
+#include "MVVM/ViewModels/SectionModel.h"
+#include "MVVM/ViewModels/TrackModel.h"
+#include "MVVM/ViewModels/TrackRowModel.h"
+#include "MVVM/ViewModels/SequencerEditorViewModel.h"
+#include "MVVM/ViewModels/TrackAreaViewModel.h"
+#include "MVVM/ViewModels/VirtualTrackArea.h"
 #include "Algo/AllOf.h"
 #include "MovieSceneTimeHelpers.h"
 #include "Modules/ModuleManager.h"
@@ -16,26 +23,26 @@
 #include "Channels/MovieSceneChannelProxy.h"
 #include "ISequencerModule.h"
 
-struct FInvalidKeyAndSectionSnappingCandidates : ISequencerSnapCandidate
+struct FInvalidKeyAndSectionSnappingCandidates : UE::Sequencer::ISnapCandidate
 {
 	/**
-	 * Keys and Sections added to this ISequencerSnapCandidate will be ignored as potential candidates for snapping.
+	 * Keys and Sections added to this ISnapField will be ignored as potential candidates for snapping.
 	 */
-	FInvalidKeyAndSectionSnappingCandidates(const TSet<FSequencerSelectedKey>& InKeysToIgnore, const TArray<TWeakObjectPtr<UMovieSceneSection>>& InSectionsToIgnore)
+	FInvalidKeyAndSectionSnappingCandidates(const TSet<FSequencerSelectedKey>& InKeysToIgnore, const TSet<UMovieSceneSection*>& InSectionsToIgnore)
 	{
 		KeysToExclude = InKeysToIgnore;
-		for (const TWeakObjectPtr<UMovieSceneSection>& WeakSection : InSectionsToIgnore)
-		{
-			SectionsToExclude.Add(WeakSection.Get());
-		}
+		SectionsToExclude = InSectionsToIgnore;
 	}
 
-	virtual bool IsKeyApplicable(FKeyHandle KeyHandle, const TSharedPtr<IKeyArea>& KeyArea, UMovieSceneSection* Section) override
+	virtual bool IsKeyApplicable(FKeyHandle KeyHandle, const UE::Sequencer::FViewModelPtr& Owner) const override
 	{
-		return !KeysToExclude.Contains(FSequencerSelectedKey(*Section, KeyArea, KeyHandle));
+		using namespace UE::Sequencer;
+
+		TSharedPtr<FChannelModel> Channel = Owner.ImplicitCast();
+		return !Channel || (!KeysToExclude.Contains(FSequencerSelectedKey(*Channel->GetSection(), Channel, KeyHandle)) && !SectionsToExclude.Contains(Channel->GetSection()));
 	}
 
-	virtual bool AreSectionBoundsApplicable(UMovieSceneSection* Section) override
+	virtual bool AreSectionBoundsApplicable(UMovieSceneSection* Section) const override
 	{
 		return !SectionsToExclude.Contains(Section);
 	}
@@ -92,24 +99,20 @@
 	return LayerId;
 }
 
-void FEditToolDragOperation::BeginTransaction( TArray< TWeakObjectPtr<UMovieSceneSection> >& Sections, const FText& TransactionDesc )
+void FEditToolDragOperation::BeginTransaction( TSet<UMovieSceneSection*>& Sections, const FText& TransactionDesc )
 {
 	// Begin an editor transaction and mark the section as transactional so it's state will be saved
 	Transaction.Reset( new FScopedTransaction(TransactionDesc) );
 
-	for (int32 SectionIndex = 0; SectionIndex < Sections.Num(); )
+	for (auto It = Sections.CreateIterator(); It; ++It)
 	{
-		UMovieSceneSection* SectionObj = Sections[SectionIndex].Get();
+		UMovieSceneSection* SectionObj = *It;
 
 		SectionObj->SetFlags( RF_Transactional );
 		// Save the current state of the section
-		if (SectionObj->TryModify())
+		if (!SectionObj->TryModify())
 		{
-			++SectionIndex;
-		}
-		else
-		{
-			Sections.RemoveAt(SectionIndex);
+			It.RemoveCurrent();
 		}
 	}
 }
@@ -122,15 +125,28 @@
 
 FResizeSection::FResizeSection( FSequencer& InSequencer, const TSet<TWeakObjectPtr<UMovieSceneSection>>& InSections, bool bInDraggingByEnd, bool bInIsSlipping )
 	: FEditToolDragOperation( InSequencer )
-	, Sections( InSections.Array() )
 	, bDraggingByEnd(bInDraggingByEnd)
 	, bIsSlipping(bInIsSlipping)
 	, MouseDownTime(0)
 {
+	Sections.Reserve(InSections.Num());
+
+	for (TWeakObjectPtr<UMovieSceneSection> WeakSection : InSections)
+	{
+		UMovieSceneSection* Section = WeakSection.Get();
+		if (Section)
+		{
+			Sections.Add(Section);
+		}
+	}
 }
 
-void FResizeSection::OnBeginDrag(const FPointerEvent& MouseEvent, FVector2D LocalMousePos, const FVirtualTrackArea& VirtualTrackArea)
+void FResizeSection::OnBeginDrag(const FPointerEvent& MouseEvent, FVector2D LocalMousePos, const UE::Sequencer::FVirtualTrackArea& VirtualTrackArea)
 {
+	using namespace UE::Sequencer;
+
+	UE::MovieScene::FScopedSignedObjectModifyDefer DeferMarkAsChanged;
+
 	BeginTransaction( Sections, NSLOCTEXT("Sequencer", "DragSectionEdgeTransaction", "Resize section") );
 
 	MouseDownTime = VirtualTrackArea.PixelToFrame(LocalMousePos.X);
@@ -145,14 +161,8 @@
 	bool bIsDilating = MouseEvent.IsControlDown();
 	PreDragSectionData.Empty();
 
-	for (TWeakObjectPtr<UMovieSceneSection> WeakSection : Sections)
+	for (UMovieSceneSection* Section : Sections)
 	{
-		UMovieSceneSection* Section = WeakSection.Get();
-		if (Section == nullptr)
-		{
-			continue;
-		}
-
 		if (bIsDilating)
 		{
 			// Populate the resize data for this section
@@ -160,12 +170,12 @@
 			ResizeData.MovieSection = Section;
 			ResizeData.InitialRange = Section->GetRange();
 
-			TOptional<FSectionHandle> SectionHandle = Sequencer.GetNodeTree()->GetSectionHandle(Section);
+			TSharedPtr<FSectionModel> SectionHandle = Sequencer.GetNodeTree()->GetSectionModel(Section);
 			if (SectionHandle)
 			{
 				//Tell section that may not have keys it's starting to dilate (e.g. skeletal tracks will cache play rate).
 				SectionHandle->GetSectionInterface()->BeginDilateSection();
-				ResizeData.SequencerSection = &SectionHandle->GetSectionInterface().Get();
+				ResizeData.SequencerSection = SectionHandle->GetSectionInterface().Get();
 			}
 
 			// Add the key times for all keys of all channels on this section
@@ -184,7 +194,7 @@
 			}
 			PreDragSectionData.Emplace(ResizeData);
 		}
-		else if (TOptional<FSectionHandle> SectionHandle = Sequencer.GetNodeTree()->GetSectionHandle(Section))
+		else if (TSharedPtr<FSectionModel> SectionHandle = Sequencer.GetNodeTree()->GetSectionModel(Section))
 		{
 			if (bIsSlipping)
 			{
@@ -200,7 +210,7 @@
 	}
 }
 
-void FResizeSection::OnEndDrag(const FPointerEvent& MouseEvent, FVector2D LocalMousePos, const FVirtualTrackArea& VirtualTrackArea)
+void FResizeSection::OnEndDrag(const FPointerEvent& MouseEvent, FVector2D LocalMousePos, const UE::Sequencer::FVirtualTrackArea& VirtualTrackArea)
 {
 	bool bIsDilating = MouseEvent.IsControlDown();
 
@@ -211,17 +221,14 @@
 		TSet<UMovieSceneTrack*> Tracks;
 		FMovieSceneSectionMovedParams SectionMovedParams(EPropertyChangeType::ValueSet);
 		EMovieSceneSectionMovedResult SectionMovedResult(EMovieSceneSectionMovedResult::None);
-		for (const TWeakObjectPtr<UMovieSceneSection>& WeakSection : Sections)
+		for (UMovieSceneSection* Section : Sections)
 		{
-			if (UMovieSceneSection* Section = WeakSection.Get())
+			if (UMovieSceneTrack* OuterTrack = Section->GetTypedOuter<UMovieSceneTrack>())
 			{
-				if (UMovieSceneTrack* OuterTrack = Section->GetTypedOuter<UMovieSceneTrack>())
-				{
-					OuterTrack->Modify();
-					SectionMovedResult |= OuterTrack->OnSectionMoved(*Section, SectionMovedParams);
+				OuterTrack->Modify();
+				SectionMovedResult |= OuterTrack->OnSectionMoved(*Section, SectionMovedParams);
 
-					Tracks.Add(OuterTrack);
-				}
+				Tracks.Add(OuterTrack);
 			}
 		}
 		for (UMovieSceneTrack* Track : Tracks)
@@ -242,8 +249,12 @@
 	}
 }
 
-void FResizeSection::OnDrag(const FPointerEvent& MouseEvent, FVector2D LocalMousePos, const FVirtualTrackArea& VirtualTrackArea)
+void FResizeSection::OnDrag(const FPointerEvent& MouseEvent, FVector2D LocalMousePos, const UE::Sequencer::FVirtualTrackArea& VirtualTrackArea)
 {
+	using namespace UE::Sequencer;
+
+	UE::MovieScene::FScopedSignedObjectModifyDefer DeferMarkAsChanged(true);
+
 	ISequencerModule& SequencerModule = FModuleManager::Get().LoadModuleChecked<ISequencerModule>("Sequencer");
 
 	bool bIsDilating = MouseEvent.IsControlDown();
@@ -260,12 +271,9 @@
 	if ( Settings->GetIsSnapEnabled() )
 	{
 		TArray<FFrameNumber> SectionTimes;
-		for (const TWeakObjectPtr<UMovieSceneSection>& WeakSection : Sections)
+		for (UMovieSceneSection* Section : Sections)
 		{
-			if (UMovieSceneSection* Section = WeakSection.Get())
-			{
-				SectionTimes.Add(SectionInitTimes[Section] + DeltaTime);
-			}
+			SectionTimes.Add(SectionInitTimes[Section] + DeltaTime);
 		}
 
 		float SnapThresholdPx = VirtualTrackArea.PixelToSeconds(PixelSnapWidth) - VirtualTrackArea.PixelToSeconds(0.f);
@@ -293,7 +301,6 @@
 	
 	/********************************************************************/
 	EMovieSceneSectionMovedResult SectionMovedResult(EMovieSceneSectionMovedResult::None);
-
 	if (bIsDilating)
 	{
 		for(FPreDragSectionData Data: PreDragSectionData)
@@ -386,16 +393,15 @@
 		}
 	}
 	/********************************************************************/
-	else for (const TWeakObjectPtr<UMovieSceneSection>& WeakSection : Sections)
+	else for (UMovieSceneSection* Section : Sections)
 	{
-		UMovieSceneSection* Section = WeakSection.Get();
-		TOptional<FSectionHandle> SectionHandle = Sequencer.GetNodeTree()->GetSectionHandle(Section);
+		TSharedPtr<FSectionModel> SectionHandle = Sequencer.GetNodeTree()->GetSectionModel(Section);
 		if (!SectionHandle)
 		{
 			continue;
 		}
 
-		TSharedRef<ISequencerSection> SectionInterface = SectionHandle->GetSectionInterface();
+		TSharedPtr<ISequencerSection> SectionInterface = SectionHandle->GetSectionInterface();
 
 		FFrameNumber NewTime = SectionInitTimes[Section] + DeltaTime;
 
@@ -455,14 +461,11 @@
 
 	{
 		TSet<UMovieSceneTrack*> Tracks;
-		for (TWeakObjectPtr<UMovieSceneSection> WeakSection : Sections)
+		for (UMovieSceneSection* Section : Sections)
 		{
-			if (UMovieSceneSection* Section = WeakSection.Get())
+			if (UMovieSceneTrack* Track = Section->GetTypedOuter<UMovieSceneTrack>())
 			{
-				if (UMovieSceneTrack* Track = Section->GetTypedOuter<UMovieSceneTrack>())
-				{
-					Tracks.Add(Track);
-				}
+				Tracks.Add(Track);
 			}
 		}
 		for (UMovieSceneTrack* Track : Tracks)
@@ -481,7 +484,7 @@
 	}
 }
 
-void FDuplicateKeysAndSections::OnBeginDrag(const FPointerEvent& MouseEvent, FVector2D LocalMousePos, const FVirtualTrackArea& VirtualTrackArea)
+void FDuplicateKeysAndSections::OnBeginDrag(const FPointerEvent& MouseEvent, FVector2D LocalMousePos, const UE::Sequencer::FVirtualTrackArea& VirtualTrackArea)
 {
 	// Begin an editor transaction and mark the section as transactional so it's state will be saved
 	BeginTransaction( Sections, NSLOCTEXT("Sequencer", "DuplicateKeysTransaction", "Duplicate Keys or Sections") );
@@ -500,12 +503,8 @@
 	// so we don't need to actually use the new key handles.
 	DuplicateKeys(KeysAsArray, NewKeyHandles);
 
-	// Duplicate our selections as well.
-	bool bDelayedStructureRebuild = false;
-
-	for (const TWeakObjectPtr<UMovieSceneSection>& WeakSection : Sections)
+	for (UMovieSceneSection* SectionToDuplicate : Sections)
 	{
-		UMovieSceneSection* SectionToDuplicate = WeakSection.Get();
 		if (!SectionToDuplicate)
 		{
 			continue;
@@ -515,21 +514,13 @@
 		UMovieSceneTrack* OwningTrack = SectionToDuplicate->GetTypedOuter<UMovieSceneTrack>();
 		OwningTrack->Modify();
 		OwningTrack->AddSection(*DuplicatedSection);
-
-		bDelayedStructureRebuild = true;
 	}
 
 	// Now start the move drag
 	FMoveKeysAndSections::OnBeginDrag(MouseEvent, LocalMousePos, VirtualTrackArea);
-
-	if (bDelayedStructureRebuild)
-	{
-		// We need to rebuild the track layout now so that the newly added section shows up, otherwise it won't show up until a section is vertically rearranged.
-		Sequencer.NotifyMovieSceneDataChanged(EMovieSceneDataChangeType::MovieSceneStructureItemsChanged);
-	}
 }
 
-void FDuplicateKeysAndSections::OnEndDrag(const FPointerEvent& MouseEvent, FVector2D LocalMousePos, const FVirtualTrackArea& VirtualTrackArea)
+void FDuplicateKeysAndSections::OnEndDrag(const FPointerEvent& MouseEvent, FVector2D LocalMousePos, const UE::Sequencer::FVirtualTrackArea& VirtualTrackArea)
 {
 	FMoveKeysAndSections::OnEndDrag(MouseEvent, LocalMousePos, VirtualTrackArea);
 
@@ -544,8 +535,10 @@
 {
 }
 
-void FManipulateSectionEasing::OnBeginDrag(const FPointerEvent& MouseEvent, FVector2D LocalMousePos, const FVirtualTrackArea& VirtualTrackArea)
+void FManipulateSectionEasing::OnBeginDrag(const FPointerEvent& MouseEvent, FVector2D LocalMousePos, const UE::Sequencer::FVirtualTrackArea& VirtualTrackArea)
 {
+	using namespace UE::Sequencer;
+
 	Transaction.Reset( new FScopedTransaction(NSLOCTEXT("Sequencer", "DragSectionEasing", "Change Section Easing")) );
 
 	UMovieSceneSection* Section = WeakSection.Get();
@@ -562,15 +555,17 @@
 	if (Settings->GetSnapSectionTimesToSections())
 	{
 		// Construct a snap field of all section bounds
-		ISequencerSnapCandidate SnapCandidates;
+		ISnapCandidate SnapCandidates;
 		SnapField = FSequencerSnapField(Sequencer, SnapCandidates, ESequencerEntity::Section);
 	}
 
 	InitValue = bEaseIn ? Section->Easing.GetEaseInDuration() : Section->Easing.GetEaseOutDuration();
 }
 
-void FManipulateSectionEasing::OnDrag(const FPointerEvent& MouseEvent, FVector2D LocalMousePos, const FVirtualTrackArea& VirtualTrackArea)
+void FManipulateSectionEasing::OnDrag(const FPointerEvent& MouseEvent, FVector2D LocalMousePos, const UE::Sequencer::FVirtualTrackArea& VirtualTrackArea)
 {
+	UE::MovieScene::FScopedSignedObjectModifyDefer DeferMarkAsChanged(true);
+
 	ESequencerScrubberStyle ScrubStyle = Sequencer.GetScrubStyle();
 
 	FFrameRate TickResolution  = Sequencer.GetFocusedTickResolution();
@@ -625,6 +620,8 @@
 
 	const int32 MaxEasingDuration = Section->HasStartFrame() && Section->HasEndFrame() ? UE::MovieScene::DiscreteSize(Section->GetRange()) : TNumericLimits<int32>::Max() / 2;
 
+	Section->Modify();
+
 	if (bEaseIn)
 	{
 		Section->Easing.bManualEaseIn = true;
@@ -635,64 +632,69 @@
 		Section->Easing.bManualEaseOut = true;
 		Section->Easing.ManualEaseOutDuration = FMath::Clamp(InitValue.Get(0) - DeltaTime.RoundToFrame().Value, 0, MaxEasingDuration);
 	}
-
-	UMovieSceneTrack* OuterTrack = Section->GetTypedOuter<UMovieSceneTrack>();
-	if (OuterTrack)
-	{
-		OuterTrack->MarkAsChanged();
-	}
-
-	Sequencer.NotifyMovieSceneDataChanged( EMovieSceneDataChangeType::TrackValueChanged );
 }
 
-void FManipulateSectionEasing::OnEndDrag(const FPointerEvent& MouseEvent, FVector2D LocalMousePos, const FVirtualTrackArea& VirtualTrackArea)
+void FManipulateSectionEasing::OnEndDrag(const FPointerEvent& MouseEvent, FVector2D LocalMousePos, const UE::Sequencer::FVirtualTrackArea& VirtualTrackArea)
 {
 	EndTransaction();
 }
 
 
-FMoveKeysAndSections::FMoveKeysAndSections(FSequencer& InSequencer, const TSet<FSequencerSelectedKey>& InSelectedKeys, const TSet<TWeakObjectPtr<UMovieSceneSection>>& InSelectedSections, bool InbHotspotWasSection)
+FMoveKeysAndSections::FMoveKeysAndSections(FSequencer& InSequencer, ESequencerMoveOperationType MoveType)
 	: FEditToolDragOperation(InSequencer)
-	, bHotspotWasSection(InbHotspotWasSection)
 {
-	// Filter out the keys on sections that are read only
-	for (const FSequencerSelectedKey& SelectedKey : InSelectedKeys)
+	using namespace UE::Sequencer;
+
+	bAllowVerticalMovement = HotspotCast<FSectionHotspotBase>(InSequencer.GetViewModel()->GetTrackArea()->GetHotspot()) != nullptr;
+
+	if (EnumHasAnyFlags(MoveType, ESequencerMoveOperationType::MoveKeys))
 	{
-		if (!SelectedKey.Section->IsReadOnly())
+		// Filter out the keys on sections that are read only
+		for (const FSequencerSelectedKey& SelectedKey : Sequencer.GetSelection().GetSelectedKeys())
 		{
-			Keys.Add(SelectedKey);
+			if (!SelectedKey.Section->IsReadOnly())
+			{
+				Keys.Add(SelectedKey);
+			}
 		}
+
+		KeysAsArray = Keys.Array();
 	}
 
-	KeysAsArray = Keys.Array();
 
-	for (const TWeakObjectPtr<UMovieSceneSection>& WeakSection : InSelectedSections)
+	if (EnumHasAnyFlags(MoveType, ESequencerMoveOperationType::MoveSections))
 	{
-		const UMovieSceneSection* SelectedSection = WeakSection.Get();
-		Sections.AddUnique(WeakSection);
-
-		UMovieScene* MovieScene = InSequencer.GetFocusedMovieSceneSequence()->GetMovieScene();
-		if (MovieScene)
+		for (TWeakPtr<FViewModel> WeakSelectedModel : Sequencer.GetSelection().GetSelectedTrackAreaItems())
 		{
-			// If the section is in a group, we also want to add the sections it is grouped with
-			const FMovieSceneSectionGroup* SectionGroup = MovieScene->GetSectionGroup(*SelectedSection);
-			if (SectionGroup)
+			if (TSharedPtr<IDraggableTrackAreaExtension> DraggableItem = ICastable::CastWeakPtrShared<IDraggableTrackAreaExtension>(WeakSelectedModel))
 			{
-				for (TWeakObjectPtr<UMovieSceneSection> WeakGroupedSection : *SectionGroup)
-				{
-					// Verify sections are still valid, and are not infinite.
-					if (WeakGroupedSection.IsValid())
-					{
-						Sections.AddUnique(WeakGroupedSection);
-					}
-				}
+				DraggableItem->OnBeginDrag(*this);
 			}
 		}
 	}
 }
 
-void FMoveKeysAndSections::OnBeginDrag(const FPointerEvent& MouseEvent, FVector2D LocalMousePos, const FVirtualTrackArea& VirtualTrackArea)
+void FMoveKeysAndSections::AddSnapTime(FFrameNumber SnapTime)
 {
+	RelativeSnapOffsets.Add(SnapTime);
+}
+
+void FMoveKeysAndSections::AddModel(TSharedPtr<UE::Sequencer::FViewModel> Model)
+{
+	using namespace UE::Sequencer;
+
+	if (FSectionModel* SectionModel = Model->CastThis<FSectionModel>())
+	{
+		Sections.Add(SectionModel->GetSection());
+	}
+	if (TSharedPtr<IDraggableTrackAreaExtension> DraggableItem = Model->CastThisShared<IDraggableTrackAreaExtension>())
+	{
+		DraggedItems.Add(DraggableItem);
+	}
+}
+
+void FMoveKeysAndSections::OnBeginDrag(const FPointerEvent& MouseEvent, FVector2D LocalMousePos, const UE::Sequencer::FVirtualTrackArea& VirtualTrackArea)
+{
 	// Early out if we've somehow started a drag operation without any sections or keys. This prevents an empty Undo/Redo Transaction from being created.
 	if (!Sections.Num() && !Keys.Num())
 	{
@@ -703,39 +705,24 @@
 
 	// Tell the Snap Field to ignore our currently selected keys and sections. We can snap to the edges of non-selected
 	// sections and keys. The actual snapping field will add other sequencer data (play ranges, playheads, etc.) as snap targets.
-	FInvalidKeyAndSectionSnappingCandidates AvoidSnapCanidates(Keys, Sections); 
+	FInvalidKeyAndSectionSnappingCandidates AvoidSnapCanidates(Keys, Sections);
 	SnapField = FSequencerSnapField(Sequencer, AvoidSnapCanidates );
 
 	// Store the frame time of the mouse so we can see how far we've moved from the starting point.
-	MouseTimePrev = VirtualTrackArea.PixelToFrame(LocalMousePos.X).FloorToFrame();
+	MouseTimePrev = VirtualTrackArea.PixelToFrame(LocalMousePos.X).RoundToFrame();
 
-	// Now we store a relative offset to each key and section from the start position. This allows us to know how far away from
-	// the mouse each valid key/section was so we can restore their offset if needed.
-	RelativeOffsets.Reserve(Sections.Num() + Keys.Num());
-	for (TWeakObjectPtr<UMovieSceneSection> WeakSection : Sections)
+	// Convert initial snapoffsets so they are relative to the mouse time
+	for (FFrameNumber& RelativeTime : RelativeSnapOffsets)
 	{
-		UMovieSceneSection* Section = WeakSection.Get();
-		FRelativeOffset Offset;
-
-		if (Section->HasStartFrame())
-		{
-			Offset.StartOffset = Section->GetInclusiveStartFrame() - MouseTimePrev;
-		}
-		if (Section->HasEndFrame())
-		{
-			Offset.EndOffset = Section->GetExclusiveEndFrame() - MouseTimePrev;
-		}
-
-		RelativeOffsets.Add(Offset);
+		RelativeTime = (RelativeTime - MouseTimePrev).FloorToFrame();
 	}
 
 	// Sections can be dragged vertically to adjust their row up or down, so we need to store what row each section is currently on. A section
 	// can be dragged above all other sections - this is accomplished by moving all other sections down. We store the row indices for all sections
 	// in all tracks that we're modifying so we can get them later to move them.
 	TSet<UMovieSceneTrack*> Tracks;
-	for (TWeakObjectPtr<UMovieSceneSection> WeakSection : Sections)
+	for (UMovieSceneSection* Section : Sections)
 	{
-		UMovieSceneSection* Section = WeakSection.Get();
 		Tracks.Add(Section->GetTypedOuter<UMovieSceneTrack>());
 	}
 	for (UMovieSceneTrack* Track : Tracks)
@@ -752,22 +739,21 @@
 	KeyTimes.SetNum(Keys.Num());
 	GetKeyTimes(KeysAsArray, KeyTimes);
 
+	const int32 StartNum = RelativeSnapOffsets.Num();
+	RelativeSnapOffsets.SetNumUninitialized(StartNum + KeyTimes.Num());
 	for (int32 Index = 0; Index < KeyTimes.Num(); ++Index)
 	{
-		// Key offsets use only the Start offset and don't set the End offset as they do not represent ranges.
-		FRelativeOffset KeyOffset;
-		KeyOffset.StartOffset = KeyTimes[Index] - MouseTimePrev;
-
-		// These are packed with our Section Offsets (ToDo: Is this actually a good idea?)
-		RelativeOffsets.Add(KeyOffset);
+		RelativeSnapOffsets[StartNum + Index] = (KeyTimes[Index] - MouseTimePrev).RoundToFrame();
 	}
 
 	// Keys can be moved within sections without the section itself being moved, so we need to call Modify on any section that owns a key that isn't also being moved.
 	ModifyNonSelectedSections();
 }
 
-void FMoveKeysAndSections::OnDrag(const FPointerEvent& MouseEvent, FVector2D LocalMousePos, const FVirtualTrackArea& VirtualTrackArea)
+void FMoveKeysAndSections::OnDrag(const FPointerEvent& MouseEvent, FVector2D LocalMousePos, const UE::Sequencer::FVirtualTrackArea& VirtualTrackArea)
 {
+	UE::MovieScene::FScopedSignedObjectModifyDefer DeferMarkAsChanged(true);
+
 	if (!Sections.Num() && !Keys.Num())
 	{
 		return;
@@ -796,17 +782,11 @@
 		const bool bSnapToInterval = (KeysAsArray.Num() > 0 && Settings->GetSnapKeyTimesToInterval()) || (Sections.Num() > 0 && Settings->GetSnapSectionTimesToInterval());
 		const bool bSnapToLikeTypes = (KeysAsArray.Num() > 0 && Settings->GetSnapKeyTimesToKeys()) || (Sections.Num() > 0 && Settings->GetSnapSectionTimesToSections());
 
-		// RelativeOffsets contains both our sections and our keys, and we add them all as potential things that can snap to stuff.
-		for (const FRelativeOffset& Offset : RelativeOffsets)
+		// RelativeSnapOffsets contains both our sections and our keys, and we add them all as potential things that can snap to stuff.
+		ValidSnapMarkers.SetNumUninitialized(RelativeSnapOffsets.Num());
+		for (int32 Index = 0; Index < RelativeSnapOffsets.Num(); ++Index)
 		{
-			if (Offset.StartOffset.IsSet())
-			{
-				ValidSnapMarkers.Add((Offset.StartOffset.GetValue() + MouseTime).FloorToFrame());
-			}
-			if (Offset.EndOffset.IsSet())
-			{
-				ValidSnapMarkers.Add((Offset.EndOffset.GetValue() + MouseTime).FloorToFrame());
-			}
+			ValidSnapMarkers[Index] = (RelativeSnapOffsets[Index] + MouseTime).FloorToFrame();
 		}
 
 		// Now we'll try and snap all of these points to the closest valid snap marker (which may be a section or interval)
@@ -854,18 +834,14 @@
 	// Get a list of the unique tracks in this selection and update their easing so previews draw interactively as you drag.
 	TSet<UMovieSceneTrack*> Tracks;
 	FMovieSceneSectionMovedParams SectionMovedParams(EPropertyChangeType::Interactive);
-	EMovieSceneSectionMovedResult SectionMovedResult(EMovieSceneSectionMovedResult::None);
-	for (TWeakObjectPtr<UMovieSceneSection> WeakSection : Sections)
+	for (UMovieSceneSection* Section : Sections)
 	{
-		if (UMovieSceneTrack* Track = WeakSection.Get()->GetTypedOuter<UMovieSceneTrack>())
+		if (UMovieSceneTrack* Track = Section->GetTypedOuter<UMovieSceneTrack>())
 		{
-			Track->Modify();
-			SectionMovedResult |= Track->OnSectionMoved(*WeakSection.Get(), SectionMovedParams);
-
+			Track->OnSectionMoved(*Section, SectionMovedParams);
 			Tracks.Add(Track);
 		}
 	}
-	bSectionMovementModifiedStructure |= (SectionMovedResult != EMovieSceneSectionMovedResult::None);
 
 	for (UMovieSceneTrack* Track : Tracks)
 	{
@@ -883,8 +859,12 @@
 	}
 }
 
-void FMoveKeysAndSections::OnEndDrag(const FPointerEvent& MouseEvent, FVector2D LocalMousePos, const FVirtualTrackArea& VirtualTrackArea)
+void FMoveKeysAndSections::OnEndDrag(const FPointerEvent& MouseEvent, FVector2D LocalMousePos, const UE::Sequencer::FVirtualTrackArea& VirtualTrackArea)
 {
+	using namespace UE::Sequencer;
+
+	UE::MovieScene::FScopedSignedObjectModifyDefer DeferMarkAsChanged;
+
 	if (!Sections.Num() && !Keys.Num())
 	{
 		return;
@@ -894,37 +874,29 @@
 	ModifiedNonSelectedSections.Empty();
 
 	// Tracks can tell us if the row indexes for any sections were changed during our drag/drop operation.
-	bool bRowIndicesChanged = false;
 	TSet<UMovieSceneTrack*> Tracks;
 
-	for (TWeakObjectPtr<UMovieSceneSection> WeakSection : Sections)
+	for (UMovieSceneSection* Section : Sections)
 	{
 		// Grab only unique tracks as multiple sections can reside on the same track.
-		Tracks.Add(WeakSection.Get()->GetTypedOuter<UMovieSceneTrack>());
+		Tracks.Add(Section->GetTypedOuter<UMovieSceneTrack>());
 	}
 
 	for (UMovieSceneTrack* Track : Tracks)
 	{
 		// Ensure all of the tracks have updated the row indices for their sections
-		bRowIndicesChanged |= Track->FixRowIndices();
+		Track->FixRowIndices();
 	}
 
-	if (bRowIndicesChanged)
-	{
-		Sequencer.NotifyMovieSceneDataChanged(EMovieSceneDataChangeType::MovieSceneStructureItemsChanged);
-	}
-
 	FMovieSceneSectionMovedParams SectionMovedParams(EPropertyChangeType::ValueSet);
-	EMovieSceneSectionMovedResult SectionMovedResult(EMovieSceneSectionMovedResult::None);
-	for (const TWeakObjectPtr<UMovieSceneSection>& WeakSection : Sections)
+	for (UMovieSceneSection* Section : Sections)
 	{
-		UMovieSceneSection* Section = WeakSection.Get();
 		UMovieSceneTrack* OuterTrack = Cast<UMovieSceneTrack>(Section->GetOuter());
 
 		if (OuterTrack)
 		{
 			OuterTrack->Modify();
-			SectionMovedResult |= OuterTrack->OnSectionMoved(*Section, SectionMovedParams);
+			OuterTrack->OnSectionMoved(*Section, SectionMovedParams);
 		}
 	}
 
@@ -933,12 +905,15 @@
 		Track->UpdateEasing();
 	}
 
+	for (TWeakPtr<IDraggableTrackAreaExtension> WeakDraggableItem : DraggedItems)
+	{
+		if (TSharedPtr<IDraggableTrackAreaExtension> DraggableItem = WeakDraggableItem.Pin())
+		{
+			DraggableItem->OnEndDrag(*this);
+		}
+	}
+
 	EndTransaction();
-
-	if (SectionMovedResult != EMovieSceneSectionMovedResult::None)
-	{
-		Sequencer.NotifyMovieSceneDataChanged(EMovieSceneDataChangeType::MovieSceneStructureItemsChanged);
-	}
 }
 
 void FMoveKeysAndSections::ModifyNonSelectedSections()
@@ -961,18 +936,20 @@
 
 TRange<FFrameNumber> FMoveKeysAndSections::GetSectionBoundaries(const UMovieSceneSection* Section)
 {
+	using namespace UE::Sequencer;
+
 	// Find the borders of where you can drag to
 	FFrameNumber LowerBound = TNumericLimits<int32>::Lowest(), UpperBound = TNumericLimits<int32>::Max();
 
 	// Find the track node for this section
-	TOptional<FSectionHandle> SectionHandle = Sequencer.GetNodeTree()->GetSectionHandle(Section);
+	TSharedPtr<FSectionModel> SectionHandle = Sequencer.GetNodeTree()->GetSectionModel(Section);
 	if (SectionHandle)
 	{
 		// Get the closest borders on either side
-		const TArray< TSharedRef<ISequencerSection> >& AllSections = SectionHandle->GetTrackNode()->GetSections();
-		for (TSharedRef<ISequencerSection> SectionInterface : AllSections)
+		TViewModelPtr<ITrackAreaExtension> TrackModel = SectionHandle->FindAncestorOfType<ITrackAreaExtension>();
+		for (const TViewModelPtr<FSectionModel>& SectionModel : TrackModel->GetTrackAreaModelListAs<FSectionModel>())
 		{
-			const UMovieSceneSection* TestSection = SectionInterface->GetSectionObject();
+			const UMovieSceneSection* TestSection = SectionModel->GetSection();
 			if (!TestSection || Sections.Contains(TestSection))
 			{
 				continue;
@@ -1000,10 +977,9 @@
 	const FFrameNumber MouseDeltaTime = (MouseTime - MouseTimePrev).FloorToFrame();
 
 	// Disallow movement if any of the sections can't move
-	for (int32 Index = 0; Index < Sections.Num(); ++Index)
+	for (UMovieSceneSection* Section : Sections)
 	{
 		// If we're moving a section that is blending with something then it's OK if it overlaps stuff, the blend amount will get updated at the end.
-		UMovieSceneSection* Section = Sections[Index].Get();
 		if (!Section)
 		{
 			continue;
@@ -1109,6 +1085,8 @@
 
 bool FMoveKeysAndSections::HandleSectionMovement(FFrameTime MouseTime, FVector2D VirtualMousePos, FVector2D LocalMousePos, TOptional<FFrameNumber> MaxDeltaX, FFrameNumber DesiredDeltaX)
 {
+	using namespace UE::Sequencer;
+
 	// Don't try to process moving sections if we don't have any sections.
 	if (Sections.Num() == 0)
 	{
@@ -1117,41 +1095,50 @@
 
 	// If sections are all on different rows or from different tracks, don't set row indices for anything because it leads to odd behavior.
 	bool bSectionsAreOnDifferentRows = false;
-	int32 FirstRowIndex = Sections[0].Get()->GetRowIndex();
+	int32 FirstRowIndex = -1;
 	UMovieSceneTrack* FirstTrack = nullptr;
 
-	for (TWeakObjectPtr<UMovieSceneSection> WeakSection : Sections)
+	for (UMovieSceneSection* Section : Sections)
 	{
-		UMovieSceneSection* Section = WeakSection.Get();
-		if (Section)
+		UMovieSceneTrack* Track = Section->GetTypedOuter<UMovieSceneTrack>();
+		if (FirstRowIndex == -1)
 		{
-			UMovieSceneTrack* Track = Section->GetTypedOuter<UMovieSceneTrack>();
-			if (FirstRowIndex != Section->GetRowIndex())
+			FirstRowIndex = Section->GetRowIndex();
+		}
+		if (FirstRowIndex != Section->GetRowIndex())
+		{
+			bSectionsAreOnDifferentRows = true;
+		}
+		if (FirstTrack)
+		{
+			if (FirstTrack != Track)
 			{
 				bSectionsAreOnDifferentRows = true;
 			}
-			if (FirstTrack)
-			{
-				if (FirstTrack != Track)
-				{
-					bSectionsAreOnDifferentRows = true;
-				}
-			}
-			else
-			{
-				FirstTrack = Track;
-			}
+		}
+		else
+		{
+			FirstTrack = Track;
 		}
 	}
 
-	TArray<TSharedRef<FSequencerTrackNode> > TrackNodes;
+	TArray<TSharedPtr<FViewModel>> Tracks;
 
+	// @todo_sequencer_mvvm: really this code should all be operating on models rather than looking directly at UMovieSceneSections and tracks
+	FSectionModelStorageExtension* SectionStorage = Sequencer.GetViewModel()->GetRootModel()->CastDynamic<FSectionModelStorageExtension>();
+
 	bool bRowIndexChanged = false;
-	for (TWeakObjectPtr<UMovieSceneSection> WeakSection : Sections)
+	for (UMovieSceneSection* Section : Sections)
 	{
-		UMovieSceneSection* Section = WeakSection.Get();
-		UMovieSceneTrack* Track = Section->GetTypedOuter<UMovieSceneTrack>();
+		TSharedPtr<FSectionModel>      SectionModel = SectionStorage->FindModelForSection(Section);
+		TSharedPtr<FViewModel> TrackModel   = SectionModel ? SectionModel->GetParentTrackModel() : nullptr;
+		if (!SectionModel || !TrackModel)
+		{
+			continue;
+		}
 
+		UMovieSceneTrack* Track = TrackModel->CastThis<ITrackExtension>()->GetTrack();
+
 		const TArray<UMovieSceneSection*>& AllSections = Track->GetAllSections();
 
 		TArray<UMovieSceneSection*> NonDraggedSections;
@@ -1163,20 +1150,16 @@
 			}
 		}
 
-		TOptional<FSectionHandle> SectionHandle = Sequencer.GetNodeTree()->GetSectionHandle(Section);
-		if (!SectionHandle)
-		{
-			continue;
-		}
+		Tracks.AddUnique(TrackModel);
 
-		TSharedRef<FSequencerTrackNode> TrackNode = SectionHandle->GetTrackNode();
-		TrackNodes.AddUnique(TrackNode);
-
 		int32 TargetRowIndex = Section->GetRowIndex();
 
+		// @todo_sequencer_mvvm: need to go through this dragging code and figure out what's going on and what it means
+		//                       now that models and views have separated concerns
+
 		// Handle vertical dragging to re-arrange tracks. We don't support vertical rearranging if you're dragging via
 		// a key, as the built in offset causes it to always jump down a row even without moving the mouse.
-		if (Track->SupportsMultipleRows() && AllSections.Num() > 1 && bHotspotWasSection)
+		if (Track->SupportsMultipleRows() && AllSections.Num() > 1 && bAllowVerticalMovement)
 		{
 			// Compute the max row index whilst disregarding the one we're dragging
 			int32 MaxRowIndex = 0;
@@ -1189,51 +1172,54 @@
 			}
 
 			// Handle sub-track and non-sub-track dragging
-			if (TrackNode->GetSubTrackMode() == FSequencerTrackNode::ESubTrackMode::None)
+			if (TrackModel->IsA<FTrackModel>())
 			{
 				const int32 NumRows = FMath::Max(Section->GetRowIndex() + 1, MaxRowIndex);
 
 				// Find the total height of the track - this is necessary because tracks may contain key areas, but they will not use sub tracks unless there is more than one row
-				float VirtualSectionBottom = 0.f;
-				TrackNode->TraverseVisible_ParentFirst([&](FSequencerDisplayNode& Node) { VirtualSectionBottom = Node.GetVirtualBottom(); return true; }, true);
+				const FVirtualGeometry VirtualGeometry = TrackModel->CastThis<FTrackModel>()->GetVirtualGeometry();
 
 				// Assume same height rows
-				const float VirtualSectionTop = TrackNode->GetVirtualTop();
-				const float VirtualSectionHeight = VirtualSectionBottom - TrackNode->GetVirtualTop();
-
+				const float VirtualSectionHeight = VirtualGeometry.NestedBottom - VirtualGeometry.Top;
 				const float VirtualRowHeight = VirtualSectionHeight / NumRows;
-				const float MouseOffsetWithinRow = VirtualMousePos.Y - (VirtualSectionTop + (VirtualRowHeight * TargetRowIndex));
+				const float MouseOffsetWithinRow = VirtualMousePos.Y - (VirtualGeometry.Top + (VirtualRowHeight * TargetRowIndex));
 
 				if (MouseOffsetWithinRow < VirtualRowHeight || MouseOffsetWithinRow > VirtualRowHeight)
 				{
-					const int32 NewIndex = FMath::FloorToInt((VirtualMousePos.Y - VirtualSectionTop) / VirtualRowHeight);
+					const int32 NewIndex = FMath::FloorToInt((VirtualMousePos.Y - VirtualGeometry.Top) / VirtualRowHeight);
 					TargetRowIndex = FMath::Clamp(NewIndex, 0, MaxRowIndex);
 				}
 
 				// If close to the top of the row, move else everything down
-				if (VirtualMousePos.Y <= VirtualSectionTop || LocalMousePos.Y <= 0)
+				if (VirtualMousePos.Y <= VirtualGeometry.Top || LocalMousePos.Y <= 0)
 				{
 					TargetRowIndex = -1;
 				}
 			}
-			else if (TrackNode->GetSubTrackMode() == FSequencerTrackNode::ESubTrackMode::SubTrack)
+			else if (TrackModel->IsA<FTrackRowModel>())
 			{
-				TSharedPtr<FSequencerTrackNode> ParentTrack = StaticCastSharedPtr<FSequencerTrackNode>(TrackNode->GetParent());
+				TSharedPtr<FTrackModel> ParentTrack = TrackModel->FindAncestorOfType<FTrackModel>();
 				if (ensure(ParentTrack.IsValid()))
 				{
-					for (int32 ChildIndex = 0; ChildIndex < ParentTrack->GetChildNodes().Num(); ++ChildIndex)
+					int32 ChildIndex = 0;
+					for (TSharedPtr<FViewModel> ChildNode : ParentTrack->GetChildren())
 					{
-						TSharedRef<FSequencerDisplayNode> ChildNode = ParentTrack->GetChildNodes()[ChildIndex];
-						float VirtualSectionTop = ChildNode->GetVirtualTop();
-						float VirtualSectionBottom = 0.f;
-						ChildNode->TraverseVisible_ParentFirst([&](FSequencerDisplayNode& Node) { VirtualSectionBottom = Node.GetVirtualBottom(); return true; }, true);
+						FVirtualGeometry ChildVirtualGeometry;
+						if (IGeometryExtension* ChildGeometryExtension = ChildNode->CastThis<IGeometryExtension>())
+						{
+							ChildVirtualGeometry = ChildGeometryExtension->GetVirtualGeometry();
+						}
 
+						float VirtualSectionTop = ChildVirtualGeometry.Top;
+						float VirtualSectionBottom = ChildVirtualGeometry.NestedBottom;
+
 						if (ChildIndex == 0 && (VirtualMousePos.Y <= VirtualSectionTop || LocalMousePos.Y <= 0))
 						{
 							TargetRowIndex = 0;
-							for (TSharedRef<ISequencerSection> TrackSection : TrackNode->GetSections())
+							FTrackRowModel* TrackRowModel = TrackModel->CastThis<FTrackRowModel>();
+							for (TSharedPtr<FSectionModel> SectionNode : TrackRowModel->GetTrackAreaModelListAs<FSectionModel>())
 							{
-								if (!Sections.Contains(TrackSection->GetSectionObject()))
+								if (!Sections.Contains(SectionNode->GetSection()))
 								{
 									TargetRowIndex = -1;
 									break;
@@ -1250,6 +1236,8 @@
 						{
 							TargetRowIndex = ChildIndex + 1;
 						}
+
+						++ChildIndex;
 					}
 				
 					// Track if we're expanding a parent track so we can unexpand it if we stop targeting it
@@ -1257,20 +1245,20 @@
 					{
 						if (!ParentTrack->IsExpanded() && ParentTrack != ExpandedParentTrack)
 						{
-							if (ExpandedParentTrack.IsValid())
+							if (TSharedPtr<FTrackModel> ExpandedParentTrackPinned = ExpandedParentTrack.Pin())
 							{
-								ExpandedParentTrack->SetExpansionState(false);
+								ExpandedParentTrackPinned->SetExpansion(false);
 								ExpandedParentTrack = nullptr;
 							}
 							ExpandedParentTrack = ParentTrack;
-							ParentTrack->SetExpansionState(true);
+							ParentTrack->SetExpansion(true);
 						}
 					}
-					else if (ExpandedParentTrack.IsValid())
+					else if (TSharedPtr<FTrackModel> ExpandedParentTrackPinned = ExpandedParentTrack.Pin())
 					{
-						ExpandedParentTrack->SetExpansionState(false);
+						ExpandedParentTrackPinned->SetExpansion(false);
 						ExpandedParentTrack = nullptr;
-					}				
+					}
 				}
 			}
 		}
@@ -1352,27 +1340,30 @@
 	{
 		PrevMousePosY = LocalMousePos.Y;
 
+		// @todo_sequence_mvvm: Fix this code as well - is it necessary if we no longer destructively re-populate the tree when data changes?
 		// Expand track node if it wasn't already expanded. This ensures that multi row tracks will show multiple rows if regenerated
-		for (TSharedRef<FSequencerTrackNode> TrackNode : TrackNodes)
-		{
-			if (!TrackNode->IsExpanded() && TrackNode->GetSubTrackMode() == FSequencerTrackNode::ESubTrackMode::None)
-			{
-				TArray<TSharedRef<ISequencerSection> > TrackNodeSections = TrackNode->GetSections();
-				if (TrackNodeSections.Num() && TrackNodeSections[0]->GetSectionObject())
-				{
-					int32 SectionFirstRowIndex = TrackNodeSections[0]->GetSectionObject()->GetRowIndex();
+		//for (TSharedPtr<FViewModel> TrackModel : Tracks)
+		//{
+		//	IOutlinerExtension* OutlinerExtension = TrackModel->CastThis<IOutlinerExtension>(TrackModel);
+		//	bool bIsTrackExpanded = (!OutlinerExtension || OutlinerExtension->IsExpanded());
+		//	if (!bIsTrackExpanded && TrackModel->GetSubTrackMode() == FSequencerTrackNode::ESubTrackMode::None)
+		//	{
+		//		TArray<TSharedRef<ISequencerSection> > TrackNodeSections = TrackModel->GetSections();
+		//		if (TrackNodeSections.Num() && TrackNodeSections[0]->GetSectionObject())
+		//		{
+		//			int32 SectionFirstRowIndex = TrackNodeSections[0]->GetSectionObject()->GetRowIndex();
 
-					for (TSharedRef<ISequencerSection> TrackNodeSection : TrackNodeSections)
-					{
-						if (TrackNodeSection->GetSectionObject() && SectionFirstRowIndex != TrackNodeSection->GetSectionObject()->GetRowIndex())
-						{
-							TrackNode->SetExpansionState(true);
-							break;
-						}
-					}
-				}
-			}
-		}
+		//			for (TSharedRef<ISequencerSection> TrackNodeSection : TrackNodeSections)
+		//			{
+		//				if (TrackNodeSection->GetSectionObject() && SectionFirstRowIndex != TrackNodeSection->GetSectionObject()->GetRowIndex())
+		//				{
+		//					TrackNode->SetExpansion(true);
+		//					break;
+		//				}
+		//			}
+		//		}
+		//	}
+		//}
 	}
 
 	return bRowIndexChanged;
@@ -1444,11 +1435,13 @@
 		}
 	}
 
+	// Explicitly mark everything as changed to ensure that the UI is responsive during a drag
 	for (TWeakObjectPtr<UMovieSceneSection> Section : ModifiedNonSelectedSections)
 	{
 		if (Section.Get())
 		{
 			Section.Get()->MarkAsChanged();
+			Section.Get()->BroadcastChanged();
 		}
 	}
 }
diff --git a/Engine/Source/Editor/Sequencer/Private/Tools/EditToolDragOperations.h b/Engine/Source/Editor/Sequencer/Private/Tools/EditToolDragOperations.h
--- a/Engine/Source/Editor/Sequencer/Private/Tools/EditToolDragOperations.h
+++ b/Engine/Source/Editor/Sequencer/Private/Tools/EditToolDragOperations.h
@@ -11,18 +11,40 @@
 #include "ScopedTransaction.h"
 #include "Tools/SequencerSnapField.h"
 #include "Channels/MovieSceneChannelHandle.h"
+#include "MVVM/Extensions/IDraggableTrackAreaExtension.h"
 
 
 class FSequencer;
 class FSlateWindowElementList;
-class FVirtualTrackArea;
 class USequencerSettings;
 
+namespace UE
+{
+namespace Sequencer
+{
+
+class IDraggableTrackAreaExtension;
+class FSectionModel;
+class FTrackModel;
+class FVirtualTrackArea;
+
+} // namespace Sequencer
+} // namespace UE
+
+
+
+enum class ESequencerMoveOperationType
+{
+	MoveKeys     = 1<<0,
+	MoveSections = 1<<1,
+};
+ENUM_CLASS_FLAGS(ESequencerMoveOperationType)
+
 /**
  * Abstract base class for drag operations that handle an operation for an edit tool.
  */
 class FEditToolDragOperation
-	: public ISequencerEditToolDragOperation
+	: public UE::Sequencer::ISequencerEditToolDragOperation
 {
 public:
 
@@ -39,7 +61,7 @@
 protected:
 
 	/** begin a new scoped transaction for this drag */
-	void BeginTransaction( TArray< TWeakObjectPtr<UMovieSceneSection> >& Sections, const FText& TransactionDesc );
+	void BeginTransaction( TSet<UMovieSceneSection*>& Sections, const FText& TransactionDesc );
 
 	/** End an existing scoped transaction if one exists */
 	void EndTransaction();
@@ -72,15 +94,15 @@
 
 	// FEditToolDragOperation interface
 
-	virtual void OnBeginDrag(const FPointerEvent& MouseEvent, FVector2D LocalMousePos, const FVirtualTrackArea& VirtualTrackArea) override;
-	virtual void OnDrag(const FPointerEvent& MouseEvent, FVector2D LocalMousePos, const FVirtualTrackArea& VirtualTrackArea) override;
-	virtual void OnEndDrag(const FPointerEvent& MouseEvent, FVector2D LocalMousePos, const FVirtualTrackArea& VirtualTrackArea) override;
+	virtual void OnBeginDrag(const FPointerEvent& MouseEvent, FVector2D LocalMousePos, const UE::Sequencer::FVirtualTrackArea& VirtualTrackArea) override;
+	virtual void OnDrag(const FPointerEvent& MouseEvent, FVector2D LocalMousePos, const UE::Sequencer::FVirtualTrackArea& VirtualTrackArea) override;
+	virtual void OnEndDrag(const FPointerEvent& MouseEvent, FVector2D LocalMousePos, const UE::Sequencer::FVirtualTrackArea& VirtualTrackArea) override;
 	virtual FCursorReply GetCursor() const override { return FCursorReply::Cursor( EMouseCursor::ResizeLeftRight ); }
 
 private:
 
 	/** The sections we are interacting with */
-	TArray<TWeakObjectPtr<UMovieSceneSection>> Sections;
+	TSet<UMovieSceneSection*> Sections;
 
 	/********************************************************/
 	struct FPreDragChannelData
@@ -129,17 +151,23 @@
  */
 class FMoveKeysAndSections
 	: public FEditToolDragOperation
+	, public UE::Sequencer::IDragOperation
 {
 public:
-	FMoveKeysAndSections(FSequencer& InSequencer, const TSet<FSequencerSelectedKey>& InSelectedKeys, const TSet<TWeakObjectPtr<UMovieSceneSection>>& InSelectedSections, bool InbHotspotWasSection);
+
+	FMoveKeysAndSections(FSequencer& InSequencer, ESequencerMoveOperationType InMoveType);
 
 	// FEditToolDragOperation interface
-	virtual void OnBeginDrag(const FPointerEvent& MouseEvent, FVector2D LocalMousePos, const FVirtualTrackArea& VirtualTrackArea) override;
-	virtual void OnDrag(const FPointerEvent& MouseEvent, FVector2D LocalMousePos, const FVirtualTrackArea& VirtualTrackArea) override;
-	virtual void OnEndDrag(const FPointerEvent& MouseEvent, FVector2D LocalMousePos, const FVirtualTrackArea& VirtualTrackArea) override;
+	virtual void OnBeginDrag(const FPointerEvent& MouseEvent, FVector2D LocalMousePos, const UE::Sequencer::FVirtualTrackArea& VirtualTrackArea) override;
+	virtual void OnDrag(const FPointerEvent& MouseEvent, FVector2D LocalMousePos, const UE::Sequencer::FVirtualTrackArea& VirtualTrackArea) override;
+	virtual void OnEndDrag(const FPointerEvent& MouseEvent, FVector2D LocalMousePos, const UE::Sequencer::FVirtualTrackArea& VirtualTrackArea) override;
 	virtual FCursorReply GetCursor() const override { return FCursorReply::Cursor(EMouseCursor::CardinalCross); }
 	// ~FEditToolDragOperation interface
 
+	/* UE::Sequencer:IDragOperation Interface */
+	virtual void AddSnapTime(FFrameNumber SnapTime) override;
+	virtual void AddModel(TSharedPtr<UE::Sequencer::FViewModel> Model) override;
+
 protected:
 	/** Calculate the possible horizontal movement we can, constrained by sections running into things. */
 	TOptional<FFrameNumber> GetMovementDeltaX(FFrameTime MouseTime);
@@ -158,8 +186,11 @@
 	void ModifyNonSelectedSections();
 
 protected:
+	/** Array of models that we're moving. */
+	TSet<TWeakPtr<UE::Sequencer::IDraggableTrackAreaExtension>> DraggedItems;
+
 	/** Array of sections that we're moving. */
-	TArray<TWeakObjectPtr<UMovieSceneSection>> Sections;
+	TSet<UMovieSceneSection*> Sections;
 
 	/** Set of keys that are being moved. */
 	TSet<FSequencerSelectedKey> Keys;
@@ -171,28 +202,8 @@
 	/** The position of the mouse when the last section move occurred */
 	TOptional<float> PrevMousePosY;
 
-	struct FRelativeOffset
-	{
-		FRelativeOffset()
-			: StartOffset()
-			, EndOffset()
-		{
-		}
-
-		/**
-		 * The offset for the start of the section. Can be unset in the case of a section with no lower bound.
-		 * Keys are represented only by StartOffset and do not have an End Offset (which would imply a range).
-		 */
-		TOptional<FFrameTime> StartOffset;
-
-		/**
-		 * The offset for the end of the section. Can be unset in the case of a section with no upper bound.
-		 */
-		TOptional<FFrameTime> EndOffset;
-	};
-
 	/** Array of relative offsets for each selected item. Keys + Sections are both added to this array. */
-	TArray<FRelativeOffset> RelativeOffsets;
+	TArray<FFrameNumber> RelativeSnapOffsets;
 
 	struct FInitialRowIndex
 	{
@@ -210,10 +221,10 @@
 	TOptional<FSequencerSnapField> SnapField;
 
 	/** If we expanded a parent track while dragging, track it here so we can re-collapse it if not dropping on it. */
-	TSharedPtr<FSequencerTrackNode> ExpandedParentTrack;
+	TWeakPtr<UE::Sequencer::FTrackModel> ExpandedParentTrack;
 
 	/** If the user is moving them via clicking on the Section then we'll allow vertical re-arranging, otherwise not. */
-	bool bHotspotWasSection;
+	bool bAllowVerticalMovement;
 };
 
 /**
@@ -223,14 +234,14 @@
 {
 public:
 
-	FDuplicateKeysAndSections( FSequencer& InSequencer, const TSet<FSequencerSelectedKey>& InSelectedKeys, const TSet<TWeakObjectPtr<UMovieSceneSection>>& InSelectedSections, bool InbHotspotWasSection)
-		: FMoveKeysAndSections(InSequencer, InSelectedKeys, InSelectedSections, InbHotspotWasSection)
+	FDuplicateKeysAndSections( FSequencer& InSequencer, ESequencerMoveOperationType Type)
+		: FMoveKeysAndSections(InSequencer, Type)
 	{}
 
 public:
 
-	virtual void OnBeginDrag(const FPointerEvent& MouseEvent, FVector2D LocalMousePos, const FVirtualTrackArea& VirtualTrackArea) override;
-	virtual void OnEndDrag(const FPointerEvent& MouseEvent, FVector2D LocalMousePos, const FVirtualTrackArea& VirtualTrackArea) override;
+	virtual void OnBeginDrag(const FPointerEvent& MouseEvent, FVector2D LocalMousePos, const UE::Sequencer::FVirtualTrackArea& VirtualTrackArea) override;
+	virtual void OnEndDrag(const FPointerEvent& MouseEvent, FVector2D LocalMousePos, const UE::Sequencer::FVirtualTrackArea& VirtualTrackArea) override;
 };
 
 /**
@@ -248,9 +259,9 @@
 
 	// FEditToolDragOperation interface
 
-	virtual void OnBeginDrag(const FPointerEvent& MouseEvent, FVector2D LocalMousePos, const FVirtualTrackArea& VirtualTrackArea) override;
-	virtual void OnDrag(const FPointerEvent& MouseEvent, FVector2D LocalMousePos, const FVirtualTrackArea& VirtualTrackArea) override;
-	virtual void OnEndDrag(const FPointerEvent& MouseEvent, FVector2D LocalMousePos, const FVirtualTrackArea& VirtualTrackArea) override;
+	virtual void OnBeginDrag(const FPointerEvent& MouseEvent, FVector2D LocalMousePos, const UE::Sequencer::FVirtualTrackArea& VirtualTrackArea) override;
+	virtual void OnDrag(const FPointerEvent& MouseEvent, FVector2D LocalMousePos, const UE::Sequencer::FVirtualTrackArea& VirtualTrackArea) override;
+	virtual void OnEndDrag(const FPointerEvent& MouseEvent, FVector2D LocalMousePos, const UE::Sequencer::FVirtualTrackArea& VirtualTrackArea) override;
 	virtual FCursorReply GetCursor() const override { return FCursorReply::Cursor( EMouseCursor::ResizeLeftRight ); }
 
 private:
diff --git a/Engine/Source/Editor/Sequencer/Private/Tools/SequencerEditTool.h b/Engine/Source/Editor/Sequencer/Private/Tools/SequencerEditTool.h
--- a/Engine/Source/Editor/Sequencer/Private/Tools/SequencerEditTool.h
+++ b/Engine/Source/Editor/Sequencer/Private/Tools/SequencerEditTool.h
@@ -5,8 +5,9 @@
 #include "CoreMinimal.h"
 #include "Input/CursorReply.h"
 #include "Input/Reply.h"
-#include "Sequencer.h"
+#include "Framework/DelayedDrag.h"
 #include "ISequencerEditTool.h"
+#include "Sequencer.h"
 
 class FSlateWindowElementList;
 
@@ -14,7 +15,7 @@
  * Abstract base class for edit tools.
  */
 class FSequencerEditTool
-	: public ISequencerEditTool
+	: public UE::Sequencer::ISequencerEditTool
 {
 public:
 
@@ -74,7 +75,7 @@
 		return FCursorReply::Unhandled();
 	}
 
-	virtual ISequencer& GetSequencer() const override
+	ISequencer& GetSequencer() const
 	{
 		return Sequencer;
 	}
@@ -84,3 +85,24 @@
 	/** This edit tool's sequencer */
 	FSequencer& Sequencer;
 };
+
+namespace UE
+{
+namespace Sequencer
+{
+
+struct FDelayedDrag_Hotspot : FDelayedDrag
+{
+	FDelayedDrag_Hotspot(FVector2D InInitialPosition, FKey InApplicableKey, TSharedPtr<ITrackAreaHotspot> InHotspot)
+		: FDelayedDrag(InInitialPosition, InApplicableKey)
+		, Hotspot(MoveTemp(InHotspot))
+	{
+		SetTriggerScaleFactor(0.1f);
+	}
+
+	TSharedPtr<ITrackAreaHotspot> Hotspot;
+};
+
+
+} // namespace Sequencer
+} // namespace UE
diff --git a/Engine/Source/Editor/Sequencer/Private/Tools/SequencerEditTool_Movement.cpp b/Engine/Source/Editor/Sequencer/Private/Tools/SequencerEditTool_Movement.cpp
--- a/Engine/Source/Editor/Sequencer/Private/Tools/SequencerEditTool_Movement.cpp
+++ b/Engine/Source/Editor/Sequencer/Private/Tools/SequencerEditTool_Movement.cpp
@@ -1,15 +1,17 @@
 // Copyright Epic Games, Inc. All Rights Reserved.
 
 #include "Tools/SequencerEditTool_Movement.h"
+#include "MVVM/ViewModels/SectionModel.h"
+#include "MVVM/ViewModels/SequencerEditorViewModel.h"
 #include "Editor.h"
 #include "Fonts/FontMeasure.h"
 #include "Styling/CoreStyle.h"
 #include "Styling/AppStyle.h"
 #include "SequencerCommonHelpers.h"
 #include "SSequencer.h"
-#include "ISequencerHotspot.h"
+#include "MVVM/Views/ITrackAreaHotspot.h"
 #include "SequencerHotspots.h"
-#include "VirtualTrackArea.h"
+#include "MVVM/ViewModels/VirtualTrackArea.h"
 #include "SequencerSettings.h"
 #include "Tools/EditToolDragOperations.h"
 #include "IKeyArea.h"
@@ -29,9 +31,11 @@
 
 FReply FSequencerEditTool_Movement::OnMouseButtonDown(SWidget& OwnerWidget, const FGeometry& MyGeometry, const FPointerEvent& MouseEvent)
 {
+	using namespace UE::Sequencer;
+
 	TSharedRef<SSequencer> SequencerWidget = StaticCastSharedRef<SSequencer>(Sequencer.GetSequencerWidget());
 
-	TSharedPtr<ISequencerHotspot> Hotspot = Sequencer.GetHotspot();
+	TSharedPtr<ITrackAreaHotspot> Hotspot = Sequencer.GetViewModel()->GetTrackArea()->GetHotspot();
 
 	DelayedDrag.Reset();
 
@@ -43,15 +47,11 @@
 
  		if (Sequencer.GetSequencerSettings()->GetSnapPlayTimeToPressedKey() || (MouseEvent.IsShiftDown() && MouseEvent.GetEffectingButton() == EKeys::LeftMouseButton) )
 		{
-			if (DelayedDrag->Hotspot.IsValid())
+			if (TSharedPtr<FKeyHotspot> KeyHotspot = HotspotCast<FKeyHotspot>(DelayedDrag->Hotspot))
 			{
-				if (DelayedDrag->Hotspot->GetType() == ESequencerHotspot::Key)
+				if (TOptional<FFrameNumber> Time = KeyHotspot->GetTime())
 				{
-					TOptional<FFrameNumber> Time = StaticCastSharedPtr<FKeyHotspot>(DelayedDrag->Hotspot)->GetTime();
-					if (Time.IsSet())
-					{
-						Sequencer.SetLocalTime(Time.GetValue());
-					}
+					Sequencer.SetLocalTime(Time.GetValue());
 				}
 			}
 		}
@@ -69,6 +69,8 @@
 
 FReply FSequencerEditTool_Movement::OnMouseMove(SWidget& OwnerWidget, const FGeometry& MyGeometry, const FPointerEvent& MouseEvent)
 {
+	using namespace UE::Sequencer;
+
 	if (DelayedDrag.IsSet())
 	{
 		TSharedRef<SSequencer> SequencerWidget = StaticCastSharedRef<SSequencer>(Sequencer.GetSequencerWidget());
@@ -83,14 +85,6 @@
 			{
 				DragPosition = MyGeometry.AbsoluteToLocal(MouseEvent.GetScreenSpacePosition());
 
-				if (Sequencer.GetSequencerSettings()->GetIsSnapEnabled() && Sequencer.GetSequencerSettings()->GetSnapKeysAndSectionsToPlayRange() && !Sequencer.GetSequencerSettings()->ShouldKeepPlayRangeInSectionBounds())
-				{
-					DragPosition.X = FMath::Max(DragPosition.X, 0.f);
-					FFrameTime CurrentTime = VirtualTrackArea.PixelToFrame(DragPosition.X);
-					CurrentTime = UE::MovieScene::ClampToDiscreteRange(CurrentTime, Sequencer.GetPlaybackRange());
-					DragPosition.X = VirtualTrackArea.FrameToPixel(CurrentTime);
-				}
-					
 				double CurrentTime = VirtualTrackArea.PixelToSeconds(DragPosition.X);
 				Sequencer.UpdateAutoScroll(CurrentTime);
 
@@ -149,8 +143,10 @@
 	return CurrentTime - OriginalHotspotTime;
 }
 
-TSharedPtr<ISequencerEditToolDragOperation> FSequencerEditTool_Movement::CreateDrag(const FPointerEvent& MouseEvent)
+TSharedPtr<UE::Sequencer::ISequencerEditToolDragOperation> FSequencerEditTool_Movement::CreateDrag(const FPointerEvent& MouseEvent)
 {
+	using namespace UE::Sequencer;
+
 	FSequencerSelection& Selection = Sequencer.GetSelection();
 	TSharedRef<SSequencer> SequencerWidget = StaticCastSharedRef<SSequencer>(Sequencer.GetSequencerWidget());
 
@@ -159,30 +155,32 @@
 	if (DelayedDrag->Hotspot.IsValid())
 	{
 		// Let the hotspot start a drag first, if it wants to
-		auto HotspotDrag = DelayedDrag->Hotspot->InitiateDrag(Sequencer);
+		TSharedPtr<ISequencerEditToolDragOperation> HotspotDrag = DelayedDrag->Hotspot->InitiateDrag(MouseEvent);
 		if (HotspotDrag.IsValid())
 		{
 			return HotspotDrag;
 		}
-		auto HotspotType = DelayedDrag->Hotspot->GetType();
 
-		const bool bSectionsSelected = Selection.GetSelectedSections().Num() > 0;
+		// Gather the draggable sections from all selected models
+		const bool bModelsSelected = Selection.GetSelectedTrackAreaItems().Num() > 0;
 		const bool bKeySelected = Selection.GetSelectedKeys().Num() > 0;
 		// @todo sequencer: Make this a customizable UI command modifier?
 		const bool bIsDuplicateEvent = MouseEvent.IsAltDown() || MouseEvent.GetEffectingButton() == EKeys::MiddleMouseButton;
-		const bool bHotspotIsSection = HotspotType == ESequencerHotspot::Section;
 
+		FSectionHotspotBase* SectionHotspot = DelayedDrag->Hotspot->CastThis<FSectionHotspotBase>();
+		FKeyHotspot*         KeyHotspot     = DelayedDrag->Hotspot->CastThis<FKeyHotspot>();
+
 		// If they have both keys and sections selected then we only support moving them right now, so we
 		// check for that first before trying to figure out if they're resizing or dilating.
-		if (bSectionsSelected && bKeySelected && !bIsDuplicateEvent)
+		if (bModelsSelected && bKeySelected && !bIsDuplicateEvent)
 		{
-			return MakeShareable(new FMoveKeysAndSections(Sequencer, Selection.GetSelectedKeys(), Selection.GetSelectedSections(), bHotspotIsSection));
+			return MakeShareable(new FMoveKeysAndSections(Sequencer, ESequencerMoveOperationType::MoveKeys | ESequencerMoveOperationType::MoveSections));
 		}
 		else if (bIsDuplicateEvent)
 		{
-			if (HotspotType == ESequencerHotspot::Key)
+			if (KeyHotspot)
 			{
-				TArrayView<const FSequencerSelectedKey> HoveredKeys = StaticCastSharedPtr<FKeyHotspot>(DelayedDrag->Hotspot)->Keys;
+				TArrayView<const FSequencerSelectedKey> HoveredKeys = KeyHotspot->Keys;
 
 				auto AnyUnselectedKey = [&Selection](const FSequencerSelectedKey& InKey)
 				{
@@ -193,49 +191,34 @@
 				{
 					// If any are not selected, we'll treat this as a unique drag
 					Selection.EmptySelectedKeys();
-					Selection.EmptySelectedSections();
-					Selection.EmptyNodesWithSelectedKeysOrSections();
+					Selection.EmptySelectedTrackAreaItems();
 					for (const FSequencerSelectedKey& Key : HoveredKeys)
 					{
 						Selection.AddToSelection(Key);
 					}
-					SequencerHelpers::UpdateHoveredNodeFromSelectedKeys(Sequencer);
 				}
 			}
-			else if (HotspotType == ESequencerHotspot::Section)
+			else if (SectionHotspot)
 			{
-				UMovieSceneSection* HoveredSection = StaticCastSharedPtr<FSectionHotspot>(DelayedDrag->Hotspot)->WeakSection.Get();
-
-				if (!Selection.IsSelected(HoveredSection))
+				if (!Selection.IsSelected(SectionHotspot->WeakSectionModel.Pin()))
 				{
 					Selection.EmptySelectedKeys();
-					Selection.EmptySelectedSections();
-					Selection.EmptyNodesWithSelectedKeysOrSections();
-					Selection.AddToSelection(HoveredSection);
-					SequencerHelpers::UpdateHoveredNodeFromSelectedSections(Sequencer);
+					Selection.EmptySelectedTrackAreaItems();
+					Selection.AddToSelection(SectionHotspot->WeakSectionModel.Pin());
 				}
 			}
 
-			return MakeShareable(new FDuplicateKeysAndSections(Sequencer, Selection.GetSelectedKeys(), Selection.GetSelectedSections(), bHotspotIsSection));
+			return MakeShareable(new FDuplicateKeysAndSections(Sequencer, ESequencerMoveOperationType::MoveKeys | ESequencerMoveOperationType::MoveSections));
 		}
 
-
-		UMovieSceneSection* SectionToDrag = nullptr;
-		if (HotspotType == ESequencerHotspot::Section || HotspotType == ESequencerHotspot::EasingArea)
-		{
-			SectionToDrag = StaticCastSharedPtr<FSectionHotspot>(DelayedDrag->Hotspot)->WeakSection.Get();
-		}
-
 		// Moving section(s)?
-		if (SectionToDrag)
+		if (SectionHotspot)
 		{
-			if (!Selection.IsSelected(SectionToDrag))
+			if (!Selection.IsSelected(SectionHotspot->WeakSectionModel.Pin()))
 			{
 				Selection.EmptySelectedKeys();
-				Selection.EmptySelectedSections();
-				Selection.EmptyNodesWithSelectedKeysOrSections();
-				Selection.AddToSelection(SectionToDrag);
-				SequencerHelpers::UpdateHoveredNodeFromSelectedSections(Sequencer);
+				Selection.EmptySelectedTrackAreaItems();
+				Selection.AddToSelection(SectionHotspot->WeakSectionModel.Pin());
 			}
 
 			if (MouseEvent.IsShiftDown())
@@ -246,14 +229,13 @@
 			}
 			else
 			{
-				TSet<FSequencerSelectedKey> EmptyKeySet;
-				return MakeShareable( new FMoveKeysAndSections( Sequencer, EmptyKeySet, Selection.GetSelectedSections(), true) );
+				return MakeShareable( new FMoveKeysAndSections( Sequencer, ESequencerMoveOperationType::MoveSections) );
 			}
 		}
 		// Moving key(s)?
-		else if (HotspotType == ESequencerHotspot::Key)
+		else if (KeyHotspot)
 		{
-			TArrayView<const FSequencerSelectedKey> HoveredKeys = StaticCastSharedPtr<FKeyHotspot>(DelayedDrag->Hotspot)->Keys;
+			TArrayView<const FSequencerSelectedKey> HoveredKeys = KeyHotspot->Keys;
 
 			auto AnyUnselectedKey = [&Selection](const FSequencerSelectedKey& InKey)
 			{
@@ -264,29 +246,25 @@
 			{
 				// If any are not selected, we'll treat this as a unique drag
 				Selection.EmptySelectedKeys();
-				Selection.EmptySelectedSections();
-				Selection.EmptyNodesWithSelectedKeysOrSections();
+				Selection.EmptySelectedTrackAreaItems();
 				for (const FSequencerSelectedKey& Key : HoveredKeys)
 				{
 					Selection.AddToSelection(Key);
 				}
-				SequencerHelpers::UpdateHoveredNodeFromSelectedKeys(Sequencer);
 			}
 
 			TSet<TWeakObjectPtr<UMovieSceneSection>> NoSections;
-			return MakeShareable( new FMoveKeysAndSections( Sequencer, Selection.GetSelectedKeys(), NoSections, false) );
+			return MakeShareable( new FMoveKeysAndSections( Sequencer, ESequencerMoveOperationType::MoveKeys) );
 		}
 	}
 	// If we're not dragging a hotspot, sections take precedence over keys
-	else if (Selection.GetSelectedSections().Num())
+	else if (Selection.GetSelectedTrackAreaItems().Num())
 	{
-		TSet<FSequencerSelectedKey> EmptyKeySet;
-		return MakeShareable( new FMoveKeysAndSections( Sequencer, EmptyKeySet, Selection.GetSelectedSections(), true ) );
+		return MakeShareable( new FMoveKeysAndSections( Sequencer, ESequencerMoveOperationType::MoveSections ) );
 	}
 	else if (Selection.GetSelectedKeys().Num())
 	{
-		TSet<TWeakObjectPtr<UMovieSceneSection>> NoSections;
-		return MakeShareable( new FMoveKeysAndSections( Sequencer, Selection.GetSelectedKeys(), NoSections, false) );
+		return MakeShareable( new FMoveKeysAndSections( Sequencer, ESequencerMoveOperationType::MoveKeys) );
 	}
 
 	return nullptr;
@@ -295,8 +273,12 @@
 
 FReply FSequencerEditTool_Movement::OnMouseButtonUp(SWidget& OwnerWidget, const FGeometry& MyGeometry, const FPointerEvent& MouseEvent)
 {
+	using namespace UE::Sequencer;
+
 	DelayedDrag.Reset();
 
+	TSharedPtr<FTrackAreaViewModel> TrackArea = Sequencer.GetViewModel()->GetTrackArea();
+
 	if (DragOperation.IsValid())
 	{
 		TSharedRef<SSequencer> SequencerWidget = StaticCastSharedRef<SSequencer>(Sequencer.GetSequencerWidget());
@@ -333,25 +315,18 @@
 				);
 
 			// Lock the hotspot while the menu is open
-			TSharedPtr<ISequencerHotspot> ExistingHotspot = Sequencer.GetHotspot();
-			if (ExistingHotspot.IsValid())
-			{
-				ExistingHotspot->bIsLocked = true;
-			}
+			TrackArea->LockHotspot(true);
 
 			// Unlock and reset the hotspot when the menu closes
 			{
-				FSequencer* SequencerPtr = &Sequencer;
+				TSharedPtr<ITrackAreaHotspot> ExistingHotspot = TrackArea->GetHotspot();
 				Menu->GetOnMenuDismissed().AddLambda(
-					[=](TSharedRef<IMenu>)
+					[TrackArea, ExistingHotspot](TSharedRef<IMenu>)
 					{
-						if (ExistingHotspot.IsValid())
+						TrackArea->LockHotspot(false);
+						if (TrackArea->GetHotspot() == ExistingHotspot)
 						{
-							ExistingHotspot->bIsLocked = false;
-						}
-						if (SequencerPtr->GetHotspot() == ExistingHotspot)
-						{
-							SequencerPtr->SetHotspot(nullptr);
+							TrackArea->SetHotspot(nullptr);
 						}
 					}
 				);
@@ -383,6 +358,8 @@
 
 int32 FSequencerEditTool_Movement::OnPaint(const FGeometry& AllottedGeometry, const FSlateRect& MyCullingRect, FSlateWindowElementList& OutDrawElements, int32 LayerId) const
 {
+	using namespace UE::Sequencer;
+
 	if (CursorDecorator)
 	{
 		FSlateDrawElement::MakeBox(
@@ -395,7 +372,7 @@
 
 	if (DelayedDrag.IsSet() && DelayedDrag->IsDragging())
 	{
-		const TSharedPtr<ISequencerHotspot>& Hotspot = DelayedDrag->Hotspot;
+		const TSharedPtr<ITrackAreaHotspot>& Hotspot = DelayedDrag->Hotspot;
 
 		if (Hotspot.IsValid())
 		{
@@ -505,11 +482,13 @@
 
 void FSequencerEditTool_Movement::UpdateCursorDecorator(const FGeometry& MyGeometry, const FPointerEvent& CursorEvent)
 {
+	using namespace UE::Sequencer;
+
 	DragPosition = MyGeometry.AbsoluteToLocal(CursorEvent.GetScreenSpacePosition());
-	
-	TSharedPtr<ISequencerHotspot> Hotspot = DelayedDrag.IsSet()
+	TSharedPtr<FTrackAreaViewModel> TrackArea = Sequencer.GetViewModel()->GetTrackArea();
+	TSharedPtr<ITrackAreaHotspot> Hotspot = DelayedDrag.IsSet()
 		? DelayedDrag->Hotspot
-		: Sequencer.GetHotspot();
+		: TrackArea->GetHotspot();
 
 	if (Hotspot.IsValid())
 	{
@@ -519,9 +498,12 @@
 
 FCursorReply FSequencerEditTool_Movement::OnCursorQuery(const FGeometry& MyGeometry, const FPointerEvent& CursorEvent) const
 {
-	TSharedPtr<ISequencerHotspot> Hotspot = DelayedDrag.IsSet()
+	using namespace UE::Sequencer;
+
+	TSharedPtr<FTrackAreaViewModel> TrackArea = Sequencer.GetViewModel()->GetTrackArea();
+	TSharedPtr<ITrackAreaHotspot> Hotspot = DelayedDrag.IsSet()
 		? DelayedDrag->Hotspot
-		: Sequencer.GetHotspot();
+		: TrackArea->GetHotspot();
 
 	if (Hotspot.IsValid())
 	{
@@ -603,7 +585,7 @@
 	return FString();
 }
 
-const ISequencerHotspot* FSequencerEditTool_Movement::GetDragHotspot() const
+TSharedPtr<UE::Sequencer::ITrackAreaHotspot> FSequencerEditTool_Movement::GetDragHotspot() const
 {
-	return DelayedDrag.IsSet() ? DelayedDrag->Hotspot.Get() : nullptr;
+	return DelayedDrag.IsSet() ? DelayedDrag->Hotspot : nullptr;
 }
diff --git a/Engine/Source/Editor/Sequencer/Private/Tools/SequencerEditTool_Movement.h b/Engine/Source/Editor/Sequencer/Private/Tools/SequencerEditTool_Movement.h
--- a/Engine/Source/Editor/Sequencer/Private/Tools/SequencerEditTool_Movement.h
+++ b/Engine/Source/Editor/Sequencer/Private/Tools/SequencerEditTool_Movement.h
@@ -7,10 +7,8 @@
 #include "Input/CursorReply.h"
 #include "Input/Reply.h"
 #include "Tools/SequencerEditTool.h"
-#include "Framework/DelayedDrag.h"
 
 class SSequencer;
-struct ISequencerHotspot;
 
 class FSequencerEditTool_Movement
 	: public FSequencerEditTool
@@ -35,7 +33,7 @@
 	virtual FCursorReply OnCursorQuery(const FGeometry& MyGeometry, const FPointerEvent& CursorEvent) const override;
 	virtual FName GetIdentifier() const override;
 	virtual bool CanDeactivate() const override;
-	virtual const ISequencerHotspot* GetDragHotspot() const override;
+	virtual TSharedPtr<UE::Sequencer::ITrackAreaHotspot> GetDragHotspot() const override;
 
 protected:
 
@@ -43,29 +41,17 @@
 
 private:
 
-	TSharedPtr<ISequencerEditToolDragOperation> CreateDrag(const FPointerEvent& MouseEvent);
+	TSharedPtr<UE::Sequencer::ISequencerEditToolDragOperation> CreateDrag(const FPointerEvent& MouseEvent);
 
 	bool GetHotspotTime(FFrameTime& HotspotTime) const;
 	FFrameTime GetHotspotOffsetTime(FFrameTime CurrentTime) const;
 	void UpdateCursorDecorator(const FGeometry& MyGeometry, const FPointerEvent& MouseEvent);
 
-	struct FDelayedDrag_Hotspot : FDelayedDrag
-	{
-		FDelayedDrag_Hotspot(FVector2D InInitialPosition, FKey InApplicableKey, TSharedPtr<ISequencerHotspot> InHotspot)
-			: FDelayedDrag(InInitialPosition, InApplicableKey)
-			, Hotspot(MoveTemp(InHotspot))
-		{
-			SetTriggerScaleFactor(0.1f);
-		}
-
-		TSharedPtr<ISequencerHotspot> Hotspot;
-	};
-
 	/** Helper class responsible for handling delayed dragging */
-	TOptional<FDelayedDrag_Hotspot> DelayedDrag;
+	TOptional<UE::Sequencer::FDelayedDrag_Hotspot> DelayedDrag;
 
 	/** Current drag operation if any */
-	TSharedPtr<ISequencerEditToolDragOperation> DragOperation;
+	TSharedPtr<UE::Sequencer::ISequencerEditToolDragOperation> DragOperation;
 
 	/** Current local position the mouse is dragged to. */
 	FVector2D DragPosition;
diff --git a/Engine/Source/Editor/Sequencer/Private/Tools/SequencerEditTool_Selection.cpp b/Engine/Source/Editor/Sequencer/Private/Tools/SequencerEditTool_Selection.cpp
--- a/Engine/Source/Editor/Sequencer/Private/Tools/SequencerEditTool_Selection.cpp
+++ b/Engine/Source/Editor/Sequencer/Private/Tools/SequencerEditTool_Selection.cpp
@@ -4,10 +4,21 @@
 #include "Styling/AppStyle.h"
 #include "SequencerCommonHelpers.h"
 #include "SSequencer.h"
+#include "ISequencerEditTool.h"
+#include "SequencerHotspots.h"
 #include "Tools/SequencerEntityVisitor.h"
-#include "SSequencerTreeView.h"
-#include "ISequencerHotspot.h"
-#include "VirtualTrackArea.h"
+#include "MVVM/Views/ISequencerTreeView.h"
+#include "MVVM/Views/ITrackAreaHotspot.h"
+#include "MVVM/ViewModels/SequencerEditorViewModel.h"
+#include "MVVM/ViewModels/VirtualTrackArea.h"
+
+#include "MVVM/ViewModels/ViewModel.h"
+#include "MVVM/ViewModels/SequenceModel.h"
+#include "MVVM/Views/STrackAreaView.h"
+#include "MVVM/Views/SOutlinerView.h"
+#include "MVVM/Extensions/IOutlinerExtension.h"
+#include "MVVM/Extensions/IPinnableExtension.h"
+#include "MVVM/Extensions/ISelectableExtension.h"
 
 
 struct FSelectionPreviewVisitor final
@@ -20,14 +31,17 @@
 		, bPinned(bInPinned)
 	{}
 
-	virtual void VisitKey(FKeyHandle KeyHandle, FFrameNumber KeyTime, const TSharedPtr<IKeyArea>& KeyArea, UMovieSceneSection* Section, TSharedRef<FSequencerDisplayNode> Node) const override
+	virtual void VisitKey(FKeyHandle KeyHandle, FFrameNumber KeyTime, const UE::Sequencer::TViewModelPtr<UE::Sequencer::FChannelModel>& Channel, UMovieSceneSection* Section) const override
 	{
-		if (Node->IsPinned() != bPinned)
+		using namespace UE::Sequencer;
+
+		TSharedPtr<IPinnableExtension> PinnableItem = Channel->GetLinkedOutlinerItem().ImplicitCast();
+		if (PinnableItem && PinnableItem->IsPinned() != bPinned)
 		{
 			return;
 		}
 
-		FSequencerSelectedKey Key(*Section, KeyArea, KeyHandle);
+		FSequencerSelectedKey Key(*Section, TSharedPtr<FChannelModel>(Channel), KeyHandle);
 
 		// Under default behavior keys have priority, so if a key is changing selection state then we remove any sections from the selection. The user can bypass this
 		// by holding down the control key which will allow selecting both keys and sections.
@@ -38,26 +52,23 @@
 			((bKeyIsSelected && SetStateTo == ESelectionPreviewState::NotSelected) ||
 			(!bKeyIsSelected && SetStateTo == ESelectionPreviewState::Selected)))
 		{
-			// Clear selected nodes
-			for (TSharedRef<FSequencerDisplayNode> SelectedNode : NodesSelectedBySections)
-			{
-				if (!NodesSelectedByKeys.Contains(SelectedNode))
-				{
-					SelectionPreview.SetSelectionState(SelectedNode, ESelectionPreviewState::Undefined);
-				}
-			}
-			
-			// Clear selected sections
-			SelectionPreview.EmptyDefinedSectionStates();
+			// Clear selected models
+			SelectionPreview.EmptyDefinedModelStates();
 		}
 
 		SelectionPreview.SetSelectionState(Key, SetStateTo);
-		SelectionPreview.SetSelectionState(Node, SetStateTo);
-		NodesSelectedByKeys.Add(Node);
 	}
 
-	virtual void VisitSection(UMovieSceneSection* Section, TSharedRef<FSequencerDisplayNode> Node) const
+	virtual void VisitDataModel(UE::Sequencer::FViewModel* DataModel) const override
 	{
+		using namespace UE::Sequencer;
+
+		ISelectableExtension* Selectable = DataModel->CastThis<ISelectableExtension>();
+		if (!Selectable || !EnumHasAnyFlags(Selectable->IsSelectable(), ESelectionIntent::PersistentSelection))
+		{
+			return;
+		}
+
 		// If key selection has priority then we check to see if there are any keys selected. If there are key selected, we don't add this section.
 		// Otherwise, we bypass this check
 		bool bKeySelectionHasPriority = !FSlateApplication::Get().GetModifierKeys().IsControlDown();
@@ -65,15 +76,11 @@
 
 		if (bKeyStateCheck)
 		{
-			SelectionPreview.SetSelectionState(Section, SetStateTo);
-			SelectionPreview.SetSelectionState(Node, SetStateTo);
-			NodesSelectedBySections.Add(Node);
+			SelectionPreview.SetSelectionState(DataModel->AsShared(), SetStateTo);
 		}
 	}
 
 private:
-	mutable TSet<TSharedRef<FSequencerDisplayNode>> NodesSelectedBySections;
-	mutable TSet<TSharedRef<FSequencerDisplayNode>> NodesSelectedByKeys;
 
 	FSequencerSelectionPreview& SelectionPreview;
 	FSequencerSelection& ExistingSelection;
@@ -83,11 +90,11 @@
 
 
 class FMarqueeDragOperation
-	: public ISequencerEditToolDragOperation
+	: public UE::Sequencer::ISequencerEditToolDragOperation
 {
 public:
 
-	FMarqueeDragOperation(FSequencer& InSequencer, SSequencerTrackArea& InTrackArea)
+	FMarqueeDragOperation(FSequencer& InSequencer, UE::Sequencer::STrackAreaView& InTrackArea)
 		: Sequencer(InSequencer)
 		, TrackArea(InTrackArea)
 		, SequencerWidget(StaticCastSharedRef<SSequencer>(InSequencer.GetSequencerWidget()))
@@ -103,7 +110,7 @@
 		return FCursorReply::Cursor( EMouseCursor::Default );
 	}
 
-	virtual void OnBeginDrag(const FPointerEvent& MouseEvent, FVector2D LocalMousePos, const FVirtualTrackArea& VirtualTrackArea) override
+	virtual void OnBeginDrag(const FPointerEvent& MouseEvent, FVector2D LocalMousePos, const UE::Sequencer::FVirtualTrackArea& VirtualTrackArea) override
 	{
 		// Start a new marquee selection
 		InitialPosition = CurrentPosition = VirtualTrackArea.PhysicalToVirtual(LocalMousePos);
@@ -126,12 +133,14 @@
 			// @todo: selection in transactions
 			//leave selections in the tree view alone so that dragging operations act similarly to click operations which don't change treeview selection state.
 			Sequencer.GetSelection().EmptySelectedKeys();
-			Sequencer.GetSelection().EmptySelectedSections();
+			Sequencer.GetSelection().EmptySelectedTrackAreaItems();
 		}
 	}
 
-	virtual void OnDrag(const FPointerEvent& MouseEvent, FVector2D LocalMousePos, const FVirtualTrackArea& VirtualTrackArea) override
+	virtual void OnDrag(const FPointerEvent& MouseEvent, FVector2D LocalMousePos, const UE::Sequencer::FVirtualTrackArea& VirtualTrackArea) override
 	{
+		using namespace UE::Sequencer;
+
 		// hange the current marquee selection
 		const FVector2D MouseDelta = MouseEvent.GetCursorDelta();
 
@@ -139,7 +148,7 @@
 		{
 			const float ScrollThresholdV = VirtualTrackArea.GetPhysicalSize().Y * 0.025f;
 
-			TSharedPtr<SSequencerTreeView> TreeView = TrackArea.GetTreeView().Pin();
+			TSharedPtr<SOutlinerView> TreeView = TrackArea.GetOutliner().Pin();
 
 			float Difference = LocalMousePos.Y - ScrollThresholdV;
 			if (Difference < 0 && MouseDelta.Y < 0)
@@ -199,14 +208,12 @@
 		// Ensure the preview is empty before calculating the intersection
 		SelectionPreview.Empty();
 
-		const auto& RootNodes = SequencerWidget->GetTreeView()->GetNodeTree()->GetRootNodes();
-
 		// Now walk everything within the current marquee range, setting preview selection states as we go
 		FSequencerEntityWalker Walker(FSequencerEntityRange(TopLeft(), BottomRight(), VirtualTrackArea.GetTickResolution()), VirtualKeySize);
-		Walker.Traverse(FSelectionPreviewVisitor(SelectionPreview, Sequencer.GetSelection(), PreviewState, TrackArea.ShowPinned()), RootNodes);
+		Walker.Traverse(FSelectionPreviewVisitor(SelectionPreview, Sequencer.GetSelection(), PreviewState, TrackArea.ShowPinned()), Sequencer.GetViewModel()->GetRootModel());
 	}
 
-	virtual void OnEndDrag(const FPointerEvent& MouseEvent, FVector2D LocalMousePos, const FVirtualTrackArea& VirtualTrackArea) override
+	virtual void OnEndDrag(const FPointerEvent& MouseEvent, FVector2D LocalMousePos, const UE::Sequencer::FVirtualTrackArea& VirtualTrackArea) override
 	{
 		// finish dragging the marquee selection
 		auto& Selection = Sequencer.GetSelection();
@@ -226,54 +233,30 @@
 			}
 		}
 
-		for (const auto& Pair : SelectionPreview.GetDefinedSectionStates())
+		for (const auto& Pair : SelectionPreview.GetDefinedModelStates())
 		{
-			UMovieSceneSection* Section = Pair.Key.Get();
-
 			if (Pair.Value == ESelectionPreviewState::Selected)
 			{
 				// Select it in the main selection
-				Selection.AddToSelection(Section);
+				Selection.AddToSelection(Pair.Key.Pin());
 			}
 			else
 			{
-				Selection.RemoveFromSelection(Section);
+				Selection.RemoveFromSelection(Pair.Key.Pin());
 			}
 		}
 
-		for (const auto& Pair : SelectionPreview.GetDefinedOutlinerNodeStates())
-		{
-			if (Pair.Value == ESelectionPreviewState::Selected)
-			{
-				Selection.AddToNodesWithSelectedKeysOrSections(Pair.Key);
-
-				// If this node exists inside a collapsed parent, we add such parents to the selection preview as well so that the node highlight is correct for accumulated keys
-				TSharedPtr<FSequencerDisplayNode> ParentNode = Pair.Key->GetParent();
-				while (ParentNode.IsValid())
-				{
-					if (!ParentNode->IsExpanded())
-					{
-						Selection.AddToNodesWithSelectedKeysOrSections(ParentNode.ToSharedRef());
-					}
-					ParentNode = ParentNode->GetParent();
-				}
-			}
-			else
-			{
-				Selection.RemoveFromNodesWithSelectedKeysOrSections(Pair.Key);
-			}
-		}
-
 		Selection.ResumeBroadcast();
 		Selection.RequestOutlinerNodeSelectionChangedBroadcast();
 
 		// We're done with this now
 		SelectionPreview.Empty();
-		SequencerHelpers::ValidateNodesWithSelectedKeysOrSections(Sequencer);
 	}
 
-	virtual int32 OnPaint(const FGeometry& AllottedGeometry, const FSlateRect& MyCullingRect, FSlateWindowElementList& OutDrawElements, int32 LayerId) const
+	virtual int32 OnPaint(const FGeometry& AllottedGeometry, const FSlateRect& MyCullingRect, FSlateWindowElementList& OutDrawElements, int32 LayerId) const override
 	{
+		using namespace UE::Sequencer;
+
 		// convert to physical space for rendering
 		const FVirtualTrackArea VirtualTrackArea = SequencerWidget->GetVirtualTrackArea(&TrackArea);
 
@@ -311,7 +294,7 @@
 	/** The sequencer itself */
 	FSequencer& Sequencer;
 
-	SSequencerTrackArea& TrackArea;
+	UE::Sequencer::STrackAreaView& TrackArea;
 
 	/** Sequencer widget */
 	TSharedRef<SSequencer> SequencerWidget;
@@ -327,7 +310,7 @@
 
 const FName FSequencerEditTool_Selection::Identifier = "Selection";
 
-FSequencerEditTool_Selection::FSequencerEditTool_Selection(FSequencer& InSequencer, SSequencerTrackArea& InTrackArea)
+FSequencerEditTool_Selection::FSequencerEditTool_Selection(FSequencer& InSequencer, UE::Sequencer::STrackAreaView& InTrackArea)
 	: FSequencerEditTool(InSequencer)
 	, TrackArea(InTrackArea)
 	, CursorDecorator(nullptr)
@@ -363,13 +346,16 @@
 
 FReply FSequencerEditTool_Selection::OnMouseButtonDown(SWidget& OwnerWidget, const FGeometry& MyGeometry, const FPointerEvent& MouseEvent)
 {
+	using namespace UE::Sequencer;
+
 	UpdateCursor(MyGeometry, MouseEvent);
 
 	DelayedDrag.Reset();
 
 	if (MouseEvent.GetEffectingButton() == EKeys::LeftMouseButton)
 	{
-		DelayedDrag = FDelayedDrag_Hotspot(MyGeometry.AbsoluteToLocal(MouseEvent.GetScreenSpacePosition()), EKeys::LeftMouseButton, Sequencer.GetHotspot());
+		TSharedPtr<ITrackAreaHotspot> Hotspot = Sequencer.GetViewModel()->GetTrackArea()->GetHotspot();
+		DelayedDrag = FDelayedDrag_Hotspot(MyGeometry.AbsoluteToLocal(MouseEvent.GetScreenSpacePosition()), EKeys::LeftMouseButton, Hotspot);
 		return FReply::Handled();
 	}
 	return FReply::Unhandled();
@@ -378,6 +364,8 @@
 
 FReply FSequencerEditTool_Selection::OnMouseMove(SWidget& OwnerWidget, const FGeometry& MyGeometry, const FPointerEvent& MouseEvent)
 {
+	using namespace UE::Sequencer;
+
 	UpdateCursor(MyGeometry, MouseEvent);
 
 	if (DelayedDrag.IsSet())
@@ -397,10 +385,9 @@
 			if (DelayedDrag->Hotspot.IsValid())
 			{
 				// We only allow resizing with the marquee selection tool enabled
-				auto HotspotType = DelayedDrag->Hotspot->GetType();
-				if (HotspotType != ESequencerHotspot::Section && HotspotType != ESequencerHotspot::Key)
+				if (DelayedDrag->Hotspot->CastThis<FSectionHotspot>() == nullptr && DelayedDrag->Hotspot->CastThis<FKeyHotspot>() == nullptr)
 				{
-					DragOperation = DelayedDrag->Hotspot->InitiateDrag(Sequencer);
+					DragOperation = DelayedDrag->Hotspot->InitiateDrag(MouseEvent);
 				}
 			}
 
@@ -500,9 +487,9 @@
 	return !DelayedDrag.IsSet();
 }
 
-const ISequencerHotspot* FSequencerEditTool_Selection::GetDragHotspot() const
+TSharedPtr<UE::Sequencer::ITrackAreaHotspot> FSequencerEditTool_Selection::GetDragHotspot() const
 {
-	return DelayedDrag.IsSet() ? DelayedDrag->Hotspot.Get() : nullptr;
+	return DelayedDrag.IsSet() ? DelayedDrag->Hotspot : nullptr;
 }
 
 void FSequencerEditTool_Selection::UpdateCursor(const FGeometry& MyGeometry, const FPointerEvent& MouseEvent)
diff --git a/Engine/Source/Editor/Sequencer/Private/Tools/SequencerEditTool_Selection.h b/Engine/Source/Editor/Sequencer/Private/Tools/SequencerEditTool_Selection.h
--- a/Engine/Source/Editor/Sequencer/Private/Tools/SequencerEditTool_Selection.h
+++ b/Engine/Source/Editor/Sequencer/Private/Tools/SequencerEditTool_Selection.h
@@ -7,12 +7,17 @@
 #include "Input/CursorReply.h"
 #include "Input/Reply.h"
 #include "Tools/SequencerEditTool.h"
-#include "Framework/DelayedDrag.h"
 
 class SSequencer;
-class SSequencerTrackArea;
-struct ISequencerHotspot;
 
+namespace UE
+{
+namespace Sequencer
+{
+	class STrackAreaView;
+}
+}
+
 class FSequencerEditTool_Selection
 	: public FSequencerEditTool
 {
@@ -22,7 +27,7 @@
 	static const FName Identifier;
 
 	/** Create and initialize a new instance. */
-	FSequencerEditTool_Selection(FSequencer& InSequencer, SSequencerTrackArea& InTrackArea);
+	FSequencerEditTool_Selection(FSequencer& InSequencer, UE::Sequencer::STrackAreaView& InTrackArea);
 
 public:
 
@@ -36,7 +41,7 @@
 	virtual FCursorReply OnCursorQuery(const FGeometry& MyGeometry, const FPointerEvent& CursorEvent) const override;
 	virtual FName GetIdentifier() const override;
 	virtual bool CanDeactivate() const override;
-	virtual const ISequencerHotspot* GetDragHotspot() const override;
+	virtual TSharedPtr<UE::Sequencer::ITrackAreaHotspot> GetDragHotspot() const override;
 	
 private:
 
@@ -45,29 +50,17 @@
 
 private:
 
-	struct FDelayedDrag_Hotspot : FDelayedDrag
-	{
-		FDelayedDrag_Hotspot(FVector2D InInitialPosition, FKey InApplicableKey, TSharedPtr<ISequencerHotspot> InHotspot)
-			: FDelayedDrag(InInitialPosition, InApplicableKey)
-			, Hotspot(MoveTemp(InHotspot))
-		{
-			SetTriggerScaleFactor(0.1f);
-		}
-
-		TSharedPtr<ISequencerHotspot> Hotspot;
-	};
-
 	/** Helper class responsible for handling delayed dragging */
-	TOptional<FDelayedDrag_Hotspot> DelayedDrag;
+	TOptional<UE::Sequencer::FDelayedDrag_Hotspot> DelayedDrag;
 	
 	/** Current drag operation if any */
-	TSharedPtr<ISequencerEditToolDragOperation> DragOperation;
+	TSharedPtr<UE::Sequencer::ISequencerEditToolDragOperation> DragOperation;
 
 	/** Cached mouse position for software cursor rendering */
 	FVector2D MousePosition;
 
 	/** TrackArea this object belongs to */
-	SSequencerTrackArea& TrackArea;
+	UE::Sequencer::STrackAreaView& TrackArea;
 
 	/** Software cursor decorator brush */
 	const FSlateBrush* CursorDecorator;
diff --git a/Engine/Source/Editor/Sequencer/Private/Tools/SequencerEntityVisitor.cpp b/Engine/Source/Editor/Sequencer/Private/Tools/SequencerEntityVisitor.cpp
--- a/Engine/Source/Editor/Sequencer/Private/Tools/SequencerEntityVisitor.cpp
+++ b/Engine/Source/Editor/Sequencer/Private/Tools/SequencerEntityVisitor.cpp
@@ -1,163 +1,217 @@
-// Copyright Epic Games, Inc. All Rights Reserved.
-
-#include "Tools/SequencerEntityVisitor.h"
-#include "IKeyArea.h"
-#include "DisplayNodes/SequencerTrackNode.h"
-#include "MovieSceneTrack.h"
-#include "MovieSceneTimeHelpers.h"
-
-FSequencerEntityRange::FSequencerEntityRange(const TRange<double>& InRange, FFrameRate InTickResolution)
-	: TickResolution(InTickResolution)
-	, Range(InRange)
-{
-}
-
-FSequencerEntityRange::FSequencerEntityRange(FVector2D TopLeft, FVector2D BottomRight, FFrameRate InTickResolution)
-	: TickResolution(InTickResolution)
-	, Range(TopLeft.X, BottomRight.X)
-	, VerticalTop(TopLeft.Y), VerticalBottom(BottomRight.Y)
-{
-}
-
-bool FSequencerEntityRange::IntersectSection(const UMovieSceneSection* InSection) const
-{
-	// Test horizontal bounds
-	return (InSection->GetRange() / TickResolution).Overlaps(Range);
-}
-
-bool FSequencerEntityRange::IntersectNode(TSharedRef<FSequencerDisplayNode> InNode) const
-{
-	if (VerticalTop.IsSet())
-	{
-		return InNode->GetVirtualTop() <= VerticalBottom.GetValue() && InNode->GetVirtualBottom() >= VerticalTop.GetValue();
-	}
-	return true;
-}
-
-bool FSequencerEntityRange::IntersectKeyArea(TSharedRef<FSequencerDisplayNode> InNode, float VirtualKeyHeight) const
-{
-	if (VerticalTop.IsSet())
-	{
-		const float NodeCenter = InNode->GetVirtualTop() + (InNode->GetVirtualBottom() - InNode->GetVirtualTop())/2;
-		return NodeCenter + VirtualKeyHeight/2 > VerticalTop.GetValue() && NodeCenter - VirtualKeyHeight/2 < VerticalBottom.GetValue();
-	}
-	return true;
-}
-
-FSequencerEntityWalker::FSequencerEntityWalker(const FSequencerEntityRange& InRange, FVector2D InVirtualKeySize)
-	: Range(InRange), VirtualKeySize(InVirtualKeySize)
-{}
-
-/* @todo: Could probably optimize this by not walking every single node, and binary searching the begin<->end ranges instead */
-void FSequencerEntityWalker::Traverse(const ISequencerEntityVisitor& Visitor, const TArray< TSharedRef<FSequencerDisplayNode> >& Nodes)
-{
-	for (const TSharedRef<FSequencerDisplayNode>& Child : Nodes)
-	{
-		if (!Child->IsHidden())
-		{
-			ConditionallyIntersectNode(Visitor, Child);
-		}
-	}
-}
-
-void FSequencerEntityWalker::ConditionallyIntersectNode(const ISequencerEntityVisitor& Visitor, const TSharedRef<FSequencerDisplayNode>& InNode)
-{
-	if (Range.IntersectNode(InNode))
-	{
-		// Visit sections within this track
-		if (InNode->GetType() == ESequencerNode::Track && Visitor.CheckEntityMask(ESequencerEntity::Section))
-		{
-			TSharedRef<FSequencerTrackNode> TrackNode = StaticCastSharedRef<FSequencerTrackNode>(InNode);
-
-			for (TSharedRef<ISequencerSection> SectionInterface : TrackNode->GetSections())
-			{
-				UMovieSceneSection* Section = SectionInterface->GetSectionObject();
-				if (Range.IntersectSection(Section))
-				{
-					Visitor.VisitSection(Section, InNode);
-				}
-			}
-		}
-
-		if (InNode->GetType() == ESequencerNode::Track || InNode->GetType() == ESequencerNode::KeyArea || InNode->GetType() == ESequencerNode::Category)
-		{
-			if (Range.IntersectKeyArea(InNode, VirtualKeySize.Y))
-			{
-				VisitKeyAnyAreas(Visitor, InNode, !InNode->IsExpanded());
-			}
-		}
-	}
-
-	// Iterate into expanded nodes
-	if (InNode->IsExpanded())
-	{
-		for (const TSharedRef<FSequencerDisplayNode>& Child : InNode->GetChildNodes())
-		{
-			// Do not visit nodes that are currently filtered out
-			if (!Child->IsHidden())
-			{
-				ConditionallyIntersectNode(Visitor, Child);
-			}
-		}
-	}
-}
-
-void FSequencerEntityWalker::VisitKeyAnyAreas(const ISequencerEntityVisitor& Visitor, const TSharedRef<FSequencerDisplayNode>& InNode, bool bAnyParentCollapsed )
-{
-	if (!Visitor.CheckEntityMask(ESequencerEntity::Key))
-	{
-		return;
-	}
-
-	TSharedPtr<FSequencerSectionKeyAreaNode> KeyAreaNode;
-	if (InNode->GetType() == ESequencerNode::KeyArea)
-	{
-		KeyAreaNode = StaticCastSharedRef<FSequencerSectionKeyAreaNode>(InNode);
-	}
-	else if (InNode->GetType() == ESequencerNode::Track)
-	{
-		KeyAreaNode = StaticCastSharedRef<FSequencerTrackNode>(InNode)->GetTopLevelKeyNode();
-	}
-
-	// If this node has or is a key area, visit all the keys on the track
-	if (KeyAreaNode)
-	{
-		for (TSharedRef<IKeyArea> KeyArea : KeyAreaNode->GetAllKeyAreas())
-		{
-			UMovieSceneSection* Section = KeyArea->GetOwningSection();
-			if (Section)
-			{
-				VisitKeyArea(Visitor, KeyArea, Section, InNode);
-			}
-		}
-	}
-	// Otherwise it might be a collapsed node that contains key areas as children. If so we visit them as if they were a part of this track so that key groupings are visited properly.
-	else if (bAnyParentCollapsed)
-	{
-		for (TSharedRef<FSequencerDisplayNode> ChildNode : InNode->GetChildNodes())
-		{
-			VisitKeyAnyAreas(Visitor, ChildNode, bAnyParentCollapsed || !InNode->IsExpanded());
-		}
-	}
-}
-
-void FSequencerEntityWalker::VisitKeyArea(const ISequencerEntityVisitor& Visitor, const TSharedRef<IKeyArea>& KeyArea, UMovieSceneSection* Section, const TSharedRef<FSequencerDisplayNode>& InNode)
-{
-	TArray<FKeyHandle> Handles;
-	TArray<FFrameNumber> Times;
-
-	const FFrameTime HalfKeySizeFrames   = (VirtualKeySize.X*.5f) * Range.TickResolution;
-	const FFrameTime RangeStartFrame     = Range.Range.GetLowerBoundValue() * Range.TickResolution;
-	const FFrameTime RangeEndFrame       = Range.Range.GetUpperBoundValue() * Range.TickResolution;
-
-	TRange<FFrameNumber> VisitRangeFrames( (RangeStartFrame-HalfKeySizeFrames).CeilToFrame(), (RangeEndFrame+HalfKeySizeFrames).FloorToFrame() );
-
-	VisitRangeFrames = TRange<FFrameNumber>::Intersection(Section->GetRange(), VisitRangeFrames);
-
-	KeyArea->GetKeyInfo(&Handles, &Times, VisitRangeFrames);
-
-	for (int32 Index = 0; Index < Times.Num(); ++Index)
-	{
-		Visitor.VisitKey(Handles[Index], Times[Index], KeyArea, Section, InNode);
-	}
-}
+// Copyright Epic Games, Inc. All Rights Reserved.
+
+#include "Tools/SequencerEntityVisitor.h"
+#include "IKeyArea.h"
+#include "MovieSceneTrack.h"
+#include "MovieSceneTimeHelpers.h"
+#include "MVVM/ViewModels/ViewModel.h"
+#include "MVVM/ViewModels/CategoryModel.h"
+#include "MVVM/ViewModels/ChannelModel.h"
+#include "MVVM/ViewModels/TrackModel.h"
+#include "MVVM/ViewModels/ViewModelIterators.h"
+#include "MVVM/Extensions/ITrackLaneExtension.h"
+#include "MVVM/Extensions/IGeometryExtension.h"
+
+FSequencerEntityRange::FSequencerEntityRange(const TRange<double>& InRange, FFrameRate InTickResolution)
+	: TickResolution(InTickResolution)
+	, Range(InRange)
+{
+}
+
+FSequencerEntityRange::FSequencerEntityRange(FVector2D TopLeft, FVector2D BottomRight, FFrameRate InTickResolution)
+	: TickResolution(InTickResolution)
+	, Range(TopLeft.X, BottomRight.X)
+	, VerticalTop(TopLeft.Y), VerticalBottom(BottomRight.Y)
+{
+}
+
+bool FSequencerEntityRange::IntersectSection(const UMovieSceneSection* InSection) const
+{
+	// Test horizontal bounds
+	return (InSection->GetRange() / TickResolution).Overlaps(Range);
+}
+
+bool FSequencerEntityRange::IntersectKeyArea(TSharedPtr<UE::Sequencer::FViewModel> InNode, float VirtualKeyHeight) const
+{
+	using namespace UE::Sequencer;
+
+	IGeometryExtension* GeometryExtension = InNode->CastThis<IGeometryExtension>();
+	if (VerticalTop.IsSet() && GeometryExtension)
+	{
+		const FVirtualGeometry VirtualGeometry = GeometryExtension->GetVirtualGeometry();
+		const float NodeCenter = VirtualGeometry.GetTop() + VirtualGeometry.GetHeight()/2;
+		return NodeCenter + VirtualKeyHeight/2 > VerticalTop.GetValue() && NodeCenter - VirtualKeyHeight/2 < VerticalBottom.GetValue();
+	}
+	return true;
+}
+
+bool FSequencerEntityRange::IntersectVertical(float Top, float Bottom) const
+{
+	if (VerticalTop.IsSet())
+	{
+		return Top <= VerticalBottom.GetValue() && Bottom >= VerticalTop.GetValue();
+	}
+	return true;
+}
+
+FSequencerEntityWalker::FSequencerEntityWalker(const FSequencerEntityRange& InRange, FVector2D InVirtualKeySize)
+	: Range(InRange), VirtualKeySize(InVirtualKeySize)
+{}
+
+void FSequencerEntityWalker::Traverse(const ISequencerEntityVisitor& Visitor, TSharedPtr<UE::Sequencer::FViewModel> Item)
+{
+	using namespace UE::Sequencer;
+
+	IOutlinerExtension* OutlinerItem = Item->CastThis<IOutlinerExtension>();
+	if (!OutlinerItem || !OutlinerItem->IsFilteredOut())
+	{
+		ConditionallyIntersectModel(Visitor, Item);
+	}
+}
+
+void FSequencerEntityWalker::ConditionallyIntersectModel(const ISequencerEntityVisitor& Visitor, const TSharedPtr<UE::Sequencer::FViewModel>& DataModel)
+{
+	using namespace UE::Sequencer;
+
+	const ITrackAreaExtension* TrackArea         = DataModel->CastThis<ITrackAreaExtension>();
+	const IGeometryExtension*  GeometryExtension = DataModel->CastThis<IGeometryExtension>();
+	if (GeometryExtension && TrackArea)
+	{
+		const FTrackAreaParameters TrackAreaParameters = TrackArea->GetTrackAreaParameters();
+		const FVirtualGeometry     VirtualGeometry     = GeometryExtension->GetVirtualGeometry();
+
+		const float Height = TrackAreaParameters.LaneType == ETrackAreaLaneType::Nested
+			? VirtualGeometry.NestedBottom - VirtualGeometry.Top
+			: VirtualGeometry.Height;
+
+		if (Range.IntersectVertical(VirtualGeometry.Top, VirtualGeometry.Top + Height))
+		{
+			for (const TViewModelPtr<ITrackLaneExtension>& TrackLane
+				: DataModel->GetChildrenOfType<ITrackLaneExtension>(EViewModelListType::TrackArea))
+			{
+				FTrackLaneVirtualAlignment Alignment = TrackLane->ArrangeVirtualTrackLaneView();
+				if (!Alignment.IsVisible())
+				{
+					continue;
+				}
+
+				FTrackLaneVerticalArrangement VerticalArrange = Alignment.VerticalAlignment.ArrangeWithin(Height);
+
+				if (Range.IntersectVertical(
+							VirtualGeometry.Top + VerticalArrange.Offset,
+							VirtualGeometry.Top + VerticalArrange.Offset + VerticalArrange.Height) )
+				{
+					if (Range.Range.Overlaps(Alignment.Range / Range.TickResolution))
+					{
+						Visitor.VisitDataModel(TrackLane.AsModel().Get());
+
+						if (TViewModelPtr<FChannelModel> ChannelModel = TrackLane.ImplicitCast())
+						{
+							VisitChannel(Visitor, ChannelModel);
+						}
+					}
+				}
+			}
+
+			if (Range.IntersectKeyArea(DataModel, VirtualKeySize.Y))
+			{
+				const IOutlinerExtension* OutlinerItem = DataModel->CastThis<IOutlinerExtension>();
+				const bool bIsExpanded = !OutlinerItem || OutlinerItem->IsExpanded();
+				VisitAnyChannels(Visitor, DataModel.ToSharedRef(), !bIsExpanded);
+			}
+		}
+	}
+
+	// Iterate into expanded nodes
+	const IOutlinerExtension* OutlinerExtension = DataModel->CastThis<IOutlinerExtension>();
+	if (OutlinerExtension == nullptr || OutlinerExtension->IsExpanded())
+	{
+		for (TSharedPtr<UE::Sequencer::FViewModel> Child : DataModel->GetChildren(EViewModelListType::Outliner))
+		{
+			// Do not visit nodes that are currently filtered out
+			const IOutlinerExtension* OutlinerItem = DataModel->CastThis<IOutlinerExtension>();
+			if (!OutlinerItem || !OutlinerItem->IsFilteredOut())
+			{
+				ConditionallyIntersectModel(Visitor, Child);
+			}
+		}
+	}
+}
+
+void FSequencerEntityWalker::VisitAnyChannels(const ISequencerEntityVisitor& Visitor, const TSharedRef<UE::Sequencer::FViewModel>& InNode, bool bAnyParentCollapsed )
+{
+	using namespace UE::Sequencer;
+
+	if (!Visitor.CheckEntityMask(ESequencerEntity::Key))
+	{
+		return;
+	}
+
+	// If this node has or is a key area, visit all the keys on the track
+	FChannelGroupModel* ChannelGroup = nullptr;
+	if (FChannelGroupModel* ChannelGroupNode = InNode->CastThis<FChannelGroupModel>())
+	{
+		ChannelGroup = ChannelGroupNode;
+	}
+	if (FTrackModel* TrackModel = InNode->CastThis<FTrackModel>())
+	{
+		// TODO: Top-level channel group -- there used to be only one at most but this assumption isn't necessarily true anymore
+		for (const TViewModelPtr<FChannelGroupModel>& TopLevelChannel : TrackModel->GetTopLevelChannels().IterateSubList<FChannelGroupModel>())
+		{
+			ChannelGroup = TopLevelChannel.Get();
+			break;
+		}
+	}
+
+	if (ChannelGroup)
+	{
+		for (const TWeakViewModelPtr<FChannelModel>& WeakChannel : ChannelGroup->GetChannels())
+		{
+			if (TViewModelPtr<FChannelModel> Channel = WeakChannel.Pin())
+			{
+				VisitChannel(Visitor, Channel);
+			}
+		}
+	}
+	// Otherwise it might be a collapsed node that contains key areas as children. If so we visit them as if they were a part of this track so that key groupings are visited properly.
+	else if (bAnyParentCollapsed)
+	{
+		const IOutlinerExtension* OutlinerItem = InNode->CastThis<IOutlinerExtension>();
+		const bool bIsExpanded = !OutlinerItem || OutlinerItem->IsExpanded();
+
+		for (TSharedPtr<FViewModel> ChildNode : InNode->GetChildren(EViewModelListType::Outliner))
+		{
+			VisitAnyChannels(Visitor, ChildNode.ToSharedRef(), bAnyParentCollapsed || !bIsExpanded);
+		}
+	}
+}
+
+void FSequencerEntityWalker::VisitChannel(const ISequencerEntityVisitor& Visitor, const UE::Sequencer::TViewModelPtr<UE::Sequencer::FChannelModel>& Channel)
+{
+	using namespace UE::Sequencer;
+
+	TSharedPtr<IKeyArea> KeyArea = Channel->GetKeyArea();
+	UMovieSceneSection* Section = KeyArea->GetOwningSection();
+	if (!Section)
+	{
+		return;
+	}
+
+	TArray<FKeyHandle> Handles;
+	TArray<FFrameNumber> Times;
+
+	const FFrameTime HalfKeySizeFrames   = (VirtualKeySize.X*.5f) * Range.TickResolution;
+	const FFrameTime RangeStartFrame     = Range.Range.GetLowerBoundValue() * Range.TickResolution;
+	const FFrameTime RangeEndFrame       = Range.Range.GetUpperBoundValue() * Range.TickResolution;
+
+	TRange<FFrameNumber> VisitRangeFrames( (RangeStartFrame-HalfKeySizeFrames).CeilToFrame(), (RangeEndFrame+HalfKeySizeFrames).FloorToFrame() );
+
+	VisitRangeFrames = TRange<FFrameNumber>::Intersection(Section->GetRange(), VisitRangeFrames);
+
+	KeyArea->GetKeyInfo(&Handles, &Times, VisitRangeFrames);
+
+	for (int32 Index = 0; Index < Times.Num(); ++Index)
+	{
+		Visitor.VisitKey(Handles[Index], Times[Index], Channel, Section);
+	}
+}
diff --git a/Engine/Source/Editor/Sequencer/Private/Tools/SequencerEntityVisitor.h b/Engine/Source/Editor/Sequencer/Private/Tools/SequencerEntityVisitor.h
--- a/Engine/Source/Editor/Sequencer/Private/Tools/SequencerEntityVisitor.h
+++ b/Engine/Source/Editor/Sequencer/Private/Tools/SequencerEntityVisitor.h
@@ -2,15 +2,23 @@
 
 #pragma once
 
-#include "CoreMinimal.h"
+#include "SequencerCoreFwd.h"
+#include "CoreTypes.h"
 #include "Curves/KeyHandle.h"
-#include "DisplayNodes/SequencerDisplayNode.h"
-#include "DisplayNodes/SequencerSectionKeyAreaNode.h"
 #include "ISequencerSection.h"
+#include "MVVM/ViewModels/ViewModel.h"
 
-class FSequencerTrackNode;
 class IKeyArea;
 
+namespace UE
+{
+namespace Sequencer
+{
+	class FTrackModel;
+	class FChannelModel;
+}
+}
+
 /** Enum of different types of entities that are available in the sequencer */
 namespace ESequencerEntity
 {
@@ -28,9 +36,9 @@
 {
 	ISequencerEntityVisitor(uint32 InEntityMask = ESequencerEntity::Everything) : EntityMask(InEntityMask) {}
 
-	virtual void VisitKey(FKeyHandle KeyHandle, FFrameNumber KeyTime, const TSharedPtr<IKeyArea>& KeyArea, UMovieSceneSection* Section, TSharedRef<FSequencerDisplayNode>) const { }
-	virtual void VisitSection(UMovieSceneSection* Section, TSharedRef<FSequencerDisplayNode>) const { }
-	
+	virtual void VisitKey(FKeyHandle KeyHandle, FFrameNumber KeyTime, const UE::Sequencer::TViewModelPtr<UE::Sequencer::FChannelModel>& Channel, UMovieSceneSection* Section) const { }
+	virtual void VisitDataModel(UE::Sequencer::FViewModel* Item) const { }
+
 	/** Check if the specified type of entity is applicable to this visitor */
 	bool CheckEntityMask(ESequencerEntity::Type Type) const { return (EntityMask & Type) != 0; }
 
@@ -50,12 +58,12 @@
 	/** Check whether the specified section intersects the horizontal range */
 	bool IntersectSection(const UMovieSceneSection* InSection) const;
 
-	/** Check whether the specified node intersects the vertical range */
-	bool IntersectNode(TSharedRef<FSequencerDisplayNode> InNode) const;
-
 	/** Check whether the specified node's key area intersects this range */
-	bool IntersectKeyArea(TSharedRef<FSequencerDisplayNode> InNode, float VirtualKeyHeight) const;
+	bool IntersectKeyArea(TSharedPtr<UE::Sequencer::FViewModel> InNode, float VirtualKeyHeight) const;
 
+	/** Check whether the specified vertical range intersects this range */
+	bool IntersectVertical(float Top, float Bottom) const;
+
 	/** tick resolution of the current time-base */
 	FFrameRate TickResolution;
 
@@ -72,19 +80,17 @@
 	/** Construction from the range itself, and an optional virtual key size, where key bounds must be taken into consideration */
 	FSequencerEntityWalker(const FSequencerEntityRange& InRange, FVector2D InVirtualKeySize);
 
-	/** Visit the specified nodes (recursively) with this range and a user-supplied visitor */
-	void Traverse(const ISequencerEntityVisitor& Visitor, const TArray< TSharedRef<FSequencerDisplayNode> >& Nodes);
+	/** Visit the specified node (recursively) with this range and a user-supplied visitor */
+	void Traverse(const ISequencerEntityVisitor& Visitor, TSharedPtr<UE::Sequencer::FViewModel> Item);
 
 private:
 
 	/** Check whether the specified node intersects the range's vertical space, and visit any keys within it if so */
-	void ConditionallyIntersectNode(const ISequencerEntityVisitor& Visitor, const TSharedRef<FSequencerDisplayNode>& InNode);
-	/** Visit any keys within the specified track node that overlap the range's horizontal space */
-	void VisitTrackNode(const ISequencerEntityVisitor& Visitor, const TSharedRef<FSequencerTrackNode>& InNode);
+	void ConditionallyIntersectModel(const ISequencerEntityVisitor& Visitor, const TSharedPtr<UE::Sequencer::FViewModel>& DataModel);
 	/** Visit any keys within any key area nodes that belong to the specified node that overlap the range's horizontal space */
-	void VisitKeyAnyAreas(const ISequencerEntityVisitor& Visitor, const TSharedRef<FSequencerDisplayNode>& InNode, bool bAnyParentCollapsed);
+	void VisitAnyChannels(const ISequencerEntityVisitor& Visitor, const TSharedRef<UE::Sequencer::FViewModel>& InNode, bool bAnyParentCollapsed);
 	/** Visit any keys within the specified key area that overlap the range's horizontal space */
-	void VisitKeyArea(const ISequencerEntityVisitor& Visitor, const TSharedRef<IKeyArea>& KeyArea, UMovieSceneSection* Section, const TSharedRef<FSequencerDisplayNode>& InNode);
+	void VisitChannel(const ISequencerEntityVisitor& Visitor, const UE::Sequencer::TViewModelPtr<UE::Sequencer::FChannelModel>& Channel);
 
 	/** The bounds of the range */
 	FSequencerEntityRange Range;
