Sequencer: MVVM2 branch and Layer Bars

diff --git a/Engine/Source/Editor/Sequencer/Private/Tools/SequencerSnapField.cpp b/Engine/Source/Editor/Sequencer/Private/Tools/SequencerSnapField.cpp
--- a/Engine/Source/Editor/Sequencer/Private/Tools/SequencerSnapField.cpp
+++ b/Engine/Source/Editor/Sequencer/Private/Tools/SequencerSnapField.cpp
@@ -3,65 +3,68 @@
 #include "Tools/SequencerSnapField.h"
 #include "MovieScene.h"
 #include "SSequencer.h"
-#include "SSequencerTreeView.h"
 #include "MovieSceneTimeHelpers.h"
 #include "MovieSceneSequence.h"
 #include "ISequencerSection.h"
+#include "MVVM/ViewModels/ViewModel.h"
+#include "MVVM/Views/SOutlinerView.h"
+#include "MVVM/Extensions/ISnappableExtension.h"
 
-struct FSnapGridVisitor : ISequencerEntityVisitor
+struct FSnapGridVisitor : ISequencerEntityVisitor, UE::Sequencer::ISnapField
 {
-	FSnapGridVisitor(ISequencerSnapCandidate& InCandidate, uint32 EntityMask)
+	FSnapGridVisitor(UE::Sequencer::ISnapCandidate& InCandidate, uint32 EntityMask)
 		: ISequencerEntityVisitor(EntityMask)
 		, Candidate(InCandidate)
 	{}
 
-	virtual void VisitKey(FKeyHandle KeyHandle, FFrameNumber KeyTime, const TSharedPtr<IKeyArea>& KeyArea, UMovieSceneSection* Section, TSharedRef<FSequencerDisplayNode>) const
+	virtual void VisitKey(FKeyHandle KeyHandle, FFrameNumber KeyTime, const UE::Sequencer::TViewModelPtr<UE::Sequencer::FChannelModel>& Channel, UMovieSceneSection* Section) const
 	{
-		if (Candidate.IsKeyApplicable(KeyHandle, KeyArea, Section))
+		using namespace UE::Sequencer;
+
+		if (Candidate.IsKeyApplicable(KeyHandle, Channel))
 		{
-			Snaps.Add(FSequencerSnapPoint{ FSequencerSnapPoint::Key, KeyTime });
+			Snaps.Add(FSnapPoint{ FSnapPoint::Key, KeyTime });
 		}
 	}
-	virtual void VisitSection(UMovieSceneSection* Section, TSharedRef<FSequencerDisplayNode> Node) const
+	virtual void VisitDataModel(UE::Sequencer::FViewModel* DataModel) const
 	{
-		if (Candidate.AreSectionBoundsApplicable(Section))
+		using namespace UE::Sequencer;
+
+		if (ISnappableExtension* Snappable = DataModel->CastThis<ISnappableExtension>())
 		{
-			if (Section->HasStartFrame())
-			{
-				Snaps.Add(FSequencerSnapPoint{ FSequencerSnapPoint::SectionBounds, Section->GetInclusiveStartFrame() });
-			}
-
-			if (Section->HasEndFrame())
-			{
-				Snaps.Add(FSequencerSnapPoint{ FSequencerSnapPoint::SectionBounds, Section->GetExclusiveEndFrame() });
-			}
-		}
-
-		if (Candidate.AreSectionCustomSnapsApplicable(Section))
-		{
-			TArray<FFrameNumber> CustomSnaps;
-			Section->GetSnapTimes(CustomSnaps, false);
-			for (FFrameNumber Time : CustomSnaps)
-			{
-				Snaps.Add(FSequencerSnapPoint{ FSequencerSnapPoint::CustomSection, Time });
-			}
+			Snappable->AddToSnapField(Candidate, *const_cast<FSnapGridVisitor*>(this));
 		}
 	}
 
-	ISequencerSnapCandidate& Candidate;
-	mutable TArray<FSequencerSnapPoint> Snaps;
+	virtual void AddSnapPoint(const UE::Sequencer::FSnapPoint& SnapPoint) override
+	{
+		Snaps.Add(SnapPoint);
+	}
+
+	UE::Sequencer::ISnapCandidate& Candidate;
+	mutable TArray<UE::Sequencer::FSnapPoint> Snaps;
 };
 
-FSequencerSnapField::FSequencerSnapField(const FSequencer& InSequencer, ISequencerSnapCandidate& Candidate, uint32 EntityMask)
+FSequencerSnapField::FSequencerSnapField(const FSequencer& InSequencer, UE::Sequencer::ISnapCandidate& Candidate, uint32 EntityMask)
 {
-	TSharedPtr<SSequencerTreeView> TreeView = StaticCastSharedRef<SSequencer>(InSequencer.GetSequencerWidget())->GetTreeView();
+	Initialize(InSequencer, Candidate, EntityMask);
+	Finalize();
+}
 
-	TArray<TSharedRef<FSequencerDisplayNode>> VisibleNodes;
-	for (const SSequencerTreeView::FCachedGeometry& Geometry : TreeView->GetAllVisibleNodes())
-	{
-		VisibleNodes.Add(Geometry.Node);
-	}
+void FSequencerSnapField::AddExplicitSnap(UE::Sequencer::FSnapPoint InSnap)
+{
+	SortedSnaps.Add(InSnap);
+}
 
+void FSequencerSnapField::Initialize(const FSequencer& InSequencer, UE::Sequencer::ISnapCandidate& Candidate, uint32 EntityMask)
+{
+	using namespace UE::Sequencer;
+
+	TSharedPtr<SOutlinerView> TreeView = StaticCastSharedRef<SSequencer>(InSequencer.GetSequencerWidget())->GetTreeView();
+
+	TArray<TViewModelPtr<IOutlinerExtension>> VisibleItems;
+	TreeView->GetVisibleItems(VisibleItems);
+
 	TRange<double> ViewRange = InSequencer.GetViewRange();
 	FSequencerEntityWalker Walker(
 		FSequencerEntityRange(ViewRange, InSequencer.GetFocusedTickResolution()),
@@ -69,53 +72,71 @@
 
 	// Traverse the visible space, collecting snapping times as we go
 	FSnapGridVisitor Visitor(Candidate, EntityMask);
-	Walker.Traverse(Visitor, VisibleNodes);
+	for (const TViewModelPtr<IOutlinerExtension>& Item : VisibleItems)
+	{
+		Walker.Traverse(Visitor, Item);
+	}
 
 	// Add the playback range start/end bounds as potential snap candidates
 	TRange<FFrameNumber> PlaybackRange = InSequencer.GetFocusedMovieSceneSequence()->GetMovieScene()->GetPlaybackRange();
 	if(UE::MovieScene::DiscreteSize(PlaybackRange) > 0)
 	{
-		Visitor.Snaps.Add(FSequencerSnapPoint{ FSequencerSnapPoint::PlaybackRange, UE::MovieScene::DiscreteInclusiveLower(PlaybackRange)});
-		Visitor.Snaps.Add(FSequencerSnapPoint{ FSequencerSnapPoint::PlaybackRange, UE::MovieScene::DiscreteExclusiveUpper(PlaybackRange)});
+		Visitor.Snaps.Add(FSnapPoint{ FSnapPoint::PlaybackRange, UE::MovieScene::DiscreteInclusiveLower(PlaybackRange)});
+		Visitor.Snaps.Add(FSnapPoint{ FSnapPoint::PlaybackRange, UE::MovieScene::DiscreteExclusiveUpper(PlaybackRange)});
 	}
 
 	// Add the current time as a potential snap candidate
-	Visitor.Snaps.Add(FSequencerSnapPoint{ FSequencerSnapPoint::CurrentTime, InSequencer.GetLocalTime().Time.FrameNumber });
+	Visitor.Snaps.Add(FSnapPoint{ FSnapPoint::CurrentTime, InSequencer.GetLocalTime().Time.FrameNumber });
 
 	// Add the selection range bounds as a potential snap candidate
 	TRange<FFrameNumber> SelectionRange = InSequencer.GetFocusedMovieSceneSequence()->GetMovieScene()->GetSelectionRange();
 	if (UE::MovieScene::DiscreteSize(SelectionRange) > 0)
 	{
-		Visitor.Snaps.Add(FSequencerSnapPoint{ FSequencerSnapPoint::InOutRange, UE::MovieScene::DiscreteInclusiveLower(SelectionRange)});
-		Visitor.Snaps.Add(FSequencerSnapPoint{ FSequencerSnapPoint::InOutRange, UE::MovieScene::DiscreteExclusiveUpper(SelectionRange) - 1});
+		Visitor.Snaps.Add(FSnapPoint{ FSnapPoint::InOutRange, UE::MovieScene::DiscreteInclusiveLower(SelectionRange)});
+		Visitor.Snaps.Add(FSnapPoint{ FSnapPoint::InOutRange, UE::MovieScene::DiscreteExclusiveUpper(SelectionRange) - 1});
 	}
 
 	// Add in the marked frames
 	for (const FMovieSceneMarkedFrame& MarkedFrame : InSequencer.GetFocusedMovieSceneSequence()->GetMovieScene()->GetMarkedFrames())
 	{
-		Visitor.Snaps.Add( FSequencerSnapPoint{ FSequencerSnapPoint::Mark, MarkedFrame.FrameNumber } );
+		Visitor.Snaps.Add( FSnapPoint{ FSnapPoint::Mark, MarkedFrame.FrameNumber } );
 	}
 
 	// Add in the global marked frames
 	for (const FMovieSceneMarkedFrame& MarkedFrame : InSequencer.GetGlobalMarkedFrames())
 	{
-		Visitor.Snaps.Add(FSequencerSnapPoint{ FSequencerSnapPoint::Mark, MarkedFrame.FrameNumber });
+		Visitor.Snaps.Add(FSnapPoint{ FSnapPoint::Mark, MarkedFrame.FrameNumber });
 	}
 
+
+	if (SortedSnaps.Num() == 0)
+	{
+		SortedSnaps = MoveTemp(Visitor.Snaps);
+	}
+	else
+	{
+		SortedSnaps.Append(Visitor.Snaps);
+	}
+}
+
+void FSequencerSnapField::Finalize()
+{
+	using namespace UE::Sequencer;
+
 	// Sort
-	Visitor.Snaps.Sort([](const FSequencerSnapPoint& A, const FSequencerSnapPoint& B){
+	SortedSnaps.Sort([](const FSnapPoint& A, const FSnapPoint& B){
 		return A.Time < B.Time;
 	});
 
 	// Remove duplicates
-	for (int32 Index = 0; Index < Visitor.Snaps.Num(); ++Index)
+	for (int32 Index = 0; Index < SortedSnaps.Num(); ++Index)
 	{
-		const FFrameNumber CurrentTime = Visitor.Snaps[Index].Time;
+		const FFrameNumber CurrentTime = SortedSnaps[Index].Time;
 
 		int32 NumToMerge = 0;
-		for (int32 DuplIndex = Index + 1; DuplIndex < Visitor.Snaps.Num(); ++DuplIndex)
+		for (int32 DuplIndex = Index + 1; DuplIndex < SortedSnaps.Num(); ++DuplIndex)
 		{
-			if (CurrentTime != Visitor.Snaps[DuplIndex].Time)
+			if (CurrentTime != SortedSnaps[DuplIndex].Time)
 			{
 				break;
 			}
@@ -124,11 +145,9 @@
 
 		if (NumToMerge)
 		{
-			Visitor.Snaps.RemoveAt(Index + 1, NumToMerge, false);
+			SortedSnaps.RemoveAt(Index + 1, NumToMerge, false);
 		}
 	}
-
-	SortedSnaps = MoveTemp(Visitor.Snaps);
 }
 
 TOptional<FFrameNumber> FSequencerSnapField::Snap(FFrameNumber InTime, int32 Threshold) const
diff --git a/Engine/Source/Editor/Sequencer/Private/Tools/SequencerSnapField.h b/Engine/Source/Editor/Sequencer/Private/Tools/SequencerSnapField.h
--- a/Engine/Source/Editor/Sequencer/Private/Tools/SequencerSnapField.h
+++ b/Engine/Source/Editor/Sequencer/Private/Tools/SequencerSnapField.h
@@ -1,69 +1,51 @@
-// Copyright Epic Games, Inc. All Rights Reserved.
-
-#pragma once
-
-#include "CoreMinimal.h"
-#include "Curves/KeyHandle.h"
-#include "Tools/SequencerEntityVisitor.h"
-
-class IKeyArea;
-class ISequencer;
-
-/** Structure defining a point to snap to in the sequencer */
-struct FSequencerSnapPoint
-{
-	enum ESnapType { Key, SectionBounds, CustomSection, PlaybackRange, CurrentTime, InOutRange, Mark };
-
-	/** The type of snap */
-	ESnapType Type;
-
-	/** The time of the snap */
-	FFrameNumber Time;
-};
-
-/** Interface that defines how to construct an FSequencerSnapField */
-struct ISequencerSnapCandidate
-{
-	virtual ~ISequencerSnapCandidate() { }
-
-	/** Return true to include the specified key in the snap field */
-	virtual bool IsKeyApplicable(FKeyHandle KeyHandle, const TSharedPtr<IKeyArea>& KeyArea, UMovieSceneSection* Section) { return true; }
-
-	/** Return true to include the specified section's bounds in the snap field */
-	virtual bool AreSectionBoundsApplicable(UMovieSceneSection* Section) { return true; }
-
-	/** Return true to include the specified section's custom snap points in the snap field */
-	virtual bool AreSectionCustomSnapsApplicable(UMovieSceneSection* Section) { return true; }
-};
-
-/** A snapping field that provides efficient snapping calculations on a range of values */
-class FSequencerSnapField
-{
-public:
-
-	/** A snap result denoting the time that was snapped, and the resulting snapped time */
-	struct FSnapResult
-	{
-		/** The time before it was snapped */
-		FFrameNumber Original;
-		/** The time after it was snapped */
-		FFrameNumber Snapped;
-	};
-
-	/** Construction from a sequencer and a snap canidate implementation. Optionally provide an entity mask to completely ignore some entity types */
-	FSequencerSnapField(const FSequencer& InSequencer, ISequencerSnapCandidate& Candidate, uint32 EntityMask = ESequencerEntity::Everything);
-
-	/** Move construction / assignment */
-	FSequencerSnapField(FSequencerSnapField&& In) : SortedSnaps(MoveTemp(In.SortedSnaps)) {}
-	FSequencerSnapField& operator=(FSequencerSnapField&& In) { SortedSnaps = MoveTemp(In.SortedSnaps); return *this; }
-	
-	/** Snap the specified time to this field with the given threshold */
-	TOptional<FFrameNumber> Snap(FFrameNumber InTime, int32 Threshold) const;
-
-	/** Snap the specified times to this field with the given threshold. Will return the closest snap value of the entire intersection. */
-	TOptional<FSnapResult> Snap(const TArray<FFrameNumber>& InTimes, int32 Threshold) const;
-
-private:
-	/** Array of snap points, approximately grouped, and sorted in ascending order by time */
-	TArray<FSequencerSnapPoint> SortedSnaps;
-};
+// Copyright Epic Games, Inc. All Rights Reserved.
+
+#pragma once
+
+#include "CoreTypes.h"
+#include "Containers/Array.h"
+#include "Tools/SequencerEntityVisitor.h"
+#include "MVVM/Extensions/ISnappableExtension.h"
+
+class IKeyArea;
+class FSequencer;
+
+/** A snapping field that provides efficient snapping calculations on a range of values */
+class FSequencerSnapField
+{
+public:
+
+	/** A snap result denoting the time that was snapped, and the resulting snapped time */
+	struct FSnapResult
+	{
+		/** The time before it was snapped */
+		FFrameNumber Original;
+		/** The time after it was snapped */
+		FFrameNumber Snapped;
+	};
+
+	FSequencerSnapField(){}
+
+	/** Construction from a sequencer and a snap canidate implementation. Optionally provide an entity mask to completely ignore some entity types */
+	FSequencerSnapField(const FSequencer& InSequencer, UE::Sequencer::ISnapCandidate& Candidate, uint32 EntityMask = ESequencerEntity::Everything);
+
+	/** Move construction / assignment */
+	FSequencerSnapField(FSequencerSnapField&& In) : SortedSnaps(MoveTemp(In.SortedSnaps)) {}
+	FSequencerSnapField& operator=(FSequencerSnapField&& In) { SortedSnaps = MoveTemp(In.SortedSnaps); return *this; }
+
+	void Initialize(const FSequencer& InSequencer, UE::Sequencer::ISnapCandidate& Candidate, uint32 EntityMask = ESequencerEntity::Everything);
+
+	void AddExplicitSnap(UE::Sequencer::FSnapPoint InSnap);
+
+	void Finalize();
+
+	/** Snap the specified time to this field with the given threshold */
+	TOptional<FFrameNumber> Snap(FFrameNumber InTime, int32 Threshold) const;
+
+	/** Snap the specified times to this field with the given threshold. Will return the closest snap value of the entire intersection. */
+	TOptional<FSnapResult> Snap(const TArray<FFrameNumber>& InTimes, int32 Threshold) const;
+
+private:
+	/** Array of snap points, approximately grouped, and sorted in ascending order by time */
+	TArray<UE::Sequencer::FSnapPoint> SortedSnaps;
+};
diff --git a/Engine/Source/Editor/Sequencer/Public/IKeyArea.h b/Engine/Source/Editor/Sequencer/Public/IKeyArea.h
--- a/Engine/Source/Editor/Sequencer/Public/IKeyArea.h
+++ b/Engine/Source/Editor/Sequencer/Public/IKeyArea.h
@@ -259,7 +259,7 @@
 	 * @param SrcEnvironment   The environment the source data was copied from
 	 * @param DstEnvironment   The environment we're pasting into
 	 */
-	void PasteKeys(const FMovieSceneClipboardKeyTrack& KeyTrack, const FMovieSceneClipboardEnvironment& SrcEnvironment, const FSequencerPasteEnvironment& DstEnvironment);
+	TArray<FKeyHandle> PasteKeys(const FMovieSceneClipboardKeyTrack& KeyTrack, const FMovieSceneClipboardEnvironment& SrcEnvironment, const FSequencerPasteEnvironment& DstEnvironment);
 
 
 	/**
diff --git a/Engine/Source/Editor/Sequencer/Public/ISequencer.h b/Engine/Source/Editor/Sequencer/Public/ISequencer.h
--- a/Engine/Source/Editor/Sequencer/Public/ISequencer.h
+++ b/Engine/Source/Editor/Sequencer/Public/ISequencer.h
@@ -32,6 +32,7 @@
 class UMovieSceneSequence;
 class UMovieSceneSubSection;
 class UMovieSceneCinematicShotSection;
+class ISequencerTrackEditor;
 class IDetailsView;
 class IKeyArea;
 enum class EMapChangeType : uint8;
@@ -41,6 +42,16 @@
 struct FMovieSceneChannelMetaData;
 class IToolkitHost;
 
+namespace UE
+{
+namespace Sequencer
+{
+
+class FEditorViewModel;
+
+} // namespace Sequencer
+} // namespace UE
+
 /**
  * Defines auto change modes.
  */
@@ -234,6 +245,11 @@
 	virtual void FocusSequenceInstance(UMovieSceneSubSection& Section) = 0;
 
 	/**
+	 * Retrieve the top level view model for this sequence
+	 */
+	virtual TSharedPtr<UE::Sequencer::FEditorViewModel> GetViewModel() const = 0;
+
+	/**
 	 * Suppresses automatic evaluation the specified sequence and signature are the only difference that would prompt a re-evaluation
 	 * 
 	 * @param Sequence        The sequence that the signature relates to
@@ -707,6 +723,12 @@
 	* @see SetSelectionRange, SetSelectionRangeEnd, SetSelectionRangeStart
 	*/
 	virtual TRange<FFrameNumber> GetSelectionRange() const = 0;
+
+	/**
+	* Retrieve or create a track editor for the specified track
+	*/
+	virtual TSharedPtr<ISequencerTrackEditor> GetTrackEditor(UMovieSceneTrack* InTrack) = 0;
+
 public:
 
 	/**
@@ -777,3 +799,4 @@
 	FOnGetIsRecording GetIsRecording;
 	FOnGetCanRecord GetCanRecord;
 };
+
diff --git a/Engine/Source/Editor/Sequencer/Public/ISequencerModule.h b/Engine/Source/Editor/Sequencer/Public/ISequencerModule.h
--- a/Engine/Source/Editor/Sequencer/Public/ISequencerModule.h
+++ b/Engine/Source/Editor/Sequencer/Public/ISequencerModule.h
@@ -24,6 +24,8 @@
 class UMovieSceneSequence;
 struct FSequencerInitParams;
 
+namespace UE { namespace Sequencer { class FTrackModel; } }
+
 enum class ECurveEditorTreeFilterType : uint32;
 
 /** Forward declaration for the default templated channel interface. Include SequencerChannelInterface.h for full definition. */
@@ -50,7 +52,13 @@
 /** A delegate which will create an object binding handler. */
 DECLARE_DELEGATE_RetVal_OneParam(TSharedRef<ISequencerEditorObjectBinding>, FOnCreateEditorObjectBinding, TSharedRef<ISequencer>);
 
+/** A delegate which will create a track model. */
+DECLARE_DELEGATE_RetVal_OneParam(TSharedPtr<UE::Sequencer::FTrackModel>, FOnCreateTrackModel, UMovieSceneTrack*);
+
 /** A delegate that is executed when adding menu content. */
+DECLARE_DELEGATE_OneParam(FOnGetContextMenuContent, FMenuBuilder& /*MenuBuilder*/);
+
+/** A delegate that is executed when adding menu content. */
 DECLARE_DELEGATE_TwoParams(FOnGetAddMenuContent, FMenuBuilder& /*MenuBuilder*/, TSharedRef<ISequencer>);
 
 /** A delegate that is executed when initializing tool menu context. */
@@ -70,8 +78,10 @@
  */
 struct FSequencerViewParams
 {
+	/** Called when building the add track menu */
 	FOnGetAddMenuContent OnGetAddMenuContent;
 
+	/** Called when building a context menu for an object binding */
 	FOnBuildCustomContextMenuForGuid OnBuildCustomContextMenuForGuid;
 
 	/** Called when this sequencer has received user focus */
@@ -206,6 +216,21 @@
 	 */
 	virtual void UnRegisterTrackEditor(FDelegateHandle InHandle) = 0;
 
+	/**
+	 * Registers a delegate that will create a track view model for a given track
+	 *
+	 * @param InCreator Delegate to register
+	 * @return A handle to the newly-added delegate
+	 */
+	virtual FDelegateHandle RegisterTrackModel(FOnCreateTrackModel InCreator) = 0;
+
+	/**
+	 * Unregisters a previously registered delegate for creating track view models
+	 *
+	 * @param InHandle Handle to the delegate to unregister
+	 */
+	virtual void UnregisterTrackModel(FDelegateHandle InHandle) = 0;
+
 	/** 
 	 * Registers a delegate that will be called when a sequencer is created
 	 *
diff --git a/Engine/Source/Editor/Sequencer/Public/ISequencerTrackEditor.h b/Engine/Source/Editor/Sequencer/Public/ISequencerTrackEditor.h
--- a/Engine/Source/Editor/Sequencer/Public/ISequencerTrackEditor.h
+++ b/Engine/Source/Editor/Sequencer/Public/ISequencerTrackEditor.h
@@ -207,8 +207,8 @@
 	/** Called when the instance of this track editor is released */
 	virtual void OnRelease() = 0;
 
-	/** Allows the track editor to paint on a track area. */
-	virtual int32 PaintTrackArea(const FPaintArgs& Args, const FGeometry& AllottedGeometry, const FSlateRect& MyCullingRect, FSlateWindowElementList& OutDrawElements, int32 LayerId, const FWidgetStyle& InWidgetStyle) = 0;
+	UE_DEPRECATED(5.0, "This function is no longer used.")
+	virtual int32 PaintTrackArea(const FPaintArgs& Args, const FGeometry& AllottedGeometry, const FSlateRect& MyCullingRect, FSlateWindowElementList& OutDrawElements, int32 LayerId, const FWidgetStyle& InWidgetStyle) { return LayerId; }
 
 	/**
 	 * Returns whether a track class is supported by this tool.
diff --git a/Engine/Source/Editor/Sequencer/Public/MVVM/CurveEditorIntegrationExtension.h b/Engine/Source/Editor/Sequencer/Public/MVVM/CurveEditorIntegrationExtension.h
new file mode 100644
--- /dev/null
+++ b/Engine/Source/Editor/Sequencer/Public/MVVM/CurveEditorIntegrationExtension.h
@@ -0,0 +1,57 @@
+// Copyright Epic Games, Inc. All Rights Reserved.
+
+#pragma once
+
+#include "MVVM/ViewModels/ViewModel.h"
+#include "MVVM/ViewModelPtr.h"
+#include "MVVM/Extensions/DynamicExtensionContainer.h"
+#include "MVVM/Extensions/ICurveEditorTreeItemExtension.h"
+#include "CurveEditorTypes.h"
+#include "Tree/ICurveEditorTreeItem.h"
+
+namespace UE
+{
+namespace Sequencer
+{
+
+class FSequenceModel;
+
+/**
+ * Extension for managing integration between outliner items and the curve editor.
+ *
+ * It relies on the following:
+ *
+ * - Extension owner (generally the root view-model) implements ICurveEditorExtension, to get
+ *   access to the curve editor itself.
+ *
+ * - Outliner items implementing ICurveEditorTreeItemExtension (or its default shim) if they 
+ *   want to show up in the curve editor.
+ */
+class FCurveEditorIntegrationExtension : public IDynamicExtension
+{
+public:
+
+	UE_SEQUENCER_DECLARE_VIEW_MODEL_TYPE_ID(FCurveEditorIntegrationExtension)
+
+	FCurveEditorIntegrationExtension();
+
+	virtual void OnCreated(TSharedRef<FViewModel> InWeakOwner) override;
+
+	void UpdateCurveEditor();
+	void RecreateCurveEditor();
+
+private:
+
+	void OnHierarchyChanged();
+
+	FCurveEditorTreeItemID AddToCurveEditor(TViewModelPtr<ICurveEditorTreeItemExtension> InViewModel, TSharedPtr<FCurveEditor> CurveEditor);
+
+private:
+
+	TWeakPtr<FSequenceModel> WeakOwnerModel;
+	TMap<TWeakPtr<FViewModel>, FCurveEditorTreeItemID> ViewModelToTreeItemIDMap;
+};
+
+} // namespace Sequencer
+} // namespace UE
+
diff --git a/Engine/Source/Editor/Sequencer/Public/MVVM/Extensions/IObjectBindingExtension.h b/Engine/Source/Editor/Sequencer/Public/MVVM/Extensions/IObjectBindingExtension.h
new file mode 100644
--- /dev/null
+++ b/Engine/Source/Editor/Sequencer/Public/MVVM/Extensions/IObjectBindingExtension.h
@@ -0,0 +1,27 @@
+// Copyright Epic Games, Inc. All Rights Reserved.
+
+#pragma once
+
+#include "MVVM/ViewModelTypeID.h"
+
+struct FGuid;
+
+namespace UE
+{
+namespace Sequencer
+{
+
+class SEQUENCER_API IObjectBindingExtension
+{
+public:
+
+	UE_SEQUENCER_DECLARE_VIEW_MODEL_TYPE_ID(IObjectBindingExtension)
+
+	virtual ~IObjectBindingExtension(){}
+
+	virtual FGuid GetObjectGuid() const = 0;
+};
+
+} // namespace Sequencer
+} // namespace UE
+
diff --git a/Engine/Source/Editor/Sequencer/Public/MVVM/ObjectBindingModelStorageExtension.h b/Engine/Source/Editor/Sequencer/Public/MVVM/ObjectBindingModelStorageExtension.h
new file mode 100644
--- /dev/null
+++ b/Engine/Source/Editor/Sequencer/Public/MVVM/ObjectBindingModelStorageExtension.h
@@ -0,0 +1,74 @@
+// Copyright Epic Games, Inc. All Rights Reserved.
+
+#pragma once
+
+#include "Misc/Guid.h"
+#include "MVVM/Extensions/DynamicExtensionContainer.h"
+#include "EventHandlers/ISequenceDataEventHandler.h"
+
+struct FMovieSceneBinding;
+
+namespace UE
+{
+namespace Sequencer
+{
+
+struct FViewModelChildren;
+
+class FSequenceModel;
+class FViewModel;
+class FObjectBindingModel;
+class FPlaceholderObjectBindingModel;
+
+class FObjectBindingModelStorageExtension
+	: public IDynamicExtension
+	, private UE::MovieScene::TIntrusiveEventHandler<UE::MovieScene::ISequenceDataEventHandler>
+{
+public:
+
+	UE_SEQUENCER_DECLARE_VIEW_MODEL_TYPE_ID(FObjectBindingModelStorageExtension);
+
+	FObjectBindingModelStorageExtension();
+
+	TSharedPtr<FViewModel> GetOrCreateModelForBinding(const FGuid& Binding);
+	TSharedPtr<FViewModel> GetOrCreateModelForBinding(const FMovieSceneBinding& Binding);
+
+	TSharedPtr<FObjectBindingModel> FindModelForObjectBinding(const FGuid& InObjectBindingID) const;
+
+	virtual void OnCreated(TSharedRef<FViewModel> InWeakOwner) override;
+	virtual void OnReinitialize() override;
+
+private:
+
+	void OnBindingAdded(const FMovieSceneBinding& Binding) override;
+	void OnBindingRemoved(const FGuid& ObjectBindingID) override;
+	void OnTrackAddedToBinding(UMovieSceneTrack* Track, const FGuid& Binding) override;
+	void OnTrackRemovedFromBinding(UMovieSceneTrack* Track, const FGuid& Binding) override;
+	void OnBindingParentChanged(const FGuid& Binding, const FGuid& NewParent) override;
+
+
+	TSharedPtr<FObjectBindingModel> CreateModelForObjectBinding(const FMovieSceneBinding& Binding);
+
+	TSharedPtr<FViewModel> CreatePlaceholderForObjectBinding(const FGuid& ObjectID);
+	TSharedPtr<FViewModel> FindPlaceholderForObjectBinding(const FGuid& InObjectBindingID) const;
+
+	/**
+	 * Implementation function for creating a new model for a binding
+	 * @param Binding The binding to create a model for
+	 * @param RootChildren (optional) When specified, will place the object into this child list, if it does not define a desired parent ID.
+	 */
+	TSharedPtr<FViewModel> GetOrCreateModelForBinding(const FMovieSceneBinding& Binding, FViewModelChildren* RootChildren);
+
+	void Compact();
+
+private:
+
+	TMap<FGuid, TWeakPtr<FObjectBindingModel>> ObjectBindingToModel;
+	TMap<FGuid, TWeakPtr<FPlaceholderObjectBindingModel>> ObjectBindingToPlaceholder;
+
+	FSequenceModel* OwnerModel;
+};
+
+} // namespace Sequencer
+} // namespace UE
+
diff --git a/Engine/Source/Editor/Sequencer/Public/MVVM/SectionModelStorageExtension.h b/Engine/Source/Editor/Sequencer/Public/MVVM/SectionModelStorageExtension.h
new file mode 100644
--- /dev/null
+++ b/Engine/Source/Editor/Sequencer/Public/MVVM/SectionModelStorageExtension.h
@@ -0,0 +1,41 @@
+// Copyright Epic Games, Inc. All Rights Reserved.
+
+#pragma once
+
+#include "UObject/ObjectKey.h"
+#include "MVVM/Extensions/DynamicExtensionContainer.h"
+
+class ISequencerSection;
+class UMovieSceneSection;
+
+namespace UE
+{
+namespace Sequencer
+{
+
+class FSectionModel;
+class FSequenceModel;
+
+class FSectionModelStorageExtension
+	: public IDynamicExtension
+{
+public:
+
+	UE_SEQUENCER_DECLARE_VIEW_MODEL_TYPE_ID(FSectionModelStorageExtension);
+
+	FSectionModelStorageExtension();
+
+	TSharedPtr<FSectionModel> CreateModelForSection(UMovieSceneSection* InSection, TSharedRef<ISequencerSection> SectionInterface);
+
+	TSharedPtr<FSectionModel> FindModelForSection(const UMovieSceneSection* InSection) const;
+
+	virtual void OnReinitialize() override;
+
+private:
+
+	TMap<FObjectKey, TWeakPtr<FSectionModel>> SectionToModel;
+};
+
+} // namespace Sequencer
+} // namespace UE
+
diff --git a/Engine/Source/Editor/Sequencer/Public/MVVM/TrackRowModelStorageExtension.h b/Engine/Source/Editor/Sequencer/Public/MVVM/TrackRowModelStorageExtension.h
new file mode 100644
--- /dev/null
+++ b/Engine/Source/Editor/Sequencer/Public/MVVM/TrackRowModelStorageExtension.h
@@ -0,0 +1,46 @@
+// Copyright Epic Games, Inc. All Rights Reserved.
+
+#pragma once
+
+#include "UObject/ObjectKey.h"
+#include "MVVM/Extensions/DynamicExtensionContainer.h"
+
+class UMovieSceneTrack;
+
+namespace UE
+{
+namespace Sequencer
+{
+
+class FTrackRowModel;
+class FSequenceModel;
+class FViewModel;
+
+class FTrackRowModelStorageExtension
+	: public IDynamicExtension
+{
+public:
+
+	UE_SEQUENCER_DECLARE_VIEW_MODEL_TYPE_ID(FTrackRowModelStorageExtension)
+
+	FTrackRowModelStorageExtension();
+
+	virtual void OnCreated(TSharedRef<FViewModel> InWeakOwner) override;
+	virtual void OnReinitialize() override;
+
+	TSharedPtr<FTrackRowModel> CreateModelForTrackRow(UMovieSceneTrack* InTrack, int32 InRowIndex, TSharedPtr<FViewModel> DesiredParent = nullptr);
+
+	TSharedPtr<FTrackRowModel> FindModelForTrackRow(UMovieSceneTrack* InTrack, int32 InRowIndex) const;
+
+private:
+
+	using KeyType = TTuple<FObjectKey, int32>;
+
+	TMap<KeyType, TWeakPtr<FTrackRowModel>> TrackToModel;
+
+	FSequenceModel* OwnerModel;
+};
+
+} // namespace Sequencer
+} // namespace UE
+
diff --git a/Engine/Source/Editor/Sequencer/Public/MVVM/ViewModels/CategoryModel.h b/Engine/Source/Editor/Sequencer/Public/MVVM/ViewModels/CategoryModel.h
new file mode 100644
--- /dev/null
+++ b/Engine/Source/Editor/Sequencer/Public/MVVM/ViewModels/CategoryModel.h
@@ -0,0 +1,126 @@
+// Copyright Epic Games, Inc. All Rights Reserved.
+
+#pragma once
+
+#include "MVVM/ViewModels/ViewModel.h"
+#include "MVVM/ViewModels/OutlinerItemModel.h"
+#include "MVVM/Extensions/LinkedOutlinerExtension.h"
+#include "MVVM/Extensions/IGeometryExtension.h"
+#include "MVVM/Extensions/ITrackLaneExtension.h"
+#include "MVVM/Extensions/ITrackAreaExtension.h"
+#include "MVVM/Extensions/IDeletableExtension.h"
+
+#include "UObject/NameTypes.h"
+#include "Internationalization/Text.h"
+
+class UMovieSceneSection;
+class FSequencerSectionCategoryNode;
+
+namespace UE
+{
+namespace Sequencer
+{
+
+/**
+ * Model for a single channel category inside a section.
+ * For instance, this represents the "Location" category of a single transform section, which would
+ * contain the X, Y, and Z translation channels.
+ */
+class FCategoryModel
+	: public FViewModel
+	, public FLinkedOutlinerExtension
+	, public FGeometryExtensionShim
+	, public ITrackLaneExtension
+	, public FLinkedOutlinerComputedSizingShim
+{
+public:
+	UE_SEQUENCER_DECLARE_CASTABLE(FCategoryModel, FViewModel, FLinkedOutlinerExtension, IGeometryExtension, ITrackLaneExtension);
+
+	explicit FCategoryModel(FName InCategoryName);
+
+	/** Whether any of the channels within this category are animated.  */
+	bool IsAnimated() const;
+
+	/** Returns the category's name */
+	FName GetCategoryName() const { return CategoryName; }
+
+	/** Returns the desired sizing for the track area row */
+	FOutlinerSizing GetDesiredSizing() const;
+
+	/*~ ITrackLaneExtension */
+	TSharedPtr<ITrackLaneWidget> CreateTrackLaneView(const FCreateTrackLaneViewParams& InParams) override;
+	FTrackLaneVirtualAlignment ArrangeVirtualTrackLaneView() const override;
+
+private:
+
+	FViewModelListHead Children;
+	FName CategoryName;
+};
+
+/**
+ * Model for the outliner entry associated with all sections' channel categories of a given common name.
+ * For instance, this represents the "Location" category entry in the Sequence outliner, which would contain
+ * the X, Y, and Z translation channels of all the corresponding sections in the track area.
+ */
+class FCategoryGroupModel
+	: public FOutlinerItemModel
+	, public ITrackAreaExtension
+	, public ICompoundOutlinerExtension
+	, public IDeletableExtension
+{
+public:
+	UE_SEQUENCER_DECLARE_CASTABLE(FCategoryGroupModel, FOutlinerItemModel, ITrackAreaExtension, ICompoundOutlinerExtension, IDeletableExtension);
+
+	explicit FCategoryGroupModel(FName InCategoryName, const FText& InDisplayText);
+
+	~FCategoryGroupModel();
+
+	/**
+	 * @return Whether any of the channels within this category group have any keyframes on them
+	 */
+	bool IsAnimated() const;
+
+	FName GetCategoryName() const
+	{
+		return CategoryName;
+	}
+
+	FText GetDisplayText() const
+	{
+		return DisplayText;
+	}
+
+	void AddCategory(TWeakViewModelPtr<FCategoryModel> InCategory);
+
+	TArrayView<const TWeakViewModelPtr<FCategoryModel>> GetCategories() const;
+
+public:
+
+	/*~ ICompoundOutlinerExtension */
+	FOutlinerSizing RecomputeSizing() override;
+
+	/*~ IOutlinerExtension */
+	FOutlinerSizing GetOutlinerSizing() const override;
+	FText GetLabel() const override;
+	FSlateFontInfo GetLabelFont() const override;
+	TSharedRef<SWidget> CreateOutlinerView(const FCreateOutlinerViewParams& InParams) override;
+
+	/*~ ITrackAreaExtension */
+	FTrackAreaParameters GetTrackAreaParameters() const override;
+	FViewModelVariantIterator GetTrackAreaModelList() const override;
+
+	/*~ IDeletableExtension */
+	bool CanDelete(FText* OutErrorMessage) const override;
+	void Delete() override;
+
+private:
+
+	TArray<TWeakViewModelPtr<FCategoryModel>> Categories;
+	FName CategoryName;
+	FText DisplayText;
+	FOutlinerSizing ComputedSizing;
+};
+
+} // namespace Sequencer
+} // namespace UE
+
diff --git a/Engine/Source/Editor/Sequencer/Public/MVVM/ViewModels/ChannelModel.h b/Engine/Source/Editor/Sequencer/Public/MVVM/ViewModels/ChannelModel.h
new file mode 100644
--- /dev/null
+++ b/Engine/Source/Editor/Sequencer/Public/MVVM/ViewModels/ChannelModel.h
@@ -0,0 +1,163 @@
+// Copyright Epic Games, Inc. All Rights Reserved.
+
+#pragma once
+
+#include "MVVM/ViewModels/ViewModel.h"
+#include "MVVM/Extensions/IOutlinerExtension.h"
+#include "MVVM/Extensions/IGeometryExtension.h"
+#include "MVVM/ViewModels/OutlinerItemModel.h"
+#include "MVVM/Extensions/LinkedOutlinerExtension.h"
+#include "MVVM/Extensions/ITrackAreaExtension.h"
+#include "MVVM/Extensions/IDeletableExtension.h"
+#include "EventHandlers/ISignedObjectEventHandler.h"
+#include "Tree/ICurveEditorTreeItem.h"
+
+#include "UObject/NameTypes.h"
+#include "Internationalization/Text.h"
+
+struct FMovieSceneChannel;
+struct FMovieSceneChannelHandle;
+
+class IKeyArea;
+class ISequencerSection;
+class UMovieSceneSection;
+class FSequencerSectionKeyAreaNode;
+
+namespace UE
+{
+namespace Sequencer
+{
+
+class FSectionModel;
+
+/**
+ * Model for a single channel inside a section.
+ * For instance, this represents the "Location.X" channel of a single transform section.
+ */
+class FChannelModel
+	: public FViewModel
+	, public FLinkedOutlinerExtension
+	, public FGeometryExtensionShim
+	, public FLinkedOutlinerComputedSizingShim
+{
+public:
+
+	UE_SEQUENCER_DECLARE_CASTABLE(FChannelModel, FViewModel, FLinkedOutlinerExtension, IGeometryExtension);
+
+	FChannelModel(FName InChannelName, TWeakPtr<ISequencerSection> InSection, FMovieSceneChannelHandle InChannel);
+	~FChannelModel();
+
+	void Initialize(TWeakPtr<ISequencerSection> InSection, FMovieSceneChannelHandle InChannel);
+
+	/** Returns the section object that owns the associated channel */
+	UMovieSceneSection* GetSection() const;
+
+	/** Returns the associated channel object */
+	FMovieSceneChannel* GetChannel() const;
+
+	/** Returns whether this channel has any keyframes on it */
+	bool IsAnimated() const;
+
+	/** Returns the channel's name */
+	FName GetChannelName() const { return ChannelName; }
+
+	/** Returns the key area for the channel */
+	TSharedPtr<IKeyArea> GetKeyArea() const { return KeyArea; }
+
+	/** Create the curve editor model for the associated channel */
+	void CreateCurveModels(TArray<TUniquePtr<FCurveModel>>& OutCurveModels);
+
+	/** Returns the desired sizing for the track area row */
+	FOutlinerSizing GetDesiredSizing() const;
+
+private:
+
+	TSharedPtr<IKeyArea> KeyArea;
+	FName ChannelName;
+};
+
+/**
+ * Model for the outliner entry associated with all sections' channels of a given common name.
+ * For instance, this represents the "Location.X" entry in the Sequencer outliner.
+ */
+class FChannelGroupModel
+	: public FOutlinerItemModel
+	, public ITrackAreaExtension
+	, public ICompoundOutlinerExtension
+	, public IDeletableExtension
+{
+public:
+
+	UE_SEQUENCER_DECLARE_CASTABLE(FChannelGroupModel, FOutlinerItemModel, ITrackAreaExtension, ICompoundOutlinerExtension, IDeletableExtension);
+
+	FChannelGroupModel(FName InChannelName, const FText& InDisplayText);
+	~FChannelGroupModel();
+
+	/** Returns whether any of the channels within this group have any keyframes on them */
+	bool IsAnimated() const;
+
+	/** Returns the common name for all channels in this group */
+	FName GetChannelName() const { return ChannelName; }
+
+	/** Returns the label for this group */
+	FText GetDisplayText() const { return DisplayText; }
+
+	/** Gets all the channel models in this group */
+	TArrayView<const TWeakViewModelPtr<FChannelModel>> GetChannels() const;
+	/** Adds a channel model to this group */
+	void AddChannel(TWeakViewModelPtr<FChannelModel> InChannel);
+
+	/** Get the key area of the channel associated with the given section */
+	TSharedPtr<IKeyArea> GetKeyArea(TSharedPtr<FSectionModel> InOwnerSection) const;
+
+	/** Get the key area of the channel associated with the given section */
+	TSharedPtr<IKeyArea> GetKeyArea(const UMovieSceneSection* InOwnerSection) const;
+
+	/** Get the channel model of the channel associated with the given section */
+	TSharedPtr<FChannelModel> GetChannel(TSharedPtr<FSectionModel> InOwnerSection) const;
+
+	/** Get the channel model of the channel associated with the given section */
+	TSharedPtr<FChannelModel> GetChannel(const UMovieSceneSection* InOwnerSection) const;
+
+	/** Get the key areas of all channels */
+	TArray<TSharedRef<IKeyArea>> GetAllKeyAreas() const;
+
+public:
+
+	/*~ ICompoundOutlinerExtension */
+	FOutlinerSizing RecomputeSizing() override;
+
+	/*~ IOutlinerExtension */
+	FOutlinerSizing GetOutlinerSizing() const override;
+	FText GetLabel() const override;
+	FSlateFontInfo GetLabelFont() const override;
+	TSharedRef<SWidget> CreateOutlinerView(const FCreateOutlinerViewParams& InParams) override;
+
+	/*~ ITrackAreaExtension */
+	FTrackAreaParameters GetTrackAreaParameters() const override;
+	FViewModelVariantIterator GetTrackAreaModelList() const override;
+
+	/*~ ICurveEditorTreeItem */
+	void CreateCurveModels(TArray<TUniquePtr<FCurveModel>>& OutCurveModels) override;
+
+	/*~ IDeletableExtension */
+	bool CanDelete(FText* OutErrorMessage) const override;
+	void Delete() override;
+
+	/*~ FOutlinerItemModel */
+	bool HasCurves() const override;
+
+private:
+
+	EVisibility GetKeyEditorVisibility() const;
+
+	TArray<TWeakViewModelPtr<FChannelModel>> Channels;
+	FName ChannelName;
+	FText DisplayText;
+	FOutlinerSizing ComputedSizing;
+};
+
+
+} // namespace Sequencer
+} // namespace UE
+
diff --git a/Engine/Source/Editor/Sequencer/Public/MVVM/ViewModels/LayerBarModel.h b/Engine/Source/Editor/Sequencer/Public/MVVM/ViewModels/LayerBarModel.h
new file mode 100644
--- /dev/null
+++ b/Engine/Source/Editor/Sequencer/Public/MVVM/ViewModels/LayerBarModel.h
@@ -0,0 +1,82 @@
+// Copyright Epic Games, Inc. All Rights Reserved.
+
+#pragma once
+
+#include "MVVM/ViewModels/ViewModel.h"
+#include "MVVM/Extensions/ITrackLaneExtension.h"
+#include "MVVM/Extensions/ISelectableExtension.h"
+#include "MVVM/Extensions/ISnappableExtension.h"
+#include "MVVM/Extensions/IStretchableExtension.h"
+#include "MVVM/Extensions/IDraggableTrackAreaExtension.h"
+#include "MVVM/Extensions/ILayerBarExtension.h"
+#include "MVVM/Extensions/LinkedOutlinerExtension.h"
+#include "MVVM/Extensions/ViewModelExtensionCollection.h"
+
+namespace UE
+{
+namespace Sequencer
+{
+
+class FSequenceModel;
+
+class FLayerBarModel
+	: public FViewModel
+	, public FLinkedOutlinerExtension
+	, public ITrackLaneExtension
+	, public ISelectableExtension
+	, public ISnappableExtension
+	, public IDraggableTrackAreaExtension
+	, public IStretchableExtension
+	, protected TViewModelExtensionCollection<ILayerBarExtension>
+{
+public:
+
+	UE_SEQUENCER_DECLARE_CASTABLE(FLayerBarModel, FViewModel
+		, FLinkedOutlinerExtension
+		, ITrackLaneExtension
+		, ISelectableExtension
+		, ISnappableExtension
+		, IDraggableTrackAreaExtension
+		, IStretchableExtension
+	);
+
+	FLayerBarModel(TWeakPtr<FViewModel> LayerRoot);
+	~FLayerBarModel();
+
+	/*~ ITrackLaneExtension Interface */
+	TSharedPtr<ITrackLaneWidget> CreateTrackLaneView(const FCreateTrackLaneViewParams& InParams) override;
+	FTrackLaneVirtualAlignment ArrangeVirtualTrackLaneView() const override;
+
+	/*~ ISelectableExtension Interface */
+	ESelectionIntent IsSelectable() const override;
+
+	/*~ ISnappableExtension Interface */
+	void AddToSnapField(const ISnapCandidate& Candidate, ISnapField& SnapField) const override;
+
+	/*~ IDraggableTrackAreaExtension Interface */
+	bool CanDrag() const override;
+	void OnBeginDrag(IDragOperation& DragOperation) override;
+	void OnEndDrag(IDragOperation& DragOperation) override;
+
+	/*~ IStretchableExtension Interface */
+	void OnInitiateStretch(IStretchOperation& StretchOperation, EStretchConstraint Constraint, FStretchParameters* InOutGlobalParameters) override;
+
+public:
+
+	TRange<FFrameNumber> ComputeRange() const;
+
+	void Offset(FFrameNumber Offset);
+
+private:
+
+	/*~ FViewModel interface */
+	void OnConstruct() override;
+	void OnDestruct() override;
+
+	/*~ TViewModelExtensionCollection Interface */
+	void OnExtensionsDirtied() override;
+};
+
+} // namespace Sequencer
+} // namespace UE
+
diff --git a/Engine/Source/Editor/Sequencer/Public/MVVM/ViewModels/ObjectBindingModel.h b/Engine/Source/Editor/Sequencer/Public/MVVM/ViewModels/ObjectBindingModel.h
new file mode 100644
--- /dev/null
+++ b/Engine/Source/Editor/Sequencer/Public/MVVM/ViewModels/ObjectBindingModel.h
@@ -0,0 +1,155 @@
+// Copyright Epic Games, Inc. All Rights Reserved.
+
+#pragma once
+
+#include "Misc/Guid.h"
+#include "MVVM/ViewModels/ViewModel.h"
+#include "MVVM/ViewModels/ViewModelHierarchy.h"
+#include "MVVM/ViewModels/OutlinerItemModel.h"
+#include "MVVM/Extensions/IRenameableExtension.h"
+#include "MVVM/Extensions/ITrackAreaExtension.h"
+#include "MVVM/Extensions/IGroupableExtension.h"
+#include "MVVM/Extensions/IObjectBindingExtension.h"
+#include "MVVM/Extensions/ISortableExtension.h"
+#include "MVVM/Extensions/IDraggableOutlinerExtension.h"
+#include "MVVM/Extensions/IDeletableExtension.h"
+
+struct FMovieSceneBinding;
+
+class UMovieScene;
+class UMovieSceneTrack;
+class FMenuBuilder;
+class FPropertyPath;
+enum class ECheckBoxState : uint8;
+
+namespace UE
+{
+namespace Sequencer
+{
+
+class FSequenceModel;
+class FLayerBarModel;
+class FTrackModelStorageExtension;
+
+/** Enumeration specifying what kind of object binding this is */
+enum class EObjectBindingType
+{
+	Possessable, Spawnable, Unknown
+};
+
+class FObjectBindingModel
+	: public FOutlinerItemModel
+	, public IObjectBindingExtension
+	, public IDraggableOutlinerExtension
+	, public ITrackAreaExtension
+	, public IGroupableExtension
+	, public IRenameableExtension
+	, public ISortableExtension
+	, public IDeletableExtension
+{
+public:
+
+	UE_SEQUENCER_DECLARE_CASTABLE(FObjectBindingModel
+		, FOutlinerItemModel
+		, IObjectBindingExtension
+		, IDraggableOutlinerExtension
+		, ITrackAreaExtension
+		, IGroupableExtension
+		, IRenameableExtension
+		, ISortableExtension
+		, IDeletableExtension);
+
+	FObjectBindingModel(FSequenceModel* OwnerModel, const FMovieSceneBinding& InBinding);
+	~FObjectBindingModel();
+
+	void AddTrack(UMovieSceneTrack* Track);
+	void RemoveTrack(UMovieSceneTrack* Track);
+
+	/*~ IObjectBindingExtension */
+	FGuid GetObjectGuid() const override;
+
+	/*~ IRenameableExtension */
+	bool CanRename() const override;
+	void Rename(const FText& NewName) override;
+
+	/*~ IOutlinerExtension */
+	FOutlinerSizing GetOutlinerSizing() const override;
+	FText GetLabel() const override;
+	FSlateColor GetLabelColor() const override;
+	FText GetLabelToolTipText() const override;
+	const FSlateBrush* GetIconBrush() const override;
+	TSharedRef<SWidget> CreateOutlinerView(const FCreateOutlinerViewParams& InParams) override;
+
+	/*~ ITrackAreaExtension */
+	FTrackAreaParameters GetTrackAreaParameters() const override;
+	FViewModelVariantIterator GetTrackAreaModelList() const override;
+
+	/*~ IGroupableExtension */
+	void GetIdentifierForGrouping(TStringBuilder<128>& OutString) const override;
+
+	/*~ ISortableExtension */
+	void SortChildren() override;
+	FSortingKey GetSortingKey() const override;
+	void SetCustomOrder(int32 InCustomOrder) override;
+
+	/*~ FOutlinerItemModel */
+	void BuildContextMenu(FMenuBuilder& MenuBuilder) override;
+	void BuildOrganizeContextMenu(FMenuBuilder& MenuBuilder) override;
+	bool GetDefaultExpansionState() const override;
+
+	/*~ IDraggableOutlinerExtension */
+	bool CanDrag() const override;
+
+	/*~ IDeletableExtension */
+	bool CanDelete(FText* OutErrorMessage) const override;
+	void Delete() override;
+
+public:
+
+	virtual void SetParentBindingID(const FGuid& InObjectBindingID);
+	virtual FGuid GetDesiredParentBinding() const;
+	virtual EObjectBindingType GetType() const;
+	virtual FText GetTooltipForSingleObjectBinding() const;
+	virtual const UClass* FindObjectClass() const;
+	virtual bool SupportsRebinding() const;
+
+public:
+
+	TSharedRef<SWidget> GetAddTrackMenuContent();
+
+protected:
+
+	/*~ FViewModel interface */
+	void OnConstruct() override;
+
+private:
+
+	void AddSpawnOwnershipMenu(FMenuBuilder& MenuBuilder);
+	void AddSpawnLevelMenu(FMenuBuilder& MenuBuilder);
+	void AddTagMenu(FMenuBuilder& MenuBuilder);
+	void AddChangeClassMenu(FMenuBuilder& MenuBuilder);
+
+private:
+
+	void AddPropertyMenuItems(FMenuBuilder& AddTrackMenuBuilder, TArray<FPropertyPath> KeyableProperties, int32 PropertyNameIndexStart, int32 PropertyNameIndexEnd);
+	void HandleAddTrackSubMenuNew(FMenuBuilder& AddTrackMenuBuilder, TArray<FPropertyPath> KeyablePropertyPaths, int32 PropertyNameIndexStart);
+	void HandlePropertyMenuItemExecute(FPropertyPath PropertyPath);
+
+	ECheckBoxState GetTagCheckState(FName TagName);
+	void ToggleTag(FName TagName);
+	void HandleDeleteTag(FName TagName);
+	void HandleAddTag(FName TagName);
+	void HandleTemplateActorClassPicked(UClass* ChosenClass);
+
+protected:
+
+	FGuid ObjectBindingID;
+	FGuid ParentObjectBindingID;
+	FViewModelListHead TrackAreaList;
+	TSharedPtr<FLayerBarModel> LayerBar;
+	FSequenceModel* OwnerModel;
+};
+
+} // namespace Sequencer
+} // namespace UE
+
diff --git a/Engine/Source/Editor/Sequencer/Public/MVVM/ViewModels/OutlinerItemModel.h b/Engine/Source/Editor/Sequencer/Public/MVVM/ViewModels/OutlinerItemModel.h
new file mode 100644
--- /dev/null
+++ b/Engine/Source/Editor/Sequencer/Public/MVVM/ViewModels/OutlinerItemModel.h
@@ -0,0 +1,118 @@
+// Copyright Epic Games, Inc. All Rights Reserved.
+
+#pragma once
+
+#include "MVVM/ICastable.h"
+#include "MVVM/ViewModels/ViewModel.h"
+#include "MVVM/Extensions/ICurveEditorTreeItemExtension.h"
+#include "MVVM/Extensions/IDimmableExtension.h"
+#include "MVVM/Extensions/IGeometryExtension.h"
+#include "MVVM/Extensions/IHoveredExtension.h"
+#include "MVVM/Extensions/IOutlinerExtension.h"
+#include "MVVM/Extensions/IPinnableExtension.h"
+#include "MVVM/Extensions/ISoloableExtension.h"
+#include "MVVM/Extensions/IMutableExtension.h"
+#include "CurveEditorTypes.h"
+#include "Tree/ICurveEditorTreeItem.h"
+
+class UMovieSceneSequence;
+class FSequencer;
+
+namespace UE
+{
+namespace Sequencer
+{
+
+class FSequencerEditorViewModel;
+
+class SEQUENCER_API FOutlinerItemModel
+	: public FViewModel
+	, public FOutlinerExtensionShim
+	, public FGeometryExtensionShim
+	, public FPinnableExtensionShim
+	, public FHoveredExtensionShim
+	, public IDimmableExtension
+	, public ISoloableExtension
+	, public IMutableExtension
+	, public FCurveEditorTreeItemExtensionShim
+	, public ICurveEditorTreeItem
+{
+public:
+
+	UE_SEQUENCER_DECLARE_CASTABLE(FOutlinerItemModel, FViewModel, IOutlinerExtension, IGeometryExtension, IPinnableExtension, IHoveredExtension, IDimmableExtension, IMutableExtension, ISoloableExtension, ICurveEditorTreeItemExtension);
+
+	FOutlinerItemModel();
+
+	TSharedPtr<FSequencerEditorViewModel> GetEditor() const;
+
+	bool IsForceFilteredOut() const;
+	void SetForceFilteredOut(bool bInIsForcedFilteredOut);
+	
+	/*~ IOutlinerExtension */
+	FName GetIdentifier() const override;
+	bool IsExpanded() const override;
+	void SetExpansion(bool bInIsExpanded) override;
+	bool IsFilteredOut() const override;
+	TSharedPtr<SWidget> CreateContextMenuWidget(const FCreateOutlinerContextMenuWidgetParams& InParams) override;
+
+	/*~ ICurveEditorTreeItemExtension */
+	virtual bool HasCurves() const override;
+	virtual TSharedPtr<ICurveEditorTreeItem> GetCurveEditorTreeItem() const override;
+
+	/*~ ICurveEditorTreeItem */
+	virtual TSharedPtr<SWidget> GenerateCurveEditorTreeWidget(const FName& InColumnName, TWeakPtr<FCurveEditor> InCurveEditor, FCurveEditorTreeItemID InTreeItemID, const TSharedRef<ITableRow>& InTableRow) override;
+	virtual void CreateCurveModels(TArray<TUniquePtr<FCurveModel>>& OutCurveModels) override;
+	virtual bool PassesFilter(const FCurveEditorTreeFilter* InFilter) const override;
+
+	/*~ IPinnableExtension */
+	bool IsPinned() const override;
+
+	/*~ IDimmableExtension */
+	bool IsDimmed() const override;
+
+	/*~ ISoloableExtension */
+	bool IsSolo() const override;
+
+	/*~ IMutableExtension */
+	bool IsMuted() const override;
+
+protected:
+
+	/** Get context menu contents. */
+	virtual void BuildContextMenu(FMenuBuilder& MenuBuilder);
+	virtual void BuildOrganizeContextMenu(FMenuBuilder& MenuBuilder);
+
+	/** Set identifier for computing node paths */
+	void SetIdentifier(FName InNewIdentifier);
+
+	/** Get the default expansion state if it wasn't saved in the movie-scene data */
+	virtual bool GetDefaultExpansionState() const;
+
+	/** Set expansion state without saving it in the movie-scene data */
+	void SetExpansionWithoutSaving(bool bInIsExpanded);
+
+protected:
+	FViewModelListHead OutlinerChildList;
+private:
+
+	bool IsRootModelPinned() const;
+	void ToggleRootModelPinned();
+
+	bool IsSelectedModelsSolo() const;
+	void ToggleSelectedModelsSolo();
+
+	bool IsSelectedModelsMuted() const;
+	void ToggleSelectedModelsMuted();
+
+private:
+
+	FName TreeItemIdentifier;
+	FCurveEditorTreeItemID CurveEditorItemID;
+	mutable bool bInitializedExpansion;
+	mutable bool bInitializedPinnedState;
+	bool bIsForcedFilteredOut;
+};
+
+} // namespace Sequencer
+} // namespace UE
+
diff --git a/Engine/Source/Editor/Sequencer/Public/MVVM/ViewModels/SectionModel.h b/Engine/Source/Editor/Sequencer/Public/MVVM/ViewModels/SectionModel.h
new file mode 100644
--- /dev/null
+++ b/Engine/Source/Editor/Sequencer/Public/MVVM/ViewModels/SectionModel.h
@@ -0,0 +1,137 @@
+// Copyright Epic Games, Inc. All Rights Reserved.
+
+#pragma once
+
+#include "MVVM/ViewModels/ViewModel.h"
+#include "MVVM/ViewModels/ViewModelHierarchy.h"
+#include "MVVM/Extensions/ILayerBarExtension.h"
+#include "MVVM/Extensions/ITrackLaneExtension.h"
+#include "MVVM/Extensions/IGeometryExtension.h"
+#include "MVVM/Extensions/ISelectableExtension.h"
+#include "MVVM/Extensions/ISnappableExtension.h"
+#include "MVVM/Extensions/IDraggableTrackAreaExtension.h"
+#include "MVVM/Extensions/IStretchableExtension.h"
+#include "Delegates/DelegateCombinations.h"
+#include "EventHandlers/ISignedObjectEventHandler.h"
+#include "EventHandlers/ISectionEventHandler.h"
+
+class ISequencerSection;
+class UMovieSceneSection;
+
+namespace UE
+{
+namespace Sequencer
+{
+
+class ITrackExtension;
+struct FOverlappingSections;
+
+/**
+ * Model for a sequencer section
+ */
+class FSectionModel
+	: public FViewModel
+	, public ILayerBarExtension
+	, public ITrackLaneExtension
+	, public ISelectableExtension
+	, public ISnappableExtension
+	, public IDraggableTrackAreaExtension
+	, public IStretchableExtension
+	, public UE::MovieScene::TIntrusiveEventHandler<UE::MovieScene::ISignedObjectEventHandler>
+	, public UE::MovieScene::TIntrusiveEventHandler<UE::MovieScene::ISectionEventHandler>
+{
+public:
+
+	UE_SEQUENCER_DECLARE_CASTABLE(FSectionModel, FViewModel
+		, ILayerBarExtension
+		, ITrackLaneExtension
+		, ISelectableExtension
+		, ISnappableExtension
+		, IDraggableTrackAreaExtension
+		, IStretchableExtension
+	);
+
+	DECLARE_MULTICAST_DELEGATE_OneParam(FOnModelUpdated, FSectionModel*)
+
+	FOnModelUpdated OnUpdated;
+
+	explicit FSectionModel(UMovieSceneSection* InSection, TSharedPtr<ISequencerSection> InSectionInterface);
+	~FSectionModel();
+
+	/** Gets the time range of the section */
+	TRange<FFrameNumber> GetRange() const;
+
+	/** Gets the section object */
+	UMovieSceneSection* GetSection() const;
+
+	/** Gets the section interface */
+	TSharedPtr<ISequencerSection> GetSectionInterface() const;
+
+	/** Gets the parent track or track row model */
+	TViewModelPtr<ITrackExtension> GetParentTrackModel() const;
+
+	/** Gets the parent track or track row model as an ITrackExtension */
+	TViewModelPtr<ITrackExtension> GetParentTrackExtension() const;
+
+	/** Gets an array of sections that underlap the specified section */
+	TArray<FOverlappingSections> GetUnderlappingSections();
+	/** Gets an array of sections whose easing bounds underlap the specified section */
+	TArray<FOverlappingSections> GetEasingSegments();
+
+public:
+
+	/*~ ILayerBarExtension Interface */
+	TRange<FFrameNumber> GetLayerBarRange() const override;
+	void OffsetLayerBar(FFrameNumber Amount) override;
+
+	/*~ ITrackLaneExtension Interface */
+	TSharedPtr<ITrackLaneWidget> CreateTrackLaneView(const FCreateTrackLaneViewParams& InParams) override;
+	FTrackLaneVirtualAlignment ArrangeVirtualTrackLaneView() const override;
+
+	/*~ ISelectableExtension Interface */
+	ESelectionIntent IsSelectable() const override;
+
+	/*~ ISnappableExtension Interface */
+	void AddToSnapField(const ISnapCandidate& Candidate, ISnapField& SnapField) const override;
+
+	/*~ ISignedObjectEventHandler Interface */
+	void OnModifiedDirectly(UMovieSceneSignedObject*) override;
+
+	/*~ ISectionEventHandler Interface */
+	void OnRowChanged(UMovieSceneSection*) override;
+
+	/*~ IDraggableTrackAreaExtension Interface */
+	bool CanDrag() const override;
+	void OnBeginDrag(IDragOperation& DragOperation) override;
+	void OnEndDrag(IDragOperation& DragOperation) override;
+
+	/*~ IStretchableExtension Interface */
+	void OnInitiateStretch(IStretchOperation& StretchOperation, EStretchConstraint Constraint, FStretchParameters* InOutGlobalParameters) override;
+	EStretchResult OnBeginStretch(const IStretchOperation& StretchOperation, const FStretchScreenParameters& ScreenParameters, FStretchParameters* InOutParameters) override;
+	void OnStretch(const IStretchOperation& StretchOperation, const FStretchScreenParameters& ScreenParameters, FStretchParameters* InOutParameters) override;
+	void OnEndStretch(const IStretchOperation& StretchOperation, const FStretchScreenParameters& ScreenParameters, FStretchParameters* InOutParameters) override;
+
+private:
+
+	void UpdateCachedData();
+
+private:
+
+	FViewModelListHead ChannelList;
+	TSharedPtr<ISequencerSection> SectionInterface;
+	TWeakObjectPtr<UMovieSceneSection> WeakSection;
+	TRange<FFrameNumber> SectionRange;
+	TRange<FFrameNumber> LayerBarRange;
+};
+
+struct FOverlappingSections
+{
+	/** The range for the overlap */
+	TRange<FFrameNumber> Range;
+	/** The sections that occupy this range, sorted by overlap priority */
+	TArray<TWeakPtr<FSectionModel>> Sections;
+};
+
+} // namespace Sequencer
+} // namespace UE
+
diff --git a/Engine/Source/Editor/Sequencer/Public/MVVM/ViewModels/SequenceModel.h b/Engine/Source/Editor/Sequencer/Public/MVVM/ViewModels/SequenceModel.h
new file mode 100644
--- /dev/null
+++ b/Engine/Source/Editor/Sequencer/Public/MVVM/ViewModels/SequenceModel.h
@@ -0,0 +1,97 @@
+// Copyright Epic Games, Inc. All Rights Reserved.
+
+#pragma once
+
+#include "Templates/SharedPointer.h"
+#include "Containers/SortedMap.h"
+#include "Misc/Guid.h"
+#include "Delegates/DelegateCombinations.h"
+#include "UObject/WeakObjectPtr.h"
+#include "UObject/WeakObjectPtrTemplates.h"
+
+#include "MovieSceneSequenceID.h"
+#include "EventHandlers/ISignedObjectEventHandler.h"
+
+#include "SequencerCoreFwd.h"
+#include "MVVM/ViewModels/ViewModel.h"
+#include "MVVM/Extensions/IOutlinerExtension.h"
+#include "MVVM/Extensions/IGeometryExtension.h"
+#include "MVVM/Extensions/ISortableExtension.h"
+#include "MVVM/Extensions/ICurveEditorTreeItemExtension.h"
+
+class ISequencer;
+class FSequencer;
+class FCurveEditor;
+class UMovieScene;
+class UMovieSceneSequence;
+
+namespace UE
+{
+namespace Sequencer
+{
+
+class FSequenceModel;
+class FSequencerEditorViewModel;
+
+DECLARE_MULTICAST_DELEGATE_OneParam(FOnInitializeSequenceModel, FSequenceModel*);
+
+class SEQUENCER_API FSequenceModel
+	: public FViewModel
+	, public ISortableExtension
+	, public ICurveEditorExtension
+	, public IOutlinerDropTargetOutlinerExtension
+	, public UE::MovieScene::ISignedObjectEventHandler
+{
+public:
+
+	UE_SEQUENCER_DECLARE_CASTABLE(FSequenceModel, FViewModel, ISortableExtension, ICurveEditorExtension, IOutlinerDropTargetOutlinerExtension);
+
+	static FOnInitializeSequenceModel CreateExtensionsEvent;
+
+	FSequenceModel(TWeakPtr<FSequencerEditorViewModel> InEditorViewModel);
+
+	void InitializeExtensions();
+
+	FMovieSceneSequenceID GetSequenceID() const;
+	UMovieSceneSequence* GetSequence() const;
+	UMovieScene* GetMovieScene() const;
+	void SetSequence(UMovieSceneSequence* InSequence, FMovieSceneSequenceID InSequenceID);
+
+	TSharedPtr<FSequencerEditorViewModel> GetEditor() const;
+	TSharedPtr<ISequencer> GetSequencer() const;
+	TSharedPtr<FSequencer> GetSequencerImpl() const;
+
+	TSharedPtr<FViewModel> GetBottomSpacer() const { return BottomSpacer; }
+
+	/*~ ISignedObjectEventHandler */
+	void OnPostUndo() override;
+
+	/*~ ISortableExtension */
+	void SortChildren() override;
+	FSortingKey GetSortingKey() const override;
+	void SetCustomOrder(int32 InCustomOrder) override;
+
+	/*~ ICurveEditorExtension */
+	TSharedPtr<FCurveEditor> GetCurveEditor() const override;
+
+	/*~ IOutlinerDropTargetOutlinerExtension */
+	TOptional<EItemDropZone> CanAcceptDrop(const FViewModelPtr& TargetModel, const FDragDropEvent& DragDropEvent, EItemDropZone InItemDropZone) override;
+	void PerformDrop(const FViewModelPtr& TargetModel, const FDragDropEvent& DragDropEvent, EItemDropZone InItemDropZone) override;
+
+private:
+
+	FViewModelListHead RootOutlinerItems;
+
+	TWeakObjectPtr<UMovieSceneSequence> WeakSequence;
+	TWeakPtr<FSequencerEditorViewModel> WeakEditor;
+
+	TSharedPtr<FViewModel> BottomSpacer;
+	FMovieSceneSequenceID SequenceID;
+
+	MovieScene::TNonIntrusiveEventHandler<MovieScene::ISignedObjectEventHandler> SequenceEventHandler;
+	MovieScene::TNonIntrusiveEventHandler<MovieScene::ISignedObjectEventHandler> MovieSceneEventHandler;
+};
+
+} // namespace Sequencer
+} // namespace UE
+
diff --git a/Engine/Source/Editor/Sequencer/Public/MVVM/ViewModels/SequencerEditorViewModel.h b/Engine/Source/Editor/Sequencer/Public/MVVM/ViewModels/SequencerEditorViewModel.h
new file mode 100644
--- /dev/null
+++ b/Engine/Source/Editor/Sequencer/Public/MVVM/ViewModels/SequencerEditorViewModel.h
@@ -0,0 +1,52 @@
+// Copyright Epic Games, Inc. All Rights Reserved.
+
+#pragma once
+
+#include "MVVM/ViewModels/EditorViewModel.h"
+
+struct FMovieSceneSequenceID;
+
+class ISequencer;
+class FSequencer;
+class UMovieSceneSequence;
+
+namespace UE
+{
+namespace Sequencer
+{
+
+class FSequenceModel;
+
+/**
+ * Main view-model for the Sequencer editor.
+ */
+class FSequencerEditorViewModel : public FEditorViewModel
+{
+public:
+
+	UE_SEQUENCER_DECLARE_CASTABLE(FSequencerEditorViewModel, FEditorViewModel);
+
+	FSequencerEditorViewModel(TSharedRef<ISequencer> InSequencer);
+
+	// @todo_sequencer_mvvm remove this later
+	TSharedPtr<ISequencer> GetSequencer() const;
+	// @todo_sequencer_mvvm remove this ASAP
+	TSharedPtr<FSequencer> GetSequencerImpl() const;
+
+	// @todo_sequencer_mvvm move this to the root view-model
+	void SetSequence(UMovieSceneSequence* InRootSequence);
+
+protected:
+
+	virtual TSharedPtr<FViewModel> CreateRootModelImpl() override;
+	virtual TSharedPtr<FOutlinerViewModel> CreateOutlinerImpl() override;
+	virtual TSharedPtr<FTrackAreaViewModel> CreateTrackAreaImpl() override;
+
+protected:
+
+	TWeakPtr<ISequencer> WeakSequencer;
+};
+
+} // namespace Sequencer
+} // namespace UE
+
diff --git a/Engine/Source/Editor/Sequencer/Public/MVVM/ViewModels/TrackModel.h b/Engine/Source/Editor/Sequencer/Public/MVVM/ViewModels/TrackModel.h
new file mode 100644
--- /dev/null
+++ b/Engine/Source/Editor/Sequencer/Public/MVVM/ViewModels/TrackModel.h
@@ -0,0 +1,155 @@
+// Copyright Epic Games, Inc. All Rights Reserved.
+
+#pragma once
+
+#include "MVVM/ViewModels/ViewModel.h"
+#include "MVVM/ViewModels/OutlinerItemModel.h"
+#include "MVVM/Extensions/IGeometryExtension.h"
+#include "MVVM/Extensions/IGroupableExtension.h"
+#include "MVVM/Extensions/IRenameableExtension.h"
+#include "MVVM/Extensions/IResizableExtension.h"
+#include "MVVM/Extensions/ISortableExtension.h"
+#include "MVVM/Extensions/ITrackAreaExtension.h"
+#include "MVVM/Extensions/ITrackExtension.h"
+#include "MVVM/Extensions/IDeletableExtension.h"
+#include "MVVM/Extensions/IDraggableOutlinerExtension.h"
+#include "MVVM/Extensions/IDeletableExtension.h"
+#include "MVVM/ViewModels/ViewModelHierarchy.h"
+#include "EventHandlers/ISignedObjectEventHandler.h"
+
+#include "MovieSceneSignedObject.h"
+
+#include "UObject/WeakObjectPtr.h"
+#include "UObject/WeakObjectPtrTemplates.h"
+
+class ISequencer;
+class ISequencerSection;
+class ISequencerTrackEditor;
+class UMovieSceneTrack;
+
+namespace UE
+{
+namespace Sequencer
+{
+
+class FSectionModel;
+
+class SEQUENCER_API FTrackModel
+	: public FOutlinerItemModel
+	, public IRenameableExtension
+	, public IResizableExtension
+	, public ITrackExtension
+	, public ITrackAreaExtension
+	, public IGroupableExtension
+	, public ISortableExtension
+	, public IDraggableOutlinerExtension
+	, public IDeletableExtension
+	, public UE::MovieScene::TIntrusiveEventHandler<UE::MovieScene::ISignedObjectEventHandler>
+	, public UE::MovieScene::IDeferredSignedObjectFlushSignal
+{
+public:
+
+	UE_SEQUENCER_DECLARE_CASTABLE(FTrackModel
+		, FOutlinerItemModel
+		, IRenameableExtension
+		, IResizableExtension
+		, ITrackExtension
+		, ITrackAreaExtension
+		, IGroupableExtension
+		, ISortableExtension
+		, IDraggableOutlinerExtension
+		, IDeletableExtension);
+
+	explicit FTrackModel(UMovieSceneTrack* Track);
+	~FTrackModel();
+
+	FViewModelChildren GetSections();
+	FViewModelChildren GetTopLevelChannels();
+
+public:
+
+	static EViewModelListType GetTopLevelChannelType();
+	static EViewModelListType GetTopLevelChannelGroupType();
+
+	/*~ ITrackExtension */
+	UMovieSceneTrack* GetTrack() const override;
+	int32 GetRowIndex() const override;
+	TSharedPtr<ISequencerTrackEditor> GetTrackEditor() const override;
+
+	/*~ ISignedObjectEventHandler */
+	void OnModifiedDirectly(UMovieSceneSignedObject*) override;
+	void OnModifiedIndirectly(UMovieSceneSignedObject*) override;
+
+	/* IDeferredSignedObjectFlushSignal */
+	void OnDeferredModifyFlush() override;
+
+	/*~ IOutlinerExtension */
+	FOutlinerSizing GetOutlinerSizing() const override;
+	TSharedRef<SWidget> CreateOutlinerView(const FCreateOutlinerViewParams& InParams) override;
+	FSlateFontInfo GetLabelFont() const override;
+	const FSlateBrush* GetIconBrush() const override;
+	FText GetLabel() const override;
+	FSlateColor GetLabelColor() const override;
+	
+	/*~ IDimmableExtension */
+	bool IsDimmed() const override;
+
+	/*~ IResizableExtension */
+	bool IsResizable() const override;
+	void Resize(float NewSize) override;
+
+	/*~ ITrackAreaExtension */
+	FTrackAreaParameters GetTrackAreaParameters() const override;
+	FViewModelVariantIterator GetTrackAreaModelList() const override;
+
+	/*~ ICurveEditorTreeItem */
+	void CreateCurveModels(TArray<TUniquePtr<FCurveModel>>& OutCurveModels) override;
+
+	/*~ IGroupableIdentifier */
+	void GetIdentifierForGrouping(TStringBuilder<128>& OutString) const override;
+
+	/*~ IRenameableExtension */
+	bool CanRename() const override;
+	void Rename(const FText& NewName) override;
+	bool IsRenameValidImpl(const FText& NewName, FText& OutErrorMessage) const override;
+
+	/*~ ISortableExtension */
+	void SortChildren() override;
+	FSortingKey GetSortingKey() const override;
+	void SetCustomOrder(int32 InCustomOrder) override;
+
+	/*~ IDraggableOutlinerExtension */
+	bool CanDrag() const override;
+
+	/*~ FOutlinerItemModel */
+	bool HasCurves() const override;
+	void BuildContextMenu(FMenuBuilder& MenuBuilder) override;
+	bool GetDefaultExpansionState() const override;
+
+	/*~ IDeletableExtension */
+	bool CanDelete(FText* OutErrorMessage) const override;
+	void Delete() override;
+
+private:
+
+	/*~ FViewModel interface */
+	void OnConstruct() override;
+
+	void ForceUpdate();
+
+	/** A second children list for the sections inside this track */
+	FViewModelListHead SectionList;
+	FViewModelListHead TopLevelChannelList;
+
+	/** The actual track wrapped by this data model */
+	TWeakObjectPtr<UMovieSceneTrack> WeakTrack;
+
+	// @todo_sequencer_mvvm: move all the track editor behavior into the view model
+	TSharedPtr<ISequencerTrackEditor> TrackEditor;
+
+	bool bNeedsUpdate;
+};
+
+} // namespace Sequencer
+} // namespace UE
+
diff --git a/Engine/Source/Editor/Sequencer/Public/MVVM/ViewModels/TrackRowModel.h b/Engine/Source/Editor/Sequencer/Public/MVVM/ViewModels/TrackRowModel.h
new file mode 100644
--- /dev/null
+++ b/Engine/Source/Editor/Sequencer/Public/MVVM/ViewModels/TrackRowModel.h
@@ -0,0 +1,92 @@
+// Copyright Epic Games, Inc. All Rights Reserved.
+
+#pragma once
+
+#include "MVVM/ViewModels/ViewModel.h"
+#include "MVVM/ViewModels/OutlinerItemModel.h"
+#include "MVVM/Extensions/ITrackExtension.h"
+#include "MVVM/Extensions/ITrackAreaExtension.h"
+#include "MVVM/Extensions/IResizableExtension.h"
+#include "MVVM/Extensions/IDeletableExtension.h"
+
+#include "UObject/WeakObjectPtr.h"
+#include "UObject/WeakObjectPtrTemplates.h"
+
+class ISequencer;
+class ISequencerTrackEditor;
+class UMovieSceneTrack;
+
+namespace UE
+{
+namespace Sequencer
+{
+
+class FSectionModel;
+
+class FTrackRowModel
+	: public FOutlinerItemModel
+	, public ITrackAreaExtension
+	, public ITrackExtension
+	, public IResizableExtension
+	, public IDeletableExtension
+{
+public:
+
+	UE_SEQUENCER_DECLARE_CASTABLE(FTrackRowModel, FOutlinerItemModel, ITrackAreaExtension, ITrackExtension, IDeletableExtension);
+
+	explicit FTrackRowModel(UMovieSceneTrack* InTrack, int32 InRowIndex);
+	~FTrackRowModel();
+
+	void Initialize();
+
+	FViewModelChildren GetSections();
+
+	TSharedPtr<ISequencerTrackEditor> GetTrackEditor() const { return TrackEditor; }
+
+	/*~ FOutlinerItemModel */
+	void BuildContextMenu(FMenuBuilder& MenuBuilder) override;
+	void CreateCurveModels(TArray<TUniquePtr<FCurveModel>>& OutCurveModels) override;
+
+	/*~ IOutlinerExtension */
+	FOutlinerSizing GetOutlinerSizing() const override;
+	TSharedRef<SWidget> CreateOutlinerView(const FCreateOutlinerViewParams& InParams) override;
+	FText GetLabel() const override;
+	FSlateColor GetLabelColor() const override;
+	FSlateFontInfo GetLabelFont() const override;
+	const FSlateBrush* GetIconBrush() const override;
+
+	/*~ IResizableExtension */
+	bool IsResizable() const override;
+	void Resize(float NewSize) override;
+
+	/*~ ITrackExtension */
+	UMovieSceneTrack* GetTrack() const override;
+	int32 GetRowIndex() const override;
+
+	/*~ ITrackAreaExtension */
+	FTrackAreaParameters GetTrackAreaParameters() const override;
+	FViewModelVariantIterator GetTrackAreaModelList() const override;
+
+	/*~ IDimmableExtension */
+	bool IsDimmed() const override;
+
+	/*~ IDeletableExtension */
+	bool CanDelete(FText* OutErrorMessage) const override;
+	void Delete() override;
+
+private:
+
+	FViewModelListHead SectionList;
+	FViewModelListHead TopLevelChannelList;
+	TWeakObjectPtr<UMovieSceneTrack> WeakTrack;
+	int32 RowIndex;
+
+	// @todo_sequencer_mvvm: move all the track editor behavior into the view model
+	TSharedPtr<ISequencerTrackEditor> TrackEditor;
+
+	friend class FTrackModel;
+};
+
+} // namespace Sequencer
+} // namespace UE
+
diff --git a/Engine/Source/Editor/Sequencer/Public/MovieSceneTrackEditor.h b/Engine/Source/Editor/Sequencer/Public/MovieSceneTrackEditor.h
--- a/Engine/Source/Editor/Sequencer/Public/MovieSceneTrackEditor.h
+++ b/Engine/Source/Editor/Sequencer/Public/MovieSceneTrackEditor.h
@@ -184,8 +184,6 @@
 	virtual void OnInitialize() override;
 	virtual void OnRelease() override;
 
-	virtual int32 PaintTrackArea(const FPaintArgs& Args, const FGeometry& AllottedGeometry, const FSlateRect& MyCullingRect, FSlateWindowElementList& OutDrawElements, int32 LayerId, const FWidgetStyle& InWidgetStyle);
-
 	virtual bool SupportsType( TSubclassOf<class UMovieSceneTrack> TrackClass ) const override = 0;
 	virtual bool SupportsSequence(UMovieSceneSequence* InSequence) const override { return true; }
 	virtual void Tick(float DeltaTime) override;
diff --git a/Engine/Source/Editor/Sequencer/Public/SequencerAddKeyOperation.h b/Engine/Source/Editor/Sequencer/Public/SequencerAddKeyOperation.h
--- a/Engine/Source/Editor/Sequencer/Public/SequencerAddKeyOperation.h
+++ b/Engine/Source/Editor/Sequencer/Public/SequencerAddKeyOperation.h
@@ -5,6 +5,7 @@
 #include "SequencerKeyParams.h"
 #include "Containers/Array.h"
 #include "Misc/FrameNumber.h"
+#include "MVVM/ViewModelPtr.h"
 
 class IKeyArea;
 class ISequencer;
@@ -12,15 +13,14 @@
 
 class UMovieSceneTrack;
 
-class FSequencerTrackNode;
-class FSequencerDisplayNode;
-class FSequencerSectionKeyAreaNode;
-
 namespace UE
 {
 namespace Sequencer
 {
 
+class ITrackExtension;
+class FViewModel;
+class FChannelGroupModel;
 
 /**
  * Temporary structure used for consistent add-key behavior for a set of display nodes
@@ -31,19 +31,15 @@
 	/**
 	 * Construct an operation from any set of display nodes. Each node in the set will receive keys for all decendant key areas.
 	 *
-	 * @param InNodes A set of all the nodes to key
-	 * @param bJustVisible Only do it on visible key area nodes, so we can respect filtering when keying
+	 * @param InModels A set of all the models to key
 	 */
-	static FAddKeyOperation FromNodes(const TSet<TSharedRef<FSequencerDisplayNode>>& InNodes, bool bJustVisible = false);
+	static FAddKeyOperation FromNodes(const TSet<TWeakPtr<FViewModel>>& InModels);
 
 
 	/**
 	 * Construct an operation from a single display node. Every key area underneath this node will receive keys.
-	 * 
-	 * @param InNode Node to key
-	 * @param bJustVisible Only do it on visible key area nodes, so we can respect filtering when keying
 	 */
-	static FAddKeyOperation FromNode(TSharedRef<FSequencerDisplayNode> InNode, bool bJustVisible = false);
+	static FAddKeyOperation FromNode(TWeakPtr<FViewModel> InModel);
 
 
 	/**
@@ -66,19 +62,17 @@
 	 * Add a set of nodes to this operation that have already had child nodes removed (ie only parent nodes should exist in the set)
 	 *
 	 * @param InNodes     A set of nodes to add to this operation that contains no child nodes
-	 * @param bJustVisible Only do it on visible key area nodes, so we can respect filtering when keying
 	 */
-	void AddPreFilteredNodes(TArrayView<const TSharedRef<FSequencerDisplayNode>> InNodes, bool bJustVisible);
+	void AddPreFilteredNodes(TArrayView<const TWeakPtr<FViewModel>> InNodes);
 
 
 	/**
 	 * Add any keyable areas to the list of potential things to key
 	 *
-	 * @param InTrackNode           The current track node
+	 * @param InTrackEditor         The track editor
 	 * @param InKeyAnythingBeneath  A node to search within for key areas
-	 * @param bJustVisible Only do it on visible key area nodes, so we can respect filtering when keying
 	 */
-	bool ConsiderKeyableAreas(FSequencerTrackNode* InTrackNode, FSequencerDisplayNode* InKeyAnythingBeneath, bool bJustVisible);
+	bool ConsiderKeyableAreas(TSharedPtr<ITrackExtension> InTrackModel, FViewModelPtr InKeyAnythingBeneath);
 
 
 	/**
@@ -86,23 +80,13 @@
 	 *
 	 * @param InTrackNode         The current track node
 	 * @param InKeyAreaNode       The key area node to add key areas from
-	 * @param bJustVisible Only do it on visible key area nodes, so we can respect filtering when keying
 	 */
-	bool ProcessKeyAreaNode(FSequencerTrackNode* InTrackNode, const FSequencerSectionKeyAreaNode* InKeyAreaNode, bool bJustVisible);
+	bool ProcessKeyArea(TSharedPtr<ITrackExtension> InTrackModel, TViewModelPtr<FChannelGroupModel> InChannelGroupModel);
 
 
 	/**
 	 * Add a key area to this operation
 	 *
-	 * @param InTrackNode         The current track node
-	 * @param InKeyArea           The key area to add
-	 */
-	bool ProcessKeyArea(FSequencerTrackNode* InTrackNode, TSharedPtr<IKeyArea> InKeyArea);
-
-
-	/**
-	 * Add a key area to this operation
-	 *
 	 * @param InTrackEditor       The track editor responsible for the key area
 	 * @param InKeyArea           The key area to add
 	 */
@@ -122,6 +106,6 @@
 	TMap<ISequencerTrackEditor*, FKeyOperation> OperationsByTrackEditor;
 };
 
-
 } // namespace Sequencer
 } // namespace UE
+
diff --git a/Engine/Source/Editor/Sequencer/Public/SequencerClipboardReconciler.h b/Engine/Source/Editor/Sequencer/Public/SequencerClipboardReconciler.h
--- a/Engine/Source/Editor/Sequencer/Public/SequencerClipboardReconciler.h
+++ b/Engine/Source/Editor/Sequencer/Public/SequencerClipboardReconciler.h
@@ -6,45 +6,50 @@
 #include "Curves/KeyHandle.h"
 #include "MovieSceneClipboard.h"
 
-class IKeyArea;
+namespace UE
+{
+namespace Sequencer
+{
+	class FChannelModel;
+}
+}
 
-typedef TArray<IKeyArea*> FKeyAreaArray;
 
 struct FSequencerPasteEnvironment : FMovieSceneClipboardEnvironment
 {
-	void ReportPastedKey(FKeyHandle KeyHandle, IKeyArea& KeyArea) const
+	void ReportPastedKey(FKeyHandle KeyHandle, TSharedPtr<UE::Sequencer::FChannelModel> ChannelModel) const
 	{
 		if (OnKeyPasted)
 		{
-			OnKeyPasted(KeyHandle, KeyArea);
+			OnKeyPasted(KeyHandle, ChannelModel);
 		}
 	}
 
-	TFunction<void(FKeyHandle, IKeyArea&)> OnKeyPasted;
+	TFunction<void(FKeyHandle, TSharedPtr<UE::Sequencer::FChannelModel>)> OnKeyPasted;
 };
 
 /** Struct responsible for adding key areas to a group */
 struct FSequencerClipboardPasteGroup
 {
 	/** Constructor that takes an array to add our group of key areas to */
-	FSequencerClipboardPasteGroup(TArray<FKeyAreaArray>& InOwnerArray)
+	FSequencerClipboardPasteGroup(TArray<TArray<TSharedPtr<UE::Sequencer::FChannelModel>>>& InOwnerArray)
 		: OwnerArray(InOwnerArray)
 		, Index(INDEX_NONE)
 	{}
 
 	/** Add a key area to this group */
-	void Add(IKeyArea& InKeyArea)
+	void Add(TSharedPtr<UE::Sequencer::FChannelModel> InChannel)
 	{
 		if (Index == INDEX_NONE)
 		{
 			Index = OwnerArray.AddDefaulted();
 		}
-		OwnerArray[Index].Add(&InKeyArea);
+		OwnerArray[Index].Add(InChannel);
 	}
 
 private:
 	/** Owner array to create our group within */
-	TArray<FKeyAreaArray>& OwnerArray;
+	TArray<TArray<TSharedPtr<UE::Sequencer::FChannelModel>>>& OwnerArray;
 
 	/** Index into the above array at which our group resides */
 	int32 Index;
@@ -96,7 +101,7 @@
 	 * @param bAllowAliases	true to allow synonymous names to match, false if only exact matches should be allowed
 	 * @return true if any matches were found, false otherwise
 	 */
-	bool FindMatchingGroup(const FKeyAreaArray& Destination, const TArray<FMovieSceneClipboardKeyTrack>& Source, TMap<int32, int32>& Map, bool bAllowAliases);
+	bool FindMatchingGroup(const TArray<TSharedPtr<UE::Sequencer::FChannelModel>>& Destination, const TArray<FMovieSceneClipboardKeyTrack>& Source, TMap<int32, int32>& Map, bool bAllowAliases);
 
 	/** Reconcile a single key area group, to one or more destination groups */
 	bool ReconcileOneToMany();
@@ -133,7 +138,7 @@
 	TSharedRef<const FMovieSceneClipboard> Clipboard;
 
 	/** Array of paste destinations. A paste destination will consist of one or more key areas. */
-	TArray<FKeyAreaArray> PasteDestination;
+	TArray<TArray<TSharedPtr<UE::Sequencer::FChannelModel>>> PasteDestination;
 
 	/** Optional cached reconciliation result for the current data set */
 	TOptional<bool> bReconcileResult;
diff --git a/Engine/Source/Editor/Sequencer/Public/SequencerCommands.h b/Engine/Source/Editor/Sequencer/Public/SequencerCommands.h
--- a/Engine/Source/Editor/Sequencer/Public/SequencerCommands.h
+++ b/Engine/Source/Editor/Sequencer/Public/SequencerCommands.h
@@ -421,6 +421,9 @@
 	/** Toggle piloting the last camera or the camera cut camera */
 	TSharedPtr< FUICommandInfo > TogglePilotCamera;
 
+	/** Forcibly refresh the UI */
+	TSharedPtr< FUICommandInfo > RefreshUI;
+
 	/**
 	 * Initialize commands
 	 */
diff --git a/Engine/Source/Editor/Sequencer/Public/SequencerKeyCollection.h b/Engine/Source/Editor/Sequencer/Public/SequencerKeyCollection.h
--- a/Engine/Source/Editor/Sequencer/Public/SequencerKeyCollection.h
+++ b/Engine/Source/Editor/Sequencer/Public/SequencerKeyCollection.h
@@ -7,12 +7,20 @@
 #include "Misc/FrameNumber.h"
 #include "Misc/Guid.h"
 #include "Misc/OptionalFwd.h"
+#include "MVVM/ViewModels/ViewModelHierarchy.h"
 #include "Templates/SharedPointer.h"
 
 class UMovieSceneSection;
-class FSequencerDisplayNode;
 class IKeyArea;
 
+namespace UE
+{
+namespace Sequencer
+{
+	class FViewModel;
+}
+}
+
 template<typename> class TRange;
 
 /** Enumeration used to define how to search for keys */
@@ -23,17 +31,19 @@
 
 struct FSequencerKeyCollectionSignature
 {
+	using FViewModel = UE::Sequencer::FViewModel;
+
 	FSequencerKeyCollectionSignature()
 	{}
 
 	/** Initialize this key collection from the specified nodes. Only gathers keys from those explicitly specified. */
-	SEQUENCER_API static FSequencerKeyCollectionSignature FromNodes(const TArray<FSequencerDisplayNode*>& InNodes, FFrameNumber InDuplicateThreshold);
+	SEQUENCER_API static FSequencerKeyCollectionSignature FromNodes(const TArray<TSharedRef<FViewModel>>& InNodes, FFrameNumber InDuplicateThreshold);
 
 	/** Initialize this key collection from the specified nodes. Gathers keys from all child nodes. */
-	SEQUENCER_API static FSequencerKeyCollectionSignature FromNodesRecursive(const TArray<FSequencerDisplayNode*>& InNodes, FFrameNumber InDuplicateThreshold);
+	SEQUENCER_API static FSequencerKeyCollectionSignature FromNodesRecursive(const TArray<TSharedRef<FViewModel>>& InNodes, FFrameNumber InDuplicateThreshold);
 
 	/** Initialize this key collection from the specified node and section index. */
-	SEQUENCER_API static FSequencerKeyCollectionSignature FromNodeRecursive(FSequencerDisplayNode& InNode, UMovieSceneSection* InSection, FFrameNumber InDuplicateThreshold);
+	SEQUENCER_API static FSequencerKeyCollectionSignature FromNodeRecursive(TSharedRef<FViewModel> InNode, UMovieSceneSection* InSection, FFrameNumber InDuplicateThreshold);
 
 	/** Compare this signature for inequality with another */
 	SEQUENCER_API friend bool operator!=(const FSequencerKeyCollectionSignature& A, const FSequencerKeyCollectionSignature& B);
diff --git a/Engine/Source/Editor/Sequencer/Public/SequencerObjectBindingDragDropOp.h b/Engine/Source/Editor/Sequencer/Public/SequencerObjectBindingDragDropOp.h
--- a/Engine/Source/Editor/Sequencer/Public/SequencerObjectBindingDragDropOp.h
+++ b/Engine/Source/Editor/Sequencer/Public/SequencerObjectBindingDragDropOp.h
@@ -1,25 +1,28 @@
-// Copyright Epic Games, Inc. All Rights Reserved.
-
-#pragma once
-
-#include "CoreMinimal.h"
-#include "Input/DragAndDrop.h"
-#include "GraphEditorDragDropAction.h"
-
-namespace UE
-{
-namespace MovieScene
-{
-
-struct FFixedObjectBindingID;
-
-} // namespace MovieScene
-} // namespace UE
-
-class FSequencerObjectBindingDragDropOp : public FGraphEditorDragDropAction
-{
-public:
-	DRAG_DROP_OPERATOR_TYPE( FSequencerDisplayNodeDragDropOpBase, FGraphEditorDragDropAction )
-
-	virtual TArray<UE::MovieScene::FFixedObjectBindingID> GetDraggedBindings() const = 0;
-};
\ No newline at end of file
+// Copyright Epic Games, Inc. All Rights Reserved.
+
+#pragma once
+
+#include "MVVM/ViewModels/OutlinerViewModelDragDropOp.h"
+
+namespace UE
+{
+namespace MovieScene
+{
+	struct FFixedObjectBindingID;
+} // namespace MovieScene
+
+namespace Sequencer
+{
+
+class FSequencerObjectBindingDragDropOp : public FOutlinerViewModelDragDropOp
+{
+public:
+	DRAG_DROP_OPERATOR_TYPE( FSequencerObjectBindingDragDropOp, FOutlinerViewModelDragDropOp )
+
+	virtual TArray<UE::MovieScene::FFixedObjectBindingID> GetDraggedBindings() const = 0;
+	virtual TArray<UE::MovieScene::FFixedObjectBindingID> GetDraggedRebindableBindings() const = 0;
+};
+
+
+} // namespace Sequencer
+} // namespace UE
\ No newline at end of file
diff --git a/Engine/Source/Editor/Sequencer/Public/SequencerSectionPainter.h b/Engine/Source/Editor/Sequencer/Public/SequencerSectionPainter.h
--- a/Engine/Source/Editor/Sequencer/Public/SequencerSectionPainter.h
+++ b/Engine/Source/Editor/Sequencer/Public/SequencerSectionPainter.h
@@ -2,7 +2,7 @@
 
 #pragma once
 
-#include "CoreMinimal.h"
+#include "CoreTypes.h"
 #include "Layout/SlateRect.h"
 #include "Layout/Geometry.h"
 
@@ -12,6 +12,15 @@
 struct FTimeToPixel;
 class IKeyArea;
 
+namespace UE
+{
+namespace Sequencer
+{
+	class FSectionModel;
+}
+}
+
+
 struct FKeyAreaElement
 {
 	TArray<TSharedRef<IKeyArea>, TInlineAllocator<1>> KeyAreas;
@@ -28,7 +37,7 @@
 {
 public:
 	/** Constructor */
-	FSequencerSectionPainter(FSlateWindowElementList& OutDrawElements, const FGeometry& InSectionGeometry, UMovieSceneSection& Section);
+	FSequencerSectionPainter(FSlateWindowElementList& OutDrawElements, const FGeometry& InSectionGeometry, TSharedPtr<UE::Sequencer::FSectionModel> Section);
 
 	/** Virtual destructor */
 	virtual ~FSequencerSectionPainter();
@@ -53,8 +62,12 @@
 public:
 
 	/** The movie scene section we're painting */
+	UE_DEPRECATED(5.1, "Please use SectionModel->GetSection() instead")
 	UMovieSceneSection& Section;
 
+	/** The section we are painting */
+	TSharedPtr<UE::Sequencer::FSectionModel> SectionModel;
+
 	/** List of slate draw elements - publicly modifiable */
 	FSlateWindowElementList& DrawElements;
 
diff --git a/Engine/Source/Editor/Sequencer/Sequencer.Build.cs b/Engine/Source/Editor/Sequencer/Sequencer.Build.cs
--- a/Engine/Source/Editor/Sequencer/Sequencer.Build.cs
+++ b/Engine/Source/Editor/Sequencer/Sequencer.Build.cs
@@ -7,18 +7,17 @@
 	public Sequencer(ReadOnlyTargetRules Target) : base(Target)
 	{
 		PrivateIncludePaths.AddRange(
-            new string[] {
-                "Editor/Sequencer/Private",
-                "Editor/Sequencer/Private/DisplayNodes",
+			new string[] {
+				"Editor/Sequencer/Private",
 				"Editor/UnrealEd/Private" // TODO: Fix this, for now it's needed for the fbx exporter
-	}
-        );
+				}
+			);
 
 		PublicDependencyModuleNames.AddRange(
 			new string[] {
 				"TimeManagement",
-			}
-		);
+				}
+			);
 
 		PrivateDependencyModuleNames.AddRange(
 			new string[] {
@@ -29,68 +28,68 @@
 				"ContentBrowser",
 				"Core", 
 				"CoreUObject", 
-                "CurveEditor",
-                "InputCore",
+				"CurveEditor",
+				"InputCore",
 				"Engine", 
 				"Slate", 
 				"SlateCore",
 				"SceneOutliner",
-                "EditorStyle",
+				"SequencerCore",
+				"EditorStyle",
 				"EditorFramework",
 				"UnrealEd", 
 				"MovieScene", 
 				"MovieSceneTracks", 
 				"MovieSceneTools", 
 				"MovieSceneCapture", 
-                "MovieSceneCaptureDialog", 
+				"MovieSceneCaptureDialog", 
 				"EditorWidgets", 
 				"SequencerWidgets",
 				"BlueprintGraph",
 				"LevelSequence",
 				"GraphEditor",
 				"PropertyEditor",
-                "ViewportInteraction",
-                "SerializedRecorderInterface",
+				"ViewportInteraction",
+				"SerializedRecorderInterface",
 				"ToolMenus",
 				"ToolWidgets",
-			}
-		);
+				}
+			);
 
 		CircularlyReferencedDependentModules.AddRange(
-			new string[]
-			{
+			new string[] {
 				"ViewportInteraction",
-			}
-		);
+				}
+			);
 
 		PrivateIncludePathModuleNames.AddRange(
 			new string[] {
-                "ContentBrowser",
+				"ContentBrowser",
 				"PropertyEditor",
 				"Kismet",
-                "LevelEditor",
+				"LevelEditor",
 				"MainFrame",
 				"DesktopPlatform",
-                "SerializedRecorderInterface"
-            }
-        );
+				"SerializedRecorderInterface"
+				}
+			);
 
 		PublicIncludePathModuleNames.AddRange(
 			new string[] {
 				"PropertyEditor",
-                "SceneOutliner",
-                "CurveEditor",
-                "Analytics",
-            }
-        );
+				"SceneOutliner",
+				"CurveEditor",
+				"Analytics",
+				}
+			);
 
 		DynamicallyLoadedModuleNames.AddRange(
 			new string[] {
 				"LevelEditor",
 				"WorkspaceMenuStructure",
 				"MainFrame",
-			}
-		);
+				}
+			);
 
 		CircularlyReferencedDependentModules.Add("MovieSceneTools");
 		AddEngineThirdPartyPrivateStaticDependencies(Target, "FBX");
diff --git a/Engine/Source/Editor/SequencerCore/Private/MVVM/ICastable.cpp b/Engine/Source/Editor/SequencerCore/Private/MVVM/ICastable.cpp
new file mode 100644
--- /dev/null
+++ b/Engine/Source/Editor/SequencerCore/Private/MVVM/ICastable.cpp
@@ -0,0 +1,25 @@
+// Copyright Epic Games, Inc. All Rights Reserved.
+
+#include "MVVM/ICastable.h"
+
+namespace UE
+{
+namespace Sequencer
+{
+
+void* ICastable::CastRaw(FViewModelTypeID InType)
+{
+	const void* Result = static_cast<const ICastable*>(this)->CastRaw(InType);
+	return const_cast<void*>(Result);
+}
+
+const void* ICastable::CastRaw(FViewModelTypeID InType) const
+{
+	const void* ExtensionPtr = nullptr;
+	CastImpl(InType, ExtensionPtr);
+	return ExtensionPtr;
+}
+
+} // namespace Sequencer
+} // namespace UE
+
diff --git a/Engine/Source/Editor/SequencerCore/Private/MVVM/ViewModelPtr.cpp b/Engine/Source/Editor/SequencerCore/Private/MVVM/ViewModelPtr.cpp
new file mode 100644
--- /dev/null
+++ b/Engine/Source/Editor/SequencerCore/Private/MVVM/ViewModelPtr.cpp
@@ -0,0 +1,28 @@
+// Copyright Epic Games, Inc. All Rights Reserved.
+
+#include "MVVM/ViewModelPtr.h"
+
+namespace UE
+{
+namespace Sequencer
+{
+
+FImplicitViewModelCast FViewModelPtr::ImplicitCast() const
+{
+	return FImplicitViewModelCast{Model};
+}
+
+FImplicitWeakViewModelCast FWeakViewModelPtr::ImplicitCast() const
+{
+	return FImplicitWeakViewModelCast{WeakModel};
+}
+
+FImplicitWeakViewModelPin FWeakViewModelPtr::ImplicitPin() const
+{
+	return FImplicitWeakViewModelPin{WeakModel};
+}
+
+
+} // namespace Sequencer
+} // namespace UE
+
diff --git a/Engine/Source/Editor/SequencerCore/Private/MVVM/ViewModelTypeID.cpp b/Engine/Source/Editor/SequencerCore/Private/MVVM/ViewModelTypeID.cpp
new file mode 100644
--- /dev/null
+++ b/Engine/Source/Editor/SequencerCore/Private/MVVM/ViewModelTypeID.cpp
@@ -0,0 +1,23 @@
+// Copyright Epic Games, Inc. All Rights Reserved.
+
+#include "MVVM/ViewModelTypeID.h"
+
+namespace UE
+{
+namespace Sequencer
+{
+
+FViewModelTypeID FViewModelTypeID::RegisterNew()
+{
+	static uint32 ID = 0;
+	return FViewModelTypeID(ID++);
+}
+
+FViewModelTypeID FViewModelTypeID::Invalid()
+{
+	return FViewModelTypeID();
+}
+
+} // namespace Sequencer
+} // namespace UE
+
diff --git a/Engine/Source/Editor/SequencerCore/Private/MVVM/Views/SKeyNavigationButtons.cpp b/Engine/Source/Editor/SequencerCore/Private/MVVM/Views/SKeyNavigationButtons.cpp
new file mode 100644
--- /dev/null
+++ b/Engine/Source/Editor/SequencerCore/Private/MVVM/Views/SKeyNavigationButtons.cpp
@@ -0,0 +1,205 @@
+// Copyright Epic Games, Inc. All Rights Reserved.
+
+#include "MVVM/Views/SKeyNavigationButtons.h"
+#include "Widgets/SBoxPanel.h"
+#include "Widgets/Text/STextBlock.h"
+#include "Widgets/Layout/SBorder.h"
+#include "Widgets/Input/SButton.h"
+#include "Styling/AppStyle.h"
+#include "Misc/FrameTime.h"
+#include "MVVM/ViewModels/ViewModel.h"
+#include "MVVM/ViewModels/ViewModelIterators.h"
+#include "MVVM/Extensions/IHoveredExtension.h"
+
+namespace UE
+{
+namespace Sequencer
+{
+
+void SKeyNavigationButtons::Construct(const FArguments& InArgs, const TSharedPtr<FViewModel>& InModel)
+{
+	WeakModel = InModel;
+
+	GetNavigatableTimesEvent = InArgs._GetNavigatableTimes;
+	AddKeyEvent = InArgs._OnAddKey;
+	SetTimeEvent = InArgs._OnSetTime;
+	TimeAttribute = InArgs._Time;
+
+	const FSlateBrush* NoBorder = FAppStyle::GetBrush( "NoBorder" );
+
+	TAttribute<FLinearColor> HoverTint(this, &SKeyNavigationButtons::GetHoverTint);
+
+	ChildSlot
+	[
+		SNew(SHorizontalBox)
+		
+		// Previous key slot
+		+ SHorizontalBox::Slot()
+		.VAlign(VAlign_Center)
+		.AutoWidth()
+		.Padding(3, 0, 0, 0)
+		[
+			SNew(SBorder)
+			.Padding(0)
+			.BorderImage(NoBorder)
+			.ColorAndOpacity(HoverTint)
+			[
+				SNew(SButton)
+				.ButtonStyle(FAppStyle::Get(), "HoverHintOnly")
+				.ToolTipText(InArgs._PreviousKeyToolTip)
+				.OnClicked(this, &SKeyNavigationButtons::OnPreviousKeyClicked)
+				.ForegroundColor( FSlateColor::UseForeground() )
+				.ContentPadding(2)
+				.IsFocusable(false)
+				[
+					SNew(STextBlock)
+					.Font(FAppStyle::Get().GetFontStyle("FontAwesome.7"))
+					.Text(FText::FromString(FString(TEXT("\xf060"))) /*fa-arrow-left*/)
+				]
+			]
+		]
+		// Add key slot
+		+ SHorizontalBox::Slot()
+		.VAlign(VAlign_Center)
+		.AutoWidth()
+		[
+			SNew(SBorder)
+			.Padding(0)
+			.BorderImage(NoBorder)
+			.ColorAndOpacity(HoverTint)
+			.IsEnabled(InArgs._IsEnabled)
+			[
+				SNew(SButton)
+				.ButtonStyle(FAppStyle::Get(), "HoverHintOnly")
+				.ToolTipText(InArgs._AddKeyToolTip)
+				.OnClicked(this, &SKeyNavigationButtons::OnAddKeyClicked)
+				.ForegroundColor( FSlateColor::UseForeground() )
+				.ContentPadding(2)
+				.IsFocusable(false)
+				[
+					SNew(STextBlock)
+					.Font(FAppStyle::Get().GetFontStyle("FontAwesome.7"))
+					.Text(FText::FromString(FString(TEXT("\xf055"))) /*fa-plus-circle*/)
+				]
+			]
+		]
+		// Next key slot
+		+ SHorizontalBox::Slot()
+		.VAlign(VAlign_Center)
+		.AutoWidth()
+		[
+			SNew(SBorder)
+			.Padding(0)
+			.BorderImage(NoBorder)
+			.ColorAndOpacity(HoverTint)
+			[
+				SNew(SButton)
+				.ButtonStyle(FAppStyle::Get(), "HoverHintOnly")
+				.ToolTipText(InArgs._NextKeyToolTip)
+				.OnClicked(this, &SKeyNavigationButtons::OnNextKeyClicked)
+				.ContentPadding(2)
+				.ForegroundColor( FSlateColor::UseForeground() )
+				.IsFocusable(false)
+				[
+					SNew(STextBlock)
+					.Font(FAppStyle::Get().GetFontStyle("FontAwesome.7"))
+					.Text(FText::FromString(FString(TEXT("\xf061"))) /*fa-arrow-right*/)
+				]
+			]
+		]
+	];
+}
+
+FLinearColor SKeyNavigationButtons::GetHoverTint() const
+{
+	IHoveredExtension* HoveredExtension = ICastable::CastWeakPtr<IHoveredExtension>(WeakModel);
+	return HoveredExtension && HoveredExtension->IsHovered() ? FLinearColor(1,1,1,0.9f) : FLinearColor(1,1,1,0.4f);
+}
+
+FReply SKeyNavigationButtons::OnPreviousKeyClicked()
+{
+	FFrameTime ClosestPreviousKeyDistance = FFrameTime(TNumericLimits<int32>::Max(), 0.99999f);
+	FFrameTime CurrentTime = TimeAttribute.Get();
+	TOptional<FFrameTime> PreviousTime;
+
+	TArray<FFrameNumber> AllTimes;
+	GetNavigatableTimesEvent.ExecuteIfBound(AllTimes);
+	if (AllTimes.IsEmpty())
+	{
+		return FReply::Handled();
+	}
+
+	for (FFrameNumber Time : AllTimes)
+	{
+		if (Time < CurrentTime && CurrentTime - Time < ClosestPreviousKeyDistance)
+		{
+			PreviousTime = Time;
+			ClosestPreviousKeyDistance = CurrentTime - Time;
+		}
+	}
+
+	if (!PreviousTime.IsSet() && AllTimes.Num() > 0)
+	{
+		AllTimes.Sort();
+
+		PreviousTime = AllTimes.Last();
+	}
+
+	if (PreviousTime.IsSet())
+	{
+		SetTimeEvent.ExecuteIfBound(PreviousTime.GetValue());
+	}
+	return FReply::Handled();
+}
+
+FReply SKeyNavigationButtons::OnNextKeyClicked()
+{
+	FFrameTime ClosestNextKeyDistance = FFrameTime(TNumericLimits<int32>::Max(), 0.99999f);
+	FFrameTime CurrentTime = TimeAttribute.Get();
+	TOptional<FFrameTime> NextTime;
+
+	TArray<FFrameNumber> AllTimes;
+	GetNavigatableTimesEvent.ExecuteIfBound(AllTimes);
+	if (AllTimes.IsEmpty())
+	{
+		return FReply::Handled();
+	}
+
+	for (FFrameNumber Time : AllTimes)
+	{
+		if (Time > CurrentTime && Time - CurrentTime < ClosestNextKeyDistance)
+		{
+			NextTime = Time;
+			ClosestNextKeyDistance = Time - CurrentTime;
+		}
+	}
+
+	if (!NextTime.IsSet() && AllTimes.Num() > 0)
+	{
+		AllTimes.Sort();
+
+		NextTime = AllTimes[0];
+	}
+
+	if (NextTime.IsSet())
+	{
+		SetTimeEvent.ExecuteIfBound(NextTime.GetValue());
+	}
+
+	return FReply::Handled();
+}
+
+FReply SKeyNavigationButtons::OnAddKeyClicked()
+{
+	if (TSharedPtr<FViewModel> DataModel = WeakModel.Pin())
+	{
+		FFrameTime CurrentTime = TimeAttribute.Get();
+		AddKeyEvent.ExecuteIfBound(CurrentTime, DataModel);
+	}
+
+	return FReply::Handled();
+}
+
+} // namespace Sequencer
+} // namespace UE
+
