Sequencer: Info column cleanups - header text and tooltip

diff --git a/Engine/Source/Editor/Sequencer/Private/LevelEditorSequencerIntegration.cpp b/Engine/Source/Editor/Sequencer/Private/LevelEditorSequencerIntegration.cpp
--- a/Engine/Source/Editor/Sequencer/Private/LevelEditorSequencerIntegration.cpp
+++ b/Engine/Source/Editor/Sequencer/Private/LevelEditorSequencerIntegration.cpp
@@ -1020,14 +1020,14 @@
 
 	FSceneOutlinerColumnInfo SpawnColumnInfo(ESceneOutlinerColumnVisibility::Visible, 11, 
 		FCreateSceneOutlinerColumn::CreateRaw( this, &FLevelEditorSequencerIntegration::CreateSequencerSpawnableColumn),
-		true, TOptional<float>(), LOCTEXT("SpawnableColumnName", "SpawnableInfo"));
+		true, TOptional<float>(), LOCTEXT("SpawnableColumnName", "Spawnable"));
 
 	SceneOutlinerModule.RegisterDefaultColumnType< Sequencer::FSequencerSpawnableColumn >(SpawnColumnInfo);
 	AcquiredResources.Add([=]{ this->DetachOutlinerColumn(); });
 
 	FSceneOutlinerColumnInfo ColumnInfo(ESceneOutlinerColumnVisibility::Visible, 15, 
 		FCreateSceneOutlinerColumn::CreateRaw( this, &FLevelEditorSequencerIntegration::CreateSequencerInfoColumn), 
-		true, TOptional<float>(), LOCTEXT("SequencerColumnName", "SequencerInfo"));
+		true, TOptional<float>(), LOCTEXT("SequencerColumnName", "Sequencer"));
 
 	SceneOutlinerModule.RegisterDefaultColumnType< Sequencer::FSequencerInfoColumn >(ColumnInfo);
 
diff --git a/Engine/Source/Editor/Sequencer/Private/SequencerInfoColumn.cpp b/Engine/Source/Editor/Sequencer/Private/SequencerInfoColumn.cpp
--- a/Engine/Source/Editor/Sequencer/Private/SequencerInfoColumn.cpp
+++ b/Engine/Source/Editor/Sequencer/Private/SequencerInfoColumn.cpp
@@ -76,17 +76,16 @@
 	return IDName;
 }
 
-
 FName FSequencerInfoColumn::GetColumnID()
 {
 	return GetID();
 }
 
-
 SHeaderRow::FColumn::FArguments FSequencerInfoColumn::ConstructHeaderRowColumn()
 {
 	return SHeaderRow::Column(GetColumnID())
-		.DefaultLabel(LOCTEXT("ItemLabel_HeaderText", "Sequence"))
+		.DefaultLabel(LOCTEXT("SequencerColumn", "Sequence"))
+		.DefaultTooltip(LOCTEXT("SequencerColumnTooltip", "The sequence that this actor is referenced from"))
 		.FillWidth( 5.0f );
 }
 
@@ -112,13 +111,11 @@
 	return HorizontalBox;
 }
 
-
 void FSequencerInfoColumn::PopulateSearchStrings( const ISceneOutlinerTreeItem& Item, TArray< FString >& OutSearchStrings ) const
 {
 	OutSearchStrings.Add(Item.GetDisplayString());
 }
 
-
 void FSequencerInfoColumn::SortItems(TArray<FSceneOutlinerTreeItemPtr>& OutItems, const EColumnSortMode::Type SortMode) const
 {
 	if (WeakBindingData.IsValid())
diff --git a/Engine/Source/Editor/Sequencer/Private/SequencerSpawnableColumn.cpp b/Engine/Source/Editor/Sequencer/Private/SequencerSpawnableColumn.cpp
--- a/Engine/Source/Editor/Sequencer/Private/SequencerSpawnableColumn.cpp
+++ b/Engine/Source/Editor/Sequencer/Private/SequencerSpawnableColumn.cpp
@@ -5,6 +5,8 @@
 #include "Widgets/Images/SImage.h"
 #include "GameFramework/Actor.h"
 
+#define LOCTEXT_NAMESPACE "SpawnableInfoColumn"
+
 namespace Sequencer
 {
 
@@ -27,7 +29,7 @@
 		.VAlignHeader(VAlign_Center)
 		.HAlignCell(HAlign_Left) // Note the icon has some baked in right alignment, so we are compensating here
 		.VAlignCell(VAlign_Center)
-		.DefaultTooltip(FText::FromName(GetColumnID()))
+		.DefaultTooltip(LOCTEXT("SpawnableColumnTooltip", "Whether this actor is spawned by Sequencer"))
 		.HeaderContentPadding(FMargin(0.0))
 		[
 			SNew(SImage)
@@ -59,3 +61,5 @@
 }
 
 }// End Sequencer namespace
+
+#undef LOCTEXT_NAMESPACE
