Sequencer: Added option to show/hide the status bar

diff --git a/Engine/Source/Editor/Sequencer/Private/SSequencer.cpp b/Engine/Source/Editor/Sequencer/Private/SSequencer.cpp
--- a/Engine/Source/Editor/Sequencer/Private/SSequencer.cpp
+++ b/Engine/Source/Editor/Sequencer/Private/SSequencer.cpp
@@ -951,7 +951,7 @@ void SSequencer::Construct(const FArguments& InArgs, TSharedRef<FSequencer> InSe
 								.VAlign(VAlign_Bottom)
 								[
 									SAssignNew(SequencerTreeFilterStatusBar, SSequencerTreeFilterStatusBar, InSequencer)	
-									.Visibility( EVisibility::SelfHitTestInvisible )				
+									.Visibility(this, &SSequencer::GetStatusBarVisibility)
 								]
 							]
 						]
@@ -2316,6 +2316,7 @@ TSharedRef<SWidget> SSequencer::MakeViewMenu()
 	MenuBuilder.AddMenuEntry( FSequencerCommands::Get().ToggleShowSelectedNodesOnly );
 	MenuBuilder.AddMenuEntry( FSequencerCommands::Get().ToggleCombinedKeyframes );
 	MenuBuilder.AddMenuEntry( FSequencerCommands::Get().ToggleChannelColors );
+	MenuBuilder.AddMenuEntry( FSequencerCommands::Get().ToggleShowStatusBar );
 	MenuBuilder.AddMenuEntry( FSequencerCommands::Get().ToggleShowPreAndPostRoll );
 
 	// Menu entry for zero padding
@@ -3871,6 +3872,11 @@ EVisibility SSequencer::GetTimeRangeVisibility() const
 	return GetSequencerSettings()->GetShowRangeSlider() ? EVisibility::Visible : EVisibility::Hidden;
 }
 
+EVisibility SSequencer::GetStatusBarVisibility() const
+{
+	return GetSequencerSettings()->GetShowStatusBar() ? EVisibility::SelfHitTestInvisible : EVisibility::Hidden;
+}
+
 EFrameNumberDisplayFormats SSequencer::GetTimeDisplayFormat() const
 {
 	return GetSequencerSettings()->GetTimeDisplayFormat();
diff --git a/Engine/Source/Editor/Sequencer/Private/SSequencer.h b/Engine/Source/Editor/Sequencer/Private/SSequencer.h
--- a/Engine/Source/Editor/Sequencer/Private/SSequencer.h
+++ b/Engine/Source/Editor/Sequencer/Private/SSequencer.h
@@ -581,6 +581,9 @@ private:
 	/** Gets whether or not the time range should be visible. */
 	EVisibility GetTimeRangeVisibility() const;
 
+	/** Gets whether the status bar should be visible. */
+	EVisibility GetStatusBarVisibility() const;
+
 	/** What is the preferred display format for time values. */
 	EFrameNumberDisplayFormats GetTimeDisplayFormat() const;
 
diff --git a/Engine/Source/Editor/Sequencer/Private/Sequencer.cpp b/Engine/Source/Editor/Sequencer/Private/Sequencer.cpp
--- a/Engine/Source/Editor/Sequencer/Private/Sequencer.cpp
+++ b/Engine/Source/Editor/Sequencer/Private/Sequencer.cpp
@@ -12765,6 +12765,14 @@ void FSequencer::BindCommands()
 		FCanExecuteAction::CreateLambda( []{ return true; } ),
 		FIsActionChecked::CreateLambda( [this]{ return Settings->GetShowChannelColors(); } ) );
 
+	SequencerCommandBindings->MapAction(
+		Commands.ToggleShowStatusBar,
+		FExecuteAction::CreateLambda([this] {
+			Settings->SetShowStatusBar(!Settings->GetShowStatusBar());
+		}),
+		FCanExecuteAction::CreateLambda([] { return true; }),
+		FIsActionChecked::CreateLambda([this] { return Settings->GetShowStatusBar(); }));
+
 	SequencerCommandBindings->MapAction(
 		Commands.ToggleShowSelectedNodesOnly,
 		FExecuteAction::CreateLambda( [this]{
diff --git a/Engine/Source/Editor/Sequencer/Private/SequencerCommands.cpp b/Engine/Source/Editor/Sequencer/Private/SequencerCommands.cpp
--- a/Engine/Source/Editor/Sequencer/Private/SequencerCommands.cpp
+++ b/Engine/Source/Editor/Sequencer/Private/SequencerCommands.cpp
@@ -114,7 +114,8 @@ void FSequencerCommands::RegisterCommands()
 	UI_COMMAND( OpenDirectorBlueprint, "Open Director Blueprint", "Opens the director blueprint for this sequence.", EUserInterfaceActionType::Button, FInputChord() );
 	UI_COMMAND( OpenTaggedBindingManager, "Open Binding Tag Manager", "Specifies options for tagging bindings within this sequence for external systems to reference as a persistent name.", EUserInterfaceActionType::Button, FInputChord() );
 	UI_COMMAND( OpenNodeGroupsManager, "Open Sequencer Group Manager", "Manage groups within this sequence.", EUserInterfaceActionType::Button, FInputChord());
-	UI_COMMAND( ToggleShowRangeSlider, "Range Slider", "Enables and disables showing the time range slider", EUserInterfaceActionType::ToggleButton, FInputChord() );
+	UI_COMMAND( ToggleShowRangeSlider, "Range Slider", "Enables and disables showing the time range slider", EUserInterfaceActionType::ToggleButton, FInputChord());
+	UI_COMMAND( ToggleShowStatusBar, "Status Bar", "Enables and disables showing the status bar", EUserInterfaceActionType::ToggleButton, FInputChord());
 	UI_COMMAND( ToggleIsSnapEnabled, "Enable Snapping", "Enables and disables snapping", EUserInterfaceActionType::ToggleButton, FInputChord() );
 
 	UI_COMMAND( ToggleSnapKeyTimesToInterval, "Snap to the Interval", "Snap keys to the time snapping interval", EUserInterfaceActionType::ToggleButton, FInputChord() );
diff --git a/Engine/Source/Editor/Sequencer/Private/SequencerSettings.cpp b/Engine/Source/Editor/Sequencer/Private/SequencerSettings.cpp
--- a/Engine/Source/Editor/Sequencer/Private/SequencerSettings.cpp
+++ b/Engine/Source/Editor/Sequencer/Private/SequencerSettings.cpp
@@ -47,6 +47,7 @@
 	bShowCombinedKeyframes = true;
 	bInfiniteKeyAreas = false;
 	bShowChannelColors = false;
+	bShowStatusBar = true;
 	ReduceKeysTolerance = KINDA_SMALL_NUMBER;
 	KeyAreaHeightWithCurves = SequencerLayoutConstants::KeyAreaHeight;
 	bDeleteKeysWhenTrimming = true;
@@ -581,6 +582,20 @@
 	}
 }
 
+bool USequencerSettings::GetShowStatusBar() const
+{
+	return bShowStatusBar;
+}
+
+void USequencerSettings::SetShowStatusBar(bool InbShowStatusBar)
+{
+	if (bShowStatusBar != InbShowStatusBar)
+	{
+		bShowStatusBar = InbShowStatusBar;
+		SaveConfig();
+	}
+}
+
 bool USequencerSettings::HasKeyAreaCurveExtents(const FString& ChannelName) const
 {
 	TArray<FString> ChannelsArray;
diff --git a/Engine/Source/Editor/Sequencer/Public/SequencerCommands.h b/Engine/Source/Editor/Sequencer/Public/SequencerCommands.h
--- a/Engine/Source/Editor/Sequencer/Public/SequencerCommands.h
+++ b/Engine/Source/Editor/Sequencer/Public/SequencerCommands.h
@@ -346,6 +346,9 @@ public:
 	/** Toggles whether to show channel colors in the track area. */
 	TSharedPtr< FUICommandInfo > ToggleChannelColors;
 
+	/** Toggles whether to show the status bar. */
+	TSharedPtr< FUICommandInfo > ToggleShowStatusBar;
+
 	/** Turns auto scroll on and off. */
 	TSharedPtr< FUICommandInfo > ToggleAutoScroll;
 
diff --git a/Engine/Source/Editor/Sequencer/Public/SequencerSettings.h b/Engine/Source/Editor/Sequencer/Public/SequencerSettings.h
--- a/Engine/Source/Editor/Sequencer/Public/SequencerSettings.h
+++ b/Engine/Source/Editor/Sequencer/Public/SequencerSettings.h
@@ -284,6 +284,11 @@
 	/** Set whether to show channel colors */
 	void SetShowChannelColors(bool bInShowChannelColors);
 
+	/** @return true if showing status bar */
+	bool GetShowStatusBar() const;
+	/** Set whether to show status bar */
+	void SetShowStatusBar(bool bInShowStatusBar);
+
 	/** @return Whether the given channel has curve extents */
 	bool HasKeyAreaCurveExtents(const FString& ChannelName) const;
 	/** @ Remove curve extents for the given channel */
@@ -534,6 +539,10 @@
 	UPROPERTY(config, EditAnywhere, Category = Timeline)
 	bool bShowChannelColors;
 
+	/** Enable or disable displaying the status bar for number of items. */
+	UPROPERTY(config, EditAnywhere, Category = Timeline)
+	bool bShowStatusBar;
+
 	/** The key area curve extents, stored per channel name */
 	UPROPERTY(config, EditAnywhere, Category = Timeline)
 	FString KeyAreaCurveExtents;
