Sequencer: Curve Editor:  Curve Editor Sync Isse with Filter
Since Selection ticks before any filter is applied and the curve editor tree is udpated in FSequencer::Tick we need to sync the curve editor on the next tick so the filter is applied first, otherwise the curve tree hierarchy won't be up to date and the curve may not be found.

diff --git a/Engine/Source/Editor/Sequencer/Private/SSequencer.cpp b/Engine/Source/Editor/Sequencer/Private/SSequencer.cpp
--- a/Engine/Source/Editor/Sequencer/Private/SSequencer.cpp
+++ b/Engine/Source/Editor/Sequencer/Private/SSequencer.cpp
@@ -1328,26 +1328,34 @@
 
 		if (GetSequencerSettings()->ShouldSyncCurveEditorSelection())
 		{
-			TSharedRef<FSequencerNodeTree> NodeTree = Sequencer->GetNodeTree();
-			CurveEditor->SuspendBroadcast();
-			// Clear the tree selection
-			CurveEditorTree->ClearSelection();
-			for (TWeakPtr<FViewModel> Node : SelectedDisplayNodes)
+			//since Selection ticks before any filter is applied and the curve editor tree is udpated in FSequencer::Tick we need to sync 
+			//the curve editor on the next tick so the filter is applied first, otherwise the curve tree hierarchy won't be up to date
+			//and the curve may not be found.
+			GEditor->GetTimerManager()->SetTimerForNextTick([Sequencer, this]()
 			{
-				if (ICurveEditorTreeItemExtension* CurveEditorItem = ICastable::CastWeakPtr<ICurveEditorTreeItemExtension>(Node))
+				const TSet<TWeakPtr<FViewModel>>& SelectedDisplayNodes = Sequencer->GetSelection().GetSelectedOutlinerItems();
+				TSharedPtr<FCurveEditor> CurveEditor = Sequencer->GetCurveEditor();
+				TSharedRef<FSequencerNodeTree> NodeTree = Sequencer->GetNodeTree();
+				CurveEditor->SuspendBroadcast();
+				// Clear the tree selection
+				CurveEditorTree->ClearSelection();
+				for (TWeakPtr<FViewModel> Node : SelectedDisplayNodes)
 				{
-					FCurveEditorTreeItemID CurveEditorTreeItem = CurveEditorItem->GetCurveEditorItemID();
-					if (CurveEditorTreeItem != FCurveEditorTreeItemID::Invalid())
+					if (ICurveEditorTreeItemExtension* CurveEditorItem = ICastable::CastWeakPtr<ICurveEditorTreeItemExtension>(Node))
 					{
-						if (!CurveEditorTree->IsItemSelected(CurveEditorTreeItem))
+						FCurveEditorTreeItemID CurveEditorTreeItem = CurveEditorItem->GetCurveEditorItemID();
+						if (CurveEditorTreeItem != FCurveEditorTreeItemID::Invalid())
 						{
-							CurveEditorTree->SetItemSelection(CurveEditorTreeItem, true);
-							CurveEditorTree->RequestScrollIntoView(CurveEditorTreeItem);
+							if (!CurveEditorTree->IsItemSelected(CurveEditorTreeItem))
+							{
+								CurveEditorTree->SetItemSelection(CurveEditorTreeItem, true);
+								CurveEditorTree->RequestScrollIntoView(CurveEditorTreeItem);
+							}
 						}
 					}
 				}
-			}
-			CurveEditor->ResumeBroadcast();
+				CurveEditor->ResumeBroadcast();
+			});
 		}
 	}
 
