Sequencer: Fix view range when getting the mouse position with the context menu

diff --git a/Engine/Source/Editor/Sequencer/Private/SequencerTimeSliderController.cpp b/Engine/Source/Editor/Sequencer/Private/SequencerTimeSliderController.cpp
--- a/Engine/Source/Editor/Sequencer/Private/SequencerTimeSliderController.cpp
+++ b/Engine/Source/Editor/Sequencer/Private/SequencerTimeSliderController.cpp
@@ -1343,7 +1343,7 @@
 	TSharedPtr<ISequencer> Sequencer = WeakSequencer.Pin();
 	const bool bReadOnly = Sequencer && Sequencer->IsReadOnly();
 	
-	FScrubRangeToScreen RangeToScreen = FScrubRangeToScreen(TimeSliderArgs.ViewRange.Get(), MyGeometry.Size);
+	FScrubRangeToScreen RangeToScreen = FScrubRangeToScreen(GetViewRange(), MyGeometry.Size);
 	const float MousePixel = MyGeometry.AbsoluteToLocal(MouseEvent.GetScreenSpacePosition()).X;
 	FFrameNumber FrameNumber = ComputeFrameTimeFromMouse(MyGeometry, MouseEvent.GetScreenSpacePosition(), RangeToScreen).FrameNumber;
 
