Sequencer: Don't select folder descendants when retrieving key areas. This ends up selecting all the control rig controls when you select a folder.

diff --git a/Engine/Source/Editor/Sequencer/Private/Sequencer.cpp b/Engine/Source/Editor/Sequencer/Private/Sequencer.cpp
--- a/Engine/Source/Editor/Sequencer/Private/Sequencer.cpp
+++ b/Engine/Source/Editor/Sequencer/Private/Sequencer.cpp
@@ -7092,7 +7092,7 @@
 			NodesToKey.Add(Node.ToSharedRef());
 
 			// No need to gather key areas from binding/tracks because they have no key areas
-			if (Node->IsA<IObjectBindingExtension>() || Node->IsA<ITrackExtension>())
+			if (Node->IsA<IObjectBindingExtension>() || Node->IsA<ITrackExtension>() || Node->IsA<FFolderModel>())
 			{
 				continue;
 			}
@@ -7113,7 +7113,7 @@
 	for (TSharedRef<FViewModel> Node : NodesToKey)
 	{
 		//if object or track selected we don't want all of the children only if spefically selected.
-		if (!Node->IsA<ITrackExtension>() && !Node->IsA<IObjectBindingExtension>())
+		if (!Node->IsA<ITrackExtension>() && !Node->IsA<IObjectBindingExtension>() && !Node->IsA<FFolderModel>())
 		{
 			SequencerHelpers::GetAllKeyAreas(Node, KeyAreas);
 		}
