Sequencer: Collapsed sections will only show filtered keys, only add keys onto the filtered areas, will only delete keys onto filtered areas, will cut and copy onto the filtered areas, and will step to next/previous only on filtered keys.

diff --git a/Engine/Source/Editor/Sequencer/Private/SectionLayout.cpp b/Engine/Source/Editor/Sequencer/Private/SectionLayout.cpp
--- a/Engine/Source/Editor/Sequencer/Private/SectionLayout.cpp
+++ b/Engine/Source/Editor/Sequencer/Private/SectionLayout.cpp
@@ -50,9 +50,23 @@
 	FSectionLayoutElement Tmp;
 	Tmp.Type = Group;
 
-	for (TSharedPtr<FChannelModel> Channel : InChannelRoot->GetDescendantsOfType<FChannelModel>())
+	for (TSharedPtr<FLinkedOutlinerExtension> OutlinerItemExtension: InChannelRoot->GetDescendantsOfType<FLinkedOutlinerExtension>(true /*bIncludeThis*/))
 	{
-		Tmp.WeakChannels.Add(Channel);
+		if (TViewModelPtr<IOutlinerExtension> Outliner = OutlinerItemExtension->GetLinkedOutlinerItem())
+		{
+			if (Outliner->IsFilteredOut() == false)
+			{
+				TSharedPtr<FViewModel> Model = OutlinerItemExtension->GetLinkedOutlinerItem().AsModel();
+				for (TSharedPtr<FChannelGroupModel> ChannelGroup : Model->GetDescendantsOfType<FChannelGroupModel>())
+				{
+					TArrayView<const TWeakViewModelPtr<FChannelModel>> Channels  =  ChannelGroup->GetChannels();
+					for (const TWeakViewModelPtr < FChannelModel>& Channel : Channels)
+					{
+						Tmp.WeakChannels.Add(Channel);
+					}
+				}
+			}
+		}
 	}
 	Tmp.LocalOffset = InOffset;
 	Tmp.Height = InHeight;
diff --git a/Engine/Source/Editor/Sequencer/Private/SequencerAddKeyOperation.cpp b/Engine/Source/Editor/Sequencer/Private/SequencerAddKeyOperation.cpp
--- a/Engine/Source/Editor/Sequencer/Private/SequencerAddKeyOperation.cpp
+++ b/Engine/Source/Editor/Sequencer/Private/SequencerAddKeyOperation.cpp
@@ -114,13 +114,15 @@
 {
 	bool bKeyedAnything = false;
 	ISequencerTrackEditor* TrackEditor = InTrackModel->GetTrackEditor().Get();
-
-	constexpr bool bIncludeThis = true;
-	for (const TWeakViewModelPtr<FChannelModel>& WeakChannel : InChannelGroupModel->GetChannels())
+	if (InChannelGroupModel->IsFilteredOut() == false)
 	{
-		if (TSharedPtr<FChannelModel> Channel = WeakChannel.Pin())
+		constexpr bool bIncludeThis = true;
+		for (const TWeakViewModelPtr<FChannelModel>& WeakChannel : InChannelGroupModel->GetChannels())
 		{
-			bKeyedAnything |= ProcessKeyArea(TrackEditor, Channel->GetKeyArea());
+			if (TSharedPtr<FChannelModel> Channel = WeakChannel.Pin())
+			{
+				bKeyedAnything |= ProcessKeyArea(TrackEditor, Channel->GetKeyArea());
+			}
 		}
 	}
 
diff --git a/Engine/Source/Editor/Sequencer/Private/SequencerKeyCollection.cpp b/Engine/Source/Editor/Sequencer/Private/SequencerKeyCollection.cpp
--- a/Engine/Source/Editor/Sequencer/Private/SequencerKeyCollection.cpp
+++ b/Engine/Source/Editor/Sequencer/Private/SequencerKeyCollection.cpp
@@ -17,10 +17,13 @@
 	{
 		for (TViewModelPtr<const FChannelGroupModel> ChannelGroupModel : Node->GetChildrenOfType<FChannelGroupModel>())
 		{
-			for (const TSharedRef<IKeyArea>& KeyArea : ChannelGroupModel->GetAllKeyAreas())
+			if (ChannelGroupModel->IsFilteredOut() == false)
 			{
-				const UMovieSceneSection* Section = KeyArea->GetOwningSection();
-				Result.KeyAreaToSignature.Add(KeyArea, Section ? Section->GetSignature() : FGuid());
+				for (const TSharedRef<IKeyArea>& KeyArea : ChannelGroupModel->GetAllKeyAreas())
+				{
+					const UMovieSceneSection* Section = KeyArea->GetOwningSection();
+					Result.KeyAreaToSignature.Add(KeyArea, Section ? Section->GetSignature() : FGuid());
+				}
 			}
 		}
 	}
@@ -42,7 +45,10 @@
 		const bool bIncludeThis = true;
 		for (const TViewModelPtr<FChannelGroupModel>& KeyAreaNode : Node->GetDescendantsOfType<FChannelGroupModel>(bIncludeThis))
 		{
-			AllKeyAreaNodes.Add(KeyAreaNode);
+			if (KeyAreaNode->IsFilteredOut() == false)
+			{
+				AllKeyAreaNodes.Add(KeyAreaNode);
+			}
 		}
 	}
 
@@ -69,7 +75,10 @@
 	AllKeyAreaNodes.Reserve(36);
 	for (TSharedPtr<FChannelGroupModel> KeyAreaNode : InNode->GetDescendantsOfType<FChannelGroupModel>(true))
 	{
-		AllKeyAreaNodes.Add(KeyAreaNode);
+		if (KeyAreaNode->IsFilteredOut() == false)
+		{
+			AllKeyAreaNodes.Add(KeyAreaNode);
+		}
 	}
 
 	for (const TSharedPtr<FChannelGroupModel>& Node : AllKeyAreaNodes)
diff --git a/Engine/Source/Editor/Sequencer/Private/Tools/SequencerEntityVisitor.cpp b/Engine/Source/Editor/Sequencer/Private/Tools/SequencerEntityVisitor.cpp
--- a/Engine/Source/Editor/Sequencer/Private/Tools/SequencerEntityVisitor.cpp
+++ b/Engine/Source/Editor/Sequencer/Private/Tools/SequencerEntityVisitor.cpp
@@ -165,11 +165,15 @@
 
 	if (ChannelGroup)
 	{
-		for (const TWeakViewModelPtr<FChannelModel>& WeakChannel : ChannelGroup->GetChannels())
+		const IOutlinerExtension* OutlinerItem = ChannelGroup->CastThis<IOutlinerExtension>();
+		if (OutlinerItem->IsFilteredOut() == false)
 		{
-			if (TViewModelPtr<FChannelModel> Channel = WeakChannel.Pin())
+			for (const TWeakViewModelPtr<FChannelModel>& WeakChannel : ChannelGroup->GetChannels())
 			{
-				VisitChannel(Visitor, Channel);
+				if (TViewModelPtr<FChannelModel> Channel = WeakChannel.Pin())
+				{
+					VisitChannel(Visitor, Channel);
+				}
 			}
 		}
 	}
