Sequencer: Change deferred signed object handler to dispatch events as part of PreTransactionFinalized rather than TransactionFinalized to prevent issues with nested transactions corrupting the transaction buffer while it was being torn down

diff --git a/Engine/Source/Editor/Sequencer/Private/SequencerModule.cpp b/Engine/Source/Editor/Sequencer/Private/SequencerModule.cpp
--- a/Engine/Source/Editor/Sequencer/Private/SequencerModule.cpp
+++ b/Engine/Source/Editor/Sequencer/Private/SequencerModule.cpp
@@ -99,7 +99,7 @@
 			break;
 
 		case ETransactionStateEventType::TransactionCanceled:
-		case ETransactionStateEventType::TransactionFinalized:
+		case ETransactionStateEventType::PreTransactionFinalized:
 		case ETransactionStateEventType::UndoRedoFinalized:
 			DeferTransactionChanges.Reset();
 			break;
