Sequencer: Fixed crash caused by slipping infinite sections

diff --git a/Engine/Source/Editor/Sequencer/Private/Tools/EditToolDragOperations.cpp b/Engine/Source/Editor/Sequencer/Private/Tools/EditToolDragOperations.cpp
--- a/Engine/Source/Editor/Sequencer/Private/Tools/EditToolDragOperations.cpp
+++ b/Engine/Source/Editor/Sequencer/Private/Tools/EditToolDragOperations.cpp
@@ -136,7 +136,14 @@
 		UMovieSceneSection* Section = WeakSection.Get();
 		if (Section)
 		{
-			Sections.Add(Section);
+			if (bDraggingByEnd && Section->HasEndFrame())
+			{
+				Sections.Add(Section);
+			}
+			else if (!bDraggingByEnd && Section->HasStartFrame())
+			{
+				Sections.Add(Section);
+			}
 		}
 	}
 }
