Sequencer: Make some buttons in the transport not focusable (most of them already where not but a few were left out)

diff --git a/Engine/Source/Editor/Sequencer/Private/Sequencer.cpp b/Engine/Source/Editor/Sequencer/Private/Sequencer.cpp
--- a/Engine/Source/Editor/Sequencer/Private/Sequencer.cpp
+++ b/Engine/Source/Editor/Sequencer/Private/Sequencer.cpp
@@ -632,8 +632,8 @@ void FSequencer::InitSequencer(const FSequencerInitParams& InitParams, const TSh
 	RecordingAnimation = FCurveSequence();
 	RecordingAnimation.AddCurve(0.f, 1.5f, ECurveEaseFunction::Linear);
 
-	// Update initial movie scene data
-	NotifyMovieSceneDataChanged( EMovieSceneDataChangeType::ActiveMovieSceneChanged );
+	// Evaluate now when the sequence is opened, as opposed to on Tick()
+	ForceEvaluate();
 
 	// Update the view range to the new current time
 	UpdateTimeBoundsToFocusedMovieScene();
@@ -4472,7 +4472,8 @@ TSharedRef<SWidget> FSequencer::OnCreateTransportSetPlaybackStart()
 		.OnClicked(this, &FSequencer::SetPlaybackStart)
 		.ToolTipText(SetPlaybackStartToolTip)
 		.ButtonStyle(FAppStyle::Get(), "Sequencer.Transport.SetPlayStart")
-		.ContentPadding(2.0f);
+		.ContentPadding(2.0f)
+		.IsFocusable(false);
 }
 
 TSharedRef<SWidget> FSequencer::OnCreateTransportJumpToPreviousKey()
@@ -4483,7 +4484,8 @@ TSharedRef<SWidget> FSequencer::OnCreateTransportJumpToPreviousKey()
 		.OnClicked(this, &FSequencer::JumpToPreviousKey)
 		.ToolTipText(JumpToPreviousKeyToolTip)
 		.ButtonStyle(FAppStyle::Get(), "Sequencer.Transport.JumpToPreviousKey")
-		.ContentPadding(2.0f);
+		.ContentPadding(2.0f)
+		.IsFocusable(false);
 }
 
 TSharedRef<SWidget> FSequencer::OnCreateTransportJumpToNextKey()
@@ -4494,7 +4496,8 @@ TSharedRef<SWidget> FSequencer::OnCreateTransportJumpToNextKey()
 		.OnClicked(this, &FSequencer::JumpToNextKey)
 		.ToolTipText(JumpToNextKeyToolTip)
 		.ButtonStyle(FAppStyle::Get(), "Sequencer.Transport.JumpToNextKey")
-		.ContentPadding(2.0f);
+		.ContentPadding(2.0f)
+		.IsFocusable(false);
 }
 
 TSharedRef<SWidget> FSequencer::OnCreateTransportSetPlaybackEnd()
@@ -4505,7 +4508,8 @@ TSharedRef<SWidget> FSequencer::OnCreateTransportSetPlaybackEnd()
 		.OnClicked(this, &FSequencer::SetPlaybackEnd)
 		.ToolTipText(SetPlaybackEndToolTip)
 		.ButtonStyle(FAppStyle::Get(), "Sequencer.Transport.SetPlayEnd")
-		.ContentPadding(2.0f);
+		.ContentPadding(2.0f)
+		.IsFocusable(false);
 }
 
 TSharedRef<SWidget> FSequencer::OnCreateTransportLoopMode()
@@ -4513,6 +4517,7 @@ TSharedRef<SWidget> FSequencer::OnCreateTransportLoopMode()
 	TSharedRef<SButton> LoopButton = SNew(SButton)
 		.OnClicked(this, &FSequencer::OnCycleLoopMode)
 		.ButtonStyle( FAppStyle::Get(), "NoBorder" )
+		.IsFocusable(false)
 		.ToolTipText_Lambda([&]()
 		{ 
 			if (GetLoopMode() == ESequencerLoopMode::SLM_NoLoop)
@@ -4564,6 +4569,7 @@ TSharedRef<SWidget> FSequencer::OnCreateTransportRecord()
 	TSharedRef<SButton> RecordButton = SNew(SButton)
 		.OnClicked(this, &FSequencer::OnRecord)
 		.ButtonStyle( FAppStyle::Get(), "NoBorder" )
+		.IsFocusable(false)
 		.ToolTipText_Lambda([&]()
 		{
 			FText OutTooltipText;
