Sequencer: Fix SelectBackward

diff --git a/Engine/Source/Editor/Sequencer/Private/Sequencer.cpp b/Engine/Source/Editor/Sequencer/Private/Sequencer.cpp
--- a/Engine/Source/Editor/Sequencer/Private/Sequencer.cpp
+++ b/Engine/Source/Editor/Sequencer/Private/Sequencer.cpp
@@ -2728,7 +2728,7 @@ void FSequencer::SelectBackward()
 
 	FFrameRate TickResolution = GetFocusedTickResolution();
 	FFrameNumber CurrentFrame = GetLocalTime().ConvertTo(TickResolution).CeilToFrame();
-	TRange<FFrameNumber> SelectionRange(CurrentFrame, TNumericLimits<FFrameNumber>::Max());
+	TRange<FFrameNumber> SelectionRange(TNumericLimits<FFrameNumber>::Min(), CurrentFrame);
 
 	TSet<TWeakPtr<FViewModel>> SelectedItems = Selection.GetNodesWithSelectedKeysOrSections();
 	if (SelectedItems.Num() == 0)
