Sequencer: Remove test for whether the current time is being evaluated again. This has caused some confusion where setting the time doesn't evaluate. This used to NOT evaluate so that duplicate events didn't fire again, but that's probably not necessary anymore.

diff --git a/Engine/Source/Editor/Sequencer/Private/Sequencer.cpp b/Engine/Source/Editor/Sequencer/Private/Sequencer.cpp
--- a/Engine/Source/Editor/Sequencer/Private/Sequencer.cpp
+++ b/Engine/Source/Editor/Sequencer/Private/Sequencer.cpp
@@ -3216,16 +3216,10 @@ void FSequencer::SetGlobalTime(FFrameTime NewTime, bool bEvaluate)
 		NewTime = NewTime.FloorToFrame();
 	}
 
-	// Don't update the sequence if the time hasn't changed as this will cause duplicate events and the like to fire.
-	// If we need to reevaluate the sequence at the same time for whetever reason, we should call ForceEvaluate()
-	TOptional<FFrameTime> CurrentPosition = PlayPosition.GetCurrentPosition();
-	if (PlayPosition.GetCurrentPosition() != NewTime)
+	FMovieSceneEvaluationRange EvalRange = PlayPosition.JumpTo(NewTime);
+	if (bEvaluate)
 	{
-		FMovieSceneEvaluationRange EvalRange = PlayPosition.JumpTo(NewTime);
-		if (bEvaluate)
-		{
-			EvaluateInternal(EvalRange);
-		}
+		EvaluateInternal(EvalRange);
 	}
 
 	if (AutoScrubTarget.IsSet())
