Sequencer: Remove trigger drag distance so that dragging the timeline is snappier.

FVisualLoggerTimeSliderController and SSimpleTimelineSlider have also removed this (from which this was derived from)

diff --git a/Engine/Source/Editor/Sequencer/Private/SequencerTimeSliderController.cpp b/Engine/Source/Editor/Sequencer/Private/SequencerTimeSliderController.cpp
--- a/Engine/Source/Editor/Sequencer/Private/SequencerTimeSliderController.cpp
+++ b/Engine/Source/Editor/Sequencer/Private/SequencerTimeSliderController.cpp
@@ -940,7 +940,7 @@ FReply FSequencerTimeSliderController::OnMouseMove( SWidget& WidgetOwner, const
 		if (!bPanning)
 		{
 			DistanceDragged += FMath::Abs( MouseEvent.GetCursorDelta().X );
-			if ( DistanceDragged > FSlateApplication::Get().GetDragTriggerDistance() )
+			if ( DistanceDragged > 0.f /*FSlateApplication::Get().GetDragTriggerDistance()*/ )
 			{
 				bPanning = true;
 			}
@@ -990,7 +990,7 @@ FReply FSequencerTimeSliderController::OnMouseMove( SWidget& WidgetOwner, const
 
 		if ( MouseDragType == DRAG_NONE )
 		{
-			if ( DistanceDragged > FSlateApplication::Get().GetDragTriggerDistance() )
+			if ( DistanceDragged > 0.f /*FSlateApplication::Get().GetDragTriggerDistance()*/ )
 			{
 				FFrameTime MouseDownFree = ComputeFrameTimeFromMouse(MyGeometry, MouseDownPosition[0], RangeToScreen, false);
 
