Sequencer: Fix blank properties menu

diff --git a/Engine/Source/Editor/Sequencer/Private/MVVM/ViewModels/ObjectBindingModel.cpp b/Engine/Source/Editor/Sequencer/Private/MVVM/ViewModels/ObjectBindingModel.cpp
--- a/Engine/Source/Editor/Sequencer/Private/MVVM/ViewModels/ObjectBindingModel.cpp
+++ b/Engine/Source/Editor/Sequencer/Private/MVVM/ViewModels/ObjectBindingModel.cpp
@@ -26,6 +26,7 @@
 #include "SObjectBindingTag.h"
 #include "Containers/ArrayBuilder.h"
 #include "ObjectBindingTagCache.h"
+#include "ObjectEditorUtils.h"
 #include "Styling/AppStyle.h"
 #include "Styling/SlateIconFinder.h"
 #include "Tracks/MovieSceneSpawnTrack.h"
@@ -723,7 +724,7 @@
 			}
 			else
 			{
-				//KeyableMenuData.MenuName = FObjectEditorUtils::GetCategoryFName(Property).ToString();
+				KeyableMenuData.MenuName = FObjectEditorUtils::GetCategoryFName(Property).ToString();
 			}
 
 			PropertiesTraversed.Add(Property);
