Sequencer: EmptySelectedOutlinerNodesWithoutSections should also look at ancestor track's section. For example, a category group model could be selected and if its ancestor track model has sections that overlap with the request sections, it shouldn't do anything here.

This fixes a bug with control rig selection where if you select a key on a control rig node, the corresponding node would get deselected.

diff --git a/Engine/Source/Editor/Sequencer/Private/SequencerSelection.cpp b/Engine/Source/Editor/Sequencer/Private/SequencerSelection.cpp
--- a/Engine/Source/Editor/Sequencer/Private/SequencerSelection.cpp
+++ b/Engine/Source/Editor/Sequencer/Private/SequencerSelection.cpp
@@ -2,6 +2,7 @@
 
 #include "SequencerSelection.h"
 #include "MovieSceneSection.h"
+#include "MovieSceneTrack.h"
 #include "SequencerCommonHelpers.h"
 
 #include "MVVM/SharedViewModelData.h"
@@ -578,6 +579,21 @@
 
 		bool bFoundMatch = false;
 
+		if (TSharedPtr<FTrackModel> TrackModel = SelectedModel->FindAncestorOfType<FTrackModel>())
+		{
+			if (UMovieSceneTrack* Track = TrackModel->GetTrack())
+			{
+				for (int32 SectionIndex = 0; SectionIndex < Track->GetAllSections().Num() && !bFoundMatch; ++SectionIndex)
+				{
+					if (Sections.Contains(Track->GetAllSections()[SectionIndex]))
+					{
+						bFoundMatch = true;
+						break;
+					}
+				}
+			}
+		}
+			
 		for (TSharedPtr<FSectionModel> SectionModel : SelectedModel->GetDescendantsOfType<FSectionModel>())
 		{
 			if (Sections.Contains(SectionModel->GetSection()))
