Sequencer: Add command to toggle auto expanding the outliner tree on child selection

diff --git a/Engine/Source/Editor/Sequencer/Private/SSequencer.cpp b/Engine/Source/Editor/Sequencer/Private/SSequencer.cpp
--- a/Engine/Source/Editor/Sequencer/Private/SSequencer.cpp
+++ b/Engine/Source/Editor/Sequencer/Private/SSequencer.cpp
@@ -1315,7 +1315,7 @@ void SSequencer::HandleOutlinerNodeSelectionChanged()
 				SequencerSelectionCurveEditorFilter = MakeShared<FSequencerSelectionCurveFilter>();
 			}
 
-			SequencerSelectionCurveEditorFilter->Update(Sequencer->GetSelection().GetSelectedOutlinerItems(), SequencerSettings->GetAutoExpandTreeView());
+			SequencerSelectionCurveEditorFilter->Update(Sequencer->GetSelection().GetSelectedOutlinerItems(), SequencerSettings->GetAutoExpandNodesOnSelection());
 
 			CurveEditor->GetTree()->AddFilter(SequencerSelectionCurveEditorFilter);
 		}
@@ -2355,6 +2355,7 @@ TSharedRef<SWidget> SSequencer::MakeViewMenu()
 
 	MenuBuilder.BeginSection("OrganizeAndSort", LOCTEXT("OrganizeAndSortHeader", "Organize and Sort"));
 	{
+		MenuBuilder.AddMenuEntry(FSequencerCommands::Get().ToggleAutoExpandNodesOnSelection);
 		MenuBuilder.AddMenuEntry(FSequencerCommands::Get().ToggleExpandCollapseNodes);
 		MenuBuilder.AddMenuEntry(FSequencerCommands::Get().ToggleExpandCollapseNodesAndDescendants);
 		MenuBuilder.AddMenuEntry(FSequencerCommands::Get().ExpandAllNodes);
diff --git a/Engine/Source/Editor/Sequencer/Private/Sequencer.cpp b/Engine/Source/Editor/Sequencer/Private/Sequencer.cpp
--- a/Engine/Source/Editor/Sequencer/Private/Sequencer.cpp
+++ b/Engine/Source/Editor/Sequencer/Private/Sequencer.cpp
@@ -7170,7 +7170,7 @@ void FSequencer::SelectByNthCategoryNode(UMovieSceneSection* Section, int Index,
 	}
 	if (bSelect)
 	{
-		if (Settings->GetAutoExpandTreeView())
+		if (Settings->GetAutoExpandNodesOnSelection())
 		{
 			for (const TSharedRef<FViewModel>& DisplayNode : NodesToSelect)
 			{
@@ -7315,7 +7315,7 @@ void FSequencer::SelectByChannels(UMovieSceneSection* Section, const TArray<FNam
 	{
 		for (const TSharedRef<FViewModel>& DisplayNode : Nodes)
 		{
-			if (Settings->GetAutoExpandTreeView())
+			if (Settings->GetAutoExpandNodesOnSelection())
 			{
 				if (DisplayNode->GetParent().IsValid() && DisplayNode->GetParent()->IsA<ITrackExtension>() && !DisplayNode->GetParent()->CastThisChecked<IOutlinerExtension>()->IsExpanded())
 				{
@@ -12510,6 +12510,12 @@ void FSequencer::BindCommands()
 		Commands.SortAllNodesAndDescendants,
 		FExecuteAction::CreateSP(this, &FSequencer::SortAllNodesAndDescendants));
 
+	SequencerCommandBindings->MapAction(
+		Commands.ToggleAutoExpandNodesOnSelection,
+		FExecuteAction::CreateLambda([this] { Settings->SetAutoExpandNodesOnSelection(!Settings->GetAutoExpandNodesOnSelection()); }),
+		FCanExecuteAction::CreateLambda([] { return true; }),
+		FIsActionChecked::CreateLambda([this] { return Settings->GetAutoExpandNodesOnSelection(); }) );
+
 	SequencerCommandBindings->MapAction(
 		Commands.ToggleExpandCollapseNodes,
 		FExecuteAction::CreateSP(this, &FSequencer::ToggleExpandCollapseNodes));
diff --git a/Engine/Source/Editor/Sequencer/Private/SequencerCommands.cpp b/Engine/Source/Editor/Sequencer/Private/SequencerCommands.cpp
--- a/Engine/Source/Editor/Sequencer/Private/SequencerCommands.cpp
+++ b/Engine/Source/Editor/Sequencer/Private/SequencerCommands.cpp
@@ -49,6 +49,7 @@
 	UI_COMMAND( ToggleKeepCursorInPlaybackRangeWhileScrubbing, "Keep Playhead in Playback Range While Scrubbing", "When checked, the playhead will be constrained to the current playback range while scrubbing", EUserInterfaceActionType::ToggleButton, FInputChord() );
 	UI_COMMAND( ToggleKeepPlaybackRangeInSectionBounds, "Keep Playback Range in Section Bounds", "When checked, the playback range will be synchronized to the section bounds", EUserInterfaceActionType::ToggleButton, FInputChord() );
 
+	UI_COMMAND( ToggleAutoExpandNodesOnSelection, "Auto Expand Nodes on Selection", "Toggle auto expanding the outliner tree on child selection", EUserInterfaceActionType::ToggleButton, FInputChord());
 	UI_COMMAND( ToggleExpandCollapseNodes, "Expand/Collapse Nodes", "Toggle expand or collapse selected nodes", EUserInterfaceActionType::Button, FInputChord(EKeys::V) );
 	UI_COMMAND( ToggleExpandCollapseNodesAndDescendants, "Expand/Collapse Nodes and Descendants", "Toggle expand or collapse selected nodes and descendants", EUserInterfaceActionType::Button, FInputChord(EModifierKey::Shift, EKeys::V) );
 	UI_COMMAND( ExpandAllNodes, "Expand All Nodes", "Expand all nodes and descendants", EUserInterfaceActionType::Button, FInputChord() );
diff --git a/Engine/Source/Editor/Sequencer/Private/SequencerSettings.cpp b/Engine/Source/Editor/Sequencer/Private/SequencerSettings.cpp
--- a/Engine/Source/Editor/Sequencer/Private/SequencerSettings.cpp
+++ b/Engine/Source/Editor/Sequencer/Private/SequencerSettings.cpp
@@ -59,7 +59,7 @@
 	bVisualizePreAndPostRoll = true;
 	TrajectoryPathCap = 250;
 	FrameNumberDisplayFormat = EFrameNumberDisplayFormats::Seconds;
-	bAutoExpandOutlinerTreeOnChildSelection = true;
+	bAutoExpandNodesOnSelection = true;
 }
 
 void USequencerSettings::PostEditChangeProperty(FPropertyChangedEvent& PropertyChangedEvent)
@@ -863,11 +863,11 @@
 	}
 }
 
-void USequencerSettings::SetAutoExpandTreeView(bool bInAutoExpandOutlinerTree)
+void USequencerSettings::SetAutoExpandNodesOnSelection(bool bInAutoExpandNodesOnSelection)
 {
-	if (bInAutoExpandOutlinerTree != bAutoExpandOutlinerTreeOnChildSelection)
+	if (bInAutoExpandNodesOnSelection != bAutoExpandNodesOnSelection)
 	{
-		bAutoExpandOutlinerTreeOnChildSelection = bInAutoExpandOutlinerTree;
+		bAutoExpandNodesOnSelection = bInAutoExpandNodesOnSelection;
 		SaveConfig();
 	}
 }
diff --git a/Engine/Source/Editor/Sequencer/Public/SequencerCommands.h b/Engine/Source/Editor/Sequencer/Public/SequencerCommands.h
--- a/Engine/Source/Editor/Sequencer/Public/SequencerCommands.h
+++ b/Engine/Source/Editor/Sequencer/Public/SequencerCommands.h
@@ -136,6 +136,9 @@
 	/** Toggle constraining the playback range to the section bounds */
 	TSharedPtr< FUICommandInfo > ToggleKeepPlaybackRangeInSectionBounds;
 
+	/** Toggle auto expand outliner tree on child selection */
+	TSharedPtr< FUICommandInfo > ToggleAutoExpandNodesOnSelection;
+
 	/** Expand/collapse nodes */
 	TSharedPtr< FUICommandInfo > ToggleExpandCollapseNodes;
 
diff --git a/Engine/Source/Editor/Sequencer/Public/SequencerSettings.h b/Engine/Source/Editor/Sequencer/Public/SequencerSettings.h
--- a/Engine/Source/Editor/Sequencer/Public/SequencerSettings.h
+++ b/Engine/Source/Editor/Sequencer/Public/SequencerSettings.h
@@ -377,11 +377,10 @@
 	/** Sets the movie renderer to use */
 	void SetMovieRendererName(const FString& InMovieRendererName);
 
-	
 	/** Gets whether or not to expand the outliner tree view when a child element is selected (from outside of the tree view). */
-	bool GetAutoExpandTreeView() const { return bAutoExpandOutlinerTreeOnChildSelection; }
+	bool GetAutoExpandNodesOnSelection() const { return bAutoExpandNodesOnSelection; }
 	/** Sets whether or not to expand the outliner tree view when a child element is selected (from outside of the tree view). */
-	void SetAutoExpandTreeView(bool bInAutoExpandOutlinerTree);
+	void SetAutoExpandNodesOnSelection(bool bInAutoExpandNodesOnSelection);
 
 protected:
 
@@ -605,7 +604,7 @@
 
 	/** Whether to expand the sequencer tree view when a child element is selected (from outside of the tree view). */
 	UPROPERTY(config, EditAnywhere, Category= General)
-	bool bAutoExpandOutlinerTreeOnChildSelection;
+	bool bAutoExpandNodesOnSelection;
 
 	FOnEvaluateSubSequencesInIsolationChanged OnEvaluateSubSequencesInIsolationChangedEvent;
 	FOnShowSelectedNodesOnlyChanged OnShowSelectedNodesOnlyChangedEvent;
