Sequencer: Check for valid channel model before dereferencing

diff --git a/Engine/Source/Editor/Sequencer/Private/SequencerNodeTree.cpp b/Engine/Source/Editor/Sequencer/Private/SequencerNodeTree.cpp
--- a/Engine/Source/Editor/Sequencer/Private/SequencerNodeTree.cpp
+++ b/Engine/Source/Editor/Sequencer/Private/SequencerNodeTree.cpp
@@ -707,16 +707,19 @@
 
 		for (TSharedPtr<FChannelGroupModel> ChannelGroupModel : StartNode.AsModel()->GetDescendantsOfType<FChannelGroupModel>())
 		{
-			for (const TWeakViewModelPtr<FChannelModel>& ChannelModel : ChannelGroupModel->GetChannels())
+			for (const TWeakViewModelPtr<FChannelModel>& WeakChannel : ChannelGroupModel->GetChannels())
 			{
-				const TSharedPtr<IKeyArea>& KeyArea = ChannelModel.Pin()->GetKeyArea();
-				FMovieSceneChannel* Channel = KeyArea->ResolveChannel();
-				if (Channel)
+				if (TViewModelPtr<FChannelModel> ChannelModel = WeakChannel.Pin())
 				{
-					if (Filters->Num() == 0 || Filters->PassesAnyFilters(Channel))
+					const TSharedPtr<IKeyArea>& KeyArea = ChannelModel->GetKeyArea();
+					FMovieSceneChannel* Channel = KeyArea->ResolveChannel();
+					if (Channel)
 					{
-						bPassedAnyFilters = true;
-						break;
+						if (Filters->Num() == 0 || Filters->PassesAnyFilters(Channel))
+						{
+							bPassedAnyFilters = true;
+							break;
+						}
 					}
 				}
 			}
@@ -836,16 +839,19 @@
 	}
 	else if (TViewModelPtr<FChannelGroupModel> ChannelGroupModel = StartNode.ImplicitCast())
 	{
-		for (const TWeakViewModelPtr<FChannelModel>& ChannelModel : ChannelGroupModel->GetChannels())
+		for (const TWeakViewModelPtr<FChannelModel>& WeakChannel : ChannelGroupModel->GetChannels())
 		{
-			const TSharedPtr<IKeyArea>& KeyArea = ChannelModel.Pin()->GetKeyArea();
-			FMovieSceneChannel* Channel = KeyArea->ResolveChannel();
-			if (Channel)
+			if (TViewModelPtr<FChannelModel> ChannelModel = WeakChannel.Pin())
 			{
-				if (Filters->Num() == 0 || Filters->PassesAnyFilters(Channel))
+				const TSharedPtr<IKeyArea>& KeyArea = ChannelModel->GetKeyArea();
+				FMovieSceneChannel* Channel = KeyArea->ResolveChannel();
+				if (Channel)
 				{
-					bPassedAnyFilters = true;
-					break;
+					if (Filters->Num() == 0 || Filters->PassesAnyFilters(Channel))
+					{
+						bPassedAnyFilters = true;
+						break;
+					}
 				}
 			}
 		}
