Sequencer: Added ensure when InModel is null

diff --git a/Engine/Source/Editor/Sequencer/Private/SequencerSelection.cpp b/Engine/Source/Editor/Sequencer/Private/SequencerSelection.cpp
--- a/Engine/Source/Editor/Sequencer/Private/SequencerSelection.cpp
+++ b/Engine/Source/Editor/Sequencer/Private/SequencerSelection.cpp
@@ -188,6 +188,11 @@ TArray<UMovieSceneTrack*> FSequencerSelection::GetSelectedTracks() const
 
 void FSequencerSelection::AddToSelection(TSharedPtr<UE::Sequencer::FViewModel> InModel)
 {
+	if (!ensureAlwaysMsgf(InModel, TEXT("AddToSelection with null view model")))
+	{
+		return;
+	}
+
 	if (InModel->IsA<UE::Sequencer::IOutlinerExtension>())
 	{
 		AddToOutlinerSelection(InModel);
@@ -200,6 +205,11 @@ void FSequencerSelection::AddToSelection(TSharedPtr<UE::Sequencer::FViewModel> I
 
 void FSequencerSelection::AddToTrackAreaSelection(TSharedPtr<UE::Sequencer::FViewModel> InModel)
 {
+	if (!ensureAlwaysMsgf(InModel, TEXT("AddToSelection with null view model")))
+	{
+		return;
+	}
+
 	using namespace UE::Sequencer;
 
 	ISelectableExtension* Selectable = InModel->CastThis<ISelectableExtension>();
@@ -266,6 +276,11 @@ void FSequencerSelection::AddToSelection(const FSequencerSelectedKey& Key)
 
 void FSequencerSelection::AddToOutlinerSelection(TSharedPtr<UE::Sequencer::FViewModel> InModel)
 {
+	if (!ensureAlwaysMsgf(InModel, TEXT("AddToSelection with null view model")))
+	{
+		return;
+	}
+
 	using namespace UE::Sequencer;
 
 	ISelectableExtension* Selectable = InModel->CastThis<ISelectableExtension>();
@@ -300,10 +315,13 @@ void FSequencerSelection::AddToSelection(const TArrayView<TSharedPtr<UE::Sequenc
 
 	for (TSharedPtr<FViewModel> InModel : InModels)
 	{
-		SelectedOutlinerItems.Add(InModel);
-		if (IOutlinerExtension* OutlinerItem = InModel->CastThis<IOutlinerExtension>())
+		if (ensureAlwaysMsgf(InModel, TEXT("AddToSelection with null view model")))
 		{
-			OutlinerItem->SetSelectionState(EOutlinerSelectionState::SelectedDirectly);
+			SelectedOutlinerItems.Add(InModel);
+			if (IOutlinerExtension* OutlinerItem = InModel->CastThis<IOutlinerExtension>())
+			{
+				OutlinerItem->SetSelectionState(EOutlinerSelectionState::SelectedDirectly);
+			}
 		}
 	}
 
