Sequencer: Back out inadvertent ForceEvaluate

diff --git a/Engine/Source/Editor/Sequencer/Private/Sequencer.cpp b/Engine/Source/Editor/Sequencer/Private/Sequencer.cpp
--- a/Engine/Source/Editor/Sequencer/Private/Sequencer.cpp
+++ b/Engine/Source/Editor/Sequencer/Private/Sequencer.cpp
@@ -632,8 +632,8 @@ void FSequencer::InitSequencer(const FSequencerInitParams& InitParams, const TSh
 	RecordingAnimation = FCurveSequence();
 	RecordingAnimation.AddCurve(0.f, 1.5f, ECurveEaseFunction::Linear);
 
-	// Evaluate now when the sequence is opened, as opposed to on Tick()
-	ForceEvaluate();
+	// Update initial movie scene data
+	NotifyMovieSceneDataChanged( EMovieSceneDataChangeType::ActiveMovieSceneChanged );
 
 	// Update the view range to the new current time
 	UpdateTimeBoundsToFocusedMovieScene();
@@ -3216,10 +3216,16 @@ void FSequencer::SetGlobalTime(FFrameTime NewTime, bool bEvaluate)
 		NewTime = NewTime.FloorToFrame();
 	}
 
-	FMovieSceneEvaluationRange EvalRange = PlayPosition.JumpTo(NewTime);
-	if (bEvaluate)
+	// Don't update the sequence if the time hasn't changed as this will cause duplicate events and the like to fire.
+	// If we need to reevaluate the sequence at the same time for whetever reason, we should call ForceEvaluate()
+	TOptional<FFrameTime> CurrentPosition = PlayPosition.GetCurrentPosition();
+	if (PlayPosition.GetCurrentPosition() != NewTime)
 	{
-		EvaluateInternal(EvalRange);
+		FMovieSceneEvaluationRange EvalRange = PlayPosition.JumpTo(NewTime);
+		if (bEvaluate)
+		{
+			EvaluateInternal(EvalRange);
+		}
 	}
 
 	if (AutoScrubTarget.IsSet())
