Sequencer: Fixed key navigation buttons not working

diff --git a/Engine/Source/Editor/Sequencer/Private/MVVM/Views/SSequencerKeyNavigationButtons.h b/Engine/Source/Editor/Sequencer/Private/MVVM/Views/SSequencerKeyNavigationButtons.h
--- a/Engine/Source/Editor/Sequencer/Private/MVVM/Views/SSequencerKeyNavigationButtons.h
+++ b/Engine/Source/Editor/Sequencer/Private/MVVM/Views/SSequencerKeyNavigationButtons.h
@@ -34,6 +34,7 @@
 
 	void Construct(const FArguments& InArgs, const TSharedPtr<FViewModel>& InModel, TSharedPtr<ISequencer> InSequencer)
 	{
+		WeakModel = InModel;
 		WeakSequencer = InSequencer;
 
 		ChildSlot
@@ -42,6 +43,7 @@
 			.AddKeyToolTip(FText::Format(LOCTEXT("AddKeyButton", "Add a new key at the current time ({0})"), FSequencerCommands::Get().SetKey->GetInputText()))
 			.PreviousKeyToolTip(FText::Format(LOCTEXT("PreviousKeyButton", "Set the time to the previous key ({0})"), FSequencerCommands::Get().StepToPreviousKey->GetInputText()))
 			.NextKeyToolTip(FText::Format(LOCTEXT("NextKeyButton", "Set the time to the next key ({0})"), FSequencerCommands::Get().StepToNextKey->GetInputText()))
+			.GetNavigatableTimes(this, &SSequencerKeyNavigationButtons::GetNavigatableTimes)
 			.Time(this, &SSequencerKeyNavigationButtons::GetTime)
 			.OnSetTime(this, &SSequencerKeyNavigationButtons::SetTime)
 			.OnAddKey(this, &SSequencerKeyNavigationButtons::HandleAddKey)
@@ -60,10 +62,9 @@
 		TSet<TSharedPtr<IKeyArea>> KeyAreas;
 		SequencerHelpers::GetAllKeyAreas( DataModel, KeyAreas );
 
-		TArray<FFrameNumber> AllTimes;
 		for ( TSharedPtr<IKeyArea> Keyarea : KeyAreas )
 		{
-			Keyarea->GetKeyTimes(AllTimes);
+			Keyarea->GetKeyTimes(NavigatableTimes);
 		}
 
 		TSet<TWeakObjectPtr<UMovieSceneSection> > Sections;
@@ -72,7 +73,7 @@
 		{
 			if (Section.IsValid())
 			{
-				Section->GetSnapTimes(AllTimes, true);
+				Section->GetSnapTimes(NavigatableTimes, true);
 			}
 		}
 	}
