Sequencer: Fixed construction of child object binding models

The initial population of object binding models is precarious due to the unstructured nature of the UMovieScene data (there is no authoritative view of ancestry for objects, folders etc) and the lack of context for where objects should live.
This results in guesswork during construction for some objects (specifically, child components that are created _before_ their parent) where the complete ancestry is not known.

The bug manifested itself by child object bindings happening to be created before their parent, which would implicitly construct the parent, but leave it detached from the model hierarchy. It would then get cleaned up along with its child binding, only to be re-created in a further iteration of the construction loop, but this time without its child (which was assumed to have already been constructed).

The fix is to detect detached sub-trees on population, and always add them to the root model's children if so. This may result in object bindings being subsequently moved into folders when the folder models are created, but that is acceptable imo.

diff --git a/Engine/Source/Editor/Sequencer/Private/MVVM/ObjectBindingModelStorageExtension.cpp b/Engine/Source/Editor/Sequencer/Private/MVVM/ObjectBindingModelStorageExtension.cpp
--- a/Engine/Source/Editor/Sequencer/Private/MVVM/ObjectBindingModelStorageExtension.cpp
+++ b/Engine/Source/Editor/Sequencer/Private/MVVM/ObjectBindingModelStorageExtension.cpp
@@ -124,26 +124,15 @@
 	{
 		MovieScene->EventHandlers.Link(this);
 
-		TSet<FGuid> ObjectBindingsInFolders;
-
-		TArray<UMovieSceneFolder*> AllFolders;
-		GetMovieSceneFoldersRecursive(MovieScene->GetRootFolders(), AllFolders);
-		for (UMovieSceneFolder* Folder : AllFolders)
-		{
-			for (const FGuid& ObjectBindingID : Folder->GetChildObjectBindings())
-			{
-				ObjectBindingsInFolders.Add(ObjectBindingID);
-			}
-		}
-
 		FViewModelChildren RootChildren = OwnerModel->GetChildList(EViewModelListType::Outliner);
 
 		for (const FMovieSceneBinding& Binding : MovieScene->GetBindings())
 		{
-			if (!ObjectBindingsInFolders.Contains(Binding.GetObjectGuid()))
+			// Need to create a new one for this object binding
+			TSharedPtr<FViewModel> Model = GetOrCreateModelForBinding(Binding.GetObjectGuid());
+			if (!Model->IsConstructed())
 			{
-				// If this is a spawnable or a possessable with no parent, it belongs in the root
-				GetOrCreateModelForBinding(Binding, &RootChildren);
+				RootChildren.AddChild(Model->GetRoot());
 			}
 		}
 	}
@@ -184,16 +173,12 @@
 		return GetOrCreateModelForBinding(*ObjectBinding);
 	}
 
+	// Placeholders are always created at the root
 	return CreatePlaceholderForObjectBinding(Binding);
 }
 
 TSharedPtr<FViewModel> FObjectBindingModelStorageExtension::GetOrCreateModelForBinding(const FMovieSceneBinding& Binding)
 {
-	return GetOrCreateModelForBinding(Binding, nullptr);
-}
-
-TSharedPtr<FViewModel> FObjectBindingModelStorageExtension::GetOrCreateModelForBinding(const FMovieSceneBinding& Binding, FViewModelChildren* RootChildren)
-{
 	TSharedPtr<FObjectBindingModel> ObjectModel = CreateModelForObjectBinding(Binding);
 
 	// Set up parentage
@@ -201,20 +186,9 @@
 	if (DesiredParent.IsValid())
 	{
 		TSharedPtr<FViewModel> Parent = GetOrCreateModelForBinding(DesiredParent);
-
-		// Explicitly remove from parent first to guarantee that the model is re-constructed
-		ObjectModel->RemoveFromParent();
-
 		Parent->GetChildList(EViewModelListType::Outliner).AddChild(ObjectModel);
 	}
-	else if (RootChildren)
-	{
-		// Explicitly remove from parent first to guarantee that the model is re-constructed
-		ObjectModel->RemoveFromParent();
 
-		RootChildren->AddChild(ObjectModel);
-	}
-
 	TSharedPtr<FViewModel> Placeholder = FindPlaceholderForObjectBinding(Binding.GetObjectGuid());
 	if (Placeholder)
 	{
@@ -236,8 +210,13 @@
 
 void FObjectBindingModelStorageExtension::OnBindingAdded(const FMovieSceneBinding& Binding)
 {
-	FViewModelChildren RootChildren = OwnerModel->GetChildList(EViewModelListType::Outliner);
-	TSharedPtr<FViewModel> Model = GetOrCreateModelForBinding(Binding, &RootChildren);
+	TSharedPtr<FViewModel> Model = GetOrCreateModelForBinding(Binding);
+
+	if (!Model->IsConstructed())
+	{
+		OwnerModel->GetChildList(EViewModelListType::Outliner)
+			.AddChild(Model);
+	}
 }
 
 void FObjectBindingModelStorageExtension::OnBindingRemoved(const FGuid& ObjectBindingID)
diff --git a/Engine/Source/Editor/Sequencer/Public/MVVM/ObjectBindingModelStorageExtension.h b/Engine/Source/Editor/Sequencer/Public/MVVM/ObjectBindingModelStorageExtension.h
--- a/Engine/Source/Editor/Sequencer/Public/MVVM/ObjectBindingModelStorageExtension.h
+++ b/Engine/Source/Editor/Sequencer/Public/MVVM/ObjectBindingModelStorageExtension.h
@@ -30,7 +30,16 @@
 
 	FObjectBindingModelStorageExtension();
 
+	/**
+	 * Implementation function for creating a new model for a binding from its ID. May return a placeholder if the binding does not yet exist.
+	 * @param Binding The binding to create a model for
+	 */
 	TSharedPtr<FViewModel> GetOrCreateModelForBinding(const FGuid& Binding);
+
+	/**
+	 * Implementation function for creating a new model for a binding
+	 * @param Binding The binding to create a model for
+	 */
 	TSharedPtr<FViewModel> GetOrCreateModelForBinding(const FMovieSceneBinding& Binding);
 
 	TSharedPtr<FObjectBindingModel> FindModelForObjectBinding(const FGuid& InObjectBindingID) const;
@@ -52,13 +61,6 @@
 	TSharedPtr<FViewModel> CreatePlaceholderForObjectBinding(const FGuid& ObjectID);
 	TSharedPtr<FViewModel> FindPlaceholderForObjectBinding(const FGuid& InObjectBindingID) const;
 
-	/**
-	 * Implementation function for creating a new model for a binding
-	 * @param Binding The binding to create a model for
-	 * @param RootChildren (optional) When specified, will place the object into this child list, if it does not define a desired parent ID.
-	 */
-	TSharedPtr<FViewModel> GetOrCreateModelForBinding(const FMovieSceneBinding& Binding, FViewModelChildren* RootChildren);
-
 	void Compact();
 
 private:
