Sequencer: Outliner and top-level channel groups are now distinct classes

This removes the need for the Force Filtered Out hack which previously prevented top-level channels from ever being keyed.
The outliner item model is now a mixin-type class that can be added to different model types, rather than being a concrete base class type.
This allows the top level and outliner channel group models to share the same base channel class, where the latter can also be added to the outliner directly.

diff --git a/Engine/Source/Editor/Sequencer/Private/MVVM/SequencerViewModelTypeIDs.cpp b/Engine/Source/Editor/Sequencer/Private/MVVM/SequencerViewModelTypeIDs.cpp
--- a/Engine/Source/Editor/Sequencer/Private/MVVM/SequencerViewModelTypeIDs.cpp
+++ b/Engine/Source/Editor/Sequencer/Private/MVVM/SequencerViewModelTypeIDs.cpp
@@ -39,6 +39,7 @@
 UE_SEQUENCER_DEFINE_CASTABLE(FCategoryGroupModel);
 UE_SEQUENCER_DEFINE_CASTABLE(FCategoryModel);
 UE_SEQUENCER_DEFINE_CASTABLE(FChannelGroupModel);
+UE_SEQUENCER_DEFINE_CASTABLE(FChannelGroupOutlinerModel);
 UE_SEQUENCER_DEFINE_CASTABLE(FChannelModel);
 UE_SEQUENCER_DEFINE_CASTABLE(FFolderModel);
 UE_SEQUENCER_DEFINE_CASTABLE(FLayerBarModel);
diff --git a/Engine/Source/Editor/Sequencer/Private/MVVM/ViewModels/ChannelModel.cpp b/Engine/Source/Editor/Sequencer/Private/MVVM/ViewModels/ChannelModel.cpp
--- a/Engine/Source/Editor/Sequencer/Private/MVVM/ViewModels/ChannelModel.cpp
+++ b/Engine/Source/Editor/Sequencer/Private/MVVM/ViewModels/ChannelModel.cpp
@@ -76,7 +76,6 @@
 	: ChannelName(InChannelName)
 	, DisplayText(InDisplayText)
 {
-	SetIdentifier(InChannelName);
 }
 
 FChannelGroupModel::~FChannelGroupModel()
@@ -159,8 +158,69 @@
 	return KeyAreas;
 }
 
-FOutlinerSizing FChannelGroupModel::RecomputeSizing()
+FTrackAreaParameters FChannelGroupModel::GetTrackAreaParameters() const
 {
+	FTrackAreaParameters Parameters;
+	Parameters.LaneType = ETrackAreaLaneType::None;
+	return Parameters;
+}
+
+FViewModelVariantIterator FChannelGroupModel::GetTrackAreaModelList() const
+{
+	return &Channels;
+}
+
+void FChannelGroupModel::CreateCurveModels(TArray<TUniquePtr<FCurveModel>>& OutCurveModels)
+{
+	TSharedPtr<FSequenceModel> SequenceModel = FindAncestorOfType<FSequenceModel>();
+	if (!SequenceModel)
+	{
+		return;
+	}
+
+	TSharedPtr<ISequencer> Sequencer = SequenceModel->GetSequencer();
+	if (!Sequencer)
+	{
+		return;
+	}
+
+	for (TWeakViewModelPtr<FChannelModel> WeakChannel : GetChannels())
+	{
+		if (TSharedPtr<FChannelModel> Channel = WeakChannel.Pin())
+		{
+			TUniquePtr<FCurveModel> NewCurve = Channel->GetKeyArea()->CreateCurveEditorModel(Sequencer.ToSharedRef());
+			if (NewCurve.IsValid())
+			{
+				OutCurveModels.Add(MoveTemp(NewCurve));
+			}
+		}
+	}
+}
+
+bool FChannelGroupModel::HasCurves() const
+{
+	for (const TSharedRef<IKeyArea>& KeyArea : GetAllKeyAreas())
+	{
+		const ISequencerChannelInterface* EditorInterface = KeyArea->FindChannelEditorInterface();
+		if (EditorInterface && EditorInterface->SupportsCurveEditorModels_Raw(KeyArea->GetChannel()))
+		{
+			return true;
+		}
+	}
+	return false;
+}
+
+FChannelGroupOutlinerModel::FChannelGroupOutlinerModel(FName InChannelName, const FText& InDisplayText)
+	: TOutlinerModelMixin<FChannelGroupModel>(InChannelName, InDisplayText)
+{
+	SetIdentifier(InChannelName);
+}
+
+FChannelGroupOutlinerModel::~FChannelGroupOutlinerModel()
+{}
+
+FOutlinerSizing FChannelGroupOutlinerModel::RecomputeSizing()
+{
 	FOutlinerSizing MaxSizing;
 
 	for (TWeakViewModelPtr<FChannelModel> WeakChannel : Channels)
@@ -188,25 +248,13 @@
 	return MaxSizing;
 }
 
-FOutlinerSizing FChannelGroupModel::GetOutlinerSizing() const
+FOutlinerSizing FChannelGroupOutlinerModel::GetOutlinerSizing() const
 {
 	return ComputedSizing;
 }
 
-FTrackAreaParameters FChannelGroupModel::GetTrackAreaParameters() const
+TSharedRef<SWidget> FChannelGroupOutlinerModel::CreateOutlinerView(const FCreateOutlinerViewParams& InParams)
 {
-	FTrackAreaParameters Parameters;
-	Parameters.LaneType = ETrackAreaLaneType::None;
-	return Parameters;
-}
-
-FViewModelVariantIterator FChannelGroupModel::GetTrackAreaModelList() const
-{
-	return &Channels;
-}
-
-TSharedRef<SWidget> FChannelGroupModel::CreateOutlinerView(const FCreateOutlinerViewParams& InParams)
-{
 	TSharedPtr<FSequencerEditorViewModel> EditorViewModel = GetEditor();
 
 	return SNew(SOutlinerItemViewBase, SharedThis(this), InParams.Editor, InParams.TreeViewRow)
@@ -215,7 +263,7 @@
 		// Even if this key area node doesn't have any key areas right now, it may in the future
 		// so we always create the switcher, and just hide it if it is not relevant
 		SNew(SHorizontalBox)
-		.Visibility(this, &FChannelGroupModel::GetKeyEditorVisibility)
+		.Visibility(this, &FChannelGroupOutlinerModel::GetKeyEditorVisibility)
 
 		+ SHorizontalBox::Slot()
 		.HAlign(HAlign_Right)
@@ -233,69 +281,34 @@
 	];
 }
 
-EVisibility FChannelGroupModel::GetKeyEditorVisibility() const
+EVisibility FChannelGroupOutlinerModel::GetKeyEditorVisibility() const
 {
 	return GetChannels().Num() == 0 ? EVisibility::Collapsed : EVisibility::Visible;
 }
 
-FText FChannelGroupModel::GetLabel() const
+FText FChannelGroupOutlinerModel::GetLabel() const
 {
 	return GetDisplayText();
 }
 
-FSlateFontInfo FChannelGroupModel::GetLabelFont() const
+FSlateFontInfo FChannelGroupOutlinerModel::GetLabelFont() const
 {
 	return IsAnimated()
 		? FAppStyle::GetFontStyle("Sequencer.AnimationOutliner.ItalicFont")
-		: FOutlinerItemModel::GetLabelFont();
+		: FOutlinerItemModelMixin::GetLabelFont();
 }
 
-bool FChannelGroupModel::HasCurves() const
+bool FChannelGroupOutlinerModel::HasCurves() const
 {
-	for (const TSharedRef<IKeyArea>& KeyArea : GetAllKeyAreas())
-	{
-		const ISequencerChannelInterface* EditorInterface = KeyArea->FindChannelEditorInterface();
-		if (EditorInterface && EditorInterface->SupportsCurveEditorModels_Raw(KeyArea->GetChannel()))
-		{
-			return true;
-		}
-	}
-	return false;
+	return FChannelGroupModel::HasCurves();
 }
 
-void FChannelGroupModel::CreateCurveModels(TArray<TUniquePtr<FCurveModel>>& OutCurveModels)
+bool FChannelGroupOutlinerModel::CanDelete(FText* OutErrorMessage) const
 {
-	TSharedPtr<FSequenceModel> SequenceModel = FindAncestorOfType<FSequenceModel>();
-	if (!SequenceModel)
-	{
-		return;
-	}
-
-	TSharedPtr<ISequencer> Sequencer = SequenceModel->GetSequencer();
-	if (!Sequencer)
-	{
-		return;
-	}
-
-	for (TWeakViewModelPtr<FChannelModel> WeakChannel : GetChannels())
-	{
-		if (TSharedPtr<FChannelModel> Channel = WeakChannel.Pin())
-		{
-			TUniquePtr<FCurveModel> NewCurve = Channel->GetKeyArea()->CreateCurveEditorModel(Sequencer.ToSharedRef());
-			if (NewCurve.IsValid())
-			{
-				OutCurveModels.Add(MoveTemp(NewCurve));
-			}
-		}
-	}
-}
-
-bool FChannelGroupModel::CanDelete(FText* OutErrorMessage) const
-{
 	return true;
 }
 
-void FChannelGroupModel::Delete()
+void FChannelGroupOutlinerModel::Delete()
 {
 	TArray<FName> PathFromTrack;
 	TViewModelPtr<ITrackExtension> Track = GetParentTrackNodeAndNamePath(this, PathFromTrack);
@@ -311,6 +324,11 @@
 	}
 }
 
+void FChannelGroupOutlinerModel::CreateCurveModels(TArray<TUniquePtr<FCurveModel>>& OutCurveModels)
+{
+	FChannelGroupModel::CreateCurveModels(OutCurveModels);
+}
+
 } // namespace Sequencer
 } // namespace UE
 
diff --git a/Engine/Source/Editor/Sequencer/Private/MVVM/ViewModels/OutlinerItemModel.cpp b/Engine/Source/Editor/Sequencer/Private/MVVM/ViewModels/OutlinerItemModel.cpp
--- a/Engine/Source/Editor/Sequencer/Private/MVVM/ViewModels/OutlinerItemModel.cpp
+++ b/Engine/Source/Editor/Sequencer/Private/MVVM/ViewModels/OutlinerItemModel.cpp
@@ -126,42 +126,41 @@
 	);
 }
 
-FOutlinerItemModel::FOutlinerItemModel()
+FOutlinerItemModelMixin::FOutlinerItemModelMixin()
 	: OutlinerChildList(EViewModelListType::Outliner)
 	, bInitializedExpansion(false)
 	, bInitializedPinnedState(false)
 	, bIsForcedFilteredOut(false)
 {
-	RegisterChildList(&OutlinerChildList);
 }
 
-TSharedPtr<FSequencerEditorViewModel> FOutlinerItemModel::GetEditor() const
+TSharedPtr<FSequencerEditorViewModel> FOutlinerItemModelMixin::GetEditor() const
 {
-	TSharedPtr<FSequenceModel> SequenceModel = FindAncestorOfType<FSequenceModel>();
+	TSharedPtr<FSequenceModel> SequenceModel = AsViewModel()->FindAncestorOfType<FSequenceModel>();
 	return SequenceModel ? SequenceModel->GetEditor() : nullptr;
 }
 
-bool FOutlinerItemModel::IsForceFilteredOut() const
+bool FOutlinerItemModelMixin::IsForceFilteredOut() const
 {
 	return bIsForcedFilteredOut;
 }
 
-void FOutlinerItemModel::SetForceFilteredOut(bool bInIsForcedFilteredOut)
+void FOutlinerItemModelMixin::SetForceFilteredOut(bool bInIsForcedFilteredOut)
 {
 	bIsForcedFilteredOut = bInIsForcedFilteredOut;
 }
 
-FName FOutlinerItemModel::GetIdentifier() const
+FName FOutlinerItemModelMixin::GetIdentifier() const
 {
 	return TreeItemIdentifier;
 }
 
-void FOutlinerItemModel::SetIdentifier(FName InNewIdentifier)
+void FOutlinerItemModelMixin::SetIdentifier(FName InNewIdentifier)
 {
 	TreeItemIdentifier = InNewIdentifier;
 }
 
-bool FOutlinerItemModel::IsExpanded() const
+bool FOutlinerItemModelMixin::IsExpanded() const
 {
 	if (bInitializedExpansion)
 	{
@@ -171,9 +170,9 @@
 	bInitializedExpansion = true;
 
 	TStringBuilder<256> StringBuilder;
-	IOutlinerExtension::GetPathName(*this, StringBuilder);
+	IOutlinerExtension::GetPathName(*AsViewModel(), StringBuilder);
 
-	TSharedPtr<FSequenceModel> SequenceModel = FindAncestorOfType<FSequenceModel>();
+	TSharedPtr<FSequenceModel> SequenceModel = AsViewModel()->FindAncestorOfType<FSequenceModel>();
 	UMovieSceneSequence*       Sequence      = SequenceModel ? SequenceModel->GetSequence() : nullptr;
 
 	if (Sequence)
@@ -183,39 +182,40 @@
 		FMovieSceneEditorData& EditorData = MovieScene->GetEditorData();
 		if (const FMovieSceneExpansionState* Expansion = EditorData.ExpansionStates.FindByHash(GetTypeHash(StringView), StringView))
 		{
-			const_cast<FOutlinerItemModel*>(this)->bIsExpanded = Expansion->bExpanded;
+			const_cast<FOutlinerItemModelMixin*>(this)->bIsExpanded = Expansion->bExpanded;
 		}
 		else
 		{
-			const_cast<FOutlinerItemModel*>(this)->bIsExpanded = GetDefaultExpansionState();
+			const_cast<FOutlinerItemModelMixin*>(this)->bIsExpanded = GetDefaultExpansionState();
 		}
 	}
 
 	return bIsExpanded;
 }
 
-bool FOutlinerItemModel::GetDefaultExpansionState() const
+bool FOutlinerItemModelMixin::GetDefaultExpansionState() const
 {
 	return false;
 }
 
-void FOutlinerItemModel::SetExpansion(bool bInIsExpanded)
+void FOutlinerItemModelMixin::SetExpansion(bool bInIsExpanded)
 {
 	SetExpansionWithoutSaving(bInIsExpanded);
 
-	if (GetParent())
+	FViewModel* ViewModel = AsViewModel();
+	if (ViewModel->GetParent())
 	{
 		// Expansion state has changed, save it to the movie scene now
-		TSharedPtr<FSequenceModel> SequenceModel = FindAncestorOfType<FSequenceModel>();
+		TSharedPtr<FSequenceModel> SequenceModel = ViewModel->FindAncestorOfType<FSequenceModel>();
 		if (SequenceModel)
 		{
 			TSharedPtr<FSequencer> Sequencer = SequenceModel->GetSequencerImpl();
-			Sequencer->GetNodeTree()->SaveExpansionState(*this, bInIsExpanded);
+			Sequencer->GetNodeTree()->SaveExpansionState(*ViewModel, bInIsExpanded);
 		}
 	}
 }
 
-void FOutlinerItemModel::SetExpansionWithoutSaving(bool bInIsExpanded)
+void FOutlinerItemModelMixin::SetExpansionWithoutSaving(bool bInIsExpanded)
 {
 	FOutlinerExtensionShim::SetExpansion(bInIsExpanded);
 
@@ -224,12 +224,12 @@
 	bInitializedExpansion = true;
 }
 
-bool FOutlinerItemModel::IsFilteredOut() const
+bool FOutlinerItemModelMixin::IsFilteredOut() const
 {
 	return bIsFilteredOut || bIsForcedFilteredOut;
 }
 
-bool FOutlinerItemModel::IsPinned() const
+bool FOutlinerItemModelMixin::IsPinned() const
 {
 	if (bInitializedPinnedState)
 	{
@@ -240,25 +240,25 @@
 
 	// Initialize expansion states for tree items
 	// Assign the saved expansion state when this node is initialized for the first time
-	const bool bIsRootModel = (GetHierarchicalDepth() == 1);
+	const bool bIsRootModel = (AsViewModel()->GetHierarchicalDepth() == 1);
 	if (bIsRootModel)
 	{
-		TSharedPtr<FSequenceModel> SequenceModel = FindAncestorOfType<FSequenceModel>();
+		TSharedPtr<FSequenceModel> SequenceModel = AsViewModel()->FindAncestorOfType<FSequenceModel>();
 		TSharedPtr<FSequencer> Sequencer = SequenceModel->GetSequencerImpl();
-		const bool bWasPinned = Sequencer->GetNodeTree()->GetSavedPinnedState(*this);
-		const_cast<FOutlinerItemModel*>(this)->FPinnableExtensionShim::SetPinned(bWasPinned);
+		const bool bWasPinned = Sequencer->GetNodeTree()->GetSavedPinnedState(*AsViewModel());
+		const_cast<FOutlinerItemModelMixin*>(this)->FPinnableExtensionShim::SetPinned(bWasPinned);
 	}
 
 	return FPinnableExtensionShim::IsPinned();
 }
 
-bool FOutlinerItemModel::IsDimmed() const
+bool FOutlinerItemModelMixin::IsDimmed() const
 {
 	TSharedPtr<FSequencerEditorViewModel> EditorViewModel = GetEditor();
 	TSharedPtr<FSequencer> Sequencer = EditorViewModel ? EditorViewModel->GetSequencerImpl() : nullptr;
 	if (Sequencer)
 	{
-		TSharedPtr<FViewModel> ThisShared(const_cast<FOutlinerItemModel*>(this)->AsShared());
+		TSharedPtr<FViewModel> ThisShared(const_cast<FViewModel*>(AsViewModel())->AsShared());
 		FSequencerNodeTree& ParentTree = Sequencer->GetNodeTree().Get();
 		if (ParentTree.IsNodeMute(ThisShared) || (ParentTree.HasSoloNodes() && !ParentTree.IsNodeSolo(ThisShared)))
 		{
@@ -268,40 +268,40 @@
 	return false;
 }
 
-bool FOutlinerItemModel::IsSolo() const
+bool FOutlinerItemModelMixin::IsSolo() const
 {
 	TSharedPtr<FSequencerEditorViewModel> EditorViewModel = GetEditor();
 	TSharedPtr<FSequencer> Sequencer = EditorViewModel ? EditorViewModel->GetSequencerImpl() : nullptr;
 	if (Sequencer)
 	{
-		TSharedPtr<FViewModel> ThisShared(const_cast<FOutlinerItemModel*>(this)->AsShared());
+		TSharedPtr<FViewModel> ThisShared(const_cast<FViewModel*>(AsViewModel())->AsShared());
 		return Sequencer->GetNodeTree()->IsNodeSolo(ThisShared);
 	}
 	return false;
 }
 
-bool FOutlinerItemModel::IsMuted() const
+bool FOutlinerItemModelMixin::IsMuted() const
 {
 	TSharedPtr<FSequencerEditorViewModel> EditorViewModel = GetEditor();
 	TSharedPtr<FSequencer> Sequencer = EditorViewModel ? EditorViewModel->GetSequencerImpl() : nullptr;
 	if (Sequencer)
 	{
-		TSharedPtr<FViewModel> ThisShared(const_cast<FOutlinerItemModel*>(this)->AsShared());
+		TSharedPtr<FViewModel> ThisShared(const_cast<FViewModel*>(AsViewModel())->AsShared());
 		return Sequencer->GetNodeTree()->IsNodeMute(ThisShared);
 	}
 	return false;
 }
 
-bool FOutlinerItemModel::IsRootModelPinned() const
+bool FOutlinerItemModelMixin::IsRootModelPinned() const
 {
-	TSharedPtr<IPinnableExtension> PinnableParent = FindAncestorOfType<IPinnableExtension>(true);
+	TSharedPtr<IPinnableExtension> PinnableParent = AsViewModel()->FindAncestorOfType<IPinnableExtension>(true);
 	return PinnableParent && PinnableParent->IsPinned();
 }
 
-void FOutlinerItemModel::ToggleRootModelPinned()
+void FOutlinerItemModelMixin::ToggleRootModelPinned()
 {
-	FSequenceModel* RootModel = GetRoot()->CastThis<FSequenceModel>();
-	TSharedPtr<IPinnableExtension> PinnableParent = FindAncestorOfType<IPinnableExtension>(true);
+	FSequenceModel* RootModel = AsViewModel()->GetRoot()->CastThis<FSequenceModel>();
+	TSharedPtr<IPinnableExtension> PinnableParent = AsViewModel()->FindAncestorOfType<IPinnableExtension>(true);
 	if (RootModel && PinnableParent)
 	{
 		TSharedPtr<FOutlinerViewModel> Outliner = RootModel->GetEditor()->GetOutliner();
@@ -311,12 +311,12 @@
 		PinnableParent->SetPinned(bShouldPin);
 
 		TSharedPtr<FSequencer> Sequencer = RootModel->GetSequencerImpl();
-		Sequencer->GetNodeTree()->SavePinnedState(*this, bShouldPin);
+		Sequencer->GetNodeTree()->SavePinnedState(*AsViewModel(), bShouldPin);
 		Sequencer->RefreshTree();
 	}
 }
 
-bool FOutlinerItemModel::IsSelectedModelsSolo() const
+bool FOutlinerItemModelMixin::IsSelectedModelsSolo() const
 {
 	TSharedPtr<FSequencerEditorViewModel> EditorViewModel = GetEditor();
 	TSharedPtr<FSequencer> Sequencer = EditorViewModel ? EditorViewModel->GetSequencerImpl() : nullptr;
@@ -327,7 +327,7 @@
 	return false;
 }
 
-void FOutlinerItemModel::ToggleSelectedModelsSolo()
+void FOutlinerItemModelMixin::ToggleSelectedModelsSolo()
 {
 	TSharedPtr<FSequencerEditorViewModel> EditorViewModel = GetEditor();
 	TSharedPtr<FSequencer> Sequencer = EditorViewModel ? EditorViewModel->GetSequencerImpl() : nullptr;
@@ -337,7 +337,7 @@
 	}
 }
 
-bool FOutlinerItemModel::IsSelectedModelsMuted() const
+bool FOutlinerItemModelMixin::IsSelectedModelsMuted() const
 {
 	TSharedPtr<FSequencerEditorViewModel> EditorViewModel = GetEditor();
 	TSharedPtr<FSequencer> Sequencer = EditorViewModel ? EditorViewModel->GetSequencerImpl() : nullptr;
@@ -348,7 +348,7 @@
 	return false;
 }
 
-void FOutlinerItemModel::ToggleSelectedModelsMuted()
+void FOutlinerItemModelMixin::ToggleSelectedModelsMuted()
 {
 	TSharedPtr<FSequencerEditorViewModel> EditorViewModel = GetEditor();
 	TSharedPtr<FSequencer> Sequencer = EditorViewModel ? EditorViewModel->GetSequencerImpl() : nullptr;
@@ -358,7 +358,7 @@
 	}
 }
 
-TSharedPtr<SWidget> FOutlinerItemModel::CreateContextMenuWidget(const FCreateOutlinerContextMenuWidgetParams& InParams)
+TSharedPtr<SWidget> FOutlinerItemModelMixin::CreateContextMenuWidget(const FCreateOutlinerContextMenuWidgetParams& InParams)
 {
 	TSharedPtr<FSequencerEditorViewModel> EditorViewModel = GetEditor();
 	TSharedPtr<ISequencer> Sequencer = EditorViewModel->GetSequencer();
@@ -376,7 +376,7 @@
 	return nullptr;
 }
 
-void FOutlinerItemModel::BuildContextMenu(FMenuBuilder& MenuBuilder)
+void FOutlinerItemModelMixin::BuildContextMenu(FMenuBuilder& MenuBuilder)
 {
 	TSharedPtr<FSequencer> Sequencer = StaticCastSharedPtr<FSequencer>(GetEditor()->GetSequencer());
 
@@ -385,6 +385,8 @@
 		return;
 	}
 
+	TSharedRef<FOutlinerItemModelMixin> SharedThis(AsViewModel()->AsShared(), this);
+
 	const bool bIsReadOnly = Sequencer->IsReadOnly();
 	FCanExecuteAction CanExecute = FCanExecuteAction::CreateLambda([bIsReadOnly]{ return !bIsReadOnly; });
 
@@ -404,7 +406,7 @@
 		);
 
 		// Only support pinning root nodes
-		const bool bIsRootModel = (GetHierarchicalDepth() == 1);
+		const bool bIsRootModel = (AsViewModel()->GetHierarchicalDepth() == 1);
 		if (bIsRootModel)
 		{
 			MenuBuilder.AddMenuEntry(
@@ -412,9 +414,9 @@
 				LOCTEXT("ToggleNodePinTooltip", "Pin or unpin this node or selected tracks"),
 				FSlateIcon(),
 				FUIAction(
-					FExecuteAction::CreateSP(this, &FOutlinerItemModel::ToggleRootModelPinned),
+					FExecuteAction::CreateSP(SharedThis, &FOutlinerItemModelMixin::ToggleRootModelPinned),
 					FCanExecuteAction(),
-					FIsActionChecked::CreateSP(this, &FOutlinerItemModel::IsRootModelPinned)
+					FIsActionChecked::CreateSP(SharedThis, &FOutlinerItemModelMixin::IsRootModelPinned)
 				),
 				NAME_None,
 				EUserInterfaceActionType::ToggleButton
@@ -427,9 +429,9 @@
 			LOCTEXT("ToggleNodeSoloTooltip", "Solo or unsolo this node or selected tracks"),
 			FSlateIcon(),
 			FUIAction(
-				FExecuteAction::CreateSP(this, &FOutlinerItemModel::ToggleSelectedModelsSolo),
+				FExecuteAction::CreateSP(SharedThis, &FOutlinerItemModelMixin::ToggleSelectedModelsSolo),
 				CanExecute,
-				FIsActionChecked::CreateSP(this, &FOutlinerItemModel::IsSelectedModelsSolo)
+				FIsActionChecked::CreateSP(SharedThis, &FOutlinerItemModelMixin::IsSelectedModelsSolo)
 			),
 			NAME_None,
 			EUserInterfaceActionType::ToggleButton
@@ -440,9 +442,9 @@
 			LOCTEXT("ToggleNodeMuteTooltip", "Mute or unmute this node or selected tracks"),
 			FSlateIcon(),
 			FUIAction(
-				FExecuteAction::CreateSP(this, &FOutlinerItemModel::ToggleSelectedModelsMuted),
+				FExecuteAction::CreateSP(SharedThis, &FOutlinerItemModelMixin::ToggleSelectedModelsMuted),
 				CanExecute,
-				FIsActionChecked::CreateSP(this, &FOutlinerItemModel::IsSelectedModelsMuted)
+				FIsActionChecked::CreateSP(SharedThis, &FOutlinerItemModelMixin::IsSelectedModelsMuted)
 			),
 			NAME_None,
 			EUserInterfaceActionType::ToggleButton
@@ -457,7 +459,7 @@
 		
 		MenuBuilder.AddMenuEntry(FGenericCommands::Get().Duplicate);
 		
-		TSharedRef<FViewModel> ThisNode = AsShared();
+		TSharedRef<FViewModel> ThisNode = AsViewModel()->AsShared();
 
 		MenuBuilder.AddMenuEntry(
 			LOCTEXT("DeleteNode", "Delete"),
@@ -466,7 +468,7 @@
 			FUIAction(FExecuteAction::CreateSP(Sequencer.Get(), &FSequencer::DeleteNode, ThisNode, false), CanExecute)
 		);
 
-		if (IsA<IObjectBindingExtension>())
+		if (ThisNode->IsA<IObjectBindingExtension>())
 		{
 			MenuBuilder.AddMenuEntry(
 				LOCTEXT("DeleteNodeAndKeepState", "Delete and Keep State"),
@@ -539,12 +541,12 @@
 	}
 }
 
-void FOutlinerItemModel::BuildOrganizeContextMenu(FMenuBuilder& MenuBuilder)
+void FOutlinerItemModelMixin::BuildOrganizeContextMenu(FMenuBuilder& MenuBuilder)
 {
-	TSharedRef<FViewModel> ThisNode = SharedThis(this);
+	TSharedRef<FViewModel> ThisNode = AsViewModel()->AsShared();
 	TSharedPtr<FSequencer> Sequencer = GetEditor()->GetSequencerImpl();
 
-	const bool bFilterableNode = (IsA<ITrackExtension>() || IsA<IObjectBindingExtension>() || IsA<FFolderModel>());
+	const bool bFilterableNode = (ThisNode->IsA<ITrackExtension>() || ThisNode->IsA<IObjectBindingExtension>() || ThisNode->IsA<FFolderModel>());
 	const bool bIsReadOnly = Sequencer->IsReadOnly();
 	
 	TArray<UMovieSceneTrack*> AllTracks;
@@ -597,20 +599,23 @@
 	}
 }
 
-bool FOutlinerItemModel::HasCurves() const
+bool FOutlinerItemModelMixin::HasCurves() const
 {
 	return false;
 }
 
-TSharedPtr<ICurveEditorTreeItem> FOutlinerItemModel::GetCurveEditorTreeItem() const
+TSharedPtr<ICurveEditorTreeItem> FOutlinerItemModelMixin::GetCurveEditorTreeItem() const
 {
-	return StaticCastSharedRef<ICurveEditorTreeItem>(SharedThis(const_cast<FOutlinerItemModel*>(this)));
+	TSharedRef<FViewModel> ThisShared(const_cast<FViewModel*>(AsViewModel())->AsShared());
+	return TSharedPtr<ICurveEditorTreeItem>(ThisShared, const_cast<FOutlinerItemModelMixin*>(this));
 }
 
-TSharedPtr<SWidget> FOutlinerItemModel::GenerateCurveEditorTreeWidget(const FName& InColumnName, TWeakPtr<FCurveEditor> InCurveEditor, FCurveEditorTreeItemID InTreeItemID, const TSharedRef<ITableRow>& TableRow)
+TSharedPtr<SWidget> FOutlinerItemModelMixin::GenerateCurveEditorTreeWidget(const FName& InColumnName, TWeakPtr<FCurveEditor> InCurveEditor, FCurveEditorTreeItemID InTreeItemID, const TSharedRef<ITableRow>& TableRow)
 {
 	using namespace UE::Sequencer;
 
+	TSharedRef<FOutlinerItemModelMixin> SharedThis(AsViewModel()->AsShared(), this);
+
 	auto GetCurveEditorHighlightText = [](TWeakPtr<FCurveEditor> InCurveEditor) -> FText 
 	{
 		TSharedPtr<FCurveEditor> PinnedCurveEditor = InCurveEditor.Pin();
@@ -642,8 +647,8 @@
 				+ SOverlay::Slot()
 				[
 					SNew(SImage)
-					.Image(this, &FOutlinerItemModel::GetIconBrush)
-					.ColorAndOpacity(this, &FOutlinerItemModel::GetIconTint)
+					.Image(SharedThis, &FOutlinerItemModelMixin::GetIconBrush)
+					.ColorAndOpacity(SharedThis, &FOutlinerItemModelMixin::GetIconTint)
 				]
 
 				+ SOverlay::Slot()
@@ -651,14 +656,14 @@
 				.HAlign(HAlign_Right)
 				[
 					SNew(SImage)
-					.Image(this, &FOutlinerItemModel::GetIconOverlayBrush)
+					.Image(SharedThis, &FOutlinerItemModelMixin::GetIconOverlayBrush)
 				]
 
 				+ SOverlay::Slot()
 				[
 					SNew(SSpacer)
 					.Visibility(EVisibility::Visible)
-					.ToolTipText(this, &FOutlinerItemModel::GetIconToolTipText)
+					.ToolTipText(SharedThis, &FOutlinerItemModelMixin::GetIconToolTipText)
 				]
 			]
 
@@ -667,10 +672,10 @@
 			.Padding(FMargin(0.f, 4.f, 0.f, 4.f))
 			[
 				SNew(STextBlock)
-				.Text(this, &FOutlinerItemModel::GetLabel)
-				.Font(this, &FOutlinerItemModel::GetLabelFont)
+				.Text(SharedThis, &FOutlinerItemModelMixin::GetLabel)
+				.Font(SharedThis, &FOutlinerItemModelMixin::GetLabelFont)
 				.HighlightText_Static(GetCurveEditorHighlightText, InCurveEditor)
-				.ToolTipText(this, &FOutlinerItemModel::GetLabelToolTipText)
+				.ToolTipText(SharedThis, &FOutlinerItemModelMixin::GetLabelToolTipText)
 			];
 	}
 	else if (InColumnName == ColumnNames.SelectHeader)
@@ -685,17 +690,17 @@
 	return nullptr;
 }
 
-void FOutlinerItemModel::CreateCurveModels(TArray<TUniquePtr<FCurveModel>>& OutCurveModels)
+void FOutlinerItemModelMixin::CreateCurveModels(TArray<TUniquePtr<FCurveModel>>& OutCurveModels)
 {
 }
 
-bool FOutlinerItemModel::PassesFilter(const FCurveEditorTreeFilter* InFilter) const
+bool FOutlinerItemModelMixin::PassesFilter(const FCurveEditorTreeFilter* InFilter) const
 {
 	if (InFilter->GetType() == ECurveEditorTreeFilterType::Text)
 	{
 		const FCurveEditorTreeTextFilter* Filter = static_cast<const FCurveEditorTreeTextFilter*>(InFilter);
 
-		TSharedPtr<const IOutlinerExtension> This = CastThisShared<IOutlinerExtension>();
+		TSharedPtr<const IOutlinerExtension> This = AsViewModel()->CastThisShared<IOutlinerExtension>();
 		for (const FCurveEditorTreeTextFilterTerm& Term : Filter->GetTerms())
 		{
 			if (NodeMatchesTextFilterTerm(This, Term))
@@ -709,7 +714,7 @@
 	else if (InFilter->GetType() == ISequencerModule::GetSequencerSelectionFilterType())
 	{
 		const FSequencerSelectionCurveFilter* Filter = static_cast<const FSequencerSelectionCurveFilter*>(InFilter);
-		return Filter->Match(AsShared());
+		return Filter->Match(AsViewModel()->AsShared());
 	}
 	return false;
 }
diff --git a/Engine/Source/Editor/Sequencer/Private/MVVM/ViewModels/TrackModelLayoutBuilder.cpp b/Engine/Source/Editor/Sequencer/Private/MVVM/ViewModels/TrackModelLayoutBuilder.cpp
--- a/Engine/Source/Editor/Sequencer/Private/MVVM/ViewModels/TrackModelLayoutBuilder.cpp
+++ b/Engine/Source/Editor/Sequencer/Private/MVVM/ViewModels/TrackModelLayoutBuilder.cpp
@@ -243,11 +243,26 @@
 	//       of channels. We have to be careful to ensure that we reuse existing items that may have been added
 	//       to the outliner from a previously built section.
 
-	TSharedPtr<FChannelGroupModel> OutlinerModel = OutlinerList.FindItem<FChannelGroupModel>(ChannelNamePredicate);
+	TSharedPtr<FChannelGroupModel> OutlinerModel;
+	if (bIsTopLevel)
+	{
+		OutlinerModel = OutlinerList.FindItem<FChannelGroupModel>(ChannelNamePredicate);
+	}
+	else
+	{
+		OutlinerModel = OutlinerList.FindItem<FChannelGroupOutlinerModel>(ChannelNamePredicate);
+	}
+
 	if (!OutlinerModel)
 	{
-		OutlinerModel = MakeShared<FChannelGroupModel>(ChannelName, MetaData->DisplayText);
-		OutlinerModel->SetForceFilteredOut(bIsTopLevel);
+		if (bIsTopLevel)
+		{
+			OutlinerModel = MakeShared<FChannelGroupModel>(ChannelName, MetaData->DisplayText);
+		}
+		else
+		{
+			OutlinerModel = MakeShared<FChannelGroupOutlinerModel>(ChannelName, MetaData->DisplayText);
+		}
 	}
 
 	OutlinerList.Link(OutlinerModel);
diff --git a/Engine/Source/Editor/Sequencer/Private/MVVM/ViewModels/TrackModelLayoutBuilder.h b/Engine/Source/Editor/Sequencer/Private/MVVM/ViewModels/TrackModelLayoutBuilder.h
--- a/Engine/Source/Editor/Sequencer/Private/MVVM/ViewModels/TrackModelLayoutBuilder.h
+++ b/Engine/Source/Editor/Sequencer/Private/MVVM/ViewModels/TrackModelLayoutBuilder.h
@@ -35,14 +35,14 @@
 	void Pop();
 
 	template<typename ModelType, typename Predicate>
-	TViewModelPtr<ModelType> FindItem(Predicate&& InPredicate) const
+	TSharedPtr<ModelType> FindItem(Predicate&& InPredicate) const
 	{
 		TViewModelPtr<ModelType> Model = FindExistingItem<ModelType>(InPredicate);
 		return Model ? Model : FindRecycledItem<ModelType>(InPredicate);
 	}
 
 	template<typename ModelType, typename Predicate>
-	TViewModelPtr<ModelType> FindExistingItem(Predicate&& InPredicate) const
+	TSharedPtr<ModelType> FindExistingItem(Predicate&& InPredicate) const
 	{
 		// Look for existing children
 		for (const TViewModelPtr<ModelType>& Item : ListData.Last().Children.IterateSubList<ModelType>())
diff --git a/Engine/Source/Editor/Sequencer/Private/SequencerAddKeyOperation.cpp b/Engine/Source/Editor/Sequencer/Private/SequencerAddKeyOperation.cpp
--- a/Engine/Source/Editor/Sequencer/Private/SequencerAddKeyOperation.cpp
+++ b/Engine/Source/Editor/Sequencer/Private/SequencerAddKeyOperation.cpp
@@ -114,7 +114,9 @@ bool FAddKeyOperation::ProcessKeyArea(TSharedPtr<ITrackExtension> InTrackModel,
 {
 	bool bKeyedAnything = false;
 	ISequencerTrackEditor* TrackEditor = InTrackModel->GetTrackEditor().Get();
-	if (InChannelGroupModel->IsFilteredOut() == false)
+
+	IOutlinerExtension* OutlinerExtension = InChannelGroupModel->CastThis<IOutlinerExtension>();
+	if (!OutlinerExtension || OutlinerExtension->IsFilteredOut() == false)
 	{
 		constexpr bool bIncludeThis = true;
 		for (const TWeakViewModelPtr<FChannelModel>& WeakChannel : InChannelGroupModel->GetChannels())
diff --git a/Engine/Source/Editor/Sequencer/Private/SequencerKeyCollection.cpp b/Engine/Source/Editor/Sequencer/Private/SequencerKeyCollection.cpp
--- a/Engine/Source/Editor/Sequencer/Private/SequencerKeyCollection.cpp
+++ b/Engine/Source/Editor/Sequencer/Private/SequencerKeyCollection.cpp
@@ -17,7 +17,8 @@ FSequencerKeyCollectionSignature FSequencerKeyCollectionSignature::FromNodes(con
 	{
 		for (TViewModelPtr<const FChannelGroupModel> ChannelGroupModel : Node->GetChildrenOfType<FChannelGroupModel>())
 		{
-			if (ChannelGroupModel->IsFilteredOut() == false)
+			const IOutlinerExtension* OutlinerExtension = ChannelGroupModel->CastThis<IOutlinerExtension>();
+			if (!OutlinerExtension || OutlinerExtension->IsFilteredOut() == false)
 			{
 				for (const TSharedRef<IKeyArea>& KeyArea : ChannelGroupModel->GetAllKeyAreas())
 				{
@@ -45,7 +46,8 @@ FSequencerKeyCollectionSignature FSequencerKeyCollectionSignature::FromNodesRecu
 		const bool bIncludeThis = true;
 		for (const TViewModelPtr<FChannelGroupModel>& KeyAreaNode : Node->GetDescendantsOfType<FChannelGroupModel>(bIncludeThis))
 		{
-			if (KeyAreaNode->IsFilteredOut() == false)
+			IOutlinerExtension* OutlinerExtension = KeyAreaNode->CastThis<IOutlinerExtension>();
+			if (!OutlinerExtension || OutlinerExtension->IsFilteredOut() == false)
 			{
 				AllKeyAreaNodes.Add(KeyAreaNode);
 			}
@@ -75,7 +77,8 @@ FSequencerKeyCollectionSignature FSequencerKeyCollectionSignature::FromNodeRecur
 	AllKeyAreaNodes.Reserve(36);
 	for (TSharedPtr<FChannelGroupModel> KeyAreaNode : InNode->GetDescendantsOfType<FChannelGroupModel>(true))
 	{
-		if (KeyAreaNode->IsFilteredOut() == false)
+		IOutlinerExtension* OutlinerExtension = KeyAreaNode->CastThis<IOutlinerExtension>();
+		if (!OutlinerExtension || OutlinerExtension->IsFilteredOut() == false)
 		{
 			AllKeyAreaNodes.Add(KeyAreaNode);
 		}
diff --git a/Engine/Source/Editor/Sequencer/Private/Tools/SequencerEntityVisitor.cpp b/Engine/Source/Editor/Sequencer/Private/Tools/SequencerEntityVisitor.cpp
--- a/Engine/Source/Editor/Sequencer/Private/Tools/SequencerEntityVisitor.cpp
+++ b/Engine/Source/Editor/Sequencer/Private/Tools/SequencerEntityVisitor.cpp
@@ -166,7 +166,7 @@ void FSequencerEntityWalker::VisitAnyChannels(const ISequencerEntityVisitor& Vis
 	if (ChannelGroup)
 	{
 		const IOutlinerExtension* OutlinerItem = ChannelGroup->CastThis<IOutlinerExtension>();
-		if (OutlinerItem->IsFilteredOut() == false)
+		if (!OutlinerItem || OutlinerItem->IsFilteredOut() == false)
 		{
 			for (const TWeakViewModelPtr<FChannelModel>& WeakChannel : ChannelGroup->GetChannels())
 			{
diff --git a/Engine/Source/Editor/Sequencer/Public/MVVM/ViewModels/ChannelModel.h b/Engine/Source/Editor/Sequencer/Public/MVVM/ViewModels/ChannelModel.h
--- a/Engine/Source/Editor/Sequencer/Public/MVVM/ViewModels/ChannelModel.h
+++ b/Engine/Source/Editor/Sequencer/Public/MVVM/ViewModels/ChannelModel.h
@@ -81,14 +81,12 @@
  * For instance, this represents the "Location.X" entry in the Sequencer outliner.
  */
 class FChannelGroupModel
-	: public FOutlinerItemModel
+	: public FViewModel
 	, public ITrackAreaExtension
-	, public ICompoundOutlinerExtension
-	, public IDeletableExtension
 {
 public:
 
-	UE_SEQUENCER_DECLARE_CASTABLE(FChannelGroupModel, FOutlinerItemModel, ITrackAreaExtension, ICompoundOutlinerExtension, IDeletableExtension);
+	UE_SEQUENCER_DECLARE_CASTABLE(FChannelGroupModel, FViewModel, ITrackAreaExtension);
 
 	FChannelGroupModel(FName InChannelName, const FText& InDisplayText);
 	~FChannelGroupModel();
@@ -124,6 +122,39 @@
 
 public:
 
+	/*~ ITrackAreaExtension */
+	FTrackAreaParameters GetTrackAreaParameters() const override;
+	FViewModelVariantIterator GetTrackAreaModelList() const override;
+
+	void CreateCurveModels(TArray<TUniquePtr<FCurveModel>>& OutCurveModels);
+	bool HasCurves() const;
+
+protected:
+
+	TArray<TWeakViewModelPtr<FChannelModel>> Channels;
+	FName ChannelName;
+	FText DisplayText;
+};
+
+
+/**
+ * Model for the outliner entry associated with all sections' channels of a given common name.
+ * For instance, this represents the "Location.X" entry in the Sequencer outliner.
+ */
+class FChannelGroupOutlinerModel
+	: public TOutlinerModelMixin<FChannelGroupModel>
+	, public ICompoundOutlinerExtension
+	, public IDeletableExtension
+{
+public:
+
+	UE_SEQUENCER_DECLARE_CASTABLE(FChannelGroupOutlinerModel, FChannelGroupModel, FOutlinerItemModelMixin, ICompoundOutlinerExtension, IDeletableExtension);
+
+	FChannelGroupOutlinerModel(FName InChannelName, const FText& InDisplayText);
+	~FChannelGroupOutlinerModel();
+
+public:
+
 	/*~ ICompoundOutlinerExtension */
 	FOutlinerSizing RecomputeSizing() override;
 
@@ -133,10 +164,6 @@
 	FSlateFontInfo GetLabelFont() const override;
 	TSharedRef<SWidget> CreateOutlinerView(const FCreateOutlinerViewParams& InParams) override;
 
-	/*~ ITrackAreaExtension */
-	FTrackAreaParameters GetTrackAreaParameters() const override;
-	FViewModelVariantIterator GetTrackAreaModelList() const override;
-
 	/*~ ICurveEditorTreeItem */
 	void CreateCurveModels(TArray<TUniquePtr<FCurveModel>>& OutCurveModels) override;
 
@@ -144,16 +171,12 @@
 	bool CanDelete(FText* OutErrorMessage) const override;
 	void Delete() override;
 
-	/*~ FOutlinerItemModel */
+	/*~ ICurveEditorTreeItemExtension */
 	bool HasCurves() const override;
 
 private:
 
 	EVisibility GetKeyEditorVisibility() const;
-
-	TArray<TWeakViewModelPtr<FChannelModel>> Channels;
-	FName ChannelName;
-	FText DisplayText;
 	FOutlinerSizing ComputedSizing;
 };
 
diff --git a/Engine/Source/Editor/Sequencer/Public/MVVM/ViewModels/OutlinerItemModel.h b/Engine/Source/Editor/Sequencer/Public/MVVM/ViewModels/OutlinerItemModel.h
--- a/Engine/Source/Editor/Sequencer/Public/MVVM/ViewModels/OutlinerItemModel.h
+++ b/Engine/Source/Editor/Sequencer/Public/MVVM/ViewModels/OutlinerItemModel.h
@@ -25,9 +25,8 @@
 
 class FSequencerEditorViewModel;
 
-class SEQUENCER_API FOutlinerItemModel
-	: public FViewModel
-	, public FOutlinerExtensionShim
+class SEQUENCER_API FOutlinerItemModelMixin
+	: public FOutlinerExtensionShim
 	, public FGeometryExtensionShim
 	, public FPinnableExtensionShim
 	, public FHoveredExtensionShim
@@ -39,9 +38,9 @@
 {
 public:
 
-	UE_SEQUENCER_DECLARE_CASTABLE(FOutlinerItemModel, FViewModel, IOutlinerExtension, IGeometryExtension, IPinnableExtension, IHoveredExtension, IDimmableExtension, IMutableExtension, ISoloableExtension, ICurveEditorTreeItemExtension);
+	using Implements = TImplements<IOutlinerExtension, IGeometryExtension, IPinnableExtension, IHoveredExtension, IDimmableExtension, IMutableExtension, ISoloableExtension, ICurveEditorTreeItemExtension>;
 
-	FOutlinerItemModel();
+	FOutlinerItemModelMixin();
 
 	TSharedPtr<FSequencerEditorViewModel> GetEditor() const;
 
@@ -91,8 +90,14 @@
 	/** Set expansion state without saving it in the movie-scene data */
 	void SetExpansionWithoutSaving(bool bInIsExpanded);
 
+private:
+
+	virtual FViewModel* AsViewModel() = 0;
+	virtual const FViewModel* AsViewModel() const = 0;
+
 protected:
 	FViewModelListHead OutlinerChildList;
+
 private:
 
 	bool IsRootModelPinned() const;
@@ -106,6 +111,7 @@
 
 private:
 
+	ICastable* CastableThis;
 	FName TreeItemIdentifier;
 	FCurveEditorTreeItemID CurveEditorItemID;
 	mutable bool bInitializedExpansion;
@@ -113,6 +119,38 @@
 	bool bIsForcedFilteredOut;
 };
 
+template<typename BaseType>
+class TOutlinerModelMixin : public BaseType, public FOutlinerItemModelMixin
+{
+public:
+	TOutlinerModelMixin()
+	{
+		this->RegisterChildList(&this->OutlinerChildList);
+	}
+
+	template<typename... ArgTypes>
+	TOutlinerModelMixin(ArgTypes&&... InArgs)
+		: BaseType(Forward<ArgTypes>(InArgs)...)
+	{
+		this->RegisterChildList(&this->OutlinerChildList);
+	}
+
+	TOutlinerModelMixin(const TOutlinerModelMixin<BaseType>&) = delete;
+	TOutlinerModelMixin<BaseType> operator=(const TOutlinerModelMixin<BaseType>&) = delete;
+
+	TOutlinerModelMixin(TOutlinerModelMixin<BaseType>&&) = delete;
+	TOutlinerModelMixin<BaseType> operator=(TOutlinerModelMixin<BaseType>&&) = delete;
+
+	virtual FViewModel* AsViewModel() { return this; }
+	virtual const FViewModel* AsViewModel() const { return this; }
+};
+
+class SEQUENCER_API FOutlinerItemModel : public TOutlinerModelMixin<FViewModel>
+{
+public:
+	UE_SEQUENCER_DECLARE_CASTABLE(FOutlinerItemModel, FOutlinerItemModelMixin);
+};
+
 } // namespace Sequencer
 } // namespace UE
 
