Sequencer: Removed ForceFilteredOut hack now that it is no longer required

diff --git a/Engine/Source/Editor/Sequencer/Private/MVVM/ViewModels/OutlinerItemModel.cpp b/Engine/Source/Editor/Sequencer/Private/MVVM/ViewModels/OutlinerItemModel.cpp
--- a/Engine/Source/Editor/Sequencer/Private/MVVM/ViewModels/OutlinerItemModel.cpp
+++ b/Engine/Source/Editor/Sequencer/Private/MVVM/ViewModels/OutlinerItemModel.cpp
@@ -130,7 +130,6 @@ FOutlinerItemModelMixin::FOutlinerItemModelMixin()
 	: OutlinerChildList(EViewModelListType::Outliner)
 	, bInitializedExpansion(false)
 	, bInitializedPinnedState(false)
-	, bIsForcedFilteredOut(false)
 {
 }
 
@@ -140,16 +139,6 @@ TSharedPtr<FSequencerEditorViewModel> FOutlinerItemModelMixin::GetEditor() const
 	return SequenceModel ? SequenceModel->GetEditor() : nullptr;
 }
 
-bool FOutlinerItemModelMixin::IsForceFilteredOut() const
-{
-	return bIsForcedFilteredOut;
-}
-
-void FOutlinerItemModelMixin::SetForceFilteredOut(bool bInIsForcedFilteredOut)
-{
-	bIsForcedFilteredOut = bInIsForcedFilteredOut;
-}
-
 FName FOutlinerItemModelMixin::GetIdentifier() const
 {
 	return TreeItemIdentifier;
@@ -226,7 +215,7 @@ void FOutlinerItemModelMixin::SetExpansionWithoutSaving(bool bInIsExpanded)
 
 bool FOutlinerItemModelMixin::IsFilteredOut() const
 {
-	return bIsFilteredOut || bIsForcedFilteredOut;
+	return bIsFilteredOut;
 }
 
 bool FOutlinerItemModelMixin::IsPinned() const
diff --git a/Engine/Source/Editor/Sequencer/Private/SequencerNodeTree.cpp b/Engine/Source/Editor/Sequencer/Private/SequencerNodeTree.cpp
--- a/Engine/Source/Editor/Sequencer/Private/SequencerNodeTree.cpp
+++ b/Engine/Source/Editor/Sequencer/Private/SequencerNodeTree.cpp
@@ -1002,16 +1002,6 @@ bool FSequencerNodeTree::UpdateFilters()
 		const bool bIsNodeFilteredIn = !bHasActiveFilter || FilteredNodes.Contains(Item);
 		Item->SetFilteredOut(!bIsNodeFilteredIn);
 
-		if (FOutlinerItemModel* OutlinerItemModel = Item.AsModel()->CastThis<FOutlinerItemModel>())
-		{
-			if (OutlinerItemModel->IsForceFilteredOut())
-			{
-				// Skip things that are force-filtered-out because we don't want to see or count them
-				// in the outliner hierarchy.
-				continue;
-			}
-		}
-
 		++DisplayNodeCount;
 	}
 
diff --git a/Engine/Source/Editor/Sequencer/Public/MVVM/ViewModels/OutlinerItemModel.h b/Engine/Source/Editor/Sequencer/Public/MVVM/ViewModels/OutlinerItemModel.h
--- a/Engine/Source/Editor/Sequencer/Public/MVVM/ViewModels/OutlinerItemModel.h
+++ b/Engine/Source/Editor/Sequencer/Public/MVVM/ViewModels/OutlinerItemModel.h
@@ -43,9 +43,6 @@ public:
 	FOutlinerItemModelMixin();
 
 	TSharedPtr<FSequencerEditorViewModel> GetEditor() const;
-
-	bool IsForceFilteredOut() const;
-	void SetForceFilteredOut(bool bInIsForcedFilteredOut);
 	
 	/*~ IOutlinerExtension */
 	FName GetIdentifier() const override;
@@ -116,7 +113,6 @@ private:
 	FCurveEditorTreeItemID CurveEditorItemID;
 	mutable bool bInitializedExpansion;
 	mutable bool bInitializedPinnedState;
-	bool bIsForcedFilteredOut;
 };
 
 template<typename BaseType>
