Curve Editor: Add marker hotkeys to curve editor commands

Step to Next Mark (ctrl-shift-.)
Step to Previous Mark (ctrl-shift-,)
Toggle Mark at Play Position (m)

diff --git a/Engine/Source/Editor/Sequencer/Private/Sequencer.cpp b/Engine/Source/Editor/Sequencer/Private/Sequencer.cpp
--- a/Engine/Source/Editor/Sequencer/Private/Sequencer.cpp
+++ b/Engine/Source/Editor/Sequencer/Private/Sequencer.cpp
@@ -13371,6 +13371,9 @@ void FSequencer::BindCommands()
 	CurveEditorSharedBindings->MapAction(Commands.JumpBackward,         *SequencerCommandBindings->GetActionForCommand(Commands.JumpBackward));
 	CurveEditorSharedBindings->MapAction(Commands.StepToNextKey,		*SequencerCommandBindings->GetActionForCommand(Commands.StepToNextKey));
 	CurveEditorSharedBindings->MapAction(Commands.StepToPreviousKey, 	*SequencerCommandBindings->GetActionForCommand(Commands.StepToPreviousKey));
+	CurveEditorSharedBindings->MapAction(Commands.StepToNextMark,       *SequencerCommandBindings->GetActionForCommand(Commands.StepToNextMark));
+	CurveEditorSharedBindings->MapAction(Commands.StepToPreviousMark,   *SequencerCommandBindings->GetActionForCommand(Commands.StepToPreviousMark));
+	CurveEditorSharedBindings->MapAction(Commands.ToggleMarkAtPlayPosition, *SequencerCommandBindings->GetActionForCommand(Commands.ToggleMarkAtPlayPosition));
 
 	CurveEditorSharedBindings->MapAction(Commands.AddTransformKey, *SequencerCommandBindings->GetActionForCommand(Commands.AddTransformKey));
 	CurveEditorSharedBindings->MapAction(Commands.AddTranslationKey, *SequencerCommandBindings->GetActionForCommand(Commands.AddTranslationKey));
