Sequencer: Removed setting to toggle Combined Keyframes and added a setting to toggle Layer Bars

Disable layer bars for Niagara by default

diff --git a/Engine/Config/BaseEditorPerProjectUserSettings.ini b/Engine/Config/BaseEditorPerProjectUserSettings.ini
--- a/Engine/Config/BaseEditorPerProjectUserSettings.ini
+++ b/Engine/Config/BaseEditorPerProjectUserSettings.ini
@@ -880,6 +880,7 @@ bKeepCursorInPlayRange=false
 bShowRangeSlider=true
 LoopMode=SLM_Loop
 bCleanPlaybackMode=false
+bShowLayerBars=false
 
 [TimeOfDaySequenceEditor SequencerSettings]
 bShowRangeSlider=true
diff --git a/Engine/Source/Editor/Sequencer/Private/MVVM/ViewModels/FolderModel.cpp b/Engine/Source/Editor/Sequencer/Private/MVVM/ViewModels/FolderModel.cpp
--- a/Engine/Source/Editor/Sequencer/Private/MVVM/ViewModels/FolderModel.cpp
+++ b/Engine/Source/Editor/Sequencer/Private/MVVM/ViewModels/FolderModel.cpp
@@ -17,6 +17,7 @@
 
 #include "Sequencer.h"
 #include "SequencerNodeTree.h"
+#include "SequencerSettings.h"
 #include "SequencerUtilities.h"
 #include "Styling/AppStyle.h"
 #include "Framework/MultiBox/MultiBoxBuilder.h"
@@ -96,10 +97,16 @@
 	// Only need to every create a layer bar once
 	if (!LayerBar)
 	{
-		LayerBar = MakeShared<FLayerBarModel>(This);
-		LayerBar->SetLinkedOutlinerItem(AsShared());
+		TSharedPtr<FSequencerEditorViewModel> EditorViewModel = GetEditor();
+		TSharedPtr<FSequencer> Sequencer = EditorViewModel->GetSequencerImpl();
 
-		TrackAreaChildren.AddChild(LayerBar);
+		if (Sequencer->GetSequencerSettings()->GetShowLayerBars())
+		{
+			LayerBar = MakeShared<FLayerBarModel>(This);
+			LayerBar->SetLinkedOutlinerItem(AsShared());
+
+			TrackAreaChildren.AddChild(LayerBar);
+		}
 	}
 
 	// Keep our existing children alive
diff --git a/Engine/Source/Editor/Sequencer/Private/MVVM/ViewModels/ObjectBindingModel.cpp b/Engine/Source/Editor/Sequencer/Private/MVVM/ViewModels/ObjectBindingModel.cpp
--- a/Engine/Source/Editor/Sequencer/Private/MVVM/ViewModels/ObjectBindingModel.cpp
+++ b/Engine/Source/Editor/Sequencer/Private/MVVM/ViewModels/ObjectBindingModel.cpp
@@ -21,6 +21,7 @@
 #include "Sequencer.h"
 #include "SequencerNodeTree.h"
 #include "SequencerCommands.h"
+#include "SequencerSettings.h"
 #include "SequencerUtilities.h"
 #include "ISequencerTrackEditor.h"
 #include "SObjectBindingTag.h"
@@ -188,10 +189,16 @@ void FObjectBindingModel::OnConstruct()
 {
 	if (!LayerBar)
 	{
-		LayerBar = MakeShared<FLayerBarModel>(AsShared());
-		LayerBar->SetLinkedOutlinerItem(AsShared());
+		TSharedPtr<FSequencerEditorViewModel> EditorViewModel = GetEditor();
+		TSharedPtr<FSequencer> Sequencer = EditorViewModel->GetSequencerImpl();
 
-		GetChildrenForList(&TrackAreaList).AddChild(LayerBar);
+		if (Sequencer->GetSequencerSettings()->GetShowLayerBars())
+		{
+			LayerBar = MakeShared<FLayerBarModel>(AsShared());
+			LayerBar->SetLinkedOutlinerItem(AsShared());
+
+			GetChildrenForList(&TrackAreaList).AddChild(LayerBar);
+		}
 	}
 
 	UMovieScene* MovieScene = OwnerModel->GetMovieScene();
diff --git a/Engine/Source/Editor/Sequencer/Private/SSequencer.cpp b/Engine/Source/Editor/Sequencer/Private/SSequencer.cpp
--- a/Engine/Source/Editor/Sequencer/Private/SSequencer.cpp
+++ b/Engine/Source/Editor/Sequencer/Private/SSequencer.cpp
@@ -2322,7 +2322,7 @@ TSharedRef<SWidget> SSequencer::MakeViewMenu()
 	MenuBuilder.AddMenuEntry( FSequencerCommands::Get().ToggleAutoScroll );
 	MenuBuilder.AddMenuEntry( FSequencerCommands::Get().ToggleShowRangeSlider );
 	MenuBuilder.AddMenuEntry( FSequencerCommands::Get().ToggleShowSelectedNodesOnly );
-	MenuBuilder.AddMenuEntry( FSequencerCommands::Get().ToggleCombinedKeyframes );
+	MenuBuilder.AddMenuEntry( FSequencerCommands::Get().ToggleLayerBars );
 	MenuBuilder.AddMenuEntry( FSequencerCommands::Get().ToggleChannelColors );
 	MenuBuilder.AddMenuEntry( FSequencerCommands::Get().ToggleShowStatusBar );
 	MenuBuilder.AddMenuEntry( FSequencerCommands::Get().ToggleShowPreAndPostRoll );
diff --git a/Engine/Source/Editor/Sequencer/Private/Sequencer.cpp b/Engine/Source/Editor/Sequencer/Private/Sequencer.cpp
--- a/Engine/Source/Editor/Sequencer/Private/Sequencer.cpp
+++ b/Engine/Source/Editor/Sequencer/Private/Sequencer.cpp
@@ -4419,6 +4419,22 @@ void FSequencer::ResetPerMovieSceneData()
 	//  needed for audio track decompression
 }
 
+void FSequencer::RefreshUI()
+{
+	using namespace UE::Sequencer;
+
+	Selection.Empty();
+
+	UMovieSceneSequence* FocusedSequence = GetFocusedMovieSceneSequence();
+	FMovieSceneSequenceID SequenceID = GetFocusedTemplateID();
+
+	TSharedPtr<FSequenceModel> RootSequenceModel = ViewModel->GetRootModel().ImplicitCast();
+	RootSequenceModel->SetSequence(nullptr, MovieSceneSequenceID::Root);
+	RootSequenceModel->SetSequence(FocusedSequence, SequenceID);
+
+	RefreshTree();
+}
+
 TSharedRef<SWidget> FSequencer::MakeTransportControls(bool bExtended)
 {
 	FEditorWidgetsModule& EditorWidgetsModule = FModuleManager::Get().LoadModuleChecked<FEditorWidgetsModule>( "EditorWidgets" );
@@ -12746,12 +12762,13 @@ void FSequencer::BindCommands()
 		FExecuteAction::CreateSP( this, &FSequencer::FindInContentBrowser ) );
 
 	SequencerCommandBindings->MapAction(
-		Commands.ToggleCombinedKeyframes,
+		Commands.ToggleLayerBars,
 		FExecuteAction::CreateLambda( [this]{
-			Settings->SetShowCombinedKeyframes( !Settings->GetShowCombinedKeyframes() );
+			Settings->SetShowLayerBars( !Settings->GetShowLayerBars() );
+			RefreshUI();
 		} ),
 		FCanExecuteAction::CreateLambda( []{ return true; } ),
-		FIsActionChecked::CreateLambda( [this]{ return Settings->GetShowCombinedKeyframes(); } ) );
+		FIsActionChecked::CreateLambda( [this]{ return Settings->GetShowLayerBars(); } ) );
 
 	SequencerCommandBindings->MapAction(
 		Commands.ToggleChannelColors,
@@ -13335,22 +13352,7 @@ void FSequencer::BindCommands()
 
 	SequencerCommandBindings->MapAction(
 		Commands.RefreshUI,
-		FExecuteAction::CreateLambda(
-			[this]
-			{
-				this->Selection.Empty();
-
-				UMovieSceneSequence*  FocusedSequence = GetFocusedMovieSceneSequence();
-				FMovieSceneSequenceID SequenceID      = GetFocusedTemplateID();
-
-				TSharedPtr<FSequenceModel> RootSequenceModel = this->ViewModel->GetRootModel().ImplicitCast();
-				RootSequenceModel->SetSequence(nullptr, MovieSceneSequenceID::Root);
-				RootSequenceModel->SetSequence(FocusedSequence, SequenceID);
-
-				RefreshTree();
-			}
-		),
-		FCanExecuteAction());
+		FExecuteAction::CreateSP( this, &FSequencer::RefreshUI));
 
 	// We want a subset of the commands to work in the Curve Editor too, but bound to our functions. This minimizes code duplication
 	// while also freeing us up from issues that result from Sequencer already using two lists (for which our commands might be spread
diff --git a/Engine/Source/Editor/Sequencer/Private/Sequencer.h b/Engine/Source/Editor/Sequencer/Private/Sequencer.h
--- a/Engine/Source/Editor/Sequencer/Private/Sequencer.h
+++ b/Engine/Source/Editor/Sequencer/Private/Sequencer.h
@@ -854,6 +854,9 @@ protected:
 	/** Reset data about a movie scene when pushing or popping a movie scene. */
 	void ResetPerMovieSceneData();
 
+	/** Forcibly refresh the UI */
+	void RefreshUI();
+
 	/** Update the time bounds to the focused movie scene */
 	void UpdateTimeBoundsToFocusedMovieScene();
 
diff --git a/Engine/Source/Editor/Sequencer/Private/SequencerCommands.cpp b/Engine/Source/Editor/Sequencer/Private/SequencerCommands.cpp
--- a/Engine/Source/Editor/Sequencer/Private/SequencerCommands.cpp
+++ b/Engine/Source/Editor/Sequencer/Private/SequencerCommands.cpp
@@ -136,7 +136,7 @@ void FSequencerCommands::RegisterCommands()
 	UI_COMMAND( ToggleSnapCurveValueToInterval, "Snap Curve Key Values", "Snap curve keys to the value snapping interval", EUserInterfaceActionType::ToggleButton, FInputChord() );
 
 	UI_COMMAND( FindInContentBrowser, "Find in Content Browser", "Find the viewed sequence asset in the content browser", EUserInterfaceActionType::Button, FInputChord() );
-	UI_COMMAND( ToggleCombinedKeyframes, "Combined Keyframes", "Show/hide the combined keyframes at the top node level", EUserInterfaceActionType::ToggleButton, FInputChord() );
+	UI_COMMAND( ToggleLayerBars, "Layer Bars", "Show/hide the layer bars to edit keyframes in bulk", EUserInterfaceActionType::ToggleButton, FInputChord() );
 	UI_COMMAND( ToggleChannelColors, "Channel Colors", "Show/hide the channel colors in the track area", EUserInterfaceActionType::ToggleButton, FInputChord() );
 	UI_COMMAND( ToggleShowSelectedNodesOnly, "Selected Nodes Only", "Show selected nodes only", EUserInterfaceActionType::ToggleButton, FInputChord() );
 	
diff --git a/Engine/Source/Editor/Sequencer/Private/SequencerSettings.cpp b/Engine/Source/Editor/Sequencer/Private/SequencerSettings.cpp
--- a/Engine/Source/Editor/Sequencer/Private/SequencerSettings.cpp
+++ b/Engine/Source/Editor/Sequencer/Private/SequencerSettings.cpp
@@ -44,7 +44,7 @@ USequencerSettings::USequencerSettings( const FObjectInitializer& ObjectInitiali
 	bLeftMouseDragDoesMarquee = false;
 	ZeroPadFrames = 0;
 	JumpFrameIncrement = FFrameNumber(5);
-	bShowCombinedKeyframes = true;
+	bShowLayerBars = true;
 	bInfiniteKeyAreas = false;
 	bShowChannelColors = false;
 	bShowStatusBar = true;
@@ -538,16 +538,16 @@ void USequencerSettings::SetJumpFrameIncrement(FFrameNumber InJumpFrameIncrement
 	}
 }
 
-bool USequencerSettings::GetShowCombinedKeyframes() const
+bool USequencerSettings::GetShowLayerBars() const
 {
-	return bShowCombinedKeyframes;
+	return bShowLayerBars;
 }
 
-void USequencerSettings::SetShowCombinedKeyframes(bool InbShowCombinedKeyframes)
+void USequencerSettings::SetShowLayerBars(bool InbShowLayerBars)
 {
-	if (bShowCombinedKeyframes != InbShowCombinedKeyframes)
+	if (bShowLayerBars != InbShowLayerBars)
 	{
-		bShowCombinedKeyframes = InbShowCombinedKeyframes;
+		bShowLayerBars = InbShowLayerBars;
 		SaveConfig();
 	}
 }
diff --git a/Engine/Source/Editor/Sequencer/Public/SequencerCommands.h b/Engine/Source/Editor/Sequencer/Public/SequencerCommands.h
--- a/Engine/Source/Editor/Sequencer/Public/SequencerCommands.h
+++ b/Engine/Source/Editor/Sequencer/Public/SequencerCommands.h
@@ -343,8 +343,8 @@ public:
 	/** Finds the viewed sequence asset in the content browser. */
 	TSharedPtr< FUICommandInfo > FindInContentBrowser;
 
-	/** Toggles whether to show combined keys at the top node level. */
-	TSharedPtr< FUICommandInfo > ToggleCombinedKeyframes;
+	/** Toggles whether to show layer bars to edit keyframes in bulk. */
+	TSharedPtr< FUICommandInfo > ToggleLayerBars;
 
 	/** Toggles whether to show channel colors in the track area. */
 	TSharedPtr< FUICommandInfo > ToggleChannelColors;
diff --git a/Engine/Source/Editor/Sequencer/Public/SequencerSettings.h b/Engine/Source/Editor/Sequencer/Public/SequencerSettings.h
--- a/Engine/Source/Editor/Sequencer/Public/SequencerSettings.h
+++ b/Engine/Source/Editor/Sequencer/Public/SequencerSettings.h
@@ -269,10 +269,10 @@ public:
 	/** Set the number of frames to increment when jumping forwards/backwards */
 	void SetJumpFrameIncrement(FFrameNumber InJumpFrameIncrement);
 
-	/** @return true if showing combined keyframes at the top node */
-	bool GetShowCombinedKeyframes() const;
-	/** Set whether to show combined keyframes at the top node */ 
-	void SetShowCombinedKeyframes(bool bInShowCombinedKeyframes);
+	/** @return true if showing layer bars */
+	bool GetShowLayerBars() const;
+	/** Set whether to show layer bars */ 
+	void SetShowLayerBars(bool bInShowLayerBars);
 
 	/** @return true if key areas are infinite */
 	bool GetInfiniteKeyAreas() const;
@@ -526,9 +526,9 @@ protected:
 	UPROPERTY( config, EditAnywhere, Category=Timeline )
 	FFrameNumber JumpFrameIncrement;
 
-	/** Enable or disable the combined keyframes at the top node level. Disabling can improve editor performance. */
+	/** Enable or disable the layer bars to edit keyframes in bulk. */
 	UPROPERTY( config, EditAnywhere, Category=Timeline )
-	bool bShowCombinedKeyframes;
+	bool bShowLayerBars;
 
 	/** Enable or disable setting key area sections as infinite by default. */
 	UPROPERTY( config, EditAnywhere, Category=Timeline )
