Sequencer: Fixed a UI bug where having multiple custom movie renderers showed all of them being checked instead of just the current one.

diff --git a/Engine/Source/Editor/Sequencer/Private/SSequencer.cpp b/Engine/Source/Editor/Sequencer/Private/SSequencer.cpp
--- a/Engine/Source/Editor/Sequencer/Private/SSequencer.cpp
+++ b/Engine/Source/Editor/Sequencer/Private/SSequencer.cpp
@@ -2757,11 +2757,9 @@ TSharedRef<SWidget> SSequencer::MakeRenderMovieMenu()
 				FUIAction(
 					FExecuteAction::CreateLambda([this, MovieRendererName] { SequencerPtr.Pin()->GetSequencerSettings()->SetMovieRendererName(MovieRendererName); }),
 					FCanExecuteAction(),
-					FIsActionChecked::CreateLambda([this] 
+					FIsActionChecked::CreateLambda([this, MovieRendererName]
 					{ 
-						ISequencerModule& SequencerModule = FModuleManager::LoadModuleChecked<ISequencerModule>("Sequencer");
-						IMovieRendererInterface* MovieRender = SequencerModule.GetMovieRenderer(SequencerPtr.Pin()->GetMovieRendererName());
-						return MovieRender != nullptr;
+						return MovieRendererName == SequencerPtr.Pin()->GetMovieRendererName();
 					})),
 				NAME_None,
 				EUserInterfaceActionType::RadioButton
@@ -2777,9 +2775,7 @@ TSharedRef<SWidget> SSequencer::MakeRenderMovieMenu()
 				FCanExecuteAction(),
 				FIsActionChecked::CreateLambda([this]
 				{
-					ISequencerModule& SequencerModule = FModuleManager::LoadModuleChecked<ISequencerModule>("Sequencer");
-					IMovieRendererInterface* MovieRender = SequencerModule.GetMovieRenderer(SequencerPtr.Pin()->GetMovieRendererName());
-					return MovieRender == nullptr;
+					return SequencerPtr.Pin()->GetMovieRendererName() == TEXT("MovieSceneCapture");
 				})),
 			NAME_None,
 			EUserInterfaceActionType::RadioButton
