Don't capture shared-pointers because they can extend the lifetime of the Sequencer editor past the point where we need them freed

diff --git a/Engine/Source/Editor/Sequencer/Private/SSequencer.cpp b/Engine/Source/Editor/Sequencer/Private/SSequencer.cpp
--- a/Engine/Source/Editor/Sequencer/Private/SSequencer.cpp
+++ b/Engine/Source/Editor/Sequencer/Private/SSequencer.cpp
@@ -1331,8 +1331,15 @@ void SSequencer::HandleOutlinerNodeSelectionChanged()
 			//since Selection ticks before any filter is applied and the curve editor tree is udpated in FSequencer::Tick we need to sync 
 			//the curve editor on the next tick so the filter is applied first, otherwise the curve tree hierarchy won't be up to date
 			//and the curve may not be found.
-			GEditor->GetTimerManager()->SetTimerForNextTick([Sequencer, this]()
+			TWeakPtr<FSequencer> WeakSequencer = SequencerPtr;
+			GEditor->GetTimerManager()->SetTimerForNextTick([WeakSequencer, this]()
 			{
+				TSharedPtr<FSequencer> Sequencer = WeakSequencer.Pin();
+				if (!Sequencer.IsValid())
+				{
+					return;
+				}
+
 				const TSet<TWeakPtr<FViewModel>>& SelectedDisplayNodes = Sequencer->GetSelection().GetSelectedOutlinerItems();
 				TSharedPtr<FCurveEditor> CurveEditor = Sequencer->GetCurveEditor();
 				TSharedRef<FSequencerNodeTree> NodeTree = Sequencer->GetNodeTree();
