Fix hover state of key navigation buttons

We were passing the first top-level channel to the widget but that was incorrect because it's not hoverable, and even if it was, it wouldn't be the thing that is hovered when the mouse cursor is on the widget. The widget is on the track item, so that's what should be passed.

Also, when clicking on the "add key" button and next/previous key buttons, we're not supposed to only look at the first top-level channel. If there are multiple top-level channels, we should key all of them (as applicable) if we click the "add key" button, and we should take into account all of their keys to navigate to the next/previous ones.

diff --git a/Engine/Source/Editor/Sequencer/Private/MVVM/Views/SOutlinerTrackView.cpp b/Engine/Source/Editor/Sequencer/Private/MVVM/Views/SOutlinerTrackView.cpp
--- a/Engine/Source/Editor/Sequencer/Private/MVVM/Views/SOutlinerTrackView.cpp
+++ b/Engine/Source/Editor/Sequencer/Private/MVVM/Views/SOutlinerTrackView.cpp
@@ -113,7 +113,7 @@
 		.AutoWidth()
 		.VAlign(VAlign_Center)
 		[
-			SNew(SSequencerKeyNavigationButtons, TopLevelChannel, EditorViewModel->GetSequencer())
+			SNew(SSequencerKeyNavigationButtons, DataModel, EditorViewModel->GetSequencer())
 		];
 	}
 	else
